Renamed experimental plugin from DataInterface to AnimNext

diff --git a/Engine/Plugins/Experimental/AnimNext/AnimNext.uplugin b/Engine/Plugins/Experimental/AnimNext/AnimNext.uplugin
new file mode 100644
--- /dev/null
+++ b/Engine/Plugins/Experimental/AnimNext/AnimNext.uplugin
@@ -0,0 +1,47 @@
+{
+	"FileVersion" : 3,
+	"Version" : 1,
+	"VersionName" : "0.1",
+	"FriendlyName" : "Anim Next",
+	"Description" : "Framework for defining functional data flow for animation systems",
+	"Category" : "Animation",
+	"CreatedBy" : "Epic Games, Inc.",
+	"CreatedByURL" : "https://epicgames.com",
+	"DocsURL" : "",
+	"MarketplaceURL" : "",
+	"SupportURL" : "",
+	"EnabledByDefault" : false,
+	"CanContainContent" : false,
+	"IsBetaVersion" : true,
+	"Installed" : false,
+	"Plugins":
+	[
+		{
+			"Name": "ControlRig",
+			"Enabled": true
+		}
+	],
+	"Modules" :
+	[
+		{
+			"Name" : "AnimNextInterface",
+			"Type" : "Runtime",
+			"LoadingPhase" : "PreDefault"
+		},
+		{
+			"Name" : "AnimNextInterfaceGraph",
+			"Type" : "Runtime",
+			"LoadingPhase" : "Default"
+		},
+		{
+			"Name" : "AnimNextInterfaceGraphEditor",
+			"Type" : "Editor",
+			"LoadingPhase" : "Default"
+		},
+		{
+			"Name" : "AnimNextInterfaceGraphUncookedOnly",
+			"Type" : "UncookedOnly",
+			"LoadingPhase" : "Default"
+		}
+	]
+}
diff --git a/Engine/Plugins/Experimental/AnimNext/Source/AnimNextInterface/AnimNextInterface.Build.cs b/Engine/Plugins/Experimental/AnimNext/Source/AnimNextInterface/AnimNextInterface.Build.cs
new file mode 100644
--- /dev/null
+++ b/Engine/Plugins/Experimental/AnimNext/Source/AnimNextInterface/AnimNextInterface.Build.cs
@@ -0,0 +1,19 @@
+// Copyright Epic Games, Inc. All Rights Reserved.
+
+namespace UnrealBuildTool.Rules
+{
+	public class AnimNextInterface : ModuleRules
+	{
+		public AnimNextInterface(ReadOnlyTargetRules Target) : base(Target)
+		{	
+			PublicDependencyModuleNames.AddRange(
+				new string[]
+				{
+					"Core",
+					"CoreUObject",
+				}
+			);
+
+		}
+	}
+}
\ No newline at end of file
diff --git a/Engine/Plugins/Experimental/AnimNext/Source/AnimNextInterface/Private/AnimNextInterface.cpp b/Engine/Plugins/Experimental/AnimNext/Source/AnimNextInterface/Private/AnimNextInterface.cpp
new file mode 100644
--- /dev/null
+++ b/Engine/Plugins/Experimental/AnimNext/Source/AnimNextInterface/Private/AnimNextInterface.cpp
@@ -0,0 +1,58 @@
+// Copyright Epic Games, Inc. All Rights Reserved.
+
+#include "AnimNextInterface.h"
+#include "IAnimNextInterface.h"
+#include "AnimNextInterfaceContext.h"
+
+namespace UE::AnimNext::Interface::Private
+{
+
+static bool CheckCompatibility(const IAnimNextInterface* InAnimNextInterface, const UE::AnimNext::Interface::FContext& InContext)
+{
+	check(InAnimNextInterface);
+	return InContext.GetResultParam().GetType().GetName() == InAnimNextInterface->GetReturnTypeName();
+}
+
+}
+
+bool IAnimNextInterface::GetDataIfCompatibleInternal(const UE::AnimNext::Interface::FContext& InContext) const
+{
+	if(UE::AnimNext::Interface::Private::CheckCompatibility(this, InContext))
+	{
+		return GetDataRawInternal(InContext);
+	}
+
+	return false;
+}
+
+bool IAnimNextInterface::GetData(const UE::AnimNext::Interface::FContext& Context) const
+{
+	return GetDataIfCompatibleInternal(Context);
+}
+
+bool IAnimNextInterface::GetDataChecked(const UE::AnimNext::Interface::FContext& Context) const
+{
+	check(UE::AnimNext::Interface::Private::CheckCompatibility(this, Context));
+	return GetDataRawInternal(Context);
+}
+
+bool IAnimNextInterface::GetData(const UE::AnimNext::Interface::FContext& Context, UE::AnimNext::Interface::FParam& OutResult) const
+{
+	const UE::AnimNext::Interface::FContext CallingContext = Context.WithResult(OutResult);
+	return GetDataIfCompatibleInternal(CallingContext);
+}
+
+bool IAnimNextInterface::GetDataChecked(const UE::AnimNext::Interface::FContext& Context, UE::AnimNext::Interface::FParam& OutResult) const
+{
+	const UE::AnimNext::Interface::FContext CallingContext = Context.WithResult(OutResult);
+	check(UE::AnimNext::Interface::Private::CheckCompatibility(this, CallingContext));
+	return GetDataRawInternal(CallingContext);
+}
+
+bool IAnimNextInterface::GetDataRawInternal(const UE::AnimNext::Interface::FContext& InContext) const
+{
+	// TODO: debug recording here on context
+
+	const UE::AnimNext::Interface::FContext CallContext = InContext.WithCallRaw(this);
+	return GetDataImpl(CallContext);
+}
diff --git a/Engine/Plugins/Experimental/AnimNext/Source/AnimNextInterface/Private/AnimNextInterfaceContext.cpp b/Engine/Plugins/Experimental/AnimNext/Source/AnimNextInterface/Private/AnimNextInterfaceContext.cpp
new file mode 100644
--- /dev/null
+++ b/Engine/Plugins/Experimental/AnimNext/Source/AnimNextInterface/Private/AnimNextInterfaceContext.cpp
@@ -0,0 +1,282 @@
+// Copyright Epic Games, Inc. All Rights Reserved.
+
+#include "AnimNextInterfaceContext.h"
+#include "IAnimNextInterface.h"
+#include "IAnimNextInterfaceParameters.h"
+#include "AnimNextInterfaceTypes.h"
+
+namespace UE::AnimNext::Interface
+{
+
+struct FThreadContextData : TThreadSingleton<FThreadContextData>
+{
+	TArray<const FContext*> ContextStack;
+};
+
+// @TODO: we might need to do this for all FContext constructions?
+// As parameters are gathered using the FContext linked list, we always need to access the top of the stack to get 
+// correct parameters. Shortcutting to the bottom of the stack via FThreadContext misses intervening parameters
+FThreadContext::FThreadContext(const FContext& InContext)
+{
+	FThreadContextData& ContextData = FThreadContextData::Get();
+	ContextData.ContextStack.Push(&InContext);
+}
+
+FThreadContext::~FThreadContext()
+{
+	FThreadContextData& ContextData = FThreadContextData::Get();
+	ContextData.ContextStack.Pop();
+}
+
+const FContext& FThreadContext::Get()
+{
+	const FThreadContextData& ContextData = FThreadContextData::Get();
+	return *ContextData.ContextStack.Top();
+}
+
+FContext::FContext(float InDeltaTime, FState& InState, FParamStorage& InParamStorage, IAnimNextInterfaceParameters* InParameters)
+	: State(&InState)
+	, ParamStorage(&InParamStorage)
+	, Parameters(InParameters)
+	, DeltaTime(InDeltaTime)
+	, BranchingMask(MAX_uint64)
+{
+}
+
+FContext::~FContext()
+{
+	// Remove any handles before the storage goes out of scope
+	AdditionalParameterHandles.Empty();
+
+	if (BlockHandle != InvalidBlockHandle)
+	{
+		ParamStorage->ReleaseBlock(BlockHandle);
+	}
+}
+
+FContext FContext::WithResult(FParam& InResult) const
+{
+	FContext NewContext(DeltaTime, *State, *ParamStorage, InResult);
+	NewContext.Parent = this;
+	NewContext.Root = (Root == nullptr) ? this : Root;
+	NewContext.CallstackHash = CallstackHash;
+	NewContext.BranchingMask = BranchingMask;
+
+	return NewContext;
+}
+
+FContext FContext::WithParameter(FName ParameterId, const FParam& InParameter) const
+{
+	FContext NewContext(DeltaTime, *State, *ParamStorage, *Result);
+	NewContext.Parent = this;
+	NewContext.Root = (Root == nullptr) ? this : Root;
+	NewContext.CallstackHash = CallstackHash;
+	NewContext.BranchingMask = BranchingMask;
+	NewContext.AddParameters({ TPairInitializer(ParameterId, InParameter) });
+	return NewContext;
+}
+
+FContext FContext::WithParameters(TArrayView<const TPair<FName, FParam>> InParameters) const
+{
+	FContext NewContext(DeltaTime, *State, *ParamStorage, *Result);
+	NewContext.Parent = this;
+	NewContext.Root = (Root == nullptr) ? this : Root;
+	NewContext.CallstackHash = CallstackHash;
+	NewContext.BranchingMask = BranchingMask;
+	NewContext.AddParameters(InParameters);
+
+	return NewContext;
+}
+
+FContext FContext::WithResultAndParameters(FParam& InResult, TArrayView<const TPair<FName, FParam>> InParameters) const
+{
+	FContext NewContext(DeltaTime, *State, *ParamStorage, InResult);
+	NewContext.Parent = this;
+	NewContext.Root = (Root == nullptr) ? this : Root;
+	NewContext.CallstackHash = CallstackHash;
+	NewContext.BranchingMask = BranchingMask;
+	NewContext.AddParameters(InParameters);
+
+	return NewContext;
+}
+
+int32 FContext::GetParametersSize(TArrayView<const TPair<FName, FParam>> InParameters, TArray<int32> &ParamAllocSizes) const
+{
+	int32 TotalParamAllocSize = 0;
+
+	const int32 NumParameters = InParameters.Num();
+
+	for (int i = 0; i < NumParameters; i++)
+	{
+		const TPair<FName, FParam>& ParamPair = InParameters[i];
+		const FParam& SourceParam = ParamPair.Value;
+
+		if (EnumHasAnyFlags(SourceParam.GetFlags(), FParam::EFlags::Stored))
+		{
+			const FParamType& ParamType = SourceParam.GetType();
+
+			const int32 NumElem = SourceParam.GetNumElements();
+			const int32 ParamAlignment = ParamType.GetAlignment();
+			const int32 ParamSize = ParamType.GetSize();
+
+			const int32 ParamAllocSize = NumElem * Align(ParamSize, ParamAlignment);
+
+			ParamAllocSizes[i] = ParamAllocSize;
+			TotalParamAllocSize += ParamAllocSize;
+		}
+	}
+
+	return TotalParamAllocSize;
+}
+
+void FContext::AddParameters(TArrayView<const TPair<FName, FParam>> InParameters)
+{
+	int32 TotalParamAllocSize = 0;
+
+	const int32 NumParameters = InParameters.Num();
+
+	TArray<int32> ParamAllocSizes;
+	ParamAllocSizes.Reserve(NumParameters);
+	ParamAllocSizes.AddZeroed(NumParameters);
+
+	TotalParamAllocSize = GetParametersSize(InParameters, ParamAllocSizes);
+
+	uint8* TargetMemory = nullptr;
+	check(BlockHandle == InvalidBlockHandle);
+
+	if (TotalParamAllocSize > 0)
+	{
+		const FParamStorage::TBlockDataPair BlockData = ParamStorage->RequestBlock(TotalParamAllocSize);
+
+		BlockHandle = BlockData.Key;
+		TargetMemory = BlockData.Value;
+	}
+
+	for(int i = 0; i < NumParameters; i++)
+	{
+		const TPair<FName, FParam>& ParamPair = InParameters[i];
+		const FParam& SourceParam = ParamPair.Value;
+		const int32 ParamAllocatedSize = ParamAllocSizes[i];
+
+		if (ParamAllocatedSize > 0 && EnumHasAnyFlags(SourceParam.GetFlags(), FParam::EFlags::Stored))
+		{
+			const FParamType& ParamType = ParamPair.Value.GetType();
+
+			const FParamType::ParamCopyFunction ParamCopyFunction = ParamType.GetParamCopyFunction();
+
+			FParam ClonedParam = ParamCopyFunction(SourceParam, TargetMemory, ParamAllocatedSize);
+
+			TargetMemory += ParamAllocatedSize;
+
+			AdditionalParameters.Add(ParamPair.Key, ClonedParam);
+		}
+		else
+		{
+			AdditionalParameters.Add(ParamPair.Key, ParamPair.Value);
+		}
+	}
+}
+
+FContext FContext::WithParameters(IAnimNextInterfaceParameters* InParameters) const
+{
+	FContext NewContext(DeltaTime, *State, *ParamStorage, *Result);
+	NewContext.Parameters = InParameters;
+	NewContext.Parent = this;
+	NewContext.Root = (Root == nullptr) ? this : Root;
+	NewContext.CallstackHash = CallstackHash;
+	NewContext.BranchingMask = BranchingMask;
+
+	return NewContext;
+}
+
+FContext FContext::CreateSubContext() const
+{
+	FContext NewContext(DeltaTime, *State, *ParamStorage, *Result);
+	NewContext.Parent = this;
+	NewContext.Root = (Root == nullptr) ? this : Root;
+	NewContext.CallstackHash = CallstackHash;
+	NewContext.BranchingMask = BranchingMask;
+
+	return NewContext;
+}
+
+bool FContext::GetParameter(FName InKey, FParam& OutParam) const
+{
+	// Check parent contexts in turn. Parameters can be overriden by each scoped context.
+	for(const FContext* CurrentContext = this; CurrentContext != nullptr; CurrentContext = CurrentContext->Parent)
+	{
+		// Check parameter provider
+		if(CurrentContext->Parameters)
+		{
+			if(CurrentContext->Parameters->GetParameter(InKey, OutParam))
+			{
+				return true;
+			}
+		}
+		// Check additional parameters
+		// Note this is currently 'else if' because we dont allow creating new contexts with both a IAnimNextInterfaceParameters
+		// interface and AdditionalParameters
+		else if(!CurrentContext->AdditionalParameters.IsEmpty())
+		{
+			// Find the parameter
+			if(const FParam* FoundParameter = CurrentContext->AdditionalParameters.Find(InKey))
+			{
+				// Check type compatibility
+				if(FoundParameter->CanAssignTo(OutParam))
+				{
+					OutParam = *FoundParameter;
+					return true;
+				}
+			}
+		}
+		// Check additional parameter handles
+		// Note this is currently 'else if' because we dont allow creating new contexts with both 
+		// AdditionalParameters and AdditionalParametersHandles
+		else if (!CurrentContext->AdditionalParameterHandles.IsEmpty())
+		{
+			// Find the parameter
+			if (const FHParam* HParam = CurrentContext->AdditionalParameterHandles.Find(InKey))
+			{
+				const FParam* Param = ParamStorage->GetParam(HParam->ParamHandle);
+				if (Param != nullptr && Param->CanAssignTo(OutParam))
+				{
+					OutParam = Param;
+					return true;
+				}
+			}
+		}
+	}
+
+	return false;
+}
+
+
+TParam<const float> FContext::GetDeltaTimeParam() const
+{
+	return TWrapParam<const float>(&DeltaTime);
+}
+
+FContext FContext::WithCallRaw(const IAnimNextInterface* InAnimNextInterface) const
+{
+	FContext NewContext(DeltaTime, *State, *ParamStorage, *Result);
+	NewContext.Parent = this;
+	NewContext.Root = (Root == nullptr) ? this : Root;
+	NewContext.CallstackHash = HashCombineFast(CallstackHash, GetTypeHash(InAnimNextInterface));
+	NewContext.BranchingMask = BranchingMask;
+	return NewContext;
+}
+
+void FContext::FlushRelevancy() const
+{
+	// Flush all relevancy-based allocations if they were not used this update
+	// @TODO: this could be more efficient if we use a more linear iteration here 
+	for(auto It = State->RelevancyValueMap.CreateIterator(); It; ++It)
+	{
+		if(It->Value.UpdateCounter != UpdateCounter)
+		{
+			It.RemoveCurrent();
+		}
+	}
+}
+
+}
\ No newline at end of file
diff --git a/Engine/Plugins/Experimental/AnimNext/Source/AnimNextInterface/Private/AnimNextInterfaceModule.cpp b/Engine/Plugins/Experimental/AnimNext/Source/AnimNextInterface/Private/AnimNextInterfaceModule.cpp
new file mode 100644
--- /dev/null
+++ b/Engine/Plugins/Experimental/AnimNext/Source/AnimNextInterface/Private/AnimNextInterfaceModule.cpp
@@ -0,0 +1,26 @@
+// Copyright Epic Games, Inc. All Rights Reserved.
+
+#include "CoreMinimal.h"
+#include "Modules/ModuleInterface.h"
+#include "Modules/ModuleManager.h"
+#include "Misc/CoreDelegates.h"
+#include "AnimNextInterfaceParam.h"
+
+namespace UE::AnimNext::Interface
+{
+
+class FModule : public IModuleInterface
+{
+public:
+	virtual void StartupModule() override
+	{
+		FCoreDelegates::OnPostEngineInit.AddLambda([]()
+		{
+			FParamType::FRegistrar::RegisterDeferredTypes();
+		});
+	}
+};
+
+}
+
+IMPLEMENT_MODULE(UE::AnimNext::Interface::FModule, AnimNextInterface)
\ No newline at end of file
diff --git a/Engine/Plugins/Experimental/AnimNext/Source/AnimNextInterface/Private/AnimNextInterfaceParam.cpp b/Engine/Plugins/Experimental/AnimNext/Source/AnimNextInterface/Private/AnimNextInterfaceParam.cpp
new file mode 100644
--- /dev/null
+++ b/Engine/Plugins/Experimental/AnimNext/Source/AnimNextInterface/Private/AnimNextInterfaceParam.cpp
@@ -0,0 +1,196 @@
+﻿// Copyright Epic Games, Inc. All Rights Reserved.
+
+#include "AnimNextInterfaceParam.h"
+#include "AnimNextInterfaceParamStorage.h"
+
+namespace UE::AnimNext::Interface
+{
+
+static FParamType DefaultType;
+static TArray<FParamType*> GTypeRegistry;
+static TArray<TUniqueFunction<void(void)>> GDeferredTypes;
+
+FParam::FParam(const FParamType& InType, void* InData, EFlags InFlags)
+	: Data(InData)
+	, NumElements((InData != nullptr) ? 1 : 0)
+	, TypeId(InType.GetTypeId())
+	, Flags(InFlags)
+{
+	check(TypeId != INDEX_NONE);
+	check(Data);
+}
+
+FParam::FParam(const FParamType& InType, void* InData, int32 InNumElements, EFlags InFlags)
+	: Data(InData)
+	, NumElements(InNumElements)
+	, TypeId(InType.GetTypeId())
+	, Flags(InFlags)
+{
+}
+
+FParam::FParam(const FParamType& InType, EFlags InFlags)
+	: Data(nullptr)
+	, NumElements(0)
+	, TypeId(InType.GetTypeId())
+	, Flags(InFlags)
+{
+	check(TypeId != INDEX_NONE);
+}
+
+FParam::FParam(const FParam* InOtherParam)
+	: Data(EnumHasAnyFlags(InOtherParam->Flags, FParam::EFlags::Embedded) ? (void*)&InOtherParam->Data : InOtherParam->Data)
+	, NumElements(InOtherParam->NumElements)
+	, TypeId(InOtherParam->TypeId)
+	, Flags(InOtherParam->Flags)
+{
+	EnumRemoveFlags(Flags, FParam::EFlags::Embedded); // For compatibility reasons, I can not keep the embedded value if we create a copy of the FParam
+}
+
+bool FParam::CanAssignTo(const FParam& InParam) const
+{
+	return CanAssignWith(InParam.GetType(), InParam.Flags, InParam.GetNumElements());
+}
+
+bool FParam::CanAssignWith(const FParamType& InType, EFlags InFlags, int32 InNumElements, FStringBuilderBase* OutReasonPtr) const
+{
+	// Check type
+	if(TypeId != InType.GetTypeId())
+	{
+		if(OutReasonPtr)
+		{
+			const FParamType& ThisType = FParamType::FRegistrar::GetTypeById(TypeId);
+			OutReasonPtr->Appendf(TEXT("Types do not match: %s and %s"), *ThisType.GetName().ToString(), *InType.GetName().ToString());
+		}
+		return false;
+	}
+	
+	// Check const if not empty, else we assume it is an empty param and can be written, mutable or not
+	if(InNumElements > 0 && !EnumHasAnyFlags(InFlags, EFlags::Mutable))
+	{
+		if(OutReasonPtr)
+		{
+			OutReasonPtr->Append(TEXT("Target is not mutable"));
+		}
+		return false;
+	}
+
+	// Check batching
+	if (EnumHasAnyFlags(Flags, EFlags::Batched) != EnumHasAnyFlags(InFlags, EFlags::Batched))
+	{
+		// If different batching, allow it if the target is empty and the source is 1 (this enables passing a single value to a context that expects N)
+		if (NumElements != 1 || InNumElements != 0)
+		{
+			if (OutReasonPtr)
+			{
+				OutReasonPtr->Append(TEXT("Incompatible batching"));
+			}
+			return false;
+		}
+	}
+	
+	return true;
+}
+
+void FParamType::FRegistrar::RegisterType(FParamType& InType, const UScriptStruct* InStruct, FName InName, uint32 InSize, uint32 InAlignment, const ParamCopyFunction InParamCopyFunctionPtr)
+{
+	InType.Struct = InStruct;
+	InType.Name = InName;
+	InType.Size = InSize;
+	InType.Alignment = InAlignment;
+	InType.ParamCopyFunctionPtr = InParamCopyFunctionPtr;
+
+	check(InType.Name != NAME_None);
+	check(InType.Size != 0);
+	check(InType.Alignment != 0);
+	check(InType.ParamCopyFunctionPtr != nullptr);
+	
+	InType.TypeId = GTypeRegistry.Add(&InType);
+}
+
+const FParamType& FParamType::FRegistrar::FindTypeByName(FName InName)
+{
+	for(FParamType* Type : GTypeRegistry)
+	{
+		if(Type->Name == InName)
+		{
+			return *Type;
+		}
+	}
+
+	checkf(false, TEXT("Type %s not found"), *InName.ToString());
+	return DefaultType;
+}
+
+const FParamType& FParamType::FRegistrar::GetTypeById(uint16 InTypeId)
+{
+	if(GTypeRegistry.IsValidIndex(InTypeId))
+	{
+		return *GTypeRegistry[InTypeId];
+	}
+
+	checkf(false, TEXT("Type ID %d not found"), InTypeId);
+	return DefaultType;	
+}
+
+FParamType::FRegistrar::FRegistrar(TUniqueFunction<void(void)>&& InFunction)
+{
+	check(IsInGameThread());
+	GDeferredTypes.Add(MoveTemp(InFunction));
+}
+
+void FParamType::FRegistrar::RegisterDeferredTypes()
+{
+	check(IsInGameThread());
+	
+	for(TUniqueFunction<void(void)>& Function : GDeferredTypes)
+	{
+		Function();
+	}
+
+	GDeferredTypes.Empty();
+}
+
+const FParamType& FParam::GetType() const
+{
+	check(GTypeRegistry.IsValidIndex(TypeId));
+	return *GTypeRegistry[TypeId];
+}
+/*
+static FParamType& GetTypeForProperty(const FProperty* InProperty)
+{
+	const FName TypeName = *InProperty->GetCPPType();
+	return FParamType::FRegistrar::FindTypeByName(TypeName);
+}
+
+FPropertyParam::FPropertyParam(const FProperty* InProperty, UObject* InContainer)
+	: FParam(GetTypeForProperty(InProperty), InProperty->ContainerPtrToValuePtr<void>(InContainer), EFlags::Const)	// Property parameters are assumed to be const
+{
+}
+*/
+
+FHParam::FHParam(FParamStorage* InOwnerStorage, FParamHandle InParamHandle)
+	: OwnerStorage(InOwnerStorage)
+	, ParamHandle(InParamHandle)
+{
+}
+
+FHParam::FHParam(const FHParam& Other)
+	: OwnerStorage(Other.OwnerStorage)
+	, ParamHandle(Other.ParamHandle)
+{
+	if (ParamHandle != InvalidParamHandle)
+	{
+		OwnerStorage->IncRefCount(ParamHandle);
+	}
+}
+
+FHParam::~FHParam()
+{
+	if (ParamHandle != InvalidParamHandle)
+	{
+		check(OwnerStorage != nullptr);
+		OwnerStorage->DecRefCount(ParamHandle);
+	}
+}
+
+} // end namespace UE::AnimNext::Interface
diff --git a/Engine/Plugins/Experimental/AnimNext/Source/AnimNextInterface/Private/AnimNextInterfaceParamStorage.cpp b/Engine/Plugins/Experimental/AnimNext/Source/AnimNextInterface/Private/AnimNextInterfaceParamStorage.cpp
new file mode 100644
--- /dev/null
+++ b/Engine/Plugins/Experimental/AnimNext/Source/AnimNextInterface/Private/AnimNextInterfaceParamStorage.cpp
@@ -0,0 +1,158 @@
+// Copyright Epic Games, Inc. All Rights Reserved.
+
+#include "AnimNextInterfaceParamStorage.h"
+#include "HAL/UnrealMemory.h"
+
+namespace UE::AnimNext::Interface
+{
+
+
+FParamStorage::FParamStorage(int32 InMaxParams, int32 InAllocatedMemorySize, int32 InMaxBlocks, bool InAllowGrowing)
+	: AllowGrowing(InAllowGrowing)
+{
+	check(InMaxParams > 0 && InMaxBlocks > 0 && InAllocatedMemorySize > 0);
+
+	Parameters.Reserve(InMaxParams);
+	Parameters.AddDefaulted(InMaxParams);
+
+	RefCounts.Reserve(InMaxParams);
+	RefCounts.AddZeroed(InMaxParams);
+
+	ParamMemoryHandles.Reserve(InMaxParams);
+	for (int32 i = 0; i < InMaxParams; i++)
+	{
+		ParamMemoryHandles.Add(InvalidBlockHandle);
+	}
+
+	// --- ---
+
+	RawMemory = (uint8*)FMemory::Malloc(InAllocatedMemorySize);
+	RawMemorySize = InAllocatedMemorySize;
+
+	BlockOffets.Reserve(InMaxBlocks);
+	BlockOffets.AddZeroed(InMaxBlocks);
+
+	BlockSizes.Reserve(InMaxBlocks);
+	BlockSizes.AddZeroed(InMaxBlocks);
+
+	BlockFreeFlag.Reserve(InMaxBlocks);
+	BlockFreeFlag.Init(true, InMaxBlocks);
+
+	FreeBlockIndex = 0;
+}
+
+FParamStorage::~FParamStorage()
+{
+	// There should be no blocks of memory or HParams in memory at this point
+	check(FreeBlockIndex == 0);
+	check(ParamFreeIndex == 0);
+
+	FMemory::Free(RawMemory);
+}
+
+FParamStorage::TBlockDataPair FParamStorage::RequestBlock(int32 RequestedBlockSize)
+{
+	FParamStorageHandle BlockHandle = InvalidBlockHandle;
+	uint8* BlockMemory = nullptr;
+
+	check((CurrentBlocksSize + RequestedBlockSize < RawMemorySize) || AllowGrowing);
+	check((FreeBlockIndex >= 0 && FreeBlockIndex < (BlockSizes.Num() - 1)) || AllowGrowing);
+
+	if (CheckBufferSize(RequestedBlockSize) && CheckNumBlocks(FreeBlockIndex))
+	{
+		const int32 BlockIndex = FreeBlockIndex;
+
+		BlockOffets[BlockIndex] = (BlockIndex > 0) ? BlockOffets[BlockIndex - 1] + BlockSizes[BlockIndex - 1] : 0;
+		BlockSizes[BlockIndex] = RequestedBlockSize;
+		CurrentBlocksSize += RequestedBlockSize;
+		BlockFreeFlag[BlockIndex] = false;
+		FreeBlockIndex++;
+
+		BlockHandle = BlockIndex;
+		BlockMemory = RawMemory + BlockOffets[BlockIndex];
+	}
+
+	return TPair<FParamStorageHandle, uint8*>(BlockHandle, BlockMemory);
+}
+
+void FParamStorage::ReleaseBlock(FParamStorageHandle BlockHandle)
+{
+	int32 BlockIndex = BlockHandle;
+
+	check(BlockIndex >= 0 && BlockIndex < FreeBlockIndex);
+
+	// just remove the offset here
+	BlockFreeFlag[BlockIndex] = true;
+	
+	// if it is the last block, release the memory and deal with the prev blocks flagged as free, if any
+	if (BlockIndex == FreeBlockIndex - 1)
+	{
+		// Recover free slots
+		while (BlockIndex >= 0 && BlockFreeFlag[BlockIndex] == true)
+		{
+			BlockOffets[BlockIndex] = 0;
+			CurrentBlocksSize -= BlockSizes[BlockIndex];
+			BlockSizes[BlockIndex] = 0;
+			BlockIndex--;
+			FreeBlockIndex--;
+		}
+	}
+}
+
+uint8* FParamStorage::GetBlockMemory(FParamStorageHandle BlockHandle)
+{
+	uint8* BlockMemory = nullptr;
+
+	const int32 BlockIndex = BlockHandle;
+
+	if (ensure(BlockIndex >= 0 && BlockIndex < FreeBlockIndex))
+	{
+		bool IsSet = BlockFreeFlag[BlockIndex];
+		ensure(IsSet == false);
+		BlockMemory = (IsSet == false) ? RawMemory + BlockOffets[BlockIndex] : nullptr;
+	}
+
+	return BlockMemory;
+}
+
+bool FParamStorage::CheckBufferSize(int32 RequestedBlockSize)
+{
+	bool BufferSizeCorrect = true;
+
+	if (CurrentBlocksSize + RequestedBlockSize > RawMemorySize)
+	{
+		if (AllowGrowing)
+		{
+			// TODO : Grow
+			check(false);
+		}
+		else
+		{
+			BufferSizeCorrect = false;
+		}
+	}
+
+	return BufferSizeCorrect;
+}
+
+bool FParamStorage::CheckNumBlocks(int32 BlockIndex)
+{
+	bool ValidIndex = true;
+
+	if (BlockIndex >= BlockSizes.Num())
+	{
+		if (AllowGrowing)
+		{
+			// TODO : Grow
+		}
+		else
+		{
+			ValidIndex = false;
+		}
+		// TODO : Grow
+	}
+
+	return ValidIndex;
+}
+
+} // end namespace UE::AnimNext::Interface
diff --git a/Engine/Plugins/Experimental/AnimNext/Source/AnimNextInterface/Private/AnimNextInterfaceState.cpp b/Engine/Plugins/Experimental/AnimNext/Source/AnimNextInterface/Private/AnimNextInterfaceState.cpp
new file mode 100644
--- /dev/null
+++ b/Engine/Plugins/Experimental/AnimNext/Source/AnimNextInterface/Private/AnimNextInterfaceState.cpp
@@ -0,0 +1,58 @@
+﻿// Copyright Epic Games, Inc. All Rights Reserved.
+
+#include "AnimNextInterfaceState.h"
+#include "AnimNextInterfaceContext.h"
+
+namespace UE::AnimNext::Interface
+{
+
+FParam* FState::FindStateRaw(const FInterfaceKeyWithIdAndStack& InKey, const FContext& InContext, EStatePersistence InPersistence)
+{
+	check(NumElements == InContext.GetNum())
+
+	FParam* ExistingValue = nullptr;
+	TMap<FInterfaceKeyWithIdAndStack, FParam>* ValueMap = nullptr; 
+
+	switch(InPersistence)
+	{
+	case EStatePersistence::Relevancy:
+		ExistingValue = RelevancyValueMap.Find(InKey);
+		if(ExistingValue)
+		{
+			// An access to a state param counts as a 'relevant use' for this update
+			static_cast<FRelevancyParam*>(ExistingValue)->UpdateCounter = InContext.UpdateCounter;
+		}
+		break;
+	case EStatePersistence::Permanent:
+		ExistingValue = PermanentValueMap.Find(InKey);
+		break;
+	default:
+		check(false);
+		break;
+	}
+	
+	return ExistingValue;
+}
+
+FParam* FState::AllocateState(const FInterfaceKeyWithIdAndStack& InKey, const FContext& InContext, const FParamType& InType, EStatePersistence InPersistence)
+{
+	void* Data = FMemory::Malloc(InType.GetSize() * NumElements, InType.GetAlignment());
+	
+	switch(InPersistence)
+	{
+	case EStatePersistence::Relevancy:
+		// TODO: chunked allocator for relevancy-based stuff?
+		return &RelevancyValueMap.Add(InKey, FRelevancyParam(InType, Data, InContext.GetBatchedFlags(), InContext.UpdateCounter));
+		break;
+	case EStatePersistence::Permanent:
+		return &PermanentValueMap.Add(InKey, FParam(InType, Data, InContext.GetBatchedFlags()));
+		break;
+	default:
+		check(false);
+		break;
+	}
+	
+	return nullptr;
+}
+
+}
diff --git a/Engine/Plugins/Experimental/AnimNext/Source/AnimNextInterface/Private/AnimNextInterfaceTypes.cpp b/Engine/Plugins/Experimental/AnimNext/Source/AnimNextInterface/Private/AnimNextInterfaceTypes.cpp
new file mode 100644
--- /dev/null
+++ b/Engine/Plugins/Experimental/AnimNext/Source/AnimNextInterface/Private/AnimNextInterfaceTypes.cpp
@@ -0,0 +1,13 @@
+﻿// Copyright Epic Games, Inc. All Rights Reserved.
+
+#include "AnimNextInterfaceTypes.h"
+#include "IAnimNextInterface.h"
+
+namespace UE::AnimNext::Interface
+{
+
+#define ANIM_NEXT_INTERFACE_TYPE(Type, Identifier) IMPLEMENT_ANIM_NEXT_INTERFACE_PARAM_TYPE(Type, Identifier)
+#include "AnimNextInterfaceTypes.inl"
+#undef ANIM_NEXT_INTERFACE_TYPE
+
+}
\ No newline at end of file
diff --git a/Engine/Plugins/Experimental/AnimNext/Source/AnimNextInterface/Public/AnimNextInterface.h b/Engine/Plugins/Experimental/AnimNext/Source/AnimNextInterface/Public/AnimNextInterface.h
new file mode 100644
--- /dev/null
+++ b/Engine/Plugins/Experimental/AnimNext/Source/AnimNextInterface/Public/AnimNextInterface.h
@@ -0,0 +1,99 @@
+﻿// Copyright Epic Games, Inc. All Rights Reserved.
+
+#pragma once
+
+#include "CoreMinimal.h"
+#include "AnimNextInterfaceContext.h"
+
+namespace UE::AnimNext::Interface
+{
+
+// Call the interface with the provided context (which should contain the parameters and result for the interface)
+template<typename InterfaceType>
+static bool GetDataSafe(const TScriptInterface<InterfaceType> AnimNextInterface, const FContext& InContext)
+{
+	if(AnimNextInterface.GetInterface() != nullptr)
+	{
+		return AnimNextInterface.GetInterface()->GetData(InContext);
+	}
+
+	return false;
+}
+
+//// Set the Result into the context and call the interface
+//template<typename InterfaceType, typename ValueType>
+//static bool GetDataSafe(const TScriptInterface<InterfaceType> AnimNextInterface, FContext& InContext, ValueType& Result)
+//{
+//	if (AnimNextInterface.GetInterface() != nullptr)
+//	{
+//		if constexpr (TIsDerivedFrom<ValueType, FParam>::Value)
+//		{
+//			const FContext ResultHolder = InContext.WithResult(Result);
+//			return AnimNextInterface.GetInterface()->GetData(ResultHolder);
+//		}
+//		else
+//		{
+//			// Support containers with Num() and GetData()
+//			if constexpr (TModels<Private::CSizedContainerWithAccessibleDataAsRawPtr, ValueType>::Value)
+//			{
+//				TWrapParam<ValueType::ElementType> WrappedResult(Result.GetData(), Result.Num());
+//				const FContext ResultHolder = InContext.WithResult(WrappedResult);
+//				return AnimNextInterface.GetInterface()->GetData(ResultHolder);
+//			}
+//			else
+//			{
+//				TWrapParam<ValueType> WrappedResult(&Result);
+//				const FContext ResultHolder = InContext.WithResult(WrappedResult);
+//				return AnimNextInterface.GetInterface()->GetData(ResultHolder);
+//			}
+//		}
+//	}
+//
+//	return false;
+//}
+
+// This version accepts a const Context, but it crates a sub context in order to be able to pass the Result
+// Using the non const version is a better option if possible
+template<typename InterfaceType, typename ValueType>
+static bool GetDataSafe(const TScriptInterface<InterfaceType> AnimNextInterface, const FContext& InContext, ValueType& Result)
+{
+	if(AnimNextInterface.GetInterface() != nullptr)
+	{
+		if constexpr (TIsDerivedFrom<ValueType, FParam>::Value)
+		{
+			return AnimNextInterface.GetInterface()->GetData(InContext, Result);
+		}
+		else
+		{
+			// Support containers with Num() and GetData()
+			if constexpr (TModels<Private::CSizedContainerWithAccessibleDataAsRawPtr, ValueType>::Value)
+			{
+				TWrapParam<typename ValueType::ElementType> WrappedResult(Result.GetData(), Result.Num());
+				return AnimNextInterface.GetInterface()->GetData(InContext, WrappedResult);
+			}
+			else
+			{
+				TWrapParam<ValueType> WrappedResult(&Result);
+				return AnimNextInterface.GetInterface()->GetData(InContext, WrappedResult);
+			}
+		}
+	}
+
+	return false;
+}
+
+template<typename ValueType>
+static bool GetParameter(const FContext& InContext, FName InKey, ValueType& Result)
+{
+	if constexpr (TIsDerivedFrom<ValueType, FParam>::Value)
+	{
+		return InContext.GetParameter(InKey, Result);
+	}
+	else
+	{
+		TWrapParam<ValueType> WrappedResult(&Result);
+		return InContext.GetParameter(InKey, WrappedResult);
+	}
+}
+
+}
\ No newline at end of file
diff --git a/Engine/Plugins/Experimental/AnimNext/Source/AnimNextInterface/Public/AnimNextInterfaceContext.h b/Engine/Plugins/Experimental/AnimNext/Source/AnimNextInterface/Public/AnimNextInterfaceContext.h
new file mode 100644
--- /dev/null
+++ b/Engine/Plugins/Experimental/AnimNext/Source/AnimNextInterface/Public/AnimNextInterfaceContext.h
@@ -0,0 +1,492 @@
+// Copyright Epic Games, Inc. All Rights Reserved.
+
+#pragma once
+
+#include "CoreMinimal.h"
+#include "UObject/Class.h"
+#include "AnimNextInterfaceParam.h"
+#include "AnimNextInterfaceState.h"
+#include "AnimNextInterfaceParamStorage.h"
+#include "AnimNextInterfaceKey.h"
+#include "Misc/MemStack.h"
+
+class IAnimNextInterface;
+class IAnimNextInterfaceParameters;
+
+namespace UE::AnimNext::Interface
+{
+
+struct FState;
+enum class EStatePersistence : uint8;
+
+// Helper to provide a non-callstack bridge for interfacing with data interface contexts
+// Declaring one of these on the stack will enable all calls in its scope to access the passed-in context
+// via FThreadContext::Get()
+struct ANIMNEXTINTERFACE_API FThreadContext
+{
+	FThreadContext(const FContext& InContext);
+	~FThreadContext();
+
+	static const FContext& Get();
+};
+
+// Context providing methods for mutating & interrogating the data interface runtime
+struct ANIMNEXTINTERFACE_API FContext
+{
+private:
+	friend class ::IAnimNextInterface;
+	friend struct FState;
+
+public:
+	static constexpr int32 MAX_BATCH_ELEMENTS = 64;
+
+	// Root public constructor. Constructs a context given a state.
+	FContext(float InDeltaTime, FState& InState, FParamStorage& InParamStorage, IAnimNextInterfaceParameters* InParameters = nullptr);
+	
+	FContext(const FContext& other) = delete;
+	FContext& operator=(const FContext&) = delete;
+
+	FContext(FContext&& Other)
+		: FContext()
+	{
+		Swap(*this, Other);
+	}
+
+	FContext& operator= (FContext&& Other)
+	{
+		Swap(*this, Other);
+		return *this;
+	}
+
+	~FContext();
+
+public:
+	// --- Sub Context Creation --- 
+
+	// Create a sub context from this one that includes the provided Result
+	FContext WithResult(FParam& InResult) const;
+
+	// Create a sub context from this one that includes the provided parameter
+	FContext WithParameter(FName ParameterId, const FParam& InParameter) const;
+
+	// Create a sub context from this one that includes the provided parameters
+	FContext WithParameters(TArrayView<const TPair<FName, FParam>> InParameters) const;
+
+	// Create a sub context from this one that includes the provided Result and parameters
+	FContext WithResultAndParameters(FParam& InResult, TArrayView<const TPair<FName, FParam>> InParameters) const;
+
+	// Create a sub context from this one that includes the provided interface parameter
+	FContext WithParameters(IAnimNextInterfaceParameters* InParameters) const;
+
+public:
+
+	// --- Interface for direct Param storage (prototype) ---
+
+	FContext CreateSubContext() const;
+
+	enum class EParamType : uint8
+	{
+		None		= 0,
+		Input		= 1 << 0,
+		Output		= 1 << 1,
+	};
+
+	// Add a parameter by value, copying it to the shared storage
+	template<typename ValueType>
+	FHParam AddValue(EParamType InParamType, FName ParameterId, ValueType &Value)
+	{
+		// Values are added mutable even if source is const
+		const FParam::EFlags Flags = FParam::EFlags::Value | FParam::EFlags::Mutable;
+
+		return AddParameter(InParamType, ParameterId, Flags, &Value);
+	}
+
+	// Add a parameter by value, copying it to the shared storage
+	template<typename ValueType>
+	FHParam AddValue(EParamType InParamType, FName ParameterId, ValueType&& Value)
+	{
+		// Values are added mutable even if source is const
+		const FParam::EFlags Flags = FParam::EFlags::Value | FParam::EFlags::Mutable;
+
+		return AddParameter(InParamType, ParameterId, Flags, &Value);
+	}
+
+	// Add a parameter by value, copying it to the shared storage
+	template<typename ValueType>
+	FHParam AddValue(EParamType InParamType, FName ParameterId, ValueType* Value)
+	{
+		// Values are added mutable even if source is const
+		const FParam::EFlags Flags = FParam::EFlags::Value | FParam::EFlags::Mutable;
+
+		return AddParameter(InParamType, ParameterId, Flags, Value);
+	}
+
+	// Add a parameter by reference, adding just a pointer to the shared storage
+	template<typename ValueType>
+	FHParam AddReference(EParamType InParamType, FName ParameterId, ValueType& Value)
+	{
+		const FParam::EFlags Flags = FParam::EFlags::Reference;
+
+		return AddParameter(InParamType, ParameterId, Flags, &Value);
+	}
+
+	// Add a parameter by reference, adding just a pointer to the shared storage
+	template<typename ValueType>
+	FHParam AddReference(EParamType InParamType, FName ParameterId, ValueType *Value)
+	{
+		FParam::EFlags Flags = FParam::EFlags::Reference;
+
+		return AddParameter(InParamType, ParameterId, Flags, Value);
+	}
+
+	template<typename ValueType>
+	FHParam AddParameter(EParamType InParamType, FName ParameterId, FParam::EFlags Flags, ValueType* Value)
+	{
+		FHParam ParamHandle;
+
+		check(InParamType == EParamType::Input || InParamType == EParamType::Output);
+		check(EnumHasAnyFlags(Flags, FParam::EFlags::Value) || EnumHasAnyFlags(Flags, FParam::EFlags::Reference));
+
+		// check if it is a simple FHParam copy
+		if constexpr (TIsDerivedFrom<ValueType, FHParam>::Value)
+		{
+			const FParam* ExistingParam = ParamStorage->GetParam(Value->ParamHandle);
+			check(ExistingParam != nullptr);
+
+			const bool bIsReadOnly = EnumHasAnyFlags(ExistingParam->Flags, FParam::EFlags::Mutable) == false;
+
+			// Verify that we want to create an Input param or if we want an output, the source is not read only
+			check(InParamType == EParamType::Input || bIsReadOnly == false);
+
+			// For now create a copy of the Param (TODO : change ref to value or value to ref, if the caller requests it?)
+			ParamHandle = *Value;
+		}
+		else
+		{
+			using MutableValueType = typename TRemoveConst<ValueType>::Type;
+
+			// If non const set the mutable flag
+			if constexpr (TIsConst<ValueType>::Value == false)
+			{
+				EnumAddFlags(Flags, FParam::EFlags::Mutable);
+			}
+
+			// deal with containers
+			if constexpr (TModels<Private::CSizedContainerWithAccessibleDataAsRawPtr, ValueType>::Value)
+			{
+				MutableValueType* ParamData = const_cast<MutableValueType>(static_cast<ValueType*>(Value->GetData()));
+				ParamHandle = EnumHasAnyFlags(Flags, FParam::EFlags::Value) 
+					? ParamHandle = ParamStorage->AddValue(ParamData, Value->Num(), Flags)
+					: ParamStorage->AddReference(ParamData, Value->Num(), Flags);
+			}
+			// and finally with single params
+			else
+			{
+				ParamHandle = EnumHasAnyFlags(Flags, FParam::EFlags::Value)
+					? ParamHandle = ParamStorage->AddValue(const_cast<MutableValueType*>(Value), 1, Flags)
+					: ParamStorage->AddReference(const_cast<MutableValueType*>(Value), 1, Flags);
+			}
+		}
+
+		AdditionalParameterHandles.Add(ParameterId, ParamHandle);
+
+		return ParamHandle;
+	}
+
+	// Get a TParam from a handle
+	template<typename ValueType>
+	TParam<ValueType> GetParameterChecked(const FHParam &InParamHandle) const
+	{
+		const FParam* Param = ParamStorage->GetParam(InParamHandle.ParamHandle);
+
+		TParam<ValueType> RetVal(Param, Param->GetFlags());
+
+		check(RetVal.IsValid());							// found something
+		checkSlow(Private::CheckParam<ValueType>(RetVal));	// the type found is valid for what was requested
+
+		return RetVal;
+	}
+
+	// Get a parameter from a handle as a specified type
+	template <typename ValueType>
+	ValueType& GetParameterAs(const FHParam& InParamHandle, int32* OutNumElems = nullptr)
+	{
+		const FParam* Param = ParamStorage->GetParam(InParamHandle.ParamHandle);
+
+		check(Param != nullptr);
+		check(Param->GetType().GetTypeId() == Private::TParamType<ValueType>::GetType().GetTypeId());
+
+		// If a non cost is requested, check the param has mutable flag
+		if constexpr (TIsConst<ValueType>::Value == false)
+		{
+			check(Param->IsMutable());
+		}
+
+		ValueType* ValueData = EnumHasAnyFlags(Param->Flags, FParam::EFlags::Embedded)
+			? static_cast<ValueType*>((void*)&Param->Data)
+			: static_cast<ValueType*>(Param->Data);
+
+		if (OutNumElems != nullptr)
+		{
+			*OutNumElems = Param->GetNumElements();
+		}
+
+		return *ValueData;
+	}
+
+	// --- Parameter management --- 
+
+	// Get a parameter if exist, returns false if not
+	bool GetParameter(FName InKey, FParam& OutParam) const;
+
+	// Get a parameter as a specified type, checking it exist and the type
+	template<typename ValueType>
+	TParam<ValueType> GetParameterChecked(FName InKey) const
+	{
+		TParam<ValueType> RetVal(GetBatchedFlags());
+
+		GetParameter(InKey, RetVal);
+
+		check(RetVal.IsValid());							// found something
+		checkSlow(Private::CheckParam<ValueType>(RetVal));	// the type found is valid for what was requested
+
+		return RetVal;
+	}
+
+	// Get a parameter as a specified type, returning a TOptional
+	template<typename ValueType>
+	TOptional<TParam<ValueType>> GetParameter(FName InKey) const
+	{
+		TOptional<TParam<ValueType>> RetVal;
+
+		TParam<ValueType> Param(FParam::EFlags::Mutable | GetBatchedFlags());
+		if (GetParameter(InKey, Param))
+		{
+			checkSlow(Private::CheckParam<ValueType>(Param));
+
+			RetVal = Param;
+		}
+
+		return RetVal;
+	}
+
+public:
+	// --- Result Management ---
+
+	// Get the ResultParam, checking it has been set
+	FParam& GetResultParam()
+	{
+		check(Result != nullptr);
+		return *Result;
+	}
+
+	// Get the ResultParam, checking it has been set
+	FParam& GetResultParam() const
+	{
+		check(Result != nullptr);
+		return *Result;
+	}
+
+	// Set a HParam as a result (for compatibility reasons)
+	void SetHParamAsResult(const FHParam& InHParam)
+	{
+		Result = ParamStorage->GetParam(InHParam.ParamHandle);
+	}
+
+	// Set a result value directly
+	// The receiver Result Param must be set prior to this call
+	template<typename ValueType>
+	void SetResult(const ValueType& InValue) const
+	{
+		// Support containers with Num() and GetData()
+		if constexpr (TModels<Private::CSizedContainerWithAccessibleDataAsRawPtr, ValueType>::Value)
+		{
+			checkSlow(Private::CheckParam<ValueType>(*Result));
+			checkSlow(InValue.Num() == Result->NumElements);
+
+			TArrayView<ValueType> Results(*static_cast<ValueType*>(Result->Data), Result->NumElements);
+			const ValueType* Data = InValue.GetData();
+			for(int32 ResultIndex = 0; ResultIndex < Result->NumElements; ++ResultIndex)
+			{
+				Results[ResultIndex] = Data[ResultIndex];
+			}
+		}
+		else
+		{
+			checkSlow(Private::CheckParam<ValueType>(*Result));
+			checkSlow(Result->NumElements == 1);
+
+			TArrayView<ValueType> Results(static_cast<ValueType*>(Result->Data), Result->NumElements);
+			for(ValueType& ResultElement : Results)
+			{
+				ResultElement = InValue;
+			}
+		}
+	}
+
+	// Get the current result as a TParam
+	template<typename ValueType>
+	TParam<ValueType> GetResult() const
+	{
+		// Types must match
+		checkSlow(Private::CheckParam<ValueType>(*Result));
+		return TParam<ValueType>(*Result, GetBatchedFlags());
+	}
+
+	// Get the current result as a specified raw type pointer
+	template<typename ValueType>
+	ValueType* GetResultChecked() const
+	{
+		// Types must match
+		checkSlow(Private::CheckParam<ValueType>(*Result));
+		return static_cast<ValueType*>(Result->Data);
+	}
+
+public:
+	// --- State management ---
+
+	// Gets (and allocates, if necessary) state for the specified node given this calling context
+	template<typename ValueType, EStatePersistence Persistence = EStatePersistence::Relevancy>
+	TParam<ValueType> GetState(const FInterfaceKeyWithId& InKey) const
+	{
+		return State->GetState<ValueType, Persistence>(InKey, *this, CallstackHash);
+	}
+
+	// Gets (and allocates, if necessary) state for the specified node given this calling context
+	template<typename ValueType, EStatePersistence Persistence = EStatePersistence::Relevancy>
+	TParam<ValueType> GetState(const IAnimNextInterface* InAnimNextInterface, uint32 InId) const
+	{
+		return State->GetState<ValueType, Persistence>(InAnimNextInterface, InId, *this, CallstackHash);
+	}
+
+public:
+	// --- Mix Context Utils ---
+
+	// Number of elements that this context represents
+	int32 GetNum() const
+	{
+		return State->NumElements;
+	}
+
+	// Access delta time as a param
+	TParam<const float> GetDeltaTimeParam() const;
+
+	// Raw access to delta time
+	float GetDeltaTime() const { return DeltaTime; }
+
+	// --- Batching management ---
+
+	// Get parameter flags for batching
+	FParam::EFlags GetBatchedFlags() const
+	{
+		return State->NumElements > 1 ? FParam::EFlags::Batched : FParam::EFlags::None;
+	}
+
+	FORCEINLINE_DEBUGGABLE bool ShouldProcessThisElement(int32 ElementIndex) const
+	{
+		check(ElementIndex < MAX_BATCH_ELEMENTS);
+		return (BranchingMask & (1ull << ElementIndex)) != 0;
+	}
+
+	FORCEINLINE_DEBUGGABLE const uint64 GetBranchingMask() const
+	{
+		return BranchingMask;
+	}
+
+	void SetBranchingMask(const TArrayView<int32> InEnabledElementIndexes)
+	{
+		check(InEnabledElementIndexes.Num() <= GetNum());
+
+		BranchingMask = 0;
+		for (int32 ElementIndex : InEnabledElementIndexes)
+		{
+			check(ElementIndex < MAX_BATCH_ELEMENTS);
+
+			BranchingMask |= (1ull << ElementIndex);
+		}
+	}
+
+	void SetBranchingMask(uint64 InBranchingMask)
+	{
+		BranchingMask = InBranchingMask;
+	}
+
+private:
+	FContext(float InDeltaTime, FState& InState, FParamStorage& InParamStorage, FParam& InResult)
+		: State(&InState)
+		, ParamStorage(&InParamStorage)
+		, Result(&InResult)
+		, DeltaTime(InDeltaTime)
+		, BranchingMask(MAX_uint64)
+	{}
+
+	FContext() = default;
+
+	FContext WithCallRaw(const IAnimNextInterface* InAnimNextInterface) const;
+
+	void FlushRelevancy() const;
+
+	int32 GetParametersSize(TArrayView<const TPair<FName, FParam>> InParameters, TArray<int32>& ParamAllocSizes) const;
+	void AddParameters(TArrayView<const TPair<FName, FParam>> InParameters);
+
+private:
+	TMap<FName, FParam> AdditionalParameters;
+	TMap<FName, FHParam> AdditionalParameterHandles;  // Temp param storage prototype
+	const FContext* Parent = nullptr;
+	const FContext* Root = nullptr;
+	FState* State = nullptr;
+	FParamStorage* ParamStorage = nullptr;
+	FParamStorageHandle BlockHandle = InvalidBlockHandle;
+	FParam* Result = nullptr;
+	IAnimNextInterfaceParameters* Parameters = nullptr;
+	float DeltaTime = 0.0f;
+	uint32 CallstackHash = 0;
+	uint32 UpdateCounter = 0;
+	uint64 BranchingMask = MAX_uint64;
+};
+
+// A typed param that owns it's own memory
+template<typename ValueType, typename AllocatorType = TMemStackAllocator<>>
+struct TAllocParam : public TParam<ValueType>
+{
+public:
+	TAllocParam(const FContext& InContext)
+		: TParam<ValueType>(InContext.GetBatchedFlags())
+	{
+		TParam<ValueType>::Data = GetDataForConstructor(InContext);
+	}
+
+	TAllocParam(const FContext& InContext, FParam::EFlags InFlags)
+		: TParam<ValueType>(InFlags)
+	{
+		TParam<ValueType>::Data = GetDataForConstructor(InContext);
+	}
+
+
+private:
+	void* GetDataForConstructor(const FContext& InContext)
+	{
+		ValueArray.SetNum(InContext.GetNum());
+		return ValueArray.GetData();
+	}
+
+	TArray<ValueType, AllocatorType> ValueArray;
+};
+
+
+// --- Batching helper macros ---
+
+#define SKIP_CONTEXT_ELEMENT_IF_MASKED(Context) \
+		if (!Context.ShouldProcessThisElement(i)) \
+			continue;
+
+// Helper macros to allow processing context element batches
+#define PROCESS_BATCH_ELEMENTS_START(Context)  \
+	const int32 ContextNumElements = Context.GetNum(); \
+	for (int32 i = 0; i < ContextNumElements; i++) \
+	{ \
+		SKIP_CONTEXT_ELEMENT_IF_MASKED(Context)
+
+#define PROCESS_BATCH_ELEMENTS_END }
+
+}
diff --git a/Engine/Plugins/Experimental/AnimNext/Source/AnimNextInterface/Public/AnimNextInterfaceKernel.h b/Engine/Plugins/Experimental/AnimNext/Source/AnimNextInterface/Public/AnimNextInterfaceKernel.h
new file mode 100644
--- /dev/null
+++ b/Engine/Plugins/Experimental/AnimNext/Source/AnimNextInterface/Public/AnimNextInterfaceKernel.h
@@ -0,0 +1,76 @@
+﻿// Copyright Epic Games, Inc. All Rights Reserved.
+
+#pragma once
+
+#include "CoreMinimal.h"
+#include "AnimNextInterfaceContext.h"
+
+namespace UE::AnimNext::Interface
+{
+
+struct ANIMNEXTINTERFACE_API FKernel
+{
+public:
+	// Runs a 'kernel' predicate that iterates over the supplied (optionally batched) params
+	// Number and type of params must match between the predicate & param pack
+	// The kernel supports passing a single element parameter where N are expected (as a compacted constant)
+	template<typename PredicateType, typename... ParamTypes>
+	static void Run(const FContext& InContext, PredicateType InPredicate, ParamTypes&... InArgs)
+	{
+		// Iterate per-element in the context, applying kernel predicate to all elements
+		const uint32 NumElements = InContext.GetNum();
+		for (uint32 ElementIndex = 0; ElementIndex < NumElements; ++ElementIndex)
+		{
+			if (InContext.ShouldProcessThisElement(ElementIndex) == false)
+				continue;
+
+			uint32 ParameterIndex = 0;
+
+			// Call predicate, expanding InArgs parameter pack from TParam<ParamType> -> ParamType
+			// and applying correct offset for each param (using pointer arithmetic)
+			InPredicate(
+				[&ParameterIndex](uint32 InElementIndex, auto& InArg) -> auto&
+				{
+					auto GetBasePtr = [](auto& InArg)
+					{
+						if constexpr (TIsDerivedFrom<ParamTypes, FParam>::Value)
+						{
+							using Type = decltype(typename ParamTypes::TValueType{});
+							return static_cast<Type*>(InArg.Data);
+						}
+						else
+						{
+							return static_cast<ParamTypes*>(&InArg);
+						}
+					};
+
+					auto GetOffset = [&ParameterIndex](auto& InArg, int32 InElementIndex) -> int32
+					{
+						if constexpr (TIsDerivedFrom<ParamTypes, FParam>::Value)
+						{
+							if (InArg.NumElements > 1)
+							{
+								check(InElementIndex < InArg.NumElements);
+								return InElementIndex;
+							}
+							else // single elements always use 0 offset, even if used on N element batch
+							{
+								check(InArg.NumElements > 0);
+								return 0;
+							}
+						}
+						else
+						{
+							return 0;
+						}
+					};
+
+					// note that this uses pointer arithmetic, so an offset of +1 means adding the sizeof of the param type
+					return *(GetBasePtr(InArg) + GetOffset(InArg, InElementIndex));
+				}(ElementIndex, InArgs)...
+			);
+		}
+	}
+};
+
+} // end namespace UE::AnimNext::Interface
\ No newline at end of file
diff --git a/Engine/Plugins/Experimental/AnimNext/Source/AnimNextInterface/Public/AnimNextInterfaceKey.h b/Engine/Plugins/Experimental/AnimNext/Source/AnimNextInterface/Public/AnimNextInterfaceKey.h
new file mode 100644
--- /dev/null
+++ b/Engine/Plugins/Experimental/AnimNext/Source/AnimNextInterface/Public/AnimNextInterfaceKey.h
@@ -0,0 +1,68 @@
+﻿// Copyright Epic Games, Inc. All Rights Reserved.
+
+#pragma once
+
+#include "CoreMinimal.h"
+#include "UObject/Class.h"
+#include "AnimNextInterfaceParam.h"
+
+class IAnimNextInterface;
+
+namespace UE::AnimNext::Interface
+{
+
+// Key value for cached data interface state
+struct ANIMNEXTINTERFACE_API FInterfaceKey
+{
+	FInterfaceKey(const IAnimNextInterface* InAnimNextInterface)
+		: AnimNextInterface(InAnimNextInterface)
+		, Hash(::GetTypeHash(InAnimNextInterface))
+	{}
+
+	FInterfaceKey(const FInterfaceKey& InExistingKey, uint32 InHash)
+		: AnimNextInterface(InExistingKey.AnimNextInterface)
+		, Hash(HashCombineFast(InExistingKey.Hash, InHash))
+	{}
+
+	const IAnimNextInterface* AnimNextInterface = nullptr;
+	uint32 Hash = 0;
+};
+
+// Key value for cached data interface state
+struct ANIMNEXTINTERFACE_API FInterfaceKeyWithId : FInterfaceKey
+{
+	FInterfaceKeyWithId(const FInterfaceKey& InInterfaceKey, uint32 InId)
+		: FInterfaceKey(InInterfaceKey, ::GetTypeHash(InId))
+		, Id(InId)
+	{}
+
+	FInterfaceKeyWithId(const FInterfaceKeyWithId& InInterfaceKeyWithId, uint32 InHash)
+		: FInterfaceKey(InInterfaceKeyWithId, HashCombineFast(InInterfaceKeyWithId.Hash, InHash))
+		, Id(InInterfaceKeyWithId.Id)
+	{}
+	
+	uint32 Id = 0;
+};
+
+// Key value for cached data interface state
+struct ANIMNEXTINTERFACE_API FInterfaceKeyWithIdAndStack : FInterfaceKeyWithId
+{
+	FInterfaceKeyWithIdAndStack(const FInterfaceKeyWithId& InInterfaceKey, uint32 InCallStackHash)
+		: FInterfaceKeyWithId(InInterfaceKey, InCallStackHash)
+		, CallStackHash(InCallStackHash)
+	{}
+
+	FORCEINLINE friend bool operator==(const FInterfaceKeyWithIdAndStack& A, const FInterfaceKeyWithIdAndStack& B)
+	{
+		return A.AnimNextInterface == B.AnimNextInterface && A.Id == B.Id && A.CallStackHash == B.CallStackHash;
+	}
+
+	FORCEINLINE friend uint32 GetTypeHash(const FInterfaceKeyWithIdAndStack& InKey)
+	{
+		return InKey.Hash;
+	}
+
+	uint32 CallStackHash = 0;
+};
+
+}
\ No newline at end of file
diff --git a/Engine/Plugins/Experimental/AnimNext/Source/AnimNextInterface/Public/AnimNextInterfaceParam.h b/Engine/Plugins/Experimental/AnimNext/Source/AnimNextInterface/Public/AnimNextInterfaceParam.h
new file mode 100644
--- /dev/null
+++ b/Engine/Plugins/Experimental/AnimNext/Source/AnimNextInterface/Public/AnimNextInterfaceParam.h
@@ -0,0 +1,594 @@
+// Copyright Epic Games, Inc. All Rights Reserved.
+
+#pragma once
+
+#include "CoreMinimal.h"
+#include "UObject/Class.h"
+#include "UObject/ScriptInterface.h"
+
+namespace UE::AnimNext::Interface
+{
+
+struct FParamType;
+struct FParamStorage;
+
+// Parameter/result/state memory wrapper
+// @TODO can we save memory here with some base/derived system for fundamental types vs struct types?
+struct ANIMNEXTINTERFACE_API FParam
+{
+public:
+	enum class EFlags : uint8
+	{
+		None		= 0,			// No flags
+		Mutable		= 1 << 0,		// Parameter is mutable, so can be mutated at runtime
+		Batched		= 1 << 1,		// Parameter is batched, so represents a number of values rather than a single value. Batch size is determined by NumElements and associated FContext
+		Stored		= 1 << 2,		// Parameter has to be stored on context storage
+		
+		// --- Flags added for parameter storage prototype ---
+		Value		= 1 << 3,		// Parameter will be stored as a Value
+		Reference	= 1 << 4,		// Parameter will be stored as a Reference (pointer)
+		Embedded	= 1 << 5,		// Parameter will be stored as a Value, but stored directly on the Data pointer
+	};
+
+	// Get the type of this param 
+	const FParamType& GetType() const;
+
+	// Check whether the supplied param can be written to by this param
+	// Verifies that types are compatible, destination is mutable and batching matches
+	bool CanAssignTo(const FParam& InParam) const;
+
+	// Helper function for CanAssignTo
+	bool CanAssignWith(const FParamType& InType, EFlags InFlags, int32 InNumelements, FStringBuilderBase* OutReasonPtr = nullptr) const;
+
+	bool IsMutable() const { return EnumHasAnyFlags(Flags, EFlags::Mutable); }
+
+	bool IsBatched() const { return EnumHasAnyFlags(Flags, EFlags::Batched); }
+
+	int32 GetNumElements() const { return NumElements; }
+
+	EFlags GetFlags() const { return Flags; }
+	
+	void* GetData() const { return Data; }
+
+protected:
+	friend struct FContext;
+	friend struct FState;
+	friend struct FKernel;
+	friend struct FParamStorage;
+	
+	FParam(const FParam* InOtherParam);
+	FParam(const FParamType& InType, void* InData, EFlags InFlags);
+	FParam(const FParamType& InType, void* InData, int32 InNumElements, EFlags InFlags);
+	FParam(const FParamType& InType, EFlags InFlags);
+
+	void* Data = nullptr;
+	int32 NumElements = 0;
+	int16 TypeId = INDEX_NONE;
+	EFlags Flags = EFlags::None;
+
+public:
+	// Internal use, but required for default constructed elements on containers
+	FParam() = default;
+
+	static FParam MakeParam_Internal(const FParamType& InType, void* InData, int32 InNumElements, EFlags InFlags)
+	{
+		return FParam(InType, InData, InNumElements, InFlags);
+	}
+};
+
+ENUM_CLASS_FLAGS(FParam::EFlags);
+
+
+// Simple type representation for data interfaces
+struct ANIMNEXTINTERFACE_API FParamType
+{
+	typedef FParam(*ParamCopyFunction)(const FParam& InSource, void* InMemory, int32 InAllocatedMemory);
+
+	FParamType() = default;
+
+	inline const UScriptStruct* GetStruct() const
+	{
+		return Struct;
+	}
+		
+	inline FName GetName() const
+	{
+		return Name;
+	}
+		
+	inline int32 GetSize() const
+	{
+		return Size;
+	}
+		
+	int32 GetAlignment() const
+	{
+		return Alignment;
+	}
+		
+	int16 GetTypeId() const
+	{
+		return TypeId;
+	}
+
+	ParamCopyFunction GetParamCopyFunction() const
+	{
+		return ParamCopyFunctionPtr;
+	}
+
+	struct ANIMNEXTINTERFACE_API FRegistrar
+	{
+		// Helper for registration (via IMPLEMENT_DATA_INTERFACE_PARAM_TYPE).
+		// Registers the type via deferred callback
+		FRegistrar(TUniqueFunction<void(void)>&& InFunction);
+
+		// Helper for registration (via IMPLEMENT_DATA_INTERFACE_PARAM_TYPE)
+		static void RegisterType(FParamType& InType, const UScriptStruct* InStruct, FName InName, uint32 InSize, uint32 InAlignment, const ParamCopyFunction InParamCopyFunctionPtr);
+
+		// Helper for registration
+		static void RegisterDeferredTypes();
+
+		// Find a type by name. If the type was not found, returns the default (invalid) type
+		static const FParamType& FindTypeByName(FName InName);
+
+		// Get a type by ID. If the type is invalid, returns the default (invalid) type
+		static const FParamType& GetTypeById(uint16 InTypeId);
+	};
+
+private:
+	const UScriptStruct* Struct = nullptr;
+	FName Name = NAME_None;
+	uint32 Size = 0;
+	uint16 Alignment = 0;
+	uint16 TypeId = INDEX_NONE;
+	ParamCopyFunction ParamCopyFunctionPtr = nullptr;
+};
+
+// A type-erased param that wraps FProperty memory
+struct FPropertyParam : FParam
+{
+	FPropertyParam(const FProperty* InProperty, UObject* InContainer);
+};
+
+namespace Private
+{
+
+template<typename ValueType> 
+struct TParamTypeImpl
+{
+	FORCEINLINE static const FParamType& GetType()
+	{
+		static_assert(sizeof(ValueType) > 0, "Unimplemented type - please include the file AnimNextInterfaceTypes.h or where DECLARE_DATA_INTERFACE_PARAM_TYPE is defined for your type");
+		static FParamType Default;
+		return Default;
+	}
+};
+
+// Regular types need to have GetType function for this to work
+template<typename ValueType> 
+struct TParamType
+{
+	FORCEINLINE static const FParamType& GetType()
+	{
+		return TParamTypeImpl<typename TRemoveConst<ValueType>::Type>::GetType();
+	}
+};
+
+template<typename ValueType>
+static bool CheckParam(const FParam& InParam)
+{
+	return InParam.GetType().GetTypeId() == Private::TParamType<ValueType>::GetType().GetTypeId();
+}
+
+}
+
+namespace Private
+{
+
+// Concept used to decide whether to use container-based construction of wrapped results
+struct CSizedContainerWithAccessibleDataAsRawPtr
+{
+	template<typename ContainerType>
+	auto Requires(ContainerType& Container) -> decltype(
+		Container.Num(),
+		Container.GetData()
+	);
+};
+
+}
+
+// A typed result which wraps the type-erased underlying param
+template<typename ValueType>
+struct TParam : FParam
+{
+public:
+	FORCEINLINE_DEBUGGABLE bool IsValid() const
+	{
+		return ((Data != nullptr || EnumHasAnyFlags(Flags, FParam::EFlags::Embedded))
+			&& TypeId != INDEX_NONE 
+			&& NumElements > 0);
+	}
+
+	FORCEINLINE_DEBUGGABLE const ValueType& operator[](int32 InIndex) const
+	{
+		check(InIndex >= 0 && (InIndex < NumElements || NumElements == 1));
+		check(Data != nullptr || EnumHasAnyFlags(Flags, FParam::EFlags::Embedded));
+
+		const ValueType* ValueData = EnumHasAnyFlags(Flags, FParam::EFlags::Embedded)
+			? static_cast<const ValueType*>((void*)&Data)
+			: static_cast<const ValueType*>(Data);
+		return (NumElements > InIndex) ? ValueData[InIndex] : ValueData[0];
+	}
+
+	FORCEINLINE_DEBUGGABLE ValueType& operator[](int32 InIndex)
+	{
+		check(InIndex >= 0 && (InIndex < NumElements || NumElements == 1));
+		check(Data != nullptr || EnumHasAnyFlags(Flags, FParam::EFlags::Embedded));
+
+		ValueType* ValueData = EnumHasAnyFlags(Flags, FParam::EFlags::Embedded)
+			? static_cast<ValueType*>((void*) &Data)
+			: static_cast<ValueType*>(Data);
+		return (NumElements > InIndex) ? ValueData[InIndex] : ValueData[0];
+	}
+
+	FORCEINLINE_DEBUGGABLE operator const ValueType& () const
+	{
+		check((Data != nullptr || EnumHasAnyFlags(Flags, FParam::EFlags::Embedded)) && NumElements == 1);
+
+		const ValueType* ValueData = EnumHasAnyFlags(Flags, FParam::EFlags::Embedded)
+			? static_cast<const ValueType*>((void*)&Data)
+			: static_cast<const ValueType*>(Data);
+
+		return *ValueData;
+	}
+
+	FORCEINLINE_DEBUGGABLE operator ValueType& ()
+	{
+		check((Data != nullptr || EnumHasAnyFlags(Flags, FParam::EFlags::Embedded)) && NumElements == 1);
+
+		ValueType* ValueData = EnumHasAnyFlags(Flags, FParam::EFlags::Embedded)
+			? static_cast<ValueType*>((void*)&Data)
+			: static_cast<ValueType*>(Data);
+
+		return *ValueData;
+	}
+
+	FORCEINLINE_DEBUGGABLE ValueType& operator*() const
+	{
+		check((Data != nullptr || EnumHasAnyFlags(Flags, FParam::EFlags::Embedded)) && NumElements == 1);
+
+		ValueType* ValueData = EnumHasAnyFlags(Flags, FParam::EFlags::Embedded)
+			? static_cast<ValueType*>((void*)&Data)
+			: static_cast<ValueType*>(Data);
+
+		return *ValueData;
+	}
+
+	FORCEINLINE_DEBUGGABLE ValueType* operator->() const
+	{
+		check((Data != nullptr || EnumHasAnyFlags(Flags, FParam::EFlags::Embedded)) && NumElements == 1);
+
+		ValueType* ValueData = EnumHasAnyFlags(Flags, FParam::EFlags::Embedded)
+			? static_cast<ValueType*>((void*)&Data)
+			: static_cast<ValueType*>(Data);
+
+		return ValueData;
+	}
+
+protected:
+	using TValueType = ValueType;
+	
+	friend struct FContext;
+	friend struct FState;
+	friend struct FKernel;
+	friend struct FParamStorage;
+	
+	TParam(EFlags InFlags)
+		: FParam(Private::TParamType<ValueType>::GetType(), InFlags)
+	{
+	}
+
+	TParam(void* InData, EFlags InFlags)
+		: FParam(Private::TParamType<ValueType>::GetType(), InData, InFlags)
+	{
+	}
+
+	TParam(const void* InData, EFlags InFlags)
+		: FParam(Private::TParamType<ValueType>::GetType(), const_cast<void*>(InData), InFlags)
+	{
+	}
+
+	TParam(void* InData, int32 InNumElements, EFlags InFlags)
+		: FParam(Private::TParamType<ValueType>::GetType(), InData, InNumElements, InFlags)
+	{
+	}
+
+	TParam(const void* InData, int32 InNumElements, EFlags InFlags)
+		: FParam(Private::TParamType<ValueType>::GetType(), const_cast<void*>(InData), InNumElements, InFlags)
+	{
+	}
+
+	TParam(FParam& InParam, EFlags InAdditionalFlags = EFlags::None)
+		: FParam(&InParam)
+	{
+#if DO_CHECK
+		TStringBuilder<64> Error;
+		checkf(InParam.CanAssignWith(Private::TParamType<ValueType>::GetType(), GetTypeFlags(InAdditionalFlags), NumElements, &Error),
+			TEXT("Cannot assign type: %s"), Error.ToString());
+#endif
+	}
+
+	TParam(const FParam& InParam, EFlags InAdditionalFlags = EFlags::None)
+		: FParam(&InParam)
+	{
+#if DO_CHECK
+		TStringBuilder<64> Error;
+		checkf(InParam.CanAssignWith(Private::TParamType<ValueType>::GetType(), GetTypeFlags(InAdditionalFlags), NumElements, &Error),
+			TEXT("Cannot assign type: %s"), Error.ToString());
+#endif
+	}
+
+protected:
+	static EFlags GetTypeFlags(EFlags InAdditionalFlags)
+	{
+		EFlags Flags = InAdditionalFlags;
+
+		// Add const flags or not
+		if constexpr (!TIsConst<ValueType>::Value)
+		{
+			Flags |= FParam::EFlags::Mutable;
+		}
+
+		return Flags;
+	}
+};
+
+namespace Private
+{
+
+template<typename ValueType>
+struct TParamCopyBaseTemplate
+{
+	static FParam CopyParam(const FParam& InSource, void* InTargetMemory, int32 InTargetAllocatedSize)
+	{
+		const FParamType& ParamType = InSource.GetType();
+		
+		const int32 NumElem = InSource.GetNumElements();
+		const UScriptStruct* ScriptStruct = ParamType.GetStruct();
+
+		if (ScriptStruct != nullptr)
+		{
+			const int32 StructSize = ScriptStruct->GetStructureSize();
+			check(NumElem * StructSize <= InTargetAllocatedSize);
+
+			ScriptStruct->CopyScriptStruct(InTargetMemory, InSource.GetData(), NumElem);
+		}
+		else
+		{
+			const int32 ParamAlignment = ParamType.GetAlignment();
+			const int32 ParamSize = ParamType.GetSize();
+			const int32 ParamAllocSize = Align(ParamSize, ParamAlignment);
+
+			// TODO : check if this is safe to do for any type not containing a script struct, else we can specialize for POD types
+			FMemory::Memcpy(InTargetMemory, InSource.GetData(), NumElem * ParamAllocSize);
+		}
+
+		return FParam::MakeParam_Internal(ParamType, InTargetMemory, NumElem, InSource.GetFlags());
+	}
+};
+
+struct CHasStaticStruct
+{
+	template<typename ValueType>
+	auto Requires(UScriptStruct* Result) -> decltype(
+		Result = ValueType::StaticStruct()
+	);
+};
+
+template<typename ValueType> 
+struct TParamTypeHelper
+{
+	FORCEINLINE static const UScriptStruct* GetStruct()
+	{
+		if constexpr (TModels<CHasStaticStruct, ValueType>::Value)
+		{
+			return TBaseStructure<ValueType>::Get();
+		}
+		else
+		{
+			return nullptr;
+		}
+	}
+
+	FORCEINLINE static uint32 GetSize()
+	{
+		return sizeof(ValueType);
+	}
+
+	FORCEINLINE static uint32 GetAlignment()
+	{
+		return alignof(ValueType);
+	}
+
+	FORCEINLINE static const FParamType::ParamCopyFunction GetParamCloneFunction()
+	{
+		return &TParamCopyBaseTemplate<ValueType>::CopyParam;
+	}
+};
+
+}
+
+// A typed param that wraps a user ptr to single or batched elements
+template<typename ValueType>
+struct TWrapParam : TParam<ValueType>
+{
+public:
+	TWrapParam(ValueType* InValuePtrToWrap)
+		: TParam<ValueType>(GetDataForConstructor(InValuePtrToWrap), GetFlagsForConstructor(InValuePtrToWrap))
+	{
+	}
+
+	TWrapParam(ValueType& InValueToWrap)
+		: TParam<ValueType>(GetDataForConstructor(&InValueToWrap), GetFlagsForConstructor(&InValueToWrap))
+	{
+	}
+
+	TWrapParam(ValueType* InValuePtrToWrap, const int32 NumElements)
+		: TParam<ValueType>(InValuePtrToWrap, NumElements, NumElements > 1 ? FParam::EFlags::Batched : FParam::EFlags::None)
+	{
+	}
+
+	TWrapParam(const TArrayView<ValueType>& InArrayToWrap)
+		: TParam<ValueType>(
+			InArrayToWrap.GetData(), 
+			InArrayToWrap.Num(), 
+			((InArrayToWrap.Num()) > 1 ? FParam::EFlags::Batched : FParam::EFlags::None) | ((InArrayToWrap.Num() > 0) ? GetFlagsForConstructor(&InArrayToWrap[0]) : FParam::EFlags::None))
+	{
+	}
+
+protected:
+	void* GetDataForConstructor(ValueType* InValuePtrToWrap) const
+	{
+		using MutableValueType = typename TRemoveConst<ValueType>::Type;
+
+		// Support containers with Num() and GetData()
+		if constexpr (TModels<Private::CSizedContainerWithAccessibleDataAsRawPtr, ValueType>::Value)
+		{
+			return (void*)(InValuePtrToWrap->GetData());
+		}
+		else
+		{
+			return const_cast<MutableValueType*>(InValuePtrToWrap);
+		}
+	}
+
+	FParam::EFlags GetFlagsForConstructor(const ValueType* InValuePtrToWrap) const
+	{
+		FParam::EFlags NewFlags = FParam::EFlags::None;
+
+		// Add batched flags or not
+		if constexpr (TModels<Private::CSizedContainerWithAccessibleDataAsRawPtr, ValueType>::Value)
+		{
+			NewFlags |= (InValuePtrToWrap->Num() > 1) ? FParam::EFlags::Batched : FParam::EFlags::None;
+		}
+
+		// Add const flags or not
+		if constexpr (!TIsConst<ValueType>::Value)
+		{
+			NewFlags |= FParam::EFlags::Mutable;
+		}
+
+		return NewFlags;
+	}
+};
+
+// A typed param that owns it's own memory with size defined at compile time
+template<typename ValueType, int32 NumElem = 1>
+struct TParamValue : TParam<ValueType>
+{
+public:
+	TParamValue()
+		: TParam<ValueType>(&ValueArray[0], NumElem, GetFlagsForConstructor())
+	{
+	}
+
+	TParamValue(FParam::EFlags InFlags)
+		: TParam<ValueType>(&ValueArray[0], NumElem, InFlags)
+	{
+	}
+
+private:
+
+	FParam::EFlags GetFlagsForConstructor() const
+	{
+		FParam::EFlags NewFlags = FParam::EFlags::None;
+
+		NewFlags |= (NumElem > 1) ? FParam::EFlags::Batched : FParam::EFlags::None;
+
+		// Add const flags or not
+		if constexpr (!TIsConst<ValueType>::Value)
+		{
+			NewFlags |= FParam::EFlags::Mutable;
+		}
+
+		return NewFlags;
+	}
+
+	ValueType ValueArray[NumElem];
+};
+
+template<typename ValueType>
+struct TContextStorageParam : TWrapParam<ValueType>
+{
+public:
+	TContextStorageParam(ValueType* InValuePtrToWrap)
+		: TWrapParam<ValueType>(InValuePtrToWrap)
+	{
+		FParam::Flags |= FParam::EFlags::Stored;
+	}
+
+	TContextStorageParam(ValueType& InValueToWrap)
+		: TWrapParam<ValueType>(&InValueToWrap)
+	{
+		FParam::Flags |= FParam::EFlags::Stored;
+	}
+
+	TContextStorageParam(ValueType* InValuePtrToWrap, const int32 NumElements)
+		: TWrapParam<ValueType>(InValuePtrToWrap, NumElements)
+	{
+		FParam::Flags |= FParam::EFlags::Stored;
+	}
+
+	TContextStorageParam(const TArrayView<ValueType>& InArrayToWrap)
+		: TWrapParam<ValueType>(InArrayToWrap)
+	{
+		FParam::Flags |= FParam::EFlags::Stored;
+	}
+};
+
+// --- ---
+
+typedef int32 FParamHandle;
+constexpr FParamHandle InvalidParamHandle = -1;
+
+struct ANIMNEXTINTERFACE_API FHParam
+{
+	FHParam() = default;
+	~FHParam();
+
+	FHParam(FHParam&& Other)
+		: FHParam()
+	{
+		Swap(*this, Other);
+	}
+
+	FHParam& operator= (FHParam&& Other)
+	{
+		Swap(*this, Other);
+		return *this;
+	}
+
+	FHParam(const FHParam& Other);
+
+	FHParam& operator= (const FHParam& Other)
+	{
+		FHParam Temp(Other); // note that the ref count is incremented here
+
+		Swap(*this, Temp);
+
+		return *this;
+	}
+
+protected:
+	friend struct FContext;
+	friend struct FParamStorage;
+
+	FHParam(FParamStorage* InOwnerStorage, FParamHandle InHandle);
+
+	FParamStorage* OwnerStorage = nullptr;
+	FParamHandle ParamHandle = InvalidParamHandle;
+
+private:
+};
+
+} // end namespace UE::AnimNext::Interface
diff --git a/Engine/Plugins/Experimental/AnimNext/Source/AnimNextInterface/Public/AnimNextInterfaceParamStorage.h b/Engine/Plugins/Experimental/AnimNext/Source/AnimNextInterface/Public/AnimNextInterfaceParamStorage.h
new file mode 100644
--- /dev/null
+++ b/Engine/Plugins/Experimental/AnimNext/Source/AnimNextInterface/Public/AnimNextInterfaceParamStorage.h
@@ -0,0 +1,222 @@
+// Copyright Epic Games, Inc. All Rights Reserved.
+
+#pragma once
+
+#include "CoreMinimal.h"
+#include "UObject/Class.h"
+#include "UObject/ScriptInterface.h"
+#include "AnimNextInterfaceParam.h"
+
+namespace UE::AnimNext::Interface
+{
+
+typedef int32 FParamStorageHandle;
+constexpr FParamStorageHandle InvalidBlockHandle = -1;
+
+struct ANIMNEXTINTERFACE_API FParamStorage
+{
+public:
+	typedef TPair<FParamStorageHandle, uint8*> TBlockDataPair;
+
+	FParamStorage(int32 InMaxParams, int32 InAllocatedMemorySize, int32 InMaxBlocks, bool InAllowGrowing = false);
+	~FParamStorage();
+
+
+	// *** Interface for param arena ***
+	// This prototype allows storing parameters into the storage using pre-allocated blocks of memory
+	// The owner is responsible for calculating the size and requesting / releasing the memory blocks
+	// Memory blocks are stored in stack LIFO way, last one has to be the first to be released (for simplicity)
+	// If a memoy block is released and it is not the last, it will be held until the ones after it are released
+
+	// Requests a block of RequestedBlockSize size
+	// Returns a pair with the block handle and the pointer to the memory buffer
+	TBlockDataPair RequestBlock(int32 RequestedBlockSize);
+
+	// Releases a block of memory from the arena, using a hanle to identify it
+	// If it is not the last of the stack, it will be held until the blocks after this one are released
+	void ReleaseBlock(FParamStorageHandle BlockHandle);
+
+	// Obtains the block of memory from a block handle
+	uint8* GetBlockMemory(FParamStorageHandle BlockHandle);
+
+
+	// *** Interface for direct Param storage ***
+	// This prototype allows storing parameters directly into the storage
+	// Parameters are stored stack like, in a LIFO approach
+	// If a parameter is released and it is not the last of the stack, it will be held until the parameters over that one are released
+	// This is to have a very simple memory management, might be removed on final implementation, based on requirements
+
+	// Adds a value to the storage (copying the data), returning a HParam
+	// If the total allocation size is <= than size of a void *, the value is stored directly in the FParam Data
+	// If not, additional memory is requested (TODO : better memory management)
+	template<typename ValueType>
+	FHParam AddValue(ValueType* Data, int32 NumElem, FParam::EFlags Flags)
+	{
+		check(ParamFreeIndex < (Parameters.Num() - 1)); // TODO : Growing ?
+
+		const int32 ParameterIndex = ParamFreeIndex;
+		ParamFreeIndex++;
+
+		const FParamType& ParamType = Private::TParamType<ValueType>::GetType();
+		const int32 ParamAllocSize = GetParamAllocSize(ParamType, NumElem);
+
+		FParam& Param = Parameters[ParameterIndex];
+
+		// Prepare the Param for the copy
+		Param.TypeId = ParamType.GetTypeId();
+		Param.Flags = Flags;
+		Param.NumElements = NumElem;
+		Param.Data = Data;
+
+		// If size is <= sizeof(void*), put the value directly at the pointer
+		if (ParamAllocSize <= sizeof(void*)) // TODO : see if we can get this constexpr
+		{
+			ParamMemoryHandles[ParameterIndex] = InvalidBlockHandle;
+
+			// Set target memory directly at Param.Data memory
+			ParamType.GetParamCopyFunction()(Param, static_cast<ValueType*>((void*)&Param.Data), ParamAllocSize);
+
+			Param.Flags |= FParam::EFlags::Embedded; // Added so it can be checked when retrieving (TODO : see if it can be done constexpr by type)
+		}
+		else
+		{
+			TBlockDataPair BlockData = RequestBlock(ParamAllocSize);
+			
+			ParamMemoryHandles[ParameterIndex] = BlockData.Key;
+
+			// Set target memory directly at BlockData memory
+			ParamType.GetParamCopyFunction()(Param, BlockData.Value, ParamAllocSize);
+			// Once the copy is done, update the Param
+			Param.Data = BlockData.Value;
+		}
+
+		RefCounts[ParameterIndex] = 1;
+
+		return FHParam(this, static_cast<FParamHandle>(ParameterIndex));
+	}
+
+	// Adds a reference (pointer) to the storage, returning a HParam
+	template<typename ValueType>
+	FHParam AddReference(ValueType* Data, int32 NumElem, FParam::EFlags Flags)
+	{
+		check(ParamFreeIndex < (Parameters.Num() - 1)); // TODO : Growing ?
+
+		const int32 ParameterIndex = ParamFreeIndex;
+		ParamFreeIndex++;
+
+		const FParamType& ParamType = Private::TParamType<ValueType>::GetType();
+
+		FParam& Param = Parameters[ParameterIndex];
+
+		Param.TypeId = ParamType.GetTypeId();
+		Param.Flags = Flags;
+		Param.NumElements = NumElem;
+		Param.Data = Data;
+
+		ParamMemoryHandles[ParameterIndex] = InvalidBlockHandle;
+
+		RefCounts[ParameterIndex] = 1;
+
+		return FHParam(this, static_cast<FParamHandle>(ParameterIndex));
+	}
+
+	// Obtains the Param associated to the HParam
+	FParam* GetParam(FParamHandle InParamHandle)
+	{
+		FParam* RetVal = nullptr;
+
+		const int32 ParameterIndex = InParamHandle;
+
+		if (ParameterIndex < ParamFreeIndex && RefCounts[ParameterIndex] > 0)
+		{
+			RetVal = &Parameters[ParameterIndex];
+		}
+
+		return RetVal;
+	}
+
+protected:
+	friend struct FHParam;
+
+	// Adds a reference count to a param, using the passed param handle
+	void IncRefCount(const FParamHandle InParamHandle)
+	{
+		const int32 ParameterIndex = InParamHandle;
+		check(ParameterIndex < ParamFreeIndex);
+		
+		check(RefCounts[ParameterIndex] > 0);
+		
+		RefCounts[ParameterIndex]++;
+	}
+
+	// Decrements a reference count to a param, using the passed param handle
+	// If the count reaches 0:
+	// - If it is the last param of the stack, it releases it and all the previous params with ref count == 0
+	// - If it is not the last one, returns without releasing (it will be done later)
+	void DecRefCount(const FParamHandle InParamHandle)
+	{
+		int32 ParameterIndex = InParamHandle;
+		check(ParameterIndex < ParamFreeIndex);
+
+		check(RefCounts[ParameterIndex] > 0);
+		
+		if (RefCounts[ParameterIndex] > 0)
+		{
+			RefCounts[ParameterIndex]--;
+
+			if (RefCounts[ParameterIndex] == 0)
+			{
+				// if it is the last parameter, release the memory and deal with the prev params with RefCount == 0, if any
+				if (ParameterIndex == ParamFreeIndex - 1)
+				{
+					// Recover free slots
+					while (ParameterIndex >= 0 && RefCounts[ParameterIndex] == 0)
+					{
+						Parameters[ParameterIndex] = FParam();
+						if (ParamMemoryHandles[ParameterIndex] != InvalidBlockHandle)
+						{
+							ReleaseBlock(ParamMemoryHandles[ParameterIndex]);
+							ParamMemoryHandles[ParameterIndex] = InvalidBlockHandle;
+						}
+
+						ParameterIndex--;
+						ParamFreeIndex--;
+					}
+				}
+			}
+		}
+	}
+
+	// Calculates the required size for a param, inclusing allignment
+	int32 GetParamAllocSize(const FParamType& ParamType, const int32 NumElem) const
+	{
+		const int32 ParamAlignment = ParamType.GetAlignment();
+		const int32 ParamSize = ParamType.GetSize();
+		const int32 ParamAllocSize = NumElem * Align(ParamSize, ParamAlignment);
+
+		return ParamAllocSize;
+	}
+
+	// --- context arena param storage prototype ---
+	TArray<int32> BlockOffets;
+	TArray<int32> BlockSizes;
+	TBitArray<> BlockFreeFlag;
+	uint8* RawMemory = nullptr;
+	int32 RawMemorySize = 0;
+	int32 CurrentBlocksSize = 0;
+	int32 FreeBlockIndex = 0;
+	
+	// --- shared param storage prototype ---
+	TArray<FParam> Parameters;
+	TArray<FParamStorageHandle> ParamMemoryHandles;
+	TArray<int16> RefCounts;
+	int32 ParamFreeIndex = 0;
+
+	bool AllowGrowing = false;
+
+private:
+	bool CheckBufferSize(int32 RequestedBlockSize);
+	bool CheckNumBlocks(int32 BlockIndex);
+};
+
+} // end namespace UE::AnimNext::Interface
diff --git a/Engine/Plugins/Experimental/AnimNext/Source/AnimNextInterface/Public/AnimNextInterfaceState.h b/Engine/Plugins/Experimental/AnimNext/Source/AnimNextInterface/Public/AnimNextInterfaceState.h
new file mode 100644
--- /dev/null
+++ b/Engine/Plugins/Experimental/AnimNext/Source/AnimNextInterface/Public/AnimNextInterfaceState.h
@@ -0,0 +1,87 @@
+﻿// Copyright Epic Games, Inc. All Rights Reserved.
+
+#pragma once
+
+#include "CoreMinimal.h"
+#include "AnimNextInterfaceParam.h"
+#include "AnimNextInterfaceKey.h"
+
+namespace UE::AnimNext::Interface
+{
+
+struct FContext;
+struct FParamType;
+
+enum class EStatePersistence : uint8
+{
+	None,
+	Relevancy,
+	Permanent,
+};
+
+// Container for data interface state
+struct ANIMNEXTINTERFACE_API FState
+{
+	FState(int32 InNumElements = 1)
+		: NumElements(InNumElements)
+	{
+		check(InNumElements > 0);
+	}
+	
+private:
+	friend struct FContext;
+
+	// Gets (and allocates, if necessary) state for the specified node given this calling context
+	template<typename ValueType, EStatePersistence Persistence = EStatePersistence::Relevancy>
+	TParam<ValueType> GetState(const FInterfaceKeyWithId& InKey, const FContext& InContext, uint32 InHash)
+	{
+		const FInterfaceKeyWithIdAndStack CombinedKey(InKey, InHash);
+
+		if(FParam* ExistingState = FindStateRaw(CombinedKey, InContext, Persistence))
+		{
+			checkSlow(Private::CheckParam<ValueType>(*ExistingState));
+			return TParam<ValueType>(*ExistingState);
+		}
+		else
+		{
+			// Allocate new state
+			FParam* NewState = AllocateState(CombinedKey, InContext, Private::TParamType<ValueType>::GetType(), Persistence);
+
+			// Construct
+			new (NewState->Data) ValueType();
+
+			// Return as param
+			return TParam<ValueType>(*NewState);
+		}
+	}
+
+	// Gets (and allocates, if necessary) state for the specified node given this calling context
+	template<typename ValueType, EStatePersistence Persistence = EStatePersistence::Relevancy>
+	TParam<ValueType> GetState(const IAnimNextInterface* InAnimNextInterface, uint32 InId, const FContext& InContext, uint32 InContextHash)
+	{
+		const FInterfaceKeyWithId Key(InAnimNextInterface, InId);
+		return GetState<ValueType>(Key, InContext, InContextHash);
+	}
+	
+private:
+	FParam* FindStateRaw(const FInterfaceKeyWithIdAndStack& InKey, const FContext& InContext, EStatePersistence InPersistence);
+
+	FParam* AllocateState(const FInterfaceKeyWithIdAndStack& InKey, const FContext& InContext, const FParamType& InType, EStatePersistence InPersistence);
+
+private:
+	struct FRelevancyParam : FParam
+	{
+		FRelevancyParam(const FParamType& InType, void* InData, EFlags InFlags, uint32 InUpdateCounter)
+			: FParam(InType, InData, InFlags)
+			, UpdateCounter(InUpdateCounter)
+		{}
+	
+		uint32 UpdateCounter = 0;
+	};
+	
+	TMap<FInterfaceKeyWithIdAndStack, FRelevancyParam> RelevancyValueMap;
+	TMap<FInterfaceKeyWithIdAndStack, FParam> PermanentValueMap;
+	int32 NumElements = 0;
+};
+
+}
\ No newline at end of file
diff --git a/Engine/Plugins/Experimental/AnimNext/Source/AnimNextInterface/Public/AnimNextInterfaceTypes.h b/Engine/Plugins/Experimental/AnimNext/Source/AnimNextInterface/Public/AnimNextInterfaceTypes.h
new file mode 100644
--- /dev/null
+++ b/Engine/Plugins/Experimental/AnimNext/Source/AnimNextInterface/Public/AnimNextInterfaceTypes.h
@@ -0,0 +1,69 @@
+﻿// Copyright Epic Games, Inc. All Rights Reserved.
+
+#pragma once
+
+#include "CoreMinimal.h"
+#include "AnimNextInterfaceParam.h"
+#include "UObject/ScriptInterface.h"
+
+class IAnimNextInterface;
+
+#define ANIM_NEXT_INTERFACE_RETURN_TYPE(Typename) \
+virtual FName GetReturnTypeNameImpl() const final override\
+{\
+	static FName ReturnTypeNameValue = #Typename;\
+	return ReturnTypeNameValue;\
+}\
+
+#define DECLARE_ANIM_NEXT_INTERFACE_PARAM_TYPE_INTERNAL(ValueType, Identifier) \
+	template<> \
+	struct UE::AnimNext::Interface::Private::TParamTypeImpl<ValueType> \
+	{ \
+		FORCEINLINE static const UE::AnimNext::Interface::FParamType& GetType() \
+		{ \
+			return Identifier##_AnimNextInterfaceTypeInfo; \
+		} \
+	};
+
+#define IMPLEMENT_ANIM_NEXT_INTERFACE_PARAM_TYPE_INTERNAL(ValueType, Identifier) \
+	UE::AnimNext::Interface::FParamType::FRegistrar Identifier##_Registrar([](){ UE::AnimNext::Interface::FParamType::FRegistrar::RegisterType( \
+		Identifier##_AnimNextInterfaceTypeInfo, \
+		UE::AnimNext::Interface::Private::TParamTypeHelper<ValueType>::GetStruct(), \
+		FName(#Identifier), UE::AnimNext::Interface::Private::TParamTypeHelper<ValueType>::GetSize(), \
+		UE::AnimNext::Interface::Private::TParamTypeHelper<ValueType>::GetAlignment(), \
+		UE::AnimNext::Interface::Private::TParamTypeHelper<ValueType>::GetParamCloneFunction()); });
+
+// Types that are passed between data interfaces (i.e. exist in in public headers) should declare via this mechanism
+#define DECLARE_ANIM_NEXT_INTERFACE_PARAM_TYPE(APIType, ValueType, Identifier) \
+	extern APIType UE::AnimNext::Interface::FParamType Identifier##_AnimNextInterfaceTypeInfo; \
+	DECLARE_ANIM_NEXT_INTERFACE_PARAM_TYPE_INTERNAL(ValueType, Identifier)
+
+// Types that are passed between data interfaces (i.e. exist in in public headers) should declare via this mechanism
+#define DECLARE_BUILTIN_ANIM_NEXT_INTERFACE_PARAM_TYPE(ValueType, Identifier) \
+	DECLARE_ANIM_NEXT_INTERFACE_PARAM_TYPE(ANIMNEXTINTERFACE_API, ValueType, Identifier)
+
+// Types that are passed between data interfaces should register via this mechanism
+#define IMPLEMENT_ANIM_NEXT_INTERFACE_PARAM_TYPE(ValueType, Identifier) \
+	UE::AnimNext::Interface::FParamType Identifier##_AnimNextInterfaceTypeInfo; \
+	IMPLEMENT_ANIM_NEXT_INTERFACE_PARAM_TYPE_INTERNAL(ValueType, Identifier)
+
+// Structs that are only used for internal state (in a cpp file) can use this simpler registration.
+// If they are exposed in a public header please use the DECLARE_ANIM_NEXT_INTERFACE_PARAM_TYPE & IMPLEMENT_ANIM_NEXT_INTERFACE_PARAM_TYPE macros
+#define IMPLEMENT_ANIM_NEXT_INTERFACE_STATE_TYPE(ValueType, Identifier) \
+	static UE::AnimNext::Interface::FParamType Identifier##_AnimNextInterfaceTypeInfo; \
+	static UE::AnimNext::Interface::FParamType::FRegistrar Identifier##_Registrar([](){ UE::AnimNext::Interface::FParamType::FRegistrar::RegisterType( \
+		Identifier##_AnimNextInterfaceTypeInfo, \
+		UE::AnimNext::Interface::Private::TParamTypeHelper<ValueType>::GetStruct(), \
+		FName(#Identifier), UE::AnimNext::Interface::Private::TParamTypeHelper<ValueType>::GetSize(), \
+		UE::AnimNext::Interface::Private::TParamTypeHelper<ValueType>::GetAlignment(), \
+		UE::AnimNext::Interface::Private::TParamTypeHelper<ValueType>::GetParamCloneFunction()); }); \
+	DECLARE_ANIM_NEXT_INTERFACE_PARAM_TYPE_INTERNAL(ValueType, Identifier);
+
+namespace UE::AnimNext::Interface
+{
+
+#define ANIM_NEXT_INTERFACE_TYPE(Type, Identifier) DECLARE_BUILTIN_ANIM_NEXT_INTERFACE_PARAM_TYPE(Type, Identifier)
+#include "AnimNextInterfaceTypes.inl"
+#undef ANIM_NEXT_INTERFACE_TYPE
+
+}
\ No newline at end of file
diff --git a/Engine/Plugins/Experimental/AnimNext/Source/AnimNextInterface/Public/AnimNextInterfaceTypes.inl b/Engine/Plugins/Experimental/AnimNext/Source/AnimNextInterface/Public/AnimNextInterfaceTypes.inl
new file mode 100644
--- /dev/null
+++ b/Engine/Plugins/Experimental/AnimNext/Source/AnimNextInterface/Public/AnimNextInterfaceTypes.inl
@@ -0,0 +1,27 @@
+﻿// Copyright Epic Games, Inc. All Rights Reserved.
+
+/** Built in types for data interfaces */
+
+#ifndef ANIM_NEXT_INTERFACE_TYPE
+#error ANIM_NEXT_INTERFACE_TYPE is not defined!
+#endif
+
+ANIM_NEXT_INTERFACE_TYPE(bool, Bool)
+ANIM_NEXT_INTERFACE_TYPE(uint8, Uint8)
+ANIM_NEXT_INTERFACE_TYPE(uint16, Uint16)
+ANIM_NEXT_INTERFACE_TYPE(uint32, Uint32)
+ANIM_NEXT_INTERFACE_TYPE(uint64, Uint64)
+ANIM_NEXT_INTERFACE_TYPE(int8, Int8)
+ANIM_NEXT_INTERFACE_TYPE(int16, Int16)
+ANIM_NEXT_INTERFACE_TYPE(int32, Int32)
+ANIM_NEXT_INTERFACE_TYPE(int64, Int64)
+ANIM_NEXT_INTERFACE_TYPE(float, Float)
+ANIM_NEXT_INTERFACE_TYPE(double, Double)
+ANIM_NEXT_INTERFACE_TYPE(FName, Name)
+ANIM_NEXT_INTERFACE_TYPE(FVector, Vector)
+ANIM_NEXT_INTERFACE_TYPE(FVector4, Vector4)
+ANIM_NEXT_INTERFACE_TYPE(FQuat, Quat)
+ANIM_NEXT_INTERFACE_TYPE(FTransform, Transform)
+ANIM_NEXT_INTERFACE_TYPE(TScriptInterface<IAnimNextInterface>, ScriptInterface_AnimNextInterface)
+
+#undef ANIM_NEXT_INTERFACE_TYPE
\ No newline at end of file
diff --git a/Engine/Plugins/Experimental/AnimNext/Source/AnimNextInterface/Public/IAnimNextInterface.h b/Engine/Plugins/Experimental/AnimNext/Source/AnimNextInterface/Public/IAnimNextInterface.h
new file mode 100644
--- /dev/null
+++ b/Engine/Plugins/Experimental/AnimNext/Source/AnimNextInterface/Public/IAnimNextInterface.h
@@ -0,0 +1,81 @@
+// Copyright Epic Games, Inc. All Rights Reserved.
+
+#pragma once
+
+#include "CoreMinimal.h"
+#include "UObject/Interface.h"
+#include "IAnimNextInterface.generated.h"
+
+class IAnimNextInterface;
+namespace UE::AnimNext::Interface
+{
+struct FParam;
+struct FContext;
+}
+
+UINTERFACE(BlueprintType, meta = (CannotImplementInterfaceInBlueprint))
+class ANIMNEXTINTERFACE_API UAnimNextInterface : public UInterface
+{
+	GENERATED_BODY()
+};
+
+// Empty data interface to support 'any' type
+class ANIMNEXTINTERFACE_API IAnimNextInterface
+{
+	GENERATED_BODY()
+
+public:
+	/**
+	 * Gets data using context to retrieve input parameters and result
+	 * @return false if type are incompatible, or if nested calls fail, otherwise true
+	 */
+	bool GetData(const UE::AnimNext::Interface::FContext& Context) const;
+
+	/**
+	 * Gets data and stores the value in OutResult, checking whether types are compatible
+	* @return false if nested calls fail, otherwise true
+	*/
+	bool GetDataChecked(const UE::AnimNext::Interface::FContext& Context) const;
+
+	/**
+	 * Gets data and stores the value in OutResult.
+	 * @return false if type are incompatible, or if nested calls fail, otherwise true
+	 */
+	bool GetData(const UE::AnimNext::Interface::FContext& Context, UE::AnimNext::Interface::FParam& OutResult) const;
+
+	/**
+	 * Gets data and stores the value in OutResult, checking whether types are compatible
+	 * @return false if nested calls fail, otherwise true
+	 */
+	bool GetDataChecked(const UE::AnimNext::Interface::FContext& Context, UE::AnimNext::Interface::FParam& OutResult) const;
+
+	/** Get the name of the return type */
+	FName GetReturnTypeName() const
+	{
+		return GetReturnTypeNameImpl();
+	}
+
+	/** Get the struct of the return type, if any */
+	const UScriptStruct* GetReturnTypeStruct() const
+	{
+		return GetReturnTypeStructImpl();
+	}
+
+private:
+	/** Get data if the types are compatible */
+	bool GetDataIfCompatibleInternal(const UE::AnimNext::Interface::FContext& InContext) const;
+	
+	/** Get a value from a data interface with no dynamic or static type checking. Internal use only. */
+	bool GetDataRawInternal(const UE::AnimNext::Interface::FContext& InContext) const;
+
+protected:
+	/** Get the return type name, used for dynamic type checking */
+	virtual FName GetReturnTypeNameImpl() const = 0;
+
+	/** Get the return type struct, used for dynamic type checking. If the result is not a struct, this should return nullptr */
+	virtual const UScriptStruct* GetReturnTypeStructImpl() const { return nullptr; }
+	
+	/** Get the value for this interface. @return true if successful, false if unsuccessful. */
+	virtual bool GetDataImpl(const UE::AnimNext::Interface::FContext& Context) const = 0;
+};
+
diff --git a/Engine/Plugins/Experimental/AnimNext/Source/AnimNextInterface/Public/IAnimNextInterfaceParameters.h b/Engine/Plugins/Experimental/AnimNext/Source/AnimNextInterface/Public/IAnimNextInterfaceParameters.h
new file mode 100644
--- /dev/null
+++ b/Engine/Plugins/Experimental/AnimNext/Source/AnimNextInterface/Public/IAnimNextInterfaceParameters.h
@@ -0,0 +1,51 @@
+﻿// Copyright Epic Games, Inc. All Rights Reserved.
+
+#pragma once
+
+#include "CoreMinimal.h"
+#include "UObject/Interface.h"
+#include "IAnimNextInterfaceParameters.generated.h"
+
+namespace UE::AnimNext::Interface
+{
+struct FParam;
+struct FContext;
+}
+
+UINTERFACE(BlueprintType, meta = (CannotImplementInterfaceInBlueprint))
+class ANIMNEXTINTERFACE_API UAnimNextInterfaceParameters : public UInterface
+{
+	GENERATED_BODY()
+};
+
+// Interface representing parameters that can be supplied to a data interface
+class ANIMNEXTINTERFACE_API IAnimNextInterfaceParameters
+{
+	GENERATED_BODY()
+
+	friend struct UE::AnimNext::Interface::FContext;
+	
+	// Get a parameter by name
+	virtual bool GetParameter(FName InKey, UE::AnimNext::Interface::FParam& OutParam) = 0;
+
+	// Get a parameter pointer by name
+	virtual const UE::AnimNext::Interface::FParam* GetParameter(FName InKey) = 0;
+
+	// Default utility implementation for UObjects
+	// Treats key as a property name or function
+	static bool GetParameterForObject(UObject* InObject, FName InKey, UE::AnimNext::Interface::FParam& OutParam);
+/*	{
+		checkSlow(InObject);
+		const UClass* ObjectClass = InObject->GetClass();
+		if(const FProperty* Property = FindFProperty<FProperty>(ObjectClass, InKey))
+		{
+			OutParam = UE::AnimNext::Interface::FPropertyParam(Property, InObject);
+			return true;
+		}
+		else if(const UFunction* Function = FindUField<UFunction>(ObjectClass, InKey))
+		{
+			OutParam = UE::AnimNext::Interface::FFunctionParam(Property, InObject);
+		}
+		return false;
+	}*/
+};
diff --git a/Engine/Plugins/Experimental/AnimNext/Source/AnimNextInterfaceGraph/AnimNextInterfaceGraph.Build.cs b/Engine/Plugins/Experimental/AnimNext/Source/AnimNextInterfaceGraph/AnimNextInterfaceGraph.Build.cs
new file mode 100644
--- /dev/null
+++ b/Engine/Plugins/Experimental/AnimNext/Source/AnimNextInterfaceGraph/AnimNextInterfaceGraph.Build.cs
@@ -0,0 +1,28 @@
+// Copyright Epic Games, Inc. All Rights Reserved.
+
+namespace UnrealBuildTool.Rules
+{
+	public class AnimNextInterfaceGraph : ModuleRules
+	{
+		public AnimNextInterfaceGraph(ReadOnlyTargetRules Target) : base(Target)
+		{
+		
+			PublicDependencyModuleNames.AddRange(
+				new string[]
+				{
+					"RigVM",
+					"ControlRig",
+				}
+			);
+
+			PrivateDependencyModuleNames.AddRange(
+				new string[]
+				{
+					"Core",
+					"CoreUObject",
+					"AnimNextInterface",
+				}
+			);
+		}
+	}
+}
\ No newline at end of file
diff --git a/Engine/Plugins/Experimental/AnimNext/Source/AnimNextInterfaceGraph/Internal/AnimNextInterfaceExecuteContext.h b/Engine/Plugins/Experimental/AnimNext/Source/AnimNextInterfaceGraph/Internal/AnimNextInterfaceExecuteContext.h
new file mode 100644
--- /dev/null
+++ b/Engine/Plugins/Experimental/AnimNext/Source/AnimNextInterfaceGraph/Internal/AnimNextInterfaceExecuteContext.h
@@ -0,0 +1,77 @@
+﻿// Copyright Epic Games, Inc. All Rights Reserved.
+
+#pragma once
+
+#include "CoreMinimal.h"
+#include "RigVMCore/RigVMExecuteContext.h"
+#include "RigVMCore/RigVMStruct.h"
+#include "AnimNextInterfaceUnitContext.h"
+#include "AnimNextInterfaceExecuteContext.generated.h"
+
+class IAnimNextInterface;
+
+namespace UE::AnimNext::Interface
+{
+	struct FContext;
+}
+
+USTRUCT(BlueprintType)
+struct FAnimNextInterfaceExecuteContext : public FRigVMExecuteContext
+{
+	GENERATED_BODY()
+
+	FAnimNextInterfaceExecuteContext()
+		: FRigVMExecuteContext()
+	, AnimNextInterfaceContext(nullptr)
+	, Interface(nullptr)
+	, ResultPtr(nullptr)
+	, UnitContext()
+	{
+	}
+
+	const UE::AnimNext::Interface::FContext& GetContext() const
+	{
+		check(AnimNextInterfaceContext);
+		return *AnimNextInterfaceContext;
+	}
+
+	const FAnimNextInterfaceUnitContext& GetUnitContext() const
+	{
+		return UnitContext;
+	}
+
+	void SetResult(bool bInResult) const
+	{
+		check(ResultPtr);
+		*ResultPtr &= bInResult;
+	}
+
+	const IAnimNextInterface* GetInterface() const
+	{
+		check(Interface);
+		return Interface;
+	}
+	
+	virtual void Copy(const FRigVMExecuteContext* InOtherContext) override
+	{
+		Super::Copy(InOtherContext);
+
+		const FAnimNextInterfaceExecuteContext* OtherContext = (const FAnimNextInterfaceExecuteContext*)InOtherContext; 
+		AnimNextInterfaceContext = OtherContext->AnimNextInterfaceContext;
+		Interface = OtherContext->Interface; 
+		ResultPtr = OtherContext->ResultPtr; 
+	}
+
+
+private:
+	const UE::AnimNext::Interface::FContext* AnimNextInterfaceContext;
+	const IAnimNextInterface* Interface;
+	bool* ResultPtr;
+	FAnimNextInterfaceUnitContext UnitContext;
+};
+
+USTRUCT(meta=(ExecuteContext="FAnimNextInterfaceExecuteContext"))
+struct FRigUnit_AnimNextInterfaceBase : public FRigVMStruct
+{
+	GENERATED_BODY()
+};
diff --git a/Engine/Plugins/Experimental/AnimNext/Source/AnimNextInterfaceGraph/Internal/AnimNextInterfaceGraph.h b/Engine/Plugins/Experimental/AnimNext/Source/AnimNextInterfaceGraph/Internal/AnimNextInterfaceGraph.h
new file mode 100644
--- /dev/null
+++ b/Engine/Plugins/Experimental/AnimNext/Source/AnimNextInterfaceGraph/Internal/AnimNextInterfaceGraph.h
@@ -0,0 +1,58 @@
+// Copyright Epic Games, Inc. All Rights Reserved.
+
+#pragma once
+
+#include "CoreMinimal.h"
+#include "IAnimNextInterface.h"
+#include "RigVMCore/RigVM.h"
+#include "AnimNextInterfaceGraph.generated.h"
+
+class UEdGraph;
+class UAnimNextInterfaceGraph;
+namespace UE::AnimNext::InterfaceGraphUncookedOnly
+{
+	struct FUtils;
+}
+
+namespace UE::AnimNext::InterfaceGraphEditor
+{
+	class FGraphEditor;
+}
+
+namespace UE::AnimNext::InterfaceGraph
+{
+	extern ANIMNEXTINTERFACEGRAPH_API const FName EntryPointName;
+	extern ANIMNEXTINTERFACEGRAPH_API const FName ResultName;
+}
+
+// A user-created graph of logic used to supply data
+UCLASS()
+class ANIMNEXTINTERFACEGRAPH_API UAnimNextInterfaceGraph : public UObject, public IAnimNextInterface
+{
+	GENERATED_BODY()
+
+	// IAnimAnimNextInterface interface
+	virtual FName GetReturnTypeNameImpl() const final override;
+	virtual const UScriptStruct* GetReturnTypeStructImpl() const final override;
+	virtual bool GetDataImpl(const UE::AnimNext::Interface::FContext& Context) const final override;
+
+	// Support rig VM execution
+	TArray<FRigVMExternalVariable> GetRigVMExternalVariables();
+	
+	friend class UAnimNextInterfaceGraphFactory;
+	friend class UAnimNextInterfaceGraph_EditorData;
+	friend struct UE::AnimNext::InterfaceGraphUncookedOnly::FUtils;
+	friend class UE::AnimNext::InterfaceGraphEditor::FGraphEditor;
+	friend class UAnimNextInterface_Graph;
+	
+	UPROPERTY()
+	TObjectPtr<URigVM> RigVM;
+
+	UPROPERTY()
+	FRigVMRuntimeSettings VMRuntimeSettings;
+	
+#if WITH_EDITORONLY_DATA
+	UPROPERTY(VisibleAnywhere, Instanced, Category = "Graph", meta = (ShowInnerProperties))
+	TObjectPtr<UObject> EditorData;
+#endif
+};
\ No newline at end of file
diff --git a/Engine/Plugins/Experimental/AnimNext/Source/AnimNextInterfaceGraph/Internal/AnimNextInterfaceUnitContext.h b/Engine/Plugins/Experimental/AnimNext/Source/AnimNextInterfaceGraph/Internal/AnimNextInterfaceUnitContext.h
new file mode 100644
--- /dev/null
+++ b/Engine/Plugins/Experimental/AnimNext/Source/AnimNextInterfaceGraph/Internal/AnimNextInterfaceUnitContext.h
@@ -0,0 +1,28 @@
+// Copyright Epic Games, Inc. All Rights Reserved.
+
+#pragma once
+
+#include "CoreMinimal.h"
+#include "AnimNextInterfaceContext.h"
+#include "Units/RigUnitContext.h"
+
+struct FAnimNextInterfaceUnitContext : public FRigUnitContext
+{
+	FAnimNextInterfaceUnitContext()
+		: FRigUnitContext()
+		, Interface(nullptr)
+		, AnimNextInterfaceContext(nullptr)
+		, bResult(nullptr)
+	{}
+
+	FAnimNextInterfaceUnitContext(const IAnimNextInterface* InInterface, const UE::AnimNext::Interface::FContext& InAnimNextInterfaceContext, bool& bInResult)
+		: FRigUnitContext()
+		, Interface(InInterface)
+		, AnimNextInterfaceContext(&InAnimNextInterfaceContext)
+		, bResult(&bInResult)
+	{}
+
+	const IAnimNextInterface* Interface;
+	const UE::AnimNext::Interface::FContext* AnimNextInterfaceContext;
+	bool* bResult;
+};
diff --git a/Engine/Plugins/Experimental/AnimNext/Source/AnimNextInterfaceGraph/Internal/RigUnit_AnimNextInterfaceBeginExecution.h b/Engine/Plugins/Experimental/AnimNext/Source/AnimNextInterfaceGraph/Internal/RigUnit_AnimNextInterfaceBeginExecution.h
new file mode 100644
--- /dev/null
+++ b/Engine/Plugins/Experimental/AnimNext/Source/AnimNextInterfaceGraph/Internal/RigUnit_AnimNextInterfaceBeginExecution.h
@@ -0,0 +1,28 @@
+﻿// Copyright Epic Games, Inc. All Rights Reserved.
+
+#pragma once
+
+#include "Units/RigUnit.h"
+#include "AnimNextInterfaceExecuteContext.h"
+#include "RigUnit_AnimNextInterfaceBeginExecution.generated.h"
+
+/**
+ * Event for driving the skeleton hierarchy with variables and rig elements
+ */
+USTRUCT(meta=(DisplayName="Execute Data Interface", Category="Events", NodeColor="1, 0, 0", Keywords="Begin,Update,Tick,Forward,Event"))
+struct ANIMNEXTINTERFACEGRAPH_API FRigUnit_AnimNextInterfaceBeginExecution : public FRigUnit_AnimNextInterfaceBase
+{
+	GENERATED_BODY()
+
+	RIGVM_METHOD()
+	void Execute();
+
+	virtual FName GetEventName() const override { return EventName; }
+	virtual bool CanOnlyExistOnce() const override { return true; }
+
+	// The execution result
+	UPROPERTY(EditAnywhere, DisplayName = "Execute", Category = "BeginExecution", meta = (Output))
+	FAnimNextInterfaceExecuteContext ExecuteContext;
+
+	static FName EventName;
+};
diff --git a/Engine/Plugins/Experimental/AnimNext/Source/AnimNextInterfaceGraph/Internal/RigUnit_AnimNextInterfaceEndExecution.h b/Engine/Plugins/Experimental/AnimNext/Source/AnimNextInterfaceGraph/Internal/RigUnit_AnimNextInterfaceEndExecution.h
new file mode 100644
--- /dev/null
+++ b/Engine/Plugins/Experimental/AnimNext/Source/AnimNextInterfaceGraph/Internal/RigUnit_AnimNextInterfaceEndExecution.h
@@ -0,0 +1,66 @@
+﻿// Copyright Epic Games, Inc. All Rights Reserved.
+
+#pragma once
+
+#include "Units/RigUnit.h"
+#include "AnimNextInterfaceExecuteContext.h"
+#include "AnimNextInterfaceUnitContext.h"
+#include "AnimNextInterfaceTypes.h"
+#include "RigUnit_AnimNextInterfaceEndExecution.generated.h"
+
+struct FRigUnitContext;
+
+/** Event for writing back calculated results to external variables */
+USTRUCT(meta = (DisplayName = "End Execute Data Interface", Category = "Events", TitleColor = "1 0 0", NodeColor = "1 1 1"))
+struct ANIMNEXTINTERFACEGRAPH_API FRigUnit_AnimNextInterfaceEndExecution : public FRigUnit_AnimNextInterfaceBase
+{
+	GENERATED_BODY()
+
+	template<typename ValueType>
+	static void SetResult(const FRigVMExecuteContext& InExecuteContext, const ValueType& InResult)
+	{
+		const FAnimNextInterfaceExecuteContext& AnimNextInterfaceExecuteContext = static_cast<const FAnimNextInterfaceExecuteContext&>(InExecuteContext);
+		AnimNextInterfaceExecuteContext.GetContext().SetResult<ValueType>(InResult);
+	}
+
+public:
+	// The execution result
+	UPROPERTY(EditAnywhere, Transient, DisplayName = "End Execute", Category = "EndExecution", meta = (Input))
+	FAnimNextInterfaceExecuteContext ExecuteContext;
+};
+
+/** Event for writing back a calculated bool */
+USTRUCT(meta = (DisplayName = "End Execute Data Interface Bool", Category = "Events", TitleColor = "1 0 0", NodeColor = "1 1 1"))
+struct ANIMNEXTINTERFACEGRAPH_API FRigUnit_AnimNextInterfaceEndExecution_Bool : public FRigUnit_AnimNextInterfaceEndExecution
+{
+	GENERATED_BODY()
+	
+	FRigUnit_AnimNextInterfaceEndExecution_Bool()
+		: FRigUnit_AnimNextInterfaceEndExecution()
+		, Result(false)
+	{}
+	
+	RIGVM_METHOD()
+	void Execute();
+	
+	UPROPERTY(EditAnywhere, Category = Result, meta = (Input))
+	bool Result;
+};
+
+/** Event for writing back a calculated float */
+USTRUCT(meta = (DisplayName = "End Execute Data Interface Float", Category = "Events", TitleColor = "1 0 0", NodeColor = "1 1 1"))
+struct ANIMNEXTINTERFACEGRAPH_API FRigUnit_AnimNextInterfaceEndExecution_Float : public FRigUnit_AnimNextInterfaceEndExecution
+{
+	GENERATED_BODY()
+
+	FRigUnit_AnimNextInterfaceEndExecution_Float()
+		: FRigUnit_AnimNextInterfaceEndExecution()
+		, Result(0.f)
+	{}
+
+	RIGVM_METHOD()
+	void Execute();
+
+	UPROPERTY(EditAnywhere, Category = Result, meta = (Input))
+	float Result;
+};
\ No newline at end of file
diff --git a/Engine/Plugins/Experimental/AnimNext/Source/AnimNextInterfaceGraph/Private/AnimNextInterfaceGraph.cpp b/Engine/Plugins/Experimental/AnimNext/Source/AnimNextInterfaceGraph/Private/AnimNextInterfaceGraph.cpp
new file mode 100644
--- /dev/null
+++ b/Engine/Plugins/Experimental/AnimNext/Source/AnimNextInterfaceGraph/Private/AnimNextInterfaceGraph.cpp
@@ -0,0 +1,47 @@
+// Copyright Epic Games, Inc. All Rights Reserved.
+
+#include "AnimNextInterfaceGraph.h"
+#include "RigVMCore/RigVMMemoryStorage.h"
+#include "RigUnit_AnimNextInterfaceBeginExecution.h"
+#include "AnimNextInterfaceUnitContext.h"
+
+#include UE_INLINE_GENERATED_CPP_BY_NAME(AnimNextInterfaceGraph)
+
+namespace UE::AnimNext::InterfaceGraph
+{
+const FName EntryPointName("GetData");
+const FName ResultName("Result");
+}
+
+FName UAnimNextInterfaceGraph::GetReturnTypeNameImpl() const
+{
+	// TODO: this needs to be user-defined & determined at creation time, using float for now
+	return TNameOf<float>::GetName();
+}
+
+const UScriptStruct* UAnimNextInterfaceGraph::GetReturnTypeStructImpl() const
+{
+	// TODO: this needs to be user-defined & determined at creation time, using nullptr (for float) for now
+	return nullptr;
+}
+
+bool UAnimNextInterfaceGraph::GetDataImpl(const UE::AnimNext::Interface::FContext& Context) const
+{
+	bool bResult = true;
+	
+	if(RigVM)
+	{
+		FRigUnitContext RigUnitContext;
+		RigUnitContext.State = EControlRigState::Update;
+		FAnimNextInterfaceUnitContext AnimNextInterfaceUnitContext(this, Context, bResult);
+		bResult &=
+			(RigVM->Execute(TArray<URigVMMemoryStorage*>(), FRigUnit_AnimNextInterfaceBeginExecution::EventName) != ERigVMExecuteResult::Failed);
+	}
+	
+	return bResult;
+}
+
+TArray<FRigVMExternalVariable> UAnimNextInterfaceGraph::GetRigVMExternalVariables()
+{
+	return TArray<FRigVMExternalVariable>(); 
+}
diff --git a/Engine/Plugins/Experimental/AnimNext/Source/AnimNextInterfaceGraph/Private/AnimNextInterfaceGraphModule.cpp b/Engine/Plugins/Experimental/AnimNext/Source/AnimNextInterfaceGraph/Private/AnimNextInterfaceGraphModule.cpp
new file mode 100644
--- /dev/null
+++ b/Engine/Plugins/Experimental/AnimNext/Source/AnimNextInterfaceGraph/Private/AnimNextInterfaceGraphModule.cpp
@@ -0,0 +1,5 @@
+// Copyright Epic Games, Inc. All Rights Reserved.
+
+#include "Modules/ModuleManager.h"
+
+IMPLEMENT_MODULE(FDefaultModuleImpl, AnimNextInterfaceGraph);
\ No newline at end of file
diff --git a/Engine/Plugins/Experimental/AnimNext/Source/AnimNextInterfaceGraph/Private/AnimNextInterface_Bool.cpp b/Engine/Plugins/Experimental/AnimNext/Source/AnimNextInterfaceGraph/Private/AnimNextInterface_Bool.cpp
new file mode 100644
--- /dev/null
+++ b/Engine/Plugins/Experimental/AnimNext/Source/AnimNextInterfaceGraph/Private/AnimNextInterface_Bool.cpp
@@ -0,0 +1,52 @@
+﻿// Copyright Epic Games, Inc. All Rights Reserved.
+
+#include "AnimNextInterface_Bool.h"
+#include "AnimNextInterface.h"
+#include "AnimNextInterfaceTypes.h"
+#include "AnimNextInterfaceKernel.h"
+
+#include UE_INLINE_GENERATED_CPP_BY_NAME(AnimNextInterface_Bool)
+
+bool UAnimNextInterface_Bool_And::GetDataImpl(const UE::AnimNext::Interface::FContext& Context) const
+{
+	using namespace UE::AnimNext::Interface;
+	
+	check(Inputs.Num() > 0);
+
+	bool bResult = true;
+	TAllocParam<bool> IntermediateParam(Context);
+	TParam<bool> Result = Context.GetResult<bool>();
+	
+	for(const TScriptInterface<IAnimNextInterface>& Input : Inputs)
+	{
+		bResult &= UE::AnimNext::Interface::GetDataSafe(Input, Context, IntermediateParam);
+
+		FKernel::Run(Context,
+			[](bool& OutResult, bool InIntermediate)
+			{
+				OutResult = OutResult && InIntermediate;
+			},
+			Result, IntermediateParam);
+	}
+
+	return bResult;
+}
+
+bool UAnimNextInterface_Bool_Not::GetDataImpl(const UE::AnimNext::Interface::FContext& Context) const
+{
+	using namespace UE::AnimNext::Interface;
+
+	TParam<bool> Result = Context.GetResult<bool>();
+	
+	TAllocParam<bool> CurrentValue(Context);
+	bool bResult = UE::AnimNext::Interface::GetDataSafe(Input, Context, CurrentValue);
+
+	FKernel::Run(Context,
+		[](bool& OutResult, bool InCurrent)
+		{
+			OutResult = !InCurrent;
+		},
+		Result, CurrentValue);
+
+	return bResult;
+}
diff --git a/Engine/Plugins/Experimental/AnimNext/Source/AnimNextInterfaceGraph/Private/AnimNextInterface_Bool.h b/Engine/Plugins/Experimental/AnimNext/Source/AnimNextInterfaceGraph/Private/AnimNextInterface_Bool.h
new file mode 100644
--- /dev/null
+++ b/Engine/Plugins/Experimental/AnimNext/Source/AnimNextInterfaceGraph/Private/AnimNextInterface_Bool.h
@@ -0,0 +1,53 @@
+// Copyright Epic Games, Inc. All Rights Reserved.
+
+#pragma once
+
+#include "CoreMinimal.h"
+#include "AnimNextInterfaceTypes.h"
+#include "IAnimNextInterface.h"
+#include "AnimNextInterfaceContext.h"
+#include "AnimNextInterface_Bool.generated.h"
+
+UCLASS()
+class UAnimNextInterface_Bool_Literal : public UObject, public IAnimNextInterface
+{
+	GENERATED_BODY()
+	ANIM_NEXT_INTERFACE_RETURN_TYPE(bool)
+
+	// IAnimNextInterface interface
+	virtual bool GetDataImpl(const UE::AnimNext::Interface::FContext& Context) const final override
+	{
+		Context.SetResult(Value);
+		return true;
+	}
+	
+	UPROPERTY(EditAnywhere, Category = "Parameters")
+	bool Value;
+};
+
+UCLASS()
+class UAnimNextInterface_Bool_And : public UObject, public IAnimNextInterface
+{
+	GENERATED_BODY()
+	ANIM_NEXT_INTERFACE_RETURN_TYPE(bool)
+
+	// IAnimNextInterface interface
+	virtual bool GetDataImpl(const UE::AnimNext::Interface::FContext& Context) const final override;
+
+	UPROPERTY(EditAnywhere, Meta=(AnimNextInterfaceType="bool"), Category = "Parameters")
+	TArray<TScriptInterface<IAnimNextInterface>> Inputs;
+};
+
+
+UCLASS()
+class UAnimNextInterface_Bool_Not : public UObject, public IAnimNextInterface
+{
+	GENERATED_BODY()
+	ANIM_NEXT_INTERFACE_RETURN_TYPE(bool)
+
+	// IAnimNextInterface interface
+	virtual bool GetDataImpl(const UE::AnimNext::Interface::FContext& Context) const final override;
+
+	UPROPERTY(EditAnywhere, Meta=(AnimNextInterfaceType="bool"), Category = "Parameters")
+	TScriptInterface<IAnimNextInterface> Input;
+};
\ No newline at end of file
diff --git a/Engine/Plugins/Experimental/AnimNext/Source/AnimNextInterfaceGraph/Private/AnimNextInterface_Float.cpp b/Engine/Plugins/Experimental/AnimNext/Source/AnimNextInterfaceGraph/Private/AnimNextInterface_Float.cpp
new file mode 100644
--- /dev/null
+++ b/Engine/Plugins/Experimental/AnimNext/Source/AnimNextInterfaceGraph/Private/AnimNextInterface_Float.cpp
@@ -0,0 +1,112 @@
+﻿// Copyright Epic Games, Inc. All Rights Reserved.
+
+#include "AnimNextInterface_Float.h"
+#include "AnimNextInterface.h"
+#include "AnimNextInterfaceTypes.h"
+#include "AnimNextInterfaceKernel.h"
+
+#include UE_INLINE_GENERATED_CPP_BY_NAME(AnimNextInterface_Float)
+
+#define LOCTEXT_NAMESPACE "AnimNextInterface"
+
+bool UAnimNextInterface_Float_Multiply::GetDataImpl(const UE::AnimNext::Interface::FContext& Context) const
+{
+	using namespace UE::AnimNext::Interface;
+	
+	check(Inputs.Num() > 0);
+
+	bool bResult = true;
+	TAllocParam<float> IntermediateParam(Context);
+	TParam<float> Result = Context.GetResult<float>();
+	
+	for(const TScriptInterface<IAnimNextInterface>& Input : Inputs)
+	{
+		bResult &= UE::AnimNext::Interface::GetDataSafe(Input, Context, IntermediateParam);
+
+		FKernel::Run(Context,
+			[](float& OutResult, float InIntermediate)
+			{
+				OutResult *= InIntermediate;
+			},
+			Result, IntermediateParam);
+	}
+
+	return bResult;
+}
+
+bool UAnimNextInterface_Float_InterpTo::GetDataImpl(const UE::AnimNext::Interface::FContext& Context) const
+{
+	using namespace UE::AnimNext::Interface;
+
+	TParam<float> Result = Context.GetResult<float>();
+	const float DeltaTime = Context.GetDeltaTime();
+	
+	TAllocParam<float> CurrentValue(Context);
+	bool bResult = UE::AnimNext::Interface::GetDataSafe(Current, Context, CurrentValue);
+
+	TAllocParam<float> TargetValue(Context);
+	bResult &= UE::AnimNext::Interface::GetDataSafe(Target, Context, TargetValue);
+
+	TAllocParam<float> SpeedValue(Context);
+	bResult &= UE::AnimNext::Interface::GetDataSafe(Speed, Context, SpeedValue);
+
+	FKernel::Run(Context,
+		[DeltaTime](float& OutResult, float InCurrent, float InTarget, float InSpeed)
+		{
+			OutResult = FMath::FInterpConstantTo(InCurrent, InTarget, DeltaTime, InSpeed);
+		},
+		Result, CurrentValue, TargetValue, SpeedValue);
+
+	return bResult;
+}
+
+struct FSpringInterpState
+{
+	float Value;
+	float ValueRate;
+};
+
+IMPLEMENT_ANIM_NEXT_INTERFACE_STATE_TYPE(FSpringInterpState, SpringInterpState);
+
+bool UAnimNextInterface_Float_SpringInterp::GetDataImpl(const UE::AnimNext::Interface::FContext& Context) const
+{
+	using namespace UE::AnimNext::Interface;
+	
+	const TParam<FSpringInterpState> State = Context.GetState<FSpringInterpState>(this, 0);
+	const float DeltaTime = Context.GetDeltaTime();
+
+	TAllocParam<float> TargetValueParam(Context);
+	bool bResult = UE::AnimNext::Interface::GetDataSafe(TargetValue, Context, TargetValueParam);
+
+	TAllocParam<float> TargetValueRateParam(Context);
+	bResult &= UE::AnimNext::Interface::GetDataSafe(TargetValueRate, Context, TargetValueRateParam);
+
+	TAllocParam<float> SmoothingTimeParam(Context);
+	bResult &= UE::AnimNext::Interface::GetDataSafe(SmoothingTime, Context, SmoothingTimeParam);
+
+	TAllocParam<float> DampingRatioParam(Context);
+	bResult &= UE::AnimNext::Interface::GetDataSafe(DampingRatio, Context, DampingRatioParam);
+
+	FKernel::Run(Context,
+	[DeltaTime](FSpringInterpState& InOutState,
+					float InTargetValue,
+					float InTargetValueRate,
+					float InSmoothingTime,
+					float InDampingRatio)
+		{
+			FMath::SpringDamperSmoothing(
+				InOutState.Value,
+				InOutState.ValueRate,
+				InTargetValue,
+				InTargetValueRate,
+				DeltaTime,
+				InSmoothingTime,
+				InDampingRatio);
+		},
+		State, TargetValueParam, TargetValueRateParam, SmoothingTimeParam, DampingRatioParam);
+
+	return bResult;
+}
+
+
+#undef LOCTEXT_NAMESPACE
diff --git a/Engine/Plugins/Experimental/AnimNext/Source/AnimNextInterfaceGraph/Private/AnimNextInterface_Float.h b/Engine/Plugins/Experimental/AnimNext/Source/AnimNextInterfaceGraph/Private/AnimNextInterface_Float.h
new file mode 100644
--- /dev/null
+++ b/Engine/Plugins/Experimental/AnimNext/Source/AnimNextInterfaceGraph/Private/AnimNextInterface_Float.h
@@ -0,0 +1,107 @@
+// Copyright Epic Games, Inc. All Rights Reserved.
+
+#pragma once
+
+#include "CoreMinimal.h"
+#include "AnimNextInterfaceTypes.h"
+#include "IAnimNextInterface.h"
+#include "AnimNextInterfaceContext.h"
+#include "AnimNextInterface_Float.generated.h"
+
+UCLASS()
+class UAnimNextInterfaceFloat : public UObject, public IAnimNextInterface
+{
+	GENERATED_BODY()
+
+	virtual FName GetReturnTypeNameImpl() const final override
+	{
+		static const FName Name = TNameOf<float>::GetName();
+		return Name;
+	}
+
+	virtual bool GetDataImpl(const UE::AnimNext::Interface::FContext& Context) const override
+	{
+		checkf(false, TEXT("UAnimNextInterfaceFloat::GetDataImpl must be overridden"));
+		return false;
+	}
+};
+
+UCLASS()
+class UAnimNextInterface_Float_Literal : public UAnimNextInterfaceFloat
+{
+	GENERATED_BODY()
+
+	// IAnimNextInterface interface
+	virtual bool GetDataImpl(const UE::AnimNext::Interface::FContext& Context) const final override
+	{
+		Context.SetResult(Value);
+		return true;
+	}
+	
+	UPROPERTY(EditAnywhere, Category = "Parameters")
+	float Value;
+};
+
+UCLASS()
+class UAnimNextInterface_Float_Multiply : public UAnimNextInterfaceFloat
+{
+	GENERATED_BODY()
+
+	// IAnimNextInterface interface
+	virtual bool GetDataImpl(const UE::AnimNext::Interface::FContext& Context) const final override;
+
+	UPROPERTY(EditAnywhere, Category = "Parameters")
+	TArray<TScriptInterface<IAnimNextInterface>> Inputs;
+};
+
+UCLASS()
+class UAnimNextInterface_Float_InterpTo : public UAnimNextInterfaceFloat
+{
+	GENERATED_BODY()
+
+	// IAnimNextInterface interface
+	virtual bool GetDataImpl(const UE::AnimNext::Interface::FContext& Context) const final override;
+
+	UPROPERTY(EditAnywhere, Category = "Parameters")
+	TScriptInterface<IAnimNextInterface> Current;
+
+	UPROPERTY(EditAnywhere, Category = "Parameters")
+	TScriptInterface<IAnimNextInterface> Target;
+
+	UPROPERTY(EditAnywhere, Category = "Parameters")
+	TScriptInterface<IAnimNextInterface> Speed;
+};
+
+UCLASS()
+class UAnimNextInterface_Float_DeltaTime : public UAnimNextInterfaceFloat
+{
+	GENERATED_BODY()
+
+	// IAnimNextInterface interface
+	virtual bool GetDataImpl(const UE::AnimNext::Interface::FContext& Context) const final override
+	{
+		Context.SetResult(Context.GetDeltaTime());
+		return true;
+	}
+};
+
+UCLASS()
+class UAnimNextInterface_Float_SpringInterp : public UAnimNextInterfaceFloat
+{
+	GENERATED_BODY()
+
+	// IAnimNextInterface interface
+	virtual bool GetDataImpl(const UE::AnimNext::Interface::FContext& Context) const final override;
+
+	UPROPERTY(EditAnywhere, Category = "Parameters")
+	TScriptInterface<IAnimNextInterface> TargetValue;
+
+	UPROPERTY(EditAnywhere, Category = "Parameters")
+	TScriptInterface<IAnimNextInterface> TargetValueRate;
+	
+	UPROPERTY(EditAnywhere, Category = "Parameters")
+	TScriptInterface<IAnimNextInterface> SmoothingTime;
+	
+	UPROPERTY(EditAnywhere, Category = "Parameters")
+	TScriptInterface<IAnimNextInterface> DampingRatio;
+};
\ No newline at end of file
diff --git a/Engine/Plugins/Experimental/AnimNext/Source/AnimNextInterfaceGraph/Private/AnimNextInterface_Name.cpp b/Engine/Plugins/Experimental/AnimNext/Source/AnimNextInterfaceGraph/Private/AnimNextInterface_Name.cpp
new file mode 100644
--- /dev/null
+++ b/Engine/Plugins/Experimental/AnimNext/Source/AnimNextInterfaceGraph/Private/AnimNextInterface_Name.cpp
@@ -0,0 +1,6 @@
+﻿// Copyright Epic Games, Inc. All Rights Reserved.
+
+#include "AnimNextInterface_Name.h"
+
+#include UE_INLINE_GENERATED_CPP_BY_NAME(AnimNextInterface_Name)
+
diff --git a/Engine/Plugins/Experimental/AnimNext/Source/AnimNextInterfaceGraph/Private/AnimNextInterface_Name.h b/Engine/Plugins/Experimental/AnimNext/Source/AnimNextInterfaceGraph/Private/AnimNextInterface_Name.h
new file mode 100644
--- /dev/null
+++ b/Engine/Plugins/Experimental/AnimNext/Source/AnimNextInterfaceGraph/Private/AnimNextInterface_Name.h
@@ -0,0 +1,28 @@
+// Copyright Epic Games, Inc. All Rights Reserved.
+
+#pragma once
+
+#include "CoreMinimal.h"
+#include "IAnimNextInterface.h"
+#include "AnimNextInterfaceContext.h"
+#include "AnimNextInterface_Name.generated.h"
+
+UCLASS()
+class UAnimNextInterface_Name_Literal : public UObject, public IAnimNextInterface
+{
+	GENERATED_BODY()
+
+	virtual FName GetReturnTypeNameImpl() const final override
+	{
+		return NAME_Name;
+	}
+
+	virtual bool GetDataImpl(const UE::AnimNext::Interface::FContext& Context) const final override
+	{
+		Context.SetResult(Value);
+		return true;
+	}
+	
+	UPROPERTY(EditAnywhere, Category = "Parameters")
+	FName Value;
+};
diff --git a/Engine/Plugins/Experimental/AnimNext/Source/AnimNextInterfaceGraph/Private/AnimNextInterface_Wrapper.cpp b/Engine/Plugins/Experimental/AnimNext/Source/AnimNextInterfaceGraph/Private/AnimNextInterface_Wrapper.cpp
new file mode 100644
--- /dev/null
+++ b/Engine/Plugins/Experimental/AnimNext/Source/AnimNextInterfaceGraph/Private/AnimNextInterface_Wrapper.cpp
@@ -0,0 +1,25 @@
+﻿// Copyright Epic Games, Inc. All Rights Reserved.
+
+#include "AnimNextInterface_Wrapper.h"
+
+#include UE_INLINE_GENERATED_CPP_BY_NAME(AnimNextInterface_Wrapper)
+
+FName UAnimNextInterface_Wrapper::GetReturnTypeNameImpl() const
+{
+	// We mimic the return type of our output
+	return Output.GetInterface() ? Output.GetInterface()->GetReturnTypeName() : NAME_None;
+}
+
+const UScriptStruct* UAnimNextInterface_Wrapper::GetReturnTypeStructImpl() const
+{
+	// We mimic the return type of our output
+	return Output.GetInterface() ? Output.GetInterface()->GetReturnTypeStruct() : nullptr;
+}
+
+bool UAnimNextInterface_Wrapper::GetDataImpl(const UE::AnimNext::Interface::FContext& Context) const
+{
+	// Parameterize output with inputs
+	//const UE::AnimNext::Interface::FContext ParameterizedContext = Context.WithParameters(Inputs);
+	//return IAnimNextInterface::StaticGetDataRaw(Output, GetReturnTypeName(), ParameterizedContext, Result);
+	return false;
+}
diff --git a/Engine/Plugins/Experimental/AnimNext/Source/AnimNextInterfaceGraph/Private/AnimNextInterface_Wrapper.h b/Engine/Plugins/Experimental/AnimNext/Source/AnimNextInterfaceGraph/Private/AnimNextInterface_Wrapper.h
new file mode 100644
--- /dev/null
+++ b/Engine/Plugins/Experimental/AnimNext/Source/AnimNextInterfaceGraph/Private/AnimNextInterface_Wrapper.h
@@ -0,0 +1,25 @@
+// Copyright Epic Games, Inc. All Rights Reserved.
+
+#pragma once
+
+#include "CoreMinimal.h"
+#include "IAnimNextInterface.h"
+#include "AnimNextInterface_Wrapper.generated.h"
+
+UCLASS()
+class UAnimNextInterface_Wrapper : public UObject, public IAnimNextInterface
+{
+	GENERATED_BODY()
+
+private:
+	// IAnimAnimNextInterface interface
+	virtual FName GetReturnTypeNameImpl() const final override;
+	virtual const UScriptStruct* GetReturnTypeStructImpl() const final override;
+	virtual bool GetDataImpl(const UE::AnimNext::Interface::FContext& Context) const final override;
+
+//	UPROPERTY(EditAnywhere, Category = "Parameters")
+//	TArray<FAnimNextInterfaceParameter> Inputs;
+
+	UPROPERTY(EditAnywhere, Category = "Parameters")
+	TScriptInterface<IAnimNextInterface> Output;
+};
\ No newline at end of file
diff --git a/Engine/Plugins/Experimental/AnimNext/Source/AnimNextInterfaceGraph/Private/RigUnit_AnimNextInterfaceBeginExecution.cpp b/Engine/Plugins/Experimental/AnimNext/Source/AnimNextInterfaceGraph/Private/RigUnit_AnimNextInterfaceBeginExecution.cpp
new file mode 100644
--- /dev/null
+++ b/Engine/Plugins/Experimental/AnimNext/Source/AnimNextInterfaceGraph/Private/RigUnit_AnimNextInterfaceBeginExecution.cpp
@@ -0,0 +1,12 @@
+﻿// Copyright Epic Games, Inc. All Rights Reserved.
+
+#include "RigUnit_AnimNextInterfaceBeginExecution.h"
+#include "Units/RigUnitContext.h"
+
+#include UE_INLINE_GENERATED_CPP_BY_NAME(RigUnit_AnimNextInterfaceBeginExecution)
+
+FName FRigUnit_AnimNextInterfaceBeginExecution::EventName = TEXT("Update");
+
+FRigUnit_AnimNextInterfaceBeginExecution_Execute()
+{
+}
diff --git a/Engine/Plugins/Experimental/AnimNext/Source/AnimNextInterfaceGraph/Private/RigUnit_AnimNextInterfaceEndExecution.cpp b/Engine/Plugins/Experimental/AnimNext/Source/AnimNextInterfaceGraph/Private/RigUnit_AnimNextInterfaceEndExecution.cpp
new file mode 100644
--- /dev/null
+++ b/Engine/Plugins/Experimental/AnimNext/Source/AnimNextInterfaceGraph/Private/RigUnit_AnimNextInterfaceEndExecution.cpp
@@ -0,0 +1,16 @@
+﻿// Copyright Epic Games, Inc. All Rights Reserved.
+
+#include "RigUnit_AnimNextInterfaceEndExecution.h"
+#include "Units/RigUnitContext.h"
+
+#include UE_INLINE_GENERATED_CPP_BY_NAME(RigUnit_AnimNextInterfaceEndExecution)
+
+FRigUnit_AnimNextInterfaceEndExecution_Bool_Execute()
+{
+	SetResult(ExecuteContext, Result);
+}
+
+FRigUnit_AnimNextInterfaceEndExecution_Float_Execute()
+{
+	SetResult(ExecuteContext, Result);
+}
diff --git a/Engine/Plugins/Experimental/AnimNext/Source/AnimNextInterfaceGraph/Private/RigUnit_AnimNextInterfaceParameter.cpp b/Engine/Plugins/Experimental/AnimNext/Source/AnimNextInterfaceGraph/Private/RigUnit_AnimNextInterfaceParameter.cpp
new file mode 100644
--- /dev/null
+++ b/Engine/Plugins/Experimental/AnimNext/Source/AnimNextInterfaceGraph/Private/RigUnit_AnimNextInterfaceParameter.cpp
@@ -0,0 +1,105 @@
+// Copyright Epic Games, Inc. All Rights Reserved.
+
+#include "RigUnit_AnimNextInterfaceParameter.h"
+#include "AnimNextInterfaceUnitContext.h"
+#include "AnimNextInterfaceExecuteContext.h"
+#include "AnimNextInterface.h"
+#include "AnimNextInterfaceTypes.h"
+#include "AnimNextInterfaceKernel.h"
+
+#include UE_INLINE_GENERATED_CPP_BY_NAME(RigUnit_AnimNextInterfaceParameter)
+
+bool FRigUnit_AnimNextInterfaceParameter::GetParameterInternal(FName InName, const FAnimNextInterfaceUnitContext& InContext, void* OutResult)
+{
+	return true;
+}
+
+FRigUnit_AnimNextInterfaceParameter_Float_Execute()
+{
+	GetParameterInternal(Parameter, ExecuteContext.GetUnitContext(), &Result);
+}
+
+FRigUnit_AnimNextInterfaceParameter_AnimNextInterface_Execute()
+{
+	GetParameterInternal(Parameter, ExecuteContext.GetUnitContext(), &Result);
+}
+
+FRigUnit_AnimNextInterface_Float_Execute()
+{
+	using namespace UE::AnimNext::Interface;
+
+	// @TODO: ensure that context arg is always present to avoid these checks here
+	const FContext& AnimNextInterfaceContext = ExecuteContext.GetContext();
+
+	// Wrap the internal result we are going to be writing to 
+	TWrapParam<float> CallResult(&Result);
+
+	// Call the interface
+	ExecuteContext.SetResult(UE::AnimNext::Interface::GetDataSafe(AnimNextInterface, AnimNextInterfaceContext, CallResult));
+}
+/*
+FRigUnit_AnimNextInterface_Pose_Execute()
+{
+	using namespace UE::AnimNext::Interface;
+
+	// @TODO: ensure that context arg is always present to avoid these checks here
+	if(Context.State == EControlRigState::Update && ExecuteContext.OpaqueArguments.Num() > 1 && ExecuteContext.OpaqueArguments[1] != nullptr)
+	{
+		const FAnimNextInterfaceUnitContext& AnimNextInterfaceUnitContext = *static_cast<const FAnimNextInterfaceUnitContext*>(ExecuteContext.OpaqueArguments[1]);
+		const FContext& AnimNextInterfaceContext = AnimNextInterfaceUnitContext.AnimNextInterfaceContext;
+
+		// Wrap the internal result we are going to be writing to 
+		TWrapParam<FAnimNextInterfaceExecuteContext> CallResult(&Result);
+
+		// Call the interface
+		AnimNextInterfaceUnitContext.bResult &= UE::AnimNext::Interface::GetDataSafe(PoseInterface, AnimNextInterfaceContext, CallResult);
+	}
+}
+*/
+FRigUnit_FloatOperator_Execute()
+{}
+
+FRigUnit_PoseOperator_Execute()
+{}
+
+FRigUnit_AnimNextInterface_SequencePlayer_Execute()
+{}
+
+struct FSpringDamperState
+{
+	float Value = 0.0f;
+	float ValueRate = 0.0f;
+};
+
+IMPLEMENT_ANIM_NEXT_INTERFACE_STATE_TYPE(FSpringDamperState, SpringDamperState);
+
+FRigUnit_TestFloatState_Execute()
+{
+	using namespace UE::AnimNext::Interface;
+
+	// @TODO: ensure that context arg is always present to avoid these checks here
+	const FContext& AnimNextInterfaceContext = ExecuteContext.GetContext();
+
+	const TParam<FSpringDamperState> State = AnimNextInterfaceContext.GetState<FSpringDamperState>(ExecuteContext.GetInterface(), 0);
+	const float DeltaTime = AnimNextInterfaceContext.GetDeltaTime();
+
+	FKernel::Run(AnimNextInterfaceContext,
+		[DeltaTime, &Result](FSpringDamperState& InOutState,
+					float InTargetValue,
+					float InTargetValueRate,
+					float InSmoothingTime,
+					float InDampingRatio)
+			{
+				FMath::SpringDamperSmoothing(
+					InOutState.Value,
+					InOutState.ValueRate,
+					InTargetValue,
+					InTargetValueRate,
+					DeltaTime,
+					InSmoothingTime,
+					InDampingRatio);
+
+				Result = InOutState.Value;
+			},
+			State, TargetValue, TargetValueRate, SmoothingTime, DampingRatio);
+}
diff --git a/Engine/Plugins/Experimental/AnimNext/Source/AnimNextInterfaceGraph/Private/RigUnit_AnimNextInterfaceParameter.h b/Engine/Plugins/Experimental/AnimNext/Source/AnimNextInterfaceGraph/Private/RigUnit_AnimNextInterfaceParameter.h
new file mode 100644
--- /dev/null
+++ b/Engine/Plugins/Experimental/AnimNext/Source/AnimNextInterfaceGraph/Private/RigUnit_AnimNextInterfaceParameter.h
@@ -0,0 +1,304 @@
+// Copyright Epic Games, Inc. All Rights Reserved.
+
+#pragma once
+
+#include "IAnimNextInterface.h"
+#include "Units/RigUnit.h"
+#include "AnimNextInterfaceGraph/Internal/AnimNextInterfaceExecuteContext.h"
+#include "RigUnit_AnimNextInterfaceParameter.generated.h"
+
+struct FAnimNextInterfaceUnitContext;
+class UAnimNextInterfaceGraph;
+
+/** Unit for reading parameters from context */
+USTRUCT(meta = (DisplayName = "Parameter", Category = "Parameters", TitleColor = "1 0 0", NodeColor = "1 1 1"))
+struct FRigUnit_AnimNextInterfaceParameter : public FRigUnit_AnimNextInterfaceBase
+{
+	GENERATED_BODY()
+
+protected:
+	static bool GetParameterInternal(FName InName, const FAnimNextInterfaceUnitContext& InContext, void* OutResult);
+
+	UPROPERTY(EditAnywhere, Category = "Parameter", meta = (Input))
+	FName Parameter = NAME_None;
+};
+
+/** Unit for reading float parameter from context */
+USTRUCT(meta = (DisplayName = "Float Parameter", Category = "Parameters", TitleColor = "1 0 0", NodeColor = "1 1 1"))
+struct FRigUnit_AnimNextInterfaceParameter_Float : public FRigUnit_AnimNextInterfaceParameter
+{
+	GENERATED_BODY()
+
+	RIGVM_METHOD()
+	void Execute();
+	
+protected:
+	UPROPERTY(EditAnywhere, Category = "Result", meta = (Output))
+	float Result = 0.0f;
+};
+
+/** Unit for reading data interface graph parameter from context */
+USTRUCT(meta = (DisplayName = "Data Interface Parameter", Category = "Parameters", TitleColor = "1 0 0", NodeColor = "1 1 1"))
+struct FRigUnit_AnimNextInterfaceParameter_AnimNextInterface : public FRigUnit_AnimNextInterfaceParameter
+{
+	GENERATED_BODY()
+
+	RIGVM_METHOD()
+	void Execute();
+	
+protected:
+	UPROPERTY(EditAnywhere, Category = "Result", meta = (Output))
+	TScriptInterface<IAnimNextInterface> Result = nullptr;
+};
+
+USTRUCT(BlueprintType)
+struct FAnimNextInterfaceParameter
+{
+	GENERATED_BODY()
+
+	FAnimNextInterfaceParameter()
+	: Name(NAME_None)
+	{}
+
+	UPROPERTY(EditAnywhere, Category = "Data Interface")
+	FName Name;
+};
+
+USTRUCT(BlueprintType)
+struct FAnimNextInterfaceParameter_Float
+{
+	GENERATED_BODY()
+	
+	FAnimNextInterfaceParameter_Float()
+		: Name(NAME_None)
+		, Value(0.f)
+	{}
+
+	UPROPERTY(EditAnywhere, Category = "Data Interface")
+	FName Name;
+
+	UPROPERTY(EditAnywhere, Category = "Data Interface")
+	float Value;
+};
+
+USTRUCT(BlueprintType)
+struct FAnimNextInterfaceParameter_Int
+{
+	GENERATED_BODY()
+
+	FAnimNextInterfaceParameter_Int()
+		: Name(NAME_None)
+		, Value(0)
+	{}
+
+	UPROPERTY(EditAnywhere, Category = "Data Interface")
+	FName Name;
+
+	UPROPERTY(EditAnywhere, Category = "Data Interface")
+	int32 Value;
+};
+
+USTRUCT(BlueprintType)
+struct FAnimNextInterfaceParameter_Bool
+{
+	GENERATED_BODY()
+
+	FAnimNextInterfaceParameter_Bool()
+		: Name(NAME_None)
+		, bValue(false)
+	{}
+
+	UPROPERTY(EditAnywhere, Category = "Data Interface")
+	FName Name;
+
+	UPROPERTY(EditAnywhere, Category = "Data Interface")
+	bool bValue;
+};
+
+USTRUCT(BlueprintType)
+struct FAnimNextInterfaceParameters1
+{
+	GENERATED_BODY()
+
+	FAnimNextInterfaceParameters1()
+		: Param0()
+		, FloatParam(0.f)
+	{}
+
+	UPROPERTY(EditAnywhere, Category = "Data Interface")
+	FAnimNextInterfaceParameter_Int Param0;
+
+	UPROPERTY(EditAnywhere, Category = "Data Interface")
+	float FloatParam;
+};
+
+USTRUCT(BlueprintType)
+struct FAnimNextInterfaceParameter_AnimNextInterface : public FAnimNextInterfaceParameter
+{
+	GENERATED_BODY()
+
+	UPROPERTY(EditAnywhere, Category = "Data Interface")
+	TScriptInterface<IAnimNextInterface> Value = nullptr;
+};
+
+/** Base unit for calling data interfaces from graphs */
+USTRUCT(meta = (DisplayName = "Data Interface", Category = "Execution", TitleColor = "1 0 0", NodeColor = "1 1 1"))
+struct FRigUnit_AnimNextInterface : public FRigUnit_AnimNextInterfaceBase
+{
+	GENERATED_BODY()
+
+protected:
+	UPROPERTY(EditAnywhere, Category = "Data Interface", meta = (Input))
+	TScriptInterface<IAnimNextInterface> AnimNextInterface = nullptr;
+};
+
+/** Unit for getting a float via a data interface */
+USTRUCT(meta = (DisplayName = "Data Interface Float", Category = "Parameters", TitleColor = "1 0 0", NodeColor = "1 1 1"))
+struct FRigUnit_AnimNextInterface_Float : public FRigUnit_AnimNextInterface
+{
+	GENERATED_BODY()
+
+	RIGVM_METHOD()
+	void Execute();
+	
+protected:
+	UPROPERTY(EditAnywhere, Category = "Data Interface", meta = (Input))
+	TArray<FAnimNextInterfaceParameter_AnimNextInterface> Parameters;
+	
+	UPROPERTY(EditAnywhere, Category = "Result", meta = (Output))
+	float Result = 0.0f;
+};
+
+/** Unit for getting a pose via a data interface */
+/*
+USTRUCT(meta=(DisplayName="Get Pose"))
+struct FRigUnit_AnimNextInterface_Pose : public FRigUnit
+{
+	GENERATED_BODY()
+
+	RIGVM_METHOD()
+	void Execute();
+	
+protected:
+	UPROPERTY(EditAnywhere, Category = "Pose Interface", meta = (Input))
+	TScriptInterface<IAnimNextInterface> PoseInterface = nullptr;
+
+	UPROPERTY(EditAnywhere, Category = "Pose Interface", meta = (Input))
+	TArray<FAnimNextInterfaceParameter_AnimNextInterface> Parameters;
+
+	UPROPERTY(EditAnywhere, Transient, DisplayName = "Result", Category = "BeginExecution", meta = (Output))
+	FAnimNextInterfaceExecuteContext Result;
+};
+*/
+/** Unit for getting a pose via a data interface */
+USTRUCT(meta = (DisplayName = "Float Operator", Category = "Operators", TitleColor = "1 0 0", NodeColor = "1 1 1"))
+struct FRigUnit_FloatOperator : public FRigUnit_AnimNextInterfaceBase
+{
+	GENERATED_BODY()
+	
+	FRigUnit_FloatOperator()
+		: ParamA(0.f)
+		, ParamB(0.f)
+		, Result(0.f)
+	{}
+
+	RIGVM_METHOD()
+	void Execute();
+	
+protected:
+	UPROPERTY(EditAnywhere, Category = "Pose Interface", meta = (Input))
+	TScriptInterface<IAnimNextInterface> Operator = nullptr;
+
+	UPROPERTY(EditAnywhere, Category = "Pose Interface", meta = (Input))
+	float ParamA;
+
+	UPROPERTY(EditAnywhere, Category = "Pose Interface", meta = (Input))
+	float ParamB;
+
+	UPROPERTY(EditAnywhere, Transient, DisplayName = "Result", Category = "BeginExecution", meta = (Output))
+	float Result;
+};
+
+/** Unit for getting a pose via a data interface */
+USTRUCT(meta = (DisplayName = "Pose Operator", Category = "Operators", TitleColor = "1 0 0", NodeColor = "1 1 1"))
+struct FRigUnit_PoseOperator : public FRigUnit_AnimNextInterfaceBase
+{
+	GENERATED_BODY()
+
+	RIGVM_METHOD()
+	void Execute();
+	
+protected:
+	UPROPERTY(EditAnywhere, Category = "Pose Interface", meta = (Input))
+	TScriptInterface<IAnimNextInterface> Operator = nullptr;
+
+	UPROPERTY(EditAnywhere, Category = "Pose Interface", meta = (Input))
+	TArray<FRigVMExecuteContext> InputPoses;
+
+	UPROPERTY(EditAnywhere, Category = "Pose Interface", meta = (Input))
+	FAnimNextInterfaceParameters1 Parameters;
+
+	UPROPERTY(EditAnywhere, Transient, DisplayName = "Result", Category = "BeginExecution", meta = (Output))
+	FAnimNextInterfaceExecuteContext Result;
+};
+
+USTRUCT(BlueprintType)
+struct FAnimSequenceParameters
+{
+	GENERATED_BODY()
+
+	UPROPERTY(EditAnywhere, Category = "Anim Sequence")
+	float PlayRate = 1.0f;
+
+	UPROPERTY(EditAnywhere, Category = "Anim Sequence")
+	float StartPosition = 0.0f;
+
+	UPROPERTY(EditAnywhere, Category = "Anim Sequence")
+	bool bLoop = 0.0f;
+};
+
+/** Unit for getting a pose via a data interface */
+USTRUCT(meta = (DisplayName = "Anim Sequence", Category = "Animation", TitleColor = "1 0 0", NodeColor = "1 1 1"))
+struct FRigUnit_AnimNextInterface_SequencePlayer : public FRigUnit_AnimNextInterfaceBase
+{
+	GENERATED_BODY()
+
+	RIGVM_METHOD()
+	void Execute();
+	
+protected:
+	UPROPERTY(EditAnywhere, Category = "Data Interface", meta = (Input))
+	FAnimSequenceParameters Parameters;
+
+	UPROPERTY(EditAnywhere, Category = "Data Interface", meta = (Input))
+	TScriptInterface<IAnimNextInterface> Sequence = nullptr;
+	
+	UPROPERTY(EditAnywhere, Transient, DisplayName = "Result", Category = "BeginExecution", meta = (Output))
+	FAnimNextInterfaceExecuteContext Result;
+};
+
+/** Unit for getting a pose via a data interface */
+USTRUCT(meta = (DisplayName = "Test Float State - Spring Damper Smoothing", Varying, Category = "Animation", TitleColor = "1 0 0", NodeColor = "1 1 1"))
+struct FRigUnit_TestFloatState : public FRigUnit_AnimNextInterfaceBase
+{
+	GENERATED_BODY()
+
+	RIGVM_METHOD()
+	void Execute();
+	
+protected:
+	UPROPERTY(EditAnywhere, Category = "Parameters", meta=(Input))
+	float TargetValue = 0.0f;
+
+	UPROPERTY(EditAnywhere, Category = "Parameters", meta=(Input))
+	float TargetValueRate = 1.0f;
+	
+	UPROPERTY(EditAnywhere, Category = "Parameters", meta=(Input))
+	float SmoothingTime = 1.0f;
+	
+	UPROPERTY(EditAnywhere, Category = "Parameters", meta=(Input))
+	float DampingRatio = 1.0f;
+
+	UPROPERTY(EditAnywhere, Category = "Parameters", meta=(Output))
+	float Result = 0.0f;
+};
diff --git a/Engine/Plugins/Experimental/AnimNext/Source/AnimNextInterfaceGraphEditor/AnimNextInterfaceGraphEditor.Build.cs b/Engine/Plugins/Experimental/AnimNext/Source/AnimNextInterfaceGraphEditor/AnimNextInterfaceGraphEditor.Build.cs
new file mode 100644
--- /dev/null
+++ b/Engine/Plugins/Experimental/AnimNext/Source/AnimNextInterfaceGraphEditor/AnimNextInterfaceGraphEditor.Build.cs
@@ -0,0 +1,34 @@
+// Copyright Epic Games, Inc. All Rights Reserved.
+
+namespace UnrealBuildTool.Rules
+{
+	public class AnimNextInterfaceGraphEditor : ModuleRules
+	{
+		public AnimNextInterfaceGraphEditor(ReadOnlyTargetRules Target) : base(Target)
+		{
+
+			PrivateDependencyModuleNames.AddRange(
+				new string[]
+				{
+					"Core",
+					"CoreUObject",
+					"Engine",
+					"AssetTools",
+					"AnimNextInterface",
+					"AnimNextInterfaceGraph",
+					"AnimNextInterfaceGraphUncookedOnly",
+					"UnrealEd",
+					"SlateCore",
+					"Slate",
+					"PropertyEditor",
+					"RigVM",
+					"RigVMDeveloper",
+					"ControlRig",
+					"ControlRigDeveloper",
+					"ControlRigEditor",
+					"GraphEditor",
+				}
+			);
+		}
+	}
+}
\ No newline at end of file
diff --git a/Engine/Plugins/Experimental/AnimNext/Source/AnimNextInterfaceGraphEditor/Private/AnimNextInterfaceEditorUtils.cpp b/Engine/Plugins/Experimental/AnimNext/Source/AnimNextInterfaceGraphEditor/Private/AnimNextInterfaceEditorUtils.cpp
new file mode 100644
--- /dev/null
+++ b/Engine/Plugins/Experimental/AnimNext/Source/AnimNextInterfaceGraphEditor/Private/AnimNextInterfaceEditorUtils.cpp
@@ -0,0 +1,122 @@
+// Copyright Epic Games, Inc. All Rights Reserved.
+
+#include "AnimNextInterfaceEditorUtils.h"
+#include "AnimNextInterfaceGraph_EditorData.h"
+#include "Kismet2/Kismet2NameValidators.h"
+
+namespace UE::AnimNext::InterfaceGraphEditor
+{
+
+void FUtils::GetAllGraphNames(const UAnimNextInterfaceGraph_EditorData* InEditorData, TSet<FName>& OutNames)
+{
+
+}
+
+static const int32 MaxNameLength = 100;
+
+FName FUtils::ValidateName(const UAnimNextInterfaceGraph_EditorData* InEditorData, const FString& InName)
+{
+	struct FNameValidator : public INameValidatorInterface
+	{
+		FNameValidator(const UAnimNextInterfaceGraph_EditorData* InEditorData)
+			: EditorData(InEditorData)
+		{
+			GetAllGraphNames(EditorData, Names);
+		}
+		
+		virtual EValidatorResult IsValid (const FName& Name, bool bOriginal = false) override
+		{
+			EValidatorResult ValidatorResult = EValidatorResult::AlreadyInUse;
+
+			if(Name == NAME_None)
+			{
+				ValidatorResult = EValidatorResult::EmptyName;
+			}
+			else if(Name.ToString().Len() > MaxNameLength)
+			{
+				ValidatorResult = EValidatorResult::TooLong;
+			}
+			else
+			{
+				// If it is in the names list then it is already in use.
+				if(!Names.Contains(Name))
+				{
+					ValidatorResult = EValidatorResult::Ok;
+
+					// Check for collision with an existing object.
+					if (UObject* ExistingObject = StaticFindObject(/*Class=*/ nullptr, const_cast<UAnimNextInterfaceGraph_EditorData*>(EditorData), *Name.ToString(), true))
+					{
+						ValidatorResult = EValidatorResult::AlreadyInUse;
+					}
+				}
+			}
+			
+			return ValidatorResult;
+		}
+		
+		virtual EValidatorResult IsValid (const FString& Name, bool bOriginal = false) override
+		{
+			// Converting a string that is too large for an FName will cause an assert, so verify the length
+			if(Name.Len() >= NAME_SIZE)
+			{
+				return EValidatorResult::TooLong;
+			}
+			else if (!FName::IsValidXName(Name, UE_BLUEPRINT_INVALID_NAME_CHARACTERS))
+			{
+				return EValidatorResult::ContainsInvalidCharacters;
+			}
+
+			// If not defined in name table, not current graph name
+			return IsValid( FName(*Name) );
+		}
+
+		/** Name set to validate */
+		TSet<FName> Names;
+		/** The editor data to check for validity within */
+		const UAnimNextInterfaceGraph_EditorData* EditorData;
+	};
+	
+	FString Name = InName;
+	if (Name.StartsWith(TEXT("RigUnit_")))
+	{
+		Name.RightChopInline(8, false);
+	}
+
+	FNameValidator NameValidator(InEditorData);
+
+	// Clean up BaseName to not contain any invalid characters, which will mean we can never find a legal name no matter how many numbers we add
+	if (NameValidator.IsValid(Name) == EValidatorResult::ContainsInvalidCharacters)
+	{
+		for (TCHAR& TestChar : Name)
+		{
+			for (TCHAR BadChar : UE_BLUEPRINT_INVALID_NAME_CHARACTERS)
+			{
+				if (TestChar == BadChar)
+				{
+					TestChar = TEXT('_');
+					break;
+				}
+			}
+		}
+	}
+	
+	int32 Count = 0;
+	FString BaseName = Name;
+	while (NameValidator.IsValid(Name) != EValidatorResult::Ok)
+	{
+		// Calculate the number of digits in the number, adding 2 (1 extra to correctly count digits, another to account for the '_' that will be added to the name
+		int32 CountLength = Count > 0 ? (int32)log((double)Count) + 2 : 2;
+
+		// If the length of the final string will be too long, cut off the end so we can fit the number
+		if (CountLength + BaseName.Len() > MaxNameLength)
+		{
+			BaseName.LeftInline(MaxNameLength - CountLength);
+		}
+		Name = FString::Printf(TEXT("%s_%d"), *BaseName, Count);
+		Count++;
+	}
+
+	return *Name;
+}
+
+}
diff --git a/Engine/Plugins/Experimental/AnimNext/Source/AnimNextInterfaceGraphEditor/Private/AnimNextInterfaceEditorUtils.h b/Engine/Plugins/Experimental/AnimNext/Source/AnimNextInterfaceGraphEditor/Private/AnimNextInterfaceEditorUtils.h
new file mode 100644
--- /dev/null
+++ b/Engine/Plugins/Experimental/AnimNext/Source/AnimNextInterfaceGraphEditor/Private/AnimNextInterfaceEditorUtils.h
@@ -0,0 +1,20 @@
+// Copyright Epic Games, Inc. All Rights Reserved.
+
+#pragma once
+
+#include "CoreMinimal.h"
+
+class UAnimNextInterfaceGraph;
+class UAnimNextInterfaceGraph_EditorData;
+
+namespace UE::AnimNext::InterfaceGraphEditor
+{
+
+struct FUtils
+{
+	static FName ValidateName(const UAnimNextInterfaceGraph_EditorData* InEditorData, const FString& InName);
+
+	static void GetAllGraphNames(const UAnimNextInterfaceGraph_EditorData* InEditorData, TSet<FName>& OutNames);
+};
+
+}
diff --git a/Engine/Plugins/Experimental/AnimNext/Source/AnimNextInterfaceGraphEditor/Private/AnimNextInterfaceGraphEditor.cpp b/Engine/Plugins/Experimental/AnimNext/Source/AnimNextInterfaceGraphEditor/Private/AnimNextInterfaceGraphEditor.cpp
new file mode 100644
--- /dev/null
+++ b/Engine/Plugins/Experimental/AnimNext/Source/AnimNextInterfaceGraphEditor/Private/AnimNextInterfaceGraphEditor.cpp
@@ -0,0 +1,322 @@
+﻿// Copyright Epic Games, Inc. All Rights Reserved.
+
+#include "AnimNextInterfaceGraphEditor.h"
+#include "AnimNextInterfaceGraphEditorMode.h"
+#include "AnimNextInterfaceGraph.h"
+#include "AnimNextInterfaceGraph_EditorData.h"
+#include "EdGraphNode_Comment.h"
+#include "EdGraphUtilities.h"
+#include "SAnimNextInterfaceGraphEditorActionMenu.h"
+#include "AnimNextInterfaceUncookedOnlyUtils.h"
+#include "Framework/Commands/GenericCommands.h"
+#include "RigVMModel/RigVMController.h"
+#include "RigVMModel/RigVMGraph.h"
+#include "Kismet2/BlueprintEditorUtils.h"
+#include "ScopedTransaction.h"
+
+#define LOCTEXT_NAMESPACE "AnimNextInterfaceGraphEditor"
+
+namespace UE::AnimNext::InterfaceGraphEditor
+{
+
+namespace Modes
+{
+	const FName GraphEditor("AnimNextInterfaceEditorMode");
+}
+
+namespace Tabs
+{
+	const FName Details("DetailsTab");
+	const FName Document("Document");
+}
+
+const FName AppIdentifier("AnimNextInterfaceGraphEditor");
+
+void FGraphEditor::InitEditor(const EToolkitMode::Type InMode, const TSharedPtr<IToolkitHost>& InInitToolkitHost, UAnimNextInterfaceGraph* InAnimNextInterfaceGraph)
+{
+	AnimNextInterfaceGraph = InAnimNextInterfaceGraph;
+
+	AnimNextInterfaceGraph_EditorData = UE::AnimNext::InterfaceGraphUncookedOnly::FUtils::GetEditorData(AnimNextInterfaceGraph);
+	AnimNextInterfaceGraph_EditorData->Initialize(false);
+	
+	DocumentManager = MakeShared<FDocumentTracker>();
+	DocumentManager->Initialize(SharedThis(this));
+
+	TSharedRef<FGraphEditorSummoner> GraphEditorSummoner = MakeShared<FGraphEditorSummoner>(SharedThis(this));
+	GraphEditorSummoner->OnCreateGraphEditorWidget().BindSP(this, &FGraphEditor::CreateGraphEditorWidget);
+	GraphEditorSummoner->OnGraphEditorFocused().BindSP(this, &FGraphEditor::OnGraphEditorFocused);
+	GraphEditorSummoner->OnGraphEditorBackgrounded().BindSP(this, &FGraphEditor::OnGraphEditorBackgrounded);
+	GraphEditorTabFactoryPtr = GraphEditorSummoner;
+	DocumentManager->RegisterDocumentFactory(GraphEditorSummoner);
+
+	const bool bCreateDefaultStandaloneMenu = true;
+	const bool bCreateDefaultToolbar = true;
+	FAssetEditorToolkit::InitAssetEditor(InMode, InInitToolkitHost, AppIdentifier, FTabManager::FLayout::NullLayout, bCreateDefaultStandaloneMenu, bCreateDefaultToolbar, InAnimNextInterfaceGraph);
+
+	BindCommands();
+
+	AddApplicationMode(Modes::GraphEditor, MakeShared<FGraphEditorMode>(SharedThis(this)));
+	SetCurrentMode(Modes::GraphEditor);
+	
+	ExtendMenu();
+	ExtendToolbar();
+	RegenerateMenusAndToolbars();
+
+	// Open initial document
+	DocumentManager->OpenDocument(FTabPayload_UObject::Make(AnimNextInterfaceGraph_EditorData->EntryPointGraph), FDocumentTracker::EOpenDocumentCause::OpenNewDocument);
+	DocumentManager->OpenDocument(FTabPayload_UObject::Make(AnimNextInterfaceGraph_EditorData->RootGraph), FDocumentTracker::EOpenDocumentCause::OpenNewDocument);
+	DocumentManager->OpenDocument(FTabPayload_UObject::Make(AnimNextInterfaceGraph_EditorData->FunctionLibraryEdGraph), FDocumentTracker::EOpenDocumentCause::OpenNewDocument);
+}
+
+void FGraphEditor::BindCommands()
+{
+}
+
+void FGraphEditor::ExtendMenu()
+{
+	
+}
+
+void FGraphEditor::ExtendToolbar()
+{
+	
+}
+
+void FGraphEditor::RegisterTabSpawners(const TSharedRef<FTabManager>& InTabManager)
+{
+	DocumentManager->SetTabManager(InTabManager);
+
+	FWorkflowCentricApplication::RegisterTabSpawners(InTabManager);
+}
+
+void FGraphEditor::UnregisterTabSpawners(const TSharedRef<FTabManager>& InTabManager)
+{
+	
+}
+
+FName FGraphEditor::GetToolkitFName() const
+{
+	return FName("AnimNextInterfaceGraphEditor");
+}
+
+FText FGraphEditor::GetBaseToolkitName() const
+{
+	return LOCTEXT("AppLabel", "AnimNextInterfaceGraphEditor");
+}
+
+FString FGraphEditor::GetWorldCentricTabPrefix() const
+{
+	return LOCTEXT("WorldCentricTabPrefix", "AnimNextInterfaceGraphEditor ").ToString();
+}
+
+FLinearColor FGraphEditor::GetWorldCentricTabColorScale() const
+{
+	return FLinearColor(0.3f, 0.2f, 0.5f, 0.5f);
+}
+
+void FGraphEditor::InitToolMenuContext(FToolMenuContext& InMenuContext)
+{
+}
+
+void FGraphEditor::OnGraphEditorFocused(TSharedRef<SGraphEditor> InGraphEditor)
+{
+	// Update the graph editor that is currently focused
+	FocusedGraphEdPtr = InGraphEditor;
+}
+
+void FGraphEditor::OnGraphEditorBackgrounded(TSharedRef<SGraphEditor> InGraphEditor)
+{
+	// Update the graph editor that is currently focused
+	FocusedGraphEdPtr = nullptr;
+}
+
+UEdGraph* FGraphEditor::GetFocusedGraph() const
+{
+	if (FocusedGraphEdPtr.IsValid())
+	{
+		if (UEdGraph* Graph = FocusedGraphEdPtr.Pin()->GetCurrentGraph())
+		{
+			return Graph;
+		}
+	}
+	return nullptr;
+}
+
+URigVMGraph* FGraphEditor::GetFocusedVMGraph() const
+{
+	UAnimNextInterfaceGraph_EdGraph* EdGraph = Cast<UAnimNextInterfaceGraph_EdGraph>(GetFocusedGraph());
+	return AnimNextInterfaceGraph_EditorData->GetVMGraphForEdGraph(EdGraph);
+}
+
+URigVMController* FGraphEditor::GetFocusedVMController() const
+{
+	return AnimNextInterfaceGraph_EditorData->GetRigVMClient()->GetController(GetFocusedVMGraph());
+}
+
+TSharedRef<SGraphEditor> FGraphEditor::CreateGraphEditorWidget(TSharedRef<FTabInfo> InTabInfo, UEdGraph* InGraph)
+{
+	if(!GraphEditorCommands.IsValid())
+	{
+		GraphEditorCommands = MakeShared<FUICommandList>();
+	
+		GraphEditorCommands->MapAction(FGenericCommands::Get().Delete,
+			FExecuteAction::CreateSP(this, &FGraphEditor::DeleteSelectedNodes),
+			FCanExecuteAction::CreateSP(this, &FGraphEditor::CanDeleteSelectedNodes));
+	}
+	
+	SGraphEditor::FGraphEditorEvents Events;
+	Events.OnCreateActionMenu = SGraphEditor::FOnCreateActionMenu::CreateSP(this, &FGraphEditor::OnCreateGraphActionMenu);
+	
+	return SNew(SGraphEditor)
+		.AdditionalCommands(GraphEditorCommands)
+		.IsEditable(this, &FGraphEditor::IsEditable, InGraph)
+		.GraphToEdit(InGraph)
+		.GraphEvents(Events)
+		.AssetEditorToolkit(AsShared());
+}
+
+FActionMenuContent FGraphEditor::OnCreateGraphActionMenu(UEdGraph* InGraph, const FVector2D& InNodePosition, const TArray<UEdGraphPin*>& InDraggedPins, bool bAutoExpand, SGraphEditor::FActionMenuClosed InOnMenuClosed)
+{
+	TSharedRef<SActionMenu> ActionMenu = SNew(SActionMenu)
+		.AutoExpandActionMenu(bAutoExpand)
+		.Graph(InGraph)
+		.NewNodePosition(InNodePosition)
+		.DraggedFromPins(InDraggedPins)
+		.OnClosedCallback(InOnMenuClosed);
+
+	TSharedPtr<SWidget> FilterTextBox = StaticCastSharedRef<SWidget>(ActionMenu->GetFilterTextBox());
+	return FActionMenuContent(StaticCastSharedRef<SWidget>(ActionMenu), FilterTextBox);
+}
+
+void FGraphEditor::OnNodeTitleCommitted(const FText& NewText, ETextCommit::Type CommitInfo, UEdGraphNode* NodeBeingChanged)
+{
+	if (UEdGraphNode_Comment* CommentBeingChanged = Cast<UEdGraphNode_Comment>(NodeBeingChanged))
+	{
+		GetFocusedVMController()->SetCommentTextByName(CommentBeingChanged->GetFName(), NewText.ToString(), CommentBeingChanged->FontSize, CommentBeingChanged->bCommentBubbleVisible, CommentBeingChanged->bColorCommentBubble, true, true);
+	}
+}
+
+void FGraphEditor::CloseDocumentTab(const UObject* DocumentID)
+{
+	TSharedRef<FTabPayload_UObject> Payload = FTabPayload_UObject::Make(DocumentID);
+	DocumentManager->CloseTab(Payload);
+}
+
+bool FGraphEditor::InEditingMode() const
+{
+	// @TODO: disallow editing when debugging when implemented
+
+	return true;
+}
+
+bool FGraphEditor::IsEditable(UEdGraph* InGraph) const
+{
+	return InGraph && InEditingMode() && InGraph->bEditable;
+}
+
+FGraphPanelSelectionSet FGraphEditor::GetSelectedNodes() const
+{
+	FGraphPanelSelectionSet CurrentSelection;
+	TSharedPtr<SGraphEditor> FocusedGraphEd = FocusedGraphEdPtr.Pin();
+	if (FocusedGraphEd.IsValid())
+	{
+		CurrentSelection = FocusedGraphEd->GetSelectedNodes();
+	}
+	return CurrentSelection;
+}
+
+void FGraphEditor::DeleteSelectedNodes()
+{
+	TSharedPtr<SGraphEditor> FocusedGraphEd = FocusedGraphEdPtr.Pin();
+	if (!FocusedGraphEd.IsValid())
+	{
+		return;
+	}
+
+	const FScopedTransaction Transaction(FGenericCommands::Get().Delete->GetDescription());
+	FocusedGraphEd->GetCurrentGraph()->Modify();
+	
+	const FGraphPanelSelectionSet SelectedNodes = GetSelectedNodes();
+
+//	SetUISelectionState(NAME_None);
+
+	if(FocusedGraphEd)
+	{
+		FocusedGraphEd->ClearSelectionSet();
+	}
+
+	// Some nodes have sub-objects that are represented as other tabs.
+	// Close them here as a pre-pass before we remove their nodes. If the documents are left open they
+	// may reference dangling data and function incorrectly in cases such as FindBlueprintforNodeChecked
+	for (FGraphPanelSelectionSet::TConstIterator NodeIt( SelectedNodes ); NodeIt; ++NodeIt)
+	{
+		if (UEdGraphNode* Node = Cast<UEdGraphNode>(*NodeIt))
+		{
+			if (Node->CanUserDeleteNode())
+			{
+				auto CloseAllDocumentsTab = [this](const UEdGraphNode* InNode)
+				{
+					TArray<UObject*> NodesToClose;
+					GetObjectsWithOuter(InNode, NodesToClose);
+					for (UObject* Node : NodesToClose)
+					{
+						UEdGraph* NodeGraph = Cast<UEdGraph>(Node);
+						if (NodeGraph)
+						{
+							CloseDocumentTab(NodeGraph);
+						}
+					}
+				};
+				
+				if (Node->GetSubGraphs().Num() > 0)
+				{
+					CloseAllDocumentsTab(Node);
+				}
+			}
+		}
+	}
+
+	// Now remove the selected nodes
+	for (FGraphPanelSelectionSet::TConstIterator NodeIt( SelectedNodes ); NodeIt; ++NodeIt)
+	{
+		if (UEdGraphNode* Node = Cast<UEdGraphNode>(*NodeIt))
+		{
+			if (Node->CanUserDeleteNode())
+			{
+				if (Node->GetSubGraphs().Num() > 0)
+				{
+					DocumentManager->CleanInvalidTabs();
+				}
+
+				FBlueprintEditorUtils::RemoveNode(nullptr, Node);
+			}
+		}
+	}
+}
+
+bool FGraphEditor::CanDeleteSelectedNodes()
+{
+	const FGraphPanelSelectionSet SelectedNodes = GetSelectedNodes();
+
+	bool bCanUserDeleteNode = false;
+
+	if(IsEditable(GetFocusedGraph()) && SelectedNodes.Num() > 0)
+	{
+		for(UObject* NodeObject : SelectedNodes)
+		{
+			// If any nodes allow deleting, then do not disable the delete option
+			UEdGraphNode* Node = Cast<UEdGraphNode>(NodeObject);
+			if(Node->CanUserDeleteNode())
+			{
+				bCanUserDeleteNode = true;
+				break;
+			}
+		}
+	}
+
+	return bCanUserDeleteNode;
+}
+
+}
+
+#undef LOCTEXT_NAMESPACE
\ No newline at end of file
diff --git a/Engine/Plugins/Experimental/AnimNext/Source/AnimNextInterfaceGraphEditor/Private/AnimNextInterfaceGraphEditor.h b/Engine/Plugins/Experimental/AnimNext/Source/AnimNextInterfaceGraphEditor/Private/AnimNextInterfaceGraphEditor.h
new file mode 100644
--- /dev/null
+++ b/Engine/Plugins/Experimental/AnimNext/Source/AnimNextInterfaceGraphEditor/Private/AnimNextInterfaceGraphEditor.h
@@ -0,0 +1,123 @@
+﻿// Copyright Epic Games, Inc. All Rights Reserved.
+
+#pragma once
+
+#include "CoreMinimal.h"
+#include "WorkflowOrientedApp/WorkflowCentricApplication.h"
+#include "GraphEditor.h"
+
+enum class ERigVMGraphNotifType : uint8;
+class UAnimNextInterfaceGraph;
+class UAnimNextInterfaceGraph_EditorData;
+class FDocumentTracker;
+class FDocumentTabFactory;
+class FTabInfo;
+class FTabManager;
+class IToolkitHost;
+class URigVMGraph;
+class URigVMController;
+class UDetailsViewWrapperObject;
+
+namespace UE::AnimNext::InterfaceGraphEditor
+{
+
+namespace Modes
+{
+	extern const FName GraphEditor;
+}
+
+namespace Tabs
+{
+	extern const FName Details;
+	extern const FName Document;
+}
+
+class FGraphEditor : public FWorkflowCentricApplication
+{
+public:
+	/** Edits the specified asset */
+	void InitEditor(const EToolkitMode::Type InMode, const TSharedPtr<IToolkitHost>& InInitToolkitHost, UAnimNextInterfaceGraph* InAnimNextInterfaceGraph);
+
+private:
+	friend class FGraphEditorMode;
+	
+	// IToolkit interface
+	virtual void RegisterTabSpawners(const TSharedRef<FTabManager>& InTabManager) override;
+	virtual void UnregisterTabSpawners(const TSharedRef<FTabManager>& InTabManager) override;
+	virtual FName GetToolkitFName() const override;
+	virtual FText GetBaseToolkitName() const override;
+	virtual FString GetWorldCentricTabPrefix() const override;
+	virtual FLinearColor GetWorldCentricTabColorScale() const override;
+	virtual void InitToolMenuContext(FToolMenuContext& InMenuContext) override;
+
+	void BindCommands();
+	
+	void ExtendMenu();
+
+	void ExtendToolbar();
+
+	void OnGraphEditorFocused(TSharedRef<SGraphEditor> InGraphEditor);
+
+	void OnGraphEditorBackgrounded(TSharedRef<SGraphEditor> InGraphEditor);
+	
+	UEdGraph* GetFocusedGraph() const;
+
+	URigVMGraph* GetFocusedVMGraph() const;
+
+	URigVMController* GetFocusedVMController() const;
+	
+	// Create new graph editor widget for the supplied document container
+	TSharedRef<SGraphEditor> CreateGraphEditorWidget(TSharedRef<FTabInfo> InTabInfo, UEdGraph* InGraph);
+
+	// Create the action menu for the specified graph
+	FActionMenuContent OnCreateGraphActionMenu(UEdGraph* InGraph, const FVector2D& InNodePosition, const TArray<UEdGraphPin*>& InDraggedPins, bool bAutoExpand, SGraphEditor::FActionMenuClosed InOnMenuClosed);
+
+	// Handle node titles being changed
+	void OnNodeTitleCommitted(const FText& NewText, ETextCommit::Type CommitInfo, UEdGraphNode* NodeBeingChanged);
+
+	void CloseDocumentTab(const UObject* DocumentID);
+
+	bool InEditingMode() const;
+
+	bool IsEditable(UEdGraph* InGraph) const;
+	
+	FGraphPanelSelectionSet GetSelectedNodes() const;
+	
+	void DeleteSelectedNodes();
+	
+	bool CanDeleteSelectedNodes();
+
+	TSharedPtr<SGraphEditor> GetGraphEditor(UEdGraph* InEdGraph) const;
+	
+	void HandleModifiedEvent(ERigVMGraphNotifType InNotifType, URigVMGraph* InGraph, UObject* InSubject);
+
+	void HandleDetailsViewCreated(TSharedRef<IDetailsView> InDetailsView)
+	{
+		DetailsView = InDetailsView;
+	}
+	
+	// The asset we are editing
+	UAnimNextInterfaceGraph* AnimNextInterfaceGraph = nullptr;
+
+	// Cached editor data ptr
+	UAnimNextInterfaceGraph_EditorData* AnimNextInterfaceGraph_EditorData = nullptr;
+	
+	// Currently focused graph editor
+	TWeakPtr<SGraphEditor> FocusedGraphEdPtr;
+
+	// Document tracker
+	TSharedPtr<FDocumentTracker> DocumentManager;
+
+	// Factory for document tabs
+	TWeakPtr<FDocumentTabFactory> GraphEditorTabFactoryPtr;
+
+	// Command list for the graph editor
+	TSharedPtr<FUICommandList> GraphEditorCommands;
+
+	// Our details panel
+	TSharedPtr<IDetailsView> DetailsView;
+
+	TArray<TStrongObjectPtr<UDetailsViewWrapperObject>> WrapperObjects;
+};
+
+}
\ No newline at end of file
diff --git a/Engine/Plugins/Experimental/AnimNext/Source/AnimNextInterfaceGraphEditor/Private/AnimNextInterfaceGraphEditorMode.cpp b/Engine/Plugins/Experimental/AnimNext/Source/AnimNextInterfaceGraphEditor/Private/AnimNextInterfaceGraphEditorMode.cpp
new file mode 100644
--- /dev/null
+++ b/Engine/Plugins/Experimental/AnimNext/Source/AnimNextInterfaceGraphEditor/Private/AnimNextInterfaceGraphEditorMode.cpp
@@ -0,0 +1,273 @@
+﻿// Copyright Epic Games, Inc. All Rights Reserved.
+
+#include "AnimNextInterfaceGraphEditorMode.h"
+#include "AnimNextInterfaceGraphEditor.h"
+#include "PropertyEditorModule.h"
+#include "Widgets/Docking/SDockTab.h"
+#include "EdGraph/EdGraph.h"
+#include "EdGraph/EdGraphSchema.h"
+#include "Widgets/Layout/SSpacer.h"
+#include "Modules/ModuleManager.h"
+
+#define LOCTEXT_NAMESPACE "GraphEditorMode"
+
+namespace UE::AnimNext::InterfaceGraphEditor
+{
+
+struct FGraphTabHistory : public FGenericTabHistory
+{
+public:
+	FGraphTabHistory(TSharedPtr<FDocumentTabFactory> InFactory, TSharedPtr<FTabPayload> InPayload)
+		: FGenericTabHistory(InFactory, InPayload)
+		, SavedLocation(FVector2D::ZeroVector)
+		, SavedZoomAmount(INDEX_NONE)
+	{
+	}
+
+private:
+	// FGenericTabHistory interface
+	virtual void EvokeHistory(TSharedPtr<::FTabInfo> InTabInfo, bool bPrevTabMatches) override
+	{
+		FWorkflowTabSpawnInfo SpawnInfo;
+		SpawnInfo.Payload = Payload;
+		SpawnInfo.TabInfo = InTabInfo;
+
+		if(bPrevTabMatches)
+		{
+			TSharedPtr<SDockTab> DockTab = InTabInfo->GetTab().Pin();
+			GraphEditor = StaticCastSharedRef<SGraphEditor>(DockTab->GetContent());
+		}
+		else
+		{
+			TSharedRef<SGraphEditor> GraphEditorRef = StaticCastSharedRef<SGraphEditor>(FactoryPtr.Pin()->CreateTabBody(SpawnInfo));
+			GraphEditor = GraphEditorRef;
+			FactoryPtr.Pin()->UpdateTab(InTabInfo->GetTab().Pin(), SpawnInfo, GraphEditorRef);
+		}
+	}
+	
+	virtual void SaveHistory() override
+	{
+		if (IsHistoryValid())
+		{
+			check(GraphEditor.IsValid());
+			GraphEditor.Pin()->GetViewLocation(SavedLocation, SavedZoomAmount);
+			GraphEditor.Pin()->GetViewBookmark(SavedBookmarkId);
+		}
+	}
+	
+	virtual void RestoreHistory() override
+	{
+		if (IsHistoryValid())
+		{
+			check(GraphEditor.IsValid());
+			GraphEditor.Pin()->SetViewLocation(SavedLocation, SavedZoomAmount, SavedBookmarkId);
+		}
+	}
+	
+	/** The graph editor represented by this history node. While this node is inactive, the graph editor is invalid */
+	TWeakPtr<SGraphEditor> GraphEditor;
+	
+	/** Saved location the graph editor was at when this history node was last visited */
+	FVector2D SavedLocation;
+	
+	/** Saved zoom the graph editor was at when this history node was last visited */
+	float SavedZoomAmount;
+	
+	/** Saved bookmark ID the graph editor was at when this history node was last visited */
+	FGuid SavedBookmarkId;
+};
+
+FGraphEditorSummoner::FGraphEditorSummoner(TSharedPtr<FWorkflowCentricApplication> InHostingApp)
+	: FDocumentTabFactoryForObjects<UEdGraph>(Tabs::Document, InHostingApp)
+{
+}
+
+void FGraphEditorSummoner::OnTabActivated(TSharedPtr<SDockTab> Tab) const
+{
+	TSharedRef<SGraphEditor> GraphEditor = StaticCastSharedRef<SGraphEditor>(Tab->GetContent());
+	OnGraphEditorFocusedDelegate.ExecuteIfBound(GraphEditor);
+}
+
+void FGraphEditorSummoner::OnTabBackgrounded(TSharedPtr<SDockTab> Tab) const
+{
+	TSharedRef<SGraphEditor> GraphEditor = StaticCastSharedRef<SGraphEditor>(Tab->GetContent());
+	OnGraphEditorBackgroundedDelegate.ExecuteIfBound(GraphEditor);
+}
+
+void FGraphEditorSummoner::OnTabRefreshed(TSharedPtr<SDockTab> Tab) const
+{
+	TSharedRef<SGraphEditor> GraphEditor = StaticCastSharedRef<SGraphEditor>(Tab->GetContent());
+	GraphEditor->NotifyGraphChanged();
+}
+
+void FGraphEditorSummoner::SaveState(TSharedPtr<SDockTab> Tab, TSharedPtr<FTabPayload> Payload) const
+{
+	TSharedRef<SGraphEditor> GraphEditor = StaticCastSharedRef<SGraphEditor>(Tab->GetContent());
+
+	FVector2D ViewLocation;
+	float ZoomAmount;
+	GraphEditor->GetViewLocation(ViewLocation, ZoomAmount);
+
+	UEdGraph* Graph = Payload->IsValid() ? FTabPayload_UObject::CastChecked<UEdGraph>(Payload) : nullptr;
+
+	OnSaveGraphStateDelegate.ExecuteIfBound(Graph, ViewLocation, ZoomAmount);
+}
+
+TAttribute<FText> FGraphEditorSummoner::ConstructTabNameForObject(UEdGraph* DocumentID) const
+{
+	if (DocumentID)
+	{
+		if (const UEdGraphSchema* Schema = DocumentID->GetSchema())
+		{
+			FGraphDisplayInfo Info;
+			Schema->GetGraphDisplayInformation(*DocumentID, /*out*/ Info);
+			return Info.DisplayName;
+		}
+		else
+		{
+			// if we don't have a schema, we're dealing with a malformed (or incomplete graph)...
+			// possibly in the midst of some transaction - here we return the object's outer path 
+			// so we can at least get some context as to which graph we're referring
+			return FText::FromString(DocumentID->GetPathName());
+		}
+	}
+
+	return LOCTEXT("UnknownGraphName", "Unknown");
+}
+
+TSharedRef<SWidget> FGraphEditorSummoner::CreateTabBodyForObject(const FWorkflowTabSpawnInfo& Info, UEdGraph* DocumentID) const
+{
+	check(Info.TabInfo.IsValid());
+	
+	if(OnCreateGraphEditorWidgetDelegate.IsBound())
+	{
+		return OnCreateGraphEditorWidgetDelegate.Execute(Info.TabInfo.ToSharedRef(), DocumentID);
+	}
+	else
+	{
+		return SNew(SSpacer);
+	}
+}
+
+const FSlateBrush* FGraphEditorSummoner::GetTabIconForObject(const FWorkflowTabSpawnInfo& Info, UEdGraph* DocumentID) const
+{
+	if(OnGetTabIconDelegate.IsBound())
+	{
+		return OnGetTabIconDelegate.Execute(DocumentID);
+	}
+	else
+	{
+		return FAppStyle::Get().GetBrush(TEXT("GraphEditor.EventGraph_16x"));
+	}
+}
+
+TSharedRef<FGenericTabHistory> FGraphEditorSummoner::CreateTabHistoryNode(TSharedPtr<FTabPayload> Payload)
+{
+	return MakeShared<FGraphTabHistory>(SharedThis(this), Payload);
+}
+
+class FDetailsTabSummoner : public FWorkflowTabFactory
+{
+public:
+	FDetailsTabSummoner(TSharedPtr<FAssetEditorToolkit> InHostingApp, FOnDetailsViewCreated InOnDetailsViewCreated)
+		: FWorkflowTabFactory(Tabs::Details, InHostingApp)
+	{
+		TabLabel = LOCTEXT("DetailsTabLabel", "Details");
+		TabIcon = FSlateIcon("EditorStyle", "LevelEditor.Tabs.Details");
+		ViewMenuDescription = LOCTEXT("DetailsTabMenuDescription", "Details");
+		ViewMenuTooltip = LOCTEXT("DetailsTabToolTip", "Shows the details tab for selected objects.");
+		bIsSingleton = true;
+
+		FDetailsViewArgs DetailsViewArgs;
+		DetailsViewArgs.NameAreaSettings = FDetailsViewArgs::HideNameArea;
+		DetailsViewArgs.bHideSelectionTip = true;
+
+		FPropertyEditorModule& PropertyEditorModule = FModuleManager::LoadModuleChecked<FPropertyEditorModule>("PropertyEditor");
+		DetailsView = PropertyEditorModule.CreateDetailView(DetailsViewArgs);
+		
+		InOnDetailsViewCreated.ExecuteIfBound(DetailsView.ToSharedRef());
+	}
+
+	TSharedPtr<IDetailsView> GetDetailsView() const
+	{
+		return DetailsView;
+	}
+
+private:
+	virtual TSharedRef<SWidget> CreateTabBody(const FWorkflowTabSpawnInfo& Info) const override
+	{
+		return DetailsView.ToSharedRef();
+	}
+	
+	virtual FText GetTabToolTipText(const FWorkflowTabSpawnInfo& Info) const override
+	{
+		return ViewMenuTooltip;
+	}
+	
+	TSharedPtr<IDetailsView> DetailsView;
+};
+
+FGraphEditorMode::FGraphEditorMode(TSharedRef<FWorkflowCentricApplication> InHostingApp)
+	: FApplicationMode(Modes::GraphEditor)
+	, HostingAppPtr(InHostingApp)
+{
+	HostingAppPtr = InHostingApp;
+
+	TSharedRef<FGraphEditor> GraphEditor = StaticCastSharedRef<FGraphEditor>(InHostingApp);
+	
+	TabFactories.RegisterFactory(MakeShared<FDetailsTabSummoner>(GraphEditor, FOnDetailsViewCreated::CreateSP(&GraphEditor.Get(), &FGraphEditor::HandleDetailsViewCreated)));
+
+	TabLayout = FTabManager::NewLayout("Standalone_AnimNextInterfaceGraphEditor_Layout_v1.0")
+		->AddArea
+		(
+			FTabManager::NewPrimaryArea()
+			->SetOrientation(Orient_Vertical)
+			->Split
+			(	
+				FTabManager::NewSplitter()
+				->SetSizeCoefficient(1.0f)
+				->SetOrientation(Orient_Horizontal)
+				->Split
+				(
+					FTabManager::NewStack()
+					->SetSizeCoefficient(0.8f)
+					->SetHideTabWell(false)
+					->AddTab(Tabs::Document, ETabState::ClosedTab)
+				)
+				->Split
+				(
+					FTabManager::NewStack()
+					->SetSizeCoefficient(0.2f)
+					->SetHideTabWell(false)
+					->AddTab(Tabs::Details, ETabState::OpenedTab)
+				)
+			)
+		);
+
+	GraphEditor->RegisterModeToolbarIfUnregistered(GetModeName());
+}
+
+void FGraphEditorMode::RegisterTabFactories(TSharedPtr<FTabManager> InTabManager)
+{
+	TSharedPtr<FWorkflowCentricApplication> HostingApp = HostingAppPtr.Pin();
+	HostingApp->PushTabFactories(TabFactories);
+
+	FApplicationMode::RegisterTabFactories(InTabManager);
+}
+
+void FGraphEditorMode::AddTabFactory(FCreateWorkflowTabFactory FactoryCreator)
+{
+	if (FactoryCreator.IsBound())
+	{
+		TabFactories.RegisterFactory(FactoryCreator.Execute(HostingAppPtr.Pin()));
+	}
+}
+
+void FGraphEditorMode::RemoveTabFactory(FName TabFactoryID)
+{
+	TabFactories.UnregisterFactory(TabFactoryID);
+}
+
+}
+
+#undef LOCTEXT_NAMESPACE
\ No newline at end of file
diff --git a/Engine/Plugins/Experimental/AnimNext/Source/AnimNextInterfaceGraphEditor/Private/AnimNextInterfaceGraphEditorMode.h b/Engine/Plugins/Experimental/AnimNext/Source/AnimNextInterfaceGraphEditor/Private/AnimNextInterfaceGraphEditorMode.h
new file mode 100644
--- /dev/null
+++ b/Engine/Plugins/Experimental/AnimNext/Source/AnimNextInterfaceGraphEditor/Private/AnimNextInterfaceGraphEditorMode.h
@@ -0,0 +1,96 @@
+﻿// Copyright Epic Games, Inc. All Rights Reserved.
+
+#pragma once
+
+#include "CoreMinimal.h"
+#include "EdGraph/EdGraph.h"
+#include "WorkflowOrientedApp/ApplicationMode.h"
+#include "WorkflowOrientedApp/WorkflowTabManager.h"
+#include "WorkflowOrientedApp/WorkflowUObjectDocuments.h"
+
+class FWorkflowCentricApplication;
+class SGraphEditor;
+
+namespace UE::AnimNext::InterfaceGraphEditor
+{
+
+DECLARE_DELEGATE_OneParam(FOnDetailsViewCreated, TSharedRef<IDetailsView>);
+
+struct FGraphEditorSummoner : public FDocumentTabFactoryForObjects<UEdGraph>
+{
+public:
+	// Delegate called when a graph editor widget is created
+	DECLARE_DELEGATE_RetVal_TwoParams(TSharedRef<SGraphEditor>, FOnCreateGraphEditorWidget, TSharedRef<FTabInfo>, UEdGraph*);
+
+	// Delegate called when a graph editor is focused
+	DECLARE_DELEGATE_OneParam(FOnGraphEditorFocused, TSharedRef<SGraphEditor>);
+
+	// Delegate called when a graph editor is backgrounded
+	DECLARE_DELEGATE_OneParam(FOnGraphEditorBackgrounded, TSharedRef<SGraphEditor>);
+
+	// Delegate called to save the state of a graph
+	DECLARE_DELEGATE_ThreeParams(FOnSaveGraphState, UEdGraph*, FVector2D, float);
+
+	// Delegate called to save the state of a graph
+	DECLARE_DELEGATE_RetVal_OneParam(const FSlateBrush*, FOnGetTabIcon, UEdGraph*);
+	
+	FGraphEditorSummoner(TSharedPtr<FWorkflowCentricApplication> InHostingApp);
+
+	FOnCreateGraphEditorWidget& OnCreateGraphEditorWidget() { return OnCreateGraphEditorWidgetDelegate; }
+	FOnGraphEditorFocused& OnGraphEditorFocused() { return OnGraphEditorFocusedDelegate; }
+	FOnGraphEditorBackgrounded& OnGraphEditorBackgrounded() { return OnGraphEditorBackgroundedDelegate; }
+	FOnSaveGraphState& OnSaveGraphState() { return OnSaveGraphStateDelegate; }
+	FOnGetTabIcon& OnGetTabIcon() { return OnGetTabIconDelegate; }
+
+private:
+	// FWorkflowTabFactory interface
+	virtual void OnTabActivated(TSharedPtr<SDockTab> Tab) const override;
+	virtual void OnTabBackgrounded(TSharedPtr<SDockTab> Tab) const override;
+	virtual void OnTabRefreshed(TSharedPtr<SDockTab> Tab) const override;
+	virtual void SaveState(TSharedPtr<SDockTab> Tab, TSharedPtr<FTabPayload> Payload) const override;
+	
+	// FDocumentTabFactoryForObjects interface
+	virtual TAttribute<FText> ConstructTabNameForObject(UEdGraph* DocumentID) const override;
+	virtual TSharedRef<SWidget> CreateTabBodyForObject(const FWorkflowTabSpawnInfo& Info, UEdGraph* DocumentID) const override;
+	virtual const FSlateBrush* GetTabIconForObject(const FWorkflowTabSpawnInfo& Info, UEdGraph* DocumentID) const override;
+	virtual TSharedRef<FGenericTabHistory> CreateTabHistoryNode(TSharedPtr<FTabPayload> Payload) override;
+
+	// The hosting app
+	TWeakPtr<FWorkflowCentricApplication> HostingAppPtr;
+
+	// Delegate called when a graph editor widget is created
+	FOnCreateGraphEditorWidget OnCreateGraphEditorWidgetDelegate;
+
+	// Delegate called when a graph editor is focused
+	FOnGraphEditorFocused OnGraphEditorFocusedDelegate;
+
+	// Delegate called when a graph editor is backgrounded
+	FOnGraphEditorBackgrounded OnGraphEditorBackgroundedDelegate;
+
+	// Delegate called to save the state of a graph
+	FOnSaveGraphState OnSaveGraphStateDelegate;
+
+	// Delegate called to get a tab icon for a graph
+	FOnGetTabIcon OnGetTabIconDelegate;
+};
+
+class FGraphEditorMode : public FApplicationMode
+{
+public:
+	FGraphEditorMode(TSharedRef<FWorkflowCentricApplication> InHostingApp);
+
+private:
+	// FApplicationMode interface
+	virtual void RegisterTabFactories(TSharedPtr<FTabManager> InTabManager) override;
+	virtual void AddTabFactory(FCreateWorkflowTabFactory FactoryCreator) override;
+	virtual void RemoveTabFactory(FName TabFactoryID) override;
+	
+	// The hosting app
+	TWeakPtr<FWorkflowCentricApplication> HostingAppPtr;
+
+	// The tab factories we support
+	FWorkflowAllowedTabSet TabFactories;
+};
+
+}
+
diff --git a/Engine/Plugins/Experimental/AnimNext/Source/AnimNextInterfaceGraphEditor/Private/AnimNextInterfaceGraphEditorModule.cpp b/Engine/Plugins/Experimental/AnimNext/Source/AnimNextInterfaceGraphEditor/Private/AnimNextInterfaceGraphEditorModule.cpp
new file mode 100644
--- /dev/null
+++ b/Engine/Plugins/Experimental/AnimNext/Source/AnimNextInterfaceGraphEditor/Private/AnimNextInterfaceGraphEditorModule.cpp
@@ -0,0 +1,42 @@
+// Copyright Epic Games, Inc. All Rights Reserved.
+
+#include "AnimNextInterfaceGraphEditorModule.h"
+#include "IAssetTools.h"
+#include "AssetTypeActions_AnimNextInterfaceGraph.h"
+#include "AnimNextInterfacePropertyTypeCustomization.h"
+
+namespace UE::AnimNext::InterfaceGraphEditor
+{
+
+void FModule::StartupModule()
+{
+	IAssetTools& AssetTools = FModuleManager::LoadModuleChecked<FAssetToolsModule>("AssetTools").Get();
+	AssetTypeActions_AnimNextInterfaceGraph = MakeShared<UE::AnimNext::InterfaceGraphEditor::FAssetTypeActions_AnimNextInterfaceGraph>();
+	AssetTools.RegisterAssetTypeActions(AssetTypeActions_AnimNextInterfaceGraph.ToSharedRef());
+
+	AnimNextInterfacePropertyTypeIdentifier = MakeShared<UE::AnimNext::InterfaceGraphEditor::FPropertyTypeIdentifier>();
+	FPropertyEditorModule& PropertyModule = FModuleManager::LoadModuleChecked<FPropertyEditorModule>("PropertyEditor");
+	PropertyModule.RegisterCustomPropertyTypeLayout(
+		FInterfaceProperty::StaticClass()->GetFName(),
+		FOnGetPropertyTypeCustomizationInstance::CreateLambda([] { return MakeShared<UE::AnimNext::InterfaceGraphEditor::FAnimNextInterfacePropertyTypeCustomization>(); }),
+		AnimNextInterfacePropertyTypeIdentifier);
+}
+
+void FModule::ShutdownModule()
+{
+	if(FModuleManager::Get().IsModuleLoaded("AssetTools"))
+	{
+		IAssetTools& AssetTools = FModuleManager::GetModuleChecked<FAssetToolsModule>("AssetTools").Get();
+		AssetTools.UnregisterAssetTypeActions(AssetTypeActions_AnimNextInterfaceGraph.ToSharedRef());
+	}
+	
+	if(FModuleManager::Get().IsModuleLoaded("PropertyEditor"))
+	{
+		FPropertyEditorModule& PropertyModule = FModuleManager::LoadModuleChecked<FPropertyEditorModule>("PropertyEditor");
+		PropertyModule.UnregisterCustomPropertyTypeLayout(FInterfaceProperty::StaticClass()->GetFName(), AnimNextInterfacePropertyTypeIdentifier);
+	}
+}
+
+}
+
+IMPLEMENT_MODULE(UE::AnimNext::InterfaceGraphEditor::FModule, AnimNextInterfaceGraphEditor);
\ No newline at end of file
diff --git a/Engine/Plugins/Experimental/AnimNext/Source/AnimNextInterfaceGraphEditor/Private/AnimNextInterfaceGraphEditorModule.h b/Engine/Plugins/Experimental/AnimNext/Source/AnimNextInterfaceGraphEditor/Private/AnimNextInterfaceGraphEditorModule.h
new file mode 100644
--- /dev/null
+++ b/Engine/Plugins/Experimental/AnimNext/Source/AnimNextInterfaceGraphEditor/Private/AnimNextInterfaceGraphEditorModule.h
@@ -0,0 +1,25 @@
+﻿// Copyright Epic Games, Inc. All Rights Reserved.
+
+#pragma once
+
+#include "CoreMinimal.h"
+#include "Modules/ModuleManager.h"
+
+namespace UE::AnimNext::InterfaceGraphEditor
+{
+
+class FAssetTypeActions_AnimNextInterfaceGraph;
+class FAnimNextInterfacePropertyTypeCustomization;
+class FPropertyTypeIdentifier;
+
+class FModule : public IModuleInterface
+{
+private:
+	virtual void StartupModule() override;
+	virtual void ShutdownModule() override;
+
+	TSharedPtr<UE::AnimNext::InterfaceGraphEditor::FAssetTypeActions_AnimNextInterfaceGraph> AssetTypeActions_AnimNextInterfaceGraph;
+	TSharedPtr<UE::AnimNext::InterfaceGraphEditor::FPropertyTypeIdentifier> AnimNextInterfacePropertyTypeIdentifier;
+};
+
+}
\ No newline at end of file
diff --git a/Engine/Plugins/Experimental/AnimNext/Source/AnimNextInterfaceGraphEditor/Private/AnimNextInterfaceGraphEditorSchemaActions.cpp b/Engine/Plugins/Experimental/AnimNext/Source/AnimNextInterfaceGraphEditor/Private/AnimNextInterfaceGraphEditorSchemaActions.cpp
new file mode 100644
--- /dev/null
+++ b/Engine/Plugins/Experimental/AnimNext/Source/AnimNextInterfaceGraphEditor/Private/AnimNextInterfaceGraphEditorSchemaActions.cpp
@@ -0,0 +1,61 @@
+﻿// Copyright Epic Games, Inc. All Rights Reserved.
+
+#include "AnimNextInterfaceGraphEditorSchemaActions.h"
+#include "AnimNextInterfaceEditorUtils.h"
+#include "AnimNextInterfaceGraph_EdGraph.h"
+#include "AnimNextInterfaceGraph_EdGraphNode.h"
+#include "SGraphActionMenu.h"
+#include "AnimNextInterfaceUncookedOnlyUtils.h"
+#include "AnimNextInterfaceGraph_EditorData.h"
+#include "RigVMModel/Nodes/RigVMUnitNode.h"
+#include "Settings/ControlRigSettings.h"
+
+UEdGraphNode* FAnimNextInterfaceGraphSchemaAction_RigUnit::PerformAction(UEdGraph* ParentGraph, TArray<UEdGraphPin*>& FromPins, const FVector2D Location, bool bSelectNewNode)
+{
+	UAnimNextInterfaceGraph_EditorData* EditorData = ParentGraph->GetTypedOuter<UAnimNextInterfaceGraph_EditorData>();
+	UAnimNextInterfaceGraph_EdGraphNode* NewNode = nullptr;
+	UAnimNextInterfaceGraph_EdGraph* EdGraph = Cast<UAnimNextInterfaceGraph_EdGraph>(ParentGraph);
+
+	if (EditorData != nullptr && EdGraph != nullptr)
+	{
+		FName Name = UE::AnimNext::InterfaceGraphEditor::FUtils::ValidateName(EditorData, StructTemplate->GetFName().ToString());
+		URigVMController* Controller = EditorData->GetRigVMClient()->GetController(ParentGraph);
+
+		Controller->OpenUndoBracket(FString::Printf(TEXT("Add '%s' Node"), *Name.ToString()));
+
+		FRigVMUnitNodeCreatedContext& UnitNodeCreatedContext = Controller->GetUnitNodeCreatedContext();
+		FRigVMUnitNodeCreatedContext::FScope ReasonScope(UnitNodeCreatedContext, ERigVMNodeCreatedReason::NodeSpawner);
+
+		if (URigVMUnitNode* ModelNode = Controller->AddUnitNode(StructTemplate, FRigUnit::GetMethodName(), Location, Name.ToString(), true, false))
+		{
+			NewNode = Cast<UAnimNextInterfaceGraph_EdGraphNode>(EdGraph->FindNodeForModelNodeName(ModelNode->GetFName()));
+			check(NewNode);
+
+			if (NewNode)
+			{
+				Controller->ClearNodeSelection(true);
+				Controller->SelectNode(ModelNode, true, true);
+			}
+
+#if WITH_EDITORONLY_DATA
+
+			// @TODO: switch over to custom settings here
+			const FControlRigSettingsPerPinBool* ExpansionMapPtr = UControlRigEditorSettings::Get()->RigUnitPinExpansion.Find(ModelNode->GetScriptStruct()->GetName());
+			if (ExpansionMapPtr)
+			{
+				const FControlRigSettingsPerPinBool& ExpansionMap = *ExpansionMapPtr;
+
+				for (const TPair<FString, bool>& Pair : ExpansionMap.Values)
+				{
+					FString PinPath = FString::Printf(TEXT("%s.%s"), *ModelNode->GetName(), *Pair.Key);
+					Controller->SetPinExpansion(PinPath, Pair.Value, true);
+				}
+			}
+		}
+#endif
+		
+		Controller->CloseUndoBracket();
+	}
+
+	return NewNode;
+}
diff --git a/Engine/Plugins/Experimental/AnimNext/Source/AnimNextInterfaceGraphEditor/Private/AnimNextInterfaceGraphEditorSchemaActions.h b/Engine/Plugins/Experimental/AnimNext/Source/AnimNextInterfaceGraphEditor/Private/AnimNextInterfaceGraphEditorSchemaActions.h
new file mode 100644
--- /dev/null
+++ b/Engine/Plugins/Experimental/AnimNext/Source/AnimNextInterfaceGraphEditor/Private/AnimNextInterfaceGraphEditorSchemaActions.h
@@ -0,0 +1,55 @@
+﻿// Copyright Epic Games, Inc. All Rights Reserved.
+
+#pragma once
+
+#include "CoreMinimal.h"
+#include "EdGraph/EdGraphSchema.h"
+#include "Styling/AppStyle.h"
+#include "AnimNextInterfaceGraphEditorSchemaActions.generated.h"
+
+struct FSlateBrush;
+
+USTRUCT()
+struct FAnimNextInterfaceGraphSchemaAction : public FEdGraphSchemaAction
+{
+	GENERATED_BODY()
+
+	FAnimNextInterfaceGraphSchemaAction() = default;
+	
+	FAnimNextInterfaceGraphSchemaAction(FText InNodeCategory, FText InMenuDesc, FText InToolTip, FText InKeywords = FText::GetEmpty())
+		: FEdGraphSchemaAction(MoveTemp(InNodeCategory), MoveTemp(InMenuDesc), MoveTemp(InToolTip), 0, MoveTemp(InKeywords))
+	{
+	}
+
+	virtual const FSlateBrush* GetIconBrush() const
+	{
+		return FAppStyle::Get().GetBrush("NoBrush");
+	}
+
+	virtual const FLinearColor& GetIconColor() const
+	{
+		static const FLinearColor DefaultColor;
+		return DefaultColor;
+	}
+};
+
+USTRUCT()
+struct FAnimNextInterfaceGraphSchemaAction_RigUnit : public FAnimNextInterfaceGraphSchemaAction
+{
+	GENERATED_BODY()
+
+	FAnimNextInterfaceGraphSchemaAction_RigUnit() = default;
+	
+	FAnimNextInterfaceGraphSchemaAction_RigUnit(UScriptStruct* InStructTemplate, FText InNodeCategory, FText InMenuDesc, FText InToolTip, FText InKeywords = FText::GetEmpty())
+		: FAnimNextInterfaceGraphSchemaAction(InNodeCategory, InMenuDesc, InToolTip, InKeywords)
+		, StructTemplate(InStructTemplate)
+	{}
+
+	// FEdGraphSchemaAction Interface
+	virtual UEdGraphNode* PerformAction(UEdGraph* ParentGraph, TArray<UEdGraphPin*>& FromPins, const FVector2D Location, bool bSelectNewNode = true) override;
+	virtual UEdGraphNode* PerformAction(UEdGraph* ParentGraph, UEdGraphPin* FromPin, const FVector2D Location, bool bSelectNewNode = true) { return nullptr; }
+
+private:
+	// The script struct for our rig unit
+	UScriptStruct* StructTemplate = nullptr;
+};
diff --git a/Engine/Plugins/Experimental/AnimNext/Source/AnimNextInterfaceGraphEditor/Private/AnimNextInterfaceGraphFactory.cpp b/Engine/Plugins/Experimental/AnimNext/Source/AnimNextInterfaceGraphEditor/Private/AnimNextInterfaceGraphFactory.cpp
new file mode 100644
--- /dev/null
+++ b/Engine/Plugins/Experimental/AnimNext/Source/AnimNextInterfaceGraphEditor/Private/AnimNextInterfaceGraphFactory.cpp
@@ -0,0 +1,99 @@
+﻿// Copyright Epic Games, Inc. All Rights Reserved.
+
+#include "AnimNextInterfaceGraphFactory.h"
+#include "AnimNextInterfaceGraph.h"
+#include "AnimNextInterfaceGraph_EditorData.h"
+#include "AnimNextInterfaceGraph_EdGraph.h"
+#include "AnimNextInterfaceGraph_EdGraphSchema.h"
+#include "AnimNextInterfaceUncookedOnlyUtils.h"
+#include "RigUnit_AnimNextInterfaceBeginExecution.h"
+#include "RigUnit_AnimNextInterfaceEndExecution.h"
+#include "Kismet2/BlueprintEditorUtils.h"
+#include "RigVMModel/RigVMController.h"
+#include "Units/RigUnit.h"
+
+UAnimNextInterfaceGraphFactory::UAnimNextInterfaceGraphFactory()
+{
+	bCreateNew = true;
+	bEditAfterNew = true;
+	SupportedClass = UAnimNextInterfaceGraph::StaticClass();
+}
+
+bool UAnimNextInterfaceGraphFactory::ConfigureProperties()
+{
+	return true;
+}
+
+UObject* UAnimNextInterfaceGraphFactory::FactoryCreateNew(UClass* Class, UObject* InParent, FName Name, EObjectFlags Flags, UObject* Context, FFeedbackContext* Warn, FName CallingContext)
+{
+	UAnimNextInterfaceGraph* NewGraph = NewObject<UAnimNextInterfaceGraph>(InParent, Class, Name, Flags | RF_Public | RF_Standalone | RF_Transactional | RF_LoadCompleted);
+
+	// Create internal editor data & graphs
+	UAnimNextInterfaceGraph_EditorData* EditorData = NewObject<UAnimNextInterfaceGraph_EditorData>(NewGraph, TEXT("EditorData"));
+	NewGraph->EditorData = EditorData;
+	EditorData->Initialize(/*bRecompileVM*/false);
+	EditorData->GetRigVMClient()->SetExecuteContextStruct(FAnimNextInterfaceExecuteContext::StaticStruct());
+
+	// Add initial execution unit
+	URigVMController* Controller = EditorData->GetRigVMClient()->GetController(EditorData->RootGraph);
+
+	URigVMUnitNode* MainEntryPointNode = Controller->AddUnitNode(FRigUnit_AnimNextInterfaceBeginExecution::StaticStruct(), FRigUnit::GetMethodName(), FVector2D(-400.0f, 0.0f), FString(), false);
+	URigVMPin* BeginExecutePin = MainEntryPointNode->FindPin(GET_MEMBER_NAME_STRING_CHECKED(FRigUnit_AnimNextInterfaceBeginExecution, ExecuteContext));
+	check(BeginExecutePin);
+	check(BeginExecutePin->GetDirection() == ERigVMPinDirection::Output);
+
+	// Add function to function lib
+	{
+		FRigVMControllerGraphGuard LibraryGraphGuard(Controller, EditorData->GetRigVMClient()->GetFunctionLibrary());
+		EditorData->EntryPoint = Controller->AddFunctionToLibrary(UE::AnimNext::InterfaceGraph::EntryPointName, true, FVector2D::ZeroVector, false, false);
+
+		// Add exposed pin for result to function
+		{
+			FRigVMControllerGraphGuard FunctionGraphGuard(Controller, EditorData->EntryPoint->GetContainedGraph());
+
+			// TODO: using float for now, but needs to be user-driven
+			const FString CPPType = TEXT("float");
+			const FName CPPTypeObjectPath = NAME_None;
+			const FString DefaultValue = TEXT("0.0");
+
+			//FunctionController->AddExposedPin(UE::AnimNext::InterfaceGraph::ResultName, ERigVMPinDirection::Output, CPPType, CPPTypeObjectPath, DefaultValue, false, false);
+			Controller->AddExposedPin(UE::AnimNext::InterfaceGraph::ResultName, ERigVMPinDirection::Output, CPPType, CPPTypeObjectPath, DefaultValue, false, false);
+		}
+	}
+
+	// Add function call
+	URigVMFunctionReferenceNode* FunctionCallNode = Controller->AddFunctionReferenceNode(EditorData->EntryPoint, FVector2D::ZeroVector, UE::AnimNext::InterfaceGraph::EntryPointName.ToString(), false, false);
+	URigVMPin* FunctionExecutePin = FunctionCallNode->FindPin(FRigVMStruct::ExecuteContextName.ToString());
+	check(FunctionExecutePin);
+	check(FunctionExecutePin->GetDirection() == ERigVMPinDirection::IO);
+
+	// Link entry point to function call
+	Controller->AddLink(BeginExecutePin, FunctionExecutePin, false);
+
+	URigVMPin* FunctionResultPin = FunctionCallNode->FindPin(UE::AnimNext::InterfaceGraph::ResultName.ToString());
+	check(FunctionResultPin);
+	check(FunctionResultPin->GetDirection() == ERigVMPinDirection::Output);
+
+	// Add end-execution unit of the correct type
+	// TODO: using a float for now, but need to use a registry to determine correct type
+	URigVMUnitNode* MainExitPointNode = Controller->AddUnitNode(FRigUnit_AnimNextInterfaceEndExecution_Float::StaticStruct(), FRigUnit::GetMethodName(), FVector2D(400.0f, 0.0f), FString(), false);
+	URigVMPin* EndExecutePin = MainExitPointNode->FindPin(GET_MEMBER_NAME_STRING_CHECKED(FRigUnit_AnimNextInterfaceEndExecution, ExecuteContext));
+	check(EndExecutePin);
+	check(EndExecutePin->GetDirection() == ERigVMPinDirection::Input);
+
+	URigVMPin* EndExecuteResultPin = MainExitPointNode->FindPin(GET_MEMBER_NAME_STRING_CHECKED(FRigUnit_AnimNextInterfaceEndExecution_Float, Result));
+	check(EndExecuteResultPin);
+	check(EndExecuteResultPin->GetDirection() == ERigVMPinDirection::Input);
+
+	// Link function call to exit point 
+	Controller->AddLink(FunctionExecutePin, EndExecutePin, false);
+
+	// Link result pins 
+	Controller->AddLink(FunctionResultPin, EndExecuteResultPin, false);
+
+	// Compile the initial skeleton
+	UE::AnimNext::InterfaceGraphUncookedOnly::FUtils::Compile(NewGraph);
+	check(!EditorData->bErrorsDuringCompilation);
+	
+	return NewGraph;
+}
\ No newline at end of file
diff --git a/Engine/Plugins/Experimental/AnimNext/Source/AnimNextInterfaceGraphEditor/Private/AnimNextInterfaceGraphFactory.h b/Engine/Plugins/Experimental/AnimNext/Source/AnimNextInterfaceGraphEditor/Private/AnimNextInterfaceGraphFactory.h
new file mode 100644
--- /dev/null
+++ b/Engine/Plugins/Experimental/AnimNext/Source/AnimNextInterfaceGraphEditor/Private/AnimNextInterfaceGraphFactory.h
@@ -0,0 +1,22 @@
+﻿// Copyright Epic Games, Inc. All Rights Reserved.
+
+#pragma once
+
+#include "CoreMinimal.h"
+#include "Factories/Factory.h"
+#include "AnimNextInterfaceGraphFactory.generated.h"
+
+UCLASS(MinimalAPI)
+class UAnimNextInterfaceGraphFactory : public UFactory
+{
+	GENERATED_BODY()
+
+	UAnimNextInterfaceGraphFactory();
+	
+	virtual bool ConfigureProperties() override;
+	virtual UObject* FactoryCreateNew(UClass* Class, UObject* InParent, FName Name, EObjectFlags Flags, UObject* Context, FFeedbackContext* Warn, FName CallingContext) override;
+	virtual UObject* FactoryCreateNew(UClass* Class, UObject* InParent, FName Name, EObjectFlags Flags, UObject* Context, FFeedbackContext* Warn) override
+	{
+		return FactoryCreateNew(Class, InParent, Name, Flags, Context, Warn, NAME_None);
+	}
+};
\ No newline at end of file
diff --git a/Engine/Plugins/Experimental/AnimNext/Source/AnimNextInterfaceGraphEditor/Private/AnimNextInterfacePropertyTypeCustomization.cpp b/Engine/Plugins/Experimental/AnimNext/Source/AnimNextInterfaceGraphEditor/Private/AnimNextInterfacePropertyTypeCustomization.cpp
new file mode 100644
--- /dev/null
+++ b/Engine/Plugins/Experimental/AnimNext/Source/AnimNextInterfaceGraphEditor/Private/AnimNextInterfacePropertyTypeCustomization.cpp
@@ -0,0 +1,88 @@
+﻿// Copyright Epic Games, Inc. All Rights Reserved.
+
+#include "AnimNextInterfacePropertyTypeCustomization.h"
+
+#include "AnimNextInterfaceGraph.h"
+#include "AnimNextInterfaceWidgetFactories.h"
+#include "DetailWidgetRow.h"
+#include "Editor/PropertyEditor/Private/ObjectPropertyNode.h"
+#include "Editor/PropertyEditor/Private/PropertyNode.h"
+#include "IAnimNextInterface.h"
+#include "IDetailChildrenBuilder.h"
+#include "PropertyCustomizationHelpers.h"
+#include "PropertyHandle.h"
+#include "SClassViewer.h"
+#include "Widgets/Text/STextBlock.h"
+
+#define LOCTEXT_NAMESPACE "AnimNextInterfacePropertyTypeCustomization"
+
+namespace UE::AnimNext::InterfaceGraphEditor
+{
+
+bool FPropertyTypeIdentifier::IsPropertyTypeCustomized(const IPropertyHandle& PropertyHandle) const
+{
+	return (PropertyHandle.GetMetaData("AnimNextInterfaceType") != "");
+}
+
+void FAnimNextInterfacePropertyTypeCustomization::CustomizeHeader(TSharedRef<IPropertyHandle> PropertyHandle, FDetailWidgetRow& HeaderRow, IPropertyTypeCustomizationUtils& CustomizationUtils)
+{
+	FString TypeName = PropertyHandle->GetMetaData("AnimNextInterfaceType");
+	FName AnimNextInterfaceTypeName(TypeName);
+
+	void* ValuePtr;
+	PropertyHandle->GetValueData(ValuePtr);
+	FScriptInterface* PropertyValue = reinterpret_cast<FScriptInterface*>(ValuePtr);
+	
+	TSharedPtr<SWidget> Widget = FAnimNextInterfaceWidgetFactories::CreateAnimNextInterfaceWidget(AnimNextInterfaceTypeName, PropertyValue->GetObject(),
+			FOnClassPicked::CreateLambda([this, PropertyHandle](UClass* ChosenClass)
+			{
+				TArray<void*> RawData;
+				PropertyHandle->AccessRawData(RawData);
+				TArray<UPackage*> OuterObjects;
+				PropertyHandle->GetOuterPackages(OuterObjects);
+
+				for(int i=0;i<RawData.Num();i++)
+				{
+					UObject* NewValue = NewObject<UObject>(OuterObjects[i], ChosenClass);
+				
+					PropertyHandle->NotifyPreChange();
+					PropertyHandle->SetValue(NewValue);
+					PropertyHandle->NotifyPostChange(EPropertyChangeType::ValueSet);
+					PropertyHandle->GetPropertyNode()->GetParentNode()->RequestRebuildChildren();
+				}
+			}),
+			nullptr
+	);
+
+	HeaderRow
+	.NameContent()
+	[
+		PropertyHandle->CreatePropertyNameWidget()
+	]
+	.ValueContent()
+	[
+		Widget.ToSharedRef()
+	];
+}
+	
+void FAnimNextInterfacePropertyTypeCustomization::CustomizeChildren(TSharedRef<IPropertyHandle> PropertyHandle, IDetailChildrenBuilder& ChildBuilder, IPropertyTypeCustomizationUtils& CustomizationUtils)
+{
+	// todo: make this work correctly for multiselect
+	void* ValuePtr;
+	PropertyHandle->GetValueData(ValuePtr);
+	FScriptInterface* PropertyValue = reinterpret_cast<FScriptInterface*>(ValuePtr);
+	if (PropertyValue->GetObject())
+	{
+		TSharedPtr<FStructOnScope> StructData = MakeShareable(new FStructOnScope(PropertyValue->GetObject()->GetClass(), (uint8*)PropertyValue->GetObject()));
+		StructData->SetPackage(PropertyValue->GetObject()->GetPackage());
+		TArray<TSharedPtr<IPropertyHandle>> NewHandles = ChildBuilder.AddAllExternalStructureProperties(StructData.ToSharedRef());
+		for (auto& NewHandle : NewHandles)
+		{
+			ChildBuilder.AddProperty(NewHandle.ToSharedRef());
+		}
+	}
+}
+
+}
+
+#undef LOCTEXT_NAMESPACE
\ No newline at end of file
diff --git a/Engine/Plugins/Experimental/AnimNext/Source/AnimNextInterfaceGraphEditor/Private/AnimNextInterfacePropertyTypeCustomization.h b/Engine/Plugins/Experimental/AnimNext/Source/AnimNextInterfaceGraphEditor/Private/AnimNextInterfacePropertyTypeCustomization.h
new file mode 100644
--- /dev/null
+++ b/Engine/Plugins/Experimental/AnimNext/Source/AnimNextInterfaceGraphEditor/Private/AnimNextInterfacePropertyTypeCustomization.h
@@ -0,0 +1,24 @@
+﻿// Copyright Epic Games, Inc. All Rights Reserved.
+
+#pragma once
+
+#include "IPropertyTypeCustomization.h"
+#include "PropertyEditorModule.h"
+
+namespace UE::AnimNext::InterfaceGraphEditor
+{
+
+class FAnimNextInterfacePropertyTypeCustomization : public IPropertyTypeCustomization
+{
+private:
+	// IPropertyTypeCustomization interface
+	virtual void CustomizeHeader(TSharedRef<IPropertyHandle> PropertyHandle, FDetailWidgetRow& HeaderRow, IPropertyTypeCustomizationUtils& CustomizationUtils) override;
+	virtual void CustomizeChildren(TSharedRef<IPropertyHandle> PropertyHandle, IDetailChildrenBuilder& ChildBuilder, IPropertyTypeCustomizationUtils& CustomizationUtils) override;
+};
+
+class FPropertyTypeIdentifier : public IPropertyTypeIdentifier
+{
+	virtual bool IsPropertyTypeCustomized(const IPropertyHandle& PropertyHandle) const override;
+};
+
+}
\ No newline at end of file
diff --git a/Engine/Plugins/Experimental/AnimNext/Source/AnimNextInterfaceGraphEditor/Private/AnimNextInterfaceWidgetFactories.cpp b/Engine/Plugins/Experimental/AnimNext/Source/AnimNextInterfaceGraphEditor/Private/AnimNextInterfaceWidgetFactories.cpp
new file mode 100644
--- /dev/null
+++ b/Engine/Plugins/Experimental/AnimNext/Source/AnimNextInterfaceGraphEditor/Private/AnimNextInterfaceWidgetFactories.cpp
@@ -0,0 +1,122 @@
+// Copyright Epic Games, Inc. All Rights Reserved.
+
+#include "AnimNextInterfaceWidgetFactories.h"
+#include "Modules/ModuleManager.h"
+#include "Widgets/Input/SHyperlink.h"
+#include "Widgets/Input/SNumericEntryBox.h"
+#include "Widgets/Input/SComboButton.h"
+#include "SClassViewer.h"
+#include "DetailCategoryBuilder.h"
+#include "Widgets/Layout/SWidgetSwitcher.h"
+#include "AnimNextInterfaceClassFilter.h"
+
+#define LOCTEXT_NAMESPACE "AnimNextInterfaceEditor"
+
+namespace UE::AnimNext::InterfaceGraphEditor
+{
+
+void ConvertToText(UObject* Object, FText& OutText)
+{
+	UClass* Class = Object->GetClass();
+	while (Class)
+	{
+		if (auto TextConverter = FAnimNextInterfaceWidgetFactories::AnimNextInterfaceTextConverter.Find(Class))
+		{
+			(*TextConverter)(Object, OutText);
+			break;
+		}
+		Class = Class->GetSuperClass();
+	}
+}
+
+TSharedPtr<SWidget> FAnimNextInterfaceWidgetFactories::CreateAnimNextInterfaceWidget(FName AnimNextInterfaceTypeName, UObject* Value, const FOnClassPicked& CreateClassCallback, TSharedPtr<SBorder>* InnerWidget)
+{
+	TSharedPtr<SWidget> LeftWidget;
+
+	if (Value)
+	{
+		UClass* Class = Value->GetClass();
+		while (Class && !LeftWidget.IsValid())
+		{
+			if (auto Creator = FAnimNextInterfaceWidgetFactories::AnimNextInterfaceWidgetCreators.Find(Class))
+			{
+				LeftWidget = (*Creator)(Value);
+				break;
+			}
+			Class = Class->GetSuperClass();
+		}
+	}
+	
+	if (!LeftWidget.IsValid())
+	{
+		LeftWidget = SNew(STextBlock).Text(LOCTEXT("SelectDataType", "Select Data Type..."));
+	}
+
+
+	// button for replacing data with a different Data Interface class
+	TSharedPtr<SComboButton> Button = SNew(SComboButton)
+			.ComboButtonStyle(FAppStyle::Get(), "SimpleComboButton");
+	
+	Button->SetOnGetMenuContent(FOnGetContent::CreateLambda([AnimNextInterfaceTypeName, Button, CreateClassCallback]()
+	{
+		FClassViewerInitializationOptions Options;
+		Options.ClassFilters.Add(MakeShared<FAnimNextInterfaceClassFilter>(AnimNextInterfaceTypeName, nullptr));
+		
+		// Add class filter for columns here
+		TSharedRef<SWidget>  ClassMenu = FModuleManager::LoadModuleChecked<FClassViewerModule>("ClassViewer").
+		CreateClassViewer(Options, FOnClassPicked::CreateLambda([Button, CreateClassCallback](UClass* Class)
+		{
+			Button->SetIsOpen(false);
+			CreateClassCallback.Execute(Class);
+		}));
+		return ClassMenu;
+	}));
+
+	TSharedPtr <SBorder> Border;
+	if (InnerWidget && InnerWidget->IsValid())
+	{
+		Border = *InnerWidget;
+	}
+	else
+	{
+		Border = SNew(SBorder);
+	}
+	
+	if (InnerWidget)
+	{
+		*InnerWidget = Border;
+	}
+	
+	Border->SetContent(LeftWidget.ToSharedRef());
+
+	TSharedPtr<SWidget> Widget = SNew(SHorizontalBox)
+		+SHorizontalBox::Slot().FillWidth(75)
+		[
+			Border.ToSharedRef()
+		]
+		+SHorizontalBox::Slot().FillWidth(25)
+		[
+			Button.ToSharedRef()
+		]
+	;
+
+
+	return Widget;
+}
+
+TMap<const UClass*, TFunction<void (const UObject* Object, FText& OutText)>> FAnimNextInterfaceWidgetFactories::AnimNextInterfaceTextConverter;
+TMap<const UClass*, TFunction<TSharedRef<SWidget> (UObject* Object)>> FAnimNextInterfaceWidgetFactories::AnimNextInterfaceWidgetCreators;
+
+void ConvertToText_Base(const UObject* Object, FText& OutText)
+{
+	OutText = FText::FromString(Object->GetName());
+}
+
+void FAnimNextInterfaceWidgetFactories::RegisterWidgets()
+{
+	AnimNextInterfaceTextConverter.Add(UObject::StaticClass(), ConvertToText_Base);
+}
+	
+}
+
+#undef LOCTEXT_NAMESPACE
diff --git a/Engine/Plugins/Experimental/AnimNext/Source/AnimNextInterfaceGraphEditor/Private/AssetTypeActions_AnimNextInterfaceGraph.cpp b/Engine/Plugins/Experimental/AnimNext/Source/AnimNextInterfaceGraphEditor/Private/AssetTypeActions_AnimNextInterfaceGraph.cpp
new file mode 100644
--- /dev/null
+++ b/Engine/Plugins/Experimental/AnimNext/Source/AnimNextInterfaceGraphEditor/Private/AssetTypeActions_AnimNextInterfaceGraph.cpp
@@ -0,0 +1,37 @@
+﻿// Copyright Epic Games, Inc. All Rights Reserved.
+
+#include "AssetTypeActions_AnimNextInterfaceGraph.h"
+#include "AnimNextInterfaceGraphEditor.h"
+#include "AnimNextInterfaceGraph.h"
+
+#define LOCTEXT_NAMESPACE "AssetTypeActions_AnimNextInterfaceGraph" 
+
+namespace UE::AnimNext::InterfaceGraphEditor
+{
+
+void FAssetTypeActions_AnimNextInterfaceGraph::OpenAssetEditor(const TArray<UObject*>& InObjects, TSharedPtr<IToolkitHost> EditWithinLevelEditor)
+{
+	const EToolkitMode::Type Mode = EditWithinLevelEditor.IsValid() ? EToolkitMode::WorldCentric : EToolkitMode::Standalone;
+
+	for (UObject* Object : InObjects)
+	{
+		if(UAnimNextInterfaceGraph* AnimNextInterfaceGraph = Cast<UAnimNextInterfaceGraph>(Object))
+		{
+			TSharedRef<FGraphEditor> GraphEditor = MakeShared<FGraphEditor>();
+			GraphEditor->InitEditor(Mode, EditWithinLevelEditor, AnimNextInterfaceGraph);
+		}
+	}
+}
+
+const TArray<FText>& FAssetTypeActions_AnimNextInterfaceGraph::GetSubMenus() const
+{
+	static const TArray<FText> SubMenus
+	{
+		LOCTEXT("AnimNextInterfacesSubMenu", "Data Interfaces")
+	};
+	return SubMenus;
+}
+
+}
+
+#undef LOCTEXT_NAMESPACE
\ No newline at end of file
diff --git a/Engine/Plugins/Experimental/AnimNext/Source/AnimNextInterfaceGraphEditor/Private/AssetTypeActions_AnimNextInterfaceGraph.h b/Engine/Plugins/Experimental/AnimNext/Source/AnimNextInterfaceGraphEditor/Private/AssetTypeActions_AnimNextInterfaceGraph.h
new file mode 100644
--- /dev/null
+++ b/Engine/Plugins/Experimental/AnimNext/Source/AnimNextInterfaceGraphEditor/Private/AssetTypeActions_AnimNextInterfaceGraph.h
@@ -0,0 +1,24 @@
+﻿// Copyright Epic Games, Inc. All Rights Reserved.
+
+#pragma once
+
+#include "AssetTypeActions_Base.h"
+#include "AnimNextInterfaceGraph.h"
+
+namespace UE::AnimNext::InterfaceGraphEditor
+{
+
+class FAssetTypeActions_AnimNextInterfaceGraph : public FAssetTypeActions_Base
+{
+public:
+	// IAssetTypeActions Implementation
+	virtual FText GetName() const override { return NSLOCTEXT("AssetTypeActions", "AssetTypeActions_AnimNextInterfaceGraph", "Data Interface Graph"); }
+	virtual FColor GetTypeColor() const override { return FColor(128,128,64); }
+	virtual UClass* GetSupportedClass() const override { return UAnimNextInterfaceGraph::StaticClass(); }
+	virtual void OpenAssetEditor(const TArray<UObject*>& InObjects, TSharedPtr<IToolkitHost> EditWithinLevelEditor = TSharedPtr<IToolkitHost>()) override;
+	virtual bool CanFilter() override { return true; }
+	virtual const TArray<FText>& GetSubMenus() const override;
+	virtual uint32 GetCategories() override { return EAssetTypeCategories::Animation; }
+};
+
+}
diff --git a/Engine/Plugins/Experimental/AnimNext/Source/AnimNextInterfaceGraphEditor/Private/SAnimNextInterfaceGraphEditorActionMenu.cpp b/Engine/Plugins/Experimental/AnimNext/Source/AnimNextInterfaceGraphEditor/Private/SAnimNextInterfaceGraphEditorActionMenu.cpp
new file mode 100644
--- /dev/null
+++ b/Engine/Plugins/Experimental/AnimNext/Source/AnimNextInterfaceGraphEditor/Private/SAnimNextInterfaceGraphEditorActionMenu.cpp
@@ -0,0 +1,228 @@
+﻿// Copyright Epic Games, Inc. All Rights Reserved.
+
+#include "SAnimNextInterfaceGraphEditorActionMenu.h"
+
+#include "AnimNextInterfaceGraphEditorSchemaActions.h"
+#include "Framework/Application/SlateApplication.h"
+#include "IDocumentation.h"
+#include "SGraphActionMenu.h"
+#include "SSubobjectEditor.h"
+#include "AnimNextInterfaceGraph_EdGraphSchema.h"
+#include "Widgets/Images/SImage.h"
+#include "Widgets/Layout/SBox.h"
+#include "Widgets/SBoxPanel.h"
+#include "Widgets/Text/STextBlock.h"
+#include "SGraphPalette.h"
+#include "RigVMCore/RigVMRegistry.h"
+#include "Units/RigUnit.h"
+#include "Widgets/SToolTip.h"
+#include "AnimNextInterfaceExecuteContext.h"
+
+#define LOCTEXT_NAMESPACE "AnimNextInterfaceEditor"
+
+namespace UE::AnimNext::InterfaceGraphEditor
+{
+
+static void CollectAllAnimNextInterfaceActions(FGraphContextMenuBuilder& MenuBuilder)
+{
+	static const TArray<UScriptStruct*> AllowedExecuteContexts =
+	{
+		FRigVMExecuteContext::StaticStruct(),
+		FAnimNextInterfaceExecuteContext::StaticStruct()
+	};
+
+	const TChunkedArray<FRigVMFunction>& Functions = FRigVMRegistry::Get().GetFunctions();
+	const int32 NumFunctions = Functions.Num();
+	for(int i=0; i < NumFunctions; i++)
+	{
+		const FRigVMFunction& Function = Functions[i];
+
+		const UScriptStruct* FunctionContext = Function.GetExecuteContextStruct();
+		if (FunctionContext == nullptr || !AllowedExecuteContexts.Contains(FunctionContext))
+		{
+			continue;
+		}
+
+		UScriptStruct* Struct = Function.Struct;
+		if (Struct == nullptr)
+		{
+			continue;
+		}
+
+		FString CategoryMetadata, DisplayNameMetadata, MenuDescSuffixMetadata;
+		Struct->GetStringMetaDataHierarchical(FRigVMStruct::CategoryMetaName, &CategoryMetadata);
+		Struct->GetStringMetaDataHierarchical(FRigVMStruct::DisplayNameMetaName, &DisplayNameMetadata);
+		Struct->GetStringMetaDataHierarchical(FRigVMStruct::MenuDescSuffixMetaName, &MenuDescSuffixMetadata);
+		if (!MenuDescSuffixMetadata.IsEmpty())
+		{
+			MenuDescSuffixMetadata = TEXT(" ") + MenuDescSuffixMetadata;
+		}
+		const FText NodeCategory = FText::FromString(CategoryMetadata);
+		const FText MenuDesc = FText::FromString(DisplayNameMetadata + MenuDescSuffixMetadata);
+		const FText ToolTip = Struct->GetToolTipText();
+
+		if (MenuDesc.IsEmpty())
+		{
+			continue;
+		}
+
+		MenuBuilder.AddAction(MakeShared<FAnimNextInterfaceGraphSchemaAction_RigUnit>(Struct, NodeCategory, MenuDesc, ToolTip));
+	};
+}
+
+SActionMenu::~SActionMenu()
+{
+	OnClosedCallback.ExecuteIfBound();
+	OnCloseReasonCallback.ExecuteIfBound(bActionExecuted, false, !DraggedFromPins.IsEmpty());
+}
+
+void SActionMenu::Construct(const FArguments& InArgs)
+{
+	Graph = InArgs._Graph;
+	DraggedFromPins = InArgs._DraggedFromPins;
+	NewNodePosition = InArgs._NewNodePosition;
+	OnClosedCallback = InArgs._OnClosedCallback;
+	bAutoExpandActionMenu = InArgs._AutoExpandActionMenu;
+	OnCloseReasonCallback = InArgs._OnCloseReason;
+
+	SBorder::Construct(SBorder::FArguments()
+		.BorderImage(FAppStyle::Get().GetBrush("Menu.Background"))
+		.Padding(5)
+		[
+			SNew(SBox)
+			.WidthOverride(400)
+			.HeightOverride(400)
+			[
+				SNew(SVerticalBox)
+				+SVerticalBox::Slot()
+				.AutoHeight()
+				.Padding(2, 2, 2, 5)
+				[
+					SNew(SHorizontalBox)
+					+SHorizontalBox::Slot()
+					.AutoWidth()
+					.VAlign(VAlign_Center)
+					[
+						SNew(STextBlock)
+						.Text(LOCTEXT("ContextText", "All Data Interface Node Classes"))
+						.Font(FAppStyle::Get().GetFontStyle("BlueprintEditor.ActionMenu.ContextDescriptionFont"))
+						.ToolTip(IDocumentation::Get()->CreateToolTip(
+							LOCTEXT("ActionMenuContextTextTooltip", "Describes the current context of the action list"),
+							nullptr,
+							TEXT("Shared/Editors/AnimNextInterfaceEditor"),
+							TEXT("AnimNextInterfaceActionMenuContextText")))
+						.WrapTextAt(280)
+					]
+				]
+				+SVerticalBox::Slot()
+				[
+					SAssignNew(GraphActionMenu, SGraphActionMenu)
+						.OnActionSelected(this, &SActionMenu::OnActionSelected)
+						.OnCreateWidgetForAction(SGraphActionMenu::FOnCreateWidgetForAction::CreateSP(this, &SActionMenu::OnCreateWidgetForAction))
+						.OnCollectAllActions(this, &SActionMenu::CollectAllActions)
+						.DraggedFromPins(DraggedFromPins)
+						.GraphObj(Graph)
+						.AlphaSortItems(true)
+						.bAllowPreselectedItemActivation(true)
+				]
+			]
+		]
+	);
+}
+
+void SActionMenu::CollectAllActions(FGraphActionListBuilderBase& OutAllActions)
+{
+	if (!Graph)
+	{
+		return;
+	}
+
+	FGraphContextMenuBuilder MenuBuilder(Graph);
+	if (!DraggedFromPins.IsEmpty())
+	{
+		MenuBuilder.FromPin = DraggedFromPins[0];
+	}
+
+	// Cannot call GetGraphContextActions() during serialization and GC due to its use of FindObject()
+	if(!GIsSavingPackage && !IsGarbageCollecting())
+	{
+		CollectAllAnimNextInterfaceActions(MenuBuilder);
+	}
+
+	OutAllActions.Append(MenuBuilder);
+}
+
+TSharedRef<SEditableTextBox> SActionMenu::GetFilterTextBox()
+{
+	return GraphActionMenu->GetFilterTextBox();
+}
+
+TSharedRef<SWidget> SActionMenu::OnCreateWidgetForAction(FCreateWidgetForActionData* const InCreateData)
+{
+	check(InCreateData);
+	InCreateData->bHandleMouseButtonDown = false;
+
+	const FSlateBrush* IconBrush = nullptr;
+	FLinearColor IconColor;
+	TSharedPtr<FAnimNextInterfaceGraphSchemaAction> Action = StaticCastSharedPtr<FAnimNextInterfaceGraphSchemaAction>(InCreateData->Action);
+	if (Action.IsValid())
+	{
+		IconBrush = Action->GetIconBrush();
+		IconColor = Action->GetIconColor();
+	}
+
+	TSharedPtr<SHorizontalBox> WidgetBox = SNew(SHorizontalBox);
+	if (IconBrush)
+	{
+		WidgetBox->AddSlot()
+			.AutoWidth()
+			.VAlign(VAlign_Center)
+			.Padding(5, 0, 0, 0)
+			[
+				SNew(SImage)
+				.ColorAndOpacity(IconColor)
+				.Image(IconBrush)
+			];
+	}
+
+	WidgetBox->AddSlot()
+		.AutoWidth()
+		.VAlign(VAlign_Center)
+		.Padding(0, 0, 0, 0)
+		[
+			SNew(SGraphPaletteItem, InCreateData)
+		];
+
+	return WidgetBox->AsShared();
+}
+
+void SActionMenu::OnActionSelected(const TArray<TSharedPtr<FEdGraphSchemaAction>>& SelectedAction, ESelectInfo::Type InSelectionType)
+{
+	if (!Graph)
+	{
+		return;
+	}
+
+	if (InSelectionType != ESelectInfo::OnMouseClick  && InSelectionType != ESelectInfo::OnKeyPress && !SelectedAction.IsEmpty())
+	{
+		return;
+	}
+
+	for (const TSharedPtr<FEdGraphSchemaAction>& Action : SelectedAction)
+	{
+		if (Action.IsValid() && Graph)
+		{
+			if (!bActionExecuted && (Action->GetTypeId() != FEdGraphSchemaAction_Dummy::StaticGetTypeId()))
+			{
+				FSlateApplication::Get().DismissAllMenus();
+				bActionExecuted = true;
+			}
+
+			Action->PerformAction(Graph, DraggedFromPins, NewNodePosition);
+		}
+	}
+}
+
+} 
+
+#undef LOCTEXT_NAMESPACE
diff --git a/Engine/Plugins/Experimental/AnimNext/Source/AnimNextInterfaceGraphEditor/Private/SAnimNextInterfaceGraphEditorActionMenu.h b/Engine/Plugins/Experimental/AnimNext/Source/AnimNextInterfaceGraphEditor/Private/SAnimNextInterfaceGraphEditorActionMenu.h
new file mode 100644
--- /dev/null
+++ b/Engine/Plugins/Experimental/AnimNext/Source/AnimNextInterfaceGraphEditor/Private/SAnimNextInterfaceGraphEditorActionMenu.h
@@ -0,0 +1,72 @@
+﻿// Copyright Epic Games, Inc. All Rights Reserved.
+#pragma once
+
+#include "GraphEditor.h"
+#include "EdGraph/EdGraph.h"
+#include "EdGraph/EdGraphPin.h"
+#include "EdGraph/EdGraphSchema.h"
+#include "Layout/Margin.h"
+#include "Layout/Visibility.h"
+#include "SGraphActionMenu.h"
+#include "Widgets/DeclarativeSyntaxSupport.h"
+#include "Widgets/Input/SEditableTextBox.h"
+#include "Widgets/Layout/SBorder.h"
+#include "Widgets/SWidget.h"
+#include "Widgets/Views/SExpanderArrow.h"
+
+class SEditableTextBox;
+class SGraphActionMenu;
+class UEdGraph;
+
+namespace UE::AnimNext::InterfaceGraphEditor
+{
+
+class SActionMenu : public SBorder
+{
+public:
+	/** Delegate for the OnCloseReason event which is always raised when the SActionMenu closes */
+	DECLARE_DELEGATE_ThreeParams(FClosedReason, bool /*bActionExecuted*/, bool /*bContextSensitiveChecked*/, bool /*bGraphPinContext*/);
+
+	SLATE_BEGIN_ARGS(SActionMenu)
+		: _Graph(static_cast<UEdGraph*>(nullptr))
+		, _NewNodePosition(FVector2D::ZeroVector)
+		, _AutoExpandActionMenu(false)
+	{
+	}
+		SLATE_ARGUMENT(UEdGraph*, Graph)
+		SLATE_ARGUMENT(FVector2D, NewNodePosition)
+		SLATE_ARGUMENT(TArray<UEdGraphPin*>, DraggedFromPins)
+		SLATE_ARGUMENT(SGraphEditor::FActionMenuClosed, OnClosedCallback)
+		SLATE_ARGUMENT(bool, AutoExpandActionMenu)
+		SLATE_EVENT(FClosedReason, OnCloseReason)
+	SLATE_END_ARGS()
+
+	void Construct(const FArguments& InArgs);
+
+	~SActionMenu();
+
+	TSharedRef<SEditableTextBox> GetFilterTextBox();
+
+protected:
+	void OnActionSelected(const TArray<TSharedPtr<FEdGraphSchemaAction>>& SelectedAction, ESelectInfo::Type InSelectionType);
+
+	TSharedRef<SWidget> OnCreateWidgetForAction(struct FCreateWidgetForActionData* const InCreateData);
+
+	/** Callback used to populate all actions list in SGraphActionMenu */
+	void CollectAllActions(FGraphActionListBuilderBase& OutAllActions);
+
+private:
+	UEdGraph* Graph = nullptr;
+	bool bAutoExpandActionMenu = false;
+	bool bActionExecuted = false;
+
+	TArray<UEdGraphPin*> DraggedFromPins;
+	FVector2D NewNodePosition = FVector2D::ZeroVector;
+
+	SGraphEditor::FActionMenuClosed OnClosedCallback;
+	FClosedReason OnCloseReasonCallback;
+
+	TSharedPtr<SGraphActionMenu> GraphActionMenu;
+};
+
+}
diff --git a/Engine/Plugins/Experimental/AnimNext/Source/AnimNextInterfaceGraphEditor/Public/AnimNextInterfaceClassFilter.h b/Engine/Plugins/Experimental/AnimNext/Source/AnimNextInterfaceGraphEditor/Public/AnimNextInterfaceClassFilter.h
new file mode 100644
--- /dev/null
+++ b/Engine/Plugins/Experimental/AnimNext/Source/AnimNextInterfaceGraphEditor/Public/AnimNextInterfaceClassFilter.h
@@ -0,0 +1,50 @@
+// Copyright Epic Games, Inc. All Rights Reserved.
+#pragma once
+
+#include "ClassViewerFilter.h"
+#include "IAnimNextInterface.h"
+
+namespace UE::AnimNext::InterfaceGraphEditor
+{
+	// Filter class for things which implement IChooserColumn
+	class FAnimNextInterfaceClassFilter : public IClassViewerFilter
+	{
+	public:
+		FAnimNextInterfaceClassFilter(FName InTypeName, const UScriptStruct* InStructType) : TypeName(InTypeName), StructType(InStructType) { };
+		virtual ~FAnimNextInterfaceClassFilter() override {};
+	
+		virtual bool IsClassAllowed(const FClassViewerInitializationOptions& InInitOptions, const UClass* InClass, TSharedRef< class FClassViewerFilterFuncs > InFilterFuncs ) override
+		{
+			if (InClass->ImplementsInterface(UAnimNextInterface::StaticClass()))
+			{
+				// pass filter if the class has "" return type (for wrappers)
+				// pass all classes if the requested typename is ""
+				// otherwise fail if the TypeNames don't match
+				// for Object type, check type hierarchy. (todo)
+			
+				IAnimNextInterface* Object = static_cast<IAnimNextInterface*>(InClass->GetDefaultObject()->GetInterfaceAddress(UAnimNextInterface::StaticClass()));
+				if (Object->GetReturnTypeName() == TypeName || Object->GetReturnTypeName() == "" || TypeName == "")
+				{
+					if (const UScriptStruct* ReturnStructType = Object->GetReturnTypeStruct())
+					{
+						return true; // kytodo: compare the struct types to see if they are compatible
+					}
+					else
+					{
+						return true;
+					}
+				}
+			}
+
+			return false;
+		}
+
+		virtual bool IsUnloadedClassAllowed(const FClassViewerInitializationOptions& InInitOptions, const TSharedRef< const class IUnloadedBlueprintData > InUnloadedClassData, TSharedRef< class FClassViewerFilterFuncs > InFilterFuncs) override
+		{
+			return false;
+		}
+	private:
+		FName TypeName;
+		const UScriptStruct* StructType;
+	};
+}
\ No newline at end of file
diff --git a/Engine/Plugins/Experimental/AnimNext/Source/AnimNextInterfaceGraphEditor/Public/AnimNextInterfaceWidgetFactories.h b/Engine/Plugins/Experimental/AnimNext/Source/AnimNextInterfaceGraphEditor/Public/AnimNextInterfaceWidgetFactories.h
new file mode 100644
--- /dev/null
+++ b/Engine/Plugins/Experimental/AnimNext/Source/AnimNextInterfaceGraphEditor/Public/AnimNextInterfaceWidgetFactories.h
@@ -0,0 +1,22 @@
+// Copyright Epic Games, Inc. All Rights Reserved.
+
+#pragma once
+
+#include "CoreMinimal.h"
+#include "Widgets/SWidget.h"
+#include "Widgets/Layout/SBorder.h"
+#include "SClassViewer.h"
+
+namespace UE::AnimNext::InterfaceGraphEditor
+{
+	class ANIMNEXTINTERFACEGRAPHEDITOR_API FAnimNextInterfaceWidgetFactories 
+	{
+	public:
+		static TMap<const UClass*, TFunction<void (const UObject* Object, FText& OutText)>> AnimNextInterfaceTextConverter;
+		static TMap<const UClass*, TFunction<TSharedRef<SWidget> (UObject* Object)>> AnimNextInterfaceWidgetCreators;
+
+		static TSharedPtr<SWidget> CreateAnimNextInterfaceWidget(FName AnimNextInterfaceTypeName, UObject* Value, const FOnClassPicked& CreateClassCallback, TSharedPtr<SBorder>* InnerWidget = nullptr);
+		
+		static void RegisterWidgets();
+	};
+}
diff --git a/Engine/Plugins/Experimental/AnimNext/Source/AnimNextInterfaceGraphUncookedOnly/AnimNextInterfaceGraphUncookedOnly.Build.cs b/Engine/Plugins/Experimental/AnimNext/Source/AnimNextInterfaceGraphUncookedOnly/AnimNextInterfaceGraphUncookedOnly.Build.cs
new file mode 100644
--- /dev/null
+++ b/Engine/Plugins/Experimental/AnimNext/Source/AnimNextInterfaceGraphUncookedOnly/AnimNextInterfaceGraphUncookedOnly.Build.cs
@@ -0,0 +1,44 @@
+// Copyright Epic Games, Inc. All Rights Reserved.
+
+namespace UnrealBuildTool.Rules
+{
+	public class AnimNextInterfaceGraphUncookedOnly : ModuleRules
+	{
+		public AnimNextInterfaceGraphUncookedOnly(ReadOnlyTargetRules Target) : base(Target)
+		{
+			
+			PublicDependencyModuleNames.AddRange(
+				new string[]
+				{
+					"RigVM",
+					"RigVMDeveloper",
+					"ControlRig",
+					"ControlRigDeveloper",
+				}
+			);
+
+			PrivateDependencyModuleNames.AddRange(
+				new string[]
+				{
+					"Core",
+					"CoreUObject",
+					"Engine",
+					"AnimNextInterface",
+					"AnimNextInterfaceGraph",
+					"BlueprintGraph",	// For K2 schema
+				}
+			);
+
+
+			if (Target.bBuildEditor)
+			{
+				PrivateDependencyModuleNames.AddRange(
+					new string[]
+					{
+						"UnrealEd",
+					}
+				);
+			}
+		}
+	}
+}
\ No newline at end of file
diff --git a/Engine/Plugins/Experimental/AnimNext/Source/AnimNextInterfaceGraphUncookedOnly/Internal/AnimNextInterfaceGraph_EdGraph.h b/Engine/Plugins/Experimental/AnimNext/Source/AnimNextInterfaceGraphUncookedOnly/Internal/AnimNextInterfaceGraph_EdGraph.h
new file mode 100644
--- /dev/null
+++ b/Engine/Plugins/Experimental/AnimNext/Source/AnimNextInterfaceGraphUncookedOnly/Internal/AnimNextInterfaceGraph_EdGraph.h
@@ -0,0 +1,27 @@
+﻿// Copyright Epic Games, Inc. All Rights Reserved.
+
+#pragma once
+
+#include "CoreMinimal.h"
+#include "Graph/ControlRigGraph.h"
+#include "AnimNextInterfaceGraph_EdGraph.generated.h"
+
+class UAnimNextInterfaceGraph_EditorData;
+
+namespace UE::AnimNext::InterfaceGraphUncookedOnly
+{
+	struct FUtils;
+}
+
+UCLASS(MinimalAPI)
+class UAnimNextInterfaceGraph_EdGraph : public UControlRigGraph
+{
+	GENERATED_BODY()
+
+	friend class UAnimNextInterfaceGraph_EditorData;
+
+	// UControlRigGraph interface
+	virtual FRigVMClient* GetRigVMClient() const override;
+	
+	void Initialize(UAnimNextInterfaceGraph_EditorData* InEditorData);
+};
\ No newline at end of file
diff --git a/Engine/Plugins/Experimental/AnimNext/Source/AnimNextInterfaceGraphUncookedOnly/Internal/AnimNextInterfaceGraph_EdGraphNode.h b/Engine/Plugins/Experimental/AnimNext/Source/AnimNextInterfaceGraphUncookedOnly/Internal/AnimNextInterfaceGraph_EdGraphNode.h
new file mode 100644
--- /dev/null
+++ b/Engine/Plugins/Experimental/AnimNext/Source/AnimNextInterfaceGraphUncookedOnly/Internal/AnimNextInterfaceGraph_EdGraphNode.h
@@ -0,0 +1,13 @@
+﻿// Copyright Epic Games, Inc. All Rights Reserved.
+
+#pragma once
+
+#include "CoreMinimal.h"
+#include "Graph/ControlRigGraphNode.h"
+#include "AnimNextInterfaceGraph_EdGraphNode.generated.h"
+
+UCLASS(MinimalAPI)
+class UAnimNextInterfaceGraph_EdGraphNode : public UControlRigGraphNode
+{
+	GENERATED_BODY()
+};
\ No newline at end of file
diff --git a/Engine/Plugins/Experimental/AnimNext/Source/AnimNextInterfaceGraphUncookedOnly/Internal/AnimNextInterfaceGraph_EdGraphSchema.h b/Engine/Plugins/Experimental/AnimNext/Source/AnimNextInterfaceGraphUncookedOnly/Internal/AnimNextInterfaceGraph_EdGraphSchema.h
new file mode 100644
--- /dev/null
+++ b/Engine/Plugins/Experimental/AnimNext/Source/AnimNextInterfaceGraphUncookedOnly/Internal/AnimNextInterfaceGraph_EdGraphSchema.h
@@ -0,0 +1,20 @@
+﻿// Copyright Epic Games, Inc. All Rights Reserved.
+
+#pragma once
+
+#include "CoreMinimal.h"
+#include "Graph/ControlRigGraphSchema.h"
+#include "AnimNextInterfaceGraph_EdGraphNode.h"
+#include "AnimNextInterfaceGraph_EdGraphSchema.generated.h"
+
+UCLASS(MinimalAPI)
+class UAnimNextInterfaceGraph_EdGraphSchema : public UControlRigGraphSchema
+{
+	GENERATED_BODY()
+
+	// UControlRigGraphSchema interface
+	virtual TSubclassOf<UControlRigGraphNode> GetGraphNodeClass() const override { return UAnimNextInterfaceGraph_EdGraphNode::StaticClass(); }
+
+	// UEdGraphSchema interface
+	virtual void TrySetDefaultValue(UEdGraphPin& InPin, const FString& InNewDefaultValue, bool bMarkAsModified) const override;
+};
\ No newline at end of file
diff --git a/Engine/Plugins/Experimental/AnimNext/Source/AnimNextInterfaceGraphUncookedOnly/Internal/AnimNextInterfaceGraph_EditorData.h b/Engine/Plugins/Experimental/AnimNext/Source/AnimNextInterfaceGraphUncookedOnly/Internal/AnimNextInterfaceGraph_EditorData.h
new file mode 100644
--- /dev/null
+++ b/Engine/Plugins/Experimental/AnimNext/Source/AnimNextInterfaceGraphUncookedOnly/Internal/AnimNextInterfaceGraph_EditorData.h
@@ -0,0 +1,129 @@
+﻿// Copyright Epic Games, Inc. All Rights Reserved.
+
+#pragma once
+
+#include "CoreMinimal.h"
+#include "RigVMModel/RigVMGraph.h"
+#include "AnimNextInterfaceGraph_EdGraph.h"
+#include "RigVMCore/RigVMGraphFunctionHost.h"
+#include "AnimNextInterfaceGraph_EditorData.generated.h"
+
+class UAnimNextInterfaceGraph;
+enum class ERigVMGraphNotifType : uint8;
+namespace UE::AnimNext::InterfaceGraphUncookedOnly
+{
+	struct FUtils;
+}
+
+namespace UE::AnimNext::InterfaceGraphEditor
+{
+	class FGraphEditor;
+}
+
+UCLASS(MinimalAPI)
+class UAnimNextInterfaceGraph_EditorData : public UObject, public IRigVMClientHost, public IRigVMGraphFunctionHost
+{
+	GENERATED_BODY()
+
+	UAnimNextInterfaceGraph_EditorData(const FObjectInitializer& ObjectInitializer);
+	
+	friend class UAnimNextInterfaceGraphFactory;
+	friend class UAnimNextInterfaceGraph_EdGraph;
+	friend struct UE::AnimNext::InterfaceGraphUncookedOnly::FUtils;
+	friend class UE::AnimNext::InterfaceGraphEditor::FGraphEditor;
+	friend struct FAnimNextInterfaceGraphSchemaAction_RigUnit;
+
+	// UObject interface
+	virtual void Serialize(FArchive& Ar) override;
+	virtual void PostLoad() override;
+	virtual bool IsEditorOnly() const override { return true; }
+	
+	// IRigVMClientHost interface
+	virtual FRigVMClient* GetRigVMClient() override;
+	virtual const FRigVMClient* GetRigVMClient() const override;
+	virtual IRigVMGraphFunctionHost* GetRigVMGraphFunctionHost() override;
+	virtual const IRigVMGraphFunctionHost* GetRigVMGraphFunctionHost() const override;
+	virtual void HandleRigVMGraphAdded(const FRigVMClient* InClient, const FString& InNodePath) override;
+	virtual void HandleRigVMGraphRemoved(const FRigVMClient* InClient, const FString& InNodePath) override {}
+	virtual void HandleRigVMGraphRenamed(const FRigVMClient* InClient, const FString& InOldNodePath, const FString& InNewNodePath) override {}
+	virtual void HandleConfigureRigVMController(const FRigVMClient* InClient, URigVMController* InControllerToConfigure) override;
+	virtual UObject* GetEditorObjectForRigVMGraph(URigVMGraph* InVMGraph) const override;
+
+	// IRigVMGraphFunctionHost interface
+	virtual FRigVMGraphFunctionStore* GetRigVMGraphFunctionStore() override;
+	virtual const FRigVMGraphFunctionStore* GetRigVMGraphFunctionStore() const override;
+
+	ANIMNEXTINTERFACEGRAPHUNCOOKEDONLY_API void Initialize(bool bRecompileVM);
+
+	void RecompileVM();
+	
+	void RecompileVMIfRequired();
+
+	void RequestAutoVMRecompilation();
+	
+	void HandleModifiedEvent(ERigVMGraphNotifType InNotifType, URigVMGraph* InGraph, UObject* InSubject);
+
+	ANIMNEXTINTERFACEGRAPHUNCOOKEDONLY_API URigVMGraph* GetVMGraphForEdGraph(const UEdGraph* InGraph) const;
+
+	void CreateEdGraphForCollapseNode(URigVMCollapseNode* InNode);
+	
+	UPROPERTY()
+	TObjectPtr<UAnimNextInterfaceGraph_EdGraph> RootGraph;
+
+	UPROPERTY()
+	TObjectPtr<UAnimNextInterfaceGraph_EdGraph> EntryPointGraph;
+	
+	UPROPERTY()
+	TObjectPtr<UAnimNextInterfaceGraph_EdGraph> FunctionLibraryEdGraph;
+
+	UPROPERTY()
+	FRigVMClient RigVMClient;
+
+	UPROPERTY()
+	FRigVMGraphFunctionStore GraphFunctionStore;
+
+	UPROPERTY()
+	TObjectPtr<URigVMGraph> RigVMGraph_DEPRECATED;
+	
+	UPROPERTY()
+	TObjectPtr<URigVMFunctionLibrary> RigVMFunctionLibrary_DEPRECATED;
+
+	UPROPERTY()
+	TObjectPtr<URigVMLibraryNode> EntryPoint;
+	
+	UPROPERTY(transient)
+	TMap<TObjectPtr<URigVMGraph>, TObjectPtr<URigVMController>> Controllers;
+
+	UPROPERTY(EditAnywhere, Category = "User Interface")
+	FRigGraphDisplaySettings RigGraphDisplaySettings;
+
+	UPROPERTY(EditAnywhere, Category = "VM")
+	FRigVMRuntimeSettings VMRuntimeSettings;
+	
+	UPROPERTY(EditAnywhere, BlueprintReadWrite, Category = "VM", meta = (AllowPrivateAccess = "true"))
+	FRigVMCompileSettings VMCompileSettings;
+
+	UPROPERTY(EditAnywhere, Category = "Python Log Settings")
+	FControlRigPythonSettings PythonLogSettings;
+
+	UPROPERTY(transient, DuplicateTransient)
+	TMap<FString, FRigVMOperand> PinToOperandMap;
+
+	UPROPERTY(transient, DuplicateTransient)
+	bool bVMRecompilationRequired = false;
+
+	UPROPERTY(transient, DuplicateTransient)
+	bool bIsCompiling = false;
+	
+	FCompilerResultsLog CompileLog;
+
+	FOnVMCompiledEvent VMCompiledEvent;
+	FRigVMGraphModifiedEvent ModifiedEvent;
+
+	bool bAutoRecompileVM = true;
+	bool bErrorsDuringCompilation = false;
+	bool bSuspendModelNotificationsForSelf = false;
+	bool bSuspendModelNotificationsForOthers = false;
+	bool bSuspendAllNotifications = false;
+	bool bCompileInDebugMode = false;
+};
\ No newline at end of file
diff --git a/Engine/Plugins/Experimental/AnimNext/Source/AnimNextInterfaceGraphUncookedOnly/Internal/AnimNextInterfaceUncookedOnlyUtils.h b/Engine/Plugins/Experimental/AnimNext/Source/AnimNextInterfaceGraphUncookedOnly/Internal/AnimNextInterfaceUncookedOnlyUtils.h
new file mode 100644
--- /dev/null
+++ b/Engine/Plugins/Experimental/AnimNext/Source/AnimNextInterfaceGraphUncookedOnly/Internal/AnimNextInterfaceUncookedOnlyUtils.h
@@ -0,0 +1,30 @@
+﻿// Copyright Epic Games, Inc. All Rights Reserved.
+
+#pragma once
+
+#include "CoreMinimal.h"
+
+class UAnimNextInterfaceGraph;
+class UAnimNextInterfaceGraph_EditorData;
+class UAnimNextInterfaceGraph_EdGraph;
+class URigVMController;
+class URigVMGraph;
+
+namespace UE::AnimNext::InterfaceGraphEditor
+{
+	class FGraphEditor;
+}
+
+namespace UE::AnimNext::InterfaceGraphUncookedOnly
+{
+
+struct ANIMNEXTINTERFACEGRAPHUNCOOKEDONLY_API FUtils
+{
+	static void Compile(UAnimNextInterfaceGraph* InGraph);
+	
+	static UAnimNextInterfaceGraph_EditorData* GetEditorData(const UAnimNextInterfaceGraph* InAnimNextInterfaceGraph);
+	
+	static void RecreateVM(UAnimNextInterfaceGraph* InGraph);
+};
+
+}
\ No newline at end of file
diff --git a/Engine/Plugins/Experimental/AnimNext/Source/AnimNextInterfaceGraphUncookedOnly/Private/AnimNextInterfaceGraphUncookedOnlyModule.cpp b/Engine/Plugins/Experimental/AnimNext/Source/AnimNextInterfaceGraphUncookedOnly/Private/AnimNextInterfaceGraphUncookedOnlyModule.cpp
new file mode 100644
--- /dev/null
+++ b/Engine/Plugins/Experimental/AnimNext/Source/AnimNextInterfaceGraphUncookedOnly/Private/AnimNextInterfaceGraphUncookedOnlyModule.cpp
@@ -0,0 +1,5 @@
+// Copyright Epic Games, Inc. All Rights Reserved.
+
+#include "Modules/ModuleManager.h"
+
+IMPLEMENT_MODULE(FDefaultModuleImpl, AnimNextInterfaceGraphUncookedOnly);
\ No newline at end of file
diff --git a/Engine/Plugins/Experimental/AnimNext/Source/AnimNextInterfaceGraphUncookedOnly/Private/AnimNextInterfaceGraphUncookedOnlyUtils.cpp b/Engine/Plugins/Experimental/AnimNext/Source/AnimNextInterfaceGraphUncookedOnly/Private/AnimNextInterfaceGraphUncookedOnlyUtils.cpp
new file mode 100644
--- /dev/null
+++ b/Engine/Plugins/Experimental/AnimNext/Source/AnimNextInterfaceGraphUncookedOnly/Private/AnimNextInterfaceGraphUncookedOnlyUtils.cpp
@@ -0,0 +1,99 @@
+﻿// Copyright Epic Games, Inc. All Rights Reserved.
+
+#include "AnimNextInterfaceUncookedOnlyUtils.h"
+#include "AnimNextInterfaceGraph.h"
+#include "AnimNextInterfaceGraph_EditorData.h"
+#include "AnimNextInterfaceGraph_EdGraph.h"
+#include "RigVMCompiler/RigVMCompiler.h"
+#include "RigVMCore/RigVM.h"
+#include "Units/RigUnitContext.h"
+
+namespace UE::AnimNext::InterfaceGraphUncookedOnly
+{
+
+void FUtils::Compile(UAnimNextInterfaceGraph* InGraph)
+{
+	check(InGraph);
+	
+	UAnimNextInterfaceGraph_EditorData* EditorData = GetEditorData(InGraph);
+	
+	if(EditorData->bIsCompiling)
+	{
+		return;
+	}
+	
+	TGuardValue<bool> CompilingGuard(EditorData->bIsCompiling, true);
+	
+	EditorData->bErrorsDuringCompilation = false;
+
+	EditorData->RigGraphDisplaySettings.MinMicroSeconds = EditorData->RigGraphDisplaySettings.LastMinMicroSeconds = DBL_MAX;
+	EditorData->RigGraphDisplaySettings.MaxMicroSeconds = EditorData->RigGraphDisplaySettings.LastMaxMicroSeconds = (double)INDEX_NONE;
+	
+	TGuardValue<bool> ReentrantGuardSelf(EditorData->bSuspendModelNotificationsForSelf, true);
+	TGuardValue<bool> ReentrantGuardOthers(EditorData->bSuspendModelNotificationsForOthers, true);
+
+	RecreateVM(InGraph);
+
+	InGraph->VMRuntimeSettings = EditorData->VMRuntimeSettings;
+
+	FRigUnitContext InitContext;
+	InitContext.State = EControlRigState::Init;
+
+	FRigUnitContext UpdateContext = InitContext;
+	UpdateContext.State = EControlRigState::Update;
+
+	void* InitContextPtr = &InitContext;
+	void* UpdateContextPtr = &UpdateContext;
+
+	EditorData->CompileLog.Messages.Reset();
+	EditorData->CompileLog.NumErrors = EditorData->CompileLog.NumWarnings = 0;
+
+	URigVMCompiler* Compiler = URigVMCompiler::StaticClass()->GetDefaultObject<URigVMCompiler>();
+	EditorData->VMCompileSettings.SetExecuteContextStruct(EditorData->RigVMClient.GetExecuteContextStruct());
+	Compiler->Settings = (EditorData->bCompileInDebugMode) ? FRigVMCompileSettings::Fast(EditorData->VMCompileSettings.GetExecuteContextStruct()) : EditorData->VMCompileSettings;
+	URigVMController* RootController = EditorData->GetRigVMClient()->GetOrCreateController(EditorData->GetRigVMClient()->GetDefaultModel());
+	Compiler->Compile(EditorData->GetRigVMClient()->GetAllModels(false, false), RootController, InGraph->RigVM, InGraph->GetRigVMExternalVariables(), &EditorData->PinToOperandMap);
+
+	if (EditorData->bErrorsDuringCompilation)
+	{
+		if(Compiler->Settings.SurpressErrors)
+		{
+			Compiler->Settings.Reportf(EMessageSeverity::Info, InGraph,TEXT("Compilation Errors may be suppressed for Data Interface Graph: %s. See VM Compile Settings for more Details"), *InGraph->GetName());
+		}
+	}
+
+	EditorData->bVMRecompilationRequired = false;
+	if(InGraph->RigVM)
+	{
+		EditorData->VMCompiledEvent.Broadcast(InGraph, InGraph->RigVM);
+	}
+
+#if WITH_EDITOR
+//	RefreshBreakpoints(EditorData);
+#endif
+}
+
+void FUtils::RecreateVM(UAnimNextInterfaceGraph* InGraph)
+{
+	InGraph->RigVM = NewObject<URigVM>(InGraph, TEXT("VM"), RF_NoFlags);
+	
+	// Cooked platforms will load these pointers from disk
+	if (!FPlatformProperties::RequiresCookedData())
+	{
+		// We dont support ERigVMMemoryType::Work memory as we dont operate on an instance
+	//	InGraph->RigVM->GetMemoryByType(ERigVMMemoryType::Work, true);
+		InGraph->RigVM->GetMemoryByType(ERigVMMemoryType::Literal, true);
+		InGraph->RigVM->GetMemoryByType(ERigVMMemoryType::Debug, true);
+	}
+
+	InGraph->RigVM->Reset();
+}
+
+UAnimNextInterfaceGraph_EditorData* FUtils::GetEditorData(const UAnimNextInterfaceGraph* InAnimNextInterfaceGraph)
+{
+	check(InAnimNextInterfaceGraph);
+	
+	return CastChecked<UAnimNextInterfaceGraph_EditorData>(InAnimNextInterfaceGraph->EditorData);
+}
+
+}
diff --git a/Engine/Plugins/Experimental/AnimNext/Source/AnimNextInterfaceGraphUncookedOnly/Private/AnimNextInterfaceGraph_EdGraph.cpp b/Engine/Plugins/Experimental/AnimNext/Source/AnimNextInterfaceGraphUncookedOnly/Private/AnimNextInterfaceGraph_EdGraph.cpp
new file mode 100644
--- /dev/null
+++ b/Engine/Plugins/Experimental/AnimNext/Source/AnimNextInterfaceGraphUncookedOnly/Private/AnimNextInterfaceGraph_EdGraph.cpp
@@ -0,0 +1,22 @@
+﻿// Copyright Epic Games, Inc. All Rights Reserved.
+
+#include "AnimNextInterfaceGraph_EdGraph.h"
+#include "AnimNextInterfaceGraph_EditorData.h"
+#include "AnimNextInterfaceUncookedOnlyUtils.h"
+
+void UAnimNextInterfaceGraph_EdGraph::Initialize(UAnimNextInterfaceGraph_EditorData* InEditorData)
+{
+	InEditorData->ModifiedEvent.RemoveAll(this);
+	InEditorData->ModifiedEvent.AddUObject(this, &UAnimNextInterfaceGraph_EdGraph::HandleModifiedEvent);
+	InEditorData->VMCompiledEvent.RemoveAll(this);
+	InEditorData->VMCompiledEvent.AddUObject(this, &UAnimNextInterfaceGraph_EdGraph::HandleVMCompiledEvent);
+}
+
+FRigVMClient* UAnimNextInterfaceGraph_EdGraph::GetRigVMClient() const
+{
+	if (const UAnimNextInterfaceGraph_EditorData* EditorData = GetTypedOuter<UAnimNextInterfaceGraph_EditorData>())
+	{
+		return (FRigVMClient*)EditorData->GetRigVMClient();
+	}
+	return nullptr;
+}
diff --git a/Engine/Plugins/Experimental/AnimNext/Source/AnimNextInterfaceGraphUncookedOnly/Private/AnimNextInterfaceGraph_EdGraphSchema.cpp b/Engine/Plugins/Experimental/AnimNext/Source/AnimNextInterfaceGraphUncookedOnly/Private/AnimNextInterfaceGraph_EdGraphSchema.cpp
new file mode 100644
--- /dev/null
+++ b/Engine/Plugins/Experimental/AnimNext/Source/AnimNextInterfaceGraphUncookedOnly/Private/AnimNextInterfaceGraph_EdGraphSchema.cpp
@@ -0,0 +1,42 @@
+﻿// Copyright Epic Games, Inc. All Rights Reserved.
+
+#include "AnimNextInterfaceGraph_EdGraphSchema.h"
+#include "Editor.h"
+
+void UAnimNextInterfaceGraph_EdGraphSchema::TrySetDefaultValue(UEdGraphPin& InPin, const FString& InNewDefaultValue, bool bMarkAsModified) const
+{
+#if WITH_EDITOR
+	if (GEditor)
+	{
+		GEditor->CancelTransaction(0);
+	}
+#endif
+
+	FString UseDefaultValue;
+	TObjectPtr<UObject> UseDefaultObject = nullptr;
+	FText UseDefaultText;
+
+	GetDefault<UEdGraphSchema_K2>()->GetPinDefaultValuesFromString(InPin.PinType, InPin.GetOwningNodeUnchecked(), InNewDefaultValue, UseDefaultValue, UseDefaultObject, UseDefaultText, /*bPreserveTextIdentity*/false);
+
+	// Check the default value and make it an error if it's bogus
+	if (GetDefault<UEdGraphSchema_K2>()->DefaultValueSimpleValidation(InPin.PinType, InPin.PinName, UseDefaultValue, UseDefaultObject, UseDefaultText))
+	{
+		InPin.DefaultObject = UseDefaultObject;
+		InPin.DefaultValue = UseDefaultValue;
+		InPin.DefaultTextValue = UseDefaultText;
+
+		UEdGraphNode* Node = InPin.GetOwningNode();
+		Node->PinDefaultValueChanged(&InPin);
+
+		// If the default value is manually set then treat it as if the value was reset to default and remove the orphaned InPin
+		if (InPin.bOrphanedPin && InPin.DoesDefaultValueMatchAutogenerated())
+		{
+			Node->PinConnectionListChanged(&InPin);
+		}
+
+		if (bMarkAsModified)
+		{
+			Node->MarkPackageDirty();
+		}
+	}
+}
diff --git a/Engine/Plugins/Experimental/AnimNext/Source/AnimNextInterfaceGraphUncookedOnly/Private/AnimNextInterfaceGraph_EditorData.cpp b/Engine/Plugins/Experimental/AnimNext/Source/AnimNextInterfaceGraphUncookedOnly/Private/AnimNextInterfaceGraph_EditorData.cpp
new file mode 100644
--- /dev/null
+++ b/Engine/Plugins/Experimental/AnimNext/Source/AnimNextInterfaceGraphUncookedOnly/Private/AnimNextInterfaceGraph_EditorData.cpp
@@ -0,0 +1,355 @@
+﻿// Copyright Epic Games, Inc. All Rights Reserved.
+
+#include "AnimNextInterfaceGraph_EditorData.h"
+
+#include "ControlRigDefines.h"
+#include "AnimNextInterfaceGraph.h"
+#include "AnimNextInterfaceGraph_EdGraphSchema.h"
+#include "AnimNextInterfaceUncookedOnlyUtils.h"
+#include "AnimNextInterfaceExecuteContext.h"
+#include "Rigs/RigHierarchyPose.h"
+#include "RigVMModel/RigVMFunctionLibrary.h"
+#include "RigVMModel/RigVMNotifications.h"
+#include "RigVMModel/Nodes/RigVMCollapseNode.h"
+#include "Curves/CurveFloat.h"
+
+UAnimNextInterfaceGraph_EditorData::UAnimNextInterfaceGraph_EditorData(const FObjectInitializer& ObjectInitializer)
+{
+	RigVMClient.Reset();
+	RigVMClient.SetOuterClientHost(this, GET_MEMBER_NAME_CHECKED(UAnimNextInterfaceGraph_EditorData, RigVMClient));
+	{
+		TGuardValue<bool> DisableClientNotifs(RigVMClient.bSuspendNotifications, true);
+		RigVMClient.AddModel(TEXT("RigVMGraph"), false, &ObjectInitializer);
+		RigVMClient.GetOrCreateFunctionLibrary(false, &ObjectInitializer);
+	}
+	
+	auto MakeEdGraph = [this, &ObjectInitializer](FName InName) -> UAnimNextInterfaceGraph_EdGraph*
+	{
+		UAnimNextInterfaceGraph_EdGraph* EdGraph = ObjectInitializer.CreateDefaultSubobject<UAnimNextInterfaceGraph_EdGraph>(this, InName);
+		EdGraph->Schema = UAnimNextInterfaceGraph_EdGraphSchema::StaticClass();
+		EdGraph->bAllowRenaming = false;
+		EdGraph->bEditable = true;
+		EdGraph->bAllowDeletion = false;
+		EdGraph->bIsFunctionDefinition = false;
+		EdGraph->Initialize(this);
+
+		return EdGraph;
+	};
+
+	RootGraph = MakeEdGraph(TEXT("RootEdGraph"));
+	FunctionLibraryEdGraph = MakeEdGraph(TEXT("RigVMFunctionLibraryEdGraph"));
+}
+
+void UAnimNextInterfaceGraph_EditorData::Serialize(FArchive& Ar)
+{
+	RigVMClient.SetOuterClientHost(this, GET_MEMBER_NAME_CHECKED(UAnimNextInterfaceGraph_EditorData, RigVMClient));
+
+	UObject::Serialize(Ar);
+
+	if(Ar.IsLoading())
+	{
+		if(RigVMGraph_DEPRECATED || RigVMFunctionLibrary_DEPRECATED)
+		{
+			TGuardValue<bool> DisableClientNotifs(RigVMClient.bSuspendNotifications, true);
+			RigVMClient.SetFromDeprecatedData(RigVMGraph_DEPRECATED, RigVMFunctionLibrary_DEPRECATED);
+		}
+	}
+}
+
+void UAnimNextInterfaceGraph_EditorData::Initialize(bool bRecompileVM)
+{
+	UAnimNextInterfaceGraph* AnimNextInterfaceGraph = GetTypedOuter<UAnimNextInterfaceGraph>();
+
+	if (RigVMClient.GetController(0) == nullptr)
+	{
+		check(RigVMClient.Num() == 1);
+		check(RigVMClient.GetFunctionLibrary());
+		
+		RigVMClient.GetOrCreateController(RigVMClient.GetDefaultModel());
+		RigVMClient.GetOrCreateController(RigVMClient.GetFunctionLibrary());
+
+		// Init function library controllers
+		for(URigVMLibraryNode* LibraryNode : RigVMClient.GetFunctionLibrary()->GetFunctions())
+		{
+			RigVMClient.GetOrCreateController(LibraryNode->GetContainedGraph());
+		}
+		
+		if(bRecompileVM)
+		{
+			UE::AnimNext::InterfaceGraphUncookedOnly::FUtils::Compile(AnimNextInterfaceGraph);
+		}
+	}
+
+	RootGraph->Initialize(this);
+	FunctionLibraryEdGraph->Initialize(this);
+	if(EntryPointGraph)
+	{
+		EntryPointGraph->Initialize(this);
+	}
+}
+
+void UAnimNextInterfaceGraph_EditorData::PostLoad()
+{
+	Super::PostLoad();
+	
+	Initialize(/*bRecompileVM*/false);
+}
+
+FRigVMClient* UAnimNextInterfaceGraph_EditorData::GetRigVMClient()
+{
+	return &RigVMClient;
+}
+
+const FRigVMClient* UAnimNextInterfaceGraph_EditorData::GetRigVMClient() const
+{
+	return &RigVMClient;
+}
+
+IRigVMGraphFunctionHost* UAnimNextInterfaceGraph_EditorData::GetRigVMGraphFunctionHost()
+{
+	return this;
+}
+
+const IRigVMGraphFunctionHost* UAnimNextInterfaceGraph_EditorData::GetRigVMGraphFunctionHost() const
+{
+	return this;
+}
+
+
+void UAnimNextInterfaceGraph_EditorData::HandleRigVMGraphAdded(const FRigVMClient* InClient, const FString& InNodePath)
+{
+	if(URigVMGraph* RigVMGraph = InClient->GetModel(InNodePath))
+	{
+		RigVMGraph->SetExecuteContextStruct(FRigVMExecuteContext::StaticStruct());
+	}
+}
+
+void UAnimNextInterfaceGraph_EditorData::HandleConfigureRigVMController(const FRigVMClient* InClient,
+                                                                    URigVMController* InControllerToConfigure)
+{
+	InControllerToConfigure->OnModified().AddUObject(this, &UAnimNextInterfaceGraph_EditorData::HandleModifiedEvent);
+
+	InControllerToConfigure->UnfoldStructDelegate.BindLambda([](const UStruct* InStruct) -> bool
+	{
+		if (InStruct == TBaseStructure<FQuat>::Get())
+		{
+			return false;
+		}
+		if (InStruct == FRuntimeFloatCurve::StaticStruct())
+		{
+			return false;
+		}
+		if (InStruct == FRigPose::StaticStruct())
+		{
+			return false;
+		}
+		return true;
+	});
+
+	TWeakObjectPtr<UAnimNextInterfaceGraph_EditorData> WeakThis(this);
+
+	// this delegate is used by the controller to determine variable validity
+	// during a bind process. the controller itself doesn't own the variables,
+	// so we need a delegate to request them from the owning blueprint
+	InControllerToConfigure->GetExternalVariablesDelegate.BindLambda([](URigVMGraph* InGraph) -> TArray<FRigVMExternalVariable>
+	{
+		if (InGraph)
+		{
+			if(UAnimNextInterfaceGraph_EditorData* EditorData = InGraph->GetTypedOuter<UAnimNextInterfaceGraph_EditorData>())
+			{
+				if (UAnimNextInterfaceGraph* Graph = EditorData->GetTypedOuter<UAnimNextInterfaceGraph>())
+				{
+					return Graph->GetRigVMExternalVariables();
+				}
+			}
+		}
+		return TArray<FRigVMExternalVariable>();
+	});
+
+	// this delegate is used by the controller to retrieve the current bytecode of the VM
+	InControllerToConfigure->GetCurrentByteCodeDelegate.BindLambda([WeakThis]() -> const FRigVMByteCode*
+	{
+		if (WeakThis.IsValid())
+		{
+			if(UAnimNextInterfaceGraph* Graph = WeakThis->GetTypedOuter<UAnimNextInterfaceGraph>())
+			{
+				if (Graph->RigVM)
+				{
+					return &Graph->RigVM->GetByteCode();
+				}
+			}
+		}
+		return nullptr;
+
+	});
+
+#if WITH_EDITOR
+	InControllerToConfigure->SetupDefaultUnitNodeDelegates(TDelegate<FName(FRigVMExternalVariable, FString)>::CreateLambda(
+		[](FRigVMExternalVariable InVariableToCreate, FString InDefaultValue) -> FName
+		{
+			return NAME_None;
+		}
+	));
+#endif
+
+}
+
+UObject* UAnimNextInterfaceGraph_EditorData::GetEditorObjectForRigVMGraph(URigVMGraph* InVMGraph) const
+{
+	if(InVMGraph)
+	{
+		TArray<UAnimNextInterfaceGraph_EdGraph*> AllGraphs = {RootGraph, EntryPointGraph, FunctionLibraryEdGraph};
+		for(UAnimNextInterfaceGraph_EdGraph* EdGraph : AllGraphs)
+		{
+			if(EdGraph)
+			{
+				if(EdGraph->ModelNodePath == InVMGraph->GetNodePath())
+				{
+					return EdGraph;
+				}
+			}
+		}
+	}
+	return nullptr;
+}
+
+FRigVMGraphFunctionStore* UAnimNextInterfaceGraph_EditorData::GetRigVMGraphFunctionStore()
+{
+	return &GraphFunctionStore;
+}
+
+const FRigVMGraphFunctionStore* UAnimNextInterfaceGraph_EditorData::GetRigVMGraphFunctionStore() const
+{
+	return &GraphFunctionStore;
+}
+
+
+void UAnimNextInterfaceGraph_EditorData::RecompileVM()
+{
+	UE::AnimNext::InterfaceGraphUncookedOnly::FUtils::Compile(GetTypedOuter<UAnimNextInterfaceGraph>());
+}
+
+void UAnimNextInterfaceGraph_EditorData::RecompileVMIfRequired()
+{
+	if (bVMRecompilationRequired)
+	{
+		RecompileVM();
+	}
+}
+
+void UAnimNextInterfaceGraph_EditorData::RequestAutoVMRecompilation()
+{
+	bVMRecompilationRequired = true;
+	if (bAutoRecompileVM)
+	{
+		RecompileVMIfRequired();
+	}
+}
+
+void UAnimNextInterfaceGraph_EditorData::HandleModifiedEvent(ERigVMGraphNotifType InNotifType, URigVMGraph* InGraph, UObject* InSubject)
+{
+	bool bNotifForOthersPending = true;
+
+	switch(InNotifType)
+	{
+	case ERigVMGraphNotifType::NodeAdded:
+		{
+			if (URigVMCollapseNode* CollapseNode = Cast<URigVMCollapseNode>(InSubject))
+			{
+				CreateEdGraphForCollapseNode(CollapseNode);
+				break;
+			}
+		}
+		// Fall through to the next case
+	case ERigVMGraphNotifType::LinkAdded:
+	case ERigVMGraphNotifType::LinkRemoved:
+	case ERigVMGraphNotifType::PinArraySizeChanged:
+	case ERigVMGraphNotifType::PinDirectionChanged:
+		{
+			if (InGraph)
+			{
+				InGraph->ClearAST();
+			}
+			break;
+		}
+
+	case ERigVMGraphNotifType::PinDefaultValueChanged:
+		{
+			if (InGraph->GetRuntimeAST().IsValid())
+			{
+				URigVMPin* RootPin = CastChecked<URigVMPin>(InSubject)->GetRootPin();
+				FRigVMASTProxy RootPinProxy = FRigVMASTProxy::MakeFromUObject(RootPin);
+				const FRigVMExprAST* Expression = InGraph->GetRuntimeAST()->GetExprForSubject(RootPinProxy);
+				if (Expression == nullptr)
+				{
+					InGraph->ClearAST();
+					break;
+				}
+				else if (Expression->NumParents() > 1)
+				{
+					InGraph->ClearAST();
+					break;
+				}
+			}
+			break;
+		}
+	}
+	
+	// if the notification still has to be sent...
+	if (bNotifForOthersPending && !bSuspendModelNotificationsForOthers)
+	{
+		if (ModifiedEvent.IsBound())
+		{
+			ModifiedEvent.Broadcast(InNotifType, InGraph, InSubject);
+		}
+	}
+}
+
+URigVMGraph* UAnimNextInterfaceGraph_EditorData::GetVMGraphForEdGraph(const UEdGraph* InGraph) const
+{
+	if (InGraph == RootGraph)
+	{
+		return RigVMClient.GetDefaultModel();
+	}
+	else
+	{
+		const UAnimNextInterfaceGraph_EdGraph* Graph = Cast<UAnimNextInterfaceGraph_EdGraph>(InGraph);
+		check(Graph);
+
+		if (Graph->bIsFunctionDefinition)
+		{
+			if (URigVMLibraryNode* LibraryNode = RigVMClient.GetFunctionLibrary()->FindFunction(*Graph->ModelNodePath))
+			{
+				return LibraryNode->GetContainedGraph();
+			}
+		}
+	}
+
+	return nullptr;
+}
+
+void UAnimNextInterfaceGraph_EditorData::CreateEdGraphForCollapseNode(URigVMCollapseNode* InNode)
+{
+	if (InNode->GetGraph()->IsA<URigVMFunctionLibrary>())
+	{
+		if (URigVMGraph* ContainedGraph = InNode->GetContainedGraph())
+		{
+			if (EntryPointGraph == nullptr)
+			{
+				// create a sub graph
+				UAnimNextInterfaceGraph_EdGraph* RigFunctionGraph = NewObject<UAnimNextInterfaceGraph_EdGraph>(this, *InNode->GetName(), RF_Transactional);
+				RigFunctionGraph->Schema = UAnimNextInterfaceGraph_EdGraphSchema::StaticClass();
+				RigFunctionGraph->bAllowRenaming = true;
+				RigFunctionGraph->bEditable = true;
+				RigFunctionGraph->bAllowDeletion = true;
+				RigFunctionGraph->ModelNodePath = ContainedGraph->GetNodePath();
+				RigFunctionGraph->bIsFunctionDefinition = true;
+
+				EntryPointGraph = RigFunctionGraph;
+				RigFunctionGraph->Initialize(this);
+
+				RigVMClient.GetOrCreateController(ContainedGraph)->ResendAllNotifications();
+			}
+		}
+	}
+}
\ No newline at end of file
diff --git a/Engine/Plugins/Experimental/DataInterface/DataInterface.uplugin b/Engine/Plugins/Experimental/DataInterface/DataInterface.uplugin
deleted file mode 100644
diff --git a/Engine/Plugins/Experimental/DataInterface/Source/DataInterface/DataInterface.Build.cs b/Engine/Plugins/Experimental/DataInterface/Source/DataInterface/DataInterface.Build.cs
deleted file mode 100644
diff --git a/Engine/Plugins/Experimental/DataInterface/Source/DataInterface/Private/DataInterface.cpp b/Engine/Plugins/Experimental/DataInterface/Source/DataInterface/Private/DataInterface.cpp
deleted file mode 100644
diff --git a/Engine/Plugins/Experimental/DataInterface/Source/DataInterface/Private/DataInterfaceContext.cpp b/Engine/Plugins/Experimental/DataInterface/Source/DataInterface/Private/DataInterfaceContext.cpp
deleted file mode 100644
diff --git a/Engine/Plugins/Experimental/DataInterface/Source/DataInterface/Private/DataInterfaceModule.cpp b/Engine/Plugins/Experimental/DataInterface/Source/DataInterface/Private/DataInterfaceModule.cpp
deleted file mode 100644
diff --git a/Engine/Plugins/Experimental/DataInterface/Source/DataInterface/Private/DataInterfaceParam.cpp b/Engine/Plugins/Experimental/DataInterface/Source/DataInterface/Private/DataInterfaceParam.cpp
deleted file mode 100644
diff --git a/Engine/Plugins/Experimental/DataInterface/Source/DataInterface/Private/DataInterfaceParamStorage.cpp b/Engine/Plugins/Experimental/DataInterface/Source/DataInterface/Private/DataInterfaceParamStorage.cpp
deleted file mode 100644
diff --git a/Engine/Plugins/Experimental/DataInterface/Source/DataInterface/Private/DataInterfaceState.cpp b/Engine/Plugins/Experimental/DataInterface/Source/DataInterface/Private/DataInterfaceState.cpp
deleted file mode 100644
diff --git a/Engine/Plugins/Experimental/DataInterface/Source/DataInterface/Private/DataInterfaceTypes.cpp b/Engine/Plugins/Experimental/DataInterface/Source/DataInterface/Private/DataInterfaceTypes.cpp
deleted file mode 100644
diff --git a/Engine/Plugins/Experimental/DataInterface/Source/DataInterface/Public/DataInterface.h b/Engine/Plugins/Experimental/DataInterface/Source/DataInterface/Public/DataInterface.h
deleted file mode 100644
diff --git a/Engine/Plugins/Experimental/DataInterface/Source/DataInterface/Public/DataInterfaceContext.h b/Engine/Plugins/Experimental/DataInterface/Source/DataInterface/Public/DataInterfaceContext.h
deleted file mode 100644
diff --git a/Engine/Plugins/Experimental/DataInterface/Source/DataInterface/Public/DataInterfaceKernel.h b/Engine/Plugins/Experimental/DataInterface/Source/DataInterface/Public/DataInterfaceKernel.h
deleted file mode 100644
diff --git a/Engine/Plugins/Experimental/DataInterface/Source/DataInterface/Public/DataInterfaceKey.h b/Engine/Plugins/Experimental/DataInterface/Source/DataInterface/Public/DataInterfaceKey.h
deleted file mode 100644
diff --git a/Engine/Plugins/Experimental/DataInterface/Source/DataInterface/Public/DataInterfaceParam.h b/Engine/Plugins/Experimental/DataInterface/Source/DataInterface/Public/DataInterfaceParam.h
deleted file mode 100644
diff --git a/Engine/Plugins/Experimental/DataInterface/Source/DataInterface/Public/DataInterfaceParamStorage.h b/Engine/Plugins/Experimental/DataInterface/Source/DataInterface/Public/DataInterfaceParamStorage.h
deleted file mode 100644
diff --git a/Engine/Plugins/Experimental/DataInterface/Source/DataInterface/Public/DataInterfaceState.h b/Engine/Plugins/Experimental/DataInterface/Source/DataInterface/Public/DataInterfaceState.h
deleted file mode 100644
diff --git a/Engine/Plugins/Experimental/DataInterface/Source/DataInterface/Public/DataInterfaceTypes.h b/Engine/Plugins/Experimental/DataInterface/Source/DataInterface/Public/DataInterfaceTypes.h
deleted file mode 100644
diff --git a/Engine/Plugins/Experimental/DataInterface/Source/DataInterface/Public/DataInterfaceTypes.inl b/Engine/Plugins/Experimental/DataInterface/Source/DataInterface/Public/DataInterfaceTypes.inl
deleted file mode 100644
diff --git a/Engine/Plugins/Experimental/DataInterface/Source/DataInterface/Public/IDataInterface.h b/Engine/Plugins/Experimental/DataInterface/Source/DataInterface/Public/IDataInterface.h
deleted file mode 100644
diff --git a/Engine/Plugins/Experimental/DataInterface/Source/DataInterface/Public/IDataInterfaceParameters.h b/Engine/Plugins/Experimental/DataInterface/Source/DataInterface/Public/IDataInterfaceParameters.h
deleted file mode 100644
diff --git a/Engine/Plugins/Experimental/DataInterfaceGraph/DataInterfaceGraph.uplugin b/Engine/Plugins/Experimental/DataInterfaceGraph/DataInterfaceGraph.uplugin
deleted file mode 100644
diff --git a/Engine/Plugins/Experimental/DataInterfaceGraph/Source/DataInterfaceGraph/DataInterfaceGraph.Build.cs b/Engine/Plugins/Experimental/DataInterfaceGraph/Source/DataInterfaceGraph/DataInterfaceGraph.Build.cs
deleted file mode 100644
diff --git a/Engine/Plugins/Experimental/DataInterfaceGraph/Source/DataInterfaceGraph/Internal/DataInterfaceExecuteContext.h b/Engine/Plugins/Experimental/DataInterfaceGraph/Source/DataInterfaceGraph/Internal/DataInterfaceExecuteContext.h
deleted file mode 100644
diff --git a/Engine/Plugins/Experimental/DataInterfaceGraph/Source/DataInterfaceGraph/Internal/DataInterfaceGraph.h b/Engine/Plugins/Experimental/DataInterfaceGraph/Source/DataInterfaceGraph/Internal/DataInterfaceGraph.h
deleted file mode 100644
diff --git a/Engine/Plugins/Experimental/DataInterfaceGraph/Source/DataInterfaceGraph/Internal/DataInterfaceUnitContext.h b/Engine/Plugins/Experimental/DataInterfaceGraph/Source/DataInterfaceGraph/Internal/DataInterfaceUnitContext.h
deleted file mode 100644
diff --git a/Engine/Plugins/Experimental/DataInterfaceGraph/Source/DataInterfaceGraph/Internal/RigUnit_DataInterfaceBeginExecution.h b/Engine/Plugins/Experimental/DataInterfaceGraph/Source/DataInterfaceGraph/Internal/RigUnit_DataInterfaceBeginExecution.h
deleted file mode 100644
diff --git a/Engine/Plugins/Experimental/DataInterfaceGraph/Source/DataInterfaceGraph/Internal/RigUnit_DataInterfaceEndExecution.h b/Engine/Plugins/Experimental/DataInterfaceGraph/Source/DataInterfaceGraph/Internal/RigUnit_DataInterfaceEndExecution.h
deleted file mode 100644
diff --git a/Engine/Plugins/Experimental/DataInterfaceGraph/Source/DataInterfaceGraph/Private/DataInterfaceGraph.cpp b/Engine/Plugins/Experimental/DataInterfaceGraph/Source/DataInterfaceGraph/Private/DataInterfaceGraph.cpp
deleted file mode 100644
diff --git a/Engine/Plugins/Experimental/DataInterfaceGraph/Source/DataInterfaceGraph/Private/DataInterfaceGraphModule.cpp b/Engine/Plugins/Experimental/DataInterfaceGraph/Source/DataInterfaceGraph/Private/DataInterfaceGraphModule.cpp
deleted file mode 100644
diff --git a/Engine/Plugins/Experimental/DataInterfaceGraph/Source/DataInterfaceGraph/Private/DataInterface_Bool.cpp b/Engine/Plugins/Experimental/DataInterfaceGraph/Source/DataInterfaceGraph/Private/DataInterface_Bool.cpp
deleted file mode 100644
diff --git a/Engine/Plugins/Experimental/DataInterfaceGraph/Source/DataInterfaceGraph/Private/DataInterface_Bool.h b/Engine/Plugins/Experimental/DataInterfaceGraph/Source/DataInterfaceGraph/Private/DataInterface_Bool.h
deleted file mode 100644
diff --git a/Engine/Plugins/Experimental/DataInterfaceGraph/Source/DataInterfaceGraph/Private/DataInterface_Float.cpp b/Engine/Plugins/Experimental/DataInterfaceGraph/Source/DataInterfaceGraph/Private/DataInterface_Float.cpp
deleted file mode 100644
diff --git a/Engine/Plugins/Experimental/DataInterfaceGraph/Source/DataInterfaceGraph/Private/DataInterface_Float.h b/Engine/Plugins/Experimental/DataInterfaceGraph/Source/DataInterfaceGraph/Private/DataInterface_Float.h
deleted file mode 100644
diff --git a/Engine/Plugins/Experimental/DataInterfaceGraph/Source/DataInterfaceGraph/Private/DataInterface_Name.cpp b/Engine/Plugins/Experimental/DataInterfaceGraph/Source/DataInterfaceGraph/Private/DataInterface_Name.cpp
deleted file mode 100644
diff --git a/Engine/Plugins/Experimental/DataInterfaceGraph/Source/DataInterfaceGraph/Private/DataInterface_Name.h b/Engine/Plugins/Experimental/DataInterfaceGraph/Source/DataInterfaceGraph/Private/DataInterface_Name.h
deleted file mode 100644
diff --git a/Engine/Plugins/Experimental/DataInterfaceGraph/Source/DataInterfaceGraph/Private/DataInterface_Wrapper.cpp b/Engine/Plugins/Experimental/DataInterfaceGraph/Source/DataInterfaceGraph/Private/DataInterface_Wrapper.cpp
deleted file mode 100644
diff --git a/Engine/Plugins/Experimental/DataInterfaceGraph/Source/DataInterfaceGraph/Private/DataInterface_Wrapper.h b/Engine/Plugins/Experimental/DataInterfaceGraph/Source/DataInterfaceGraph/Private/DataInterface_Wrapper.h
deleted file mode 100644
diff --git a/Engine/Plugins/Experimental/DataInterfaceGraph/Source/DataInterfaceGraph/Private/RigUnit_DataInterfaceBeginExecution.cpp b/Engine/Plugins/Experimental/DataInterfaceGraph/Source/DataInterfaceGraph/Private/RigUnit_DataInterfaceBeginExecution.cpp
deleted file mode 100644
diff --git a/Engine/Plugins/Experimental/DataInterfaceGraph/Source/DataInterfaceGraph/Private/RigUnit_DataInterfaceEndExecution.cpp b/Engine/Plugins/Experimental/DataInterfaceGraph/Source/DataInterfaceGraph/Private/RigUnit_DataInterfaceEndExecution.cpp
deleted file mode 100644
diff --git a/Engine/Plugins/Experimental/DataInterfaceGraph/Source/DataInterfaceGraph/Private/RigUnit_DataInterfaceParameter.cpp b/Engine/Plugins/Experimental/DataInterfaceGraph/Source/DataInterfaceGraph/Private/RigUnit_DataInterfaceParameter.cpp
deleted file mode 100644
diff --git a/Engine/Plugins/Experimental/DataInterfaceGraph/Source/DataInterfaceGraph/Private/RigUnit_DataInterfaceParameter.h b/Engine/Plugins/Experimental/DataInterfaceGraph/Source/DataInterfaceGraph/Private/RigUnit_DataInterfaceParameter.h
deleted file mode 100644
diff --git a/Engine/Plugins/Experimental/DataInterfaceGraph/Source/DataInterfaceGraphEditor/DataInterfaceGraphEditor.Build.cs b/Engine/Plugins/Experimental/DataInterfaceGraph/Source/DataInterfaceGraphEditor/DataInterfaceGraphEditor.Build.cs
deleted file mode 100644
diff --git a/Engine/Plugins/Experimental/DataInterfaceGraph/Source/DataInterfaceGraphEditor/Private/AssetTypeActions_DataInterfaceGraph.cpp b/Engine/Plugins/Experimental/DataInterfaceGraph/Source/DataInterfaceGraphEditor/Private/AssetTypeActions_DataInterfaceGraph.cpp
deleted file mode 100644
diff --git a/Engine/Plugins/Experimental/DataInterfaceGraph/Source/DataInterfaceGraphEditor/Private/AssetTypeActions_DataInterfaceGraph.h b/Engine/Plugins/Experimental/DataInterfaceGraph/Source/DataInterfaceGraphEditor/Private/AssetTypeActions_DataInterfaceGraph.h
deleted file mode 100644
diff --git a/Engine/Plugins/Experimental/DataInterfaceGraph/Source/DataInterfaceGraphEditor/Private/DataInterfaceEditorUtils.cpp b/Engine/Plugins/Experimental/DataInterfaceGraph/Source/DataInterfaceGraphEditor/Private/DataInterfaceEditorUtils.cpp
deleted file mode 100644
diff --git a/Engine/Plugins/Experimental/DataInterfaceGraph/Source/DataInterfaceGraphEditor/Private/DataInterfaceEditorUtils.h b/Engine/Plugins/Experimental/DataInterfaceGraph/Source/DataInterfaceGraphEditor/Private/DataInterfaceEditorUtils.h
deleted file mode 100644
diff --git a/Engine/Plugins/Experimental/DataInterfaceGraph/Source/DataInterfaceGraphEditor/Private/DataInterfaceGraphEditor.cpp b/Engine/Plugins/Experimental/DataInterfaceGraph/Source/DataInterfaceGraphEditor/Private/DataInterfaceGraphEditor.cpp
deleted file mode 100644
diff --git a/Engine/Plugins/Experimental/DataInterfaceGraph/Source/DataInterfaceGraphEditor/Private/DataInterfaceGraphEditor.h b/Engine/Plugins/Experimental/DataInterfaceGraph/Source/DataInterfaceGraphEditor/Private/DataInterfaceGraphEditor.h
deleted file mode 100644
diff --git a/Engine/Plugins/Experimental/DataInterfaceGraph/Source/DataInterfaceGraphEditor/Private/DataInterfaceGraphEditorMode.cpp b/Engine/Plugins/Experimental/DataInterfaceGraph/Source/DataInterfaceGraphEditor/Private/DataInterfaceGraphEditorMode.cpp
deleted file mode 100644
diff --git a/Engine/Plugins/Experimental/DataInterfaceGraph/Source/DataInterfaceGraphEditor/Private/DataInterfaceGraphEditorMode.h b/Engine/Plugins/Experimental/DataInterfaceGraph/Source/DataInterfaceGraphEditor/Private/DataInterfaceGraphEditorMode.h
deleted file mode 100644
diff --git a/Engine/Plugins/Experimental/DataInterfaceGraph/Source/DataInterfaceGraphEditor/Private/DataInterfaceGraphEditorModule.cpp b/Engine/Plugins/Experimental/DataInterfaceGraph/Source/DataInterfaceGraphEditor/Private/DataInterfaceGraphEditorModule.cpp
deleted file mode 100644
diff --git a/Engine/Plugins/Experimental/DataInterfaceGraph/Source/DataInterfaceGraphEditor/Private/DataInterfaceGraphEditorModule.h b/Engine/Plugins/Experimental/DataInterfaceGraph/Source/DataInterfaceGraphEditor/Private/DataInterfaceGraphEditorModule.h
deleted file mode 100644
diff --git a/Engine/Plugins/Experimental/DataInterfaceGraph/Source/DataInterfaceGraphEditor/Private/DataInterfaceGraphEditorSchemaActions.cpp b/Engine/Plugins/Experimental/DataInterfaceGraph/Source/DataInterfaceGraphEditor/Private/DataInterfaceGraphEditorSchemaActions.cpp
deleted file mode 100644
diff --git a/Engine/Plugins/Experimental/DataInterfaceGraph/Source/DataInterfaceGraphEditor/Private/DataInterfaceGraphEditorSchemaActions.h b/Engine/Plugins/Experimental/DataInterfaceGraph/Source/DataInterfaceGraphEditor/Private/DataInterfaceGraphEditorSchemaActions.h
deleted file mode 100644
diff --git a/Engine/Plugins/Experimental/DataInterfaceGraph/Source/DataInterfaceGraphEditor/Private/DataInterfaceGraphFactory.cpp b/Engine/Plugins/Experimental/DataInterfaceGraph/Source/DataInterfaceGraphEditor/Private/DataInterfaceGraphFactory.cpp
deleted file mode 100644
diff --git a/Engine/Plugins/Experimental/DataInterfaceGraph/Source/DataInterfaceGraphEditor/Private/DataInterfaceGraphFactory.h b/Engine/Plugins/Experimental/DataInterfaceGraph/Source/DataInterfaceGraphEditor/Private/DataInterfaceGraphFactory.h
deleted file mode 100644
diff --git a/Engine/Plugins/Experimental/DataInterfaceGraph/Source/DataInterfaceGraphEditor/Private/DataInterfacePropertyTypeCustomization.cpp b/Engine/Plugins/Experimental/DataInterfaceGraph/Source/DataInterfaceGraphEditor/Private/DataInterfacePropertyTypeCustomization.cpp
deleted file mode 100644
diff --git a/Engine/Plugins/Experimental/DataInterfaceGraph/Source/DataInterfaceGraphEditor/Private/DataInterfacePropertyTypeCustomization.h b/Engine/Plugins/Experimental/DataInterfaceGraph/Source/DataInterfaceGraphEditor/Private/DataInterfacePropertyTypeCustomization.h
deleted file mode 100644
diff --git a/Engine/Plugins/Experimental/DataInterfaceGraph/Source/DataInterfaceGraphEditor/Private/DataInterfaceWidgetFactories.cpp b/Engine/Plugins/Experimental/DataInterfaceGraph/Source/DataInterfaceGraphEditor/Private/DataInterfaceWidgetFactories.cpp
deleted file mode 100644
diff --git a/Engine/Plugins/Experimental/DataInterfaceGraph/Source/DataInterfaceGraphEditor/Private/SDataInterfaceGraphEditorActionMenu.cpp b/Engine/Plugins/Experimental/DataInterfaceGraph/Source/DataInterfaceGraphEditor/Private/SDataInterfaceGraphEditorActionMenu.cpp
deleted file mode 100644
diff --git a/Engine/Plugins/Experimental/DataInterfaceGraph/Source/DataInterfaceGraphEditor/Private/SDataInterfaceGraphEditorActionMenu.h b/Engine/Plugins/Experimental/DataInterfaceGraph/Source/DataInterfaceGraphEditor/Private/SDataInterfaceGraphEditorActionMenu.h
deleted file mode 100644
diff --git a/Engine/Plugins/Experimental/DataInterfaceGraph/Source/DataInterfaceGraphEditor/Public/DataInterfaceClassFilter.h b/Engine/Plugins/Experimental/DataInterfaceGraph/Source/DataInterfaceGraphEditor/Public/DataInterfaceClassFilter.h
deleted file mode 100644
diff --git a/Engine/Plugins/Experimental/DataInterfaceGraph/Source/DataInterfaceGraphEditor/Public/DataInterfaceWidgetFactories.h b/Engine/Plugins/Experimental/DataInterfaceGraph/Source/DataInterfaceGraphEditor/Public/DataInterfaceWidgetFactories.h
deleted file mode 100644
diff --git a/Engine/Plugins/Experimental/DataInterfaceGraph/Source/DataInterfaceGraphUncookedOnly/DataInterfaceGraphUncookedOnly.Build.cs b/Engine/Plugins/Experimental/DataInterfaceGraph/Source/DataInterfaceGraphUncookedOnly/DataInterfaceGraphUncookedOnly.Build.cs
deleted file mode 100644
diff --git a/Engine/Plugins/Experimental/DataInterfaceGraph/Source/DataInterfaceGraphUncookedOnly/Internal/DataInterfaceGraph_EdGraph.h b/Engine/Plugins/Experimental/DataInterfaceGraph/Source/DataInterfaceGraphUncookedOnly/Internal/DataInterfaceGraph_EdGraph.h
deleted file mode 100644
diff --git a/Engine/Plugins/Experimental/DataInterfaceGraph/Source/DataInterfaceGraphUncookedOnly/Internal/DataInterfaceGraph_EdGraphNode.h b/Engine/Plugins/Experimental/DataInterfaceGraph/Source/DataInterfaceGraphUncookedOnly/Internal/DataInterfaceGraph_EdGraphNode.h
deleted file mode 100644
diff --git a/Engine/Plugins/Experimental/DataInterfaceGraph/Source/DataInterfaceGraphUncookedOnly/Internal/DataInterfaceGraph_EdGraphSchema.h b/Engine/Plugins/Experimental/DataInterfaceGraph/Source/DataInterfaceGraphUncookedOnly/Internal/DataInterfaceGraph_EdGraphSchema.h
deleted file mode 100644
diff --git a/Engine/Plugins/Experimental/DataInterfaceGraph/Source/DataInterfaceGraphUncookedOnly/Internal/DataInterfaceGraph_EditorData.h b/Engine/Plugins/Experimental/DataInterfaceGraph/Source/DataInterfaceGraphUncookedOnly/Internal/DataInterfaceGraph_EditorData.h
deleted file mode 100644
diff --git a/Engine/Plugins/Experimental/DataInterfaceGraph/Source/DataInterfaceGraphUncookedOnly/Internal/DataInterfaceUncookedOnlyUtils.h b/Engine/Plugins/Experimental/DataInterfaceGraph/Source/DataInterfaceGraphUncookedOnly/Internal/DataInterfaceUncookedOnlyUtils.h
deleted file mode 100644
diff --git a/Engine/Plugins/Experimental/DataInterfaceGraph/Source/DataInterfaceGraphUncookedOnly/Private/DataInterfaceGraphUncookedOnlyModule.cpp b/Engine/Plugins/Experimental/DataInterfaceGraph/Source/DataInterfaceGraphUncookedOnly/Private/DataInterfaceGraphUncookedOnlyModule.cpp
deleted file mode 100644
diff --git a/Engine/Plugins/Experimental/DataInterfaceGraph/Source/DataInterfaceGraphUncookedOnly/Private/DataInterfaceGraphUncookedOnlyUtils.cpp b/Engine/Plugins/Experimental/DataInterfaceGraph/Source/DataInterfaceGraphUncookedOnly/Private/DataInterfaceGraphUncookedOnlyUtils.cpp
deleted file mode 100644
diff --git a/Engine/Plugins/Experimental/DataInterfaceGraph/Source/DataInterfaceGraphUncookedOnly/Private/DataInterfaceGraph_EdGraph.cpp b/Engine/Plugins/Experimental/DataInterfaceGraph/Source/DataInterfaceGraphUncookedOnly/Private/DataInterfaceGraph_EdGraph.cpp
deleted file mode 100644
diff --git a/Engine/Plugins/Experimental/DataInterfaceGraph/Source/DataInterfaceGraphUncookedOnly/Private/DataInterfaceGraph_EdGraphSchema.cpp b/Engine/Plugins/Experimental/DataInterfaceGraph/Source/DataInterfaceGraphUncookedOnly/Private/DataInterfaceGraph_EdGraphSchema.cpp
deleted file mode 100644
diff --git a/Engine/Plugins/Experimental/DataInterfaceGraph/Source/DataInterfaceGraphUncookedOnly/Private/DataInterfaceGraph_EditorData.cpp b/Engine/Plugins/Experimental/DataInterfaceGraph/Source/DataInterfaceGraphUncookedOnly/Private/DataInterfaceGraph_EditorData.cpp
deleted file mode 100644
