AnimNext graph misc fixes:
- Updated to new RigVM execute context refactoring
- Fixed relinking pins at load time (PostLoad)
- Fixed graph rename crashing (PostRename)
- Removed some Data Interface expressions that were left
- Initial AnimGraphNode to be able to use AnimNextGraphs on AnimBPs

diff --git a/Engine/Plugins/Experimental/AnimNext/Source/AnimNextInterface/Public/AnimNextInterfaceContext.h b/Engine/Plugins/Experimental/AnimNext/Source/AnimNextInterface/Public/AnimNextInterfaceContext.h
--- a/Engine/Plugins/Experimental/AnimNext/Source/AnimNextInterface/Public/AnimNextInterfaceContext.h
+++ b/Engine/Plugins/Experimental/AnimNext/Source/AnimNextInterface/Public/AnimNextInterfaceContext.h
@@ -19,7 +19,7 @@ namespace UE::AnimNext::Interface
 struct FState;
 enum class EStatePersistence : uint8;
 
-// Helper to provide a non-callstack bridge for interfacing with data interface contexts
+// Helper to provide a non-callstack bridge for interfacing with anim interface contexts
 // Declaring one of these on the stack will enable all calls in its scope to access the passed-in context
 // via FThreadContext::Get()
 struct ANIMNEXTINTERFACE_API FThreadContext
@@ -30,7 +30,7 @@ struct ANIMNEXTINTERFACE_API FThreadContext
 	static const FContext& Get();
 };
 
-// Context providing methods for mutating & interrogating the data interface runtime
+// Context providing methods for mutating & interrogating the anim interface runtime
 struct ANIMNEXTINTERFACE_API FContext
 {
 private:
diff --git a/Engine/Plugins/Experimental/AnimNext/Source/AnimNextInterface/Public/AnimNextInterfaceKey.h b/Engine/Plugins/Experimental/AnimNext/Source/AnimNextInterface/Public/AnimNextInterfaceKey.h
--- a/Engine/Plugins/Experimental/AnimNext/Source/AnimNextInterface/Public/AnimNextInterfaceKey.h
+++ b/Engine/Plugins/Experimental/AnimNext/Source/AnimNextInterface/Public/AnimNextInterfaceKey.h
@@ -11,7 +11,7 @@ class IAnimNextInterface;
 namespace UE::AnimNext::Interface
 {
 
-// Key value for cached data interface state
+// Key value for cached anim interface state
 struct ANIMNEXTINTERFACE_API FInterfaceKey
 {
 	FInterfaceKey(const IAnimNextInterface* InAnimNextInterface)
@@ -28,7 +28,7 @@ struct ANIMNEXTINTERFACE_API FInterfaceKey
 	uint32 Hash = 0;
 };
 
-// Key value for cached data interface state
+// Key value for cached anim interface state
 struct ANIMNEXTINTERFACE_API FInterfaceKeyWithId : FInterfaceKey
 {
 	FInterfaceKeyWithId(const FInterfaceKey& InInterfaceKey, uint32 InId)
@@ -44,7 +44,7 @@ struct ANIMNEXTINTERFACE_API FInterfaceKeyWithId : FInterfaceKey
 	uint32 Id = 0;
 };
 
-// Key value for cached data interface state
+// Key value for cached anim interface state
 struct ANIMNEXTINTERFACE_API FInterfaceKeyWithIdAndStack : FInterfaceKeyWithId
 {
 	FInterfaceKeyWithIdAndStack(const FInterfaceKeyWithId& InInterfaceKey, uint32 InCallStackHash)
diff --git a/Engine/Plugins/Experimental/AnimNext/Source/AnimNextInterface/Public/AnimNextInterfaceParam.h b/Engine/Plugins/Experimental/AnimNext/Source/AnimNextInterface/Public/AnimNextInterfaceParam.h
--- a/Engine/Plugins/Experimental/AnimNext/Source/AnimNextInterface/Public/AnimNextInterfaceParam.h
+++ b/Engine/Plugins/Experimental/AnimNext/Source/AnimNextInterface/Public/AnimNextInterfaceParam.h
@@ -79,7 +79,7 @@ public:
 ENUM_CLASS_FLAGS(FParam::EFlags);
 
 
-// Simple type representation for data interfaces
+// Simple type representation for anim interfaces
 struct ANIMNEXTINTERFACE_API FParamType
 {
 	typedef FParam(*ParamCopyFunction)(const FParam& InSource, void* InMemory, int32 InAllocatedMemory);
diff --git a/Engine/Plugins/Experimental/AnimNext/Source/AnimNextInterface/Public/AnimNextInterfaceParamStorage.h b/Engine/Plugins/Experimental/AnimNext/Source/AnimNextInterface/Public/AnimNextInterfaceParamStorage.h
--- a/Engine/Plugins/Experimental/AnimNext/Source/AnimNextInterface/Public/AnimNextInterfaceParamStorage.h
+++ b/Engine/Plugins/Experimental/AnimNext/Source/AnimNextInterface/Public/AnimNextInterfaceParamStorage.h
@@ -52,7 +52,7 @@ public:
 	template<typename ValueType>
 	FHParam AddValue(ValueType* Data, int32 NumElem, FParam::EFlags Flags)
 	{
-		check(ParamFreeIndex < (Parameters.Num() - 1)); // TODO : Growing ?
+		check(ParamFreeIndex < (Parameters.Num())); // TODO : Growing ?
 
 		const int32 ParameterIndex = ParamFreeIndex;
 		ParamFreeIndex++;
diff --git a/Engine/Plugins/Experimental/AnimNext/Source/AnimNextInterface/Public/AnimNextInterfaceState.h b/Engine/Plugins/Experimental/AnimNext/Source/AnimNextInterface/Public/AnimNextInterfaceState.h
--- a/Engine/Plugins/Experimental/AnimNext/Source/AnimNextInterface/Public/AnimNextInterfaceState.h
+++ b/Engine/Plugins/Experimental/AnimNext/Source/AnimNextInterface/Public/AnimNextInterfaceState.h
@@ -19,7 +19,7 @@ enum class EStatePersistence : uint8
 	Permanent,
 };
 
-// Container for data interface state
+// Container for anim interface state
 struct ANIMNEXTINTERFACE_API FState
 {
 	FState(int32 InNumElements = 1)
diff --git a/Engine/Plugins/Experimental/AnimNext/Source/AnimNextInterface/Public/AnimNextInterfaceTypes.h b/Engine/Plugins/Experimental/AnimNext/Source/AnimNextInterface/Public/AnimNextInterfaceTypes.h
--- a/Engine/Plugins/Experimental/AnimNext/Source/AnimNextInterface/Public/AnimNextInterfaceTypes.h
+++ b/Engine/Plugins/Experimental/AnimNext/Source/AnimNextInterface/Public/AnimNextInterfaceTypes.h
@@ -33,16 +33,16 @@ virtual FName GetReturnTypeNameImpl() const final override\
 		UE::AnimNext::Interface::Private::TParamTypeHelper<ValueType>::GetAlignment(), \
 		UE::AnimNext::Interface::Private::TParamTypeHelper<ValueType>::GetParamCloneFunction()); });
 
-// Types that are passed between data interfaces (i.e. exist in in public headers) should declare via this mechanism
+// Types that are passed between anim interfaces (i.e. exist in in public headers) should declare via this mechanism
 #define DECLARE_ANIM_NEXT_INTERFACE_PARAM_TYPE(APIType, ValueType, Identifier) \
 	extern APIType UE::AnimNext::Interface::FParamType Identifier##_AnimNextInterfaceTypeInfo; \
 	DECLARE_ANIM_NEXT_INTERFACE_PARAM_TYPE_INTERNAL(ValueType, Identifier)
 
-// Types that are passed between data interfaces (i.e. exist in in public headers) should declare via this mechanism
+// Types that are passed between anim interfaces (i.e. exist in in public headers) should declare via this mechanism
 #define DECLARE_BUILTIN_ANIM_NEXT_INTERFACE_PARAM_TYPE(ValueType, Identifier) \
 	DECLARE_ANIM_NEXT_INTERFACE_PARAM_TYPE(ANIMNEXTINTERFACE_API, ValueType, Identifier)
 
-// Types that are passed between data interfaces should register via this mechanism
+// Types that are passed between anim interfaces should register via this mechanism
 #define IMPLEMENT_ANIM_NEXT_INTERFACE_PARAM_TYPE(ValueType, Identifier) \
 	UE::AnimNext::Interface::FParamType Identifier##_AnimNextInterfaceTypeInfo; \
 	IMPLEMENT_ANIM_NEXT_INTERFACE_PARAM_TYPE_INTERNAL(ValueType, Identifier)
diff --git a/Engine/Plugins/Experimental/AnimNext/Source/AnimNextInterface/Public/AnimNextInterfaceTypes.inl b/Engine/Plugins/Experimental/AnimNext/Source/AnimNextInterface/Public/AnimNextInterfaceTypes.inl
--- a/Engine/Plugins/Experimental/AnimNext/Source/AnimNextInterface/Public/AnimNextInterfaceTypes.inl
+++ b/Engine/Plugins/Experimental/AnimNext/Source/AnimNextInterface/Public/AnimNextInterfaceTypes.inl
@@ -1,6 +1,6 @@
 ﻿// Copyright Epic Games, Inc. All Rights Reserved.
 
-/** Built in types for data interfaces */
+/** Built in types for anim interfaces */
 
 #ifndef ANIM_NEXT_INTERFACE_TYPE
 #error ANIM_NEXT_INTERFACE_TYPE is not defined!
diff --git a/Engine/Plugins/Experimental/AnimNext/Source/AnimNextInterface/Public/IAnimNextInterface.h b/Engine/Plugins/Experimental/AnimNext/Source/AnimNextInterface/Public/IAnimNextInterface.h
--- a/Engine/Plugins/Experimental/AnimNext/Source/AnimNextInterface/Public/IAnimNextInterface.h
+++ b/Engine/Plugins/Experimental/AnimNext/Source/AnimNextInterface/Public/IAnimNextInterface.h
@@ -19,7 +19,7 @@ class ANIMNEXTINTERFACE_API UAnimNextInterface : public UInterface
 	GENERATED_BODY()
 };
 
-// Empty data interface to support 'any' type
+// Empty anim interface to support 'any' type
 class ANIMNEXTINTERFACE_API IAnimNextInterface
 {
 	GENERATED_BODY()
@@ -65,7 +65,7 @@ private:
 	/** Get data if the types are compatible */
 	bool GetDataIfCompatibleInternal(const UE::AnimNext::Interface::FContext& InContext) const;
 	
-	/** Get a value from a data interface with no dynamic or static type checking. Internal use only. */
+	/** Get a value from an anim interface with no dynamic or static type checking. Internal use only. */
 	bool GetDataRawInternal(const UE::AnimNext::Interface::FContext& InContext) const;
 
 protected:
diff --git a/Engine/Plugins/Experimental/AnimNext/Source/AnimNextInterface/Public/IAnimNextInterfaceParameters.h b/Engine/Plugins/Experimental/AnimNext/Source/AnimNextInterface/Public/IAnimNextInterfaceParameters.h
--- a/Engine/Plugins/Experimental/AnimNext/Source/AnimNextInterface/Public/IAnimNextInterfaceParameters.h
+++ b/Engine/Plugins/Experimental/AnimNext/Source/AnimNextInterface/Public/IAnimNextInterfaceParameters.h
@@ -18,7 +18,7 @@ class ANIMNEXTINTERFACE_API UAnimNextInterfaceParameters : public UInterface
 	GENERATED_BODY()
 };
 
-// Interface representing parameters that can be supplied to a data interface
+// Interface representing parameters that can be supplied to an anim interface
 class ANIMNEXTINTERFACE_API IAnimNextInterfaceParameters
 {
 	GENERATED_BODY()
diff --git a/Engine/Plugins/Experimental/AnimNext/Source/AnimNextInterfaceGraph/AnimNextInterfaceGraph.Build.cs b/Engine/Plugins/Experimental/AnimNext/Source/AnimNextInterfaceGraph/AnimNextInterfaceGraph.Build.cs
--- a/Engine/Plugins/Experimental/AnimNext/Source/AnimNextInterfaceGraph/AnimNextInterfaceGraph.Build.cs
+++ b/Engine/Plugins/Experimental/AnimNext/Source/AnimNextInterfaceGraph/AnimNextInterfaceGraph.Build.cs
@@ -21,6 +21,7 @@ namespace UnrealBuildTool.Rules
 					"Core",
 					"CoreUObject",
 					"AnimNextInterface",
+					"Engine"
 				}
 			);
 		}
diff --git a/Engine/Plugins/Experimental/AnimNext/Source/AnimNextInterfaceGraph/Internal/AnimGraph/AnimNode_AnimNextInterfaceGraph.h b/Engine/Plugins/Experimental/AnimNext/Source/AnimNextInterfaceGraph/Internal/AnimGraph/AnimNode_AnimNextInterfaceGraph.h
new file mode 100644
--- /dev/null
+++ b/Engine/Plugins/Experimental/AnimNext/Source/AnimNextInterfaceGraph/Internal/AnimGraph/AnimNode_AnimNextInterfaceGraph.h
@@ -0,0 +1,74 @@
+// Copyright Epic Games, Inc. All Rights Reserved.
+
+#pragma once
+
+#include "Animation/AnimNode_CustomProperty.h"
+#include "Animation/InputScaleBias.h"
+#include "AnimNextInterfaceGraph.h"
+#include "AnimNode_AnimNextInterfaceGraph.generated.h"
+
+class UNodeMappingContainer;
+
+/**
+ * Animation node that allows a AnimNextInterfaceGraph output to be used in an animation graph
+ */
+USTRUCT()
+struct ANIMNEXTINTERFACEGRAPH_API FAnimNode_AnimNextInterfaceGraph : public FAnimNode_CustomProperty
+{
+	GENERATED_BODY()
+
+	FAnimNode_AnimNextInterfaceGraph();
+	~FAnimNode_AnimNextInterfaceGraph();
+
+	// FAnimNode_Base interface
+	virtual void OnInitializeAnimInstance(const FAnimInstanceProxy* InProxy, const UAnimInstance* InAnimInstance) override;
+	virtual void Initialize_AnyThread(const FAnimationInitializeContext& Context) override;
+	virtual void GatherDebugData(FNodeDebugData& DebugData) override;
+	virtual void Update_AnyThread(const FAnimationUpdateContext& Context) override;
+	virtual void CacheBones_AnyThread(const FAnimationCacheBonesContext& Context) override;
+	virtual void Evaluate_AnyThread(FPoseContext & Output) override;
+	virtual int32 GetLODThreshold() const override { return LODThreshold; }
+
+	virtual void InitializeProperties(const UObject* InSourceInstance, UClass* InTargetClass) override;
+	virtual void PropagateInputProperties(const UObject* InSourceInstance) override;
+
+private:
+#if WITH_EDITOR
+	virtual void HandleObjectsReinstanced_Impl(UObject* InSourceObject, UObject* InTargetObject, const TMap<UObject*, UObject*>& OldToNewInstanceMap) override;
+#endif
+private:
+
+	/** The input pose we will pass to the graph */
+	UPROPERTY(EditAnywhere, Category = Links)
+	FPoseLink Source;
+
+	UPROPERTY(EditAnywhere, Category = Settings, meta = (PinHiddenByDefault))
+	TObjectPtr<UAnimNextInterfaceGraph> AnimNextInterfaceGraph;
+
+	/*
+	 * Max LOD that this node is allowed to run
+	 * For example if you have LODThreadhold to be 2, it will run until LOD 2 (based on 0 index)
+	 * when the component LOD becomes 3, it will stop update/evaluate
+	 * currently transition would be issue and that has to be re-visited
+	 */
+	UPROPERTY(EditAnywhere, Category = Performance, meta = (DisplayName = "LOD Threshold"))
+	int32 LODThreshold;
+
+protected:
+	virtual UClass* GetTargetClass() const override { return AnimNextInterfaceGraph ? AnimNextInterfaceGraph->StaticClass() : nullptr; }
+	
+public:
+
+	void PostSerialize(const FArchive& Ar);
+
+	friend class UAnimGraphNode_AnimNextInterfaceGraph;
+};
+
+template<>
+struct TStructOpsTypeTraits<FAnimNode_AnimNextInterfaceGraph> : public TStructOpsTypeTraitsBase2<FAnimNode_AnimNextInterfaceGraph>
+{
+	enum
+	{
+		WithPostSerialize = true,
+	};
+};
\ No newline at end of file
diff --git a/Engine/Plugins/Experimental/AnimNext/Source/AnimNextInterfaceGraph/Internal/AnimNextInterfaceExecuteContext.h b/Engine/Plugins/Experimental/AnimNext/Source/AnimNextInterfaceGraph/Internal/AnimNextInterfaceExecuteContext.h
--- a/Engine/Plugins/Experimental/AnimNext/Source/AnimNextInterfaceGraph/Internal/AnimNextInterfaceExecuteContext.h
+++ b/Engine/Plugins/Experimental/AnimNext/Source/AnimNextInterfaceGraph/Internal/AnimNextInterfaceExecuteContext.h
@@ -4,7 +4,7 @@
 
 #include "CoreMinimal.h"
 #include "RigVMCore/RigVMExecuteContext.h"
-#include "RigVMCore/RigVMStruct.h"
+#include "Units/RigUnit.h"
 #include "AnimNextInterfaceUnitContext.h"
 #include "AnimNextInterfaceExecuteContext.generated.h"
 
@@ -35,11 +35,18 @@ struct FAnimNextInterfaceExecuteContext : public FRigVMExecuteContext
 		return *AnimNextInterfaceContext;
 	}
 
-	const FAnimNextInterfaceUnitContext& GetUnitContext() const
+	const FRigUnitContext& GetUnitContext() const
 	{
 		return UnitContext;
 	}
 
+	void SetContextData(const IAnimNextInterface* InInterface, const UE::AnimNext::Interface::FContext& InAnimNextInterfaceContext, bool& bInResult)
+	{
+		Interface = InInterface;
+		AnimNextInterfaceContext = &InAnimNextInterfaceContext;
+		ResultPtr = &bInResult;
+	}
+
 	void SetResult(bool bInResult) const
 	{
 		check(ResultPtr);
@@ -67,11 +74,11 @@ private:
 	const UE::AnimNext::Interface::FContext* AnimNextInterfaceContext;
 	const IAnimNextInterface* Interface;
 	bool* ResultPtr;
-	FAnimNextInterfaceUnitContext UnitContext;
+	FRigUnitContext UnitContext;
 };
 
 USTRUCT(meta=(ExecuteContext="FAnimNextInterfaceExecuteContext"))
-struct FRigUnit_AnimNextInterfaceBase : public FRigVMStruct
+struct FRigUnit_AnimNextInterfaceBase : public FRigUnit
 {
 	GENERATED_BODY()
 };
diff --git a/Engine/Plugins/Experimental/AnimNext/Source/AnimNextInterfaceGraph/Internal/AnimNextInterfaceGraph.h b/Engine/Plugins/Experimental/AnimNext/Source/AnimNextInterfaceGraph/Internal/AnimNextInterfaceGraph.h
--- a/Engine/Plugins/Experimental/AnimNext/Source/AnimNextInterfaceGraph/Internal/AnimNextInterfaceGraph.h
+++ b/Engine/Plugins/Experimental/AnimNext/Source/AnimNextInterfaceGraph/Internal/AnimNextInterfaceGraph.h
@@ -31,6 +31,10 @@ class ANIMNEXTINTERFACEGRAPH_API UAnimNextInterfaceGraph : public UObject, publi
 {
 	GENERATED_BODY()
 
+	// UObject interface
+	virtual void PostRename(UObject* OldOuter, const FName OldName) override;
+	virtual void GetPreloadDependencies(TArray<UObject*>& OutDeps) override;
+
 	// IAnimAnimNextInterface interface
 	virtual FName GetReturnTypeNameImpl() const final override;
 	virtual const UScriptStruct* GetReturnTypeStructImpl() const final override;
@@ -44,6 +48,7 @@ class ANIMNEXTINTERFACEGRAPH_API UAnimNextInterfaceGraph : public UObject, publi
 	friend struct UE::AnimNext::InterfaceGraphUncookedOnly::FUtils;
 	friend class UE::AnimNext::InterfaceGraphEditor::FGraphEditor;
 	friend class UAnimNextInterface_Graph;
+	friend class UAnimGraphNode_AnimNextInterfaceGraph;
 	
 	UPROPERTY()
 	TObjectPtr<URigVM> RigVM;
@@ -51,8 +56,14 @@ class ANIMNEXTINTERFACEGRAPH_API UAnimNextInterfaceGraph : public UObject, publi
 	UPROPERTY()
 	FRigVMRuntimeSettings VMRuntimeSettings;
 	
+	UPROPERTY()
+	FName ReturnTypeName;
+
+	UPROPERTY()
+	TObjectPtr<UScriptStruct> ReturnTypeStruct;
+
 #if WITH_EDITORONLY_DATA
 	UPROPERTY(VisibleAnywhere, Instanced, Category = "Graph", meta = (ShowInnerProperties))
 	TObjectPtr<UObject> EditorData;
 #endif
-};
\ No newline at end of file
+};
diff --git a/Engine/Plugins/Experimental/AnimNext/Source/AnimNextInterfaceGraph/Internal/RigUnit_AnimNextInterfaceBeginExecution.h b/Engine/Plugins/Experimental/AnimNext/Source/AnimNextInterfaceGraph/Internal/RigUnit_AnimNextInterfaceBeginExecution.h
--- a/Engine/Plugins/Experimental/AnimNext/Source/AnimNextInterfaceGraph/Internal/RigUnit_AnimNextInterfaceBeginExecution.h
+++ b/Engine/Plugins/Experimental/AnimNext/Source/AnimNextInterfaceGraph/Internal/RigUnit_AnimNextInterfaceBeginExecution.h
@@ -9,7 +9,7 @@
 /**
  * Event for driving the skeleton hierarchy with variables and rig elements
  */
-USTRUCT(meta=(DisplayName="Execute Data Interface", Category="Events", NodeColor="1, 0, 0", Keywords="Begin,Update,Tick,Forward,Event"))
+USTRUCT(meta=(DisplayName="Execute Anim Interface", Category="Events", NodeColor="1, 0, 0", Keywords="Begin,Update,Tick,Forward,Event"))
 struct ANIMNEXTINTERFACEGRAPH_API FRigUnit_AnimNextInterfaceBeginExecution : public FRigUnit_AnimNextInterfaceBase
 {
 	GENERATED_BODY()
diff --git a/Engine/Plugins/Experimental/AnimNext/Source/AnimNextInterfaceGraph/Internal/RigUnit_AnimNextInterfaceEndExecution.h b/Engine/Plugins/Experimental/AnimNext/Source/AnimNextInterfaceGraph/Internal/RigUnit_AnimNextInterfaceEndExecution.h
--- a/Engine/Plugins/Experimental/AnimNext/Source/AnimNextInterfaceGraph/Internal/RigUnit_AnimNextInterfaceEndExecution.h
+++ b/Engine/Plugins/Experimental/AnimNext/Source/AnimNextInterfaceGraph/Internal/RigUnit_AnimNextInterfaceEndExecution.h
@@ -11,7 +11,7 @@
 struct FRigUnitContext;
 
 /** Event for writing back calculated results to external variables */
-USTRUCT(meta = (DisplayName = "End Execute Data Interface", Category = "Events", TitleColor = "1 0 0", NodeColor = "1 1 1"))
+USTRUCT(meta = (DisplayName = "End Execute Anim Interface", Category = "Events", TitleColor = "1 0 0", NodeColor = "1 1 1"))
 struct ANIMNEXTINTERFACEGRAPH_API FRigUnit_AnimNextInterfaceEndExecution : public FRigUnit_AnimNextInterfaceBase
 {
 	GENERATED_BODY()
@@ -30,7 +30,7 @@ public:
 };
 
 /** Event for writing back a calculated bool */
-USTRUCT(meta = (DisplayName = "End Execute Data Interface Bool", Category = "Events", TitleColor = "1 0 0", NodeColor = "1 1 1"))
+USTRUCT(meta = (DisplayName = "End Execute Anim Interface Bool", Category = "Events", TitleColor = "1 0 0", NodeColor = "1 1 1"))
 struct ANIMNEXTINTERFACEGRAPH_API FRigUnit_AnimNextInterfaceEndExecution_Bool : public FRigUnit_AnimNextInterfaceEndExecution
 {
 	GENERATED_BODY()
@@ -41,14 +41,16 @@ struct ANIMNEXTINTERFACEGRAPH_API FRigUnit_AnimNextInterfaceEndExecution_Bool :
 	{}
 	
 	RIGVM_METHOD()
-	void Execute();
+	virtual void Execute() override;
+
+	virtual bool CanOnlyExistOnce() const override { return true; }
 	
 	UPROPERTY(EditAnywhere, Category = Result, meta = (Input))
 	bool Result;
 };
 
 /** Event for writing back a calculated float */
-USTRUCT(meta = (DisplayName = "End Execute Data Interface Float", Category = "Events", TitleColor = "1 0 0", NodeColor = "1 1 1"))
+USTRUCT(meta = (DisplayName = "End Execute Anim Interface Float", Category = "Events", TitleColor = "1 0 0", NodeColor = "1 1 1"))
 struct ANIMNEXTINTERFACEGRAPH_API FRigUnit_AnimNextInterfaceEndExecution_Float : public FRigUnit_AnimNextInterfaceEndExecution
 {
 	GENERATED_BODY()
@@ -59,8 +61,10 @@ struct ANIMNEXTINTERFACEGRAPH_API FRigUnit_AnimNextInterfaceEndExecution_Float :
 	{}
 
 	RIGVM_METHOD()
-	void Execute();
+	virtual void Execute() override;
+
+	virtual bool CanOnlyExistOnce() const override { return true; }
 
 	UPROPERTY(EditAnywhere, Category = Result, meta = (Input))
 	float Result;
-};
\ No newline at end of file
+};
diff --git a/Engine/Plugins/Experimental/AnimNext/Source/AnimNextInterfaceGraph/Private/AnimGraph/AnimNode_AnimNextInterfaceGraph.cpp b/Engine/Plugins/Experimental/AnimNext/Source/AnimNextInterfaceGraph/Private/AnimGraph/AnimNode_AnimNextInterfaceGraph.cpp
new file mode 100644
--- /dev/null
+++ b/Engine/Plugins/Experimental/AnimNext/Source/AnimNextInterfaceGraph/Private/AnimGraph/AnimNode_AnimNextInterfaceGraph.cpp
@@ -0,0 +1,176 @@
+// Copyright Epic Games, Inc. All Rights Reserved.
+
+#include "AnimGraph/AnimNode_AnimNextInterfaceGraph.h"
+#include "ControlRig.h"
+#include "ControlRigComponent.h"
+#include "Components/SkeletalMeshComponent.h"
+#include "Animation/AnimInstanceProxy.h"
+#include "GameFramework/Actor.h"
+#include "Animation/NodeMappingContainer.h"
+#include "AnimationRuntime.h"
+#include "AnimNextInterfaceContext.h"
+#include "AnimNextInterfaceParamStorage.h"
+#include "AnimNextInterfaceTypes.h"
+#include "AnimNextInterface.h"
+
+#include UE_INLINE_GENERATED_CPP_BY_NAME(AnimNode_AnimNextInterfaceGraph)
+
+#if WITH_EDITOR
+#include "Editor.h"
+#endif
+
+FAnimNode_AnimNextInterfaceGraph::FAnimNode_AnimNextInterfaceGraph()
+	: FAnimNode_CustomProperty()
+	, AnimNextInterfaceGraph(nullptr)
+	, LODThreshold(INDEX_NONE)
+{
+}
+
+FAnimNode_AnimNextInterfaceGraph::~FAnimNode_AnimNextInterfaceGraph()
+{
+}
+
+void FAnimNode_AnimNextInterfaceGraph::OnInitializeAnimInstance(const FAnimInstanceProxy* InProxy, const UAnimInstance* InAnimInstance)
+{
+	DECLARE_SCOPE_HIERARCHICAL_COUNTER_FUNC()
+
+	FAnimNode_CustomProperty::OnInitializeAnimInstance(InProxy, InAnimInstance);
+
+	InitializeProperties(InAnimInstance, GetTargetClass());
+}
+
+void FAnimNode_AnimNextInterfaceGraph::GatherDebugData(FNodeDebugData& DebugData)
+{
+	DECLARE_SCOPE_HIERARCHICAL_COUNTER_FUNC()
+
+	Source.GatherDebugData(DebugData.BranchFlow(1.f));
+}
+
+void FAnimNode_AnimNextInterfaceGraph::Update_AnyThread(const FAnimationUpdateContext& Context)
+{
+	DECLARE_SCOPE_HIERARCHICAL_COUNTER_FUNC()
+
+	Source.Update(Context);
+
+	if (IsLODEnabled(Context.AnimInstanceProxy))
+	{
+		GetEvaluateGraphExposedInputs().Execute(Context);
+
+		PropagateInputProperties(Context.AnimInstanceProxy->GetAnimInstanceObject());
+	}
+
+	FAnimNode_CustomProperty::Update_AnyThread(Context);
+
+	//TRACE_ANIM_NODE_VALUE(Context, TEXT("Class"), *GetNameSafe(ControlRigClass.Get()));
+}
+
+void FAnimNode_AnimNextInterfaceGraph::Initialize_AnyThread(const FAnimationInitializeContext& Context)
+{
+	DECLARE_SCOPE_HIERARCHICAL_COUNTER_FUNC()
+
+	Source.Initialize(Context);
+
+	FAnimNode_CustomProperty::Initialize_AnyThread(Context);
+}
+
+void FAnimNode_AnimNextInterfaceGraph::CacheBones_AnyThread(const FAnimationCacheBonesContext& Context)
+{
+	DECLARE_SCOPE_HIERARCHICAL_COUNTER_FUNC()
+
+	FAnimNode_CustomProperty::CacheBones_AnyThread(Context);
+
+	Source.CacheBones(Context);
+}
+
+void FAnimNode_AnimNextInterfaceGraph::Evaluate_AnyThread(FPoseContext & Output)
+{
+	DECLARE_SCOPE_HIERARCHICAL_COUNTER_FUNC()
+
+	FPoseContext SourcePose(Output);
+
+	if (Source.GetLinkNode())
+	{
+		Source.Evaluate(SourcePose);
+	}
+	else
+	{
+		SourcePose.ResetToRefPose();
+	}
+
+	// This is just a simple test to execute the graph and obtain a float result
+	// We have to create the Pose API, in order to be able to pass / receive the animation data
+	if (AnimNextInterfaceGraph->GetReturnTypeName() == TNameOf<float>::GetName())
+	{
+		UE::AnimNext::Interface::FState RootState(1); // TODO : Get rid of num elements ? Let the param itself have it
+		UE::AnimNext::Interface::FParamStorage ParamStorage(1, 16, 1);	// TODO : see if we can have state and params together and optional. Also, default size
+		UE::AnimNext::Interface::FContext RootContext(0.f, RootState, ParamStorage);
+		TScriptInterface<IAnimNextInterface> ScriptInterface = AnimNextInterfaceGraph;
+
+		float ReturnValue = 0.f;
+		const bool bOk = UE::AnimNext::Interface::GetDataSafe(ScriptInterface, RootContext, ReturnValue);
+	}
+
+	// TODO
+	//ExecuteAnimNextInterface (SourcePose);
+
+	Output = SourcePose;
+
+	// evaluate 
+	FAnimNode_CustomProperty::Evaluate_AnyThread(Output);
+}
+
+void FAnimNode_AnimNextInterfaceGraph::PostSerialize(const FArchive& Ar)
+{
+	DECLARE_SCOPE_HIERARCHICAL_COUNTER_FUNC()
+
+	// after compile, we have to reinitialize
+	// because it needs new execution code
+	// since memory has changed
+	if (Ar.IsObjectReferenceCollector())
+	{
+		if (AnimNextInterfaceGraph)
+		{
+			//AnimNextInterfaceGraph->Initialize();
+		}
+	}
+}
+
+void FAnimNode_AnimNextInterfaceGraph::InitializeProperties(const UObject* InSourceInstance, UClass* InTargetClass)
+{
+	// Build property lists
+	SourceProperties.Reset(SourcePropertyNames.Num());
+	DestProperties.Reset(SourcePropertyNames.Num());
+
+	check(SourcePropertyNames.Num() == DestPropertyNames.Num());
+
+	for (int32 Idx = 0; Idx < SourcePropertyNames.Num(); ++Idx)
+	{
+		FName& SourceName = SourcePropertyNames[Idx];
+		UClass* SourceClass = InSourceInstance->GetClass();
+
+		FProperty* SourceProperty = FindFProperty<FProperty>(SourceClass, SourceName);
+		SourceProperties.Add(SourceProperty);
+		DestProperties.Add(nullptr);
+	}
+}
+
+void FAnimNode_AnimNextInterfaceGraph::PropagateInputProperties(const UObject* InSourceInstance)
+{
+	if (InSourceInstance)
+	{
+		// Assign value to the properties exposed as pins
+		for (int32 PropIdx = 0; PropIdx < SourceProperties.Num(); ++PropIdx)
+		{
+		}
+	}
+}
+
+
+#if WITH_EDITOR
+
+void FAnimNode_AnimNextInterfaceGraph::HandleObjectsReinstanced_Impl(UObject* InSourceObject, UObject* InTargetObject, const TMap<UObject*, UObject*>& OldToNewInstanceMap)
+{
+	FAnimNode_CustomProperty::HandleObjectsReinstanced_Impl(InSourceObject, InTargetObject, OldToNewInstanceMap);
+}
+
+#endif
diff --git a/Engine/Plugins/Experimental/AnimNext/Source/AnimNextInterfaceGraph/Private/AnimNextInterfaceGraph.cpp b/Engine/Plugins/Experimental/AnimNext/Source/AnimNextInterfaceGraph/Private/AnimNextInterfaceGraph.cpp
--- a/Engine/Plugins/Experimental/AnimNext/Source/AnimNextInterfaceGraph/Private/AnimNextInterfaceGraph.cpp
+++ b/Engine/Plugins/Experimental/AnimNext/Source/AnimNextInterfaceGraph/Private/AnimNextInterfaceGraph.cpp
@@ -4,6 +4,8 @@
 #include "RigVMCore/RigVMMemoryStorage.h"
 #include "RigUnit_AnimNextInterfaceBeginExecution.h"
 #include "AnimNextInterfaceUnitContext.h"
+#include "UObject/ObjectSaveContext.h"
+#include "UObject/Package.h"
 
 #include UE_INLINE_GENERATED_CPP_BY_NAME(AnimNextInterfaceGraph)
 
@@ -15,14 +17,12 @@ const FName ResultName("Result");
 
 FName UAnimNextInterfaceGraph::GetReturnTypeNameImpl() const
 {
-	// TODO: this needs to be user-defined & determined at creation time, using float for now
-	return TNameOf<float>::GetName();
+	return ReturnTypeName;
 }
 
 const UScriptStruct* UAnimNextInterfaceGraph::GetReturnTypeStructImpl() const
 {
-	// TODO: this needs to be user-defined & determined at creation time, using nullptr (for float) for now
-	return nullptr;
+	return ReturnTypeStruct;
 }
 
 bool UAnimNextInterfaceGraph::GetDataImpl(const UE::AnimNext::Interface::FContext& Context) const
@@ -31,10 +31,11 @@ bool UAnimNextInterfaceGraph::GetDataImpl(const UE::AnimNext::Interface::FContex
 	
 	if(RigVM)
 	{
-		FRigUnitContext RigUnitContext;
-		FAnimNextInterfaceUnitContext AnimNextInterfaceUnitContext(this, Context, bResult);
-		bResult &=
-			(RigVM->Execute(TArray<URigVMMemoryStorage*>(), FRigUnit_AnimNextInterfaceBeginExecution::EventName) != ERigVMExecuteResult::Failed);
+		FAnimNextInterfaceExecuteContext& AnimNextInterfaceContext = RigVM->GetContext().GetPublicDataSafe<FAnimNextInterfaceExecuteContext>();
+
+		AnimNextInterfaceContext.SetContextData(this, Context, bResult);
+
+		bResult &= (RigVM->Execute(TArray<URigVMMemoryStorage*>(), FRigUnit_AnimNextInterfaceBeginExecution::EventName) != ERigVMExecuteResult::Failed);
 	}
 	
 	return bResult;
@@ -44,3 +45,56 @@ TArray<FRigVMExternalVariable> UAnimNextInterfaceGraph::GetRigVMExternalVariable
 {
 	return TArray<FRigVMExternalVariable>(); 
 }
+
+static TArray<UClass*> GetClassObjectsInPackage(UPackage* InPackage)
+{
+	TArray<UObject*> Objects;
+	GetObjectsWithOuter(InPackage, Objects, false);
+
+	TArray<UClass*> ClassObjects;
+	for (UObject* Object : Objects)
+	{
+		if (UClass* Class = Cast<UClass>(Object))
+		{
+			ClassObjects.Add(Class);
+		}
+	}
+
+	return ClassObjects;
+}
+
+void UAnimNextInterfaceGraph::PostRename(UObject* OldOuter, const FName OldName)
+{
+	Super::PostRename(OldOuter, OldName);
+
+	// Whenever the asset is renamed/moved, generated classes parented to the old package
+	// are not moved to the new package automatically (see FAssetRenameManager), so we
+	// have to manually perform the move/rename, to avoid invalid reference to the old package
+
+	// Note: while asset duplication doesn't duplicate the classes either, it is not a problem there
+	// because we always recompile in post duplicate.
+	TArray<UClass*> ClassObjects = GetClassObjectsInPackage(OldOuter->GetPackage());
+
+	for (UClass* ClassObject : ClassObjects)
+	{
+		if (URigVMMemoryStorageGeneratorClass* MemoryClass = Cast<URigVMMemoryStorageGeneratorClass>(ClassObject))
+		{
+			MemoryClass->Rename(nullptr, GetPackage(), REN_ForceNoResetLoaders | REN_DoNotDirty | REN_DontCreateRedirectors | REN_NonTransactional);
+		}
+	}
+}
+
+void UAnimNextInterfaceGraph::GetPreloadDependencies(TArray<UObject*>& OutDeps)
+{
+	Super::GetPreloadDependencies(OutDeps);
+
+	TArray<UClass*> ClassObjects = GetClassObjectsInPackage(GetPackage());
+
+	for (UClass* ClassObject : ClassObjects)
+	{
+		if (URigVMMemoryStorageGeneratorClass* MemoryClass = Cast<URigVMMemoryStorageGeneratorClass>(ClassObject))
+		{
+			OutDeps.Add(MemoryClass);
+		}
+	}
+}
diff --git a/Engine/Plugins/Experimental/AnimNext/Source/AnimNextInterfaceGraph/Private/RigUnit_AnimNextInterfaceParameter.cpp b/Engine/Plugins/Experimental/AnimNext/Source/AnimNextInterfaceGraph/Private/RigUnit_AnimNextInterfaceParameter.cpp
--- a/Engine/Plugins/Experimental/AnimNext/Source/AnimNextInterfaceGraph/Private/RigUnit_AnimNextInterfaceParameter.cpp
+++ b/Engine/Plugins/Experimental/AnimNext/Source/AnimNextInterfaceGraph/Private/RigUnit_AnimNextInterfaceParameter.cpp
@@ -9,19 +9,19 @@
 
 #include UE_INLINE_GENERATED_CPP_BY_NAME(RigUnit_AnimNextInterfaceParameter)
 
-bool FRigUnit_AnimNextInterfaceParameter::GetParameterInternal(FName InName, const FAnimNextInterfaceUnitContext& InContext, void* OutResult)
+bool FRigUnit_AnimNextInterfaceParameter::GetParameterInternal(FName InName, const FAnimNextInterfaceExecuteContext& InContext, void* OutResult)
 {
 	return true;
 }
 
 FRigUnit_AnimNextInterfaceParameter_Float_Execute()
 {
-	GetParameterInternal(Parameter, ExecuteContext.GetUnitContext(), &Result);
+	GetParameterInternal(Parameter, ExecuteContext, &Result);
 }
 
 FRigUnit_AnimNextInterfaceParameter_AnimNextInterface_Execute()
 {
-	GetParameterInternal(Parameter, ExecuteContext.GetUnitContext(), &Result);
+	GetParameterInternal(Parameter, ExecuteContext, &Result);
 }
 
 FRigUnit_AnimNextInterface_Float_Execute()
diff --git a/Engine/Plugins/Experimental/AnimNext/Source/AnimNextInterfaceGraph/Private/RigUnit_AnimNextInterfaceParameter.h b/Engine/Plugins/Experimental/AnimNext/Source/AnimNextInterfaceGraph/Private/RigUnit_AnimNextInterfaceParameter.h
--- a/Engine/Plugins/Experimental/AnimNext/Source/AnimNextInterfaceGraph/Private/RigUnit_AnimNextInterfaceParameter.h
+++ b/Engine/Plugins/Experimental/AnimNext/Source/AnimNextInterfaceGraph/Private/RigUnit_AnimNextInterfaceParameter.h
@@ -17,7 +17,7 @@ struct FRigUnit_AnimNextInterfaceParameter : public FRigUnit_AnimNextInterfaceBa
 	GENERATED_BODY()
 
 protected:
-	static bool GetParameterInternal(FName InName, const FAnimNextInterfaceUnitContext& InContext, void* OutResult);
+	static bool GetParameterInternal(FName InName, const FAnimNextInterfaceExecuteContext& InContext, void* OutResult);
 
 	UPROPERTY(EditAnywhere, Category = "Parameter", meta = (Input))
 	FName Parameter = NAME_None;
@@ -37,8 +37,8 @@ protected:
 	float Result = 0.0f;
 };
 
-/** Unit for reading data interface graph parameter from context */
-USTRUCT(meta = (DisplayName = "Data Interface Parameter", Category = "Parameters", TitleColor = "1 0 0", NodeColor = "1 1 1"))
+/** Unit for reading anim interface graph parameter from context */
+USTRUCT(meta = (DisplayName = "Anim Interface Parameter", Category = "Parameters", TitleColor = "1 0 0", NodeColor = "1 1 1"))
 struct FRigUnit_AnimNextInterfaceParameter_AnimNextInterface : public FRigUnit_AnimNextInterfaceParameter
 {
 	GENERATED_BODY()
@@ -60,7 +60,7 @@ struct FAnimNextInterfaceParameter
 	: Name(NAME_None)
 	{}
 
-	UPROPERTY(EditAnywhere, Category = "Data Interface")
+	UPROPERTY(EditAnywhere, Category = "Anim Interface")
 	FName Name;
 };
 
@@ -74,10 +74,10 @@ struct FAnimNextInterfaceParameter_Float
 		, Value(0.f)
 	{}
 
-	UPROPERTY(EditAnywhere, Category = "Data Interface")
+	UPROPERTY(EditAnywhere, Category = "Anim Interface")
 	FName Name;
 
-	UPROPERTY(EditAnywhere, Category = "Data Interface")
+	UPROPERTY(EditAnywhere, Category = "Anim Interface")
 	float Value;
 };
 
@@ -91,10 +91,10 @@ struct FAnimNextInterfaceParameter_Int
 		, Value(0)
 	{}
 
-	UPROPERTY(EditAnywhere, Category = "Data Interface")
+	UPROPERTY(EditAnywhere, Category = "Anim Interface")
 	FName Name;
 
-	UPROPERTY(EditAnywhere, Category = "Data Interface")
+	UPROPERTY(EditAnywhere, Category = "Anim Interface")
 	int32 Value;
 };
 
@@ -108,10 +108,10 @@ struct FAnimNextInterfaceParameter_Bool
 		, bValue(false)
 	{}
 
-	UPROPERTY(EditAnywhere, Category = "Data Interface")
+	UPROPERTY(EditAnywhere, Category = "Anim Interface")
 	FName Name;
 
-	UPROPERTY(EditAnywhere, Category = "Data Interface")
+	UPROPERTY(EditAnywhere, Category = "Anim Interface")
 	bool bValue;
 };
 
@@ -125,10 +125,10 @@ struct FAnimNextInterfaceParameters1
 		, FloatParam(0.f)
 	{}
 
-	UPROPERTY(EditAnywhere, Category = "Data Interface")
+	UPROPERTY(EditAnywhere, Category = "Anim Interface")
 	FAnimNextInterfaceParameter_Int Param0;
 
-	UPROPERTY(EditAnywhere, Category = "Data Interface")
+	UPROPERTY(EditAnywhere, Category = "Anim Interface")
 	float FloatParam;
 };
 
@@ -137,23 +137,23 @@ struct FAnimNextInterfaceParameter_AnimNextInterface : public FAnimNextInterface
 {
 	GENERATED_BODY()
 
-	UPROPERTY(EditAnywhere, Category = "Data Interface")
+	UPROPERTY(EditAnywhere, Category = "Anim Interface")
 	TScriptInterface<IAnimNextInterface> Value = nullptr;
 };
 
-/** Base unit for calling data interfaces from graphs */
-USTRUCT(meta = (DisplayName = "Data Interface", Category = "Execution", TitleColor = "1 0 0", NodeColor = "1 1 1"))
+/** Base unit for calling anim interfaces from graphs */
+USTRUCT(meta = (DisplayName = "Anim Interface", Category = "Execution", TitleColor = "1 0 0", NodeColor = "1 1 1"))
 struct FRigUnit_AnimNextInterface : public FRigUnit_AnimNextInterfaceBase
 {
 	GENERATED_BODY()
 
 protected:
-	UPROPERTY(EditAnywhere, Category = "Data Interface", meta = (Input))
+	UPROPERTY(EditAnywhere, Category = "Anim Interface", meta = (Input))
 	TScriptInterface<IAnimNextInterface> AnimNextInterface = nullptr;
 };
 
-/** Unit for getting a float via a data interface */
-USTRUCT(meta = (DisplayName = "Data Interface Float", Category = "Parameters", TitleColor = "1 0 0", NodeColor = "1 1 1"))
+/** Unit for getting a float via an anim interface */
+USTRUCT(meta = (DisplayName = "Anim Interface Float", Category = "Parameters", TitleColor = "1 0 0", NodeColor = "1 1 1"))
 struct FRigUnit_AnimNextInterface_Float : public FRigUnit_AnimNextInterface
 {
 	GENERATED_BODY()
@@ -162,14 +162,14 @@ struct FRigUnit_AnimNextInterface_Float : public FRigUnit_AnimNextInterface
 	void Execute();
 	
 protected:
-	UPROPERTY(EditAnywhere, Category = "Data Interface", meta = (Input))
+	UPROPERTY(EditAnywhere, Category = "Anim Interface", meta = (Input))
 	TArray<FAnimNextInterfaceParameter_AnimNextInterface> Parameters;
 	
 	UPROPERTY(EditAnywhere, Category = "Result", meta = (Output))
 	float Result = 0.0f;
 };
 
-/** Unit for getting a pose via a data interface */
+/** Unit for getting a pose via an anim interface */
 /*
 USTRUCT(meta=(DisplayName="Get Pose"))
 struct FRigUnit_AnimNextInterface_Pose : public FRigUnit
@@ -190,7 +190,7 @@ protected:
 	FAnimNextInterfaceExecuteContext Result;
 };
 */
-/** Unit for getting a pose via a data interface */
+/** Unit for getting a pose via an anim interface */
 USTRUCT(meta = (DisplayName = "Float Operator", Category = "Operators", TitleColor = "1 0 0", NodeColor = "1 1 1"))
 struct FRigUnit_FloatOperator : public FRigUnit_AnimNextInterfaceBase
 {
@@ -219,7 +219,7 @@ protected:
 	float Result;
 };
 
-/** Unit for getting a pose via a data interface */
+/** Unit for getting a pose via an anim interface */
 USTRUCT(meta = (DisplayName = "Pose Operator", Category = "Operators", TitleColor = "1 0 0", NodeColor = "1 1 1"))
 struct FRigUnit_PoseOperator : public FRigUnit_AnimNextInterfaceBase
 {
@@ -257,7 +257,7 @@ struct FAnimSequenceParameters
 	bool bLoop = 0.0f;
 };
 
-/** Unit for getting a pose via a data interface */
+/** Unit for getting a pose via an anim interface */
 USTRUCT(meta = (DisplayName = "Anim Sequence", Category = "Animation", TitleColor = "1 0 0", NodeColor = "1 1 1"))
 struct FRigUnit_AnimNextInterface_SequencePlayer : public FRigUnit_AnimNextInterfaceBase
 {
@@ -267,17 +267,17 @@ struct FRigUnit_AnimNextInterface_SequencePlayer : public FRigUnit_AnimNextInter
 	void Execute();
 	
 protected:
-	UPROPERTY(EditAnywhere, Category = "Data Interface", meta = (Input))
+	UPROPERTY(EditAnywhere, Category = "Anim Interface", meta = (Input))
 	FAnimSequenceParameters Parameters;
 
-	UPROPERTY(EditAnywhere, Category = "Data Interface", meta = (Input))
+	UPROPERTY(EditAnywhere, Category = "Anim Interface", meta = (Input))
 	TScriptInterface<IAnimNextInterface> Sequence = nullptr;
 	
 	UPROPERTY(EditAnywhere, Transient, DisplayName = "Result", Category = "BeginExecution", meta = (Output))
 	FAnimNextInterfaceExecuteContext Result;
 };
 
-/** Unit for getting a pose via a data interface */
+/** Unit for getting a pose via an anim interface */
 USTRUCT(meta = (DisplayName = "Test Float State - Spring Damper Smoothing", Varying, Category = "Animation", TitleColor = "1 0 0", NodeColor = "1 1 1"))
 struct FRigUnit_TestFloatState : public FRigUnit_AnimNextInterfaceBase
 {
diff --git a/Engine/Plugins/Experimental/AnimNext/Source/AnimNextInterfaceGraphEditor/Private/AnimNextInterfaceGraphEditor.cpp b/Engine/Plugins/Experimental/AnimNext/Source/AnimNextInterfaceGraphEditor/Private/AnimNextInterfaceGraphEditor.cpp
--- a/Engine/Plugins/Experimental/AnimNext/Source/AnimNextInterfaceGraphEditor/Private/AnimNextInterfaceGraphEditor.cpp
+++ b/Engine/Plugins/Experimental/AnimNext/Source/AnimNextInterfaceGraphEditor/Private/AnimNextInterfaceGraphEditor.cpp
@@ -63,9 +63,9 @@ void FGraphEditor::InitEditor(const EToolkitMode::Type InMode, const TSharedPtr<
 	RegenerateMenusAndToolbars();
 
 	// Open initial document
-	DocumentManager->OpenDocument(FTabPayload_UObject::Make(AnimNextInterfaceGraph_EditorData->EntryPointGraph), FDocumentTracker::EOpenDocumentCause::OpenNewDocument);
+	//DocumentManager->OpenDocument(FTabPayload_UObject::Make(AnimNextInterfaceGraph_EditorData->EntryPointGraph), FDocumentTracker::EOpenDocumentCause::OpenNewDocument);
 	DocumentManager->OpenDocument(FTabPayload_UObject::Make(AnimNextInterfaceGraph_EditorData->RootGraph), FDocumentTracker::EOpenDocumentCause::OpenNewDocument);
-	DocumentManager->OpenDocument(FTabPayload_UObject::Make(AnimNextInterfaceGraph_EditorData->FunctionLibraryEdGraph), FDocumentTracker::EOpenDocumentCause::OpenNewDocument);
+	//DocumentManager->OpenDocument(FTabPayload_UObject::Make(AnimNextInterfaceGraph_EditorData->FunctionLibraryEdGraph), FDocumentTracker::EOpenDocumentCause::OpenNewDocument);
 }
 
 void FGraphEditor::BindCommands()
diff --git a/Engine/Plugins/Experimental/AnimNext/Source/AnimNextInterfaceGraphEditor/Private/AnimNextInterfaceGraphFactory.cpp b/Engine/Plugins/Experimental/AnimNext/Source/AnimNextInterfaceGraphEditor/Private/AnimNextInterfaceGraphFactory.cpp
--- a/Engine/Plugins/Experimental/AnimNext/Source/AnimNextInterfaceGraphEditor/Private/AnimNextInterfaceGraphFactory.cpp
+++ b/Engine/Plugins/Experimental/AnimNext/Source/AnimNextInterfaceGraphEditor/Private/AnimNextInterfaceGraphFactory.cpp
@@ -42,54 +42,54 @@ UObject* UAnimNextInterfaceGraphFactory::FactoryCreateNew(UClass* Class, UObject
 	check(BeginExecutePin);
 	check(BeginExecutePin->GetDirection() == ERigVMPinDirection::Output);
 
-	// Add function to function lib
-	{
-		FRigVMControllerGraphGuard LibraryGraphGuard(Controller, EditorData->GetRigVMClient()->GetFunctionLibrary());
-		EditorData->EntryPoint = Controller->AddFunctionToLibrary(UE::AnimNext::InterfaceGraph::EntryPointName, true, FVector2D::ZeroVector, false, false);
-
-		// Add exposed pin for result to function
-		{
-			FRigVMControllerGraphGuard FunctionGraphGuard(Controller, EditorData->EntryPoint->GetContainedGraph());
-
-			// TODO: using float for now, but needs to be user-driven
-			const FString CPPType = TEXT("float");
-			const FName CPPTypeObjectPath = NAME_None;
-			const FString DefaultValue = TEXT("0.0");
-
-			//FunctionController->AddExposedPin(UE::AnimNext::InterfaceGraph::ResultName, ERigVMPinDirection::Output, CPPType, CPPTypeObjectPath, DefaultValue, false, false);
-			Controller->AddExposedPin(UE::AnimNext::InterfaceGraph::ResultName, ERigVMPinDirection::Output, CPPType, CPPTypeObjectPath, DefaultValue, false, false);
-		}
-	}
-
-	// Add function call
-	URigVMFunctionReferenceNode* FunctionCallNode = Controller->AddFunctionReferenceNode(EditorData->EntryPoint, FVector2D::ZeroVector, UE::AnimNext::InterfaceGraph::EntryPointName.ToString(), false, false);
-	URigVMPin* FunctionExecutePin = FunctionCallNode->FindPin(FRigVMStruct::ExecuteContextName.ToString());
-	check(FunctionExecutePin);
-	check(FunctionExecutePin->GetDirection() == ERigVMPinDirection::IO);
-
-	// Link entry point to function call
-	Controller->AddLink(BeginExecutePin, FunctionExecutePin, false);
-
-	URigVMPin* FunctionResultPin = FunctionCallNode->FindPin(UE::AnimNext::InterfaceGraph::ResultName.ToString());
-	check(FunctionResultPin);
-	check(FunctionResultPin->GetDirection() == ERigVMPinDirection::Output);
-
-	// Add end-execution unit of the correct type
-	// TODO: using a float for now, but need to use a registry to determine correct type
-	URigVMUnitNode* MainExitPointNode = Controller->AddUnitNode(FRigUnit_AnimNextInterfaceEndExecution_Float::StaticStruct(), FRigUnit::GetMethodName(), FVector2D(400.0f, 0.0f), FString(), false);
-	URigVMPin* EndExecutePin = MainExitPointNode->FindPin(GET_MEMBER_NAME_STRING_CHECKED(FRigUnit_AnimNextInterfaceEndExecution, ExecuteContext));
-	check(EndExecutePin);
-	check(EndExecutePin->GetDirection() == ERigVMPinDirection::Input);
-
-	URigVMPin* EndExecuteResultPin = MainExitPointNode->FindPin(GET_MEMBER_NAME_STRING_CHECKED(FRigUnit_AnimNextInterfaceEndExecution_Float, Result));
-	check(EndExecuteResultPin);
-	check(EndExecuteResultPin->GetDirection() == ERigVMPinDirection::Input);
-
-	// Link function call to exit point 
-	Controller->AddLink(FunctionExecutePin, EndExecutePin, false);
-
-	// Link result pins 
-	Controller->AddLink(FunctionResultPin, EndExecuteResultPin, false);
+	//// Add function to function lib
+	//{
+	//	FRigVMControllerGraphGuard LibraryGraphGuard(Controller, EditorData->GetRigVMClient()->GetFunctionLibrary());
+	//	EditorData->EntryPoint = Controller->AddFunctionToLibrary(UE::AnimNext::InterfaceGraph::EntryPointName, true, FVector2D::ZeroVector, false, false);
+
+	//	// Add exposed pin for result to function
+	//	{
+	//		FRigVMControllerGraphGuard FunctionGraphGuard(Controller, EditorData->EntryPoint->GetContainedGraph());
+
+	//		// TODO: using float for now, but needs to be user-driven
+	//		const FString CPPType = TEXT("float");
+	//		const FName CPPTypeObjectPath = NAME_None;
+	//		const FString DefaultValue = TEXT("0.0");
+
+	//		//FunctionController->AddExposedPin(UE::AnimNext::InterfaceGraph::ResultName, ERigVMPinDirection::Output, CPPType, CPPTypeObjectPath, DefaultValue, false, false);
+	//		Controller->AddExposedPin(UE::AnimNext::InterfaceGraph::ResultName, ERigVMPinDirection::Output, CPPType, CPPTypeObjectPath, DefaultValue, false, false);
+	//	}
+	//}
+
+	//// Add function call
+	//URigVMFunctionReferenceNode* FunctionCallNode = Controller->AddFunctionReferenceNode(EditorData->EntryPoint, FVector2D::ZeroVector, UE::AnimNext::InterfaceGraph::EntryPointName.ToString(), false, false);
+	//URigVMPin* FunctionExecutePin = FunctionCallNode->FindPin(FRigVMStruct::ExecuteContextName.ToString());
+	//check(FunctionExecutePin);
+	//check(FunctionExecutePin->GetDirection() == ERigVMPinDirection::IO);
+
+	//// Link entry point to function call
+	//Controller->AddLink(BeginExecutePin, FunctionExecutePin, false);
+
+	//URigVMPin* FunctionResultPin = FunctionCallNode->FindPin(UE::AnimNext::InterfaceGraph::ResultName.ToString());
+	//check(FunctionResultPin);
+	//check(FunctionResultPin->GetDirection() == ERigVMPinDirection::Output);
+
+	//// Add end-execution unit of the correct type
+	//// TODO: using a float for now, but need to use a registry to determine correct type
+	//URigVMUnitNode* MainExitPointNode = Controller->AddUnitNode(FRigUnit_AnimNextInterfaceEndExecution_Float::StaticStruct(), FRigUnit::GetMethodName(), FVector2D(400.0f, 0.0f), FString(), false);
+	//URigVMPin* EndExecutePin = MainExitPointNode->FindPin(GET_MEMBER_NAME_STRING_CHECKED(FRigUnit_AnimNextInterfaceEndExecution, ExecuteContext));
+	//check(EndExecutePin);
+	//check(EndExecutePin->GetDirection() == ERigVMPinDirection::Input);
+
+	//URigVMPin* EndExecuteResultPin = MainExitPointNode->FindPin(GET_MEMBER_NAME_STRING_CHECKED(FRigUnit_AnimNextInterfaceEndExecution_Float, Result));
+	//check(EndExecuteResultPin);
+	//check(EndExecuteResultPin->GetDirection() == ERigVMPinDirection::Input);
+
+	//// Link function call to exit point 
+	//Controller->AddLink(FunctionExecutePin, EndExecutePin, false);
+
+	//// Link result pins 
+	//Controller->AddLink(FunctionResultPin, EndExecuteResultPin, false);
 
 	// Compile the initial skeleton
 	UE::AnimNext::InterfaceGraphUncookedOnly::FUtils::Compile(NewGraph);
diff --git a/Engine/Plugins/Experimental/AnimNext/Source/AnimNextInterfaceGraphEditor/Private/AnimNextInterfaceWidgetFactories.cpp b/Engine/Plugins/Experimental/AnimNext/Source/AnimNextInterfaceGraphEditor/Private/AnimNextInterfaceWidgetFactories.cpp
--- a/Engine/Plugins/Experimental/AnimNext/Source/AnimNextInterfaceGraphEditor/Private/AnimNextInterfaceWidgetFactories.cpp
+++ b/Engine/Plugins/Experimental/AnimNext/Source/AnimNextInterfaceGraphEditor/Private/AnimNextInterfaceWidgetFactories.cpp
@@ -53,7 +53,7 @@ TSharedPtr<SWidget> FAnimNextInterfaceWidgetFactories::CreateAnimNextInterfaceWi
 	}
 
 
-	// button for replacing data with a different Data Interface class
+	// button for replacing data with a different AnimNext Interface class
 	TSharedPtr<SComboButton> Button = SNew(SComboButton)
 			.ComboButtonStyle(FAppStyle::Get(), "SimpleComboButton");
 	
diff --git a/Engine/Plugins/Experimental/AnimNext/Source/AnimNextInterfaceGraphEditor/Private/AssetTypeActions_AnimNextInterfaceGraph.cpp b/Engine/Plugins/Experimental/AnimNext/Source/AnimNextInterfaceGraphEditor/Private/AssetTypeActions_AnimNextInterfaceGraph.cpp
--- a/Engine/Plugins/Experimental/AnimNext/Source/AnimNextInterfaceGraphEditor/Private/AssetTypeActions_AnimNextInterfaceGraph.cpp
+++ b/Engine/Plugins/Experimental/AnimNext/Source/AnimNextInterfaceGraphEditor/Private/AssetTypeActions_AnimNextInterfaceGraph.cpp
@@ -27,7 +27,7 @@ const TArray<FText>& FAssetTypeActions_AnimNextInterfaceGraph::GetSubMenus() con
 {
 	static const TArray<FText> SubMenus
 	{
-		LOCTEXT("AnimNextInterfacesSubMenu", "Data Interfaces")
+		LOCTEXT("AnimNextInterfacesSubMenu", "AnimNext")
 	};
 	return SubMenus;
 }
diff --git a/Engine/Plugins/Experimental/AnimNext/Source/AnimNextInterfaceGraphEditor/Private/AssetTypeActions_AnimNextInterfaceGraph.h b/Engine/Plugins/Experimental/AnimNext/Source/AnimNextInterfaceGraphEditor/Private/AssetTypeActions_AnimNextInterfaceGraph.h
--- a/Engine/Plugins/Experimental/AnimNext/Source/AnimNextInterfaceGraphEditor/Private/AssetTypeActions_AnimNextInterfaceGraph.h
+++ b/Engine/Plugins/Experimental/AnimNext/Source/AnimNextInterfaceGraphEditor/Private/AssetTypeActions_AnimNextInterfaceGraph.h
@@ -12,7 +12,7 @@ class FAssetTypeActions_AnimNextInterfaceGraph : public FAssetTypeActions_Base
 {
 public:
 	// IAssetTypeActions Implementation
-	virtual FText GetName() const override { return NSLOCTEXT("AssetTypeActions", "AssetTypeActions_AnimNextInterfaceGraph", "Data Interface Graph"); }
+	virtual FText GetName() const override { return NSLOCTEXT("AssetTypeActions", "AssetTypeActions_AnimNextInterfaceGraph", "AnimNext Graph"); }
 	virtual FColor GetTypeColor() const override { return FColor(128,128,64); }
 	virtual UClass* GetSupportedClass() const override { return UAnimNextInterfaceGraph::StaticClass(); }
 	virtual void OpenAssetEditor(const TArray<UObject*>& InObjects, TSharedPtr<IToolkitHost> EditWithinLevelEditor = TSharedPtr<IToolkitHost>()) override;
diff --git a/Engine/Plugins/Experimental/AnimNext/Source/AnimNextInterfaceGraphEditor/Private/SAnimNextInterfaceGraphEditorActionMenu.cpp b/Engine/Plugins/Experimental/AnimNext/Source/AnimNextInterfaceGraphEditor/Private/SAnimNextInterfaceGraphEditorActionMenu.cpp
--- a/Engine/Plugins/Experimental/AnimNext/Source/AnimNextInterfaceGraphEditor/Private/SAnimNextInterfaceGraphEditorActionMenu.cpp
+++ b/Engine/Plugins/Experimental/AnimNext/Source/AnimNextInterfaceGraphEditor/Private/SAnimNextInterfaceGraphEditorActionMenu.cpp
@@ -104,7 +104,7 @@ void SActionMenu::Construct(const FArguments& InArgs)
 					.VAlign(VAlign_Center)
 					[
 						SNew(STextBlock)
-						.Text(LOCTEXT("ContextText", "All Data Interface Node Classes"))
+						.Text(LOCTEXT("ContextText", "All AnimNext Interface Node Classes"))
 						.Font(FAppStyle::Get().GetFontStyle("BlueprintEditor.ActionMenu.ContextDescriptionFont"))
 						.ToolTip(IDocumentation::Get()->CreateToolTip(
 							LOCTEXT("ActionMenuContextTextTooltip", "Describes the current context of the action list"),
diff --git a/Engine/Plugins/Experimental/AnimNext/Source/AnimNextInterfaceGraphUncookedOnly/AnimNextInterfaceGraphUncookedOnly.Build.cs b/Engine/Plugins/Experimental/AnimNext/Source/AnimNextInterfaceGraphUncookedOnly/AnimNextInterfaceGraphUncookedOnly.Build.cs
--- a/Engine/Plugins/Experimental/AnimNext/Source/AnimNextInterfaceGraphUncookedOnly/AnimNextInterfaceGraphUncookedOnly.Build.cs
+++ b/Engine/Plugins/Experimental/AnimNext/Source/AnimNextInterfaceGraphUncookedOnly/AnimNextInterfaceGraphUncookedOnly.Build.cs
@@ -26,6 +26,11 @@ namespace UnrealBuildTool.Rules
 					"AnimNextInterface",
 					"AnimNextInterfaceGraph",
 					"BlueprintGraph",	// For K2 schema
+					"AnimationCore",
+					"AnimGraph",
+					"Kismet",
+					"Slate",
+					"SlateCore",
 				}
 			);
 
diff --git a/Engine/Plugins/Experimental/AnimNext/Source/AnimNextInterfaceGraphUncookedOnly/Internal/AnimGraph/AnimGraphNode_AnimNextInterfaceGraph.h b/Engine/Plugins/Experimental/AnimNext/Source/AnimNextInterfaceGraphUncookedOnly/Internal/AnimGraph/AnimGraphNode_AnimNextInterfaceGraph.h
new file mode 100644
--- /dev/null
+++ b/Engine/Plugins/Experimental/AnimNext/Source/AnimNextInterfaceGraphUncookedOnly/Internal/AnimGraph/AnimGraphNode_AnimNextInterfaceGraph.h
@@ -0,0 +1,54 @@
+// Copyright Epic Games, Inc. All Rights Reserved.
+
+#pragma once
+
+#include "AnimGraphNode_CustomProperty.h"
+#include "AnimGraph/AnimNode_AnimNextInterfaceGraph.h"
+#include "SSearchableComboBox.h"
+#include "AnimGraphNode_AnimNextInterfaceGraph.generated.h"
+
+struct FVariableMappingInfo;
+
+UCLASS(MinimalAPI)
+class UAnimGraphNode_AnimNextInterfaceGraph : public UAnimGraphNode_CustomProperty
+{
+	GENERATED_UCLASS_BODY()
+
+	UPROPERTY(EditAnywhere, Category = Settings)
+	FAnimNode_AnimNextInterfaceGraph Node;
+
+	virtual void CustomizeDetails(IDetailLayoutBuilder& DetailBuilder) override;
+	virtual void PostEditChangeProperty(FPropertyChangedEvent& PropertyChangedEvent) override;
+	virtual void PostReconstructNode() override;
+
+private:
+	// UEdGraphNode interface
+	virtual FText GetNodeTitle(ENodeTitleType::Type TitleType) const override;
+	virtual void CustomizePinData(UEdGraphPin* Pin, FName SourcePropertyName, int32 ArrayIndex) const override;
+	virtual FText GetTooltipText() const override;
+	virtual void PreloadRequiredAssets() override;
+	//virtual bool IsSinkNode() const { return true; }
+
+	virtual FAnimNode_CustomProperty* GetCustomPropertyNode() override { return &Node; }
+	virtual const FAnimNode_CustomProperty* GetCustomPropertyNode() const override { return &Node; }
+
+	// property related things
+	void GetVariables(bool bInput, TMap<FName, FRigVMExternalVariable>& OutParameters) const;
+
+	TMap<FName, FRigVMExternalVariable> InputVariables;
+	TMap<FName, FRigVMExternalVariable> OutputVariables;
+
+	void RebuildExposedProperties();
+	virtual void CreateCustomPins(TArray<UEdGraphPin*>* OldPins) override;
+	virtual void ValidateAnimNodeDuringCompilation(USkeleton* ForSkeleton, FCompilerResultsLog& MessageLog) override;
+	
+	// pin option related
+	bool IsPropertyExposeEnabled(FName PropertyName) const;
+	ECheckBoxState IsPropertyExposed(FName PropertyName) const;
+	void OnPropertyExposeCheckboxChanged(ECheckBoxState NewState, FName PropertyName);
+
+	bool IsInputProperty(const FName& PropertyName) const;
+
+	virtual UObject* GetJumpTargetForDoubleClick() const override;
+};
+
diff --git a/Engine/Plugins/Experimental/AnimNext/Source/AnimNextInterfaceGraphUncookedOnly/Internal/AnimNextInterfaceGraph_EditorData.h b/Engine/Plugins/Experimental/AnimNext/Source/AnimNextInterfaceGraphUncookedOnly/Internal/AnimNextInterfaceGraph_EditorData.h
--- a/Engine/Plugins/Experimental/AnimNext/Source/AnimNextInterfaceGraphUncookedOnly/Internal/AnimNextInterfaceGraph_EditorData.h
+++ b/Engine/Plugins/Experimental/AnimNext/Source/AnimNextInterfaceGraphUncookedOnly/Internal/AnimNextInterfaceGraph_EditorData.h
@@ -20,6 +20,12 @@ namespace UE::AnimNext::InterfaceGraphEditor
 	class FGraphEditor;
 }
 
+enum class EAnimNextInterfaceGraphLoadType : uint8
+{
+	PostLoad,
+	CheckUserDefinedStructs
+};
+
 UCLASS(MinimalAPI)
 class UAnimNextInterfaceGraph_EditorData : public UObject, public IRigVMClientHost, public IRigVMGraphFunctionHost
 {
@@ -37,7 +43,11 @@ class UAnimNextInterfaceGraph_EditorData : public UObject, public IRigVMClientHo
 	virtual void Serialize(FArchive& Ar) override;
 	virtual void PostLoad() override;
 	virtual bool IsEditorOnly() const override { return true; }
-	
+#if WITH_EDITOR
+	void HandlePackageDone(const FEndLoadPackageContext& Context);
+	void HandlePackageDone();
+#endif // WITH_EDITOR
+
 	// IRigVMClientHost interface
 	virtual FRigVMClient* GetRigVMClient() override;
 	virtual const FRigVMClient* GetRigVMClient() const override;
@@ -55,6 +65,8 @@ class UAnimNextInterfaceGraph_EditorData : public UObject, public IRigVMClientHo
 
 	ANIMNEXTINTERFACEGRAPHUNCOOKEDONLY_API void Initialize(bool bRecompileVM);
 
+	void RefreshAllModels(EAnimNextInterfaceGraphLoadType InLoadType);
+
 	void RecompileVM();
 	
 	void RecompileVMIfRequired();
@@ -66,7 +78,15 @@ class UAnimNextInterfaceGraph_EditorData : public UObject, public IRigVMClientHo
 	ANIMNEXTINTERFACEGRAPHUNCOOKEDONLY_API URigVMGraph* GetVMGraphForEdGraph(const UEdGraph* InGraph) const;
 
 	void CreateEdGraphForCollapseNode(URigVMCollapseNode* InNode);
-	
+
+	void UpdateGraphReturnType();
+
+	FEdGraphPinType FindGraphReturnPinType() const;
+
+	bool IsNodeExecConnected(const UEdGraphNode* Node) const;
+
+	FName GetPinTypeName(const FEdGraphPinType& EdGraphPinType);
+
 	UPROPERTY()
 	TObjectPtr<UAnimNextInterfaceGraph_EdGraph> RootGraph;
 
@@ -126,4 +146,4 @@ class UAnimNextInterfaceGraph_EditorData : public UObject, public IRigVMClientHo
 	bool bSuspendModelNotificationsForOthers = false;
 	bool bSuspendAllNotifications = false;
 	bool bCompileInDebugMode = false;
-};
\ No newline at end of file
+};
diff --git a/Engine/Plugins/Experimental/AnimNext/Source/AnimNextInterfaceGraphUncookedOnly/Private/AnimGraph/AnimGraphNode_AnimNextInterfaceGraph.cpp b/Engine/Plugins/Experimental/AnimNext/Source/AnimNextInterfaceGraphUncookedOnly/Private/AnimGraph/AnimGraphNode_AnimNextInterfaceGraph.cpp
new file mode 100644
--- /dev/null
+++ b/Engine/Plugins/Experimental/AnimNext/Source/AnimNextInterfaceGraphUncookedOnly/Private/AnimGraph/AnimGraphNode_AnimNextInterfaceGraph.cpp
@@ -0,0 +1,212 @@
+// Copyright Epic Games, Inc. All Rights Reserved.
+
+#include "AnimGraph/AnimGraphNode_AnimNextInterfaceGraph.h"
+#include "Kismet2/CompilerResultsLog.h"
+#include "DetailLayoutBuilder.h"
+#include "DetailCategoryBuilder.h"
+#include "Widgets/Text/STextBlock.h"
+#include "Widgets/Input/SCheckBox.h"
+#include "DetailWidgetRow.h"
+#include "Widgets/Input/SButton.h"
+//#include "SVariableMappingWidget.h"
+#include "ScopedTransaction.h"
+#include "Kismet2/BlueprintEditorUtils.h"
+#include "AnimationGraphSchema.h"
+#include "ControlRigBlueprintGeneratedClass.h"
+#include "ControlRigBlueprint.h"
+#include "Misc/DefaultValueHelper.h"
+#include "AnimNextInterfaceUncookedOnlyUtils.h"
+#include "AnimNextInterfaceGraph_EditorData.h"
+
+#include UE_INLINE_GENERATED_CPP_BY_NAME(AnimGraphNode_AnimNextInterfaceGraph)
+
+#define LOCTEXT_NAMESPACE "AnimGraphNode_AnimNextInterfaceGraph"
+
+UAnimGraphNode_AnimNextInterfaceGraph::UAnimGraphNode_AnimNextInterfaceGraph(const FObjectInitializer& ObjectInitializer)
+	: Super(ObjectInitializer)
+{
+}
+
+FText UAnimGraphNode_AnimNextInterfaceGraph::GetNodeTitle(ENodeTitleType::Type TitleType) const
+{
+	// display control rig here
+	return LOCTEXT("AnimGraphNode_AnimNextInterfaceGraph_Title", "AnimNextInterface Graph");
+}
+
+FText UAnimGraphNode_AnimNextInterfaceGraph::GetTooltipText() const
+{
+	// display control rig here
+	return LOCTEXT("AnimGraphNode_AnimNextInterfaceGraph_Tooltip", "Evaluates a AnimNextInterface Graph");
+}
+
+void UAnimGraphNode_AnimNextInterfaceGraph::PreloadRequiredAssets()
+{
+	Super::PreloadRequiredAssets();
+
+	if (Node.AnimNextInterfaceGraph)
+	{
+		PreloadObject(Node.AnimNextInterfaceGraph);
+		PreloadObject(Node.AnimNextInterfaceGraph->EditorData);
+	}
+}
+
+void UAnimGraphNode_AnimNextInterfaceGraph::CreateCustomPins(TArray<UEdGraphPin*>* OldPins)
+{
+	DECLARE_SCOPE_HIERARCHICAL_COUNTER_FUNC()
+	
+	// we do this to refresh input variables
+	//RebuildExposedProperties();
+
+	// Grab the SKELETON class here as when we are reconstructed during during BP compilation
+	// the full generated class is not yet present built.
+	if (!GetTargetSkeletonClass())
+	{
+		// Nothing to search for properties
+		return;
+	}
+}
+
+void UAnimGraphNode_AnimNextInterfaceGraph::ValidateAnimNodeDuringCompilation(USkeleton* ForSkeleton, FCompilerResultsLog& MessageLog)
+{
+	if (UClass* TargetClass = GetTargetClass())
+	{
+	}
+
+	Super::ValidateAnimNodeDuringCompilation(ForSkeleton, MessageLog);
+}
+
+void UAnimGraphNode_AnimNextInterfaceGraph::RebuildExposedProperties()
+{
+	DECLARE_SCOPE_HIERARCHICAL_COUNTER_FUNC()
+
+}
+
+bool UAnimGraphNode_AnimNextInterfaceGraph::IsInputProperty(const FName& PropertyName) const
+{
+	DECLARE_SCOPE_HIERARCHICAL_COUNTER_FUNC()
+
+	// this is true for both input variables and controls
+	return InputVariables.Contains(PropertyName) || !OutputVariables.Contains(PropertyName);
+}
+
+bool UAnimGraphNode_AnimNextInterfaceGraph::IsPropertyExposeEnabled(FName PropertyName) const
+{
+	DECLARE_SCOPE_HIERARCHICAL_COUNTER_FUNC()
+
+	// if known exposable, and and if it hasn't been exposed yet
+	if (CustomPinProperties.ContainsByPredicate([PropertyName](const FOptionalPinFromProperty& InOptionalPin){ return InOptionalPin.PropertyName == PropertyName; }))
+	{
+		return IsInputProperty(PropertyName);
+	}
+
+	return false;
+}
+
+ECheckBoxState UAnimGraphNode_AnimNextInterfaceGraph::IsPropertyExposed(FName PropertyName) const
+{
+	if (CustomPinProperties.ContainsByPredicate([PropertyName](const FOptionalPinFromProperty& InOptionalPin){ return InOptionalPin.bShowPin && InOptionalPin.PropertyName == PropertyName; }))
+	{
+		return ECheckBoxState::Checked;
+	}
+
+	return ECheckBoxState::Unchecked;
+}
+
+void UAnimGraphNode_AnimNextInterfaceGraph::OnPropertyExposeCheckboxChanged(ECheckBoxState NewState, FName PropertyName)
+{
+	DECLARE_SCOPE_HIERARCHICAL_COUNTER_FUNC()
+}
+
+void UAnimGraphNode_AnimNextInterfaceGraph::CustomizeDetails(IDetailLayoutBuilder& DetailBuilder)
+{
+	DECLARE_SCOPE_HIERARCHICAL_COUNTER_FUNC()
+
+	Super::CustomizeDetails(DetailBuilder);
+
+	// We dont allow multi-select here
+	if (DetailBuilder.GetSelectedObjects().Num() > 1)
+	{
+		return;
+	}
+
+	// input/output exposure feature START
+	RebuildExposedProperties();
+
+	// ***************************************************************
+
+
+	TSharedRef<IPropertyHandle> ClassHandle = DetailBuilder.GetProperty(TEXT("Node.AnimNextInterfaceGraph"), GetClass());
+	if (ClassHandle->IsValidHandle())
+	{
+		ClassHandle->SetOnPropertyValueChanged(FSimpleDelegate::CreateUObject(this, &UAnimGraphNode_AnimNextInterfaceGraph::OnInstanceClassChanged, &DetailBuilder));
+	}
+}
+
+void UAnimGraphNode_AnimNextInterfaceGraph::GetVariables(bool bInput, TMap<FName, FRigVMExternalVariable>& OutVariables) const
+{
+	DECLARE_SCOPE_HIERARCHICAL_COUNTER_FUNC()
+
+	OutVariables.Reset();
+
+	if (UControlRigBlueprintGeneratedClass* TargetClass = Cast<UControlRigBlueprintGeneratedClass>(GetTargetClass()))
+	{
+		if (UControlRigBlueprint* RigBlueprint = Cast<UControlRigBlueprint>(TargetClass->ClassGeneratedBy))
+		{
+			//RigBlueprint->CleanupVariables();
+			UControlRig* ControlRig = TargetClass->GetDefaultObject<UControlRig>();
+			if (ControlRig)
+			{
+				const TArray<FRigVMExternalVariable>& PublicVariables = ControlRig->GetPublicVariables();
+				for (const FRigVMExternalVariable& PublicVariable : PublicVariables)
+				{
+					if (!bInput || !PublicVariable.bIsReadOnly)
+					{
+						OutVariables.Add(PublicVariable.Name, PublicVariable);
+					}
+				}
+			}
+		}
+	}
+}
+
+void UAnimGraphNode_AnimNextInterfaceGraph::PostEditChangeProperty(FPropertyChangedEvent& PropertyChangedEvent)
+{
+	DECLARE_SCOPE_HIERARCHICAL_COUNTER_FUNC()
+
+	Super::PostEditChangeProperty(PropertyChangedEvent);
+
+	bool bRequiresNodeReconstruct = false;
+	FProperty* ChangedProperty = PropertyChangedEvent.Property;
+
+	if (ChangedProperty)
+	{
+		if (ChangedProperty->GetFName() == GET_MEMBER_NAME_CHECKED(FAnimNode_AnimNextInterfaceGraph, AnimNextInterfaceGraph))
+		{
+			bRequiresNodeReconstruct = true;
+			RebuildExposedProperties();
+		}
+	}
+
+	if (bRequiresNodeReconstruct)
+	{
+		ReconstructNode();
+		FBlueprintEditorUtils::MarkBlueprintAsStructurallyModified(GetBlueprint());
+	}
+}
+
+void UAnimGraphNode_AnimNextInterfaceGraph::PostReconstructNode()
+{
+	Super::PostReconstructNode();
+}
+
+void UAnimGraphNode_AnimNextInterfaceGraph::CustomizePinData(UEdGraphPin* Pin, FName SourcePropertyName, int32 ArrayIndex) const
+{
+	Super::CustomizePinData(Pin, SourcePropertyName, ArrayIndex);
+}
+
+UObject* UAnimGraphNode_AnimNextInterfaceGraph::GetJumpTargetForDoubleClick() const
+{
+	return Node.AnimNextInterfaceGraph;
+}
+
+#undef LOCTEXT_NAMESPACE
diff --git a/Engine/Plugins/Experimental/AnimNext/Source/AnimNextInterfaceGraphUncookedOnly/Private/AnimNextInterfaceGraphUncookedOnlyUtils.cpp b/Engine/Plugins/Experimental/AnimNext/Source/AnimNextInterfaceGraphUncookedOnly/Private/AnimNextInterfaceGraphUncookedOnlyUtils.cpp
--- a/Engine/Plugins/Experimental/AnimNext/Source/AnimNextInterfaceGraphUncookedOnly/Private/AnimNextInterfaceGraphUncookedOnlyUtils.cpp
+++ b/Engine/Plugins/Experimental/AnimNext/Source/AnimNextInterfaceGraphUncookedOnly/Private/AnimNextInterfaceGraphUncookedOnlyUtils.cpp
@@ -6,7 +6,7 @@
 #include "AnimNextInterfaceGraph_EdGraph.h"
 #include "RigVMCompiler/RigVMCompiler.h"
 #include "RigVMCore/RigVM.h"
-#include "Units/RigUnitContext.h"
+#include "AnimNextInterfaceExecuteContext.h"
 
 namespace UE::AnimNext::InterfaceGraphUncookedOnly
 {
@@ -49,7 +49,7 @@ void FUtils::Compile(UAnimNextInterfaceGraph* InGraph)
 	{
 		if(Compiler->Settings.SurpressErrors)
 		{
-			Compiler->Settings.Reportf(EMessageSeverity::Info, InGraph,TEXT("Compilation Errors may be suppressed for Data Interface Graph: %s. See VM Compile Settings for more Details"), *InGraph->GetName());
+			Compiler->Settings.Reportf(EMessageSeverity::Info, InGraph,TEXT("Compilation Errors may be suppressed for AnimNext Interface Graph: %s. See VM Compile Settings for more Details"), *InGraph->GetName());
 		}
 	}
 
@@ -67,6 +67,7 @@ void FUtils::Compile(UAnimNextInterfaceGraph* InGraph)
 void FUtils::RecreateVM(UAnimNextInterfaceGraph* InGraph)
 {
 	InGraph->RigVM = NewObject<URigVM>(InGraph, TEXT("VM"), RF_NoFlags);
+	InGraph->RigVM->SetContextPublicDataStruct(FAnimNextInterfaceExecuteContext::StaticStruct());
 	
 	// Cooked platforms will load these pointers from disk
 	if (!FPlatformProperties::RequiresCookedData())
diff --git a/Engine/Plugins/Experimental/AnimNext/Source/AnimNextInterfaceGraphUncookedOnly/Private/AnimNextInterfaceGraph_EditorData.cpp b/Engine/Plugins/Experimental/AnimNext/Source/AnimNextInterfaceGraphUncookedOnly/Private/AnimNextInterfaceGraph_EditorData.cpp
--- a/Engine/Plugins/Experimental/AnimNext/Source/AnimNextInterfaceGraphUncookedOnly/Private/AnimNextInterfaceGraph_EditorData.cpp
+++ b/Engine/Plugins/Experimental/AnimNext/Source/AnimNextInterfaceGraphUncookedOnly/Private/AnimNextInterfaceGraph_EditorData.cpp
@@ -12,6 +12,8 @@
 #include "RigVMModel/RigVMNotifications.h"
 #include "RigVMModel/Nodes/RigVMCollapseNode.h"
 #include "Curves/CurveFloat.h"
+#include "UObject/ObjectSaveContext.h"
+#include "Kismet2/BlueprintEditorUtils.h"
 
 UAnimNextInterfaceGraph_EditorData::UAnimNextInterfaceGraph_EditorData(const FObjectInitializer& ObjectInitializer)
 {
@@ -22,6 +24,7 @@ UAnimNextInterfaceGraph_EditorData::UAnimNextInterfaceGraph_EditorData(const FOb
 		RigVMClient.AddModel(TEXT("RigVMGraph"), false, &ObjectInitializer);
 		RigVMClient.GetOrCreateFunctionLibrary(false, &ObjectInitializer);
 	}
+	RigVMClient.SetExecuteContextStruct(FAnimNextInterfaceExecuteContext::StaticStruct());
 	
 	auto MakeEdGraph = [this, &ObjectInitializer](FName InName) -> UAnimNextInterfaceGraph_EdGraph*
 	{
@@ -93,8 +96,160 @@ void UAnimNextInterfaceGraph_EditorData::PostLoad()
 	Super::PostLoad();
 	
 	Initialize(/*bRecompileVM*/false);
+	RefreshAllModels(EAnimNextInterfaceGraphLoadType::PostLoad);
+
+#if WITH_EDITOR
+	if (GIsEditor)
+	{
+		// delay compilation until the package has been loaded
+		FCoreUObjectDelegates::OnEndLoadPackage.AddUObject(this, &UAnimNextInterfaceGraph_EditorData::HandlePackageDone);
+	}
+#else // !WITH_EDITOR
+	RecompileVMIfRequired();
+#endif // WITH_EDITOR
+}
+
+#if WITH_EDITOR
+void UAnimNextInterfaceGraph_EditorData::HandlePackageDone(const FEndLoadPackageContext& Context)
+{
+	if (!Context.LoadedPackages.Contains(GetPackage()))
+	{
+		return;
+	}
+	HandlePackageDone();
+}
+
+void UAnimNextInterfaceGraph_EditorData::HandlePackageDone()
+{
+	FCoreUObjectDelegates::OnEndLoadPackage.RemoveAll(this);
+
+	RecompileVM();
 }
 
+void UAnimNextInterfaceGraph_EditorData::RefreshAllModels(EAnimNextInterfaceGraphLoadType InLoadType)
+{
+	const bool bIsPostLoad = InLoadType == EAnimNextInterfaceGraphLoadType::PostLoad;
+
+	TGuardValue<bool> IsCompilingGuard(bIsCompiling, true);
+	TGuardValue<bool> ClientIgnoreModificationsGuard(RigVMClient.bIgnoreModelNotifications, true);
+
+	TArray<URigVMGraph*> GraphsToDetach = RigVMClient.GetAllModels(true, false);
+
+	if (ensure(IsInGameThread()))
+	{
+		for (URigVMGraph* GraphToDetach : GraphsToDetach)
+		{
+			URigVMController* Controller = RigVMClient.GetOrCreateController(GraphToDetach);
+			// temporarily disable default value validation during load time, serialized values should always be accepted
+			TGuardValue<bool> PerGraphDisablePinDefaultValueValidation(Controller->bValidatePinDefaults, false);
+			FRigVMControllerNotifGuard NotifGuard(Controller, true);
+			Controller->DetachLinksFromPinObjects();
+			TArray<URigVMNode*> Nodes = GraphToDetach->GetNodes();
+			for (URigVMNode* Node : Nodes)
+			{
+				Controller->RepopulatePinsOnNode(Node, true, true);
+			}
+		}
+		//SetupPinRedirectorsForBackwardsCompatibility();
+	}
+
+	for (URigVMGraph* GraphToDetach : GraphsToDetach)
+	{
+		URigVMController* Controller = RigVMClient.GetOrCreateController(GraphToDetach);
+		// at this stage, allow all links to be reattached,
+		// RecomputeAllTemplateFilteredPermutations() later should break any invalid links
+		FRigVMControllerNotifGuard NotifGuard(Controller, true);
+		Controller->ReattachLinksToPinObjects(true /* follow redirectors */, nullptr, true, true);
+	}
+
+	if (bIsPostLoad)
+	{
+		//PatchTemplateNodesWithPreferredPermutation();
+	}
+
+	TArray<URigVMGraph*> GraphsToClean = RigVMClient.GetAllModels(true, true);
+
+	// Sort from leaf graphs to root
+	TArray<URigVMGraph*> SortedGraphsToClean;
+	SortedGraphsToClean.Reserve(GraphsToClean.Num());
+	while (SortedGraphsToClean.Num() < GraphsToClean.Num())
+	{
+		bool bGraphAdded = false;
+		for (URigVMGraph* Graph : GraphsToClean)
+		{
+			if (SortedGraphsToClean.Contains(Graph))
+			{
+				continue;
+			}
+
+			TArray<URigVMGraph*> ContainedGraphs;
+			for (URigVMNode* Node : Graph->GetNodes())
+			{
+				if (URigVMLibraryNode* LibraryNode = Cast<URigVMLibraryNode>(Node))
+				{
+					if (URigVMFunctionReferenceNode* FunctionReferenceNode = Cast<URigVMFunctionReferenceNode>(LibraryNode))
+					{
+						if (FunctionReferenceNode->GetReferencedFunctionHeader().LibraryPointer.LibraryNode.GetLongPackageName() != GetPackage()->GetPathName())
+						{
+							continue;
+						}
+						if (URigVMLibraryNode* ReferencedNode = FunctionReferenceNode->LoadReferencedNode())
+						{
+							ContainedGraphs.Add(ReferencedNode->GetContainedGraph());
+							continue;
+						}
+					}
+
+					if (URigVMGraph* ContainedGraph = LibraryNode->GetContainedGraph())
+					{
+						ContainedGraphs.Add(ContainedGraph);
+					}
+				}
+			}
+
+			bool bAllContained = true;
+			for (URigVMGraph* Contained : ContainedGraphs)
+			{
+				if (!SortedGraphsToClean.Contains(Contained))
+				{
+					bAllContained = false;
+					break;
+				}
+			}
+			if (bAllContained)
+			{
+				SortedGraphsToClean.Add(Graph);
+				bGraphAdded = true;
+			}
+		}
+		ensure(bGraphAdded);
+	}
+
+	for (int32 GraphIndex = 0; GraphIndex < SortedGraphsToClean.Num(); GraphIndex++)
+	{
+		URigVMGraph* GraphToClean = SortedGraphsToClean[GraphIndex];
+		URigVMController* Controller = RigVMClient.GetOrCreateController(GraphToClean);
+		TGuardValue<bool> RecomputeGuard(Controller->bSuspendRecomputingOuterTemplateFilters, true);
+		//TGuardValue<bool> GuardEditGraph(GraphToClean->bEditable, true);
+		FRigVMControllerNotifGuard NotifGuard(Controller, true);
+
+		for (URigVMNode* ModelNode : GraphToClean->GetNodes())
+		{
+			Controller->RemoveUnusedOrphanedPins(ModelNode);
+		}
+
+		if (bIsPostLoad)
+		{
+			if (URigVMLibraryNode* LibraryNode = GraphToClean->GetTypedOuter<URigVMLibraryNode>())
+			{
+				Controller->UpdateLibraryTemplate(LibraryNode, false);
+			}
+		}
+	}
+}
+
+#endif // WITH_EDITOR
+
 FRigVMClient* UAnimNextInterfaceGraph_EditorData::GetRigVMClient()
 {
 	return &RigVMClient;
@@ -120,7 +275,7 @@ void UAnimNextInterfaceGraph_EditorData::HandleRigVMGraphAdded(const FRigVMClien
 {
 	if(URigVMGraph* RigVMGraph = InClient->GetModel(InNodePath))
 	{
-		RigVMGraph->SetExecuteContextStruct(FRigVMExecuteContext::StaticStruct());
+		RigVMGraph->SetExecuteContextStruct(FAnimNextInterfaceExecuteContext::StaticStruct());
 	}
 }
 
@@ -266,10 +421,8 @@ void UAnimNextInterfaceGraph_EditorData::HandleModifiedEvent(ERigVMGraphNotifTyp
 	case ERigVMGraphNotifType::PinArraySizeChanged:
 	case ERigVMGraphNotifType::PinDirectionChanged:
 		{
-			if (InGraph)
-			{
-				InGraph->ClearAST();
-			}
+			UpdateGraphReturnType();
+			RecompileVM();
 			break;
 		}
 
@@ -352,4 +505,119 @@ void UAnimNextInterfaceGraph_EditorData::CreateEdGraphForCollapseNode(URigVMColl
 			}
 		}
 	}
-}
\ No newline at end of file
+}
+
+FEdGraphPinType UAnimNextInterfaceGraph_EditorData::FindGraphReturnPinType() const
+{
+	static const TCHAR* EndExecStr = TEXT("AnimNextInterfaceEndExec");
+	static const FName ResultName = TEXT(".Result");
+
+	const UAnimNextInterfaceGraph_EdGraph* AnimNextInterfaceEdGraph = RootGraph.Get();
+	if (AnimNextInterfaceEdGraph != nullptr)
+	{
+		for (const TObjectPtr<UEdGraphNode>& Node : AnimNextInterfaceEdGraph->Nodes)
+		{
+			if (IsNodeExecConnected(Node))
+			{
+				const FString& NodeName = Node->GetName();
+				if (NodeName.Contains(EndExecStr))
+				{
+					for (const UEdGraphPin* Pin : Node->Pins)
+					{
+						const FString PinName = Pin->PinName.ToString();
+
+						if (PinName.Contains(ResultName.ToString()) && PinName.Contains(EndExecStr))
+						{
+							return Pin->PinType;
+						}
+					}
+				}
+			}
+		}
+	}
+
+	return FEdGraphPinType();
+}
+
+bool UAnimNextInterfaceGraph_EditorData::IsNodeExecConnected(const UEdGraphNode* Node) const
+{
+	static const TCHAR* ExecuteContextStr = TEXT(".ExecuteContext");
+
+	for (const UEdGraphPin* Pin : Node->Pins)
+	{
+		const FString PinName = Pin->PinName.ToString();
+
+		if (PinName.Contains(ExecuteContextStr) && Pin->LinkedTo.Num() > 0)
+		{
+			return true;
+		}
+	}
+
+	return false;
+}
+
+void UAnimNextInterfaceGraph_EditorData::UpdateGraphReturnType()
+{
+	if (UAnimNextInterfaceGraph* Graph = GetTypedOuter<UAnimNextInterfaceGraph>())
+	{
+		FEdGraphPinType GraphReturnPinType = FindGraphReturnPinType();
+		const FName PinTypeName = GetPinTypeName(GraphReturnPinType);
+		if (Graph->ReturnTypeName != PinTypeName || Graph->ReturnTypeStruct != GraphReturnPinType.PinSubCategoryObject)
+		{
+			Graph->Modify();
+			Graph->ReturnTypeName = PinTypeName;
+			Graph->ReturnTypeStruct = TObjectPtr<UScriptStruct>(Cast<UScriptStruct>(GraphReturnPinType.PinSubCategoryObject));
+		}
+	}
+}
+
+FName UAnimNextInterfaceGraph_EditorData::GetPinTypeName(const FEdGraphPinType& EdGraphPinType)
+{
+	FName DataType = EdGraphPinType.PinCategory;
+
+	if (EdGraphPinType.PinCategory == UEdGraphSchema_K2::PC_Real)
+	{
+		if (EdGraphPinType.PinSubCategory == UEdGraphSchema_K2::PC_Float)
+		{
+			DataType = TNameOf<float>::GetName();
+		}
+		else if (EdGraphPinType.PinSubCategory == UEdGraphSchema_K2::PC_Double)
+		{
+			DataType = TNameOf<double>::GetName();
+		}
+		else
+		{
+			ensure(false);
+		}
+	}
+	else if (DataType != NAME_None)
+	{
+		if (DataType == UEdGraphSchema_K2::PC_Struct)
+		{
+			UObject* DataTypeObject = nullptr;
+			DataType = NAME_None;
+			if (UScriptStruct* DataStruct = Cast<UScriptStruct>(EdGraphPinType.PinSubCategoryObject))
+			{
+				DataTypeObject = DataStruct;
+				DataType = *DataStruct->GetStructCPPName();
+			}
+		}
+
+		if (DataType == TEXT("int"))
+		{
+			DataType = TNameOf<int32>::GetName();
+		}
+		else if (DataType == TEXT("name"))
+		{
+			DataType = TNameOf<FName>::GetName();
+		}
+		else if (DataType == TEXT("string"))
+		{
+			DataType = TNameOf<FString>::GetName();;
+		}
+	}
+
+	return DataType;
+}
+
+#undef LOCTEXT_NAMESPACE
