Fixed LOC error due to #undef LOCTEXT_NAMESPACE on a file with no #define LOCTEXT_NAMESPACE

diff --git a/Engine/Plugins/Experimental/AnimNext/Source/AnimNextInterfaceGraphUncookedOnly/Private/AnimNextInterfaceGraph_EditorData.cpp b/Engine/Plugins/Experimental/AnimNext/Source/AnimNextInterfaceGraphUncookedOnly/Private/AnimNextInterfaceGraph_EditorData.cpp
--- a/Engine/Plugins/Experimental/AnimNext/Source/AnimNextInterfaceGraphUncookedOnly/Private/AnimNextInterfaceGraph_EditorData.cpp
+++ b/Engine/Plugins/Experimental/AnimNext/Source/AnimNextInterfaceGraphUncookedOnly/Private/AnimNextInterfaceGraph_EditorData.cpp
@@ -619,5 +619,3 @@ FName UAnimNextInterfaceGraph_EditorData::GetPinTypeName(const FEdGraphPinType&
 
 	return DataType;
 }
-
-#undef LOCTEXT_NAMESPACE
