Fixed bad merge.

diff --git a/Engine/Plugins/Experimental/AnimNext/Source/AnimNextInterface/Public/AnimNextInterfaceContext.h b/Engine/Plugins/Experimental/AnimNext/Source/AnimNextInterface/Public/AnimNextInterfaceContext.h
--- a/Engine/Plugins/Experimental/AnimNext/Source/AnimNextInterface/Public/AnimNextInterfaceContext.h
+++ b/Engine/Plugins/Experimental/AnimNext/Source/AnimNextInterface/Public/AnimNextInterfaceContext.h
@@ -172,7 +172,7 @@ public:
 			}
 
 			// deal with containers
-			if constexpr (TModels_V<Private::CSizedContainerWithAccessibleDataAsRawPtr, ValueType>::Value)
+			if constexpr (TModels_V<Private::CSizedContainerWithAccessibleDataAsRawPtr, ValueType>)
 			{
 				MutableValueType* ParamData = const_cast<MutableValueType>(static_cast<ValueType*>(Value->GetData()));
 				ParamHandle = EnumHasAnyFlags(Flags, FParam::EFlags::Value) 
@@ -299,7 +299,7 @@ public:
 	void SetResult(const ValueType& InValue) const
 	{
 		// Support containers with Num() and GetData()
-		if constexpr (TModels_V<Private::CSizedContainerWithAccessibleDataAsRawPtr, ValueType>::Value)
+		if constexpr (TModels_V<Private::CSizedContainerWithAccessibleDataAsRawPtr, ValueType>)
 		{
 			checkSlow(Private::CheckParam<ValueType>(*Result));
 			checkSlow(InValue.Num() == Result->NumElements);
