Fixed some remaining TModels usage.

diff --git a/Engine/Plugins/Experimental/AnimNext/Source/AnimNextInterface/Public/AnimNextInterface.h b/Engine/Plugins/Experimental/AnimNext/Source/AnimNextInterface/Public/AnimNextInterface.h
--- a/Engine/Plugins/Experimental/AnimNext/Source/AnimNextInterface/Public/AnimNextInterface.h
+++ b/Engine/Plugins/Experimental/AnimNext/Source/AnimNextInterface/Public/AnimNextInterface.h
@@ -34,7 +34,7 @@ static bool GetDataSafe(const TScriptInterface<InterfaceType> AnimNextInterface,
 //		else
 //		{
 //			// Support containers with Num() and GetData()
-//			if constexpr (TModels<Private::CSizedContainerWithAccessibleDataAsRawPtr, ValueType>::Value)
+//			if constexpr (TModels_V<Private::CSizedContainerWithAccessibleDataAsRawPtr, ValueType>)
 //			{
 //				TWrapParam<ValueType::ElementType> WrappedResult(Result.GetData(), Result.Num());
 //				const FContext ResultHolder = InContext.WithResult(WrappedResult);
@@ -66,7 +66,7 @@ static bool GetDataSafe(const TScriptInterface<InterfaceType> AnimNextInterface,
 		else
 		{
 			// Support containers with Num() and GetData()
-			if constexpr (TModels<Private::CSizedContainerWithAccessibleDataAsRawPtr, ValueType>::Value)
+			if constexpr (TModels_V<Private::CSizedContainerWithAccessibleDataAsRawPtr, ValueType>)
 			{
 				TWrapParam<typename ValueType::ElementType> WrappedResult(Result.GetData(), Result.Num());
 				return AnimNextInterface.GetInterface()->GetData(InContext, WrappedResult);
diff --git a/Engine/Plugins/Experimental/AnimNext/Source/AnimNextInterface/Public/AnimNextInterfaceParam.h b/Engine/Plugins/Experimental/AnimNext/Source/AnimNextInterface/Public/AnimNextInterfaceParam.h
--- a/Engine/Plugins/Experimental/AnimNext/Source/AnimNextInterface/Public/AnimNextInterfaceParam.h
+++ b/Engine/Plugins/Experimental/AnimNext/Source/AnimNextInterface/Public/AnimNextInterfaceParam.h
@@ -452,7 +452,7 @@ protected:
 		using MutableValueType = std::remove_const_t<ValueType>;
 
 		// Support containers with Num() and GetData()
-		if constexpr (TModels<Private::CSizedContainerWithAccessibleDataAsRawPtr, ValueType>::Value)
+		if constexpr (TModels_V<Private::CSizedContainerWithAccessibleDataAsRawPtr, ValueType>)
 		{
 			return (void*)(InValuePtrToWrap->GetData());
 		}
@@ -467,7 +467,7 @@ protected:
 		FParam::EFlags NewFlags = FParam::EFlags::None;
 
 		// Add batched flags or not
-		if constexpr (TModels<Private::CSizedContainerWithAccessibleDataAsRawPtr, ValueType>::Value)
+		if constexpr (TModels_V<Private::CSizedContainerWithAccessibleDataAsRawPtr, ValueType>)
 		{
 			NewFlags |= (InValuePtrToWrap->Num() > 1) ? FParam::EFlags::Batched : FParam::EFlags::None;
 		}
diff --git a/Engine/Source/Developer/RigVMDeveloper/Public/RigVMModel/Nodes/RigVMDispatchNode.h b/Engine/Source/Developer/RigVMDeveloper/Public/RigVMModel/Nodes/RigVMDispatchNode.h
--- a/Engine/Source/Developer/RigVMDeveloper/Public/RigVMModel/Nodes/RigVMDispatchNode.h
+++ b/Engine/Source/Developer/RigVMDeveloper/Public/RigVMModel/Nodes/RigVMDispatchNode.h
@@ -59,7 +59,7 @@ public:
 	// Returns a copy of the struct with the current values
 	template <
 		typename T,
-		typename TEnableIf<TModels<CRigVMUStruct, T>::Value>::Type* = nullptr
+		typename TEnableIf<TModels_V<CRigVMUStruct, T>>::Type* = nullptr
 	>
 	T ConstructFactoryInstance() const
 	{
diff --git a/Engine/Source/Runtime/RigVM/Public/RigVMCore/RigVMDispatchFactory.h b/Engine/Source/Runtime/RigVM/Public/RigVMCore/RigVMDispatchFactory.h
--- a/Engine/Source/Runtime/RigVM/Public/RigVMCore/RigVMDispatchFactory.h
+++ b/Engine/Source/Runtime/RigVM/Public/RigVMCore/RigVMDispatchFactory.h
@@ -189,7 +189,7 @@ protected:
 
 	template <
 		typename T,
-		typename TEnableIf<TModels<CRigVMUStruct, T>::Value>::Type * = nullptr
+		typename TEnableIf<TModels_V<CRigVMUStruct, T>>::Type * = nullptr
 	>
 	static FString GetDefaultValueForStruct(const T& InValue)
 	{
diff --git a/Engine/Source/Runtime/RigVM/Public/RigVMCore/RigVMNativized.h b/Engine/Source/Runtime/RigVM/Public/RigVMCore/RigVMNativized.h
--- a/Engine/Source/Runtime/RigVM/Public/RigVMCore/RigVMNativized.h
+++ b/Engine/Source/Runtime/RigVM/Public/RigVMCore/RigVMNativized.h
@@ -388,7 +388,7 @@ public:
 
 	template <
 		typename T,
-		typename TEnableIf<TModels<CRigVMUStruct, T>::Value>::Type* = nullptr
+		typename TEnableIf<TModels_V<CRigVMUStruct, T>>::Type* = nullptr
 	>
 	static TArray<TArray<T>> GetStructArrayArrayConstant(const FString& InDefaultValue)
 	{
diff --git a/Engine/Source/Runtime/RigVM/Public/RigVMCore/RigVMRegistry.h b/Engine/Source/Runtime/RigVM/Public/RigVMCore/RigVMRegistry.h
--- a/Engine/Source/Runtime/RigVM/Public/RigVMCore/RigVMRegistry.h
+++ b/Engine/Source/Runtime/RigVM/Public/RigVMCore/RigVMRegistry.h
@@ -231,7 +231,7 @@ public:
 	// Returns a dispatch factory given its static struct (or nullptr)
 	template <
 		typename T,
-		typename TEnableIf<TModels<CRigVMUStruct, T>::Value>::Type * = nullptr
+		typename TEnableIf<TModels_V<CRigVMUStruct, T>>::Type * = nullptr
 	>
 	FRigVMDispatchFactory* FindOrAddDispatchFactory()
 	{
@@ -244,7 +244,7 @@ public:
 	// Returns a dispatch factory's singleton function name if that exists
 	template <
 		typename T,
-		typename TEnableIf<TModels<CRigVMUStruct, T>::Value>::Type * = nullptr
+		typename TEnableIf<TModels_V<CRigVMUStruct, T>>::Type * = nullptr
 	>
 	FString FindOrAddSingletonDispatchFunction()
 	{
