Experimental plugin prototype using RigVM to evaluate animation (WIP) :
New AnimationRegistry for data storage and ReferencePose generation for SkeletalMesh / SkeletalMeshComponent assets.
Created the concept of ReferencePose and LODPose (the second being just a colletion of animation atoms to be filled by the animation code). Both use a transforms array with separated arrays for translation, rotation and scale3D.
Cloned part of the animation decompression code to the plugin, so it can be used by RigVM Units (and adapted it to work with the new data structures).
Added an engine test to check performance between AoS poses and SoA poses with a pose blend.

diff --git a/Engine/Plugins/Experimental/AnimNext/Source/AnimNextInterface/AnimNextInterface.Build.cs b/Engine/Plugins/Experimental/AnimNext/Source/AnimNextInterface/AnimNextInterface.Build.cs
--- a/Engine/Plugins/Experimental/AnimNext/Source/AnimNextInterface/AnimNextInterface.Build.cs
+++ b/Engine/Plugins/Experimental/AnimNext/Source/AnimNextInterface/AnimNextInterface.Build.cs
@@ -13,7 +13,13 @@ namespace UnrealBuildTool.Rules
 					"CoreUObject",
 				}
 			);
-
+			PrivateDependencyModuleNames.AddRange(
+				new string[]
+				{
+					"AssetRegistry",
+					"Engine"
+				}
+			);
 		}
 	}
 }
\ No newline at end of file
diff --git a/Engine/Plugins/Experimental/AnimNext/Source/AnimNextInterface/Private/AnimNextInterfaceModule.cpp b/Engine/Plugins/Experimental/AnimNext/Source/AnimNextInterface/Private/AnimNextInterfaceModule.cpp
--- a/Engine/Plugins/Experimental/AnimNext/Source/AnimNextInterface/Private/AnimNextInterfaceModule.cpp
+++ b/Engine/Plugins/Experimental/AnimNext/Source/AnimNextInterface/Private/AnimNextInterfaceModule.cpp
@@ -5,6 +5,7 @@
 #include "Modules/ModuleManager.h"
 #include "Misc/CoreDelegates.h"
 #include "AnimNextInterfaceParam.h"
+#include "AnimationDataRegistry.h"
 
 namespace UE::AnimNext::Interface
 {
@@ -14,10 +15,18 @@ class FModule : public IModuleInterface
 public:
 	virtual void StartupModule() override
 	{
-		FCoreDelegates::OnPostEngineInit.AddLambda([]()
+		FCoreDelegates::OnFEngineLoopInitComplete.AddLambda([]()
 		{
 			FParamType::FRegistrar::RegisterDeferredTypes();
+
+			UE::AnimNext::Interface::FAnimationDataRegistry::Init();
+		});
+
+		FCoreDelegates::OnEnginePreExit.AddLambda([]()
+		{
+			UE::AnimNext::Interface::FAnimationDataRegistry::Destroy();
 		});
+		
 	}
 };
 
diff --git a/Engine/Plugins/Experimental/AnimNext/Source/AnimNextInterface/Private/AnimNextInterfaceParam.cpp b/Engine/Plugins/Experimental/AnimNext/Source/AnimNextInterface/Private/AnimNextInterfaceParam.cpp
--- a/Engine/Plugins/Experimental/AnimNext/Source/AnimNextInterface/Private/AnimNextInterfaceParam.cpp
+++ b/Engine/Plugins/Experimental/AnimNext/Source/AnimNextInterface/Private/AnimNextInterfaceParam.cpp
@@ -8,7 +8,14 @@ namespace UE::AnimNext::Interface
 
 static FParamType DefaultType;
 static TArray<FParamType*> GTypeRegistry;
-static TArray<TUniqueFunction<void(void)>> GDeferredTypes;
+
+// This avoids c++ static ordering problems with the types definitions (for the deferred type registration)
+/*static*/ TArray<TUniqueFunction<void(void)>>& FParamType::FRegistrar::GetDeferredTypes()
+{
+	static TArray<TUniqueFunction<void(void)>> GDeferredTypes;
+	return GDeferredTypes;
+}
+
 
 FParam::FParam(const FParamType& InType, void* InData, EFlags InFlags)
 	: Data(InData)
@@ -135,19 +142,19 @@ const FParamType& FParamType::FRegistrar::GetTypeById(uint16 InTypeId)
 FParamType::FRegistrar::FRegistrar(TUniqueFunction<void(void)>&& InFunction)
 {
 	check(IsInGameThread());
-	GDeferredTypes.Add(MoveTemp(InFunction));
+	GetDeferredTypes().Add(MoveTemp(InFunction));
 }
 
 void FParamType::FRegistrar::RegisterDeferredTypes()
 {
 	check(IsInGameThread());
-	
-	for(TUniqueFunction<void(void)>& Function : GDeferredTypes)
+
+	for(TUniqueFunction<void(void)>& Function : FRegistrar::GetDeferredTypes())
 	{
 		Function();
 	}
 
-	GDeferredTypes.Empty();
+	FRegistrar::GetDeferredTypes().Empty();
 }
 
 const FParamType& FParam::GetType() const
diff --git a/Engine/Plugins/Experimental/AnimNext/Source/AnimNextInterface/Private/AnimNextInterfaceParamStorage.cpp b/Engine/Plugins/Experimental/AnimNext/Source/AnimNextInterface/Private/AnimNextInterfaceParamStorage.cpp
--- a/Engine/Plugins/Experimental/AnimNext/Source/AnimNextInterface/Private/AnimNextInterfaceParamStorage.cpp
+++ b/Engine/Plugins/Experimental/AnimNext/Source/AnimNextInterface/Private/AnimNextInterfaceParamStorage.cpp
@@ -10,7 +10,7 @@ namespace UE::AnimNext::Interface
 FParamStorage::FParamStorage(int32 InMaxParams, int32 InAllocatedMemorySize, int32 InMaxBlocks, bool InAllowGrowing)
 	: AllowGrowing(InAllowGrowing)
 {
-	check(InMaxParams > 0 && InMaxBlocks > 0 && InAllocatedMemorySize > 0);
+	check(InMaxParams > 0 && (InMaxBlocks == 0 || (InMaxBlocks > 0 && InAllocatedMemorySize > 0)));
 
 	Parameters.Reserve(InMaxParams);
 	Parameters.AddDefaulted(InMaxParams);
@@ -26,17 +26,23 @@ FParamStorage::FParamStorage(int32 InMaxParams, int32 InAllocatedMemorySize, int
 
 	// --- ---
 
-	RawMemory = (uint8*)FMemory::Malloc(InAllocatedMemorySize);
-	RawMemorySize = InAllocatedMemorySize;
+	if (InAllocatedMemorySize > 0)
+	{
+		RawMemory = (uint8*)FMemory::Malloc(InAllocatedMemorySize);
+		RawMemorySize = InAllocatedMemorySize;
+	}
 
-	BlockOffets.Reserve(InMaxBlocks);
-	BlockOffets.AddZeroed(InMaxBlocks);
+	if (InMaxBlocks > 0)
+	{
+		BlockOffets.Reserve(InMaxBlocks);
+		BlockOffets.AddZeroed(InMaxBlocks);
 
-	BlockSizes.Reserve(InMaxBlocks);
-	BlockSizes.AddZeroed(InMaxBlocks);
+		BlockSizes.Reserve(InMaxBlocks);
+		BlockSizes.AddZeroed(InMaxBlocks);
 
-	BlockFreeFlag.Reserve(InMaxBlocks);
-	BlockFreeFlag.Init(true, InMaxBlocks);
+		BlockFreeFlag.Reserve(InMaxBlocks);
+		BlockFreeFlag.Init(true, InMaxBlocks);
+	}
 
 	FreeBlockIndex = 0;
 }
diff --git a/Engine/Plugins/Experimental/AnimNext/Source/AnimNextInterface/Private/AnimationDataRegistry.cpp b/Engine/Plugins/Experimental/AnimNext/Source/AnimNextInterface/Private/AnimationDataRegistry.cpp
new file mode 100644
--- /dev/null
+++ b/Engine/Plugins/Experimental/AnimNext/Source/AnimNextInterface/Private/AnimationDataRegistry.cpp
@@ -0,0 +1,219 @@
+// Copyright Epic Games, Inc. All Rights Reserved.
+
+#include "AnimationDataRegistry.h"
+
+#include "Misc/ScopeRWLock.h"
+#include "AnimationDataRegistryTypes.h"
+#include "AnimationGenerationTools.h"
+#include "AnimNextInterfaceTypes.h"
+#include "UObject/UObjectGlobals.h"
+#include "Components/SkeletalMeshComponent.h"
+#include "Engine/SkeletalMesh.h"
+#include "Rendering/SkeletalMeshModel.h"
+#include "BoneContainer.h"
+#include "Engine/SkinnedAssetCommon.h"
+#include "Rendering/SkeletalMeshRenderData.h"
+
+#include "AssetRegistry/AssetRegistryModule.h"
+#include "AssetRegistry/AssetData.h"
+#include "Rendering/SkeletalMeshRenderData.h"
+
+
+namespace // Private
+{
+
+UE::AnimNext::Interface::FAnimationDataRegistry* GAnimationDataRegistry = nullptr;
+constexpr int32 BASIC_TYPE_ALLOCK_BLOCK = 1000;
+FDelegateHandle PostGarbageCollectHandle;
+
+} // end namespace
+
+namespace UE::AnimNext::Interface
+{
+
+/*static*/ void FAnimationDataRegistry::Init()
+{
+	if (ensure(GAnimationDataRegistry == nullptr))
+	{
+		GAnimationDataRegistry = new FAnimationDataRegistry();
+
+		PostGarbageCollectHandle = FCoreUObjectDelegates::GetPostGarbageCollect().AddStatic(&FAnimationDataRegistry::HandlePostGarbageCollect);
+	}
+}
+
+/*static*/ void FAnimationDataRegistry::Destroy()
+{
+	if (ensure(GAnimationDataRegistry != nullptr))
+	{
+		FCoreUObjectDelegates::GetPostGarbageCollect().Remove(PostGarbageCollectHandle);
+
+		GAnimationDataRegistry->ReleaseReferencePoseData(); // release any registered poses
+
+		check(GAnimationDataRegistry->AllocatedBlocks.Num() == 0); // any other data should have been released at this point
+		check(GAnimationDataRegistry->StoredData.Num() == 0);
+
+		delete GAnimationDataRegistry;
+		GAnimationDataRegistry = nullptr;
+	}
+}
+
+FAnimationDataRegistry* FAnimationDataRegistry::Get()
+{
+	checkf(GAnimationDataRegistry, TEXT("Animation Data Registry is not instanced. It is only valid to access this while the engine module is loaded."));
+	return GAnimationDataRegistry;
+}
+
+
+/*static*/ void FAnimationDataRegistry::HandlePostGarbageCollect()
+{
+	// Compact the registry on GC
+	if (GAnimationDataRegistry)
+	{
+		FRWScopeLock Lock(GAnimationDataRegistry->SkeletalMeshReferencePosesLock, SLT_Write);
+
+		for (auto Iter = GAnimationDataRegistry->SkeletalMeshReferencePoses.CreateIterator(); Iter; ++Iter)
+		{
+			const TWeakObjectPtr<const USkeletalMeshComponent>& SkeletalMeshComponentPtr = Iter.Key();
+			if (SkeletalMeshComponentPtr.Get() == nullptr)
+			{
+				Iter.RemoveCurrent();
+			}
+		}
+	}
+}
+
+FAnimationDataHandle FAnimationDataRegistry::RegisterReferencePose(USkeletalMeshComponent* SkeletalMeshComponent)
+{
+	FAnimationDataHandle Handle = AllocateData<FAnimationReferencePose>(1);
+
+	FAnimationReferencePose& AnimationReferencePose = Handle.GetRef<FAnimationReferencePose>();
+
+	FGenerationTools::GenerateReferencePose(SkeletalMeshComponent, SkeletalMeshComponent->GetSkeletalMeshAsset(), AnimationReferencePose);
+
+	// register even if it fails to generate (register an empty ref pose)
+	{
+		FDelegateHandle DelegateHandle = SkeletalMeshComponent->RegisterOnLODRequiredBonesUpdate(FOnLODRequiredBonesUpdate::CreateRaw(this, &FAnimationDataRegistry::OnLODRequiredBonesUpdate));
+
+		FRWScopeLock Lock(SkeletalMeshReferencePosesLock, SLT_Write);
+		SkeletalMeshReferencePoses.Add(SkeletalMeshComponent, FReferencePoseData(Handle, DelegateHandle));
+	}
+
+	return Handle;
+}
+
+void FAnimationDataRegistry::OnLODRequiredBonesUpdate(USkeletalMeshComponent* SkeletalMeshComponent, int32 LODLevel, const TArray<FBoneIndexType>& LODRequiredBones)
+{
+	// TODO : Check if the LDO bomes are different from the currently calculated ReferencePose data (for now just delete the cached data)
+	RemoveReferencePose(SkeletalMeshComponent);
+}
+
+FAnimationDataHandle FAnimationDataRegistry::GetOrGenerateReferencePose(USkeletalMeshComponent* SkeletalMeshComponent)
+{
+	FAnimationDataHandle ReturnHandle;
+
+	{
+		FRWScopeLock Lock(SkeletalMeshReferencePosesLock, SLT_ReadOnly);
+
+		if (const FReferencePoseData* ReferencePoseData = SkeletalMeshReferencePoses.Find(SkeletalMeshComponent))
+		{
+			ReturnHandle = ReferencePoseData->AnimationDataHandle;
+		}
+	}
+	
+	if (ReturnHandle.IsValid() == false)
+	{
+		ReturnHandle = RegisterReferencePose(SkeletalMeshComponent);
+	}
+
+	return ReturnHandle;
+}
+
+void FAnimationDataRegistry::RemoveReferencePose(USkeletalMeshComponent* SkeletalMeshComponent)
+{
+	if (SkeletalMeshComponent != nullptr)
+	{
+		FRWScopeLock Lock(SkeletalMeshReferencePosesLock, SLT_Write);
+
+		FReferencePoseData ReferencePoseData;
+		if (SkeletalMeshReferencePoses.RemoveAndCopyValue(SkeletalMeshComponent, ReferencePoseData))
+		{
+			SkeletalMeshComponent->UnregisterOnLODRequiredBonesUpdate(ReferencePoseData.DelegateHandle);
+		}
+	}
+}
+
+
+
+void FAnimationDataRegistry::RegisterData(const FName& Id, const FAnimationDataHandle& AnimationDataHandle)
+{
+	FRWScopeLock Lock(StoredDataLock, SLT_Write);
+	StoredData.Add(Id, AnimationDataHandle);
+}
+
+void FAnimationDataRegistry::UnregisterData(const FName& Id)
+{
+	FRWScopeLock Lock(StoredDataLock, SLT_Write);
+	StoredData.Remove(Id);
+}
+
+FAnimationDataHandle FAnimationDataRegistry::GetRegisteredData(const FName& Id) const
+{
+	FAnimationDataHandle ReturnHandle;
+
+	{
+		FRWScopeLock Lock(StoredDataLock, SLT_ReadOnly);
+
+		if (const FAnimationDataHandle* HandlePtr = StoredData.Find(Id))
+		{
+			ReturnHandle = *HandlePtr;
+		}
+	}
+
+	return ReturnHandle;
+}
+
+void FAnimationDataRegistry::FreeAllocatedBlock(Private::FAllocatedBlock* AllocatedBlock)
+{
+	FRWScopeLock Lock(DataTypeDefsLock, SLT_Write);
+
+	if (ensure(AllocatedBlock != nullptr && AllocatedBlocks.Find(AllocatedBlock)))
+	{
+		if (AllocatedBlock->Memory != nullptr)
+		{
+			void* Memory = AllocatedBlock->Memory;
+
+			FDataTypeDef* TypeDef = DataTypeDefs.Find(AllocatedBlock->TypeId);
+			if (ensure(TypeDef != nullptr))
+			{
+				TypeDef->DestrooyTypeFn((uint8*)AllocatedBlock->Memory, AllocatedBlock->NumElem);
+
+				FMemory::Free(AllocatedBlock->Memory); // TODO : This should come from preallocated chunks, use malloc / free for now
+				AllocatedBlock->Memory = nullptr;
+			}
+
+			delete AllocatedBlock; // TODO : avoid memory fragmentation
+			AllocatedBlocks.Remove(AllocatedBlock);
+		}
+	}
+}
+
+// Remove any ReferencePoses and unregister all the SkeletalMeshComponent delegates (if any still alive)
+void FAnimationDataRegistry::ReleaseReferencePoseData()
+{
+	FRWScopeLock Lock(SkeletalMeshReferencePosesLock, SLT_Write);
+
+	for (auto Iter = GAnimationDataRegistry->SkeletalMeshReferencePoses.CreateIterator(); Iter; ++Iter)
+	{
+		const TWeakObjectPtr<USkeletalMeshComponent>& SkeletalMeshComponentPtr = Iter.Key();
+		const FReferencePoseData& ReferencePoseData = Iter.Value();
+
+		if (USkeletalMeshComponent* SkeletalMeshComponent = SkeletalMeshComponentPtr.Get())
+		{
+			SkeletalMeshComponent->UnregisterOnLODRequiredBonesUpdate(ReferencePoseData.DelegateHandle);
+		}
+	}
+
+	SkeletalMeshReferencePoses.Empty();
+}
+
+} // end namespace UE::AnimNext::Interface
diff --git a/Engine/Plugins/Experimental/AnimNext/Source/AnimNextInterface/Private/AnimationDataRegistryTypes.cpp b/Engine/Plugins/Experimental/AnimNext/Source/AnimNextInterface/Private/AnimationDataRegistryTypes.cpp
new file mode 100644
--- /dev/null
+++ b/Engine/Plugins/Experimental/AnimNext/Source/AnimNextInterface/Private/AnimationDataRegistryTypes.cpp
@@ -0,0 +1,26 @@
+// Copyright Epic Games, Inc. All Rights Reserved.
+
+#include "AnimationDataRegistryTypes.h"
+#include "AnimationDataRegistry.h"
+#include "AnimationReferencePose.h"
+
+namespace UE::AnimNext::Interface
+{
+
+FAnimationDataHandle::~FAnimationDataHandle()
+{
+	if (AllocatedBlock != nullptr)
+	{
+		check(AllocatedBlock->GetRefCount() > 0);
+
+		const int32 CurrentCount = AllocatedBlock->Release();
+		check(CurrentCount >= 0);
+
+		if (CurrentCount == 0)
+		{
+			FAnimationDataRegistry::Get()->FreeAllocatedBlock(AllocatedBlock);
+		}
+	}
+}
+
+} // namespace UE::AnimNext::Interface
diff --git a/Engine/Plugins/Experimental/AnimNext/Source/AnimNextInterface/Private/AnimationDecompressionTools.cpp b/Engine/Plugins/Experimental/AnimNext/Source/AnimNextInterface/Private/AnimationDecompressionTools.cpp
new file mode 100644
--- /dev/null
+++ b/Engine/Plugins/Experimental/AnimNext/Source/AnimNextInterface/Private/AnimationDecompressionTools.cpp
@@ -0,0 +1,663 @@
+// Copyright Epic Games, Inc. All Rights Reserved.
+
+#include "AnimationDecompressionTools.h"
+#include "Animation/AnimSequenceDecompressionContext.h"
+#include "Animation/AnimCompressionTypes.h"
+#include "Animation/AnimBoneCompressionCodec.h"
+
+#include "AnimationRuntime.h"
+#include "AnimEncoding.h"
+#include "Animation/AnimBoneDecompressionData.h"
+#include "Animation/SkeletonRemappingRegistry.h"
+#include "Animation/SkeletonRemapping.h"
+#include "Animation/AnimRootMotionProvider.h"
+#include "Animation/AnimSequence.h"
+#include "HAL/ConsoleManager.h"
+
+CSV_DECLARE_CATEGORY_MODULE_EXTERN(ENGINE_API, Animation);
+
+DECLARE_CYCLE_STAT(TEXT("Build Anim Track Pairs"), STAT_BuildAnimTrackPairs, STATGROUP_Anim);
+DECLARE_CYCLE_STAT(TEXT("Extract Pose From Anim Data"), STAT_ExtractPoseFromAnimData, STATGROUP_Anim);
+
+DECLARE_CYCLE_STAT(TEXT("AnimSeq GetBonePose"), STAT_AnimSeq_GetBonePose, STATGROUP_Anim);
+DECLARE_CYCLE_STAT(TEXT("AnimSeq EvalCurveData"), STAT_AnimSeq_EvalCurveData, STATGROUP_Anim);
+
+namespace UE::AnimNext::Interface
+{
+
+static IConsoleVariable* CVarForceEvalRawData = IConsoleManager::Get().FindConsoleVariable(TEXT("a.ForceEvalRawData"));
+static bool GetForceRawData()
+{
+	return (CVarForceEvalRawData != nullptr) ? CVarForceEvalRawData->GetBool() : false;
+}
+
+
+//****************************************************************************
+// This file contains code extracted from AnimSequence and AnimationDecompression,
+// in order to be able to decompress an anim sequence using AnimNext format
+//****************************************************************************
+
+struct FGetBonePoseScratchArea : public TThreadSingleton<FGetBonePoseScratchArea>
+{
+	BoneTrackArray RotationScalePairs;
+	BoneTrackArray TranslationPairs;
+	BoneTrackArray AnimScaleRetargetingPairs;
+	BoneTrackArray AnimRelativeRetargetingPairs;
+	BoneTrackArray OrientAndScaleRetargetingPairs;
+};
+
+
+static bool CanEvaluateRawAnimationData(const UAnimSequence* AnimSequence)
+{
+#if WITH_EDITOR
+	return AnimSequence->IsDataModelValid();
+#endif
+
+	return false;
+}
+
+static bool UseRawDataForPoseExtraction(const UAnimSequence* AnimSequence, FAnimationLODPose& AnimationPoseData)
+{
+	return CanEvaluateRawAnimationData(AnimSequence) && (
+#if WITH_EDITOR
+		PRAGMA_DISABLE_DEPRECATION_WARNINGS
+		AnimSequence->OnlyUseRawData() ||
+		PRAGMA_ENABLE_DEPRECATION_WARNINGS
+#endif
+		(AnimSequence->GetSkeletonVirtualBoneGuid() != AnimSequence->GetSkeleton()->GetVirtualBoneGuid()) || AnimationPoseData.GetDisableRetargeting() || AnimationPoseData.ShouldUseRawData() ||
+#if WITH_EDITOR
+		GetForceRawData() ||
+#endif // WITH_EDITOR
+		AnimationPoseData.ShouldUseSourceData());
+}
+
+// --- ---
+
+/*static*/ void FDecompressionTools::GetAnimationPose(const UAnimSequence* AnimSequence, FAnimationLODPose& OutAnimationPoseData, const FAnimExtractContext& ExtractionContext)
+{
+	if (OutAnimationPoseData.GetRefPose().IsValid() == false)
+	{
+		return;
+	}
+
+	// @todo anim: if compressed and baked in the future, we don't have to do this 
+	if (UseRawDataForPoseExtraction(AnimSequence, OutAnimationPoseData) && AnimSequence->IsValidAdditive())
+	{
+		switch (AnimSequence->GetAdditiveAnimType())
+		{
+			case AAT_LocalSpaceBase:
+			{
+				GetBonePose_Additive(AnimSequence, ExtractionContext, OutAnimationPoseData);
+			}
+			break;
+
+			case AAT_RotationOffsetMeshSpace:
+			{
+				GetBonePose_AdditiveMeshRotationOnly(AnimSequence, ExtractionContext, OutAnimationPoseData);
+			}
+			break;
+
+			default:
+				break;
+		}
+	}
+	else
+	{
+		GetBonePose(AnimSequence, ExtractionContext, OutAnimationPoseData);
+	}
+
+	// If the sequence has root motion enabled, allow sampling of a root motion delta into the custom attribute container of the outgoing pose
+	if (AnimSequence->HasRootMotion())
+	{
+		if (const UE::Anim::IAnimRootMotionProvider* RootMotionProvider = UE::Anim::IAnimRootMotionProvider::Get())
+		{
+			//RootMotionProvider->SampleRootMotion(ExtractionContext.DeltaTimeRecord, *AnimSequence, ExtractionContext.bLooping, OutAnimationPoseData.GetAttributes());
+		}
+	}
+
+	// Check that all bone atoms coming from animation are normalized
+#if DO_CHECK && WITH_EDITORONLY_DATA
+	//check(OutAnimationPoseData.LocalTransforms.IsNormalized());
+#endif
+
+}
+
+void FDecompressionTools::GetBonePose(const UAnimSequence* AnimSequence, const FAnimExtractContext& ExtractionContext, FAnimationLODPose& OutAnimationPoseData, bool bForceUseRawData /*= false*/)
+{
+	SCOPE_CYCLE_COUNTER(STAT_AnimSeq_GetBonePose);
+	CSV_SCOPED_TIMING_STAT(Animation, AnimSeq_GetBonePose);
+
+	const TArrayView<const FBoneIndexType> LODBoneIndexes = OutAnimationPoseData.GetLODBoneIndexes();
+
+	check(!bForceUseRawData || CanEvaluateRawAnimationData(AnimSequence));
+	const bool bUseRawDataForPoseExtraction = (CanEvaluateRawAnimationData(AnimSequence) && bForceUseRawData) ||
+#if WITH_EDITOR
+		GetForceRawData() ||
+#endif // WITH_EDITOR
+		UseRawDataForPoseExtraction(AnimSequence, OutAnimationPoseData);
+
+	const bool bIsBakedAdditive = !bUseRawDataForPoseExtraction && AnimSequence->IsValidAdditive();
+
+	const USkeleton* MySkeleton = AnimSequence->GetSkeleton();
+	if (!MySkeleton)
+	{
+		OutAnimationPoseData.SetRefPose(bIsBakedAdditive);
+		return;
+	}
+
+	const bool bDisableRetargeting = OutAnimationPoseData.GetDisableRetargeting();
+
+	// initialize with ref-pose
+	if (bIsBakedAdditive)
+	{
+		//When using baked additive ref pose is identity
+		OutAnimationPoseData.SetRefPose(bIsBakedAdditive);
+	}
+	else
+	{
+		// if retargeting is disabled, we initialize pose with 'Retargeting Source' ref pose.
+		if (bDisableRetargeting)
+		{
+			TArray<FTransform> const& AuthoredOnRefSkeleton = AnimSequence->GetRetargetTransforms();
+
+			const int32 NumBones = LODBoneIndexes.Num();
+			const int32 NumRawBones = AnimSequence->GetSkeleton()->GetReferenceSkeleton().GetRawBoneNum();
+
+			for (int i = 0; i < NumBones; i++)
+			{
+				const int32 SkeletonBoneIndex = LODBoneIndexes[i];
+
+				// Virtual bones are part of the retarget transform pose, so if the pose has not been updated (recently) there might be a mismatch
+				if (SkeletonBoneIndex < NumRawBones || AuthoredOnRefSkeleton.IsValidIndex(SkeletonBoneIndex))
+				{
+					OutAnimationPoseData.LocalTransforms[i] = AuthoredOnRefSkeleton[SkeletonBoneIndex];
+				}
+			}
+		}
+		else
+		{
+			OutAnimationPoseData.SetRefPose();
+		}
+	}
+
+#if WITH_EDITOR
+	const int32 NumTracks = bUseRawDataForPoseExtraction ? AnimSequence->GetDataModelInterface()->GetNumBoneTracks() : AnimSequence->CompressedData.CompressedTrackToSkeletonMapTable.Num();
+#else
+	const int32 NumTracks = AnimSequence->CompressedData.CompressedTrackToSkeletonMapTable.Num();
+#endif 
+	const bool bTreatAnimAsAdditive = (AnimSequence->IsValidAdditive() && !bUseRawDataForPoseExtraction); // Raw data is never additive
+	const FRootMotionReset RootMotionReset(AnimSequence->bEnableRootMotion, AnimSequence->RootMotionRootLock, AnimSequence->bForceRootLock, AnimSequence->ExtractRootTrackTransform(0.0f, /*&RequiredBones*/nullptr), bTreatAnimAsAdditive);
+
+#if WITH_EDITOR
+	// Evaluate raw (source) curve and bone data
+	if (bUseRawDataForPoseExtraction)
+	{
+		// TODO : Curves support
+		//{
+		//	FSkeletonRemappingCurve RemappedCurve(OutAnimationPoseData.GetCurve(), RequiredBones, AnimSequence->GetSkeleton());
+		//	FAnimationPoseData RemappedPoseData(OutAnimationPoseData.GetPose(), RemappedCurve.GetCurve(), OutAnimationPoseData.GetAttributes());
+
+		//	const UE::Anim::DataModel::FEvaluationContext EvaluationContext(ExtractionContext.CurrentTime, DataModelInterface->GetFrameRate(), GetRetargetTransformsSourceName(), GetRetargetTransforms(), Interpolation);
+		//	DataModelInterface->Evaluate(RemappedPoseData, EvaluationContext);
+		//}
+
+		//if ((ExtractionContext.bExtractRootMotion && RootMotionReset.bEnableRootMotion) || RootMotionReset.bForceRootLock)
+		//{
+		//	RootMotionReset.ResetRootBoneForRootMotion(OutPose[FCompactPoseBoneIndex(0)], RequiredBones);
+		//}
+
+		return;
+	}
+	else
+#endif // WITH_EDITOR
+		// Only try and evaluate compressed bone data if the animation contains any bone tracks
+		if (NumTracks != 0)
+		{
+			// Evaluate compressed bone data
+			FAnimSequenceDecompressionContext DecompContext(AnimSequence->GetSamplingFrameRate()
+				, AnimSequence->GetSamplingFrameRate().AsFrameTime(AnimSequence->GetPlayLength()).RoundToFrame().Value
+				, AnimSequence->Interpolation
+				, AnimSequence->GetRetargetTransformsSourceName()
+				, *AnimSequence->CompressedData.CompressedDataStructure
+				, AnimSequence->GetSkeleton()->GetRefLocalPoses()
+				, AnimSequence->CompressedData.CompressedTrackToSkeletonMapTable
+				, AnimSequence->GetSkeleton()
+				, AnimSequence->IsValidAdditive());
+
+			DecompressPose(OutAnimationPoseData, AnimSequence->CompressedData, ExtractionContext, DecompContext, AnimSequence->GetRetargetTransforms(), RootMotionReset);
+		}
+
+//	// TODO : Curves support
+//	// (Always) evaluate compressed curve data
+//	{
+//		// Scoped so that the RemappedCurve destructs and isn't kept alive longer than needed.
+//		FSkeletonRemappingCurve RemappedCurve(OutAnimationPoseData.GetCurve(), RequiredBones, GetSkeleton());
+//#if WITH_EDITOR
+//		// When evaluating from raw animation data, UE::Anim::BuildPoseFromModel will populate the curve data
+//		if (!bUseRawDataForPoseExtraction)
+//#endif // WITH_EDITOR
+//		{
+//			AnimSequence->EvaluateCurveData(RemappedCurve.GetCurve(), ExtractionContext.CurrentTime, bUseRawDataForPoseExtraction);
+//		}
+//	}
+
+	// TODO : attributes support
+	// Evaluate animation attributes (no compressed format yet)
+	//EvaluateAttributes(OutAnimationPoseData, ExtractionContext, false);
+}
+
+void FDecompressionTools::GetBonePose_Additive(const UAnimSequence* AnimSequence, const FAnimExtractContext& ExtractionContext, FAnimationLODPose& OutAnimationPoseData)
+{
+}
+
+void FDecompressionTools::GetBonePose_AdditiveMeshRotationOnly(const UAnimSequence* AnimSequence, const FAnimExtractContext& ExtractionContext, FAnimationLODPose& OutAnimationPoseData)
+{
+}
+
+
+// --- ---
+
+void FDecompressionTools::DecompressPose(FAnimationLODPose& OutAnimationPoseData,
+										const FCompressedAnimSequence& CompressedData,
+										const FAnimExtractContext& ExtractionContext,
+										FAnimSequenceDecompressionContext& DecompressionContext,
+										FName RetargetSource,
+										const FRootMotionReset& RootMotionReset)
+{
+	const TArray<FTransform>& RetargetTransforms = DecompressionContext.GetSourceSkeleton()->GetRefLocalPoses(RetargetSource);
+	DecompressPose(OutAnimationPoseData, CompressedData, ExtractionContext, DecompressionContext, RetargetTransforms, RootMotionReset);
+}
+
+void FDecompressionTools::DecompressPose(FAnimationLODPose& OutAnimationPoseData,
+										const FCompressedAnimSequence& CompressedData,
+										const FAnimExtractContext& ExtractionContext,
+										FAnimSequenceDecompressionContext& DecompressionContext,
+										const TArray<FTransform>& RetargetTransforms,
+										const FRootMotionReset& RootMotionReset)
+{
+	const FAnimationReferencePose& ReferencePose = OutAnimationPoseData.GetRefPose();
+	const TArrayView<const FBoneIndexType>LODBoneIndexes = OutAnimationPoseData.GetLODBoneIndexes();
+	const TArrayView<const FBoneIndexType> SkeletoonToLODBoneIndexes = ReferencePose.GetSkeletoonToLODBoneIndexes(0); // Full list of Skeleton to LOD conversion
+	const int32 NumLODBoneIndexes = LODBoneIndexes.Num();
+
+	const int32 NumTracks = CompressedData.CompressedTrackToSkeletonMapTable.Num();
+
+
+	const USkeleton* TargetSkeleton = OutAnimationPoseData.GetSkeletonAsset();
+	const FSkeletonRemapping& SkeletonRemapping = UE::Anim::FSkeletonRemappingRegistry::Get().GetRemapping(DecompressionContext.GetSourceSkeleton(), TargetSkeleton);
+
+	BoneTrackArray& RotationScalePairs = FGetBonePoseScratchArea::Get().RotationScalePairs;
+	BoneTrackArray& TranslationPairs = FGetBonePoseScratchArea::Get().TranslationPairs;
+	BoneTrackArray& AnimScaleRetargetingPairs = FGetBonePoseScratchArea::Get().AnimScaleRetargetingPairs;
+	BoneTrackArray& AnimRelativeRetargetingPairs = FGetBonePoseScratchArea::Get().AnimRelativeRetargetingPairs;
+	BoneTrackArray& OrientAndScaleRetargetingPairs = FGetBonePoseScratchArea::Get().OrientAndScaleRetargetingPairs;
+
+	// build a list of desired bones
+	RotationScalePairs.Reset();
+	TranslationPairs.Reset();
+	AnimScaleRetargetingPairs.Reset();
+	AnimRelativeRetargetingPairs.Reset();
+	OrientAndScaleRetargetingPairs.Reset();
+
+	// Optimization: assuming first index is root bone. That should always be the case in Skeletons.
+	checkSlow((LODBoneIndexes[0] == FMeshPoseBoneIndex(0).GetInt()));
+	// this is not guaranteed for AnimSequences though... If Root is not animated, Track will not exist.
+	const bool bFirstTrackIsRootBone = (CompressedData.GetSkeletonIndexFromTrackIndex(0) == 0);
+
+	{
+		SCOPE_CYCLE_COUNTER(STAT_BuildAnimTrackPairs);
+
+		// Handle root bone separately if it is track 0. so we start w/ Index 1.
+		for (int32 TrackIndex = (bFirstTrackIsRootBone ? 1 : 0); TrackIndex < NumTracks; TrackIndex++)
+		{
+			const int32 SourceSkeletonBoneIndex = CompressedData.GetSkeletonIndexFromTrackIndex(TrackIndex);
+			const int32 TargetSkeletonBoneIndex = SkeletonRemapping.IsValid() ? SkeletonRemapping.GetTargetSkeletonBoneIndex(SourceSkeletonBoneIndex) : SourceSkeletonBoneIndex;
+
+			if (TargetSkeletonBoneIndex != INDEX_NONE)
+			{
+				const int32 LODBoneIndex = TargetSkeletonBoneIndex < SkeletoonToLODBoneIndexes.Num() ? SkeletoonToLODBoneIndexes[TargetSkeletonBoneIndex] : INDEX_NONE;
+
+				if (LODBoneIndex != INDEX_NONE && LODBoneIndex < NumLODBoneIndexes) // skip bones not in current LOD
+				{
+					RotationScalePairs.Add(BoneTrackPair(LODBoneIndex, TrackIndex));
+
+					// Skip extracting translation component for EBoneTranslationRetargetingMode::Skeleton.
+					switch (TargetSkeleton->GetBoneTranslationRetargetingMode(TargetSkeletonBoneIndex, OutAnimationPoseData.GetDisableRetargeting()))
+					{
+					case EBoneTranslationRetargetingMode::Animation:
+						TranslationPairs.Add(BoneTrackPair(LODBoneIndex, TrackIndex));
+						break;
+					case EBoneTranslationRetargetingMode::AnimationScaled:
+						TranslationPairs.Add(BoneTrackPair(LODBoneIndex, TrackIndex));
+						AnimScaleRetargetingPairs.Add(BoneTrackPair(LODBoneIndex, SourceSkeletonBoneIndex));
+						break;
+					case EBoneTranslationRetargetingMode::AnimationRelative:
+						TranslationPairs.Add(BoneTrackPair(LODBoneIndex, TrackIndex));
+
+						// With baked additives, we can skip 'AnimationRelative' tracks, as the relative transform gets canceled out.
+						// (A1 + Rel) - (A2 + Rel) = A1 - A2.
+						if (!DecompressionContext.IsAdditiveAnimation())
+						{
+							AnimRelativeRetargetingPairs.Add(BoneTrackPair(LODBoneIndex, SourceSkeletonBoneIndex));
+						}
+						break;
+					case EBoneTranslationRetargetingMode::OrientAndScale:
+						TranslationPairs.Add(BoneTrackPair(LODBoneIndex, TrackIndex));
+
+						// Additives remain additives, they're not retargeted.
+						if (!DecompressionContext.IsAdditiveAnimation())
+						{
+							OrientAndScaleRetargetingPairs.Add(BoneTrackPair(LODBoneIndex, SourceSkeletonBoneIndex));
+						}
+						break;
+					}
+				}
+			}
+		}
+	}
+
+	{
+		SCOPE_CYCLE_COUNTER(STAT_ExtractPoseFromAnimData);
+		CSV_SCOPED_TIMING_STAT(Animation, ExtractPoseFromAnimData);
+		CSV_CUSTOM_STAT(Animation, NumberOfExtractedAnimations, 1, ECsvCustomStatOp::Accumulate);
+
+		DecompressionContext.Seek(ExtractionContext.CurrentTime);
+
+		// Handle Root Bone separately
+		if (bFirstTrackIsRootBone)
+		{
+			const int32 TrackIndex = 0;
+			const int32 LODRootBone = 0;
+			FTransform RootAtom = OutAnimationPoseData.LocalTransforms[0];
+
+			CompressedData.BoneCompressionCodec->DecompressBone(DecompressionContext, TrackIndex, RootAtom);
+
+			// Retarget the root onto the target skeleton (correcting for differences in rest poses)
+			if (SkeletonRemapping.RequiresReferencePoseRetarget())
+			{
+				const int32 TargetSkeletonBoneIndex = 0;
+
+				if (DecompressionContext.IsAdditiveAnimation())
+				{
+					RootAtom.SetRotation(SkeletonRemapping.RetargetAdditiveRotationToTargetSkeleton(TargetSkeletonBoneIndex, RootAtom.GetRotation()));
+					if (TargetSkeleton->GetBoneTranslationRetargetingMode(TargetSkeletonBoneIndex, OutAnimationPoseData.GetDisableRetargeting()) != EBoneTranslationRetargetingMode::Skeleton)
+					{
+						RootAtom.SetTranslation(SkeletonRemapping.RetargetAdditiveTranslationToTargetSkeleton(TargetSkeletonBoneIndex, RootAtom.GetTranslation()));
+					}
+				}
+				else
+				{
+					RootAtom.SetRotation(SkeletonRemapping.RetargetBoneRotationToTargetSkeleton(TargetSkeletonBoneIndex, RootAtom.GetRotation()));
+					if (TargetSkeleton->GetBoneTranslationRetargetingMode(TargetSkeletonBoneIndex, OutAnimationPoseData.GetDisableRetargeting()) != EBoneTranslationRetargetingMode::Skeleton)
+					{
+						RootAtom.SetTranslation(SkeletonRemapping.RetargetBoneTranslationToTargetSkeleton(TargetSkeletonBoneIndex, RootAtom.GetTranslation()));
+					}
+				}
+			}
+
+			// @laurent - we should look into splitting rotation and translation tracks, so we don't have to process translation twice.
+			RetargetBoneTransform(OutAnimationPoseData.GetRefPose()
+				, DecompressionContext.GetSourceSkeleton()
+				, OutAnimationPoseData.GetSkeletonAsset()
+				, DecompressionContext.AnimName
+				, RetargetTransforms
+				, RootAtom
+				, 0
+				, LODRootBone
+				, DecompressionContext.IsAdditiveAnimation()
+				, OutAnimationPoseData.GetDisableRetargeting());
+
+			OutAnimationPoseData.LocalTransforms[0] = RootAtom;
+		}
+
+		if (RotationScalePairs.Num() > 0)
+		{
+#if DEFAULT_SOA
+			// get the remaining bone atoms
+			CompressedData.BoneCompressionCodec->DecompressPose(DecompressionContext, UE::Anim::FAnimPoseDecompressionData(RotationScalePairs, TranslationPairs, RotationScalePairs, OutAnimationPoseData.LocalTransforms.Rotations, OutAnimationPoseData.LocalTransforms.Translations, OutAnimationPoseData.LocalTransforms.Scales3D));
+#else
+			// get the remaining bone atoms
+			TArrayView<FTransform> OutPoseBones = OutAnimationPoseData.LocalTransforms.Transforms;
+			CompressedData.BoneCompressionCodec->DecompressPose(DecompressionContext, RotationScalePairs, TranslationPairs, RotationScalePairs, OutPoseBones);
+#endif
+		}
+	}
+
+	// Retarget the pose onto the target skeleton (correcting for differences in rest poses)
+	if (SkeletonRemapping.RequiresReferencePoseRetarget())
+	{
+		const int32 LODNumBones = OutAnimationPoseData.GetNumBones();
+
+		if (DecompressionContext.IsAdditiveAnimation())
+		{
+			for (int32 LODBoneIndex = bFirstTrackIsRootBone ? 1 : 0; LODBoneIndex < LODNumBones; ++LODBoneIndex)
+			{
+				const int32 TargetSkeletonBoneIndex = ReferencePose.GetSkeletonBoneIndexFromLODBoneIndex(LODBoneIndex);
+				OutAnimationPoseData.LocalTransforms[LODBoneIndex].SetRotation(SkeletonRemapping.RetargetAdditiveRotationToTargetSkeleton(TargetSkeletonBoneIndex, OutAnimationPoseData.LocalTransforms[LODBoneIndex].GetRotation()));
+				if (TargetSkeleton->GetBoneTranslationRetargetingMode(TargetSkeletonBoneIndex, OutAnimationPoseData.GetDisableRetargeting()) != EBoneTranslationRetargetingMode::Skeleton)
+				{
+					OutAnimationPoseData.LocalTransforms[LODBoneIndex].SetTranslation(SkeletonRemapping.RetargetAdditiveTranslationToTargetSkeleton(TargetSkeletonBoneIndex, OutAnimationPoseData.LocalTransforms[LODBoneIndex].GetTranslation()));
+				}
+			}
+		}
+		else
+		{
+			for (int32 LODBoneIndex = (bFirstTrackIsRootBone ? 1 : 0); LODBoneIndex < LODNumBones; ++LODBoneIndex)
+			{
+				const int32 TargetSkeletonBoneIndex = LODBoneIndexes[LODBoneIndex]; // ReferencePose.GetSkeletonBoneIndexFromLODBoneIndex(LODBoneIndex);
+				OutAnimationPoseData.LocalTransforms[LODBoneIndex].SetRotation(SkeletonRemapping.RetargetBoneRotationToTargetSkeleton(TargetSkeletonBoneIndex, OutAnimationPoseData.LocalTransforms[LODBoneIndex].GetRotation()));
+				if (TargetSkeleton->GetBoneTranslationRetargetingMode(TargetSkeletonBoneIndex, OutAnimationPoseData.GetDisableRetargeting()) != EBoneTranslationRetargetingMode::Skeleton)
+				{
+					OutAnimationPoseData.LocalTransforms[LODBoneIndex].SetTranslation(SkeletonRemapping.RetargetBoneTranslationToTargetSkeleton(TargetSkeletonBoneIndex, OutAnimationPoseData.LocalTransforms[LODBoneIndex].GetTranslation()));
+				}
+			}
+		}
+	}
+
+	// Once pose has been extracted, snap root bone back to first frame if we are extracting root motion.
+	if ((ExtractionContext.bExtractRootMotion && RootMotionReset.bEnableRootMotion) || RootMotionReset.bForceRootLock)
+	{
+		//RootMotionReset.ResetRootBoneForRootMotion(OutAnimationPoseData.LocalTransforms[FCompactPoseBoneIndex(0)], RequiredBones);
+	}
+
+	// Anim Scale Retargeting
+	int32 const NumBonesToScaleRetarget = AnimScaleRetargetingPairs.Num();
+	if (NumBonesToScaleRetarget > 0)
+	{
+		TArray<FTransform> const& AuthoredOnRefSkeleton = RetargetTransforms;
+
+		for (const BoneTrackPair& BonePair : AnimScaleRetargetingPairs)
+		{
+			const int32 LODBoneIndex(BonePair.AtomIndex);
+			const int32 SourceSkeletonBoneIndex = BonePair.TrackIndex;
+
+			// @todo - precache that in FBoneContainer when we have SkeletonIndex->TrackIndex mapping. So we can just apply scale right away.
+			float const SourceTranslationLength = AuthoredOnRefSkeleton[SourceSkeletonBoneIndex].GetTranslation().Size();
+			if (SourceTranslationLength > UE_KINDA_SMALL_NUMBER)
+			{
+				float const TargetTranslationLength = ReferencePose.GetRefPoseTranslation(LODBoneIndex).Size();
+				OutAnimationPoseData.LocalTransforms[LODBoneIndex].ScaleTranslation(TargetTranslationLength / SourceTranslationLength);
+			}
+		}
+	}
+
+	// Anim Relative Retargeting
+	int32 const NumBonesToRelativeRetarget = AnimRelativeRetargetingPairs.Num();
+	if (NumBonesToRelativeRetarget > 0)
+	{
+		TArray<FTransform> const& AuthoredOnRefSkeleton = RetargetTransforms;
+
+		for (const BoneTrackPair& BonePair : AnimRelativeRetargetingPairs)
+		{
+			const int32 LODBoneIndex(BonePair.AtomIndex);
+			const int32 SourceSkeletonBoneIndex = BonePair.TrackIndex;
+
+			const FTransform& RefPoseTransform = ReferencePose.GetRefPoseTransform(LODBoneIndex);
+
+			// Remap the base pose onto the target skeleton so that we are working entirely in target space
+			FTransform BaseTransform = AuthoredOnRefSkeleton[SourceSkeletonBoneIndex];
+			if (SkeletonRemapping.RequiresReferencePoseRetarget())
+			{
+				const int32 TargetSkeletonBoneIndex = SkeletonRemapping.GetTargetSkeletonBoneIndex(SourceSkeletonBoneIndex);
+				BaseTransform = SkeletonRemapping.RetargetBoneTransformToTargetSkeleton(TargetSkeletonBoneIndex, BaseTransform);
+			}
+
+			// Apply the retargeting as if it were an additive difference between the current skeleton and the retarget skeleton. 
+			OutAnimationPoseData.LocalTransforms[LODBoneIndex].SetRotation(OutAnimationPoseData.LocalTransforms[LODBoneIndex].GetRotation() * BaseTransform.GetRotation().Inverse() * RefPoseTransform.GetRotation());
+			OutAnimationPoseData.LocalTransforms[LODBoneIndex].SetTranslation(OutAnimationPoseData.LocalTransforms[LODBoneIndex].GetTranslation() + (RefPoseTransform.GetTranslation() - BaseTransform.GetTranslation()));
+			OutAnimationPoseData.LocalTransforms[LODBoneIndex].SetScale3D(OutAnimationPoseData.LocalTransforms[LODBoneIndex].GetScale3D() * (RefPoseTransform.GetScale3D() * BaseTransform.GetSafeScaleReciprocal(BaseTransform.GetScale3D())));
+			OutAnimationPoseData.LocalTransforms[LODBoneIndex].NormalizeRotation();
+		}
+	}
+
+	// TODO : Have to recreate GetRetargetSourceCachedData
+	// Translation 'Orient and Scale' Translation Retargeting
+	//const int32 NumBonesToOrientAndScaleRetarget = OrientAndScaleRetargetingPairs.Num();
+	//if (NumBonesToOrientAndScaleRetarget > 0)
+	//{
+	//	const FRetargetSourceCachedData& RetargetSourceCachedData = RequiredBones.GetRetargetSourceCachedData(DecompressionContext.AnimName, SkeletonRemapping, RetargetTransforms);
+	//	const TArray<FOrientAndScaleRetargetingCachedData>& OrientAndScaleDataArray = RetargetSourceCachedData.OrientAndScaleData;
+	//	const TArray<int32>& CompactPoseIndexToOrientAndScaleIndex = RetargetSourceCachedData.CompactPoseIndexToOrientAndScaleIndex;
+
+	//	// If we have any cached retargeting data.
+	//	if ((OrientAndScaleDataArray.Num() > 0) && (CompactPoseIndexToOrientAndScaleIndex.Num() == RequiredBones.GetCompactPoseNumBones()))
+	//	{
+	//		for (int32 Index = 0; Index < NumBonesToOrientAndScaleRetarget; Index++)
+	//		{
+	//			const BoneTrackPair& BonePair = OrientAndScaleRetargetingPairs[Index];
+	//			const FCompactPoseBoneIndex CompactPoseBoneIndex(BonePair.AtomIndex);
+	//			const int32 OrientAndScaleIndex = CompactPoseIndexToOrientAndScaleIndex[CompactPoseBoneIndex.GetInt()];
+	//			if (OrientAndScaleIndex != INDEX_NONE)
+	//			{
+	//				const FOrientAndScaleRetargetingCachedData& OrientAndScaleData = OrientAndScaleDataArray[OrientAndScaleIndex];
+	//				FTransform& BoneTransform = OutAnimationPoseData.LocalTransforms[CompactPoseBoneIndex];
+	//				const FVector AnimatedTranslation = BoneTransform.GetTranslation();
+
+	//				// If Translation is not animated, we can just copy the TargetTranslation. No retargeting needs to be done.
+	//				const FVector NewTranslation = (AnimatedTranslation - OrientAndScaleData.SourceTranslation).IsNearlyZero(BONE_TRANS_RT_ORIENT_AND_SCALE_PRECISION) ?
+	//					OrientAndScaleData.TargetTranslation :
+	//					OrientAndScaleData.TranslationDeltaOrient.RotateVector(AnimatedTranslation) * OrientAndScaleData.TranslationScale;
+
+	//				BoneTransform.SetTranslation(NewTranslation);
+	//			}
+	//		}
+	//	}
+	//}
+}
+
+void FDecompressionTools::RetargetBoneTransform(const FAnimationReferencePose& ReferencePose
+	, const USkeleton* SourceSkeleton
+	, const USkeleton* TargetSkeleton
+	, const FName& RetargetSource
+	, FTransform& BoneTransform
+	, const int32 SkeletonBoneIndex
+	, const int32 LODBoneIndex
+	, const bool bIsBakedAdditive
+	, const bool bDisableRetargeting)
+{
+	if (SourceSkeleton)
+	{
+		const TArray<FTransform>& RetargetTransforms = SourceSkeleton->GetRefLocalPoses(RetargetSource);
+		RetargetBoneTransform(ReferencePose, SourceSkeleton, TargetSkeleton, RetargetSource, RetargetTransforms, BoneTransform, SkeletonBoneIndex, LODBoneIndex, bIsBakedAdditive, bDisableRetargeting);
+	}
+}
+
+void FDecompressionTools::RetargetBoneTransform(const FAnimationReferencePose& ReferencePose
+	, const USkeleton* SourceSkeleton
+	, const USkeleton* TargetSkeleton
+	, const FName& SourceName
+	, const TArray<FTransform>& RetargetTransforms
+	, FTransform& BoneTransform
+	, const int32 SkeletonBoneIndex
+	, const int32 LODBoneIndex
+	, const bool bIsBakedAdditive
+	, const bool bDisableRetargeting)
+{
+	check(!RetargetTransforms.IsEmpty());
+	if (SourceSkeleton)
+	{
+		//const USkeleton* TargetSkeleton = RequiredBones.GetSkeletonAsset();
+		const FSkeletonRemapping& SkeletonRemapping = UE::Anim::FSkeletonRemappingRegistry::Get().GetRemapping(SourceSkeleton, TargetSkeleton);
+
+		const int32 TargetSkeletonBoneIndex = ReferencePose.GetSkeletonBoneIndexFromLODBoneIndex(LODBoneIndex);
+		const int32 SourceSkeletonBoneIndex = SkeletonRemapping.IsValid() ? SkeletonRemapping.GetSourceSkeletonBoneIndex(TargetSkeletonBoneIndex) : SkeletonBoneIndex;
+
+		switch (TargetSkeleton->GetBoneTranslationRetargetingMode(TargetSkeletonBoneIndex, bDisableRetargeting))
+		{
+		case EBoneTranslationRetargetingMode::AnimationScaled:
+		{
+			// @todo - precache that in FBoneContainer when we have SkeletonIndex->TrackIndex mapping. So we can just apply scale right away.
+			const TArray<FTransform>& SkeletonRefPoseArray = RetargetTransforms;
+			const float SourceTranslationLength = SkeletonRefPoseArray[SourceSkeletonBoneIndex].GetTranslation().Size();
+			if (SourceTranslationLength > UE_KINDA_SMALL_NUMBER)
+			{
+				const float TargetTranslationLength = ReferencePose.GetRefPoseTranslation(LODBoneIndex).Size();
+				BoneTransform.ScaleTranslation(TargetTranslationLength / SourceTranslationLength);
+			}
+			break;
+		}
+
+		case EBoneTranslationRetargetingMode::Skeleton:
+		{
+			BoneTransform.SetTranslation(bIsBakedAdditive ? FVector::ZeroVector : ReferencePose.GetRefPoseTranslation(LODBoneIndex));
+			break;
+		}
+
+		case EBoneTranslationRetargetingMode::AnimationRelative:
+		{
+			// With baked additive animations, Animation Relative delta gets canceled out, so we can skip it.
+			// (A1 + Rel) - (A2 + Rel) = A1 - A2.
+			if (!bIsBakedAdditive)
+			{
+				const TArray<FTransform>& AuthoredOnRefSkeleton = RetargetTransforms;
+				const FTransform& RefPoseTransform = ReferencePose.GetRefPoseTransform(LODBoneIndex);
+
+				// Remap the base pose onto the target skeleton so that we are working entirely in target space
+				FTransform BaseTransform = AuthoredOnRefSkeleton[SourceSkeletonBoneIndex];
+				if (SkeletonRemapping.RequiresReferencePoseRetarget())
+				{
+					BaseTransform = SkeletonRemapping.RetargetBoneTransformToTargetSkeleton(TargetSkeletonBoneIndex, BaseTransform);
+				}
+
+				// Apply the retargeting as if it were an additive difference between the current skeleton and the retarget skeleton. 
+				BoneTransform.SetRotation(BoneTransform.GetRotation() * BaseTransform.GetRotation().Inverse() * RefPoseTransform.GetRotation());
+				BoneTransform.SetTranslation(BoneTransform.GetTranslation() + (RefPoseTransform.GetTranslation() - BaseTransform.GetTranslation()));
+				BoneTransform.SetScale3D(BoneTransform.GetScale3D() * (RefPoseTransform.GetScale3D() * BaseTransform.GetSafeScaleReciprocal(BaseTransform.GetScale3D())));
+				BoneTransform.NormalizeRotation();
+			}
+			break;
+		}
+
+		case EBoneTranslationRetargetingMode::OrientAndScale:
+		{
+			if (!bIsBakedAdditive)
+			{
+				//const FRetargetSourceCachedData& RetargetSourceCachedData = RequiredBones.GetRetargetSourceCachedData(SourceName, SkeletonRemapping, RetargetTransforms);
+				//const TArray<FOrientAndScaleRetargetingCachedData>& OrientAndScaleDataArray = RetargetSourceCachedData.OrientAndScaleData;
+				//const TArray<int32>& CompactPoseIndexToOrientAndScaleIndex = RetargetSourceCachedData.CompactPoseIndexToOrientAndScaleIndex;
+
+				//// If we have any cached retargeting data.
+				//if ((OrientAndScaleDataArray.Num() > 0) && (CompactPoseIndexToOrientAndScaleIndex.Num() == RequiredBones.GetCompactPoseNumBones()))
+				//{
+				//	const int32 OrientAndScaleIndex = CompactPoseIndexToOrientAndScaleIndex[BoneIndex.GetInt()];
+				//	if (OrientAndScaleIndex != INDEX_NONE)
+				//	{
+				//		const FOrientAndScaleRetargetingCachedData& OrientAndScaleData = OrientAndScaleDataArray[OrientAndScaleIndex];
+				//		const FVector AnimatedTranslation = BoneTransform.GetTranslation();
+
+				//		// If Translation is not animated, we can just copy the TargetTranslation. No retargeting needs to be done.
+				//		const FVector NewTranslation = (AnimatedTranslation - OrientAndScaleData.SourceTranslation).IsNearlyZero(BONE_TRANS_RT_ORIENT_AND_SCALE_PRECISION) ?
+				//			OrientAndScaleData.TargetTranslation :
+				//			OrientAndScaleData.TranslationDeltaOrient.RotateVector(AnimatedTranslation) * OrientAndScaleData.TranslationScale;
+
+				//		BoneTransform.SetTranslation(NewTranslation);
+				//	}
+				//}
+			}
+			break;
+		}
+		}
+	}
+}
+
+} // end namespace UE::AnimNext::Interface
diff --git a/Engine/Plugins/Experimental/AnimNext/Source/AnimNextInterface/Private/AnimationGenerationTools.cpp b/Engine/Plugins/Experimental/AnimNext/Source/AnimNextInterface/Private/AnimationGenerationTools.cpp
new file mode 100644
--- /dev/null
+++ b/Engine/Plugins/Experimental/AnimNext/Source/AnimNextInterface/Private/AnimationGenerationTools.cpp
@@ -0,0 +1,485 @@
+// Copyright Epic Games, Inc. All Rights Reserved.
+
+#include "AnimationGenerationTools.h"
+
+#include "AnimationReferencePose.h"
+#include "AnimationRuntime.h"
+#include "Animation/AnimNodeBase.h"
+#include "BoneContainer.h"
+#include "Components/SkeletalMeshComponent.h"
+#include "Engine/SkeletalMesh.h"
+#include "Engine/SkeletalMeshSocket.h"
+#include "SkeletalMeshSceneProxy.h"
+#include "Engine/SkinnedAssetCommon.h"
+#include "PhysicsEngine/PhysicsAsset.h"
+#include "Rendering/SkeletalMeshModel.h"
+#include "Rendering/SkeletalMeshRenderData.h"
+
+// auto command
+#include "AssetRegistry/AssetData.h"
+#include "AssetRegistry/AssetRegistryModule.h"
+
+
+DEFINE_LOG_CATEGORY_STATIC(LogAnimGenerationTools, Log, All)
+
+
+namespace UE::AnimNext::Interface
+{
+
+struct FCompareBoneIndexType
+{
+	FORCEINLINE bool operator()(const FBoneIndexType& A, const FBoneIndexType& B) const
+	{
+		return A < B;
+	}
+};
+
+
+/*static*/ bool FGenerationTools::GenerateReferencePose(const USkeletalMeshComponent* SkeletalMeshComponent
+	, const USkeletalMesh* SkeletalMesh
+	, FAnimationReferencePose& OutAnimationReferencePose)
+{
+	using namespace UE::AnimNext::Interface;
+
+	bool ReferencePoseGenerated = false;
+
+	if (SkeletalMesh == nullptr)
+	{
+		return false;
+	}
+
+	UE_LOG(LogAnimGenerationTools, VeryVerbose, TEXT("Generating CanonicalBoneSet for SkeletalMesh %s."), *SkeletalMesh->GetPathName());
+
+	const FSkeletalMeshRenderData* SkelMeshRenderData = (SkeletalMeshComponent != nullptr)
+		? SkeletalMeshComponent->GetSkeletalMeshRenderData()
+		: SkeletalMesh->GetResourceForRendering();
+
+	if (!SkelMeshRenderData)
+	{
+		UE_LOG(LogAnimGenerationTools, Warning, TEXT("Error generating CanonicalBoneSet for SkeletalMesh %s. No SkeletalMeshRenderData."), *SkeletalMesh->GetPathName());
+		return false;
+	}
+
+	const TIndirectArray<FSkeletalMeshLODRenderData>& LODRenderData = SkelMeshRenderData->LODRenderData;
+	const int32 NumLODs = LODRenderData.Num();
+
+	TArray<FGenerationLODData> GenerationLODData;
+	GenerationLODData.Reserve(NumLODs);
+	GenerationLODData.AddDefaulted(NumLODs);
+
+	TArray<FGenerationLODData> ComponentSpaceGenerationLODData;
+	ComponentSpaceGenerationLODData.Reserve(NumLODs);
+	ComponentSpaceGenerationLODData.AddDefaulted(NumLODs);
+
+	if (NumLODs > 0)
+	{
+		// Generate LOD0 bones
+		TArray<FBoneIndexType>& RequiredBones_LOD0 = GenerationLODData[0].RequiredBones;
+
+		constexpr int32 LOD0Index = 0;
+		GenerateRawLODData(SkeletalMeshComponent, SkeletalMesh, LOD0Index, LODRenderData, RequiredBones_LOD0, ComponentSpaceGenerationLODData[LOD0Index].RequiredBones);
+
+		// Now calculate the LODs > 1
+
+		constexpr int32 StartLOD = 1;
+		GenerateLODData(SkeletalMeshComponent, SkeletalMesh, StartLOD, NumLODs, LODRenderData, RequiredBones_LOD0, GenerationLODData, ComponentSpaceGenerationLODData);
+
+		// Check if the sockets are set to always animate, else the component space requires separated data (different bone indexes)
+		bool bCanGenerateSingleBonesList = CheckSkeletalAllMeshSocketsAlwaysAnimate(SkeletalMesh);
+
+		if (bCanGenerateSingleBonesList)
+		{
+			bCanGenerateSingleBonesList &= CheckExcludedBones(NumLODs, GenerationLODData, SkeletalMesh);
+			//bCanGenerateSingleBonesList &= CheckExcludedBones(NumLODs, ComponentSpaceGenerationLODData, SkeletalMesh); // Commented : right now we only support skeletal meshes with all the sockets set to always animate
+
+			TArray<FBoneIndexType> OrderedBoneList;
+			TArray<FBoneIndexType> ComponentSpaceOrderedBoneList;
+
+			if (bCanGenerateSingleBonesList)
+			{
+				bCanGenerateSingleBonesList &= GenerateOrderedBoneList(SkeletalMesh, GenerationLODData, OrderedBoneList);
+				//bCanGenerateSingleBonesList &= GenerateOrderedBoneList(SkeletalMesh, ComponentSpaceGenerationLODData, ComponentSpaceOrderedBoneList); // Commented : right now we only support skeletal meshes with all the sockets set to always animate
+				//bCanGenerateSingleBonesList &= OrderedBoneList == ComponentSpaceOrderedBoneList; // Commented : right now we only support skeletal meshes with all the sockets set to always animate
+			}
+
+			if (bCanGenerateSingleBonesList)
+			{
+				OutAnimationReferencePose.GenerationFlags = FAnimationReferencePose::EGenerationFlags::FastPath;
+				OutAnimationReferencePose.SkeletalMesh = SkeletalMesh;
+				OutAnimationReferencePose.Skeleton = SkeletalMesh->GetSkeleton();
+
+				TArray<int32> LODNumBones;
+				LODNumBones.Reset(NumLODs);
+				for (int32 LODIndex = 0; LODIndex < NumLODs; ++LODIndex)
+				{
+					LODNumBones.Add(GenerationLODData[LODIndex].RequiredBones.Num());
+				}
+
+				// Generate a Skeleton to LOD look up table
+				const int32 NumOrderedBones = OrderedBoneList.Num();
+
+				TArray<FBoneIndexType> SkeletonToLODBoneList;
+				SkeletonToLODBoneList.SetNumZeroed(NumOrderedBones);
+				for (int i = 0; i < NumOrderedBones; ++i)
+				{
+					SkeletonToLODBoneList[OrderedBoneList[i]] = i;
+				}
+
+				OutAnimationReferencePose.Initialize(SkeletalMesh->GetRefSkeleton(), { OrderedBoneList }, { SkeletonToLODBoneList }, LODNumBones, bCanGenerateSingleBonesList);
+
+				ReferencePoseGenerated = true;
+			}
+		}
+		// TODO : Decide if we only allow skeletal meshes that allow for single bone list generation
+		//        Or create an alternative
+	}
+
+	return ReferencePoseGenerated;
+}
+
+/*static*/ void FGenerationTools::GenerateRawLODData(const USkeletalMeshComponent* SkeletalMeshComponent
+	, const USkeletalMesh* SkeletalMesh
+	, const int32 LODIndex
+	, const TIndirectArray<FSkeletalMeshLODRenderData>& LODRenderData
+	, TArray<FBoneIndexType>& OutRequiredBones
+	, TArray<FBoneIndexType>& OutFillComponentSpaceTransformsRequiredBones)
+{
+	const FSkeletalMeshLODRenderData& LODModel = LODRenderData[LODIndex];
+
+	if (LODRenderData[LODIndex].RequiredBones.Num() > 0)
+	{
+		// Start with the LODModel RequiredBones (precalculated LOD data)
+		OutRequiredBones = LODModel.RequiredBones;
+
+		// Add the Virtual bones from the skeleton
+		USkeletalMeshComponent::GetRequiredVirtualBones(SkeletalMesh, OutRequiredBones);
+
+		// Add any bones used by physics SkeletalBodySetups
+		const UPhysicsAsset* const PhysicsAsset = SkeletalMesh->GetPhysicsAsset();
+		// If we have a PhysicsAsset, we also need to make sure that all the bones used by it are always updated, as its used
+		// by line checks etc. We might also want to kick in the physics, which means having valid bone transforms.
+		if (PhysicsAsset)
+		{
+			USkeletalMeshComponent::GetPhysicsRequiredBones(SkeletalMesh, PhysicsAsset, OutRequiredBones);
+		}
+
+		// TODO - Make sure that bones with per-poly collision are also always updated.
+
+		// If we got a SkeletalMeshComponent, we can exclude invisible bones
+		if (SkeletalMeshComponent != nullptr)
+		{
+			USkeletalMeshComponent::ExcludeHiddenBones(SkeletalMeshComponent, SkeletalMesh, OutRequiredBones);
+		}
+
+		// Add in any bones that may be required when mirroring.
+		PRAGMA_DISABLE_DEPRECATION_WARNINGS
+		if (SkeletalMesh->SkelMirrorTable.Num() > 0 &&
+			(SkeletalMeshComponent == nullptr || SkeletalMesh->SkelMirrorTable.Num() == SkeletalMeshComponent->BoneSpaceTransforms.Num()))
+		{
+			USkeletalMeshComponent::GetMirroringRequiredBones(SkeletalMesh, OutRequiredBones);
+		}
+		PRAGMA_ENABLE_DEPRECATION_WARNINGS
+
+			// Get socket bones set to animate and bones required to fill the component space base transforms
+		TArray<FBoneIndexType> NeededBonesForFillComponentSpaceTransforms;
+		USkeletalMeshComponent::GetSocketRequiredBones(SkeletalMesh, OutRequiredBones, NeededBonesForFillComponentSpaceTransforms);
+
+		// If we got a SkeletalMeshComponent, we can include shadow shapes referenced bones
+		if (SkeletalMeshComponent != nullptr)
+		{
+			USkeletalMeshComponent::GetShadowShapeRequiredBones(SkeletalMeshComponent, OutRequiredBones);
+		}
+
+		// Ensure that we have a complete hierarchy down to those bones.
+		// This is needed because when we add bones (i.e. physics), the parent might not be in the list
+		FAnimationRuntime::EnsureParentsPresent(OutRequiredBones, SkeletalMesh->GetRefSkeleton());
+
+		OutFillComponentSpaceTransformsRequiredBones.Reset(OutRequiredBones.Num() + NeededBonesForFillComponentSpaceTransforms.Num());
+		OutFillComponentSpaceTransformsRequiredBones = OutRequiredBones;
+
+		NeededBonesForFillComponentSpaceTransforms.Sort();
+		USkeletalMeshComponent::MergeInBoneIndexArrays(OutFillComponentSpaceTransformsRequiredBones, NeededBonesForFillComponentSpaceTransforms);
+		FAnimationRuntime::EnsureParentsPresent(OutFillComponentSpaceTransformsRequiredBones, SkeletalMesh->GetRefSkeleton());
+	}
+}
+
+/*static*/ void FGenerationTools::GenerateLODData(const USkeletalMeshComponent* SkeletalMeshComponent
+	, const USkeletalMesh* SkeletalMesh
+	, const int32 StartLOD
+	, const int32 NumLODs
+	, const TIndirectArray<FSkeletalMeshLODRenderData>& LODRenderData
+	, const TArray<FBoneIndexType>& RequiredBones_LOD0
+	, TArray<FGenerationLODData>& GenerationLODData
+	, TArray<FGenerationLODData>& GenerationComponentSpaceLODData)
+{
+	for (int32 LODIndex = StartLOD; LODIndex < NumLODs; ++LODIndex)
+	{
+		TArray<FBoneIndexType>& RequiredBones = GenerationLODData[LODIndex].RequiredBones;
+		TArray<FBoneIndexType>& RequiredComponentSpaceBones = GenerationComponentSpaceLODData[LODIndex].RequiredBones;
+		GenerateRawLODData(SkeletalMeshComponent, SkeletalMesh, LODIndex, LODRenderData, RequiredBones, RequiredComponentSpaceBones);
+
+		const int32 ParentLODIndex = LODIndex - 1;
+
+		CalculateDifferenceFromParentLOD(LODIndex, GenerationLODData);
+		CalculateDifferenceFromParentLOD(LODIndex, GenerationComponentSpaceLODData);
+	}
+}
+
+// Calculate the bone indexes difference from LOD0 for LODIndex
+/*static*/ void FGenerationTools::CalculateDifferenceFromParentLOD(int32 LODIndex, TArray<FGenerationLODData>& GenerationLODData)
+{
+	const int32 ParentLODIndex = LODIndex - 1;
+
+	const TArray<FBoneIndexType>& RequiredBones_LOD0 = GenerationLODData[0].RequiredBones;
+	const TArray<FBoneIndexType>& RequiredBones = GenerationLODData[LODIndex].RequiredBones;
+	const TArray<FBoneIndexType>& RequiredBones_ParentLOD = GenerationLODData[ParentLODIndex].RequiredBones;
+	
+	TArray<FBoneIndexType>& ExcludedBonesFromLOD0 = GenerationLODData[LODIndex].ExcludedBones;
+	TArray<FBoneIndexType>& ExcludedBonesFromPrevLOD = GenerationLODData[LODIndex].ExcludedBonesFromPrevLOD;
+
+	DifferenceBoneIndexArrays(RequiredBones_LOD0, RequiredBones, ExcludedBonesFromLOD0);
+	DifferenceBoneIndexArrays(RequiredBones_ParentLOD, RequiredBones, ExcludedBonesFromPrevLOD);
+}
+
+/*static*/ bool FGenerationTools::CheckExcludedBones(const int32 NumLODs
+	, const TArray<FGenerationLODData>& GenerationLODData
+	, const USkeletalMesh* SkeletalMesh)
+{
+	bool bCanGenerateSingleBonesList = true;
+
+	for (int32 LODIndex = NumLODs - 1; LODIndex > 1; --LODIndex)
+	{
+		const FGenerationLODData& LODData = GenerationLODData[LODIndex];
+		const FGenerationLODData& PrevLODData = GenerationLODData[LODIndex - 1];
+
+		const bool bPrevSmaller = (PrevLODData.ExcludedBones.Num() <= LODData.ExcludedBones.Num());
+		if (bPrevSmaller == false)
+		{
+			bCanGenerateSingleBonesList = false;
+			UE_LOG(LogAnimGenerationTools, Warning, TEXT("SkeletalMesh %s canonical ordered bone set can not be stored as single bones list. LOD %d does not contain all the bones of LOD %d"), *SkeletalMesh->GetPathName(), LODIndex, LODIndex - 1);
+			break;
+		}
+
+		for (TArray<FBoneIndexType>::TConstIterator LODExcludedBonesIt(PrevLODData.ExcludedBones); LODExcludedBonesIt; ++LODExcludedBonesIt)
+		{
+			if (LODData.ExcludedBones.Contains(*LODExcludedBonesIt) == false)
+			{
+				bCanGenerateSingleBonesList = false;
+				UE_LOG(LogAnimGenerationTools, Warning, TEXT("SkeletalMesh %s canonical ordered bone set can not be stored in LOD order. LOD %d does not contain all the bones of LOD %d"), *SkeletalMesh->GetPathName(), LODIndex, LODIndex - 1);
+				break;
+			}
+		}
+	}
+	return bCanGenerateSingleBonesList;
+}
+
+/*static*/ bool FGenerationTools::GenerateOrderedBoneList(const USkeletalMesh* SkeletalMesh
+	, TArray<FGenerationLODData>& GenerationLODData
+	, TArray<FBoneIndexType>& OrderedBoneList)
+{
+	bool bCanFastPath = true;
+
+	OrderedBoneList = GenerationLODData[0].RequiredBones;
+
+	const int32 NumLODs = GenerationLODData.Num();
+
+	// Compute the common set of bones for all LODS (remove excluded bones for LODS > 0)
+	for (int32 LODIndex = 1; LODIndex < NumLODs; ++LODIndex)
+	{
+		const TArray<FBoneIndexType>& ExcludedBonesFromPrevLOD = GenerationLODData[LODIndex].ExcludedBonesFromPrevLOD;
+
+		const int32 NumExcludedBones = ExcludedBonesFromPrevLOD.Num();
+		for (int i = NumExcludedBones - 1; i >= 0; --i)
+		{
+			OrderedBoneList.Remove(ExcludedBonesFromPrevLOD[i]);
+		}
+	}
+
+	// Add the ExcludedBonesFromPrevLOD of each LOD, in inverse order 
+	for (int32 LODIndex = NumLODs - 1; LODIndex > 0; --LODIndex)
+	{
+		for (TArray<FBoneIndexType>::TConstIterator SetIt(GenerationLODData[LODIndex].ExcludedBonesFromPrevLOD); SetIt; ++SetIt)
+		{
+			OrderedBoneList.Add(*SetIt);
+		}
+	}
+
+	// Check if all the bones have the parents before themselves in the array
+	const int NumBones = OrderedBoneList.Num();
+	for (int i = 0; i < NumBones; ++i)
+	{
+		const FBoneIndexType BoneIndex = OrderedBoneList[i];
+		const FBoneIndexType BoneIndexParentIndex = SkeletalMesh->GetRefSkeleton().GetParentIndex(BoneIndex);
+
+		if (BoneIndexParentIndex >= 0)
+		{
+			int32 ParentIndexAtOrderedBoneList = -1;
+			OrderedBoneList.Find(BoneIndexParentIndex, ParentIndexAtOrderedBoneList);
+
+			if (ParentIndexAtOrderedBoneList >= i)
+			{
+				bCanFastPath = false;
+				UE_LOG(LogAnimGenerationTools, Warning, TEXT("Warning : SkeletalMesh [%s] has an invalid LOD setup."), *SkeletalMesh->GetPathName());
+				break;
+			}
+		}
+	}
+
+	return bCanFastPath;
+}
+
+/**
+ *	Utility for taking two arrays of bone indices, which must be strictly increasing, and finding the A - B.
+ *	That is - any items left in A, after removing B.
+ */
+/*static*/ void FGenerationTools::DifferenceBoneIndexArrays(const TArray<FBoneIndexType>& A, const TArray<FBoneIndexType>& B, TArray<FBoneIndexType>& Output)
+{
+	int32 APos = 0;
+	int32 BPos = 0;
+
+	while (APos < A.Num())
+	{
+		// check if any elements left in B
+		if (BPos < B.Num())
+		{
+			// If A Value < B Value, we have to add A Value to the output (these indexes are not in the substract array)
+			if (A[APos] < B[BPos])
+			{
+				Output.Add(A[APos]);
+				APos++;
+			}
+			// If APos value == BPos value, we have to skip A Value in the output (we want to substract B values from A)
+			// We increase BPos as we assume no duplicated indexes in the arrays
+			else if (A[APos] == B[BPos])
+			{
+				APos++;
+				BPos++;
+			}
+			// If APos value > BPos value, we have to increase BPos, until any of the other two conditions are valid again or we finish the elements in B
+			else
+			{
+				BPos++;
+			}
+		}
+		// If B is finished (no more elements), we just keep adding A to the output
+		else
+		{
+			Output.Add(A[APos]);
+			APos++;
+		}
+	}
+}
+
+/*static*/ bool FGenerationTools::CheckSkeletalAllMeshSocketsAlwaysAnimate(const USkeletalMesh* SkeletalMesh)
+{
+	bool bAllSocketsAlwaysAnimate = true;
+
+	TArray<USkeletalMeshSocket*> ActiveSocketList = SkeletalMesh->GetActiveSocketList();
+	for (const USkeletalMeshSocket* Socket : ActiveSocketList)
+	{
+		const int32 BoneIndex = SkeletalMesh->GetRefSkeleton().FindBoneIndex(Socket->BoneName);
+		if (BoneIndex != INDEX_NONE)
+		{
+			if (Socket->bForceAlwaysAnimated == false)
+			{
+				UE_LOG(LogSkeletalMesh, Warning, TEXT("SkeletalMesh %s canonical ordered bone set can not be stored as single bones list. Socket [%s] is not set to always animate."), *SkeletalMesh->GetPathName(), *Socket->GetName());
+				bAllSocketsAlwaysAnimate = false;
+			}
+		}
+	}
+
+	return bAllSocketsAlwaysAnimate;
+}
+
+// Converts AnimBP pose to AnimNext Pose
+// This function expects both poses to have the same LOD (number of bones and indexes)
+// The target pose should be assigned to the correct reference pose prior to this call
+/*static*/ void FGenerationTools::RemapPose(int32 LODLevel, const FPoseContext& SourcePose, const FAnimationReferencePose& RefPose, FAnimationLODPose& TargetPose)
+{
+	const TArrayView<const FBoneIndexType> LODBoneIndexes = RefPose.GetLODBoneIndexes(LODLevel);
+	const int32 NumBones = LODBoneIndexes.Num();
+
+	for (int i = 0; i < NumBones; ++i)
+	{
+		const auto& SkeletonBoneIndex = LODBoneIndexes[i];
+
+		const FCompactPoseBoneIndex CompactBoneIndex = SourcePose.Pose.GetBoneContainer().GetCompactPoseIndexFromSkeletonIndex(SkeletonBoneIndex);
+		if (CompactBoneIndex.GetInt() != INDEX_NONE)
+		{
+			TargetPose.LocalTransforms[i] = SourcePose.Pose[CompactBoneIndex];
+		}
+	}
+}
+
+// Converts AnimNext pose to AnimBP Pose
+// This function expects both poses to have the same LOD (number of bones and indexes)
+// The target pose should be assigned to the correct reference pose prior to this call
+/*static*/ void FGenerationTools::RemapPose(int32 LODLevel, const FAnimationReferencePose& RefPose, const FAnimationLODPose& SourcePose, FPoseContext& TargetPose)
+{
+	const TArrayView<const FBoneIndexType> LODBoneIndexes = RefPose.GetLODBoneIndexes(LODLevel);
+	const int32 NumBones = LODBoneIndexes.Num();
+
+	for (int i = 0; i < NumBones; ++i)
+	{
+		const auto& SkeletonBoneIndex = LODBoneIndexes[i];
+
+		const FCompactPoseBoneIndex CompactBoneIndex = TargetPose.Pose.GetBoneContainer().GetCompactPoseIndexFromSkeletonIndex(SkeletonBoneIndex);
+		if (CompactBoneIndex.GetInt() != INDEX_NONE)
+		{
+			TargetPose.Pose[CompactBoneIndex] = SourcePose.LocalTransforms[i];
+		}
+	}
+}
+
+} // end namespace UE::AnimNext::Interface
+
+
+
+namespace AnimNext::Tools::ConsoleCommands
+{
+	struct FHelper
+	{
+		static TArray<FBoneIndexType> ComputeExcludedBones(const USkeletalMesh* SkeletalMesh, const TArray<FBoneIndexType>& LODRequiredBones, const TArray<FBoneIndexType>& NextLODRequiredBones)
+		{
+			TArray<FBoneIndexType> ExcludedBones;
+
+			UE::AnimNext::Interface::FGenerationTools::DifferenceBoneIndexArrays(LODRequiredBones, NextLODRequiredBones, ExcludedBones);
+
+			return ExcludedBones;
+		}
+
+		static void CheckSkeletalMeshesLODs()
+		{
+			using namespace UE::AnimNext::Interface;
+
+			TArray<FAssetData> Assets;
+			IAssetRegistry& AssetRegistry = FModuleManager::LoadModuleChecked<FAssetRegistryModule>(TEXT("AssetRegistry")).Get();
+			AssetRegistry.GetAssetsByClass(USkeletalMesh::StaticClass()->GetClassPathName(), Assets);
+
+			const int32 NumAssets = Assets.Num();
+			for (int32 Idx = 0; Idx < NumAssets; Idx++)
+			{
+				const USkeletalMesh* SkeletalMesh = Cast<USkeletalMesh>(Assets[Idx].GetAsset());
+				if (SkeletalMesh != nullptr)
+				{
+					UE::AnimNext::Interface::FAnimationReferencePose OutAnimationReferencePose;
+
+					if (FGenerationTools::GenerateReferencePose(nullptr, SkeletalMesh, OutAnimationReferencePose))
+					{
+						UE_LOG(LogAnimGenerationTools, VeryVerbose, TEXT("[%d of %d] SkeletalMesh %s BoneReferencePose generated."), Idx + 1, NumAssets, *SkeletalMesh->GetPathName());
+					}
+				}
+				else
+				{
+					UE_LOG(LogAnimGenerationTools, VeryVerbose, TEXT("[%d of %d] SkeletalMesh %s BoneReferencePose error, asset could not be loaded."), Idx + 1, NumAssets, *SkeletalMesh->GetPathName());
+				}
+			}
+		}
+	};
+
+
+	FAutoConsoleCommand CheckSkeletalMeshesLODs(TEXT("animnext.tools.checkskeletalmesheslods"), TEXT(""), FConsoleCommandDelegate::CreateLambda([]()
+	{
+		FHelper::CheckSkeletalMeshesLODs();
+	}));
+} // end namespace AnimNext::Tools::ConsoleCommands
diff --git a/Engine/Plugins/Experimental/AnimNext/Source/AnimNextInterface/Private/AnimationReferencePose.cpp b/Engine/Plugins/Experimental/AnimNext/Source/AnimNextInterface/Private/AnimationReferencePose.cpp
new file mode 100644
--- /dev/null
+++ b/Engine/Plugins/Experimental/AnimNext/Source/AnimNextInterface/Private/AnimationReferencePose.cpp
@@ -0,0 +1,15 @@
+// Copyright Epic Games, Inc. All Rights Reserved.
+
+#include "AnimationReferencePose.h"
+#include "AnimNextInterfaceParam.h"
+
+namespace UE::AnimNext::Interface
+{
+
+IMPLEMENT_ANIM_NEXT_INTERFACE_PARAM_TYPE(FAnimationReferencePose, AnimationReferencePose);
+//IMPLEMENT_ANIM_NEXT_INTERFACE_PARAM_TYPE(FAnimationLODPose, AnimationLODPose);
+//IMPLEMENT_ANIM_NEXT_INTERFACE_PARAM_TYPE(FAnimationLODPoseStack, AnimationLODPoseStack);
+
+const FAnimTransform TransformAdditiveIdentity = FAnimTransform(FQuat::Identity, FVector::ZeroVector, FVector::ZeroVector);
+
+} // namespace UE::AnimNext::Interface
diff --git a/Engine/Plugins/Experimental/AnimNext/Source/AnimNextInterface/Public/AnimNextInterfaceContext.h b/Engine/Plugins/Experimental/AnimNext/Source/AnimNextInterface/Public/AnimNextInterfaceContext.h
--- a/Engine/Plugins/Experimental/AnimNext/Source/AnimNextInterface/Public/AnimNextInterfaceContext.h
+++ b/Engine/Plugins/Experimental/AnimNext/Source/AnimNextInterface/Public/AnimNextInterfaceContext.h
@@ -91,6 +91,66 @@ public:
 		Output		= 1 << 1,
 	};
 
+	// Add an Input parameter by value, copying it to the shared storage
+	template<typename ValueType>
+	FHParam AddInputValue(FName ParameterId, ValueType& Value)
+	{
+		// Values are added mutable even if source is const
+		const FParam::EFlags Flags = FParam::EFlags::Value | FParam::EFlags::Mutable;
+
+		return AddParameter(EParamType::Input, ParameterId, Flags, &Value);
+	}
+
+	// Add an Input parameter by value, copying it to the shared storage
+	template<typename ValueType>
+	FHParam AddInputValue(FName ParameterId, ValueType&& Value)
+	{
+		// Values are added mutable even if source is const
+		const FParam::EFlags Flags = FParam::EFlags::Value | FParam::EFlags::Mutable;
+
+		return AddParameter(EParamType::Input, ParameterId, Flags, &Value);
+	}
+
+	// Add an Input parameter by value, copying it to the shared storage
+	template<typename ValueType>
+	FHParam AddInputValue(FName ParameterId, ValueType* Value)
+	{
+		// Values are added mutable even if source is const
+		const FParam::EFlags Flags = FParam::EFlags::Value | FParam::EFlags::Mutable;
+
+		return AddParameter(EParamType::Input, ParameterId, Flags, Value);
+	}
+
+	// Add an Output parameter by value, copying it to the shared storage
+	template<typename ValueType>
+	FHParam AddOutputValue(FName ParameterId, ValueType& Value)
+	{
+		// Values are added mutable even if source is const
+		const FParam::EFlags Flags = FParam::EFlags::Value | FParam::EFlags::Mutable;
+
+		return AddParameter(EParamType::Output, ParameterId, Flags, &Value);
+	}
+
+	// Add an Output parameter by value, copying it to the shared storage
+	template<typename ValueType>
+	FHParam AddOutputValue(FName ParameterId, ValueType&& Value)
+	{
+		// Values are added mutable even if source is const
+		const FParam::EFlags Flags = FParam::EFlags::Value | FParam::EFlags::Mutable;
+
+		return AddParameter(EParamType::Output, ParameterId, Flags, &Value);
+	}
+
+	// Add an Output parameter by value, copying it to the shared storage
+	template<typename ValueType>
+	FHParam AddOutputValue(FName ParameterId, ValueType* Value)
+	{
+		// Values are added mutable even if source is const
+		const FParam::EFlags Flags = FParam::EFlags::Value | FParam::EFlags::Mutable;
+
+		return AddParameter(EParamType::Output, ParameterId, Flags, Value);
+	}
+
 	// Add a parameter by value, copying it to the shared storage
 	template<typename ValueType>
 	FHParam AddValue(EParamType InParamType, FName ParameterId, ValueType &Value)
@@ -121,6 +181,44 @@ public:
 		return AddParameter(InParamType, ParameterId, Flags, Value);
 	}
 
+	// --- Parameters by refecence / pointer ---
+
+	// Add an Input parameter by reference, adding just a pointer to the shared storage
+	template<typename ValueType>
+	FHParam AddInputReference(FName ParameterId, ValueType& Value)
+	{
+		const FParam::EFlags Flags = FParam::EFlags::Reference;
+
+		return AddParameter(EParamType::Input, ParameterId, Flags, &Value);
+	}
+
+	// Add an Input parameter by reference, adding just a pointer to the shared storage
+	template<typename ValueType>
+	FHParam AddInputReference(FName ParameterId, ValueType* Value)
+	{
+		FParam::EFlags Flags = FParam::EFlags::Reference;
+
+		return AddParameter(EParamType::Input, ParameterId, Flags, Value);
+	}
+
+	// Add an Output parameter by reference, adding just a pointer to the shared storage
+	template<typename ValueType>
+	FHParam AddOutputReference(FName ParameterId, ValueType& Value)
+	{
+		const FParam::EFlags Flags = FParam::EFlags::Reference;
+
+		return AddParameter(EParamType::Output, ParameterId, Flags, &Value);
+	}
+
+	// Add an Output parameter by reference, adding just a pointer to the shared storage
+	template<typename ValueType>
+	FHParam AddOutputReference(FName ParameterId, ValueType* Value)
+	{
+		FParam::EFlags Flags = FParam::EFlags::Reference;
+
+		return AddParameter(EParamType::Output, ParameterId, Flags, Value);
+	}
+
 	// Add a parameter by reference, adding just a pointer to the shared storage
 	template<typename ValueType>
 	FHParam AddReference(EParamType InParamType, FName ParameterId, ValueType& Value)
diff --git a/Engine/Plugins/Experimental/AnimNext/Source/AnimNextInterface/Public/AnimNextInterfaceParam.h b/Engine/Plugins/Experimental/AnimNext/Source/AnimNextInterface/Public/AnimNextInterfaceParam.h
--- a/Engine/Plugins/Experimental/AnimNext/Source/AnimNextInterface/Public/AnimNextInterfaceParam.h
+++ b/Engine/Plugins/Experimental/AnimNext/Source/AnimNextInterface/Public/AnimNextInterfaceParam.h
@@ -118,6 +118,8 @@ struct ANIMNEXTINTERFACE_API FParamType
 
 	struct ANIMNEXTINTERFACE_API FRegistrar
 	{
+		static TArray<TUniqueFunction<void(void)>>& GetDeferredTypes();
+
 		// Helper for registration (via IMPLEMENT_DATA_INTERFACE_PARAM_TYPE).
 		// Registers the type via deferred callback
 		FRegistrar(TUniqueFunction<void(void)>&& InFunction);
@@ -209,6 +211,28 @@ public:
 			&& NumElements > 0);
 	}
 
+	FORCEINLINE_DEBUGGABLE const ValueType& GetDataChecked() const
+	{
+		check((Data != nullptr || EnumHasAnyFlags(Flags, FParam::EFlags::Embedded)) && NumElements == 1);
+
+		const ValueType* ValueData = EnumHasAnyFlags(Flags, FParam::EFlags::Embedded)
+			? static_cast<const ValueType*>((void*)&Data)
+			: static_cast<const ValueType*>(Data);
+
+		return *ValueData;
+	}
+
+	FORCEINLINE_DEBUGGABLE ValueType& GetDataChecked()
+	{
+		check((Data != nullptr || EnumHasAnyFlags(Flags, FParam::EFlags::Embedded)) && NumElements == 1);
+
+		ValueType* ValueData = EnumHasAnyFlags(Flags, FParam::EFlags::Embedded)
+			? static_cast<ValueType*>((void*)&Data)
+			: static_cast<ValueType*>(Data);
+
+		return *ValueData;
+	}
+
 	FORCEINLINE_DEBUGGABLE const ValueType& operator[](int32 InIndex) const
 	{
 		check(InIndex >= 0 && (InIndex < NumElements || NumElements == 1));
diff --git a/Engine/Plugins/Experimental/AnimNext/Source/AnimNextInterface/Public/AnimNextInterfaceParamStorage.h b/Engine/Plugins/Experimental/AnimNext/Source/AnimNextInterface/Public/AnimNextInterfaceParamStorage.h
--- a/Engine/Plugins/Experimental/AnimNext/Source/AnimNextInterface/Public/AnimNextInterfaceParamStorage.h
+++ b/Engine/Plugins/Experimental/AnimNext/Source/AnimNextInterface/Public/AnimNextInterfaceParamStorage.h
@@ -99,7 +99,7 @@ public:
 	template<typename ValueType>
 	FHParam AddReference(ValueType* Data, int32 NumElem, FParam::EFlags Flags)
 	{
-		check(ParamFreeIndex < (Parameters.Num() - 1)); // TODO : Growing ?
+		check(ParamFreeIndex < (Parameters.Num())); // TODO : Growing ?
 
 		const int32 ParameterIndex = ParamFreeIndex;
 		ParamFreeIndex++;
diff --git a/Engine/Plugins/Experimental/AnimNext/Source/AnimNextInterface/Public/AnimationDataRegistry.h b/Engine/Plugins/Experimental/AnimNext/Source/AnimNextInterface/Public/AnimationDataRegistry.h
new file mode 100644
--- /dev/null
+++ b/Engine/Plugins/Experimental/AnimNext/Source/AnimNextInterface/Public/AnimationDataRegistry.h
@@ -0,0 +1,237 @@
+// Copyright Epic Games, Inc. All Rights Reserved.
+
+#pragma once
+
+#include "CoreMinimal.h"
+#include "AnimationDataRegistryTypes.h"
+#include "UObject/WeakObjectPtr.h"
+#include "AnimNextInterfaceParam.h"
+
+class USkeletalMeshComponent;
+struct FReferenceSkeleton;
+
+namespace UE::AnimNext::Interface
+{
+
+class FModule;
+
+struct FAnimationDataHandle;
+
+// Global registry of animation data
+// Holds ref counted data that gets released when the last DataHandle of that element goes out of scope
+// Calling public functions from multiple threads is expected. Data races are guarded by a FRWLock.
+// TODO : Memory management will have to be implemented to avoid fragmentation and performance reasons
+class ANIMNEXTINTERFACE_API FAnimationDataRegistry
+{
+public:
+
+	// Access the global registry
+	static FAnimationDataRegistry* Get();
+
+
+	// --- Reference Pose Handling ---
+
+	// Generates and registers a reference pose for the SkeletalMesh asset of the SkeletalMeshComponent
+	// and modifies it with the additional required bones 
+	// or the visibility state of the bones of the SkeletalMeshComponent
+	FAnimationDataHandle RegisterReferencePose(USkeletalMeshComponent* SkeletalMeshComponent);
+
+	// Returns a ref counted handle to the refence pose of the given skeletal mesh component
+	FAnimationDataHandle GetOrGenerateReferencePose(USkeletalMeshComponent* SkeletalMeshComponent);
+
+	// Removes a previously registered reference pose for the given SkeletalMeshComponent
+	void RemoveReferencePose(USkeletalMeshComponent* SkeletalMeshComponent);
+
+	
+	// --- AnimationData Storage / Retrieval  --- 
+
+	// Registers an anim data handle with arbitrary data using an FName
+	// Note that AnimDataHandles are refcounted, so this makes them permanent until unregistered
+	void RegisterData(const FName& Id, const FAnimationDataHandle& AnimationDataHandle);
+
+	// Unregisters a previously registered anim data handle 
+	void UnregisterData(const FName& Id);
+
+	// Obtains the data hanle for the passed Id, if it exists.
+	// If there is no anim data handle registered, the handle IsValid will be false
+	FAnimationDataHandle GetRegisteredData(const FName& Id) const;
+
+	
+	// --- Supported types registration ---
+
+	// Registers a type and sets the desired preallocation block size
+	// If a type is allocated without registering, a default block size of 32 will be used
+	template<typename DataType>
+	inline void RegisterDataType(int32 AllocationBlockSize)
+	{
+		RegisterDataType_Impl<DataType>(AllocationBlockSize);
+	}
+
+	// --- Persistent Data ---
+
+	// Allocates uninitialized memory for a type (leaving the initialization to the caller)
+	// Returns a refcounted animation data handle
+	// Allocated memory will be released once the refcount reaches 0
+	template<typename DataType>
+	FAnimationDataHandle PreAllocateMemory(const int32 NumElements)
+	{
+		const FParamType& ParamType = Private::TParamType<DataType>::GetType();
+
+		const FDataTypeDef* TypeDef = nullptr;
+		{
+			FRWScopeLock Lock(DataTypeDefsLock, SLT_ReadOnly);
+			TypeDef = DataTypeDefs.Find(ParamType.GetTypeId());
+		}
+
+		if (TypeDef == nullptr)
+		{
+			TypeDef = RegisterDataType_Impl<DataType>(DEFAULT_BLOCK_SIZE);
+			// TODO : Log if we allocate more than DEFAULT_BLOCK_SIZE elements of that type
+		}
+
+		if (ensure(TypeDef != nullptr && TypeDef->TypeId != MAX_uint16))
+		{
+			const int32 ElementSize = TypeDef->ElementSize;
+			const int32 ElementAlign = TypeDef->ElementAlign;
+			const int32 AlignedSize = Align(ElementSize, ElementAlign);
+
+			const int32 BufferSize = NumElements * AlignedSize;
+
+			uint8* Memory = (uint8*)FMemory::Malloc(BufferSize, TypeDef->ElementAlign);    // TODO : This should come from preallocated chunks, use malloc / free for now
+
+			Private::FAllocatedBlock* AllocatedBlock = new Private::FAllocatedBlock(Memory, NumElements, TypeDef->TypeId); // TODO : avoid memory fragmentation
+			AllocatedBlock->AddRef();
+			AllocatedBlocks.Add(AllocatedBlock);
+			return FAnimationDataHandle(AllocatedBlock);
+		}
+
+		return FAnimationDataHandle();
+	}
+
+	// Allocates memory for a type, initialized with the default constructor (with optional passed arguments)
+	// Returns a refcounted animation data handle
+	// Allocated memory will be released once the refcount reaches 0
+	template<typename DataType, typename... ArgTypes>
+	FAnimationDataHandle AllocateData(const int32 NumElements, ArgTypes&&... Args)
+	{
+		FAnimationDataHandle Handle = PreAllocateMemory<DataType>(NumElements);
+
+		DataType* RetVal = Handle.GetPtr<DataType>();
+		for (int i = 0; i < NumElements; i++)
+		{
+			// perform a placement new per element
+			new ((uint8*)&RetVal[i]) DataType(Forward<ArgTypes>(Args)...);
+		}
+
+		return Handle;
+	}
+
+private:
+	typedef void (*DestroyFnSignature)(uint8* TargetBuffer, int32 NumElem);
+
+	static constexpr int32 DEFAULT_BLOCK_SIZE = 32;
+
+	// structure holding each registered type information
+	struct FDataTypeDef
+	{
+		Private::AnimDataTypeId TypeId = Private::InvalidTypeId;
+		DestroyFnSignature DestrooyTypeFn = nullptr;
+		int32 ElementSize = 0;
+		int32 ElementAlign = 0;
+		int32 AllocationBlockSize = 0;
+	};
+
+	struct FReferencePoseData
+	{
+		FReferencePoseData() = default;
+		
+		FReferencePoseData(const FAnimationDataHandle& InAnimationDataHandle, const FDelegateHandle& InDelegateHandle)
+			: AnimationDataHandle(InAnimationDataHandle)
+			, DelegateHandle(InDelegateHandle)
+		{
+		}
+
+		FReferencePoseData(FAnimationDataHandle&& InAnimationDataHandle, FDelegateHandle&& InDelegateHandle)
+			: AnimationDataHandle(MoveTemp(InAnimationDataHandle))
+			, DelegateHandle(MoveTemp(InDelegateHandle))
+		{
+		}
+
+		FAnimationDataHandle AnimationDataHandle;
+		FDelegateHandle DelegateHandle;
+	};
+
+	// Map holding registered types
+	TMap<Private::AnimDataTypeId, FDataTypeDef> DataTypeDefs;
+	// Lock for registered types map
+	FRWLock DataTypeDefsLock;
+
+	TSet<Private::FAllocatedBlock*> AllocatedBlocks;
+	// Lock for allocated blocks
+	FRWLock AllocatedBlocksLock;
+
+	// Map holding named data
+	TMap <FName, FAnimationDataHandle> StoredData;
+	mutable FRWLock StoredDataLock;
+
+	// Map holding reference poses for SkeletalMeshes
+	TMap <TWeakObjectPtr<USkeletalMeshComponent>, FReferencePoseData> SkeletalMeshReferencePoses;
+	mutable FRWLock SkeletalMeshReferencePosesLock;
+
+	// Registers a type and sets the allocation block size
+	template<typename DataType>
+	FDataTypeDef* RegisterDataType_Impl(int32 AllocationBlockSize)
+	{
+		const FParamType& ParamType = Private::TParamType<DataType>::GetType();
+		const Private::AnimDataTypeId TypeId = ParamType.GetTypeId();
+		check(TypeId != MAX_uint16); // Missing IMPLEMENT_ANIM_NEXT_INTERFACE_STATE_TYPE(Type, Name);
+
+		const int32 ElementSize = ParamType.GetSize();
+		const int32 ElementAlign = ParamType.GetAlignment();
+
+		// If we use raw types, I need a per element destructor
+		DestroyFnSignature DestroyFn = [](uint8* TargetBuffer, int32 NumElem)->void
+		{
+			const DataType* Ptr = (DataType*)TargetBuffer;
+			for (int i = 0; i < NumElem; i++)
+			{
+				Ptr[i].~DataType();
+			}
+		};
+
+		FDataTypeDef* AddedDef = nullptr;
+		{
+			FRWScopeLock WriteLock(DataTypeDefsLock, SLT_Write);
+
+			AddedDef = &DataTypeDefs.FindOrAdd(TypeId, { TypeId, DestroyFn, ElementSize, ElementAlign, AllocationBlockSize });
+			check(AddedDef->TypeId == TypeId); // check we have not added two different types with the same ID
+		}
+
+		return AddedDef;
+	}
+
+	void OnLODRequiredBonesUpdate(USkeletalMeshComponent* SkeletalMeshComponent, int32 LODLevel, const TArray<FBoneIndexType>& LODRequiredBones);
+	
+	void FreeAllocatedBlock(Private::FAllocatedBlock * AllocatedBlock);
+
+	void ReleaseReferencePoseData();
+
+// --- ---
+private:
+	friend class FModule;
+	friend struct FAnimationDataHandle;
+
+	// Initialize the global registry
+	static void Init();
+
+	// Shutdown the global registry
+	static void Destroy();
+
+	volatile int32 HandleCounter = 0;
+
+
+private:
+	static void HandlePostGarbageCollect();
+};
+
+} // namespace UE::AnimNext::Interface
diff --git a/Engine/Plugins/Experimental/AnimNext/Source/AnimNextInterface/Public/AnimationDataRegistryTypes.h b/Engine/Plugins/Experimental/AnimNext/Source/AnimNextInterface/Public/AnimationDataRegistryTypes.h
new file mode 100644
--- /dev/null
+++ b/Engine/Plugins/Experimental/AnimNext/Source/AnimNextInterface/Public/AnimationDataRegistryTypes.h
@@ -0,0 +1,175 @@
+// Copyright Epic Games, Inc. All Rights Reserved.
+
+#pragma once
+
+#include "CoreMinimal.h"
+#include "UObject/WeakObjectPtr.h"
+#include <algorithm>
+#include "Animation/AnimTypes.h"
+#include "ReferenceSkeleton.h"
+#include "BoneIndices.h"
+
+namespace UE::AnimNext::Interface
+{
+
+
+namespace Private
+{
+
+using AnimDataTypeId = uint16;
+static constexpr AnimDataTypeId InvalidTypeId = MAX_uint16;
+
+struct FAllocatedBlock
+{
+	mutable int32 NumRefs = 0;
+	void* Memory = nullptr;
+	int32 NumElem = 0;
+	AnimDataTypeId TypeId = InvalidTypeId;
+
+	FAllocatedBlock(void* InMemory, int32 InNumElem, AnimDataTypeId InTypeId)
+		: Memory(InMemory)
+		, NumElem(InNumElem)
+		, TypeId(InTypeId)
+	{
+	}
+
+	inline uint32 AddRef() const
+	{
+		return uint32(FPlatformAtomics::InterlockedIncrement(&NumRefs));
+	}
+
+	inline uint32 Release() const
+	{
+		check(NumRefs > 0);
+
+		const int32 Refs = FPlatformAtomics::InterlockedDecrement(&NumRefs);
+		check(Refs >= 0);
+
+		return uint32(Refs);
+	}
+
+	inline uint32 GetRefCount() const
+	{
+		return uint32(NumRefs);
+	}
+
+private:
+	FAllocatedBlock() = delete;
+	FAllocatedBlock(const FAllocatedBlock& Other) = delete;
+	FAllocatedBlock(FAllocatedBlock&& Other) = delete;
+};
+
+} // end namespace Private
+
+struct ANIMNEXTINTERFACE_API FAnimationDataHandle
+{
+	FAnimationDataHandle() = default;
+
+	FAnimationDataHandle(Private::FAllocatedBlock* InAllocatedBlock)
+		: AllocatedBlock(InAllocatedBlock)
+	{
+	}
+
+	~FAnimationDataHandle();
+
+	FAnimationDataHandle(const FAnimationDataHandle& Other)
+		: AllocatedBlock(Other.AllocatedBlock)
+	{
+		if (AllocatedBlock != nullptr)
+		{
+			const int32 CurrentCount = AllocatedBlock->AddRef();
+			check(CurrentCount > 1);
+		}
+	}
+
+	FAnimationDataHandle& operator= (const FAnimationDataHandle& Other)
+	{
+		FAnimationDataHandle Tmp(Other);
+
+		Swap(*this, Tmp);
+		return *this;
+	}
+
+	FAnimationDataHandle(FAnimationDataHandle&& Other)
+		: FAnimationDataHandle()
+	{
+		Swap(*this, Other);
+	}
+
+	FAnimationDataHandle& operator= (FAnimationDataHandle&& Other)
+	{
+		Swap(*this, Other);
+		return *this;
+	}
+
+	inline bool IsValid() const
+	{
+		return AllocatedBlock != nullptr;
+	}
+
+	template<typename DataType>
+	inline TArrayView<DataType> AsArrayView()
+	{
+		check(IsValid());
+		return TArrayView<DataType>((DataType*)AllocatedBlock->Memory, AllocatedBlock->NumElem);
+	}
+
+	template<typename DataType>
+	inline TArrayView<DataType> AsArrayView() const
+	{
+		check(IsValid());
+		return TArrayView<DataType>(AllocatedBlock->Memory, AllocatedBlock->NumElem);
+	}
+
+	template<typename DataType>
+	inline DataType* GetPtr()
+	{
+		check(IsValid());
+		DataType* Data = static_cast<DataType*>(AllocatedBlock->Memory);
+		check(Data != nullptr);
+		return Data;
+	}
+
+	template<typename DataType>
+	inline const DataType* GetPtr() const
+	{
+		check(IsValid());
+		const DataType* Data = static_cast<const DataType*>(AllocatedBlock->Memory);
+		return Data;
+	}
+
+	template<typename DataType>
+	inline DataType& GetRef()
+	{
+		check(IsValid());
+		DataType* Data = static_cast<DataType*>(AllocatedBlock->Memory);
+		check(Data != nullptr);
+		return *Data;
+	}
+
+	template<typename DataType>
+	inline const DataType& GetRef() const
+	{
+		check(IsValid());
+		const DataType* Data = static_cast<const DataType*>(AllocatedBlock->Memory);
+		check(Data != nullptr);
+		return *Data;
+	}
+
+	inline Private::AnimDataTypeId GetTypeId() const
+	{
+		return AllocatedBlock != nullptr ? AllocatedBlock->TypeId : Private::InvalidTypeId;
+	}
+
+private:
+	Private::FAllocatedBlock* AllocatedBlock = nullptr;
+};
+
+enum class ETransformFlags : uint8
+{
+	None = 0,
+	ComponentSpaceSet = 1 << 0
+};
+
+
+} // namespace UE::AnimNext::Interface
diff --git a/Engine/Plugins/Experimental/AnimNext/Source/AnimNextInterface/Public/AnimationDecompressionTools.h b/Engine/Plugins/Experimental/AnimNext/Source/AnimNextInterface/Public/AnimationDecompressionTools.h
new file mode 100644
--- /dev/null
+++ b/Engine/Plugins/Experimental/AnimNext/Source/AnimNextInterface/Public/AnimationDecompressionTools.h
@@ -0,0 +1,75 @@
+// Copyright Epic Games, Inc. All Rights Reserved.
+
+#pragma once
+
+#include "CoreMinimal.h"
+#include "AnimationReferencePose.h"
+
+struct FCompactPose;
+struct FCompressedAnimSequence;
+struct FAnimExtractContext;
+struct FAnimSequenceDecompressionContext;
+struct FRootMotionReset;
+class UAnimSequence;
+
+namespace UE::AnimNext::Interface
+{
+
+class ANIMNEXTINTERFACE_API FDecompressionTools
+{
+public:
+
+	// Extracted from UAnimSequence
+	// Extracts Animation Data from the provided AnimSequence, using AnimExtractContext extraction parameters
+	static void GetAnimationPose(const UAnimSequence* AnimSequence
+		, FAnimationLODPose& OutAnimationPoseData
+		, const FAnimExtractContext& ExtractionContext);
+
+	/**
+	* Get Bone Transform of the Time given, relative to Parent for all RequiredBones
+	* This returns different transform based on additive or not. Or what kind of additive.
+	*/
+	static void GetBonePose(const UAnimSequence* AnimSequence, const FAnimExtractContext& ExtractionContext, FAnimationLODPose& OutAnimationPoseData, bool bForceUseRawData = false);
+
+	static void GetBonePose_Additive(const UAnimSequence* AnimSequence, const FAnimExtractContext& ExtractionContext, FAnimationLODPose& OutAnimationPoseData);
+	static void GetBonePose_AdditiveMeshRotationOnly(const UAnimSequence* AnimSequence, const FAnimExtractContext& ExtractionContext, FAnimationLODPose& OutAnimationPoseData);
+
+	// Decompress and retarget animation data using provided RetargetTransforms
+	static void DecompressPose(FAnimationLODPose& OutAnimationPoseData,
+		const FCompressedAnimSequence& CompressedData,
+		const FAnimExtractContext& ExtractionContext,
+		FAnimSequenceDecompressionContext& DecompressionContext,
+		const TArray<FTransform>& RetargetTransforms,
+		const FRootMotionReset& RootMotionReset);
+
+	// Decompress and retarget animation data
+	static void DecompressPose(FAnimationLODPose& OutAnimationPoseData,
+		const FCompressedAnimSequence& CompressedData,
+		const FAnimExtractContext& ExtractionContext,
+		FAnimSequenceDecompressionContext& DecompressionContext,
+		FName RetargetSource,
+		const FRootMotionReset& RootMotionReset);
+
+	static void RetargetBoneTransform(const FAnimationReferencePose& ReferencePose
+		, const USkeleton* SourceSkeleton
+		, const USkeleton* TargetSkeleton
+		, const FName& RetargetSource
+		, FTransform& BoneTransform
+		, const int32 SkeletonBoneIndex
+		, const int32 LODBoneIndex
+		, const bool bIsBakedAdditive
+		, const bool bDisableRetargeting);
+
+	static void RetargetBoneTransform(const FAnimationReferencePose& ReferencePose
+		, const USkeleton* SourceSkeleton
+		, const USkeleton* TargetSkeleton
+		, const FName& SourceName
+		, const TArray<FTransform>& RetargetTransforms
+		, FTransform& BoneTransform
+		, const int32 SkeletonBoneIndex
+		, const int32 LODBoneIndex
+		, const bool bIsBakedAdditive
+		, const bool bDisableRetargeting);
+};
+
+} // end namespace UE::AnimNext::Interface
diff --git a/Engine/Plugins/Experimental/AnimNext/Source/AnimNextInterface/Public/AnimationGenerationTools.h b/Engine/Plugins/Experimental/AnimNext/Source/AnimNextInterface/Public/AnimationGenerationTools.h
new file mode 100644
--- /dev/null
+++ b/Engine/Plugins/Experimental/AnimNext/Source/AnimNextInterface/Public/AnimationGenerationTools.h
@@ -0,0 +1,101 @@
+// Copyright Epic Games, Inc. All Rights Reserved.
+
+#pragma once
+
+#include "BoneIndices.h"
+#include "Containers/IndirectArray.h"
+#include "AnimationReferencePose.h"
+
+class USkeletalMeshComponent;
+class USkeletalMesh;
+class FSkeletalMeshLODRenderData;
+struct FPoseContext;
+class UPhysicsAsset;
+
+namespace UE::AnimNext::Interface
+{
+
+struct FGenerationLODData
+{
+	TArray<FBoneIndexType> RequiredBones;									// All the bones required for the LOD
+	TArray<FBoneIndexType> ExcludedBones;									// List of bones excluded from LOD 0
+	TArray<FBoneIndexType> ExcludedBonesFromPrevLOD;						// list of bones excluded from previous LOD
+};
+
+class ANIMNEXTINTERFACE_API FGenerationTools
+{
+public:
+
+	// Generates the reference pose data from a SkeletalMeshComponent and / or Skeletal Mesh asset
+	// If no SkeletalMeshComponent is passed, the reference pose will not exclude invisible bones and will not include shadow shapes required bones
+	// If no SkeletalMesh asset is passed, there will be no generation
+	static bool GenerateReferencePose(const USkeletalMeshComponent* SkeletalMeshComponent
+		, const USkeletalMesh* SkeletalMesh
+		, FAnimationReferencePose& OutAnimationReferencePose);
+
+	// Generates the full list of bones required by a LOD, based on the SkeletalRetrieve RequiredBones
+	// Compute ExcludedBones versus LOD0 and PreviousLOD
+	// Remove ExcludedBonesFromPrevLOD from BonesInAllLODS
+	static void GenerateRawLODData(const USkeletalMeshComponent* SkeletalMeshComponent
+		, const USkeletalMesh* SkeletalMesh
+		, const int32 LODIndex
+		, const TIndirectArray<FSkeletalMeshLODRenderData>& LODRenderData
+		, TArray<FBoneIndexType>& OutRequiredBones
+		, TArray<FBoneIndexType>& OutFillComponentSpaceTransformsRequiredBones);
+
+	// For each LOD > 0:
+	// Retrieve RequiredBones
+	// Compute ExcludedBones versus LOD0 and PreviousLOD
+	static void GenerateLODData(const USkeletalMeshComponent* SkeletalMeshComponent
+		, const USkeletalMesh* SkeletalMesh
+		, const int32 StartLOD
+		, const int32 NumLODs
+		, const TIndirectArray<FSkeletalMeshLODRenderData>& LODRenderData
+		, const TArray<FBoneIndexType>& RequiredBones_LOD0
+		, TArray<FGenerationLODData>& GenerationLODData
+		, TArray<FGenerationLODData>& GenerationComponentSpaceLODData);
+
+	// Calculate the bone indexes difference from LOD0 for LODIndex
+	static void CalculateDifferenceFromParentLOD(int32 LODIndex, TArray<FGenerationLODData>& GenerationLODData);
+
+	// For each LOD :
+	// Check the excluded bones in LOD(N) contain all the bones excluded in LOD(N-1)
+	static bool CheckExcludedBones(const int32 NumLODs
+		, const TArray<FGenerationLODData>& GenerationLODData
+		, const USkeletalMesh* SkeletalMesh);
+
+	// For each LOD :
+	// Generates an unified list of bones, in LOD order (if possible)
+	// Returns true if the unified list could be created, false otherwise
+	static bool GenerateOrderedBoneList(const USkeletalMesh* SkeletalMesh
+		, TArray<FGenerationLODData>& GenerationLODData
+		, TArray<FBoneIndexType>& OrderedBoneList);
+
+	/**
+	 *	Utility for taking two arrays of bone indices, which must be strictly increasing, and finding the A - B.
+	 *	That is - any items left in A, after removing B.
+	 */
+	static void DifferenceBoneIndexArrays(const TArray<FBoneIndexType>& A, const TArray<FBoneIndexType>& B, TArray<FBoneIndexType>& Output);
+
+	// Checks if all sockets of a skeletal mesh are set to always animate, as it is a requirement for generating a single reference pose,
+	// where the local space pose and the component space pose use the same bone indexes
+	static bool CheckSkeletalAllMeshSocketsAlwaysAnimate(const USkeletalMesh* SkeletalMesh);
+
+	// Converts AnimBP pose to AnimNext Pose
+	// This function expects both poses to have the same LOD (number of bones and indexes)
+	// The target pose should be assigned to the correct reference pose prior to this call
+	static void RemapPose(int32 LODLevel
+		, const FPoseContext& SourcePose
+		, const FAnimationReferencePose& RefPose
+		, FAnimationLODPose& TargetPose); // TODO : see how to deal with heap and stack variants without duplicating functions (and if possible without templating the function)
+
+	// Converts AnimNext pose to AnimBP Pose
+	// This function expects both poses to have the same LOD (number of bones and indexes)
+	// The target pose should be assigned to the correct reference pose prior to this call
+	static void RemapPose(int32 LODLevel
+		, const FAnimationReferencePose& RefPose
+		, const FAnimationLODPose& SourcePose
+		, FPoseContext& TargetPose); // TODO : see how to deal with heap and stack variants without duplicating functions (and if possible without templating the function)
+};
+
+} // namespace UE::AnimNext::Interface
diff --git a/Engine/Plugins/Experimental/AnimNext/Source/AnimNextInterface/Public/AnimationReferencePose.h b/Engine/Plugins/Experimental/AnimNext/Source/AnimNextInterface/Public/AnimationReferencePose.h
new file mode 100644
--- /dev/null
+++ b/Engine/Plugins/Experimental/AnimNext/Source/AnimNextInterface/Public/AnimationReferencePose.h
@@ -0,0 +1,995 @@
+// Copyright Epic Games, Inc. All Rights Reserved.
+
+#pragma once
+
+#include "CoreMinimal.h"
+#include "Containers/ContainersFwd.h"
+#include "UObject/WeakObjectPtr.h"
+#include "AnimationDataRegistry.h"
+#include <algorithm>
+#include "Animation/AnimTypes.h"
+#include "ReferenceSkeleton.h"
+#include "BoneIndices.h"
+#include "AnimNextInterfaceTypes.h"
+
+class USkeletalMesh;
+class USkeleton;
+
+namespace UE::AnimNext::Interface
+{
+
+#define ANIM_ENABLE_POINTER_ITERATION 1
+
+using FAnimTranslation = FVector;
+using FAnimRotation = FQuat;
+using FAnimScale = FVector;
+using FAnimTransform = FTransform;
+
+extern const FAnimTransform ANIMNEXTINTERFACE_API TransformAdditiveIdentity;
+
+//****************************************************************************
+
+template <typename AllocatorType>
+struct TAnimTransformArrayAoS
+{
+	TArray<FAnimTransform, AllocatorType> Transforms;
+
+	// Default empty (no memory) Transform
+	TAnimTransformArrayAoS() = default;
+
+	// Create a TransformArray of N elements
+	explicit TAnimTransformArrayAoS(const int NumTransforms)
+	{
+		SetNum(NumTransforms);
+	}
+
+	~TAnimTransformArrayAoS() = default;
+	TAnimTransformArrayAoS(const TAnimTransformArrayAoS& Other) = default;
+	TAnimTransformArrayAoS& operator= (const TAnimTransformArrayAoS& Other) = default;
+	TAnimTransformArrayAoS(TAnimTransformArrayAoS&& Other) = default;
+	TAnimTransformArrayAoS& operator= (TAnimTransformArrayAoS&& Other) = default;
+
+	inline void Reset(int NumTransforms)
+	{
+		Transforms.Reset(NumTransforms);
+	}
+
+	void SetNum(int32 NumTransforms, bool bAllowShrinking = true)
+	{
+		constexpr int32 TransformSize = sizeof(FAnimTranslation) + sizeof(FAnimRotation) + sizeof(FAnimScale);
+		Transforms.SetNum(NumTransforms, bAllowShrinking);
+	}
+
+	inline void SetIdentity(bool bAdditiveIdentity = false)
+	{
+		std::fill(Transforms.begin(), Transforms.end(), bAdditiveIdentity ? TransformAdditiveIdentity : FTransform::Identity);
+	}
+
+	inline void CopyTransforms(const TAnimTransformArrayAoS& Other, int32 Index, int32 NumTransforms)
+	{
+		check(Index < Num() && Index < Other.Num());
+		check(Index + (NumTransforms - 1) < Num() && Index + (NumTransforms - 1) < Other.Num());
+		
+		std::copy(Other.Transforms.GetData() + Index, Other.Transforms.GetData() + (Index + NumTransforms), Transforms.GetData() + Index);
+	}
+
+	inline int32 Num() const
+	{
+		return Transforms.Num();
+	}
+
+	inline TArray<FTransform>& GetTransforms()
+	{
+		return Transforms;
+	}
+	inline const TArray<FTransform>& GetTransforms() const
+	{
+		return Transforms;
+	}
+
+	inline FAnimTransform& operator[](int32 Index)
+	{
+		return Transforms[Index];
+	}
+
+	inline const FAnimTransform& operator[](int32 Index) const
+	{
+		return Transforms[Index];
+	}
+
+	/** Set this transform array to the weighted blend of the supplied two transforms. */
+	void Blend(const TAnimTransformArrayAoS& AtomArray1, const TAnimTransformArrayAoS& AtomArray2, float BlendWeight)
+	{
+		if (FMath::Abs(BlendWeight) <= ZERO_ANIMWEIGHT_THRESH)
+		{
+			// if blend is all the way for child1, then just copy its bone atoms
+			//(*this) = Atom1;
+			CopyTransforms(AtomArray1, 0, Num());
+		}
+		else if (FMath::Abs(BlendWeight - 1.0f) <= ZERO_ANIMWEIGHT_THRESH)
+		{
+			// if blend is all the way for child2, then just copy its bone atoms
+			//(*this) = Atom2;
+			CopyTransforms(AtomArray2, 0, Num());
+		}
+		else
+		{
+			const int32 NumTransforms = Num();
+
+			// Right now we only support same size blend
+			check(AtomArray1.Num() == NumTransforms);
+			check(AtomArray2.Num() == NumTransforms);
+
+#if ANIM_ENABLE_POINTER_ITERATION
+			// This version is faster, but without the range checks the TArray has on operator[]
+			const FAnimTransform* RESTRICT Transforms1 = AtomArray1.Transforms.GetData();
+			const FAnimTransform* RESTRICT Transforms1End = AtomArray1.Transforms.GetData() + NumTransforms;
+			const FAnimTransform* RESTRICT Transforms2 = AtomArray2.Transforms.GetData();
+			FAnimTransform* RESTRICT TransformsResult = Transforms.GetData();
+
+			for (; Transforms1 != Transforms1End; ++Transforms1, ++Transforms2, ++TransformsResult)
+			{
+				//#if !(UE_BUILD_SHIPPING || UE_BUILD_TEST) && WITH_EDITORONLY_DATA
+				//				// Check that all bone atoms coming from animation are normalized
+				//				check(Transforms1[i].IsRotationNormalized());
+				//				check(Transforms2[i].IsRotationNormalized());
+				//#endif
+				TransformsResult->Blend(*Transforms1, *Transforms2, BlendWeight);
+			}
+#else
+			// This version is a bit slower, but has range checks on the array
+			const TArray<FAnimTransform, AllocatorType>& Transforms1 = AtomArray1.Transforms;
+			const TArray<FAnimTransform, AllocatorType>& Transforms2 = AtomArray2.Transforms;
+			for (int32 i = 0; i < NumTransforms; ++i)
+			{
+				//#if !(UE_BUILD_SHIPPING || UE_BUILD_TEST) && WITH_EDITORONLY_DATA
+				//				// Check that all bone atoms coming from animation are normalized
+				//				check(Transforms1[i].IsRotationNormalized());
+				//				check(Transforms2[i].IsRotationNormalized());
+				//#endif
+				Transforms[i].Blend(Transforms1[i], Transforms2[i], BlendWeight);
+			}
+#endif // ANIM_ENABLE_POINTER_ITERATION
+
+			DiagnosticCheckNaN_All();
+		}
+	}
+
+	bool ContainsNaN() const
+	{
+		for (const auto& Transform : Transforms)
+		{
+			if (Transform.GetRotation().ContainsNaN())
+			{
+				return true;
+			}
+			if (Transform.GetTranslation().ContainsNaN())
+			{
+				return true;
+			}
+			if (Transform.GetScale3D().ContainsNaN())
+			{
+				return true;
+			}
+		}
+
+		return false;
+	}
+
+	bool IsValid() const
+	{
+		if (ContainsNaN())
+		{
+			return false;
+		}
+
+		for (const auto& Transform : Transforms)
+		{
+			if (Transform.GetRotation().IsNormalized() == false)
+			{
+				return false;
+			}
+		}
+
+		return true;
+	}
+
+private:
+#if ENABLE_NAN_DIAGNOSTIC
+	FORCEINLINE static void DiagnosticCheckNaN_Rotate(const FQuat& Rotation)
+	{
+		if (Rotation.ContainsNaN())
+		{
+			logOrEnsureNanError(TEXT("TAnimTransformArraySoA Rotation contains NaN"));
+		}
+	}
+
+	FORCEINLINE static void DiagnosticCheckNaN_Scale3D(const FVector& Scale3D)
+	{
+		if (Scale3D.ContainsNaN())
+		{
+			logOrEnsureNanError(TEXT("TAnimTransformArraySoA Scale3D contains NaN"));
+		}
+	}
+
+	FORCEINLINE static void DiagnosticCheckNaN_Translate(const FVector& Translation)
+	{
+		if (Translation.ContainsNaN())
+		{
+			logOrEnsureNanError(TEXT("TAnimTransformArraySoA Translation contains NaN"));
+		}
+	}
+
+	FORCEINLINE void DiagnosticCheckNaN_All() const
+	{
+		for (const auto& Transform : Transforms)
+		{
+			DiagnosticCheckNaN_Rotate(Transform.GetRotation());
+			DiagnosticCheckNaN_Translate(Transform.GetTranslation());
+			DiagnosticCheckNaN_Scale3D(Transform.GetScale3D());
+		}
+	}
+
+	FORCEINLINE void DiagnosticCheck_IsValid() const
+	{
+		DiagnosticCheckNaN_All();
+		if (!IsValid())
+		{
+			logOrEnsureNanError(TEXT("TAnimTransformArraySoA is not valid"));
+		}
+	}
+
+#else
+	FORCEINLINE static void DiagnosticCheckNaN_Translate(const FVector& Translation) {}
+	FORCEINLINE static void DiagnosticCheckNaN_Rotate(const FQuat& Rotation) {}
+	FORCEINLINE static void DiagnosticCheckNaN_Scale3D(const FVector& Scale3D) {}
+	FORCEINLINE void DiagnosticCheckNaN_All() const {}
+	FORCEINLINE void DiagnosticCheck_IsValid() const {}
+#endif
+};
+
+/**
+* Transform Array using ArrayOfStructs model
+*/
+using FAnimTransformArrayAoSHeap = TAnimTransformArrayAoS<FDefaultAllocator>;
+using FAnimTransformArrayAoSStack = TAnimTransformArrayAoS<FAnimStackAllocator>;
+
+
+// This enables using a SoA array with operator[]
+struct FTransformSoAAdapter
+{
+	FQuat& Rotation;
+	FVector& Translation;
+	FVector& Scale3D;
+
+	FTransformSoAAdapter(FQuat& InRotation, FVector& InTranslation, FVector& InScale3D)
+		: Rotation(InRotation)
+		, Translation(InTranslation)
+		, Scale3D(InScale3D)
+	{
+	}
+
+	FORCEINLINE FQuat& GetRotation() const
+	{
+		return Rotation;
+	}
+
+	FORCEINLINE void SetRotation(const FQuat& InRotation)
+	{
+		Rotation = InRotation;
+	}
+
+	FORCEINLINE FVector& GetTranslation() const
+	{
+		return Translation;
+	}
+
+	FORCEINLINE void SetTranslation(const FVector& InTranslation)
+	{
+		Translation = InTranslation;
+	}
+
+	FORCEINLINE FVector& GetScale3D() const
+	{
+		return Scale3D;
+	}
+
+	FORCEINLINE void SetScale3D(const FVector& InScale3D)
+	{
+		Scale3D = InScale3D;
+	}
+
+	FORCEINLINE operator FTransform()
+	{
+		return FTransform(Rotation, Translation, Scale3D);
+	}
+
+	FORCEINLINE operator FTransform() const
+	{
+		return FTransform(Rotation, Translation, Scale3D);
+	}
+
+	FORCEINLINE void operator= (const FTransform& Transform)
+	{
+		Rotation = Transform.GetRotation();
+		Translation = Transform.GetTranslation();
+		Scale3D = Transform.GetScale3D();
+	}
+
+	FORCEINLINE void ScaleTranslation(const FVector::FReal& Scale)
+	{
+		Translation *= Scale;
+		//DiagnosticCheckNaN_Translate();
+	}
+
+	FORCEINLINE void NormalizeRotation()
+	{
+		Rotation.Normalize();
+		//DiagnosticCheckNaN_Rotate();
+	}
+};
+
+/**
+* Transform Array Test using StructOfArrays model
+*/
+template <typename AllocatorType>
+struct TAnimTransformArraySoA
+{
+	TArrayView<FAnimTranslation> Translations;
+	TArrayView<FAnimRotation> Rotations;
+	TArrayView<FAnimScale> Scales3D;
+
+	// Default empty (no memory) Transform
+	TAnimTransformArraySoA()
+	{
+		UpdateViews(nullptr, 0);
+	}
+
+	// Create a TransformArray of N elements. Optionally initializes to identity (default)
+	explicit TAnimTransformArraySoA(int NumTransforms, bool bSetIdentity = true, bool bAdditiveIdentity = false)
+	{
+		SetNum(NumTransforms);
+
+		if (bSetIdentity)
+		{
+			SetIdentity(bAdditiveIdentity);
+		}
+	}
+
+	~TAnimTransformArraySoA() = default;
+
+	TAnimTransformArraySoA(const TAnimTransformArraySoA& Other)
+	{
+		AllocatedMemory = Other.AllocatedMemory;
+
+		UpdateViews(AllocatedMemory.GetData(), Other.Num());
+	}
+
+	TAnimTransformArraySoA& operator= (const TAnimTransformArraySoA& Other)
+	{
+		TAnimTransformArraySoA Tmp(Other);
+
+		Swap(*this, Tmp);
+		return *this;
+	}
+
+	TAnimTransformArraySoA(TAnimTransformArraySoA&& Other)
+	{
+		Swap(*this, Other);
+	}
+
+	TAnimTransformArraySoA& operator= (TAnimTransformArraySoA&& Other)
+	{
+		Swap(*this, Other);
+		return *this;
+	}
+
+	inline void Reset(int32 NumTransforms)
+	{
+		constexpr int32 TransformSize = sizeof(FAnimTranslation) + sizeof(FAnimRotation) + sizeof(FAnimScale);
+		AllocatedMemory.Reset(NumTransforms * TransformSize);
+	}
+
+	void SetNum(int32 NumTransforms, bool bAllowShrinking = true)
+	{
+		constexpr int32 TransformSize = sizeof(FAnimTranslation) + sizeof(FAnimRotation) + sizeof(FAnimScale);
+		AllocatedMemory.SetNum(NumTransforms * TransformSize, bAllowShrinking);
+
+		UpdateViews(AllocatedMemory.GetData(), NumTransforms);
+	}
+
+	inline void SetIdentity(bool bAdditiveIdentity = false)
+	{
+		std::fill(Translations.begin(), Translations.end(), FVector::ZeroVector);
+		std::fill(Rotations.begin(), Rotations.end(), FQuat::Identity);
+		std::fill(Scales3D.begin(), Scales3D.end(), bAdditiveIdentity ? FVector::ZeroVector : FVector::OneVector);
+	}
+
+	inline void CopyTransforms(const TAnimTransformArraySoA& Other, int32 Index, int32 NumTransforms)
+	{
+		check(Index < Num() && Index < Other.Num());
+		check(Index + (NumTransforms - 1) < Num() && Index + (NumTransforms - 1) < Other.Num());
+
+		std::copy(Other.Translations.begin() + Index, Other.Translations.begin() + (Index + NumTransforms), Translations.begin());
+		std::copy(Other.Rotations.begin() + Index, Other.Rotations.begin() + (Index + NumTransforms), Rotations.begin());
+		std::copy(Other.Scales3D.begin() + Index, Other.Scales3D.begin() + (Index + NumTransforms), Scales3D.begin());
+	}
+
+	/** Set this transform array to the weighted blend of the supplied two transforms. */
+	void Blend(const TAnimTransformArraySoA& AtomArray1, const TAnimTransformArraySoA& AtomArray2, float BlendWeight)
+	{
+		if (FMath::Abs(BlendWeight) <= ZERO_ANIMWEIGHT_THRESH)
+		{
+			// if blend is all the way for child1, then just copy its bone atoms
+			//(*this) = Atom1;
+			CopyTransforms(AtomArray1, 0, Num());
+		}
+		else if (FMath::Abs(BlendWeight - 1.0f) <= ZERO_ANIMWEIGHT_THRESH)
+		{
+			// if blend is all the way for child2, then just copy its bone atoms
+			//(*this) = Atom2;
+			CopyTransforms(AtomArray2, 0, Num());
+		}
+		else
+		{
+			const int32 NumTransforms = Num();
+
+#if ANIM_ENABLE_POINTER_ITERATION
+			using TransformVectorRegister = TVectorRegisterType<double>;
+
+			// Right now we only support same size blend
+			check(AtomArray1.Translations.Num() == Translations.Num() && AtomArray1.Rotations.Num() == Rotations.Num() && AtomArray1.Scales3D.Num() == Scales3D.Num());
+			check(AtomArray2.Translations.Num() == Translations.Num() && AtomArray2.Rotations.Num() == Rotations.Num() && AtomArray2.Scales3D.Num() == Scales3D.Num());
+
+			const FAnimTranslation* RESTRICT Translations1 = AtomArray1.Translations.GetData();
+			const FAnimTranslation* RESTRICT Translations1End = AtomArray1.Translations.GetData() + NumTransforms;
+			const FAnimTranslation* RESTRICT Translations2 = AtomArray2.Translations.GetData();
+			FAnimTranslation* RESTRICT TranslationsResult = Translations.GetData();
+			for (; Translations1 != Translations1End; ++Translations1, ++Translations2, ++TranslationsResult)
+			{
+				*TranslationsResult = FMath::Lerp(*Translations1, *Translations2, BlendWeight);
+			}
+
+			const FAnimRotation* RESTRICT Rotations1 = AtomArray1.Rotations.GetData();
+			const FAnimRotation* RESTRICT Rotations1End = AtomArray1.Rotations.GetData() + NumTransforms;
+			const FAnimRotation* RESTRICT Rotations2 = AtomArray2.Rotations.GetData();
+			FAnimRotation* RESTRICT RotationsResult = Rotations.GetData();
+			for (; Rotations1 != Rotations1End; ++Rotations1, ++Rotations2, ++RotationsResult)
+			{
+				//#if !(UE_BUILD_SHIPPING || UE_BUILD_TEST) && WITH_EDITORONLY_DATA
+				//				// Check that all bone atoms coming from animation are normalized
+				//				check(Rotations1->IsNormalized());
+				//				check(Rotations2->IsNormalized());
+				//#endif
+
+				*RotationsResult = FQuat::FastLerp(*Rotations1, *Rotations2, BlendWeight).GetNormalized();
+			}
+
+			const FAnimScale* RESTRICT Scales3D1 = AtomArray1.Scales3D.GetData();
+			const FAnimScale* RESTRICT Scales3D1End = AtomArray1.Scales3D.GetData() + NumTransforms;
+			const FAnimScale* RESTRICT Scales3D2 = AtomArray2.Scales3D.GetData();
+			FAnimScale* RESTRICT Scales3DResult = Scales3D.GetData();
+			for (; Scales3D1 != Scales3D1End; ++Scales3D1, ++Scales3D2, ++Scales3DResult)
+			{
+				*Scales3DResult = FMath::Lerp(*Scales3D1, *Scales3D2, BlendWeight);
+			}
+#else
+			// This version is slower than AoS and a lot slower than the version using pointers, but has range checks on the TArrayView
+			const TArrayView<FAnimTranslation>& Translations1 = AtomArray1.Translations;
+			const TArrayView<FAnimTranslation>& Translations2 = AtomArray2.Translations;
+			for (int32 i = 0; i < NumTransforms; ++i)
+			{
+				Translations[i] = FMath::Lerp(Translations1[i], Translations2[i], BlendWeight);
+			}
+			const TArrayView<FAnimRotation>& Rotations1 = AtomArray1.Rotations;
+			const TArrayView<FAnimRotation>& Rotations2 = AtomArray2.Rotations;
+			for (int32 i = 0; i < NumTransforms; ++i)
+			{
+				//#if !(UE_BUILD_SHIPPING || UE_BUILD_TEST) && WITH_EDITORONLY_DATA
+				//				// Check that all bone atoms coming from animation are normalized
+				//				check(Transforms1[i].IsRotationNormalized());
+				//				check(Transforms2[i].IsRotationNormalized());
+				//#endif
+				Rotations[i] = FQuat::FastLerp(Rotations1[i], Rotations2[i], BlendWeight).GetNormalized();
+			}
+			const TArrayView<FAnimScale>& Scales3D1 = AtomArray1.Scales3D;
+			const TArrayView<FAnimScale>& Scales3D2 = AtomArray2.Scales3D;
+			for (int32 i = 0; i < NumTransforms; ++i)
+			{
+				Scales3D[i] = FMath::Lerp(Scales3D1[i], Scales3D2[i], BlendWeight);
+			}
+
+#endif // ANIM_ENABLE_POINTER_ITERATION
+
+			DiagnosticCheckNaN_All();
+		}
+	}
+
+	FORCEINLINE int32 Num() const
+	{
+		return Translations.Num();
+	}
+
+	FTransformSoAAdapter operator[] (int Index)
+	{
+		return FTransformSoAAdapter(Rotations[Index], Translations[Index], Scales3D[Index]);
+	}
+
+	const FTransformSoAAdapter operator[] (int Index) const
+	{
+		return FTransformSoAAdapter(Rotations[Index], Translations[Index], Scales3D[Index]);
+	}
+
+	bool ContainsNaN() const
+	{
+		for (const auto& Rotation : Rotations)
+		{
+			if (Rotation.ContainsNaN())
+			{
+				return true;
+			}
+		}
+		for (const auto& Translation : Translations)
+		{
+			if (Translation.ContainsNaN())
+			{
+				return true;
+			}
+		}
+		for (const auto& Scale3D : Scales3D)
+		{
+			if (Scale3D.ContainsNaN())
+			{
+				return true;
+			}
+		}
+
+		return false;
+	}
+
+	bool IsValid() const
+	{
+		if (ContainsNaN())
+		{
+			return false;
+		}
+
+		for (const auto& Rotation : Rotations)
+		{
+			if (Rotation.IsNormalized() == false)
+			{
+				return false;
+			}
+		}
+
+		return true;
+	}
+
+private:
+	TArray<uint8, AllocatorType> AllocatedMemory; // TODO : use the allocator directly
+
+	void UpdateViews(uint8* Memory, int32 NumTransforms)
+	{
+		if (NumTransforms > 0)
+		{
+			Rotations = TArrayView<FAnimRotation>((FAnimRotation*)(Memory), NumTransforms);
+
+			const int32 TranslationsOffset = sizeof(FAnimRotation) * NumTransforms;
+			Translations = TArrayView<FAnimTranslation>((FAnimTranslation*)(Memory + TranslationsOffset), NumTransforms);
+
+			const int32 Scales3DOffset = TranslationsOffset + sizeof(FAnimTranslation) * NumTransforms;
+			Scales3D = TArrayView<FAnimScale>((FAnimScale*)(Memory + Scales3DOffset), NumTransforms);
+		}
+		else
+		{
+			Rotations = TArrayView<FAnimRotation>();
+			Translations = TArrayView<FAnimTranslation>();
+			Scales3D = TArrayView<FAnimScale>();
+		}
+	}
+
+#if ENABLE_NAN_DIAGNOSTIC
+	FORCEINLINE static void DiagnosticCheckNaN_Scale3D(const FVector& Scale3D)
+	{
+		if (Scale3D.ContainsNaN())
+		{
+			logOrEnsureNanError(TEXT("TAnimTransformArraySoA Scale3D contains NaN"));
+		}
+	}
+
+	FORCEINLINE static void DiagnosticCheckNaN_Translate(const FVector& Translation)
+	{
+		if (Translation.ContainsNaN())
+		{
+			logOrEnsureNanError(TEXT("TAnimTransformArraySoA Translation contains NaN"));
+		}
+	}
+
+	FORCEINLINE static void DiagnosticCheckNaN_Rotate(const FQuat& Rotation)
+	{
+		if (Rotation.ContainsNaN())
+		{
+			logOrEnsureNanError(TEXT("TAnimTransformArraySoA Rotation contains NaN"));
+		}
+	}
+
+	FORCEINLINE void DiagnosticCheckNaN_All() const
+	{
+		for (const auto& Rotation : Rotations)
+		{
+			DiagnosticCheckNaN_Rotate(Rotation);
+		}
+		for (const auto& Translation : Translations)
+		{
+			DiagnosticCheckNaN_Translate(Translation);
+		}
+		for (const auto& Scale3D : Scales3D)
+		{
+			DiagnosticCheckNaN_Scale3D(Scale3D);
+		}
+	}
+
+	FORCEINLINE void DiagnosticCheck_IsValid() const
+	{
+		DiagnosticCheckNaN_All();
+		if (!IsValid())
+		{
+			logOrEnsureNanError(TEXT("TAnimTransformArraySoA is not valid"));
+		}
+	}
+
+#else
+	FORCEINLINE static void DiagnosticCheckNaN_Translate(const FVector& Translation) {}
+	FORCEINLINE static void DiagnosticCheckNaN_Rotate(const FQuat& Rotation) {}
+	FORCEINLINE static void DiagnosticCheckNaN_Scale3D(const FVector& Scale3D) {}
+	FORCEINLINE void DiagnosticCheckNaN_All() const {}
+	FORCEINLINE void DiagnosticCheck_IsValid() const {}
+#endif
+
+};
+
+/**
+* Transform Array using StructsOfArrays model
+*/
+using FAnimTransformArraySoAHeap = TAnimTransformArraySoA<FDefaultAllocator>;
+using FAnimTransformArraySoAStack = TAnimTransformArraySoA<FAnimStackAllocator>;
+
+
+//****************************************************************************
+
+#define DEFAULT_SOA 1
+#if DEFAULT_SOA
+template <typename AllocatorType>
+using TAnimTransformArray = TAnimTransformArraySoA<AllocatorType>;
+
+using FAnimTransformArrayHeap = FAnimTransformArraySoAHeap;
+using FAnimTransformArrayStack = FAnimTransformArraySoAStack;
+
+using FAnimTransformArray = FAnimTransformArraySoAHeap;
+
+#else
+template <typename AllocatorType>
+using TAnimTransformArray = TAnimTransformArrayAoS<AllocatorType>;
+
+using FAnimTransformArrayHeap = FAnimTransformArrayAoSHeap;
+using FAnimTransformArrayStack = FAnimTransformArrayAoSStack;
+
+using FAnimTransformArray = FAnimTransformArrayAoSHeap;
+
+#endif
+
+
+
+//****************************************************************************
+
+template <typename AllocatorType>
+struct TAnimationReferencePose
+{
+	enum EGenerationFlags : uint8
+	{
+		None = 0,
+		FastPath = 1 << 0
+	};
+
+	TAnimTransformArray<AllocatorType> ReferenceLocalTransforms;
+	TArray<TArray<FBoneIndexType, AllocatorType>, AllocatorType> LODBoneIndexes;
+	TArray<TArray<FBoneIndexType, AllocatorType>, AllocatorType> SkeletonToLODBoneIndexes;
+	TArray<int32, AllocatorType> LODNumBones;
+
+	TWeakObjectPtr<const USkeletalMesh> SkeletalMesh = nullptr;
+	TWeakObjectPtr<const USkeleton> Skeleton = nullptr;
+	EGenerationFlags GenerationFlags = EGenerationFlags::None;
+
+	TAnimationReferencePose() = default;
+
+	bool IsValid() const
+	{
+		return ReferenceLocalTransforms.Num() > 0;
+	}
+
+	int32 GetNumBonesForLOD(int32 LODLevel) const
+	{
+		const int32 NumLODS = LODNumBones.Num();
+
+		return (LODLevel < NumLODS) ? LODNumBones[LODLevel] : (NumLODS > 0) ? LODNumBones[0] : 0;
+	}
+
+	bool IsFastPath() const
+	{
+		return (GenerationFlags & EGenerationFlags::FastPath) != 0;
+	}
+
+	void Initialize(const FReferenceSkeleton& RefSkeleton
+		, const TArray<TArray<FBoneIndexType>>& InLODBoneIndexes
+		, const TArray<TArray<FBoneIndexType>>& InSkeletonToLODBoneIndexes
+		, const TArray<int32, AllocatorType>& InLODNumBones
+		, bool bFastPath = false)
+	{
+		const int32 NumBonesLOD0 = (InLODNumBones.Num()) > 0 ? InLODNumBones[0] : 0;
+
+		ReferenceLocalTransforms.SetNum(NumBonesLOD0);
+		LODBoneIndexes = InLODBoneIndexes;
+		SkeletonToLODBoneIndexes = InSkeletonToLODBoneIndexes;
+		LODNumBones = InLODNumBones;
+		
+		const TArray<FAnimTransform>& RefBonePose = RefSkeleton.GetRefBonePose();
+		const TArray<FBoneIndexType>& InBoneIndexes = InLODBoneIndexes[0]; // Fill the transforms with the LOD0 indexes
+
+		for (int32 i = 0; i < NumBonesLOD0; ++i)
+		{
+			ReferenceLocalTransforms[i] = RefBonePose[InBoneIndexes[i]]; // TODO : For SoA this is un-optimal, as we are using a TransformAdapter. Evaluate using a specific SoA iterator
+		}
+
+		GenerationFlags = bFastPath ? EGenerationFlags::FastPath : EGenerationFlags::None;
+	}
+
+	const TArrayView<const FBoneIndexType> GetLODBoneIndexes(int32 LODLevel) const
+	{
+		TArrayView<const FBoneIndexType> ArrayView;
+
+		if (LODLevel >= 0 && (IsFastPath() || LODLevel < LODBoneIndexes.Num()))
+		{
+			const int32 NumBonesForLOD = GetNumBonesForLOD(LODLevel);
+
+			if (IsFastPath())
+			{
+				ArrayView = MakeArrayView(LODBoneIndexes[0].GetData(), NumBonesForLOD);
+			}
+			else
+			{
+				ArrayView = MakeArrayView(LODBoneIndexes[LODLevel].GetData(), NumBonesForLOD);
+			}
+		}
+
+		return ArrayView;
+	}
+
+	const TArrayView<const FBoneIndexType> GetSkeletoonToLODBoneIndexes(int32 LODLevel) const
+	{
+		TArrayView<const FBoneIndexType> ArrayView;
+
+		if (LODLevel >= 0 && (IsFastPath() || LODLevel < SkeletonToLODBoneIndexes.Num()))
+		{
+			const int32 NumBonesForLOD = GetNumBonesForLOD(LODLevel);
+
+			if (IsFastPath())
+			{
+				return MakeArrayView(SkeletonToLODBoneIndexes[0].GetData(), NumBonesForLOD);
+			}
+			else
+			{
+				return MakeArrayView(SkeletonToLODBoneIndexes[LODLevel].GetData(), NumBonesForLOD);
+			}
+		}
+
+		return ArrayView;
+	}
+
+	int32 GetSkeletonBoneIndexFromLODBoneIndex(int32 LODBoneIndex) const
+	{
+		const int32 NumBonesLOD0 = LODBoneIndexes[0].Num();
+		check(LODBoneIndex < NumBonesLOD0);
+		return LODBoneIndexes[0][LODBoneIndex];
+	}
+
+	int32 GetLODBoneIndexFromSkeletonBoneIndex(int32 SkeletionBoneIndex) const
+	{
+		const int32 NumBonesLOD0 = SkeletonToLODBoneIndexes[0].Num();
+		check(SkeletionBoneIndex < NumBonesLOD0);
+		return SkeletonToLODBoneIndexes[0][SkeletionBoneIndex];
+	}
+
+	FTransform GetRefPoseTransform(int32 LODBoneIndex) const
+	{
+		const int32 NumBonesLOD0 = LODBoneIndexes[0].Num();
+		check(LODBoneIndex < NumBonesLOD0);
+
+		return ReferenceLocalTransforms[LODBoneIndex];
+	}
+
+	const FQuat& GetRefPoseRotation(int32 LODBoneIndex) const
+	{
+		const int32 NumBonesLOD0 = LODBoneIndexes[0].Num();
+		check(LODBoneIndex < NumBonesLOD0);
+
+		return ReferenceLocalTransforms[LODBoneIndex].GetRotation();
+	}
+
+	const FVector& GetRefPoseTranslation(int32 LODBoneIndex) const
+	{
+		const int32 NumBonesLOD0 = LODBoneIndexes[0].Num();
+		check(LODBoneIndex < NumBonesLOD0);
+
+		return ReferenceLocalTransforms[LODBoneIndex].GetTranslation();
+	}
+
+	const FVector& GetRefPoseScale3D(int32 LODBoneIndex) const
+	{
+		const int32 NumBonesLOD0 = LODBoneIndexes[0].Num();
+		check(LODBoneIndex < NumBonesLOD0);
+
+		return ReferenceLocalTransforms[LODBoneIndex].GetScale3D();
+	}
+
+};
+
+/**
+* Hide the allocator usage
+*/
+using FAnimationReferencePose = TAnimationReferencePose<FDefaultAllocator>;
+
+
+enum EAnimationPoseFlags : uint8
+{
+	None				= 0,
+	Additive			= 1 << 0,
+	DisableRetargeting	= 1 << 1,
+	UseRawData			= 1 << 2,
+	UseSourceData		= 1 << 3
+};
+
+template <typename AllocatorType>
+struct TAnimationLODPose
+{
+	static constexpr int32 INVALID_LOD_LEVEL = -1;
+
+	TAnimTransformArray<AllocatorType> LocalTransforms;
+	const FAnimationReferencePose* RefPose = nullptr;
+	int32 LODLevel = INVALID_LOD_LEVEL;
+	EAnimationPoseFlags Flags = EAnimationPoseFlags::None;
+
+	TAnimationLODPose() = default;
+
+	TAnimationLODPose(const FAnimationReferencePose& InRefPose, int32 InLODLevel, bool bSetRefPose = true, bool bAdditive = false)
+	{
+		PrepareForLOD(InRefPose, InLODLevel, bSetRefPose, bAdditive);
+	}
+
+	void PrepareForLOD(const FAnimationReferencePose& InRefPose, int32 InLODLevel, bool bSetRefPose = true, bool bAdditive = false)
+	{
+		LODLevel = InLODLevel;
+		RefPose = &InRefPose;
+
+		const int32 NumTransforms = InRefPose.GetNumBonesForLOD(InLODLevel);
+		LocalTransforms.SetNum(NumTransforms);
+		Flags = (EAnimationPoseFlags)(bAdditive ? (Flags | EAnimationPoseFlags::Additive) : Flags & EAnimationPoseFlags::Additive);
+
+		if (bSetRefPose && NumTransforms > 0)
+		{
+			SetRefPose(bAdditive);
+		}
+	}
+
+	void SetRefPose(bool bAdditive = false)
+	{
+		const int32 NumTransforms = LocalTransforms.Num();
+
+		if (NumTransforms > 0)
+		{
+			if (bAdditive)
+			{
+				SetIdentity(bAdditive);
+			}
+			else
+			{
+				check(RefPose != nullptr);
+				LocalTransforms.CopyTransforms(RefPose->ReferenceLocalTransforms, 0, NumTransforms);
+			}
+		}
+
+		Flags = (EAnimationPoseFlags)(bAdditive ? (Flags | EAnimationPoseFlags::Additive) : Flags & EAnimationPoseFlags::Additive);
+	}
+
+	const FAnimationReferencePose& GetRefPose() const
+	{
+		check(RefPose != nullptr);
+		return *RefPose;
+	}
+
+	void SetIdentity(bool bAdditive = false)
+	{
+		LocalTransforms.SetIdentity(bAdditive);
+	}
+
+	int32 GetNumBones() const
+	{
+		return RefPose != nullptr ? RefPose->GetNumBonesForLOD(LODLevel) : 0;
+	}
+
+	const TArrayView<const FBoneIndexType> GetLODBoneIndexes() const
+	{
+		if (LODLevel != INVALID_LOD_LEVEL && RefPose != nullptr)
+		{
+			return RefPose->GetLODBoneIndexes(LODLevel);
+		}
+		else
+		{
+			return TArrayView<const FBoneIndexType>();
+		}
+	}
+
+	const TArrayView<const FBoneIndexType> GetSkeletoonToLODBoneIndexes() const
+	{
+		if (LODLevel != INVALID_LOD_LEVEL && RefPose != nullptr)
+		{
+			return RefPose->GetSkeletoonToLODBoneIndexes(LODLevel);
+		}
+		else
+		{
+			return TArrayView<const FBoneIndexType>();
+		}
+	}
+
+	const USkeleton* GetSkeletonAsset() const
+	{
+		return RefPose != nullptr ? RefPose->Skeleton.Get() : nullptr;
+	}
+
+	/** Disable Retargeting */
+	void SetDisableRetargeting(bool bDisableRetargeting)
+	{
+		Flags = (EAnimationPoseFlags)(bDisableRetargeting ? (Flags | EAnimationPoseFlags::DisableRetargeting) : Flags & EAnimationPoseFlags::DisableRetargeting);
+	}
+
+	/** True if retargeting is disabled */
+	bool GetDisableRetargeting() const
+	{
+		return (Flags & EAnimationPoseFlags::DisableRetargeting) != 0;
+	}
+
+	/** Ignore compressed data and use RAW data instead, for debugging. */
+	void SetUseRAWData(bool bUseRAWData)
+	{
+		Flags = (EAnimationPoseFlags)(bUseRAWData ? (Flags | EAnimationPoseFlags::UseRawData) : Flags & EAnimationPoseFlags::UseRawData);
+	}
+
+	/** True if we're requesting RAW data instead of compressed data. For debugging. */
+	bool ShouldUseRawData() const
+	{
+		return (Flags & EAnimationPoseFlags::UseRawData) != 0;
+	}
+
+	/** Use Source data instead.*/
+	void SetUseSourceData(bool bUseSourceData)
+	{
+		Flags = (EAnimationPoseFlags)(bUseSourceData ? (Flags | EAnimationPoseFlags::UseSourceData) : Flags & EAnimationPoseFlags::UseSourceData);
+	}
+
+	/** True if we're requesting Source data instead of RawAnimationData. For debugging. */
+	bool ShouldUseSourceData() const
+	{
+		return (Flags & EAnimationPoseFlags::UseSourceData) != 0;
+	}
+
+};
+
+using FAnimationLODPoseHeap = TAnimationLODPose<FDefaultAllocator>;
+using FAnimationLODPoseStack = TAnimationLODPose<FAnimStackAllocator>;
+
+using FAnimationLODPose = FAnimationLODPoseHeap;
+
+DECLARE_BUILTIN_ANIM_NEXT_INTERFACE_PARAM_TYPE(FAnimationReferencePose, AnimationReferencePose);
+//DECLARE_BUILTIN_ANIM_NEXT_INTERFACE_PARAM_TYPE(FAnimationLODPose, AnimationLODPose);
+//DECLARE_BUILTIN_ANIM_NEXT_INTERFACE_PARAM_TYPE(FAnimationLODPoseStack, AnimationLODPoseStack);
+
+} // namespace UE::AnimNext::Interface
diff --git a/Engine/Plugins/Experimental/AnimNext/Source/AnimNextInterfaceGraph/Internal/AnimGraph/AnimNode_AnimNextInterfaceGraph.h b/Engine/Plugins/Experimental/AnimNext/Source/AnimNextInterfaceGraph/Internal/AnimGraph/AnimNode_AnimNextInterfaceGraph.h
--- a/Engine/Plugins/Experimental/AnimNext/Source/AnimNextInterfaceGraph/Internal/AnimGraph/AnimNode_AnimNextInterfaceGraph.h
+++ b/Engine/Plugins/Experimental/AnimNext/Source/AnimNextInterfaceGraph/Internal/AnimGraph/AnimNode_AnimNextInterfaceGraph.h
@@ -5,8 +5,12 @@
 #include "Animation/AnimNode_CustomProperty.h"
 #include "Animation/InputScaleBias.h"
 #include "AnimNextInterfaceGraph.h"
+#include "AnimNextInterfaceContext.h"
+#include "AnimNextInterfaceParamStorage.h"
+// --- ---
 #include "AnimNode_AnimNextInterfaceGraph.generated.h"
 
+
 class UNodeMappingContainer;
 
 /**
@@ -39,12 +43,15 @@ private:
 private:
 
 	/** The input pose we will pass to the graph */
-	UPROPERTY(EditAnywhere, Category = Links)
-	FPoseLink Source;
+	UPROPERTY(EditAnywhere, Category = Links, meta = (DisplayName = "Source"))
+	FPoseLink SourceLink;
 
 	UPROPERTY(EditAnywhere, Category = Settings, meta = (PinHiddenByDefault))
 	TObjectPtr<UAnimNextInterfaceGraph> AnimNextInterfaceGraph;
 
+	UPROPERTY(EditAnywhere, Category = Settings, meta = (PinHiddenByDefault, DisallowedClasses = "/Script/Engine.AnimMontage"))
+	TObjectPtr<UAnimSequence> TestSequence = nullptr;
+
 	/*
 	 * Max LOD that this node is allowed to run
 	 * For example if you have LODThreadhold to be 2, it will run until LOD 2 (based on 0 index)
@@ -54,6 +61,12 @@ private:
 	UPROPERTY(EditAnywhere, Category = Performance, meta = (DisplayName = "LOD Threshold"))
 	int32 LODThreshold;
 
+	// This should come from the FPoseContext, but as I need persistence, I store the state here
+	UE::AnimNext::Interface::FState RootState; // TODO : Get rid of num elements ? Let the param itself have it
+
+	// Delta time received accumulated in update and used at Evaluate (so we can receive multiple calls to Evaluate)
+	float GraphDeltaTime = 0.f;
+
 protected:
 	virtual UClass* GetTargetClass() const override { return AnimNextInterfaceGraph ? AnimNextInterfaceGraph->StaticClass() : nullptr; }
 	
diff --git a/Engine/Plugins/Experimental/AnimNext/Source/AnimNextInterfaceGraph/Internal/AnimNextInterface_LODPose.h b/Engine/Plugins/Experimental/AnimNext/Source/AnimNextInterfaceGraph/Internal/AnimNextInterface_LODPose.h
new file mode 100644
--- /dev/null
+++ b/Engine/Plugins/Experimental/AnimNext/Source/AnimNextInterfaceGraph/Internal/AnimNextInterface_LODPose.h
@@ -0,0 +1,82 @@
+// Copyright Epic Games, Inc. All Rights Reserved.
+
+#pragma once
+
+#include "CoreMinimal.h"
+#include "AnimNextInterfaceTypes.h"
+#include "IAnimNextInterface.h"
+#include "AnimNextInterfaceContext.h"
+#include "AnimationReferencePose.h"
+#include "AnimNextInterfaceTypes.h"
+#include "AnimNextInterface_LODPose.generated.h"
+
+USTRUCT(BlueprintType, meta = (DisplayName = "ReferencePose"))
+struct FAnimNextGraphReferencePose
+{
+	GENERATED_BODY()
+
+	FAnimNextGraphReferencePose() = default;
+
+	explicit FAnimNextGraphReferencePose(const UE::AnimNext::Interface::FAnimationReferencePose* InReferencePose)
+		: ReferencePose(InReferencePose)
+	{
+	}
+
+	const UE::AnimNext::Interface::FAnimationReferencePose* ReferencePose = nullptr;
+};
+DECLARE_ANIM_NEXT_INTERFACE_PARAM_TYPE(ANIMNEXTINTERFACEGRAPH_API, FAnimNextGraphReferencePose, FAnimNextGraphReferencePose);
+
+USTRUCT(BlueprintType, meta = (DisplayName = "LODPose"))
+struct FAnimNextGraphLODPose
+{
+	GENERATED_BODY()
+
+	FAnimNextGraphLODPose() = default;
+
+	explicit FAnimNextGraphLODPose(const UE::AnimNext::Interface::FAnimationLODPose& InLODPose)
+		: LODPose(InLODPose)
+	{
+	}
+	explicit FAnimNextGraphLODPose(UE::AnimNext::Interface::FAnimationLODPose&& InLODPose)
+		: LODPose(MoveTemp(InLODPose))
+	{
+	}
+
+	UE::AnimNext::Interface::FAnimationLODPose LODPose;
+};
+DECLARE_ANIM_NEXT_INTERFACE_PARAM_TYPE(ANIMNEXTINTERFACEGRAPH_API, FAnimNextGraphLODPose, FAnimNextGraphLODPose);
+
+
+UCLASS()
+class UAnimNextInterfaceGraphLODPose : public UObject, public IAnimNextInterface
+{
+	GENERATED_BODY()
+
+	virtual FName GetReturnTypeNameImpl() const final override
+	{
+		static const FName Name = TNameOf<FAnimNextGraphLODPose>::GetName();
+		return Name;
+	}
+
+	virtual bool GetDataImpl(const UE::AnimNext::Interface::FContext& Context) const override
+	{
+		checkf(false, TEXT("UAnimNextInterfaceLODPose::GetDataImpl must be overridden"));
+		return false;
+	}
+};
+
+UCLASS()
+class UAnimNextInterface_LODPose_Literal : public UAnimNextInterfaceGraphLODPose
+{
+	GENERATED_BODY()
+
+	// IAnimNextInterface interface
+	virtual bool GetDataImpl(const UE::AnimNext::Interface::FContext& Context) const final override
+	{
+		Context.SetResult(Value);
+		return true;
+	}
+	
+	UPROPERTY(EditAnywhere, Category = "Parameters")
+	FAnimNextGraphLODPose Value;
+};
diff --git a/Engine/Plugins/Experimental/AnimNext/Source/AnimNextInterfaceGraph/Internal/RigUnit_AnimNextInterfaceEndExecution.h b/Engine/Plugins/Experimental/AnimNext/Source/AnimNextInterfaceGraph/Internal/RigUnit_AnimNextInterfaceEndExecution.h
--- a/Engine/Plugins/Experimental/AnimNext/Source/AnimNextInterfaceGraph/Internal/RigUnit_AnimNextInterfaceEndExecution.h
+++ b/Engine/Plugins/Experimental/AnimNext/Source/AnimNextInterfaceGraph/Internal/RigUnit_AnimNextInterfaceEndExecution.h
@@ -6,6 +6,8 @@
 #include "AnimNextInterfaceExecuteContext.h"
 #include "AnimNextInterfaceUnitContext.h"
 #include "AnimNextInterfaceTypes.h"
+#include "AnimNextInterface_LODPose.h"
+// --- ---
 #include "RigUnit_AnimNextInterfaceEndExecution.generated.h"
 
 struct FRigUnitContext;
@@ -68,3 +70,22 @@ struct ANIMNEXTINTERFACEGRAPH_API FRigUnit_AnimNextInterfaceEndExecution_Float :
 	UPROPERTY(EditAnywhere, Category = Result, meta = (Input))
 	float Result;
 };
+
+/** Event for writing back a calculated pose */
+USTRUCT(meta = (DisplayName = "End Execute LODPose", Category = "Events", TitleColor = "1 0 0", NodeColor = "1 1 1"))
+struct ANIMNEXTINTERFACEGRAPH_API FRigUnit_AnimNextInterfaceEndExecution_LODPose : public FRigUnit_AnimNextInterfaceEndExecution
+{
+	GENERATED_BODY()
+
+	FRigUnit_AnimNextInterfaceEndExecution_LODPose()
+		: FRigUnit_AnimNextInterfaceEndExecution()
+	{}
+
+	RIGVM_METHOD()
+	virtual void Execute() override;
+
+	virtual bool CanOnlyExistOnce() const override { return true; }
+
+	UPROPERTY(EditAnywhere, Category = Result, meta = (Input))
+	FAnimNextGraphLODPose Result;
+};
diff --git a/Engine/Plugins/Experimental/AnimNext/Source/AnimNextInterfaceGraph/Private/AnimGraph/AnimNode_AnimNextInterfaceGraph.cpp b/Engine/Plugins/Experimental/AnimNext/Source/AnimNextInterfaceGraph/Private/AnimGraph/AnimNode_AnimNextInterfaceGraph.cpp
--- a/Engine/Plugins/Experimental/AnimNext/Source/AnimNextInterfaceGraph/Private/AnimGraph/AnimNode_AnimNextInterfaceGraph.cpp
+++ b/Engine/Plugins/Experimental/AnimNext/Source/AnimNextInterfaceGraph/Private/AnimGraph/AnimNode_AnimNextInterfaceGraph.cpp
@@ -1,6 +1,8 @@
 // Copyright Epic Games, Inc. All Rights Reserved.
 
 #include "AnimGraph/AnimNode_AnimNextInterfaceGraph.h"
+#include "AnimationDataRegistry.h"
+#include "AnimationReferencePose.h"
 #include "ControlRig.h"
 #include "ControlRigComponent.h"
 #include "Components/SkeletalMeshComponent.h"
@@ -8,10 +10,15 @@
 #include "GameFramework/Actor.h"
 #include "Animation/NodeMappingContainer.h"
 #include "AnimationRuntime.h"
-#include "AnimNextInterfaceContext.h"
-#include "AnimNextInterfaceParamStorage.h"
 #include "AnimNextInterfaceTypes.h"
 #include "AnimNextInterface.h"
+#include "AnimationDataRegistry.h"
+#include "AnimationGenerationTools.h"
+#include "AnimationReferencePose.h"
+#include "AnimNextInterface_LODPose.h"
+#include "RigUnit_AnimNextAnimSequence.h" // TEST
+#include "Engine/SkeletalMesh.h"
+#include "BoneContainer.h"
 
 #include UE_INLINE_GENERATED_CPP_BY_NAME(AnimNode_AnimNextInterfaceGraph)
 
@@ -23,6 +30,7 @@ FAnimNode_AnimNextInterfaceGraph::FAnimNode_AnimNextInterfaceGraph()
 	: FAnimNode_CustomProperty()
 	, AnimNextInterfaceGraph(nullptr)
 	, LODThreshold(INDEX_NONE)
+	, RootState(1) // TODO : Get rid of num elements ? Let the param itself have it
 {
 }
 
@@ -43,14 +51,16 @@ void FAnimNode_AnimNextInterfaceGraph::GatherDebugData(FNodeDebugData& DebugData
 {
 	DECLARE_SCOPE_HIERARCHICAL_COUNTER_FUNC()
 
-	Source.GatherDebugData(DebugData.BranchFlow(1.f));
+	SourceLink.GatherDebugData(DebugData.BranchFlow(1.f));
 }
 
 void FAnimNode_AnimNextInterfaceGraph::Update_AnyThread(const FAnimationUpdateContext& Context)
 {
 	DECLARE_SCOPE_HIERARCHICAL_COUNTER_FUNC()
 
-	Source.Update(Context);
+	GraphDeltaTime += Context.GetDeltaTime();
+
+	SourceLink.Update(Context);
 
 	if (IsLODEnabled(Context.AnimInstanceProxy))
 	{
@@ -68,7 +78,7 @@ void FAnimNode_AnimNextInterfaceGraph::Initialize_AnyThread(const FAnimationInit
 {
 	DECLARE_SCOPE_HIERARCHICAL_COUNTER_FUNC()
 
-	Source.Initialize(Context);
+	SourceLink.Initialize(Context);
 
 	FAnimNode_CustomProperty::Initialize_AnyThread(Context);
 }
@@ -79,41 +89,69 @@ void FAnimNode_AnimNextInterfaceGraph::CacheBones_AnyThread(const FAnimationCach
 
 	FAnimNode_CustomProperty::CacheBones_AnyThread(Context);
 
-	Source.CacheBones(Context);
+	SourceLink.CacheBones(Context);
 }
 
 void FAnimNode_AnimNextInterfaceGraph::Evaluate_AnyThread(FPoseContext & Output)
 {
 	DECLARE_SCOPE_HIERARCHICAL_COUNTER_FUNC()
 
+	using namespace UE::AnimNext::Interface;
+
 	FPoseContext SourcePose(Output);
 
-	if (Source.GetLinkNode())
+	if (SourceLink.GetLinkNode())
 	{
-		Source.Evaluate(SourcePose);
+		SourceLink.Evaluate(SourcePose);
 	}
 	else
 	{
 		SourcePose.ResetToRefPose();
 	}
 
-	// This is just a simple test to execute the graph and obtain a float result
-	// We have to create the Pose API, in order to be able to pass / receive the animation data
-	if (AnimNextInterfaceGraph->GetReturnTypeName() == TNameOf<float>::GetName())
-	{
-		UE::AnimNext::Interface::FState RootState(1); // TODO : Get rid of num elements ? Let the param itself have it
-		UE::AnimNext::Interface::FParamStorage ParamStorage(1, 16, 1);	// TODO : see if we can have state and params together and optional. Also, default size
-		UE::AnimNext::Interface::FContext RootContext(0.f, RootState, ParamStorage);
-		TScriptInterface<IAnimNextInterface> ScriptInterface = AnimNextInterfaceGraph;
+	USkeletalMeshComponent* SkeletalMeshComponent = Output.AnimInstanceProxy->GetSkelMeshComponent();
+	check(SkeletalMeshComponent != nullptr);
 
-		float ReturnValue = 0.f;
-		const bool bOk = UE::AnimNext::Interface::GetDataSafe(ScriptInterface, RootContext, ReturnValue);
-	}
+	FAnimationDataHandle RefPoseHandle = FAnimationDataRegistry::Get()->GetOrGenerateReferencePose(SkeletalMeshComponent);
+	const FAnimationReferencePose& RefPose = RefPoseHandle.GetRef<FAnimationReferencePose>();
+
+	const int32 LODLevel = Output.AnimInstanceProxy->GetLODLevel();
+
+	// TODO : In a normal graph, FState/FParamStorage/FContext would come from FAnimationUpdateContext, here I have to create one the fly for the AnimNext graph
+	//UE::AnimNext::Interface::FState RootState; // moved to .h as I need state persistence in the AnimNext graph
+	UE::AnimNext::Interface::FParamStorage ParamStorage(5, 0, 0);	// TODO : see if we can have state and params together and optional. Also, default size
+	
+	// TODO : Using AnimInstanceProxy->GetDeltaSeconds() makes the preview to advance  multiple times when debug options are activated (i.e. ShowUncompressedAnim)
+	// See where to get / how to calculate the correct delta value
+	UE::AnimNext::Interface::FContext RootContext(GraphDeltaTime, RootState, ParamStorage);
+	GraphDeltaTime = 0.f; // Reset for the case that we receive multiple calls to Evaluate (debug options)
+
+	FAnimNextGraphReferencePose GraphReferencePose(&RefPose);
+	FAnimNextGraph_AnimSequence GraphTestSequence(TestSequence);  // TEST
+
+	FAnimNextGraphLODPose GraphSourceLODPose(FAnimationLODPose(RefPose, LODLevel, false, Output.ExpectsAdditivePose()));
+	FGenerationTools::RemapPose(LODLevel, SourcePose, RefPose, GraphSourceLODPose.LODPose);
 
-	// TODO
-	//ExecuteAnimNextInterface (SourcePose);
+	RootContext.AddInputReference(FName(TEXT("GraphReferencePose")), GraphReferencePose);			// TODO : Should this go into the AnimNextVM Context ?
+	RootContext.AddInputValue(FName(TEXT("GraphLODLevel")), LODLevel);								// TODO : Should this go into the AnimNextVM Context ?
+	RootContext.AddInputValue(FName(TEXT("GraphExpectsAdditive")), Output.ExpectsAdditivePose());	// TODO : Should this go into the AnimNextVM Context ?
+
+	RootContext.AddInputReference(FName(TEXT("SourcePose")), GraphSourceLODPose);		// TODO : Pass this as a external variable maybe ? When we have support for variables in the rigvm graph
+	RootContext.AddInputReference(FName(TEXT("TestSequence")), GraphTestSequence);		// TEST anim decompression - Anim Sequence to decompress
+		
+	TScriptInterface<IAnimNextInterface> ScriptInterface = AnimNextInterfaceGraph;
+	FAnimNextGraphLODPose GraphResultLODPose(FAnimationLODPose(RefPose, LODLevel, true, Output.ExpectsAdditivePose()));
+
+	const bool bOk = UE::AnimNext::Interface::GetDataSafe(ScriptInterface, RootContext, GraphResultLODPose);
+	if (bOk)
+	{
+		FGenerationTools::RemapPose(LODLevel, RefPose, GraphResultLODPose.LODPose, Output);
+	}
+	else
+	{
+		Output.ResetToRefPose();
+	}
 
-	Output = SourcePose;
 
 	// evaluate 
 	FAnimNode_CustomProperty::Evaluate_AnyThread(Output);
diff --git a/Engine/Plugins/Experimental/AnimNext/Source/AnimNextInterfaceGraph/Private/AnimNextInterface_LODPose.cpp b/Engine/Plugins/Experimental/AnimNext/Source/AnimNextInterfaceGraph/Private/AnimNextInterface_LODPose.cpp
new file mode 100644
--- /dev/null
+++ b/Engine/Plugins/Experimental/AnimNext/Source/AnimNextInterfaceGraph/Private/AnimNextInterface_LODPose.cpp
@@ -0,0 +1,14 @@
+// Copyright Epic Games, Inc. All Rights Reserved.
+
+#include "AnimNextInterface_LODPose.h"
+#include "AnimNextInterface.h"
+#include "AnimNextInterfaceTypes.h"
+#include "AnimNextInterfaceKernel.h"
+
+IMPLEMENT_ANIM_NEXT_INTERFACE_PARAM_TYPE(FAnimNextGraphLODPose, FAnimNextGraphLODPose);
+IMPLEMENT_ANIM_NEXT_INTERFACE_PARAM_TYPE(FAnimNextGraphReferencePose, FAnimNextGraphReferencePose);
+
+//#define LOCTEXT_NAMESPACE "AnimNextInterface"
+
+
+//#undef LOCTEXT_NAMESPACE
diff --git a/Engine/Plugins/Experimental/AnimNext/Source/AnimNextInterfaceGraph/Private/RigUnit_AnimNextAnimSequence.cpp b/Engine/Plugins/Experimental/AnimNext/Source/AnimNextInterfaceGraph/Private/RigUnit_AnimNextAnimSequence.cpp
new file mode 100644
--- /dev/null
+++ b/Engine/Plugins/Experimental/AnimNext/Source/AnimNextInterfaceGraph/Private/RigUnit_AnimNextAnimSequence.cpp
@@ -0,0 +1,106 @@
+// Copyright Epic Games, Inc. All Rights Reserved.
+
+#include "RigUnit_AnimNextAnimSequence.h"
+#include "AnimNextInterfaceUnitContext.h"
+#include "AnimNextInterfaceExecuteContext.h"
+#include "AnimNextInterface.h"
+#include "AnimNextInterfaceTypes.h"
+#include "AnimNextInterfaceKernel.h"
+#include "AnimationDecompressionTools.h"
+#include "Animation/AnimSequence.h"
+
+#include UE_INLINE_GENERATED_CPP_BY_NAME(RigUnit_AnimNextAnimSequence)
+
+
+IMPLEMENT_ANIM_NEXT_INTERFACE_PARAM_TYPE(FAnimNextGraph_AnimSequence, FAnimNextGraph_AnimSequence);
+
+
+struct FAnimSequencePlayerState
+{
+	float InternalTimeAccumulator = 0.0f;
+	float PrevInternalTimeAccumulator = 0.0f;
+};
+
+IMPLEMENT_ANIM_NEXT_INTERFACE_STATE_TYPE(FAnimSequencePlayerState, AnimSequencePlayerState);
+
+// --- FRigUnit_AnimNext_AnimSequenceAsset --- 
+
+FRigUnit_AnimNext_AnimSequenceAsset_Execute()
+{
+	using namespace UE::AnimNext::Interface;
+
+	//Sequence.AnimSequence = AnimSequence; // TODO : Enable once we have RigVM object support enabled
+}
+
+// --- FRigUnit_AnimNext_AnimSequencePlayer --- 
+
+FRigUnit_AnimNext_AnimSequencePlayer_Initialize()
+{
+	using namespace UE::AnimNext::Interface;
+
+	const FContext& AnimNextInterfaceContext = ExecuteContext.GetContext();
+
+	const UAnimSequenceBase* AnimSequenceBase = Sequence.AnimSequence;
+	if (AnimSequenceBase != nullptr)
+	{
+		// Get or Create the internal state
+		const TParam<FAnimSequencePlayerState> AnimSequencePlayerState = AnimNextInterfaceContext.GetState<FAnimSequencePlayerState>(ExecuteContext.GetInterface(), 0);
+
+		const float SequenceLength = AnimSequenceBase->GetPlayLength();
+		AnimSequencePlayerState->InternalTimeAccumulator = FMath::Clamp(Parameters.StartPosition, 0.f, SequenceLength);
+		AnimSequencePlayerState->PrevInternalTimeAccumulator = AnimSequencePlayerState->InternalTimeAccumulator;
+	}
+}
+
+FRigUnit_AnimNext_AnimSequencePlayer_Execute()
+{
+	using namespace UE::AnimNext::Interface;
+
+	const FContext& AnimNextInterfaceContext = ExecuteContext.GetContext();
+	
+	const UAnimSequenceBase* AnimSequenceBase = Sequence.AnimSequence;
+	if (AnimSequenceBase != nullptr)
+	{
+		// Get or Create the internal state
+		//const TParam<FAnimSequencePlayerState> AnimSequencePlayerState = AnimNextInterfaceContext.GetState<FAnimSequencePlayerState>(ExecuteContext.GetInterface(), 0);
+		FAnimSequencePlayerState& AnimSequencePlayerState = AnimNextInterfaceContext.GetState<FAnimSequencePlayerState, EStatePersistence::Permanent>(ExecuteContext.GetInterface(), 0).GetDataChecked();
+
+		const float DeltaTime = AnimNextInterfaceContext.GetDeltaTime();
+		const float EffectiveDelta = FMath::IsNearlyZero(DeltaTime) || FMath::IsNearlyZero(Parameters.PlayRate) ? 0.f : DeltaTime * Parameters.PlayRate;
+
+		const float SequenceLength = AnimSequenceBase->GetPlayLength();
+		float CurrentTime = Parameters.bLoop
+			? FMath::Fmod(AnimSequencePlayerState.InternalTimeAccumulator + EffectiveDelta, SequenceLength)
+			: FMath::Clamp(AnimSequencePlayerState.InternalTimeAccumulator + EffectiveDelta, 0.f, SequenceLength);
+
+		if (Parameters.bLoop && CurrentTime < 0.f)
+		{
+			CurrentTime += SequenceLength;
+		}
+
+		AnimSequencePlayerState.PrevInternalTimeAccumulator = AnimSequencePlayerState.InternalTimeAccumulator;
+		AnimSequencePlayerState.InternalTimeAccumulator = CurrentTime;
+
+		FDeltaTimeRecord DeltaTimeRecord;
+		DeltaTimeRecord.Set(AnimSequencePlayerState.PrevInternalTimeAccumulator, AnimSequencePlayerState.InternalTimeAccumulator);
+
+		const FAnimExtractContext ExtractionContext(static_cast<double>(AnimSequencePlayerState.InternalTimeAccumulator)
+			, false /*Output.AnimInstanceProxy->ShouldExtractRootMotion()*/
+			, DeltaTimeRecord
+			, Parameters.bLoop);
+
+		const FAnimNextGraphReferencePose& GraphReferencePose = AnimNextInterfaceContext.GetParameterChecked<FAnimNextGraphReferencePose>(TEXT("GraphReferencePose")).GetDataChecked();
+		const int32 GraphLODLevel = AnimNextInterfaceContext.GetParameterChecked<int32>(TEXT("GraphLODLevel")).GetDataChecked();
+		const bool bGraphExpectsAdditive = AnimNextInterfaceContext.GetParameterChecked<bool>(TEXT("GraphExpectsAdditive")).GetDataChecked();
+
+		LODPose.LODPose.PrepareForLOD(*GraphReferencePose.ReferencePose, GraphLODLevel, true, bGraphExpectsAdditive);
+		
+		// Note : calling FDecompressionTools instead of UAnimSequence / UAnimSequenceBase, as I can not modify the engine code (so I extracted the function)
+		// TODO : Revisit this, so we can plug other sequence types
+		const UAnimSequence* AnimSequence = Cast<UAnimSequence>(AnimSequenceBase);
+		if (AnimSequence != nullptr)
+		{
+			FDecompressionTools::GetAnimationPose(AnimSequence, LODPose.LODPose, ExtractionContext);
+		}
+	}
+}
diff --git a/Engine/Plugins/Experimental/AnimNext/Source/AnimNextInterfaceGraph/Private/RigUnit_AnimNextAnimSequence.h b/Engine/Plugins/Experimental/AnimNext/Source/AnimNextInterfaceGraph/Private/RigUnit_AnimNextAnimSequence.h
new file mode 100644
--- /dev/null
+++ b/Engine/Plugins/Experimental/AnimNext/Source/AnimNextInterfaceGraph/Private/RigUnit_AnimNextAnimSequence.h
@@ -0,0 +1,98 @@
+// Copyright Epic Games, Inc. All Rights Reserved.
+
+#pragma once
+
+#include "IAnimNextInterface.h"
+#include "Units/RigUnit.h"
+#include "AnimNextInterfaceExecuteContext.h"
+#include "AnimNextInterface_LODPose.h"
+// --- ---
+#include "RigUnit_AnimNextAnimSequence.generated.h"
+
+class UAnimSequenceBase;
+
+USTRUCT(BlueprintType, meta = (DisplayName = "AnimSequence"))
+struct FAnimNextGraph_AnimSequence
+{
+	GENERATED_BODY()
+
+	FAnimNextGraph_AnimSequence() = default;
+
+	explicit FAnimNextGraph_AnimSequence(const UAnimSequenceBase* InAnimSequence)
+		: AnimSequence(InAnimSequence)
+	{
+	}
+
+	const UAnimSequenceBase* AnimSequence = nullptr;
+};
+
+DECLARE_ANIM_NEXT_INTERFACE_PARAM_TYPE(ANIMNEXTINTERFACEGRAPH_API, FAnimNextGraph_AnimSequence, FAnimNextGraph_AnimSequence);
+
+
+
+struct FAnimNextInterfaceUnitContext;
+class UAnimNextInterfaceGraph;
+class UAnimSequenceBase;
+
+
+// --- FRigUnit_AnimNext_AnimSequenceAsset ---
+
+/** Unit for getting an asset reference we can pass to a player or data extractor */
+USTRUCT(meta = (DisplayName = "Anim Sequence Asset", Category = "Animation", TitleColor = "1 0 0", NodeColor = "1 1 1"))
+struct FRigUnit_AnimNext_AnimSequenceAsset : public FRigUnit_AnimNextInterfaceBase
+{
+	GENERATED_BODY()
+
+	RIGVM_METHOD()
+	void Execute();
+
+protected:
+	// TODO : Enable once we have RigVM object support enabled
+	//UPROPERTY(EditAnywhere, Category = "Asset", meta = (Input, EditInline, PinHiddenByDefault, DisallowedClasses = "/Script/Engine.AnimMontage"))
+	//TObjectPtr<UAnimSequenceBase> AnimSequence = nullptr;
+
+	UPROPERTY(EditAnywhere, Category = "Animation", DisplayName = "AnimSequence", meta = (Output))
+	FAnimNextGraph_AnimSequence Sequence;
+};
+
+// --- FRigUnit_AnimNext_AnimSequencePlayer --- 
+
+USTRUCT(BlueprintType)
+struct FAnimSequenceParameters
+{
+	GENERATED_BODY()
+
+	UPROPERTY(EditAnywhere, Category = "Anim Sequence")
+	float PlayRate = 1.0f;
+
+	UPROPERTY(EditAnywhere, Category = "Anim Sequence")
+	float StartPosition = 0.0f;
+
+	UPROPERTY(EditAnywhere, Category = "Anim Sequence")
+	bool bLoop = 0.0f;
+};
+
+/** Unit for getting a pose from an animation sequence */
+USTRUCT(meta = (DisplayName = "AnimSequencePlayer", Category = "Animation", TitleColor = "1 0 0", NodeColor = "1 1 1"))
+struct FRigUnit_AnimNext_AnimSequencePlayer : public FRigUnit_AnimNextInterfaceBase
+{
+	GENERATED_BODY()
+
+	RIGVM_METHOD()
+	void Initialize();
+
+	RIGVM_METHOD()
+	void Execute();
+
+protected:
+	UPROPERTY(EditAnywhere, Category = "Anim Interface", meta = (Input))
+	FAnimSequenceParameters Parameters;
+
+	// The animation sequence asset to play
+	UPROPERTY(EditAnywhere, Category = "Animation", DisplayName = "AnimSequence", meta = (Input))
+	FAnimNextGraph_AnimSequence Sequence;
+
+	// The output pose generated
+	UPROPERTY(EditAnywhere, Transient, Category = "Animation", DisplayName = "Result", meta = (Output))
+	FAnimNextGraphLODPose LODPose;
+};
diff --git a/Engine/Plugins/Experimental/AnimNext/Source/AnimNextInterfaceGraph/Private/RigUnit_AnimNextInterfaceEndExecution.cpp b/Engine/Plugins/Experimental/AnimNext/Source/AnimNextInterfaceGraph/Private/RigUnit_AnimNextInterfaceEndExecution.cpp
--- a/Engine/Plugins/Experimental/AnimNext/Source/AnimNextInterfaceGraph/Private/RigUnit_AnimNextInterfaceEndExecution.cpp
+++ b/Engine/Plugins/Experimental/AnimNext/Source/AnimNextInterfaceGraph/Private/RigUnit_AnimNextInterfaceEndExecution.cpp
@@ -14,3 +14,8 @@ FRigUnit_AnimNextInterfaceEndExecution_Float_Execute()
 {
 	SetResult(ExecuteContext, Result);
 }
+
+FRigUnit_AnimNextInterfaceEndExecution_LODPose_Execute()
+{
+	SetResult(ExecuteContext, Result);
+}
\ No newline at end of file
diff --git a/Engine/Plugins/Experimental/AnimNext/Source/AnimNextInterfaceGraph/Private/RigUnit_AnimNextInterfaceParameter.cpp b/Engine/Plugins/Experimental/AnimNext/Source/AnimNextInterfaceGraph/Private/RigUnit_AnimNextInterfaceParameter.cpp
--- a/Engine/Plugins/Experimental/AnimNext/Source/AnimNextInterfaceGraph/Private/RigUnit_AnimNextInterfaceParameter.cpp
+++ b/Engine/Plugins/Experimental/AnimNext/Source/AnimNextInterfaceGraph/Private/RigUnit_AnimNextInterfaceParameter.cpp
@@ -19,6 +19,25 @@ FRigUnit_AnimNextInterfaceParameter_Float_Execute()
 	GetParameterInternal(Parameter, ExecuteContext, &Result);
 }
 
+FRigUnit_AnimNextInterfaceParameter_GraphLODPose_Execute()
+{
+	const UE::AnimNext::Interface::FContext& AnimNextInterfaceContext = ExecuteContext.GetContext();
+
+	Result = AnimNextInterfaceContext.GetParameterChecked<FAnimNextGraphLODPose>(Parameter);
+
+	//GetParameterInternal(Parameter, ExecuteContext, &Result);
+}
+
+FRigUnit_AnimNextInterfaceParameter_AnimSequence_Execute()
+{
+	const UE::AnimNext::Interface::FContext& AnimNextInterfaceContext = ExecuteContext.GetContext();
+
+	Result = AnimNextInterfaceContext.GetParameterChecked<FAnimNextGraph_AnimSequence>(Parameter);
+
+	//GetParameterInternal(Parameter, ExecuteContext, &Result);
+}
+
+
 FRigUnit_AnimNextInterfaceParameter_AnimNextInterface_Execute()
 {
 	GetParameterInternal(Parameter, ExecuteContext, &Result);
@@ -63,7 +82,8 @@ FRigUnit_PoseOperator_Execute()
 {}
 
 FRigUnit_AnimNextInterface_SequencePlayer_Execute()
-{}
+{
+}
 
 struct FSpringDamperState
 {
diff --git a/Engine/Plugins/Experimental/AnimNext/Source/AnimNextInterfaceGraph/Private/RigUnit_AnimNextInterfaceParameter.h b/Engine/Plugins/Experimental/AnimNext/Source/AnimNextInterfaceGraph/Private/RigUnit_AnimNextInterfaceParameter.h
--- a/Engine/Plugins/Experimental/AnimNext/Source/AnimNextInterfaceGraph/Private/RigUnit_AnimNextInterfaceParameter.h
+++ b/Engine/Plugins/Experimental/AnimNext/Source/AnimNextInterfaceGraph/Private/RigUnit_AnimNextInterfaceParameter.h
@@ -5,8 +5,12 @@
 #include "AnimNextInterfaceExecuteContext.h"
 #include "IAnimNextInterface.h"
 #include "Units/RigUnit.h"
+#include "AnimNextInterface_LODPose.h"
+#include "RigUnit_AnimNextAnimSequence.h"
+// --- ---
 #include "RigUnit_AnimNextInterfaceParameter.generated.h"
 
+
 struct FAnimNextInterfaceUnitContext;
 class UAnimNextInterfaceGraph;
 
@@ -37,6 +41,34 @@ protected:
 	float Result = 0.0f;
 };
 
+/** Unit for reading a GraphLODPose parameter from context */
+USTRUCT(meta = (DisplayName = "LODPose Parameter", Category = "Parameters", TitleColor = "1 0 0", NodeColor = "1 1 1"))
+struct FRigUnit_AnimNextInterfaceParameter_GraphLODPose : public FRigUnit_AnimNextInterfaceParameter
+{
+	GENERATED_BODY()
+
+	RIGVM_METHOD()
+	void Execute();
+	
+protected:
+	UPROPERTY(EditAnywhere, Category = "Result", meta = (Output))
+	FAnimNextGraphLODPose Result;
+};
+
+/** Unit for reading an AnimSequence parameter from context */
+USTRUCT(meta = (DisplayName = "AnimSequence Parameter", Category = "Parameters", TitleColor = "1 0 0", NodeColor = "1 1 1"))
+struct FRigUnit_AnimNextInterfaceParameter_AnimSequence : public FRigUnit_AnimNextInterfaceParameter
+{
+	GENERATED_BODY()
+
+	RIGVM_METHOD()
+	void Execute();
+
+protected:
+	UPROPERTY(EditAnywhere, Category = "Result", meta = (Output))
+	FAnimNextGraph_AnimSequence Result;
+};
+
 /** Unit for reading anim interface graph parameter from context */
 USTRUCT(meta = (DisplayName = "Anim Interface Parameter", Category = "Parameters", TitleColor = "1 0 0", NodeColor = "1 1 1"))
 struct FRigUnit_AnimNextInterfaceParameter_AnimNextInterface : public FRigUnit_AnimNextInterfaceParameter
@@ -242,37 +274,22 @@ protected:
 	FAnimNextInterfaceExecuteContext Result;
 };
 
-USTRUCT(BlueprintType)
-struct FAnimSequenceParameters
-{
-	GENERATED_BODY()
-
-	UPROPERTY(EditAnywhere, Category = "Anim Sequence")
-	float PlayRate = 1.0f;
-
-	UPROPERTY(EditAnywhere, Category = "Anim Sequence")
-	float StartPosition = 0.0f;
-
-	UPROPERTY(EditAnywhere, Category = "Anim Sequence")
-	bool bLoop = 0.0f;
-};
-
 /** Unit for getting a pose via an anim interface */
 USTRUCT(meta = (DisplayName = "Anim Sequence", Category = "Animation", TitleColor = "1 0 0", NodeColor = "1 1 1"))
 struct FRigUnit_AnimNextInterface_SequencePlayer : public FRigUnit_AnimNextInterfaceBase
 {
 	GENERATED_BODY()
 
-	RIGVM_METHOD()
-	void Execute();
-	
+		RIGVM_METHOD()
+		void Execute();
+
 protected:
 	UPROPERTY(EditAnywhere, Category = "Anim Interface", meta = (Input))
 	FAnimSequenceParameters Parameters;
 
 	UPROPERTY(EditAnywhere, Category = "Anim Interface", meta = (Input))
 	TScriptInterface<IAnimNextInterface> Sequence = nullptr;
-	
+
 	UPROPERTY(EditAnywhere, Transient, DisplayName = "Result", Category = "BeginExecution", meta = (Output))
 	FAnimNextInterfaceExecuteContext Result;
 };
