Fixed CIS error due to multiple calls to a module init

[REDACTED_HANDLE] thomas.sarkanen

diff --git a/Engine/Plugins/Experimental/AnimNext/Source/AnimNextInterface/Private/AnimationDataRegistry.cpp b/Engine/Plugins/Experimental/AnimNext/Source/AnimNextInterface/Private/AnimationDataRegistry.cpp
--- a/Engine/Plugins/Experimental/AnimNext/Source/AnimNextInterface/Private/AnimationDataRegistry.cpp
+++ b/Engine/Plugins/Experimental/AnimNext/Source/AnimNextInterface/Private/AnimationDataRegistry.cpp
@@ -33,7 +33,7 @@ namespace UE::AnimNext::Interface
 
 /*static*/ void FAnimationDataRegistry::Init()
 {
-	if (ensure(GAnimationDataRegistry == nullptr))
+	if (GAnimationDataRegistry == nullptr)
 	{
 		GAnimationDataRegistry = new FAnimationDataRegistry();
 
@@ -43,7 +43,7 @@ namespace UE::AnimNext::Interface
 
 /*static*/ void FAnimationDataRegistry::Destroy()
 {
-	if (ensure(GAnimationDataRegistry != nullptr))
+	if (GAnimationDataRegistry != nullptr)
 	{
 		FCoreUObjectDelegates::GetPostGarbageCollect().Remove(PostGarbageCollectHandle);
 
