Fixed CIS static analysis warnings

[REDACTED_HANDLE] trivial
[REDACTED_HANDLE] none

diff --git a/Engine/Plugins/Experimental/AnimNext/Source/AnimNextInterface/Private/AnimationGenerationTools.cpp b/Engine/Plugins/Experimental/AnimNext/Source/AnimNextInterface/Private/AnimationGenerationTools.cpp
--- a/Engine/Plugins/Experimental/AnimNext/Source/AnimNextInterface/Private/AnimationGenerationTools.cpp
+++ b/Engine/Plugins/Experimental/AnimNext/Source/AnimNextInterface/Private/AnimationGenerationTools.cpp
@@ -309,7 +309,7 @@ struct FCompareBoneIndexType
 	for (int i = 0; i < NumBones; ++i)
 	{
 		const FBoneIndexType BoneIndex = OrderedBoneList[i];
-		const FBoneIndexType BoneIndexParentIndex = SkeletalMesh->GetRefSkeleton().GetParentIndex(BoneIndex);
+		const int32 BoneIndexParentIndex = SkeletalMesh->GetRefSkeleton().GetParentIndex(BoneIndex);
 
 		if (BoneIndexParentIndex >= 0)
 		{
@@ -471,7 +471,7 @@ namespace AnimNext::Tools::ConsoleCommands
 				}
 				else
 				{
-					UE_LOG(LogAnimGenerationTools, VeryVerbose, TEXT("[%d of %d] SkeletalMesh %s BoneReferencePose error, asset could not be loaded."), Idx + 1, NumAssets, *SkeletalMesh->GetPathName());
+					UE_LOG(LogAnimGenerationTools, VeryVerbose, TEXT("[%d of %d] SkeletalMesh is null. Asset : %s could not be loaded."), Idx + 1, NumAssets, *Assets[Idx].AssetName.ToString());
 				}
 			}
 		}
