Attempt to fix a CIS arror of a non unity build probably due to a missing include

[REDACTED_HANDLE] trivial
[REDACTED_HANDLE] none

diff --git a/Engine/Plugins/Experimental/AnimNext/Source/AnimNextInterfaceGraph/Private/AnimGraph/AnimNode_AnimNextInterfaceGraph.cpp b/Engine/Plugins/Experimental/AnimNext/Source/AnimNextInterfaceGraph/Private/AnimGraph/AnimNode_AnimNextInterfaceGraph.cpp
--- a/Engine/Plugins/Experimental/AnimNext/Source/AnimNextInterfaceGraph/Private/AnimGraph/AnimNode_AnimNextInterfaceGraph.cpp
+++ b/Engine/Plugins/Experimental/AnimNext/Source/AnimNextInterfaceGraph/Private/AnimGraph/AnimNode_AnimNextInterfaceGraph.cpp
@@ -17,6 +17,7 @@
 #include "AnimationReferencePose.h"
 #include "AnimNextInterface_LODPose.h"
 #include "RigUnit_AnimNextAnimSequence.h" // TEST
+#include "Animation/AnimSequence.h"
 #include "Engine/SkeletalMesh.h"
 #include "BoneContainer.h"
 
