AnimNext parameters improvements

Refactor AnimNext to use type system based on StructUtils
Access to parameter data now uses TArrayView<uint8> rather than raw void*
Began namespace reorg (mainly removing redundant Interface namespace)
Cleaned up a bunch of unused prototype code (kernels, batching etc)

diff --git a/Engine/Plugins/Experimental/AnimNext/AnimNext.uplugin b/Engine/Plugins/Experimental/AnimNext/AnimNext.uplugin
--- a/Engine/Plugins/Experimental/AnimNext/AnimNext.uplugin
+++ b/Engine/Plugins/Experimental/AnimNext/AnimNext.uplugin
@@ -19,6 +19,10 @@
 		{
 			"Name": "ControlRig",
 			"Enabled": true
+		},
+		{
+			"Name": "StructUtils",
+			"Enabled": true
 		}
 	],
 	"Modules" :
diff --git a/Engine/Plugins/Experimental/AnimNext/Source/AnimNextInterface/AnimNextInterface.Build.cs b/Engine/Plugins/Experimental/AnimNext/Source/AnimNextInterface/AnimNextInterface.Build.cs
--- a/Engine/Plugins/Experimental/AnimNext/Source/AnimNextInterface/AnimNextInterface.Build.cs
+++ b/Engine/Plugins/Experimental/AnimNext/Source/AnimNextInterface/AnimNextInterface.Build.cs
@@ -11,6 +11,7 @@ namespace UnrealBuildTool.Rules
 				{
 					"Core",
 					"CoreUObject",
+					"StructUtils"
 				}
 			);
 			PrivateDependencyModuleNames.AddRange(
diff --git a/Engine/Plugins/Experimental/AnimNext/Source/AnimNextInterface/Private/AnimNextInterface.cpp b/Engine/Plugins/Experimental/AnimNext/Source/AnimNextInterface/Private/AnimNextInterface.cpp
--- a/Engine/Plugins/Experimental/AnimNext/Source/AnimNextInterface/Private/AnimNextInterface.cpp
+++ b/Engine/Plugins/Experimental/AnimNext/Source/AnimNextInterface/Private/AnimNextInterface.cpp
@@ -4,20 +4,20 @@
 #include "IAnimNextInterface.h"
 #include "AnimNextInterfaceContext.h"
 
-namespace UE::AnimNext::Interface::Private
+namespace UE::AnimNext::Private
 {
 
-static bool CheckCompatibility(const IAnimNextInterface* InAnimNextInterface, const UE::AnimNext::Interface::FContext& InContext)
+static bool CheckCompatibility(const IAnimNextInterface* InAnimNextInterface, const UE::AnimNext::FContext& InContext)
 {
 	check(InAnimNextInterface);
-	return InContext.GetResultParam().GetType().GetName() == InAnimNextInterface->GetReturnTypeName();
+	return InContext.GetResultParam().GetTypeHandle() == InAnimNextInterface->GetReturnTypeHandle();
 }
 
 }
 
-bool IAnimNextInterface::GetDataIfCompatibleInternal(const UE::AnimNext::Interface::FContext& InContext) const
+bool IAnimNextInterface::GetDataIfCompatibleInternal(const UE::AnimNext::FContext& InContext) const
 {
-	if(UE::AnimNext::Interface::Private::CheckCompatibility(this, InContext))
+	if(UE::AnimNext::Private::CheckCompatibility(this, InContext))
 	{
 		return GetDataRawInternal(InContext);
 	}
@@ -25,34 +25,39 @@ bool IAnimNextInterface::GetDataIfCompatibleInternal(const UE::AnimNext::Interfa
 	return false;
 }
 
-bool IAnimNextInterface::GetData(const UE::AnimNext::Interface::FContext& Context) const
+bool IAnimNextInterface::GetData(const UE::AnimNext::FContext& Context) const
 {
 	return GetDataIfCompatibleInternal(Context);
 }
 
-bool IAnimNextInterface::GetDataChecked(const UE::AnimNext::Interface::FContext& Context) const
+bool IAnimNextInterface::GetDataChecked(const UE::AnimNext::FContext& Context) const
 {
-	check(UE::AnimNext::Interface::Private::CheckCompatibility(this, Context));
+	check(UE::AnimNext::Private::CheckCompatibility(this, Context));
 	return GetDataRawInternal(Context);
 }
 
-bool IAnimNextInterface::GetData(const UE::AnimNext::Interface::FContext& Context, UE::AnimNext::Interface::FParam& OutResult) const
+bool IAnimNextInterface::GetData(const UE::AnimNext::FContext& Context, UE::AnimNext::FParam& OutResult) const
 {
-	const UE::AnimNext::Interface::FContext CallingContext = Context.WithResult(OutResult);
+	const UE::AnimNext::FContext CallingContext = Context.WithResult(OutResult);
 	return GetDataIfCompatibleInternal(CallingContext);
 }
 
-bool IAnimNextInterface::GetDataChecked(const UE::AnimNext::Interface::FContext& Context, UE::AnimNext::Interface::FParam& OutResult) const
+bool IAnimNextInterface::GetDataChecked(const UE::AnimNext::FContext& Context, UE::AnimNext::FParam& OutResult) const
 {
-	const UE::AnimNext::Interface::FContext CallingContext = Context.WithResult(OutResult);
-	check(UE::AnimNext::Interface::Private::CheckCompatibility(this, CallingContext));
+	const UE::AnimNext::FContext CallingContext = Context.WithResult(OutResult);
+	check(UE::AnimNext::Private::CheckCompatibility(this, CallingContext));
 	return GetDataRawInternal(CallingContext);
 }
 
-bool IAnimNextInterface::GetDataRawInternal(const UE::AnimNext::Interface::FContext& InContext) const
+bool IAnimNextInterface::GetDataRawInternal(const UE::AnimNext::FContext& InContext) const
 {
 	// TODO: debug recording here on context
 
-	const UE::AnimNext::Interface::FContext CallContext = InContext.WithCallRaw(this);
+	const UE::AnimNext::FContext CallContext = InContext.WithCallRaw(this);
 	return GetDataImpl(CallContext);
 }
+
+UE::AnimNext::FParamTypeHandle IAnimNextInterface::GetReturnTypeHandle() const
+{
+	return GetReturnTypeHandleImpl();
+}
\ No newline at end of file
diff --git a/Engine/Plugins/Experimental/AnimNext/Source/AnimNextInterface/Private/AnimNextInterfaceContext.cpp b/Engine/Plugins/Experimental/AnimNext/Source/AnimNextInterface/Private/AnimNextInterfaceContext.cpp
--- a/Engine/Plugins/Experimental/AnimNext/Source/AnimNextInterface/Private/AnimNextInterfaceContext.cpp
+++ b/Engine/Plugins/Experimental/AnimNext/Source/AnimNextInterface/Private/AnimNextInterfaceContext.cpp
@@ -2,10 +2,11 @@
 
 #include "AnimNextInterfaceContext.h"
 #include "IAnimNextInterface.h"
-#include "IAnimNextInterfaceParameters.h"
-#include "AnimNextInterfaceTypes.h"
+#include "Param/IAnimNextParamInterface.h"
+#include "Param/ParamType.h"
+#include "Param/ParamHelpers.h"
 
-namespace UE::AnimNext::Interface
+namespace UE::AnimNext
 {
 
 struct FThreadContextData : TThreadSingleton<FThreadContextData>
@@ -34,12 +35,11 @@ const FContext& FThreadContext::Get()
 	return *ContextData.ContextStack.Top();
 }
 
-FContext::FContext(float InDeltaTime, FState& InState, FParamStorage& InParamStorage, IAnimNextInterfaceParameters* InParameters)
+FContext::FContext(float InDeltaTime, FState& InState, FParamStorage& InParamStorage, IAnimNextParamInterface* InParameters)
 	: State(&InState)
 	, ParamStorage(&InParamStorage)
 	, Parameters(InParameters)
 	, DeltaTime(InDeltaTime)
-	, BranchingMask(MAX_uint64)
 {
 }
 
@@ -60,7 +60,6 @@ FContext FContext::WithResult(FParam& InResult) const
 	NewContext.Parent = this;
 	NewContext.Root = (Root == nullptr) ? this : Root;
 	NewContext.CallstackHash = CallstackHash;
-	NewContext.BranchingMask = BranchingMask;
 
 	return NewContext;
 }
@@ -71,7 +70,6 @@ FContext FContext::WithParameter(FName ParameterId, const FParam& InParameter) c
 	NewContext.Parent = this;
 	NewContext.Root = (Root == nullptr) ? this : Root;
 	NewContext.CallstackHash = CallstackHash;
-	NewContext.BranchingMask = BranchingMask;
 	NewContext.AddParameters({ TPairInitializer(ParameterId, InParameter) });
 	return NewContext;
 }
@@ -82,7 +80,6 @@ FContext FContext::WithParameters(TArrayView<const TPair<FName, FParam>> InParam
 	NewContext.Parent = this;
 	NewContext.Root = (Root == nullptr) ? this : Root;
 	NewContext.CallstackHash = CallstackHash;
-	NewContext.BranchingMask = BranchingMask;
 	NewContext.AddParameters(InParameters);
 
 	return NewContext;
@@ -94,7 +91,6 @@ FContext FContext::WithResultAndParameters(FParam& InResult, TArrayView<const TP
 	NewContext.Parent = this;
 	NewContext.Root = (Root == nullptr) ? this : Root;
 	NewContext.CallstackHash = CallstackHash;
-	NewContext.BranchingMask = BranchingMask;
 	NewContext.AddParameters(InParameters);
 
 	return NewContext;
@@ -113,13 +109,12 @@ int32 FContext::GetParametersSize(TArrayView<const TPair<FName, FParam>> InParam
 
 		if (EnumHasAnyFlags(SourceParam.GetFlags(), FParam::EFlags::Stored))
 		{
-			const FParamType& ParamType = SourceParam.GetType();
+			const FParamTypeHandle ParamTypeHandle = SourceParam.GetTypeHandle();
 
-			const int32 NumElem = SourceParam.GetNumElements();
-			const int32 ParamAlignment = ParamType.GetAlignment();
-			const int32 ParamSize = ParamType.GetSize();
+			const int32 ParamAlignment = ParamTypeHandle.GetAlignment();
+			const int32 ParamSize = ParamTypeHandle.GetSize();
 
-			const int32 ParamAllocSize = NumElem * Align(ParamSize, ParamAlignment);
+			const int32 ParamAllocSize = Align(ParamSize, ParamAlignment);
 
 			ParamAllocSizes[i] = ParamAllocSize;
 			TotalParamAllocSize += ParamAllocSize;
@@ -160,11 +155,10 @@ void FContext::AddParameters(TArrayView<const TPair<FName, FParam>> InParameters
 
 		if (ParamAllocatedSize > 0 && EnumHasAnyFlags(SourceParam.GetFlags(), FParam::EFlags::Stored))
 		{
-			const FParamType& ParamType = ParamPair.Value.GetType();
+			const FParamTypeHandle ParamTypeHandle = ParamPair.Value.GetTypeHandle();
 
-			const FParamType::ParamCopyFunction ParamCopyFunction = ParamType.GetParamCopyFunction();
-
-			FParam ClonedParam = ParamCopyFunction(SourceParam, TargetMemory, ParamAllocatedSize);
+			TArrayView<uint8> ClonedParamMemory(TargetMemory, ParamTypeHandle.GetSize());
+			FParam ClonedParam = FParam::DuplicateParam(ParamPair.Value, ClonedParamMemory);
 
 			TargetMemory += ParamAllocatedSize;
 
@@ -177,14 +171,13 @@ void FContext::AddParameters(TArrayView<const TPair<FName, FParam>> InParameters
 	}
 }
 
-FContext FContext::WithParameters(IAnimNextInterfaceParameters* InParameters) const
+FContext FContext::WithParameters(IAnimNextParamInterface* InParameters) const
 {
 	FContext NewContext(DeltaTime, *State, *ParamStorage, *Result);
 	NewContext.Parameters = InParameters;
 	NewContext.Parent = this;
 	NewContext.Root = (Root == nullptr) ? this : Root;
 	NewContext.CallstackHash = CallstackHash;
-	NewContext.BranchingMask = BranchingMask;
 
 	return NewContext;
 }
@@ -195,7 +188,6 @@ FContext FContext::CreateSubContext() const
 	NewContext.Parent = this;
 	NewContext.Root = (Root == nullptr) ? this : Root;
 	NewContext.CallstackHash = CallstackHash;
-	NewContext.BranchingMask = BranchingMask;
 
 	return NewContext;
 }
@@ -235,7 +227,7 @@ bool FContext::GetParameter(FName InKey, FParam& OutParam) const
 		else if (!CurrentContext->AdditionalParameterHandles.IsEmpty())
 		{
 			// Find the parameter
-			if (const FHParam* HParam = CurrentContext->AdditionalParameterHandles.Find(InKey))
+			if (const FParamHandle* HParam = CurrentContext->AdditionalParameterHandles.Find(InKey))
 			{
 				const FParam* Param = ParamStorage->GetParam(HParam->ParamHandle);
 				if (Param != nullptr && Param->CanAssignTo(OutParam))
@@ -262,7 +254,6 @@ FContext FContext::WithCallRaw(const IAnimNextInterface* InAnimNextInterface) co
 	NewContext.Parent = this;
 	NewContext.Root = (Root == nullptr) ? this : Root;
 	NewContext.CallstackHash = HashCombineFast(CallstackHash, GetTypeHash(InAnimNextInterface));
-	NewContext.BranchingMask = BranchingMask;
 	return NewContext;
 }
 
@@ -279,4 +270,4 @@ void FContext::FlushRelevancy() const
 	}
 }
 
-}
\ No newline at end of file
+}
diff --git a/Engine/Plugins/Experimental/AnimNext/Source/AnimNextInterface/Private/AnimNextInterfaceModule.cpp b/Engine/Plugins/Experimental/AnimNext/Source/AnimNextInterface/Private/AnimNextInterfaceModule.cpp
--- a/Engine/Plugins/Experimental/AnimNext/Source/AnimNextInterface/Private/AnimNextInterfaceModule.cpp
+++ b/Engine/Plugins/Experimental/AnimNext/Source/AnimNextInterface/Private/AnimNextInterfaceModule.cpp
@@ -4,10 +4,9 @@
 #include "Modules/ModuleInterface.h"
 #include "Modules/ModuleManager.h"
 #include "Misc/CoreDelegates.h"
-#include "AnimNextInterfaceParam.h"
 #include "AnimationDataRegistry.h"
 
-namespace UE::AnimNext::Interface
+namespace UE::AnimNext
 {
 
 class FModule : public IModuleInterface
@@ -17,14 +16,12 @@ public:
 	{
 		FCoreDelegates::OnFEngineLoopInitComplete.AddLambda([]()
 		{
-			FParamType::FRegistrar::RegisterDeferredTypes();
-
-			UE::AnimNext::Interface::FAnimationDataRegistry::Init();
+			UE::AnimNext::FAnimationDataRegistry::Init();
 		});
 
 		FCoreDelegates::OnEnginePreExit.AddLambda([]()
 		{
-			UE::AnimNext::Interface::FAnimationDataRegistry::Destroy();
+			UE::AnimNext::FAnimationDataRegistry::Destroy();
 		});
 		
 	}
@@ -32,4 +29,4 @@ public:
 
 }
 
-IMPLEMENT_MODULE(UE::AnimNext::Interface::FModule, AnimNextInterface)
\ No newline at end of file
+IMPLEMENT_MODULE(UE::AnimNext::FModule, AnimNext)
\ No newline at end of file
diff --git a/Engine/Plugins/Experimental/AnimNext/Source/AnimNextInterface/Private/AnimNextInterfaceParam.cpp b/Engine/Plugins/Experimental/AnimNext/Source/AnimNextInterface/Private/AnimNextInterfaceParam.cpp
deleted file mode 100644
diff --git a/Engine/Plugins/Experimental/AnimNext/Source/AnimNextInterface/Private/AnimNextInterfaceParamStorage.cpp b/Engine/Plugins/Experimental/AnimNext/Source/AnimNextInterface/Private/AnimNextInterfaceParamStorage.cpp
deleted file mode 100644
diff --git a/Engine/Plugins/Experimental/AnimNext/Source/AnimNextInterface/Private/AnimNextInterfaceState.cpp b/Engine/Plugins/Experimental/AnimNext/Source/AnimNextInterface/Private/AnimNextInterfaceState.cpp
--- a/Engine/Plugins/Experimental/AnimNext/Source/AnimNextInterface/Private/AnimNextInterfaceState.cpp
+++ b/Engine/Plugins/Experimental/AnimNext/Source/AnimNextInterface/Private/AnimNextInterfaceState.cpp
@@ -3,13 +3,11 @@
 #include "AnimNextInterfaceState.h"
 #include "AnimNextInterfaceContext.h"
 
-namespace UE::AnimNext::Interface
+namespace UE::AnimNext
 {
 
 FParam* FState::FindStateRaw(const FInterfaceKeyWithIdAndStack& InKey, const FContext& InContext, EStatePersistence InPersistence)
 {
-	check(NumElements == InContext.GetNum())
-
 	FParam* ExistingValue = nullptr;
 	TMap<FInterfaceKeyWithIdAndStack, FParam>* ValueMap = nullptr; 
 
@@ -34,18 +32,18 @@ FParam* FState::FindStateRaw(const FInterfaceKeyWithIdAndStack& InKey, const FCo
 	return ExistingValue;
 }
 
-FParam* FState::AllocateState(const FInterfaceKeyWithIdAndStack& InKey, const FContext& InContext, const FParamType& InType, EStatePersistence InPersistence)
+FParam* FState::AllocateState(const FInterfaceKeyWithIdAndStack& InKey, const FContext& InContext, const FParamTypeHandle& InTypeHandle, EStatePersistence InPersistence)
 {
-	void* Data = FMemory::Malloc(InType.GetSize() * NumElements, InType.GetAlignment());
-	
+	const TArrayView<uint8> Data(static_cast<uint8*>(FMemory::Malloc(InTypeHandle.GetSize(), InTypeHandle.GetAlignment())), InTypeHandle.GetSize());
+
 	switch(InPersistence)
 	{
 	case EStatePersistence::Relevancy:
 		// TODO: chunked allocator for relevancy-based stuff?
-		return &RelevancyValueMap.Add(InKey, FRelevancyParam(InType, Data, InContext.GetBatchedFlags(), InContext.UpdateCounter));
+		return &RelevancyValueMap.Add(InKey, FRelevancyParam(InTypeHandle, Data, FParam::EFlags::Mutable, InContext.UpdateCounter));
 		break;
 	case EStatePersistence::Permanent:
-		return &PermanentValueMap.Add(InKey, FParam(InType, Data, InContext.GetBatchedFlags()));
+		return &PermanentValueMap.Add(InKey, FParam(InTypeHandle, Data, FParam::EFlags::Mutable));
 		break;
 	default:
 		check(false);
diff --git a/Engine/Plugins/Experimental/AnimNext/Source/AnimNextInterface/Private/AnimNextInterfaceTypes.cpp b/Engine/Plugins/Experimental/AnimNext/Source/AnimNextInterface/Private/AnimNextInterfaceTypes.cpp
deleted file mode 100644
diff --git a/Engine/Plugins/Experimental/AnimNext/Source/AnimNextInterface/Private/AnimationDataRegistry.cpp b/Engine/Plugins/Experimental/AnimNext/Source/AnimNextInterface/Private/AnimationDataRegistry.cpp
--- a/Engine/Plugins/Experimental/AnimNext/Source/AnimNextInterface/Private/AnimationDataRegistry.cpp
+++ b/Engine/Plugins/Experimental/AnimNext/Source/AnimNextInterface/Private/AnimationDataRegistry.cpp
@@ -22,13 +22,13 @@
 namespace // Private
 {
 
-UE::AnimNext::Interface::FAnimationDataRegistry* GAnimationDataRegistry = nullptr;
+UE::AnimNext::FAnimationDataRegistry* GAnimationDataRegistry = nullptr;
 constexpr int32 BASIC_TYPE_ALLOCK_BLOCK = 1000;
 FDelegateHandle PostGarbageCollectHandle;
 
 } // end namespace
 
-namespace UE::AnimNext::Interface
+namespace UE::AnimNext
 {
 
 /*static*/ void FAnimationDataRegistry::Init()
@@ -84,9 +84,9 @@ FAnimationDataRegistry* FAnimationDataRegistry::Get()
 
 FAnimationDataHandle FAnimationDataRegistry::RegisterReferencePose(USkeletalMeshComponent* SkeletalMeshComponent)
 {
-	FAnimationDataHandle Handle = AllocateData<FAnimationReferencePose>(1);
+	FAnimationDataHandle Handle = AllocateData<FAnimNextReferencePose>(1);
 
-	FAnimationReferencePose& AnimationReferencePose = Handle.GetRef<FAnimationReferencePose>();
+	FAnimNextReferencePose& AnimationReferencePose = Handle.GetRef<FAnimNextReferencePose>();
 
 	FGenerationTools::GenerateReferencePose(SkeletalMeshComponent, SkeletalMeshComponent->GetSkeletalMeshAsset(), AnimationReferencePose);
 
@@ -182,10 +182,10 @@ void FAnimationDataRegistry::FreeAllocatedBlock(Private::FAllocatedBlock* Alloca
 		{
 			void* Memory = AllocatedBlock->Memory;
 
-			FDataTypeDef* TypeDef = DataTypeDefs.Find(AllocatedBlock->TypeId);
+			FDataTypeDef* TypeDef = DataTypeDefs.Find(AllocatedBlock->TypeHandle);
 			if (ensure(TypeDef != nullptr))
 			{
-				TypeDef->DestrooyTypeFn((uint8*)AllocatedBlock->Memory, AllocatedBlock->NumElem);
+				TypeDef->DestroyTypeFn((uint8*)AllocatedBlock->Memory, AllocatedBlock->NumElem);
 
 				FMemory::Free(AllocatedBlock->Memory); // TODO : This should come from preallocated chunks, use malloc / free for now
 				AllocatedBlock->Memory = nullptr;
@@ -216,4 +216,4 @@ void FAnimationDataRegistry::ReleaseReferencePoseData()
 	SkeletalMeshReferencePoses.Empty();
 }
 
-} // end namespace UE::AnimNext::Interface
+} // end namespace UE::AnimNext
diff --git a/Engine/Plugins/Experimental/AnimNext/Source/AnimNextInterface/Private/AnimationDataRegistryTypes.cpp b/Engine/Plugins/Experimental/AnimNext/Source/AnimNextInterface/Private/AnimationDataRegistryTypes.cpp
--- a/Engine/Plugins/Experimental/AnimNext/Source/AnimNextInterface/Private/AnimationDataRegistryTypes.cpp
+++ b/Engine/Plugins/Experimental/AnimNext/Source/AnimNextInterface/Private/AnimationDataRegistryTypes.cpp
@@ -2,9 +2,8 @@
 
 #include "AnimationDataRegistryTypes.h"
 #include "AnimationDataRegistry.h"
-#include "AnimationReferencePose.h"
 
-namespace UE::AnimNext::Interface
+namespace UE::AnimNext
 {
 
 FAnimationDataHandle::~FAnimationDataHandle()
@@ -23,4 +22,4 @@ FAnimationDataHandle::~FAnimationDataHandle()
 	}
 }
 
-} // namespace UE::AnimNext::Interface
+} // namespace UE::AnimNext
diff --git a/Engine/Plugins/Experimental/AnimNext/Source/AnimNextInterface/Private/AnimationDecompressionTools.cpp b/Engine/Plugins/Experimental/AnimNext/Source/AnimNextInterface/Private/AnimationDecompressionTools.cpp
--- a/Engine/Plugins/Experimental/AnimNext/Source/AnimNextInterface/Private/AnimationDecompressionTools.cpp
+++ b/Engine/Plugins/Experimental/AnimNext/Source/AnimNextInterface/Private/AnimationDecompressionTools.cpp
@@ -22,7 +22,7 @@ DECLARE_CYCLE_STAT(TEXT("Extract Pose From Anim Data"), STAT_ExtractPoseFromAnim
 DECLARE_CYCLE_STAT(TEXT("AnimSeq GetBonePose"), STAT_AnimSeq_GetBonePose, STATGROUP_Anim);
 DECLARE_CYCLE_STAT(TEXT("AnimSeq EvalCurveData"), STAT_AnimSeq_EvalCurveData, STATGROUP_Anim);
 
-namespace UE::AnimNext::Interface
+namespace UE::AnimNext
 {
 
 static IConsoleVariable* CVarForceEvalRawData = IConsoleManager::Get().FindConsoleVariable(TEXT("a.ForceEvalRawData"));
@@ -660,4 +660,4 @@ void FDecompressionTools::RetargetBoneTransform(const FAnimationReferencePose& R
 	}
 }
 
-} // end namespace UE::AnimNext::Interface
+} // end namespace UE::AnimNext
diff --git a/Engine/Plugins/Experimental/AnimNext/Source/AnimNextInterface/Private/AnimationGenerationTools.cpp b/Engine/Plugins/Experimental/AnimNext/Source/AnimNextInterface/Private/AnimationGenerationTools.cpp
--- a/Engine/Plugins/Experimental/AnimNext/Source/AnimNextInterface/Private/AnimationGenerationTools.cpp
+++ b/Engine/Plugins/Experimental/AnimNext/Source/AnimNextInterface/Private/AnimationGenerationTools.cpp
@@ -23,7 +23,7 @@
 DEFINE_LOG_CATEGORY_STATIC(LogAnimGenerationTools, Log, All)
 
 
-namespace UE::AnimNext::Interface
+namespace UE::AnimNext
 {
 
 struct FCompareBoneIndexType
@@ -39,7 +39,7 @@ struct FCompareBoneIndexType
 	, const USkeletalMesh* SkeletalMesh
 	, FAnimationReferencePose& OutAnimationReferencePose)
 {
-	using namespace UE::AnimNext::Interface;
+	using namespace UE::AnimNext;
 
 	bool ReferencePoseGenerated = false;
 
@@ -104,7 +104,7 @@ struct FCompareBoneIndexType
 
 			if (bCanGenerateSingleBonesList)
 			{
-				OutAnimationReferencePose.GenerationFlags = FAnimationReferencePose::EGenerationFlags::FastPath;
+				OutAnimationReferencePose.GenerationFlags = EReferencePoseGenerationFlags::FastPath;
 				OutAnimationReferencePose.SkeletalMesh = SkeletalMesh;
 				OutAnimationReferencePose.Skeleton = SkeletalMesh->GetSkeleton();
 
@@ -431,7 +431,7 @@ struct FCompareBoneIndexType
 	}
 }
 
-} // end namespace UE::AnimNext::Interface
+} // end namespace UE::AnimNext
 
 
 
@@ -443,14 +443,14 @@ namespace AnimNext::Tools::ConsoleCommands
 		{
 			TArray<FBoneIndexType> ExcludedBones;
 
-			UE::AnimNext::Interface::FGenerationTools::DifferenceBoneIndexArrays(LODRequiredBones, NextLODRequiredBones, ExcludedBones);
+			UE::AnimNext::FGenerationTools::DifferenceBoneIndexArrays(LODRequiredBones, NextLODRequiredBones, ExcludedBones);
 
 			return ExcludedBones;
 		}
 
 		static void CheckSkeletalMeshesLODs()
 		{
-			using namespace UE::AnimNext::Interface;
+			using namespace UE::AnimNext;
 
 			TArray<FAssetData> Assets;
 			IAssetRegistry& AssetRegistry = FModuleManager::LoadModuleChecked<FAssetRegistryModule>(TEXT("AssetRegistry")).Get();
@@ -462,7 +462,7 @@ namespace AnimNext::Tools::ConsoleCommands
 				const USkeletalMesh* SkeletalMesh = Cast<USkeletalMesh>(Assets[Idx].GetAsset());
 				if (SkeletalMesh != nullptr)
 				{
-					UE::AnimNext::Interface::FAnimationReferencePose OutAnimationReferencePose;
+					FAnimationReferencePose OutAnimationReferencePose;
 
 					if (FGenerationTools::GenerateReferencePose(nullptr, SkeletalMesh, OutAnimationReferencePose))
 					{
diff --git a/Engine/Plugins/Experimental/AnimNext/Source/AnimNextInterface/Private/AnimationReferencePose.cpp b/Engine/Plugins/Experimental/AnimNext/Source/AnimNextInterface/Private/AnimationReferencePose.cpp
--- a/Engine/Plugins/Experimental/AnimNext/Source/AnimNextInterface/Private/AnimationReferencePose.cpp
+++ b/Engine/Plugins/Experimental/AnimNext/Source/AnimNextInterface/Private/AnimationReferencePose.cpp
@@ -1,15 +1,10 @@
 // Copyright Epic Games, Inc. All Rights Reserved.
 
 #include "AnimationReferencePose.h"
-#include "AnimNextInterfaceParam.h"
 
-namespace UE::AnimNext::Interface
+namespace UE::AnimNext
 {
 
-IMPLEMENT_ANIM_NEXT_INTERFACE_PARAM_TYPE(FAnimationReferencePose, AnimationReferencePose);
-//IMPLEMENT_ANIM_NEXT_INTERFACE_PARAM_TYPE(FAnimationLODPose, AnimationLODPose);
-//IMPLEMENT_ANIM_NEXT_INTERFACE_PARAM_TYPE(FAnimationLODPoseStack, AnimationLODPoseStack);
-
 const FAnimTransform TransformAdditiveIdentity = FAnimTransform(FQuat::Identity, FVector::ZeroVector, FVector::ZeroVector);
 
-} // namespace UE::AnimNext::Interface
+} // namespace UE::AnimNext
diff --git a/Engine/Plugins/Experimental/AnimNext/Source/AnimNextInterface/Private/Param/Param.cpp b/Engine/Plugins/Experimental/AnimNext/Source/AnimNextInterface/Private/Param/Param.cpp
new file mode 100644
--- /dev/null
+++ b/Engine/Plugins/Experimental/AnimNext/Source/AnimNextInterface/Private/Param/Param.cpp
@@ -0,0 +1,118 @@
+﻿// Copyright Epic Games, Inc. All Rights Reserved.
+
+#include "Param/Param.h"
+#include "Param/ParamStorage.h"
+#include "Param/ParamHelpers.h"
+
+namespace UE::AnimNext
+{
+
+FParam::FParam(const FParam* InOtherParam)
+	: Data(EnumHasAnyFlags(InOtherParam->Flags, EFlags::Embedded) ? (void*)&InOtherParam->Data : InOtherParam->Data)
+	, TypeHandle(InOtherParam->TypeHandle)
+	, Size(InOtherParam->Size)
+	, Flags(InOtherParam->Flags)
+{
+	EnumRemoveFlags(Flags, EFlags::Embedded); // For compatibility reasons, I can not keep the embedded value if we create a copy of the FParam
+
+	check(TypeHandle.IsValid());
+	check(Data);
+	check(Size);
+}
+
+FParam::FParam(const FParamTypeHandle& InTypeHandle, TArrayView<uint8> InData, EFlags InFlags)
+	: Data(InData.GetData())
+	, TypeHandle(InTypeHandle)
+	, Size(InData.Num())
+	, Flags(InFlags)
+{
+	check(TypeHandle.IsValid());
+	check(Data);
+	check(InData.Num() > 0 && InData.Num() < 0xffff);
+}
+
+FParam::FParam(const FParamTypeHandle& InTypeHandle, TConstArrayView<uint8> InData, EFlags InFlags)
+	: Data(const_cast<uint8*>(InData.GetData()))
+	, TypeHandle(InTypeHandle)
+	, Size(InData.Num())
+	, Flags(InFlags)
+{
+	check(TypeHandle.IsValid());
+	check(Data);
+	check(InData.Num() > 0 && InData.Num() < 0xffff);
+}
+
+FParam::FParam(const FParamTypeHandle& InTypeHandle, EFlags InFlags)
+	: Data(nullptr)
+	, TypeHandle(InTypeHandle)
+	, Size(0)
+	, Flags(InFlags)
+{
+	check(TypeHandle.IsValid());
+}
+
+bool FParam::CanAssignTo(const FParam& InParam) const
+{
+	return CanAssignWith(InParam.GetTypeHandle(), InParam.Flags);
+}
+
+bool FParam::CanAssignWith(const FParamTypeHandle& InTypeHandle, EFlags InFlags, FStringBuilderBase* OutReasonPtr) const
+{
+	// Check type
+	if(TypeHandle != InTypeHandle)
+	{
+		if(OutReasonPtr)
+		{
+			OutReasonPtr->Appendf(TEXT("Types do not match: %s and %s"), *TypeHandle.ToString(), *InTypeHandle.ToString());
+		}
+		return false;
+	}
+	
+	// Check mutability - we cannot return a mutable version of an immutable param
+	if(!EnumHasAnyFlags(Flags, EFlags::Mutable) && EnumHasAnyFlags(InFlags, EFlags::Mutable))
+	{
+		if(OutReasonPtr)
+		{
+			OutReasonPtr->Append(TEXT("Cannot assign to an immutable parameter"));
+		}
+		return false;
+	}
+
+	return true;
+}
+
+FParam FParam::DuplicateParam(const FParam& InSource, TArrayView<uint8> InTargetMemory)
+{
+	const FParamTypeHandle ParamType = InSource.GetTypeHandle();
+
+	FParamHelpers::Copy(ParamType, ParamType, InSource.GetMutableData(), InTargetMemory);
+
+	return FParam(ParamType, InTargetMemory, InSource.GetFlags());
+}
+
+FParamHandle::FParamHandle(FParamStorage* InOwnerStorage, FInternalHandle InParamHandle)
+	: OwnerStorage(InOwnerStorage)
+	, ParamHandle(InParamHandle)
+{
+}
+
+FParamHandle::FParamHandle(const FParamHandle& Other)
+	: OwnerStorage(Other.OwnerStorage)
+	, ParamHandle(Other.ParamHandle)
+{
+	if (ParamHandle != InvalidParamHandle)
+	{
+		OwnerStorage->IncRefCount(ParamHandle);
+	}
+}
+
+FParamHandle::~FParamHandle()
+{
+	if (ParamHandle != InvalidParamHandle)
+	{
+		check(OwnerStorage != nullptr);
+		OwnerStorage->DecRefCount(ParamHandle);
+	}
+}
+
+} // end namespace UE::AnimNext
diff --git a/Engine/Plugins/Experimental/AnimNext/Source/AnimNextInterface/Private/Param/ParamHelpers.cpp b/Engine/Plugins/Experimental/AnimNext/Source/AnimNextInterface/Private/Param/ParamHelpers.cpp
new file mode 100644
--- /dev/null
+++ b/Engine/Plugins/Experimental/AnimNext/Source/AnimNextInterface/Private/Param/ParamHelpers.cpp
@@ -0,0 +1,310 @@
+﻿// Copyright Epic Games, Inc. All Rights Reserved.
+
+#include "Param/ParamHelpers.h"
+#include "Param/ParamType.h"
+#include "Param/ParamTypeHandle.h"
+#include "Templates/SubclassOf.h"
+
+namespace UE::AnimNext
+{
+
+FParamHelpers::ECopyResult FParamHelpers::Copy(const FAnimNextParamType& InSourceType, const FAnimNextParamType& InTargetType, TConstArrayView<uint8> InSourceMemory, TArrayView<uint8> InTargetMemory)
+{
+	if(InSourceType == InTargetType)
+	{
+		switch(InSourceType.ContainerType)
+		{
+		case FAnimNextParamType::EContainerType::Array:
+			{
+				const TScriptArray<FHeapAllocator>* SourceArray = reinterpret_cast<const TScriptArray<FHeapAllocator>*>(InSourceMemory.GetData());
+				TScriptArray<FHeapAllocator>* TargetArray = reinterpret_cast<TScriptArray<FHeapAllocator>*>(InTargetMemory.GetData());
+
+				const int32 NumElements = SourceArray->Num(); 
+				const size_t ValueTypeSize = InSourceType.GetValueTypeSize();
+				const size_t ValueTypeAlignment = InSourceType.GetValueTypeSize();
+
+				// Reallocate target array
+				TargetArray->SetNumUninitialized(NumElements, ValueTypeSize, ValueTypeAlignment);
+				check(TargetArray->GetAllocatedSize(ValueTypeSize) == SourceArray->GetAllocatedSize(ValueTypeSize));
+
+				// Perform the copy according to value type
+				switch (InSourceType.ValueType)
+				{
+				case FAnimNextParamType::EValueType::None:
+					ensureMsgf(false, TEXT("Trying to copy parameter of type None"));
+					break;
+				case FAnimNextParamType::EValueType::Bool:
+				case FAnimNextParamType::EValueType::Byte:
+				case FAnimNextParamType::EValueType::Int32:
+				case FAnimNextParamType::EValueType::Int64:
+				case FAnimNextParamType::EValueType::Float:
+				case FAnimNextParamType::EValueType::Double:
+				case FAnimNextParamType::EValueType::Name:
+				case FAnimNextParamType::EValueType::Enum:
+					{
+						FMemory::Memcpy(TargetArray->GetData(), SourceArray->GetData(), SourceArray->GetAllocatedSize(ValueTypeSize));
+					}
+					break;
+				case FAnimNextParamType::EValueType::String:
+					{
+						const FString* SourceString = static_cast<const FString*>(SourceArray->GetData());
+						FString* TargetString = static_cast<FString*>(TargetArray->GetData());
+						for(int32 Index = 0; Index < NumElements; ++Index)
+						{
+							*TargetString = *SourceString;
+							++TargetString;
+							++SourceString;
+						}
+					}
+					break;
+				case FAnimNextParamType::EValueType::Text:
+					{
+						const FText* SourceText = static_cast<const FText*>(SourceArray->GetData());
+						FText* TargetText = static_cast<FText*>(TargetArray->GetData());
+						for(int32 Index = 0; Index < NumElements; ++Index)
+						{
+							*TargetText = *SourceText;
+							++TargetText;
+							++SourceText;
+						}
+					}
+					break;
+				case FAnimNextParamType::EValueType::Struct:
+					if(UScriptStruct* ScriptStruct = Cast<UScriptStruct>(InSourceType.ValueTypeObject.Get()))
+					{
+						ScriptStruct->CopyScriptStruct(TargetArray->GetData(), SourceArray->GetData(), NumElements);
+					}
+					else
+					{
+						checkf(false, TEXT("Error: FParameterHelpers::Copy: Unknown Struct Type"));
+					}
+					break;
+				case FAnimNextParamType::EValueType::Object:
+					{
+						const TObjectPtr<UObject>* SourceObjectPtr = reinterpret_cast<const TObjectPtr<UObject>*>(InSourceMemory.GetData());
+						TObjectPtr<UObject>* TargetObjectPtr = reinterpret_cast<TObjectPtr<UObject>*>(InTargetMemory.GetData());
+						for(int32 Index = 0; Index < NumElements; ++Index)
+						{
+							*TargetObjectPtr = *SourceObjectPtr;
+							++TargetObjectPtr;
+							++SourceObjectPtr;
+						}
+					}
+					break;
+				case FAnimNextParamType::EValueType::SoftObject:
+					{
+						const TSoftObjectPtr<UObject>* SourceSoftObjectPtr = reinterpret_cast<const TSoftObjectPtr<UObject>*>(InSourceMemory.GetData());
+						TSoftObjectPtr<UObject>* TargetSoftObjectPtr = reinterpret_cast<TSoftObjectPtr<UObject>*>(InTargetMemory.GetData());
+						for(int32 Index = 0; Index < NumElements; ++Index)
+						{
+							*TargetSoftObjectPtr = *SourceSoftObjectPtr;
+							++TargetSoftObjectPtr;
+							++SourceSoftObjectPtr;
+						}
+					}
+					break;
+				case FAnimNextParamType::EValueType::Class:
+					{
+						const TSubclassOf<UObject>* SourceClassPtr = reinterpret_cast<const TSubclassOf<UObject>*>(InSourceMemory.GetData());
+						TSubclassOf<UObject>* TargetClassPtr = reinterpret_cast<TSubclassOf<UObject>*>(InTargetMemory.GetData());
+						for(int32 Index = 0; Index < NumElements; ++Index)
+						{
+							*TargetClassPtr = *SourceClassPtr;
+							++TargetClassPtr;
+							++SourceClassPtr;
+						}
+					}
+					break;
+				case FAnimNextParamType::EValueType::SoftClass:
+					{
+						const TSoftClassPtr<UObject>* SourceClassPtr = reinterpret_cast<const TSoftClassPtr<UObject>*>(InSourceMemory.GetData());
+						TSoftClassPtr<UObject>* TargetClassPtr = reinterpret_cast<TSoftClassPtr<UObject>*>(InTargetMemory.GetData());
+						for(int32 Index = 0; Index < NumElements; ++Index)
+						{
+							*TargetClassPtr = *SourceClassPtr;
+							++TargetClassPtr;
+							++SourceClassPtr;
+						}
+					}
+					break;
+				default:
+					checkf(false, TEXT("Error: FParameterHelpers::Copy of unknown type"));
+					break;
+				}
+			}
+			break;
+		case FAnimNextParamType::EContainerType::None:
+			{
+				switch (InSourceType.ValueType)
+				{
+				case FAnimNextParamType::EValueType::None:
+					ensureMsgf(false, TEXT("Trying to copy parameter of type None"));
+					break;
+				case FAnimNextParamType::EValueType::Bool:
+				case FAnimNextParamType::EValueType::Byte:
+				case FAnimNextParamType::EValueType::Int32:
+				case FAnimNextParamType::EValueType::Int64:
+				case FAnimNextParamType::EValueType::Float:
+				case FAnimNextParamType::EValueType::Double:
+				case FAnimNextParamType::EValueType::Name:
+				case FAnimNextParamType::EValueType::Enum:
+					{
+						const int32 ParamAlignment = InSourceType.GetAlignment();
+						const int32 ParamSize = InSourceType.GetSize();
+						const int32 ParamAllocSize = Align(ParamSize, ParamAlignment);
+						check(InTargetMemory.Num() <= ParamAllocSize);
+						check(InSourceMemory.Num() <= ParamAllocSize);
+						check(IsAligned(InTargetMemory.GetData(), ParamAlignment));
+						check(IsAligned(InSourceMemory.GetData(), ParamAlignment));
+
+						FMemory::Memcpy(InTargetMemory.GetData(), InSourceMemory.GetData(), ParamAllocSize);
+					}
+					break;
+				case FAnimNextParamType::EValueType::String:
+					{
+						const FString* SourceString = reinterpret_cast<const FString*>(InSourceMemory.GetData());
+						FString* TargetString = reinterpret_cast<FString*>(InTargetMemory.GetData());
+						*TargetString = *SourceString;
+					}
+					break;
+				case FAnimNextParamType::EValueType::Text:
+					{
+						const FText* SourceText = reinterpret_cast<const FText*>(InSourceMemory.GetData());
+						FText* TargetText = reinterpret_cast<FText*>(InTargetMemory.GetData());
+						*TargetText = *SourceText;
+					}
+					break;
+				case FAnimNextParamType::EValueType::Struct:
+					if(UScriptStruct* ScriptStruct = Cast<UScriptStruct>(InSourceType.ValueTypeObject.Get()))
+					{
+						const int32 StructSize = ScriptStruct->GetStructureSize();
+						check(StructSize <= InTargetMemory.Num());
+
+						ScriptStruct->CopyScriptStruct(InTargetMemory.GetData(), InSourceMemory.GetData(), 1);
+					}
+					else
+					{
+						checkf(false, TEXT("Error: FParameterHelpers::Copy: Unknown Struct Type"));
+					}
+					break;
+				case FAnimNextParamType::EValueType::Object:
+					{
+						const TObjectPtr<UObject>* SourceObjectPtr = reinterpret_cast<const TObjectPtr<UObject>*>(InSourceMemory.GetData());
+						TObjectPtr<UObject>* TargetObjectPtr = reinterpret_cast<TObjectPtr<UObject>*>(InTargetMemory.GetData());
+						*TargetObjectPtr = *SourceObjectPtr;
+					}
+					break;
+				case FAnimNextParamType::EValueType::SoftObject:
+					{
+						const TSoftObjectPtr<UObject>* SourceSoftObjectPtr = reinterpret_cast<const TSoftObjectPtr<UObject>*>(InSourceMemory.GetData());
+						TSoftObjectPtr<UObject>* TargetSoftObjectPtr = reinterpret_cast<TSoftObjectPtr<UObject>*>(InTargetMemory.GetData());
+						*TargetSoftObjectPtr = *SourceSoftObjectPtr;
+					}
+					break;
+				case FAnimNextParamType::EValueType::Class:
+					{
+						const TSubclassOf<UObject>* SourceClassPtr = reinterpret_cast<const TSubclassOf<UObject>*>(InSourceMemory.GetData());
+						TSubclassOf<UObject>* TargetClassPtr = reinterpret_cast<TSubclassOf<UObject>*>(InTargetMemory.GetData());
+						*TargetClassPtr = *SourceClassPtr;
+					}
+					break;
+				case FAnimNextParamType::EValueType::SoftClass:
+					{
+						const TSoftClassPtr<UObject>* SourceClassPtr = reinterpret_cast<const TSoftClassPtr<UObject>*>(InSourceMemory.GetData());
+						TSoftClassPtr<UObject>* TargetClassPtr = reinterpret_cast<TSoftClassPtr<UObject>*>(InTargetMemory.GetData());
+						*TargetClassPtr = *SourceClassPtr;
+					}
+					break;
+				default:
+					checkf(false, TEXT("Error: FParameterHelpers::Copy of unknown type"));
+					break;
+				}
+			}
+			break;
+		}
+		
+		return ECopyResult::Succeeded;
+	}
+	else
+	{
+		/** TODO: mismatched/conversions etc. */
+		check(false);
+		return ECopyResult::Failed;
+	}
+}
+
+
+FParamHelpers::ECopyResult FParamHelpers::Copy(const FParamTypeHandle& InSourceTypeHandle, const FParamTypeHandle& InTargetTypeHandle, TConstArrayView<uint8> InSourceMemory, TArrayView<uint8> InTargetMemory)
+{
+	if(InSourceTypeHandle == InTargetTypeHandle)
+	{
+		auto SimpleCopy = [&InTargetMemory, &InSourceMemory](int32 InParamSize, int32 InAlignment)
+		{
+			const int32 ParamAllocSize = Align(InParamSize, InAlignment);
+			check(InTargetMemory.Num() <= ParamAllocSize);
+			check(InSourceMemory.Num() <= ParamAllocSize);
+			check(IsAligned(InTargetMemory.GetData(), InAlignment));
+			check(IsAligned(InSourceMemory.GetData(), InAlignment));
+
+			FMemory::Memcpy(InTargetMemory.GetData(), InSourceMemory.GetData(), ParamAllocSize);
+
+			return ECopyResult::Succeeded;
+		};
+
+		switch(InSourceTypeHandle.GetParameterType())
+		{
+		case FParamTypeHandle::EParamType::None:
+			ensureMsgf(false, TEXT("Trying to copy parameter of type None"));
+			break;
+		case FParamTypeHandle::EParamType::Bool:
+			return SimpleCopy(sizeof(bool), alignof(bool));
+		case FParamTypeHandle::EParamType::Byte:
+			return SimpleCopy(sizeof(uint8), alignof(uint8));
+		case FParamTypeHandle::EParamType::Int32:
+			return SimpleCopy(sizeof(int32), alignof(int32));
+		case FParamTypeHandle::EParamType::Int64:
+			return SimpleCopy(sizeof(int64), alignof(int64));
+		case FParamTypeHandle::EParamType::Float:
+			return SimpleCopy(sizeof(float), alignof(float));
+		case FParamTypeHandle::EParamType::Double:
+			return SimpleCopy(sizeof(double), alignof(double));
+		case FParamTypeHandle::EParamType::Name:
+			return SimpleCopy(sizeof(FName), alignof(FName));
+		case FParamTypeHandle::EParamType::String:
+			{
+				const FString* SourceString = reinterpret_cast<const FString*>(InSourceMemory.GetData());
+				FString* TargetString = reinterpret_cast<FString*>(InTargetMemory.GetData());
+				*TargetString = *SourceString;
+				return ECopyResult::Succeeded;
+			}
+		case FParamTypeHandle::EParamType::Text:
+			{
+				const FText* SourceText = reinterpret_cast<const FText*>(InSourceMemory.GetData());
+				FText* TargetText = reinterpret_cast<FText*>(InTargetMemory.GetData());
+				*TargetText = *SourceText;
+			}
+		case FParamTypeHandle::EParamType::Vector:
+			return SimpleCopy(sizeof(FVector), alignof(FVector));
+		case FParamTypeHandle::EParamType::Vector4:
+			return SimpleCopy(sizeof(FVector4), alignof(FVector4));
+		case FParamTypeHandle::EParamType::Quat:
+			return SimpleCopy(sizeof(FQuat), alignof(FQuat));
+		case FParamTypeHandle::EParamType::Transform:
+			return SimpleCopy(sizeof(FTransform), alignof(FTransform));
+		case FParamTypeHandle::EParamType::Custom:
+			{
+				FAnimNextParamType Type = InSourceTypeHandle.GetType();
+				return Copy(Type, Type, InSourceMemory, InTargetMemory);
+			}
+		default:
+			checkf(false, TEXT("Error: FParameterHelpers::Copy of unknown type"));
+			break;
+		}
+	}
+
+	/** TODO: mismatched/conversions etc. */
+	check(false);
+	return ECopyResult::Failed;
+}
+
+}
\ No newline at end of file
diff --git a/Engine/Plugins/Experimental/AnimNext/Source/AnimNextInterface/Private/Param/ParamStorage.cpp b/Engine/Plugins/Experimental/AnimNext/Source/AnimNextInterface/Private/Param/ParamStorage.cpp
new file mode 100644
--- /dev/null
+++ b/Engine/Plugins/Experimental/AnimNext/Source/AnimNextInterface/Private/Param/ParamStorage.cpp
@@ -0,0 +1,163 @@
+// Copyright Epic Games, Inc. All Rights Reserved.
+
+#include "Param/ParamStorage.h"
+#include "HAL/UnrealMemory.h"
+
+namespace UE::AnimNext
+{
+
+FParamStorage::FParamStorage(int32 InMaxParams, int32 InAllocatedMemorySize, int32 InMaxBlocks, bool InAllowGrowing)
+	: AllowGrowing(InAllowGrowing)
+{
+	check(InMaxParams > 0 && (InMaxBlocks == 0 || (InMaxBlocks > 0 && InAllocatedMemorySize > 0)));
+
+	Parameters.Reserve(InMaxParams);
+	Parameters.AddDefaulted(InMaxParams);
+
+	RefCounts.Reserve(InMaxParams);
+	RefCounts.AddZeroed(InMaxParams);
+
+	ParamMemoryHandles.Reserve(InMaxParams);
+	for (int32 i = 0; i < InMaxParams; i++)
+	{
+		ParamMemoryHandles.Add(InvalidBlockHandle);
+	}
+
+	// --- ---
+
+	if (InAllocatedMemorySize > 0)
+	{
+		RawMemory = (uint8*)FMemory::Malloc(InAllocatedMemorySize);
+		RawMemorySize = InAllocatedMemorySize;
+	}
+
+	if (InMaxBlocks > 0)
+	{
+		BlockOffets.Reserve(InMaxBlocks);
+		BlockOffets.AddZeroed(InMaxBlocks);
+
+		BlockSizes.Reserve(InMaxBlocks);
+		BlockSizes.AddZeroed(InMaxBlocks);
+
+		BlockFreeFlag.Reserve(InMaxBlocks);
+		BlockFreeFlag.Init(true, InMaxBlocks);
+	}
+
+	FreeBlockIndex = 0;
+}
+
+FParamStorage::~FParamStorage()
+{
+	// There should be no blocks of memory or HParams in memory at this point
+	check(FreeBlockIndex == 0);
+	check(ParamFreeIndex == 0);
+
+	FMemory::Free(RawMemory);
+}
+
+FParamStorage::TBlockDataPair FParamStorage::RequestBlock(int32 RequestedBlockSize)
+{
+	FParamStorageHandle BlockHandle = InvalidBlockHandle;
+	uint8* BlockMemory = nullptr;
+
+	check((CurrentBlocksSize + RequestedBlockSize < RawMemorySize) || AllowGrowing);
+	check((FreeBlockIndex >= 0 && FreeBlockIndex < (BlockSizes.Num() - 1)) || AllowGrowing);
+
+	if (CheckBufferSize(RequestedBlockSize) && CheckNumBlocks(FreeBlockIndex))
+	{
+		const int32 BlockIndex = FreeBlockIndex;
+
+		BlockOffets[BlockIndex] = (BlockIndex > 0) ? BlockOffets[BlockIndex - 1] + BlockSizes[BlockIndex - 1] : 0;
+		BlockSizes[BlockIndex] = RequestedBlockSize;
+		CurrentBlocksSize += RequestedBlockSize;
+		BlockFreeFlag[BlockIndex] = false;
+		FreeBlockIndex++;
+
+		BlockHandle = BlockIndex;
+		BlockMemory = RawMemory + BlockOffets[BlockIndex];
+	}
+
+	return TPair<FParamStorageHandle, uint8*>(BlockHandle, BlockMemory);
+}
+
+void FParamStorage::ReleaseBlock(FParamStorageHandle BlockHandle)
+{
+	int32 BlockIndex = BlockHandle;
+
+	check(BlockIndex >= 0 && BlockIndex < FreeBlockIndex);
+
+	// just remove the offset here
+	BlockFreeFlag[BlockIndex] = true;
+	
+	// if it is the last block, release the memory and deal with the prev blocks flagged as free, if any
+	if (BlockIndex == FreeBlockIndex - 1)
+	{
+		// Recover free slots
+		while (BlockIndex >= 0 && BlockFreeFlag[BlockIndex] == true)
+		{
+			BlockOffets[BlockIndex] = 0;
+			CurrentBlocksSize -= BlockSizes[BlockIndex];
+			BlockSizes[BlockIndex] = 0;
+			BlockIndex--;
+			FreeBlockIndex--;
+		}
+	}
+}
+
+uint8* FParamStorage::GetBlockMemory(FParamStorageHandle BlockHandle)
+{
+	uint8* BlockMemory = nullptr;
+
+	const int32 BlockIndex = BlockHandle;
+
+	if (ensure(BlockIndex >= 0 && BlockIndex < FreeBlockIndex))
+	{
+		bool IsSet = BlockFreeFlag[BlockIndex];
+		ensure(IsSet == false);
+		BlockMemory = (IsSet == false) ? RawMemory + BlockOffets[BlockIndex] : nullptr;
+	}
+
+	return BlockMemory;
+}
+
+bool FParamStorage::CheckBufferSize(int32 RequestedBlockSize)
+{
+	bool BufferSizeCorrect = true;
+
+	if (CurrentBlocksSize + RequestedBlockSize > RawMemorySize)
+	{
+		if (AllowGrowing)
+		{
+			// TODO : Grow
+			check(false);
+		}
+		else
+		{
+			BufferSizeCorrect = false;
+		}
+	}
+
+	return BufferSizeCorrect;
+}
+
+bool FParamStorage::CheckNumBlocks(int32 BlockIndex)
+{
+	bool ValidIndex = true;
+
+	if (BlockIndex >= BlockSizes.Num())
+	{
+		if (AllowGrowing)
+		{
+			// TODO : Grow
+		}
+		else
+		{
+			ValidIndex = false;
+		}
+		// TODO : Grow
+	}
+
+	return ValidIndex;
+}
+
+} // end namespace UE::AnimNext
diff --git a/Engine/Plugins/Experimental/AnimNext/Source/AnimNextInterface/Private/Param/ParamType.cpp b/Engine/Plugins/Experimental/AnimNext/Source/AnimNextInterface/Private/Param/ParamType.cpp
new file mode 100644
--- /dev/null
+++ b/Engine/Plugins/Experimental/AnimNext/Source/AnimNextInterface/Private/Param/ParamType.cpp
@@ -0,0 +1,485 @@
+﻿// Copyright Epic Games, Inc. All Rights Reserved.
+
+#include "Param/ParamType.h"
+#include "Param/ParamTypeHandle.h"
+#include "Misc/StringBuilder.h"
+#include "UObject/Class.h"
+#include "Templates/SubclassOf.h"
+
+FAnimNextParamType::FAnimNextParamType(EValueType InValueType, EContainerType InContainerType, UObject* InValueTypeObject)
+	: ValueTypeObject(InValueTypeObject)
+	, ValueType(InValueType)
+	, ContainerType(InContainerType)
+{
+}
+
+bool FAnimNextParamType::IsValidObject() const
+{
+	switch(ValueType)
+	{
+	default:
+	case EValueType::None:
+	case EValueType::Bool:
+	case EValueType::Byte:
+	case EValueType::Int32:
+	case EValueType::Int64:
+	case EValueType::Float:
+	case EValueType::Double:
+	case EValueType::Name:
+	case EValueType::String:
+	case EValueType::Text:
+		return false;
+	case EValueType::Enum:
+		{
+			const UObject* ResolvedObject = ValueTypeObject.Get();
+			return ResolvedObject && ResolvedObject->IsA(UEnum::StaticClass());
+		}
+	case EValueType::Struct:
+		{
+			const UObject* ResolvedObject = ValueTypeObject.Get();
+			return ResolvedObject && ResolvedObject->IsA(UScriptStruct::StaticClass());
+		}
+	case EValueType::Object:
+	case EValueType::SoftObject:
+	case EValueType::Class:
+	case EValueType::SoftClass:
+		{
+			const UObject* ResolvedObject = ValueTypeObject.Get();
+			return ResolvedObject && ResolvedObject->IsA(UClass::StaticClass());
+		}
+	}
+}
+
+UE::AnimNext::FParamTypeHandle FAnimNextParamType::GetHandle() const
+{
+	using namespace UE::AnimNext;
+
+	FParamTypeHandle Handle;
+	
+	switch(ContainerType)
+	{
+	case EPropertyBagContainerType::None:
+		switch(ValueType)
+		{
+		default:
+		case EValueType::None:
+			Handle.SetParameterType(FParamTypeHandle::EParamType::None);
+			break;
+		case EValueType::Bool:
+			Handle.SetParameterType(FParamTypeHandle::EParamType::Bool);
+			break;
+		case EValueType::Byte:
+			Handle.SetParameterType(FParamTypeHandle::EParamType::Byte);
+			break;
+		case EValueType::Int32:
+			Handle.SetParameterType(FParamTypeHandle::EParamType::Int32);
+			break;
+		case EValueType::Int64:
+			Handle.SetParameterType(FParamTypeHandle::EParamType::Int64);
+			break;
+		case EValueType::Float:
+			Handle.SetParameterType(FParamTypeHandle::EParamType::Float);
+			break;
+		case EValueType::Double:
+			Handle.SetParameterType(FParamTypeHandle::EParamType::Double);
+			break;
+		case EValueType::Name:
+			Handle.SetParameterType(FParamTypeHandle::EParamType::Name);
+			break;
+		case EValueType::String:
+			Handle.SetParameterType(FParamTypeHandle::EParamType::String);
+			break;
+		case EValueType::Text:
+			Handle.SetParameterType(FParamTypeHandle::EParamType::Text);
+			break;
+		case EValueType::Struct:
+			if(UScriptStruct* ScriptStruct = Cast<UScriptStruct>(ValueTypeObject.Get()))
+			{
+				if(ScriptStruct == TBaseStructure<FVector>::Get())
+				{
+					Handle.SetParameterType(FParamTypeHandle::EParamType::Vector);
+				}
+				else if(ScriptStruct == TBaseStructure<FVector4>::Get())
+				{
+					Handle.SetParameterType(FParamTypeHandle::EParamType::Vector4);
+				}
+				else if(ScriptStruct == TBaseStructure<FQuat>::Get())
+				{
+					Handle.SetParameterType(FParamTypeHandle::EParamType::Quat);
+				}
+				else if(ScriptStruct == TBaseStructure<FTransform>::Get())
+				{
+					Handle.SetParameterType(FParamTypeHandle::EParamType::Transform);
+				}
+				else
+				{
+					Handle.SetParameterType(FParamTypeHandle::EParamType::Custom);
+					Handle.SetCustomTypeIndex(FParamTypeHandle::GetOrAllocateCustomTypeIndex(EValueType::Struct, EContainerType::None, ScriptStruct));
+				}
+			}
+			else
+			{
+				Handle.SetParameterType(FParamTypeHandle::EParamType::None);
+			}
+			break;
+		case EValueType::Enum:
+			if(UEnum* Enum = Cast<UEnum>(ValueTypeObject.Get()))
+			{
+				Handle.SetParameterType(FParamTypeHandle::EParamType::Custom);
+				Handle.SetCustomTypeIndex(FParamTypeHandle::GetOrAllocateCustomTypeIndex(ValueType, EContainerType::None, Enum));
+			}
+			else
+			{
+				Handle.SetParameterType(FParamTypeHandle::EParamType::None);
+			}
+			break;
+		case EValueType::Object:
+		case EValueType::SoftObject:
+		case EValueType::Class:
+		case EValueType::SoftClass:
+			if(UClass* Class = Cast<UClass>(ValueTypeObject.Get()))
+			{
+				Handle.SetParameterType(FParamTypeHandle::EParamType::Custom);
+				Handle.SetCustomTypeIndex(FParamTypeHandle::GetOrAllocateCustomTypeIndex(ValueType, EContainerType::None, Class));
+			}
+			else
+			{
+				Handle.SetParameterType(FParamTypeHandle::EParamType::None);
+			}
+			break;
+		}
+		break;
+	case EPropertyBagContainerType::Array:
+	default:
+		switch(ValueType)
+		{
+		default:
+		case EValueType::None:
+			Handle.SetParameterType(FParamTypeHandle::EParamType::None);
+			break;
+		case EValueType::Bool:
+		case EValueType::Byte:
+		case EValueType::Int32:
+		case EValueType::Int64:
+		case EValueType::Float:
+		case EValueType::Double:
+		case EValueType::Name:
+		case EValueType::String:
+		case EValueType::Text:
+			Handle.SetParameterType(FParamTypeHandle::EParamType::Custom);
+			Handle.SetCustomTypeIndex(FParamTypeHandle::GetOrAllocateCustomTypeIndex(ValueType, ContainerType, nullptr));
+			break;
+		case EValueType::Struct:
+			if(UScriptStruct* ScriptStruct = Cast<UScriptStruct>(ValueTypeObject.Get()))
+			{
+				Handle.SetParameterType(FParamTypeHandle::EParamType::Custom);
+				Handle.SetCustomTypeIndex(FParamTypeHandle::GetOrAllocateCustomTypeIndex(ValueType, ContainerType, ScriptStruct));
+			}
+			else
+			{
+				Handle.SetParameterType(FParamTypeHandle::EParamType::None);
+			}
+			break;
+		case EValueType::Enum:
+			if(UEnum* Enum = Cast<UEnum>(ValueTypeObject.Get()))
+			{
+				Handle.SetParameterType(FParamTypeHandle::EParamType::Custom);
+				Handle.SetCustomTypeIndex(FParamTypeHandle::GetOrAllocateCustomTypeIndex(ValueType, ContainerType, Enum));
+			}
+			else
+			{
+				Handle.SetParameterType(FParamTypeHandle::EParamType::None);
+			}
+			break;
+		case EValueType::Object:
+		case EValueType::SoftObject:
+		case EValueType::Class:
+		case EValueType::SoftClass:
+			if(UClass* Class = Cast<UClass>(ValueTypeObject.Get()))
+			{
+				Handle.SetParameterType(FParamTypeHandle::EParamType::Custom);
+				Handle.SetCustomTypeIndex(FParamTypeHandle::GetOrAllocateCustomTypeIndex(ValueType, ContainerType, Class));
+			}
+			else
+			{
+				Handle.SetParameterType(FParamTypeHandle::EParamType::None);
+			}
+			break;
+		}
+		break;
+	}
+
+	return Handle;
+}
+
+size_t FAnimNextParamType::GetSize() const
+{
+	switch(ContainerType)
+	{
+	case EContainerType::None:
+		return GetValueTypeSize();
+	case EContainerType::Array:
+		return sizeof(TArray<uint8>);
+	default:
+		break;
+	}
+
+	checkf(false, TEXT("Error: FParameterType::GetSize: Unknown Type Container %d, Value %d"), ContainerType, ValueType);
+	return 0;
+}
+
+size_t FAnimNextParamType::GetValueTypeSize() const
+{
+	switch(ValueType)
+	{
+	case EValueType::None:
+		return 0;
+	case EValueType::Bool:
+		return sizeof(bool);
+	case EValueType::Byte:
+		return sizeof(uint8);
+	case EValueType::Int32:
+		return sizeof(int32);
+	case EValueType::Int64:
+		return sizeof(int64);
+	case EValueType::Float:
+		return sizeof(float);
+	case EValueType::Double:
+		return sizeof(double);
+	case EValueType::Name:
+		return sizeof(FName);
+	case EValueType::String:
+		return sizeof(FString);
+	case EValueType::Text:
+		return sizeof(FText);
+	case EValueType::Enum:
+		// TODO: seems to be no way to find the size of a UEnum?
+		return sizeof(uint8);
+	case EValueType::Struct:
+		if(UScriptStruct* ScriptStruct = Cast<UScriptStruct>(ValueTypeObject.Get()))
+		{
+			return ScriptStruct->GetStructureSize();
+		}
+		else
+		{
+			checkf(false, TEXT("Error: FParameterType::GetValueTypeSize: Unknown Struct Type"));
+			return 0;
+		}
+	case EValueType::Object:
+		return sizeof(TObjectPtr<UObject>);
+	case EValueType::SoftObject:
+		return sizeof(TSoftObjectPtr<UObject>);
+	case EValueType::Class:
+		return sizeof(TSubclassOf<UObject>);
+	case EValueType::SoftClass:
+		return sizeof(TSoftClassPtr<UObject>);
+	default:
+		break;
+	}
+
+	checkf(false, TEXT("Error: FParameterType::GetValueTypeSize: Unknown Type Container %d, Value %d"), ContainerType, ValueType);
+	return 0;
+}
+
+size_t FAnimNextParamType::GetAlignment() const
+{
+	switch(ContainerType)
+	{
+	case EContainerType::None:
+		return GetValueTypeAlignment();
+	case EContainerType::Array:
+		return sizeof(TArray<uint8>);
+	default:
+		break;
+	}
+
+	checkf(false, TEXT("Error: FParameterType::GetAlignment: Unknown Type: Container %d, Value %d"), ContainerType, ValueType);
+	return 0;
+}
+
+size_t FAnimNextParamType::GetValueTypeAlignment() const
+{
+	switch(ValueType)
+	{
+	case EValueType::None:
+		return 0;
+	case EValueType::Bool:
+		return alignof(bool);
+	case EValueType::Byte:
+		return alignof(uint8);
+	case EValueType::Int32:
+		return alignof(int32);
+	case EValueType::Int64:
+		return alignof(int64);
+	case EValueType::Float:
+		return alignof(float);
+	case EValueType::Double:
+		return alignof(double);
+	case EValueType::Name:
+		return alignof(FName);
+	case EValueType::String:
+		return alignof(FString);
+	case EValueType::Text:
+		return alignof(FText);
+	case EValueType::Enum:
+		// TODO: seems to be no way to find the alignment of a UEnum?
+		return alignof(uint8);
+	case EValueType::Struct:
+		if(UScriptStruct* ScriptStruct = Cast<UScriptStruct>(ValueTypeObject.Get()))
+		{
+			return ScriptStruct->GetMinAlignment();
+		}
+		else
+		{
+			checkf(false, TEXT("Error: FParameterType::GetValueTypeAlignment: Unknown Struct Type"));
+			return 0;
+		}
+	case EValueType::Object:
+		return alignof(TObjectPtr<UObject>);
+	case EValueType::SoftObject:
+		return alignof(TSoftObjectPtr<UObject>);
+	case EValueType::Class:
+		return alignof(TSubclassOf<UObject>);
+	case EValueType::SoftClass:
+		return alignof(TSoftClassPtr<UObject>);
+	default:
+		break;
+	}
+
+	checkf(false, TEXT("Error: FParameterType::GetValueTypeAlignment: Unknown Type: Container %d, Value %d"), ContainerType, ValueType);
+	return 0;
+}
+
+FString FAnimNextParamType::ToString() const
+{
+	auto GetTypeString = [this](TStringBuilder<128>& InStringBuilder)
+	{
+		switch(ValueType)
+		{
+		case EValueType::None:
+			InStringBuilder.Append(TEXT("None"));
+			break;
+		case EValueType::Bool:
+			InStringBuilder.Append(TEXT("bool"));
+			break;
+		case EValueType::Byte:
+			InStringBuilder.Append(TEXT("uint8"));
+			break;
+		case EValueType::Int32:
+			InStringBuilder.Append(TEXT("int32"));
+			break;
+		case EValueType::Int64:
+			InStringBuilder.Append(TEXT("int64"));
+			break;
+		case EValueType::Float:
+			InStringBuilder.Append(TEXT("float"));
+			break;
+		case EValueType::Double:
+			InStringBuilder.Append(TEXT("double"));
+			break;
+		case EValueType::Name:
+			InStringBuilder.Append(TEXT("FName"));
+			break;
+		case EValueType::String:
+			InStringBuilder.Append(TEXT("FString"));
+			break;
+		case EValueType::Text:
+			InStringBuilder.Append(TEXT("FText"));
+			break;
+		case EValueType::Enum:
+			if(UEnum* Enum = Cast<UEnum>(ValueTypeObject.Get()))
+			{
+				InStringBuilder.Append(TEXT("U"));
+				InStringBuilder.Append(Enum->GetName());
+			}
+			else
+			{
+				InStringBuilder.Append(TEXT("Error: Unknown Enum"));
+			}
+			break;
+		case EValueType::Struct:
+			if(UScriptStruct* Struct = Cast<UScriptStruct>(ValueTypeObject.Get()))
+			{
+				InStringBuilder.Append(TEXT("F"));
+				InStringBuilder.Append(Struct->GetName());
+			}
+			else
+			{
+				InStringBuilder.Append(TEXT("Error: Unknown Struct"));
+			}
+			break;
+		case EValueType::Object:
+			if(UClass* Class = Cast<UClass>(ValueTypeObject.Get()))
+			{
+				InStringBuilder.Append(TEXT("TObjectPtr<U"));
+				InStringBuilder.Append(Class->GetName());
+				InStringBuilder.Append(TEXT(">"));
+			}
+			else
+			{
+				InStringBuilder.Append(TEXT("Error: TObjectPtr of Unknown Class"));
+			}
+			break;
+		case EValueType::SoftObject:
+			if(UClass* Class = Cast<UClass>(ValueTypeObject.Get()))
+			{
+				InStringBuilder.Append(TEXT("TSoftObjectPtr<U"));
+				InStringBuilder.Append(Class->GetName());
+				InStringBuilder.Append(TEXT(">"));
+			}
+			else
+			{
+				InStringBuilder.Append(TEXT("Error: TSoftObjectPtr of Unknown Class"));
+			}
+			break;
+		case EValueType::Class:
+			if(UClass* Class = Cast<UClass>(ValueTypeObject.Get()))
+			{
+				InStringBuilder.Append(TEXT("TSubClassOf<U"));
+				InStringBuilder.Append(Class->GetName());
+				InStringBuilder.Append(TEXT(">"));
+			}
+			else
+			{
+				InStringBuilder.Append(TEXT("Error: TSubClassOf of Unknown Class"));
+			}
+			break;
+		case EValueType::SoftClass:
+			if(UClass* Class = Cast<UClass>(ValueTypeObject.Get()))
+			{
+				InStringBuilder.Append(TEXT("TSoftClassPtr<U"));
+				InStringBuilder.Append(Class->GetName());
+				InStringBuilder.Append(TEXT(">"));
+			}
+			else
+			{
+				InStringBuilder.Append(TEXT("Error: TSoftClassPtr of Unknown Class"));
+			}
+			break;
+		default:
+			InStringBuilder.Append(TEXT("Error: Unknown value type"));
+			break;
+		}
+	};
+
+	TStringBuilder<128> StringBuilder;
+
+	switch(ContainerType)
+	{
+	case EContainerType::None:
+		GetTypeString(StringBuilder);
+		break;
+	case EContainerType::Array:
+		{
+			StringBuilder.Append(TEXT("TArray<"));
+			GetTypeString(StringBuilder);
+			StringBuilder.Append(TEXT(">"));
+		}
+		break;
+	default:
+		StringBuilder.Append(TEXT("Error: Unknown container type"));
+		break;
+	}
+
+	return StringBuilder.ToString();
+}
diff --git a/Engine/Plugins/Experimental/AnimNext/Source/AnimNextInterface/Private/Param/ParamTypeHandle.cpp b/Engine/Plugins/Experimental/AnimNext/Source/AnimNextInterface/Private/Param/ParamTypeHandle.cpp
new file mode 100644
--- /dev/null
+++ b/Engine/Plugins/Experimental/AnimNext/Source/AnimNextInterface/Private/Param/ParamTypeHandle.cpp
@@ -0,0 +1,235 @@
+// Copyright Epic Games, Inc. All Rights Reserved.
+
+#include "Param/ParamTypeHandle.h"
+#include "Param/ParamType.h"
+#include "Misc/ScopeRWLock.h"
+#include "UObject/Class.h"
+
+namespace UE::AnimNext
+{
+
+// Array of all non built-in types. Index into this array is (CustomTypeIndex - 1)
+static TArray<FAnimNextParamType> GCustomTypes;
+
+// Map of type to index in GNonBuiltInTypes array
+static TMap<FAnimNextParamType, uint32> GTypeToIndexMap;
+
+// RW lock for types array/map
+FRWLock TypesLock;
+
+void FParamTypeHandle::ResetCustomTypes()
+{
+	FRWScopeLock ScopeLock(TypesLock, SLT_Write);
+
+	GCustomTypes.Empty();
+	GTypeToIndexMap.Empty();
+}
+
+uint32 FParamTypeHandle::GetOrAllocateCustomTypeIndex(FAnimNextParamType::EValueType InValueType, FAnimNextParamType::EContainerType InContainerType, UObject* InValueTypeObject)
+{
+	FAnimNextParamType ParameterType;
+	ParameterType.ValueTypeObject = InValueTypeObject;
+	ParameterType.ValueType = InValueType;
+	ParameterType.ContainerType = InContainerType;
+
+	const uint32 Hash = GetTypeHash(ParameterType);
+
+	// NOTE: It is tempting to try to acquire a only read lock here during the FindByHash, but doing so means that the
+	// lazy-insert operation is no longer atomic as there is no way to upgrade a read lock into a write lock with
+	// FRWScopeLock (@see FRWScopeLock comments), hence the write lock around the map/array access here.
+	{
+		// See if the type already exists in the map
+		FRWScopeLock ScopeLock(TypesLock, SLT_Write);
+
+		if(const uint32* IndexPtr = GTypeToIndexMap.FindByHash(Hash, ParameterType))
+		{
+			return *IndexPtr + 1;
+		}
+
+		// Add a new custom type
+		const uint32 Index = GCustomTypes.Add(ParameterType);
+		GTypeToIndexMap.AddByHash(Hash, ParameterType, Index);
+
+		checkf((Index + 1) < (1 << 24), TEXT("FParamTypeHandle::GetCustomTypeIndex: Type index overflowed"));
+		return Index + 1;
+	}
+}
+
+bool FParamTypeHandle::ValidateCustomTypeIndex(uint32 InCustomTypeIndex)
+{
+	FRWScopeLock ScopeLock(TypesLock, SLT_ReadOnly);
+	return GCustomTypes.IsValidIndex(InCustomTypeIndex - 1);
+}
+
+FAnimNextParamType FParamTypeHandle::GetType() const
+{
+	FAnimNextParamType ParameterType;
+
+	switch(GetParameterType())
+	{
+	default:
+	case EParamType::None:
+		break;
+	case EParamType::Bool:
+		ParameterType.ValueType = FAnimNextParamType::EValueType::Bool;
+		ParameterType.ContainerType = FAnimNextParamType::EContainerType::None;
+		break;
+	case EParamType::Byte:
+		ParameterType.ValueType = FAnimNextParamType::EValueType::Byte;
+		ParameterType.ContainerType = FAnimNextParamType::EContainerType::None;
+		break;
+	case EParamType::Int32:
+		ParameterType.ValueType = FAnimNextParamType::EValueType::Int32;
+		ParameterType.ContainerType = FAnimNextParamType::EContainerType::None;
+		break;
+	case EParamType::Int64:
+		ParameterType.ValueType = FAnimNextParamType::EValueType::Int64;
+		ParameterType.ContainerType = FAnimNextParamType::EContainerType::None;
+		break;
+	case EParamType::Float:
+		ParameterType.ValueType = FAnimNextParamType::EValueType::Float;
+		ParameterType.ContainerType = FAnimNextParamType::EContainerType::None;
+		break;
+	case EParamType::Double:
+		ParameterType.ValueType = FAnimNextParamType::EValueType::Double;
+		ParameterType.ContainerType = FAnimNextParamType::EContainerType::None;
+		break;
+	case EParamType::Name:
+		ParameterType.ValueType = FAnimNextParamType::EValueType::Name;
+		ParameterType.ContainerType = FAnimNextParamType::EContainerType::None;
+		break;
+	case EParamType::String:
+		ParameterType.ValueType = FAnimNextParamType::EValueType::String;
+		ParameterType.ContainerType = FAnimNextParamType::EContainerType::None;
+		break;
+	case EParamType::Text:
+		ParameterType.ValueType = FAnimNextParamType::EValueType::Text;
+		ParameterType.ContainerType = FAnimNextParamType::EContainerType::None;
+		break;
+	case EParamType::Vector:
+		ParameterType.ValueTypeObject = TBaseStructure<FVector>::Get();
+		ParameterType.ValueType = FAnimNextParamType::EValueType::Struct;
+		ParameterType.ContainerType = FAnimNextParamType::EContainerType::None;
+		break;
+	case EParamType::Vector4:
+		ParameterType.ValueTypeObject = TBaseStructure<FVector4>::Get();
+		ParameterType.ValueType = FAnimNextParamType::EValueType::Struct;
+		ParameterType.ContainerType = FAnimNextParamType::EContainerType::None;
+		break;
+	case EParamType::Quat:
+		ParameterType.ValueTypeObject = TBaseStructure<FQuat>::Get();
+		ParameterType.ValueType = FAnimNextParamType::EValueType::Struct;
+		ParameterType.ContainerType = FAnimNextParamType::EContainerType::None;
+		break;
+	case EParamType::Transform:
+		ParameterType.ValueTypeObject = TBaseStructure<FTransform>::Get();
+		ParameterType.ValueType = FAnimNextParamType::EValueType::Struct;
+		ParameterType.ContainerType = FAnimNextParamType::EContainerType::None;
+		break;
+	case EParamType::Custom:
+		{
+			FRWScopeLock ScopeLock(TypesLock, SLT_ReadOnly);
+			ParameterType = GCustomTypes[GetCustomTypeIndex() - 1];
+		}
+		break;
+	}
+
+	return ParameterType;
+}
+
+size_t FParamTypeHandle::GetSize() const
+{
+	switch(GetParameterType())
+	{
+	default:
+	case EParamType::None:
+		return 0;
+	case EParamType::Bool:
+		return sizeof(bool);
+	case EParamType::Byte:
+		return sizeof(uint8);
+	case EParamType::Int32:
+		return sizeof(int32);
+	case EParamType::Int64:
+		return sizeof(int64);
+	case EParamType::Float:
+		return sizeof(float);
+	case EParamType::Double:
+		return sizeof(double);
+	case EParamType::Name:
+		return sizeof(FName);
+	case EParamType::String:
+		return sizeof(FString);
+	case EParamType::Text:
+		return sizeof(FText);
+	case EParamType::Vector:
+		return sizeof(FVector);
+	case EParamType::Vector4:
+		return sizeof(FVector4);
+	case EParamType::Quat:
+		return sizeof(FQuat);
+	case EParamType::Transform:
+		return sizeof(FTransform);
+	case EParamType::Custom:
+		return GetType().GetSize();
+	}
+
+	return 0;
+}
+
+size_t FParamTypeHandle::GetValueTypeSize() const
+{
+	return GetSize();
+}
+
+size_t FParamTypeHandle::GetAlignment() const
+{
+	switch(GetParameterType())
+	{
+	default:
+	case EParamType::None:
+		return 0;
+	case EParamType::Bool:
+		return alignof(bool);
+	case EParamType::Byte:
+		return alignof(uint8);
+	case EParamType::Int32:
+		return alignof(int32);
+	case EParamType::Int64:
+		return alignof(int64);
+	case EParamType::Float:
+		return alignof(float);
+	case EParamType::Double:
+		return alignof(double);
+	case EParamType::Name:
+		return alignof(FName);
+	case EParamType::String:
+		return alignof(FString);
+	case EParamType::Text:
+		return alignof(FText);
+	case EParamType::Vector:
+		return alignof(FVector);
+	case EParamType::Vector4:
+		return alignof(FVector4);
+	case EParamType::Quat:
+		return alignof(FQuat);
+	case EParamType::Transform:
+		return alignof(FTransform);
+	case EParamType::Custom:
+		return GetType().GetAlignment();
+	}
+
+	return 0;
+}
+
+size_t FParamTypeHandle::GetValueTypeAlignment() const
+{
+	return GetAlignment();
+}
+
+FString FParamTypeHandle::ToString() const
+{
+	return GetType().ToString();
+}
+
+}
diff --git a/Engine/Plugins/Experimental/AnimNext/Source/AnimNextInterface/Public/AnimNextInterface.h b/Engine/Plugins/Experimental/AnimNext/Source/AnimNextInterface/Public/AnimNextInterface.h
--- a/Engine/Plugins/Experimental/AnimNext/Source/AnimNextInterface/Public/AnimNextInterface.h
+++ b/Engine/Plugins/Experimental/AnimNext/Source/AnimNextInterface/Public/AnimNextInterface.h
@@ -20,38 +20,6 @@ static bool GetDataSafe(const TScriptInterface<InterfaceType> AnimNextInterface,
 	return false;
 }
 
-//// Set the Result into the context and call the interface
-//template<typename InterfaceType, typename ValueType>
-//static bool GetDataSafe(const TScriptInterface<InterfaceType> AnimNextInterface, FContext& InContext, ValueType& Result)
-//{
-//	if (AnimNextInterface.GetInterface() != nullptr)
-//	{
-//		if constexpr (TIsDerivedFrom<ValueType, FParam>::Value)
-//		{
-//			const FContext ResultHolder = InContext.WithResult(Result);
-//			return AnimNextInterface.GetInterface()->GetData(ResultHolder);
-//		}
-//		else
-//		{
-//			// Support containers with Num() and GetData()
-//			if constexpr (TModels_V<Private::CSizedContainerWithAccessibleDataAsRawPtr, ValueType>)
-//			{
-//				TWrapParam<ValueType::ElementType> WrappedResult(Result.GetData(), Result.Num());
-//				const FContext ResultHolder = InContext.WithResult(WrappedResult);
-//				return AnimNextInterface.GetInterface()->GetData(ResultHolder);
-//			}
-//			else
-//			{
-//				TWrapParam<ValueType> WrappedResult(&Result);
-//				const FContext ResultHolder = InContext.WithResult(WrappedResult);
-//				return AnimNextInterface.GetInterface()->GetData(ResultHolder);
-//			}
-//		}
-//	}
-//
-//	return false;
-//}
-
 // This version accepts a const Context, but it crates a sub context in order to be able to pass the Result
 // Using the non const version is a better option if possible
 template<typename InterfaceType, typename ValueType>
diff --git a/Engine/Plugins/Experimental/AnimNext/Source/AnimNextInterface/Public/AnimNextInterfaceContext.h b/Engine/Plugins/Experimental/AnimNext/Source/AnimNextInterface/Public/AnimNextInterfaceContext.h
--- a/Engine/Plugins/Experimental/AnimNext/Source/AnimNextInterface/Public/AnimNextInterfaceContext.h
+++ b/Engine/Plugins/Experimental/AnimNext/Source/AnimNextInterface/Public/AnimNextInterfaceContext.h
@@ -4,18 +4,20 @@
 
 #include "CoreMinimal.h"
 #include "UObject/Class.h"
-#include "AnimNextInterfaceParam.h"
+#include "Param/Param.h"
 #include "AnimNextInterfaceState.h"
-#include "AnimNextInterfaceParamStorage.h"
+#include "Param/ParamStorage.h"
 #include "AnimNextInterfaceKey.h"
+#include "Param/ParamType.h"
 #include "Misc/MemStack.h"
 
 class IAnimNextInterface;
-class IAnimNextInterfaceParameters;
+class IAnimNextParamInterface;
 
-namespace UE::AnimNext::Interface
+namespace UE::AnimNext
 {
 
+struct FContext;
 struct FState;
 enum class EStatePersistence : uint8;
 
@@ -38,10 +40,8 @@ private:
 	friend struct FState;
 
 public:
-	static constexpr int32 MAX_BATCH_ELEMENTS = 64;
-
 	// Root public constructor. Constructs a context given a state.
-	FContext(float InDeltaTime, FState& InState, FParamStorage& InParamStorage, IAnimNextInterfaceParameters* InParameters = nullptr);
+	FContext(float InDeltaTime, FState& InState, FParamStorage& InParamStorage, IAnimNextParamInterface* InParameters = nullptr);
 	
 	FContext(const FContext& other) = delete;
 	FContext& operator=(const FContext&) = delete;
@@ -76,7 +76,7 @@ public:
 	FContext WithResultAndParameters(FParam& InResult, TArrayView<const TPair<FName, FParam>> InParameters) const;
 
 	// Create a sub context from this one that includes the provided interface parameter
-	FContext WithParameters(IAnimNextInterfaceParameters* InParameters) const;
+	FContext WithParameters(IAnimNextParamInterface* InParameters) const;
 
 public:
 
@@ -93,7 +93,7 @@ public:
 
 	// Add an Input parameter by value, copying it to the shared storage
 	template<typename ValueType>
-	FHParam AddInputValue(FName ParameterId, ValueType& Value)
+	FParamHandle AddInputValue(FName ParameterId, ValueType& Value)
 	{
 		// Values are added mutable even if source is const
 		const FParam::EFlags Flags = FParam::EFlags::Value | FParam::EFlags::Mutable;
@@ -103,7 +103,7 @@ public:
 
 	// Add an Input parameter by value, copying it to the shared storage
 	template<typename ValueType>
-	FHParam AddInputValue(FName ParameterId, ValueType&& Value)
+	FParamHandle AddInputValue(FName ParameterId, ValueType&& Value)
 	{
 		// Values are added mutable even if source is const
 		const FParam::EFlags Flags = FParam::EFlags::Value | FParam::EFlags::Mutable;
@@ -113,7 +113,7 @@ public:
 
 	// Add an Input parameter by value, copying it to the shared storage
 	template<typename ValueType>
-	FHParam AddInputValue(FName ParameterId, ValueType* Value)
+	FParamHandle AddInputValue(FName ParameterId, ValueType* Value)
 	{
 		// Values are added mutable even if source is const
 		const FParam::EFlags Flags = FParam::EFlags::Value | FParam::EFlags::Mutable;
@@ -123,7 +123,7 @@ public:
 
 	// Add an Output parameter by value, copying it to the shared storage
 	template<typename ValueType>
-	FHParam AddOutputValue(FName ParameterId, ValueType& Value)
+	FParamHandle AddOutputValue(FName ParameterId, ValueType& Value)
 	{
 		// Values are added mutable even if source is const
 		const FParam::EFlags Flags = FParam::EFlags::Value | FParam::EFlags::Mutable;
@@ -133,7 +133,7 @@ public:
 
 	// Add an Output parameter by value, copying it to the shared storage
 	template<typename ValueType>
-	FHParam AddOutputValue(FName ParameterId, ValueType&& Value)
+	FParamHandle AddOutputValue(FName ParameterId, ValueType&& Value)
 	{
 		// Values are added mutable even if source is const
 		const FParam::EFlags Flags = FParam::EFlags::Value | FParam::EFlags::Mutable;
@@ -143,7 +143,7 @@ public:
 
 	// Add an Output parameter by value, copying it to the shared storage
 	template<typename ValueType>
-	FHParam AddOutputValue(FName ParameterId, ValueType* Value)
+	FParamHandle AddOutputValue(FName ParameterId, ValueType* Value)
 	{
 		// Values are added mutable even if source is const
 		const FParam::EFlags Flags = FParam::EFlags::Value | FParam::EFlags::Mutable;
@@ -153,7 +153,7 @@ public:
 
 	// Add a parameter by value, copying it to the shared storage
 	template<typename ValueType>
-	FHParam AddValue(EParamType InParamType, FName ParameterId, ValueType &Value)
+	FParamHandle AddValue(EParamType InParamType, FName ParameterId, ValueType &Value)
 	{
 		// Values are added mutable even if source is const
 		const FParam::EFlags Flags = FParam::EFlags::Value | FParam::EFlags::Mutable;
@@ -163,7 +163,7 @@ public:
 
 	// Add a parameter by value, copying it to the shared storage
 	template<typename ValueType>
-	FHParam AddValue(EParamType InParamType, FName ParameterId, ValueType&& Value)
+	FParamHandle AddValue(EParamType InParamType, FName ParameterId, ValueType&& Value)
 	{
 		// Values are added mutable even if source is const
 		const FParam::EFlags Flags = FParam::EFlags::Value | FParam::EFlags::Mutable;
@@ -173,7 +173,7 @@ public:
 
 	// Add a parameter by value, copying it to the shared storage
 	template<typename ValueType>
-	FHParam AddValue(EParamType InParamType, FName ParameterId, ValueType* Value)
+	FParamHandle AddValue(EParamType InParamType, FName ParameterId, ValueType* Value)
 	{
 		// Values are added mutable even if source is const
 		const FParam::EFlags Flags = FParam::EFlags::Value | FParam::EFlags::Mutable;
@@ -185,7 +185,7 @@ public:
 
 	// Add an Input parameter by reference, adding just a pointer to the shared storage
 	template<typename ValueType>
-	FHParam AddInputReference(FName ParameterId, ValueType& Value)
+	FParamHandle AddInputReference(FName ParameterId, ValueType& Value)
 	{
 		const FParam::EFlags Flags = FParam::EFlags::Reference;
 
@@ -194,7 +194,7 @@ public:
 
 	// Add an Input parameter by reference, adding just a pointer to the shared storage
 	template<typename ValueType>
-	FHParam AddInputReference(FName ParameterId, ValueType* Value)
+	FParamHandle AddInputReference(FName ParameterId, ValueType* Value)
 	{
 		FParam::EFlags Flags = FParam::EFlags::Reference;
 
@@ -203,7 +203,7 @@ public:
 
 	// Add an Output parameter by reference, adding just a pointer to the shared storage
 	template<typename ValueType>
-	FHParam AddOutputReference(FName ParameterId, ValueType& Value)
+	FParamHandle AddOutputReference(FName ParameterId, ValueType& Value)
 	{
 		const FParam::EFlags Flags = FParam::EFlags::Reference;
 
@@ -212,7 +212,7 @@ public:
 
 	// Add an Output parameter by reference, adding just a pointer to the shared storage
 	template<typename ValueType>
-	FHParam AddOutputReference(FName ParameterId, ValueType* Value)
+	FParamHandle AddOutputReference(FName ParameterId, ValueType* Value)
 	{
 		FParam::EFlags Flags = FParam::EFlags::Reference;
 
@@ -221,7 +221,7 @@ public:
 
 	// Add a parameter by reference, adding just a pointer to the shared storage
 	template<typename ValueType>
-	FHParam AddReference(EParamType InParamType, FName ParameterId, ValueType& Value)
+	FParamHandle AddReference(EParamType InParamType, FName ParameterId, ValueType& Value)
 	{
 		const FParam::EFlags Flags = FParam::EFlags::Reference;
 
@@ -230,7 +230,7 @@ public:
 
 	// Add a parameter by reference, adding just a pointer to the shared storage
 	template<typename ValueType>
-	FHParam AddReference(EParamType InParamType, FName ParameterId, ValueType *Value)
+	FParamHandle AddReference(EParamType InParamType, FName ParameterId, ValueType *Value)
 	{
 		FParam::EFlags Flags = FParam::EFlags::Reference;
 
@@ -238,15 +238,15 @@ public:
 	}
 
 	template<typename ValueType>
-	FHParam AddParameter(EParamType InParamType, FName ParameterId, FParam::EFlags Flags, ValueType* Value)
+	FParamHandle AddParameter(EParamType InParamType, FName ParameterId, FParam::EFlags Flags, ValueType* Value)
 	{
-		FHParam ParamHandle;
+		FParamHandle ParamHandle;
 
 		check(InParamType == EParamType::Input || InParamType == EParamType::Output);
 		check(EnumHasAnyFlags(Flags, FParam::EFlags::Value) || EnumHasAnyFlags(Flags, FParam::EFlags::Reference));
 
-		// check if it is a simple FHParam copy
-		if constexpr (TIsDerivedFrom<ValueType, FHParam>::Value)
+		// check if it is a simple FParamHandle copy
+		if constexpr (TIsDerivedFrom<ValueType, FParamHandle>::Value)
 		{
 			const FParam* ExistingParam = ParamStorage->GetParam(Value->ParamHandle);
 			check(ExistingParam != nullptr);
@@ -269,21 +269,9 @@ public:
 				EnumAddFlags(Flags, FParam::EFlags::Mutable);
 			}
 
-			// deal with containers
-			if constexpr (TModels_V<Private::CSizedContainerWithAccessibleDataAsRawPtr, ValueType>)
-			{
-				MutableValueType* ParamData = const_cast<MutableValueType>(static_cast<ValueType*>(Value->GetData()));
-				ParamHandle = EnumHasAnyFlags(Flags, FParam::EFlags::Value) 
-					? ParamHandle = ParamStorage->AddValue(ParamData, Value->Num(), Flags)
-					: ParamStorage->AddReference(ParamData, Value->Num(), Flags);
-			}
-			// and finally with single params
-			else
-			{
-				ParamHandle = EnumHasAnyFlags(Flags, FParam::EFlags::Value)
-					? ParamHandle = ParamStorage->AddValue(const_cast<MutableValueType*>(Value), 1, Flags)
-					: ParamStorage->AddReference(const_cast<MutableValueType*>(Value), 1, Flags);
-			}
+			ParamHandle = EnumHasAnyFlags(Flags, FParam::EFlags::Value)
+				? ParamStorage->AddValue(const_cast<MutableValueType*>(Value), Flags)
+				: ParamStorage->AddReference(const_cast<MutableValueType*>(Value), Flags);
 		}
 
 		AdditionalParameterHandles.Add(ParameterId, ParamHandle);
@@ -293,7 +281,7 @@ public:
 
 	// Get a TParam from a handle
 	template<typename ValueType>
-	TParam<ValueType> GetParameterChecked(const FHParam &InParamHandle) const
+	TParam<ValueType> GetParameterChecked(const FParamHandle &InParamHandle) const
 	{
 		const FParam* Param = ParamStorage->GetParam(InParamHandle.ParamHandle);
 
@@ -307,12 +295,12 @@ public:
 
 	// Get a parameter from a handle as a specified type
 	template <typename ValueType>
-	ValueType& GetParameterAs(const FHParam& InParamHandle, int32* OutNumElems = nullptr)
+	ValueType& GetParameterAs(const FParamHandle& InParamHandle)
 	{
 		const FParam* Param = ParamStorage->GetParam(InParamHandle.ParamHandle);
 
 		check(Param != nullptr);
-		check(Param->GetType().GetTypeId() == Private::TParamType<ValueType>::GetType().GetTypeId());
+		check(Param->GetTypeHandle() == FParamTypeHandle::GetHandle<ValueType>());
 
 		// If a non cost is requested, check the param has mutable flag
 		if constexpr (TIsConst<ValueType>::Value == false)
@@ -324,11 +312,6 @@ public:
 			? static_cast<ValueType*>((void*)&Param->Data)
 			: static_cast<ValueType*>(Param->Data);
 
-		if (OutNumElems != nullptr)
-		{
-			*OutNumElems = Param->GetNumElements();
-		}
-
 		return *ValueData;
 	}
 
@@ -341,7 +324,7 @@ public:
 	template<typename ValueType>
 	TParam<ValueType> GetParameterChecked(FName InKey) const
 	{
-		TParam<ValueType> RetVal(GetBatchedFlags());
+		TParam<ValueType> RetVal;
 
 		GetParameter(InKey, RetVal);
 
@@ -357,7 +340,7 @@ public:
 	{
 		TOptional<TParam<ValueType>> RetVal;
 
-		TParam<ValueType> Param(FParam::EFlags::Mutable | GetBatchedFlags());
+		TParam<ValueType> Param(FParam::EFlags::Mutable);
 		if (GetParameter(InKey, Param))
 		{
 			checkSlow(Private::CheckParam<ValueType>(Param));
@@ -386,7 +369,7 @@ public:
 	}
 
 	// Set a HParam as a result (for compatibility reasons)
-	void SetHParamAsResult(const FHParam& InHParam)
+	void SetHParamAsResult(const FParamHandle& InHParam)
 	{
 		Result = ParamStorage->GetParam(InHParam.ParamHandle);
 	}
@@ -396,48 +379,35 @@ public:
 	template<typename ValueType>
 	void SetResult(const ValueType& InValue) const
 	{
-		// Support containers with Num() and GetData()
-		if constexpr (TModels_V<Private::CSizedContainerWithAccessibleDataAsRawPtr, ValueType>)
-		{
-			checkSlow(Private::CheckParam<ValueType>(*Result));
-			checkSlow(InValue.Num() == Result->NumElements);
-
-			TArrayView<ValueType> Results(*static_cast<ValueType*>(Result->Data), Result->NumElements);
-			const ValueType* Data = InValue.GetData();
-			for(int32 ResultIndex = 0; ResultIndex < Result->NumElements; ++ResultIndex)
-			{
-				Results[ResultIndex] = Data[ResultIndex];
-			}
-		}
-		else
-		{
-			checkSlow(Private::CheckParam<ValueType>(*Result));
-			checkSlow(Result->NumElements == 1);
-
-			TArrayView<ValueType> Results(static_cast<ValueType*>(Result->Data), Result->NumElements);
-			for(ValueType& ResultElement : Results)
-			{
-				ResultElement = InValue;
-			}
-		}
+		check(Result != nullptr);
+		TParam<ValueType> TypedResult(*Result);
+		*TypedResult = InValue;
 	}
 
 	// Get the current result as a TParam
 	template<typename ValueType>
-	TParam<ValueType> GetResult() const
+	TParam<ValueType> GetResultParam() const
 	{
-		// Types must match
-		checkSlow(Private::CheckParam<ValueType>(*Result));
-		return TParam<ValueType>(*Result, GetBatchedFlags());
+		check(Result != nullptr);
+		return TParam<ValueType>(*Result);
+	}
+
+	// Get the current result as a mutable reference
+	template<typename ValueType>
+	ValueType& GetResult() const
+	{
+		check(Result != nullptr);
+		TParam<ValueType> TypedResult(*Result);
+		return *TypedResult;
 	}
 
-	// Get the current result as a specified raw type pointer
+	// Get the current result as a mutable ptr
 	template<typename ValueType>
-	ValueType* GetResultChecked() const
+	ValueType* GetResultPtr() const
 	{
-		// Types must match
-		checkSlow(Private::CheckParam<ValueType>(*Result));
-		return static_cast<ValueType*>(Result->Data);
+		check(Result != nullptr);
+		TParam<ValueType> TypedResult(*Result);
+		return &*TypedResult;
 	}
 
 public:
@@ -445,26 +415,36 @@ public:
 
 	// Gets (and allocates, if necessary) state for the specified node given this calling context
 	template<typename ValueType, EStatePersistence Persistence = EStatePersistence::Relevancy>
-	TParam<ValueType> GetState(const FInterfaceKeyWithId& InKey) const
+	TParam<ValueType> GetStateParam(const FInterfaceKeyWithId& InKey) const
 	{
 		return State->GetState<ValueType, Persistence>(InKey, *this, CallstackHash);
 	}
 
 	// Gets (and allocates, if necessary) state for the specified node given this calling context
 	template<typename ValueType, EStatePersistence Persistence = EStatePersistence::Relevancy>
-	TParam<ValueType> GetState(const IAnimNextInterface* InAnimNextInterface, uint32 InId) const
+	TParam<ValueType> GetStateParam(const IAnimNextInterface* InAnimNextInterface, uint32 InId) const
 	{
 		return State->GetState<ValueType, Persistence>(InAnimNextInterface, InId, *this, CallstackHash);
 	}
 
-public:
-	// --- Mix Context Utils ---
+	// Gets (and allocates, if necessary) state for the specified node given this calling context
+	template<typename ValueType, EStatePersistence Persistence = EStatePersistence::Relevancy>
+	ValueType& GetState(const FInterfaceKeyWithId& InKey) const
+	{
+		TParam<ValueType> Param = State->GetState<ValueType, Persistence>(InKey, *this, CallstackHash); 
+		return *Param;
+	}
 
-	// Number of elements that this context represents
-	int32 GetNum() const
+	// Gets (and allocates, if necessary) state for the specified node given this calling context
+	template<typename ValueType, EStatePersistence Persistence = EStatePersistence::Relevancy>
+	ValueType& GetState(const IAnimNextInterface* InAnimNextInterface, uint32 InId) const
 	{
-		return State->NumElements;
+		TParam<ValueType> Param = State->GetState<ValueType, Persistence>(InAnimNextInterface, InId, *this, CallstackHash);
+		return *Param;
 	}
+	
+public:
+	// --- Mix Context Utils ---
 
 	// Access delta time as a param
 	TParam<const float> GetDeltaTimeParam() const;
@@ -472,50 +452,12 @@ public:
 	// Raw access to delta time
 	float GetDeltaTime() const { return DeltaTime; }
 
-	// --- Batching management ---
-
-	// Get parameter flags for batching
-	FParam::EFlags GetBatchedFlags() const
-	{
-		return State->NumElements > 1 ? FParam::EFlags::Batched : FParam::EFlags::None;
-	}
-
-	FORCEINLINE_DEBUGGABLE bool ShouldProcessThisElement(int32 ElementIndex) const
-	{
-		check(ElementIndex < MAX_BATCH_ELEMENTS);
-		return (BranchingMask & (1ull << ElementIndex)) != 0;
-	}
-
-	FORCEINLINE_DEBUGGABLE const uint64 GetBranchingMask() const
-	{
-		return BranchingMask;
-	}
-
-	void SetBranchingMask(const TArrayView<int32> InEnabledElementIndexes)
-	{
-		check(InEnabledElementIndexes.Num() <= GetNum());
-
-		BranchingMask = 0;
-		for (int32 ElementIndex : InEnabledElementIndexes)
-		{
-			check(ElementIndex < MAX_BATCH_ELEMENTS);
-
-			BranchingMask |= (1ull << ElementIndex);
-		}
-	}
-
-	void SetBranchingMask(uint64 InBranchingMask)
-	{
-		BranchingMask = InBranchingMask;
-	}
-
 private:
 	FContext(float InDeltaTime, FState& InState, FParamStorage& InParamStorage, FParam& InResult)
 		: State(&InState)
 		, ParamStorage(&InParamStorage)
 		, Result(&InResult)
 		, DeltaTime(InDeltaTime)
-		, BranchingMask(MAX_uint64)
 	{}
 
 	FContext() = default;
@@ -529,62 +471,17 @@ private:
 
 private:
 	TMap<FName, FParam> AdditionalParameters;
-	TMap<FName, FHParam> AdditionalParameterHandles;  // Temp param storage prototype
+	TMap<FName, FParamHandle> AdditionalParameterHandles;  // Temp param storage prototype
 	const FContext* Parent = nullptr;
 	const FContext* Root = nullptr;
 	FState* State = nullptr;
 	FParamStorage* ParamStorage = nullptr;
 	FParamStorageHandle BlockHandle = InvalidBlockHandle;
 	FParam* Result = nullptr;
-	IAnimNextInterfaceParameters* Parameters = nullptr;
+	IAnimNextParamInterface* Parameters = nullptr;
 	float DeltaTime = 0.0f;
 	uint32 CallstackHash = 0;
 	uint32 UpdateCounter = 0;
-	uint64 BranchingMask = MAX_uint64;
 };
 
-// A typed param that owns it's own memory
-template<typename ValueType, typename AllocatorType = TMemStackAllocator<>>
-struct TAllocParam : public TParam<ValueType>
-{
-public:
-	TAllocParam(const FContext& InContext)
-		: TParam<ValueType>(InContext.GetBatchedFlags())
-	{
-		TParam<ValueType>::Data = GetDataForConstructor(InContext);
-	}
-
-	TAllocParam(const FContext& InContext, FParam::EFlags InFlags)
-		: TParam<ValueType>(InFlags)
-	{
-		TParam<ValueType>::Data = GetDataForConstructor(InContext);
-	}
-
-
-private:
-	void* GetDataForConstructor(const FContext& InContext)
-	{
-		ValueArray.SetNum(InContext.GetNum());
-		return ValueArray.GetData();
-	}
-
-	TArray<ValueType, AllocatorType> ValueArray;
-};
-
-
-// --- Batching helper macros ---
-
-#define SKIP_CONTEXT_ELEMENT_IF_MASKED(Context) \
-		if (!Context.ShouldProcessThisElement(i)) \
-			continue;
-
-// Helper macros to allow processing context element batches
-#define PROCESS_BATCH_ELEMENTS_START(Context)  \
-	const int32 ContextNumElements = Context.GetNum(); \
-	for (int32 i = 0; i < ContextNumElements; i++) \
-	{ \
-		SKIP_CONTEXT_ELEMENT_IF_MASKED(Context)
-
-#define PROCESS_BATCH_ELEMENTS_END }
-
 }
diff --git a/Engine/Plugins/Experimental/AnimNext/Source/AnimNextInterface/Public/AnimNextInterfaceKernel.h b/Engine/Plugins/Experimental/AnimNext/Source/AnimNextInterface/Public/AnimNextInterfaceKernel.h
deleted file mode 100644
diff --git a/Engine/Plugins/Experimental/AnimNext/Source/AnimNextInterface/Public/AnimNextInterfaceKey.h b/Engine/Plugins/Experimental/AnimNext/Source/AnimNextInterface/Public/AnimNextInterfaceKey.h
--- a/Engine/Plugins/Experimental/AnimNext/Source/AnimNextInterface/Public/AnimNextInterfaceKey.h
+++ b/Engine/Plugins/Experimental/AnimNext/Source/AnimNextInterface/Public/AnimNextInterfaceKey.h
@@ -4,11 +4,10 @@
 
 #include "CoreMinimal.h"
 #include "UObject/Class.h"
-#include "AnimNextInterfaceParam.h"
 
 class IAnimNextInterface;
 
-namespace UE::AnimNext::Interface
+namespace UE::AnimNext
 {
 
 // Key value for cached anim interface state
diff --git a/Engine/Plugins/Experimental/AnimNext/Source/AnimNextInterface/Public/AnimNextInterfaceParam.h b/Engine/Plugins/Experimental/AnimNext/Source/AnimNextInterface/Public/AnimNextInterfaceParam.h
deleted file mode 100644
diff --git a/Engine/Plugins/Experimental/AnimNext/Source/AnimNextInterface/Public/AnimNextInterfaceParamStorage.h b/Engine/Plugins/Experimental/AnimNext/Source/AnimNextInterface/Public/AnimNextInterfaceParamStorage.h
deleted file mode 100644
diff --git a/Engine/Plugins/Experimental/AnimNext/Source/AnimNextInterface/Public/AnimNextInterfaceState.h b/Engine/Plugins/Experimental/AnimNext/Source/AnimNextInterface/Public/AnimNextInterfaceState.h
--- a/Engine/Plugins/Experimental/AnimNext/Source/AnimNextInterface/Public/AnimNextInterfaceState.h
+++ b/Engine/Plugins/Experimental/AnimNext/Source/AnimNextInterface/Public/AnimNextInterfaceState.h
@@ -3,14 +3,13 @@
 #pragma once
 
 #include "CoreMinimal.h"
-#include "AnimNextInterfaceParam.h"
+#include "Param/Param.h"
 #include "AnimNextInterfaceKey.h"
 
-namespace UE::AnimNext::Interface
+namespace UE::AnimNext
 {
 
 struct FContext;
-struct FParamType;
 
 enum class EStatePersistence : uint8
 {
@@ -22,12 +21,6 @@ enum class EStatePersistence : uint8
 // Container for anim interface state
 struct ANIMNEXTINTERFACE_API FState
 {
-	FState(int32 InNumElements = 1)
-		: NumElements(InNumElements)
-	{
-		check(InNumElements > 0);
-	}
-	
 private:
 	friend struct FContext;
 
@@ -45,7 +38,7 @@ private:
 		else
 		{
 			// Allocate new state
-			FParam* NewState = AllocateState(CombinedKey, InContext, Private::TParamType<ValueType>::GetType(), Persistence);
+			FParam* NewState = AllocateState(CombinedKey, InContext, FParamTypeHandle::GetHandle<ValueType>(), Persistence);
 
 			// Construct
 			new (NewState->Data) ValueType();
@@ -66,13 +59,13 @@ private:
 private:
 	FParam* FindStateRaw(const FInterfaceKeyWithIdAndStack& InKey, const FContext& InContext, EStatePersistence InPersistence);
 
-	FParam* AllocateState(const FInterfaceKeyWithIdAndStack& InKey, const FContext& InContext, const FParamType& InType, EStatePersistence InPersistence);
+	FParam* AllocateState(const FInterfaceKeyWithIdAndStack& InKey, const FContext& InContext, const FParamTypeHandle& InTypeHandle, EStatePersistence InPersistence);
 
 private:
 	struct FRelevancyParam : FParam
 	{
-		FRelevancyParam(const FParamType& InType, void* InData, EFlags InFlags, uint32 InUpdateCounter)
-			: FParam(InType, InData, InFlags)
+		FRelevancyParam(const FParamTypeHandle& InTypeHandle, TArrayView<uint8> InData, EFlags InFlags, uint32 InUpdateCounter)
+			: FParam(InTypeHandle, InData, InFlags)
 			, UpdateCounter(InUpdateCounter)
 		{}
 	
@@ -81,7 +74,6 @@ private:
 	
 	TMap<FInterfaceKeyWithIdAndStack, FRelevancyParam> RelevancyValueMap;
 	TMap<FInterfaceKeyWithIdAndStack, FParam> PermanentValueMap;
-	int32 NumElements = 0;
 };
 
 }
\ No newline at end of file
diff --git a/Engine/Plugins/Experimental/AnimNext/Source/AnimNextInterface/Public/AnimNextInterfaceTypes.h b/Engine/Plugins/Experimental/AnimNext/Source/AnimNextInterface/Public/AnimNextInterfaceTypes.h
--- a/Engine/Plugins/Experimental/AnimNext/Source/AnimNextInterface/Public/AnimNextInterfaceTypes.h
+++ b/Engine/Plugins/Experimental/AnimNext/Source/AnimNextInterface/Public/AnimNextInterfaceTypes.h
@@ -3,67 +3,10 @@
 #pragma once
 
 #include "CoreMinimal.h"
-#include "AnimNextInterfaceParam.h"
-#include "UObject/ScriptInterface.h"
-
-class IAnimNextInterface;
+#include "Param/ParamTypeHandle.h"
 
 #define ANIM_NEXT_INTERFACE_RETURN_TYPE(Typename) \
-virtual FName GetReturnTypeNameImpl() const final override\
+virtual UE::AnimNext::FParamTypeHandle GetReturnTypeHandleImpl() const final override\
 {\
-	static FName ReturnTypeNameValue = #Typename;\
-	return ReturnTypeNameValue;\
+	return UE::AnimNext::FParamTypeHandle::GetHandle<Typename>();\
 }\
-
-#define DECLARE_ANIM_NEXT_INTERFACE_PARAM_TYPE_INTERNAL(ValueType, Identifier) \
-	template<> \
-	struct UE::AnimNext::Interface::Private::TParamTypeImpl<ValueType> \
-	{ \
-		FORCEINLINE static const UE::AnimNext::Interface::FParamType& GetType() \
-		{ \
-			return Identifier##_AnimNextInterfaceTypeInfo; \
-		} \
-	};
-
-#define IMPLEMENT_ANIM_NEXT_INTERFACE_PARAM_TYPE_INTERNAL(ValueType, Identifier) \
-	UE::AnimNext::Interface::FParamType::FRegistrar Identifier##_Registrar([](){ UE::AnimNext::Interface::FParamType::FRegistrar::RegisterType( \
-		Identifier##_AnimNextInterfaceTypeInfo, \
-		UE::AnimNext::Interface::Private::TParamTypeHelper<ValueType>::GetStruct(), \
-		FName(#Identifier), UE::AnimNext::Interface::Private::TParamTypeHelper<ValueType>::GetSize(), \
-		UE::AnimNext::Interface::Private::TParamTypeHelper<ValueType>::GetAlignment(), \
-		UE::AnimNext::Interface::Private::TParamTypeHelper<ValueType>::GetParamCloneFunction()); });
-
-// Types that are passed between anim interfaces (i.e. exist in in public headers) should declare via this mechanism
-#define DECLARE_ANIM_NEXT_INTERFACE_PARAM_TYPE(APIType, ValueType, Identifier) \
-	extern APIType UE::AnimNext::Interface::FParamType Identifier##_AnimNextInterfaceTypeInfo; \
-	DECLARE_ANIM_NEXT_INTERFACE_PARAM_TYPE_INTERNAL(ValueType, Identifier)
-
-// Types that are passed between anim interfaces (i.e. exist in in public headers) should declare via this mechanism
-#define DECLARE_BUILTIN_ANIM_NEXT_INTERFACE_PARAM_TYPE(ValueType, Identifier) \
-	DECLARE_ANIM_NEXT_INTERFACE_PARAM_TYPE(ANIMNEXTINTERFACE_API, ValueType, Identifier)
-
-// Types that are passed between anim interfaces should register via this mechanism
-#define IMPLEMENT_ANIM_NEXT_INTERFACE_PARAM_TYPE(ValueType, Identifier) \
-	UE::AnimNext::Interface::FParamType Identifier##_AnimNextInterfaceTypeInfo; \
-	IMPLEMENT_ANIM_NEXT_INTERFACE_PARAM_TYPE_INTERNAL(ValueType, Identifier)
-
-// Structs that are only used for internal state (in a cpp file) can use this simpler registration.
-// If they are exposed in a public header please use the DECLARE_ANIM_NEXT_INTERFACE_PARAM_TYPE & IMPLEMENT_ANIM_NEXT_INTERFACE_PARAM_TYPE macros
-#define IMPLEMENT_ANIM_NEXT_INTERFACE_STATE_TYPE(ValueType, Identifier) \
-	static UE::AnimNext::Interface::FParamType Identifier##_AnimNextInterfaceTypeInfo; \
-	static UE::AnimNext::Interface::FParamType::FRegistrar Identifier##_Registrar([](){ UE::AnimNext::Interface::FParamType::FRegistrar::RegisterType( \
-		Identifier##_AnimNextInterfaceTypeInfo, \
-		UE::AnimNext::Interface::Private::TParamTypeHelper<ValueType>::GetStruct(), \
-		FName(#Identifier), UE::AnimNext::Interface::Private::TParamTypeHelper<ValueType>::GetSize(), \
-		UE::AnimNext::Interface::Private::TParamTypeHelper<ValueType>::GetAlignment(), \
-		UE::AnimNext::Interface::Private::TParamTypeHelper<ValueType>::GetParamCloneFunction()); }); \
-	DECLARE_ANIM_NEXT_INTERFACE_PARAM_TYPE_INTERNAL(ValueType, Identifier);
-
-namespace UE::AnimNext::Interface
-{
-
-#define ANIM_NEXT_INTERFACE_TYPE(Type, Identifier) DECLARE_BUILTIN_ANIM_NEXT_INTERFACE_PARAM_TYPE(Type, Identifier)
-#include "AnimNextInterfaceTypes.inl"
-#undef ANIM_NEXT_INTERFACE_TYPE
-
-}
\ No newline at end of file
diff --git a/Engine/Plugins/Experimental/AnimNext/Source/AnimNextInterface/Public/AnimNextInterfaceTypes.inl b/Engine/Plugins/Experimental/AnimNext/Source/AnimNextInterface/Public/AnimNextInterfaceTypes.inl
deleted file mode 100644
diff --git a/Engine/Plugins/Experimental/AnimNext/Source/AnimNextInterface/Public/AnimationDataRegistry.h b/Engine/Plugins/Experimental/AnimNext/Source/AnimNextInterface/Public/AnimationDataRegistry.h
--- a/Engine/Plugins/Experimental/AnimNext/Source/AnimNextInterface/Public/AnimationDataRegistry.h
+++ b/Engine/Plugins/Experimental/AnimNext/Source/AnimNextInterface/Public/AnimationDataRegistry.h
@@ -5,12 +5,12 @@
 #include "CoreMinimal.h"
 #include "AnimationDataRegistryTypes.h"
 #include "UObject/WeakObjectPtr.h"
-#include "AnimNextInterfaceParam.h"
+#include "Param/Param.h"
 
 class USkeletalMeshComponent;
 struct FReferenceSkeleton;
 
-namespace UE::AnimNext::Interface
+namespace UE::AnimNext
 {
 
 class FModule;
@@ -75,12 +75,12 @@ public:
 	template<typename DataType>
 	FAnimationDataHandle PreAllocateMemory(const int32 NumElements)
 	{
-		const FParamType& ParamType = Private::TParamType<DataType>::GetType();
+		FParamTypeHandle ParamTypeHandle = FParamTypeHandle::GetHandle<DataType>();
 
 		const FDataTypeDef* TypeDef = nullptr;
 		{
 			FRWScopeLock Lock(DataTypeDefsLock, SLT_ReadOnly);
-			TypeDef = DataTypeDefs.Find(ParamType.GetTypeId());
+			TypeDef = DataTypeDefs.Find(ParamTypeHandle);
 		}
 
 		if (TypeDef == nullptr)
@@ -89,7 +89,7 @@ public:
 			// TODO : Log if we allocate more than DEFAULT_BLOCK_SIZE elements of that type
 		}
 
-		if (ensure(TypeDef != nullptr && TypeDef->TypeId != MAX_uint16))
+		if (ensure(TypeDef != nullptr && TypeDef->ParamTypeHandle.IsValid()))
 		{
 			const int32 ElementSize = TypeDef->ElementSize;
 			const int32 ElementAlign = TypeDef->ElementAlign;
@@ -99,7 +99,7 @@ public:
 
 			uint8* Memory = (uint8*)FMemory::Malloc(BufferSize, TypeDef->ElementAlign);    // TODO : This should come from preallocated chunks, use malloc / free for now
 
-			Private::FAllocatedBlock* AllocatedBlock = new Private::FAllocatedBlock(Memory, NumElements, TypeDef->TypeId); // TODO : avoid memory fragmentation
+			Private::FAllocatedBlock* AllocatedBlock = new Private::FAllocatedBlock(Memory, NumElements, ParamTypeHandle); // TODO : avoid memory fragmentation
 			AllocatedBlock->AddRef();
 			AllocatedBlocks.Add(AllocatedBlock);
 			return FAnimationDataHandle(AllocatedBlock);
@@ -134,8 +134,8 @@ private:
 	// structure holding each registered type information
 	struct FDataTypeDef
 	{
-		Private::AnimDataTypeId TypeId = Private::InvalidTypeId;
-		DestroyFnSignature DestrooyTypeFn = nullptr;
+		FParamTypeHandle ParamTypeHandle;
+		DestroyFnSignature DestroyTypeFn = nullptr;
 		int32 ElementSize = 0;
 		int32 ElementAlign = 0;
 		int32 AllocationBlockSize = 0;
@@ -162,7 +162,7 @@ private:
 	};
 
 	// Map holding registered types
-	TMap<Private::AnimDataTypeId, FDataTypeDef> DataTypeDefs;
+	TMap<FParamTypeHandle, FDataTypeDef> DataTypeDefs;
 	// Lock for registered types map
 	FRWLock DataTypeDefsLock;
 
@@ -182,12 +182,11 @@ private:
 	template<typename DataType>
 	FDataTypeDef* RegisterDataType_Impl(int32 AllocationBlockSize)
 	{
-		const FParamType& ParamType = Private::TParamType<DataType>::GetType();
-		const Private::AnimDataTypeId TypeId = ParamType.GetTypeId();
-		check(TypeId != MAX_uint16); // Missing IMPLEMENT_ANIM_NEXT_INTERFACE_STATE_TYPE(Type, Name);
+		FParamTypeHandle ParamTypeHandle = FParamTypeHandle::GetHandle<DataType>();
+		check(ParamTypeHandle.IsValid());
 
-		const int32 ElementSize = ParamType.GetSize();
-		const int32 ElementAlign = ParamType.GetAlignment();
+		const int32 ElementSize = ParamTypeHandle.GetSize();
+		const int32 ElementAlign = ParamTypeHandle.GetAlignment();
 
 		// If we use raw types, I need a per element destructor
 		DestroyFnSignature DestroyFn = [](uint8* TargetBuffer, int32 NumElem)->void
@@ -203,8 +202,8 @@ private:
 		{
 			FRWScopeLock WriteLock(DataTypeDefsLock, SLT_Write);
 
-			AddedDef = &DataTypeDefs.FindOrAdd(TypeId, { TypeId, DestroyFn, ElementSize, ElementAlign, AllocationBlockSize });
-			check(AddedDef->TypeId == TypeId); // check we have not added two different types with the same ID
+			AddedDef = &DataTypeDefs.FindOrAdd(ParamTypeHandle, { ParamTypeHandle, DestroyFn, ElementSize, ElementAlign, AllocationBlockSize });
+			check(AddedDef->ParamTypeHandle == ParamTypeHandle); // check we have not added two different types with the same ID
 		}
 
 		return AddedDef;
@@ -234,4 +233,4 @@ private:
 	static void HandlePostGarbageCollect();
 };
 
-} // namespace UE::AnimNext::Interface
+} // namespace UE::AnimNext
diff --git a/Engine/Plugins/Experimental/AnimNext/Source/AnimNextInterface/Public/AnimationDataRegistryTypes.h b/Engine/Plugins/Experimental/AnimNext/Source/AnimNextInterface/Public/AnimationDataRegistryTypes.h
--- a/Engine/Plugins/Experimental/AnimNext/Source/AnimNextInterface/Public/AnimationDataRegistryTypes.h
+++ b/Engine/Plugins/Experimental/AnimNext/Source/AnimNextInterface/Public/AnimationDataRegistryTypes.h
@@ -5,31 +5,29 @@
 #include "CoreMinimal.h"
 #include "UObject/WeakObjectPtr.h"
 #include <algorithm>
+
+#include "Param/ParamTypeHandle.h"
 #include "Animation/AnimTypes.h"
 #include "ReferenceSkeleton.h"
 #include "BoneIndices.h"
 
-namespace UE::AnimNext::Interface
+namespace UE::AnimNext
 {
 
-
 namespace Private
 {
 
-using AnimDataTypeId = uint16;
-static constexpr AnimDataTypeId InvalidTypeId = MAX_uint16;
-
 struct FAllocatedBlock
 {
-	mutable int32 NumRefs = 0;
 	void* Memory = nullptr;
 	int32 NumElem = 0;
-	AnimDataTypeId TypeId = InvalidTypeId;
+	mutable int32 NumRefs = 0;
+	FParamTypeHandle TypeHandle;
 
-	FAllocatedBlock(void* InMemory, int32 InNumElem, AnimDataTypeId InTypeId)
+	FAllocatedBlock(void* InMemory, int32 InNumElem, FParamTypeHandle InTypeHandle)
 		: Memory(InMemory)
 		, NumElem(InNumElem)
-		, TypeId(InTypeId)
+		, TypeHandle(InTypeHandle)
 	{
 	}
 
@@ -156,9 +154,9 @@ struct ANIMNEXTINTERFACE_API FAnimationDataHandle
 		return *Data;
 	}
 
-	inline Private::AnimDataTypeId GetTypeId() const
+	inline FParamTypeHandle GetTypeHandle() const
 	{
-		return AllocatedBlock != nullptr ? AllocatedBlock->TypeId : Private::InvalidTypeId;
+		return AllocatedBlock != nullptr ? AllocatedBlock->TypeHandle : FParamTypeHandle();
 	}
 
 private:
@@ -172,4 +170,4 @@ enum class ETransformFlags : uint8
 };
 
 
-} // namespace UE::AnimNext::Interface
+} // namespace UE::AnimNext
diff --git a/Engine/Plugins/Experimental/AnimNext/Source/AnimNextInterface/Public/AnimationDecompressionTools.h b/Engine/Plugins/Experimental/AnimNext/Source/AnimNextInterface/Public/AnimationDecompressionTools.h
--- a/Engine/Plugins/Experimental/AnimNext/Source/AnimNextInterface/Public/AnimationDecompressionTools.h
+++ b/Engine/Plugins/Experimental/AnimNext/Source/AnimNextInterface/Public/AnimationDecompressionTools.h
@@ -12,7 +12,7 @@ struct FAnimSequenceDecompressionContext;
 struct FRootMotionReset;
 class UAnimSequence;
 
-namespace UE::AnimNext::Interface
+namespace UE::AnimNext
 {
 
 class ANIMNEXTINTERFACE_API FDecompressionTools
@@ -72,4 +72,4 @@ public:
 		, const bool bDisableRetargeting);
 };
 
-} // end namespace UE::AnimNext::Interface
+} // end namespace UE::AnimNext
diff --git a/Engine/Plugins/Experimental/AnimNext/Source/AnimNextInterface/Public/AnimationGenerationTools.h b/Engine/Plugins/Experimental/AnimNext/Source/AnimNextInterface/Public/AnimationGenerationTools.h
--- a/Engine/Plugins/Experimental/AnimNext/Source/AnimNextInterface/Public/AnimationGenerationTools.h
+++ b/Engine/Plugins/Experimental/AnimNext/Source/AnimNextInterface/Public/AnimationGenerationTools.h
@@ -12,7 +12,7 @@ class FSkeletalMeshLODRenderData;
 struct FPoseContext;
 class UPhysicsAsset;
 
-namespace UE::AnimNext::Interface
+namespace UE::AnimNext
 {
 
 struct FGenerationLODData
@@ -98,4 +98,4 @@ public:
 		, FPoseContext& TargetPose); // TODO : see how to deal with heap and stack variants without duplicating functions (and if possible without templating the function)
 };
 
-} // namespace UE::AnimNext::Interface
+} // namespace UE::AnimNext
diff --git a/Engine/Plugins/Experimental/AnimNext/Source/AnimNextInterface/Public/AnimationReferencePose.h b/Engine/Plugins/Experimental/AnimNext/Source/AnimNextInterface/Public/AnimationReferencePose.h
--- a/Engine/Plugins/Experimental/AnimNext/Source/AnimNextInterface/Public/AnimationReferencePose.h
+++ b/Engine/Plugins/Experimental/AnimNext/Source/AnimNextInterface/Public/AnimationReferencePose.h
@@ -10,12 +10,12 @@
 #include "Animation/AnimTypes.h"
 #include "ReferenceSkeleton.h"
 #include "BoneIndices.h"
-#include "AnimNextInterfaceTypes.h"
+#include "AnimationReferencePose.generated.h"
 
 class USkeletalMesh;
 class USkeleton;
 
-namespace UE::AnimNext::Interface
+namespace UE::AnimNext
 {
 
 #define ANIM_ENABLE_POINTER_ITERATION 1
@@ -682,15 +682,17 @@ using FAnimTransformArray = FAnimTransformArrayAoSHeap;
 
 //****************************************************************************
 
+enum class EReferencePoseGenerationFlags : uint8
+{
+	None = 0,
+	FastPath = 1 << 0
+};
+
+ENUM_CLASS_FLAGS(EReferencePoseGenerationFlags);
+
 template <typename AllocatorType>
 struct TAnimationReferencePose
 {
-	enum EGenerationFlags : uint8
-	{
-		None = 0,
-		FastPath = 1 << 0
-	};
-
 	TAnimTransformArray<AllocatorType> ReferenceLocalTransforms;
 	TArray<TArray<FBoneIndexType, AllocatorType>, AllocatorType> LODBoneIndexes;
 	TArray<TArray<FBoneIndexType, AllocatorType>, AllocatorType> SkeletonToLODBoneIndexes;
@@ -698,7 +700,7 @@ struct TAnimationReferencePose
 
 	TWeakObjectPtr<const USkeletalMesh> SkeletalMesh = nullptr;
 	TWeakObjectPtr<const USkeleton> Skeleton = nullptr;
-	EGenerationFlags GenerationFlags = EGenerationFlags::None;
+	EReferencePoseGenerationFlags GenerationFlags = EReferencePoseGenerationFlags::None;
 
 	TAnimationReferencePose() = default;
 
@@ -716,7 +718,7 @@ struct TAnimationReferencePose
 
 	bool IsFastPath() const
 	{
-		return (GenerationFlags & EGenerationFlags::FastPath) != 0;
+		return (GenerationFlags & EReferencePoseGenerationFlags::FastPath) != EReferencePoseGenerationFlags::None;
 	}
 
 	void Initialize(const FReferenceSkeleton& RefSkeleton
@@ -740,7 +742,7 @@ struct TAnimationReferencePose
 			ReferenceLocalTransforms[i] = RefBonePose[InBoneIndexes[i]]; // TODO : For SoA this is un-optimal, as we are using a TransformAdapter. Evaluate using a specific SoA iterator
 		}
 
-		GenerationFlags = bFastPath ? EGenerationFlags::FastPath : EGenerationFlags::None;
+		GenerationFlags = bFastPath ? EReferencePoseGenerationFlags::FastPath : EReferencePoseGenerationFlags::None;
 	}
 
 	const TArrayView<const FBoneIndexType> GetLODBoneIndexes(int32 LODLevel) const
@@ -839,7 +841,7 @@ struct TAnimationReferencePose
 using FAnimationReferencePose = TAnimationReferencePose<FDefaultAllocator>;
 
 
-enum EAnimationPoseFlags : uint8
+enum class EAnimationPoseFlags : uint8
 {
 	None				= 0,
 	Additive			= 1 << 0,
@@ -848,6 +850,8 @@ enum EAnimationPoseFlags : uint8
 	UseSourceData		= 1 << 3
 };
 
+ENUM_CLASS_FLAGS(EAnimationPoseFlags);
+
 template <typename AllocatorType>
 struct TAnimationLODPose
 {
@@ -954,7 +958,7 @@ struct TAnimationLODPose
 	/** True if retargeting is disabled */
 	bool GetDisableRetargeting() const
 	{
-		return (Flags & EAnimationPoseFlags::DisableRetargeting) != 0;
+		return (Flags & EAnimationPoseFlags::DisableRetargeting) != EAnimationPoseFlags::None;
 	}
 
 	/** Ignore compressed data and use RAW data instead, for debugging. */
@@ -966,7 +970,7 @@ struct TAnimationLODPose
 	/** True if we're requesting RAW data instead of compressed data. For debugging. */
 	bool ShouldUseRawData() const
 	{
-		return (Flags & EAnimationPoseFlags::UseRawData) != 0;
+		return (Flags & EAnimationPoseFlags::UseRawData) != EAnimationPoseFlags::None;
 	}
 
 	/** Use Source data instead.*/
@@ -978,7 +982,7 @@ struct TAnimationLODPose
 	/** True if we're requesting Source data instead of RawAnimationData. For debugging. */
 	bool ShouldUseSourceData() const
 	{
-		return (Flags & EAnimationPoseFlags::UseSourceData) != 0;
+		return (Flags & EAnimationPoseFlags::UseSourceData) != EAnimationPoseFlags::None;
 	}
 
 };
@@ -988,8 +992,24 @@ using FAnimationLODPoseStack = TAnimationLODPose<FAnimStackAllocator>;
 
 using FAnimationLODPose = FAnimationLODPoseHeap;
 
-DECLARE_BUILTIN_ANIM_NEXT_INTERFACE_PARAM_TYPE(FAnimationReferencePose, AnimationReferencePose);
-//DECLARE_BUILTIN_ANIM_NEXT_INTERFACE_PARAM_TYPE(FAnimationLODPose, AnimationLODPose);
-//DECLARE_BUILTIN_ANIM_NEXT_INTERFACE_PARAM_TYPE(FAnimationLODPoseStack, AnimationLODPoseStack);
+} // namespace UE::AnimNext
 
-} // namespace UE::AnimNext::Interface
+// USTRUCT wrapper for reference pose
+USTRUCT()
+struct FAnimNextReferencePose
+#if CPP
+	: UE::AnimNext::FAnimationReferencePose
+#endif
+{
+	GENERATED_BODY()
+};
+
+// USTRUCT wrapper for LOD pose
+USTRUCT()
+struct FAnimNextLODPose
+#if CPP
+	: UE::AnimNext::FAnimationLODPose
+#endif
+{
+	GENERATED_BODY()
+};
\ No newline at end of file
diff --git a/Engine/Plugins/Experimental/AnimNext/Source/AnimNextInterface/Public/IAnimNextInterface.h b/Engine/Plugins/Experimental/AnimNext/Source/AnimNextInterface/Public/IAnimNextInterface.h
--- a/Engine/Plugins/Experimental/AnimNext/Source/AnimNextInterface/Public/IAnimNextInterface.h
+++ b/Engine/Plugins/Experimental/AnimNext/Source/AnimNextInterface/Public/IAnimNextInterface.h
@@ -7,12 +7,15 @@
 #include "IAnimNextInterface.generated.h"
 
 class IAnimNextInterface;
-namespace UE::AnimNext::Interface
+
+namespace UE::AnimNext
 {
 struct FParam;
+struct FParamTypeHandle;
 struct FContext;
 }
 
+
 UINTERFACE(BlueprintType, meta = (CannotImplementInterfaceInBlueprint))
 class ANIMNEXTINTERFACE_API UAnimNextInterface : public UInterface
 {
@@ -29,53 +32,41 @@ public:
 	 * Gets data using context to retrieve input parameters and result
 	 * @return false if type are incompatible, or if nested calls fail, otherwise true
 	 */
-	bool GetData(const UE::AnimNext::Interface::FContext& Context) const;
+	bool GetData(const UE::AnimNext::FContext& Context) const;
 
 	/**
 	 * Gets data and stores the value in OutResult, checking whether types are compatible
 	* @return false if nested calls fail, otherwise true
 	*/
-	bool GetDataChecked(const UE::AnimNext::Interface::FContext& Context) const;
+	bool GetDataChecked(const UE::AnimNext::FContext& Context) const;
 
 	/**
 	 * Gets data and stores the value in OutResult.
 	 * @return false if type are incompatible, or if nested calls fail, otherwise true
 	 */
-	bool GetData(const UE::AnimNext::Interface::FContext& Context, UE::AnimNext::Interface::FParam& OutResult) const;
+	bool GetData(const UE::AnimNext::FContext& Context, UE::AnimNext::FParam& OutResult) const;
 
 	/**
 	 * Gets data and stores the value in OutResult, checking whether types are compatible
 	 * @return false if nested calls fail, otherwise true
 	 */
-	bool GetDataChecked(const UE::AnimNext::Interface::FContext& Context, UE::AnimNext::Interface::FParam& OutResult) const;
+	bool GetDataChecked(const UE::AnimNext::FContext& Context, UE::AnimNext::FParam& OutResult) const;
 
-	/** Get the name of the return type */
-	FName GetReturnTypeName() const
-	{
-		return GetReturnTypeNameImpl();
-	}
-
-	/** Get the struct of the return type, if any */
-	const UScriptStruct* GetReturnTypeStruct() const
-	{
-		return GetReturnTypeStructImpl();
-	}
+	/** Get the handle of the return type */
+	UE::AnimNext::FParamTypeHandle GetReturnTypeHandle() const;
 
 private:
 	/** Get data if the types are compatible */
-	bool GetDataIfCompatibleInternal(const UE::AnimNext::Interface::FContext& InContext) const;
+	bool GetDataIfCompatibleInternal(const UE::AnimNext::FContext& InContext) const;
 	
 	/** Get a value from an anim interface with no dynamic or static type checking. Internal use only. */
-	bool GetDataRawInternal(const UE::AnimNext::Interface::FContext& InContext) const;
+	bool GetDataRawInternal(const UE::AnimNext::FContext& InContext) const;
 
 protected:
-	/** Get the return type name, used for dynamic type checking */
-	virtual FName GetReturnTypeNameImpl() const = 0;
+	/** Get the handle of the return type */
+	virtual UE::AnimNext::FParamTypeHandle GetReturnTypeHandleImpl() const = 0;
 
-	/** Get the return type struct, used for dynamic type checking. If the result is not a struct, this should return nullptr */
-	virtual const UScriptStruct* GetReturnTypeStructImpl() const { return nullptr; }
-	
 	/** Get the value for this interface. @return true if successful, false if unsuccessful. */
-	virtual bool GetDataImpl(const UE::AnimNext::Interface::FContext& Context) const = 0;
+	virtual bool GetDataImpl(const UE::AnimNext::FContext& Context) const = 0;
 };
 
diff --git a/Engine/Plugins/Experimental/AnimNext/Source/AnimNextInterface/Public/IAnimNextInterfaceParameters.h b/Engine/Plugins/Experimental/AnimNext/Source/AnimNextInterface/Public/IAnimNextInterfaceParameters.h
deleted file mode 100644
diff --git a/Engine/Plugins/Experimental/AnimNext/Source/AnimNextInterface/Public/Param/IAnimNextParamInterface.h b/Engine/Plugins/Experimental/AnimNext/Source/AnimNextInterface/Public/Param/IAnimNextParamInterface.h
new file mode 100644
--- /dev/null
+++ b/Engine/Plugins/Experimental/AnimNext/Source/AnimNextInterface/Public/Param/IAnimNextParamInterface.h
@@ -0,0 +1,33 @@
+﻿// Copyright Epic Games, Inc. All Rights Reserved.
+
+#pragma once
+
+#include "CoreMinimal.h"
+#include "UObject/Interface.h"
+#include "IAnimNextParamInterface.generated.h"
+
+namespace UE::AnimNext
+{
+struct FParam;
+struct FContext;
+}
+
+UINTERFACE(BlueprintType, meta = (CannotImplementInterfaceInBlueprint))
+class ANIMNEXTINTERFACE_API UAnimNextParamInterface : public UInterface
+{
+	GENERATED_BODY()
+};
+
+// Interface representing parameters that can be supplied to an anim interface
+class ANIMNEXTINTERFACE_API IAnimNextParamInterface
+{
+	GENERATED_BODY()
+
+	friend struct UE::AnimNext::FContext;
+
+	// Get a parameter by name
+	virtual bool GetParameter(FName InKey, UE::AnimNext::FParam& OutParam) = 0;
+
+	// Get a parameter pointer by name
+	virtual const UE::AnimNext::FParam* GetParameter(FName InKey) = 0;
+};
diff --git a/Engine/Plugins/Experimental/AnimNext/Source/AnimNextInterface/Public/Param/Param.h b/Engine/Plugins/Experimental/AnimNext/Source/AnimNextInterface/Public/Param/Param.h
new file mode 100644
--- /dev/null
+++ b/Engine/Plugins/Experimental/AnimNext/Source/AnimNextInterface/Public/Param/Param.h
@@ -0,0 +1,380 @@
+// Copyright Epic Games, Inc. All Rights Reserved.
+
+#pragma once
+
+#include "CoreMinimal.h"
+#include "UObject/Class.h"
+#include "UObject/ScriptInterface.h"
+#include "Param/ParamType.h"
+#include "Param/ParamTypeHandle.h"
+
+namespace UE::AnimNext
+{
+struct FContext;
+
+struct FParamStorage;
+
+// Parameter/result/state memory wrapper
+// @TODO can we save memory here with some base/derived system for fundamental types vs struct types?
+struct ANIMNEXTINTERFACE_API FParam
+{
+public:
+	enum class EFlags : uint8
+	{
+		None		= 0,			// No flags
+		Mutable		= 1 << 0,		// Parameter is mutable, so can be mutated at runtime
+		Stored		= 1 << 1,		// Parameter has to be stored on context storage
+
+		// --- Flags added for parameter storage prototype ---
+		Value		= 1 << 2,		// Parameter will be stored as a Value
+		Reference	= 1 << 3,		// Parameter will be stored as a Reference (pointer)
+		Embedded	= 1 << 4,		// Parameter will be stored as a Value, but stored directly on the Data pointer
+	};
+
+	// Get the type handle of this param 
+	FParamTypeHandle GetTypeHandle() const { return TypeHandle; }
+
+	// Check whether the supplied param can be written to by this param
+	// Verifies that types are compatible, destination is mutable and batching matches
+	bool CanAssignTo(const FParam& InParam) const;
+
+	// Helper functions for CanAssignTo
+	bool CanAssignWith(const FParamTypeHandle& InTypeHandle, EFlags InFlags, FStringBuilderBase* OutReasonPtr = nullptr) const;
+
+	// Check whether this parameter is able to be mutated
+	bool IsMutable() const { return EnumHasAnyFlags(Flags, EFlags::Mutable); }
+
+	// Get the internal flags
+	EFlags GetFlags() const { return Flags; }
+
+	// Get an immutable view of the parameter's data
+	TConstArrayView<uint8> GetData() const { return TConstArrayView<uint8>(static_cast<uint8*>(Data), Size); }
+
+	// Get an mutable view of the parameter's data, asserts if this parameter is immutable
+	TArrayView<uint8> GetMutableData() const { check(IsMutable()); return TArrayView<uint8>(static_cast<uint8*>(Data), Size); }
+
+protected:
+	friend struct FContext;
+	friend struct FState;
+	friend struct FParamStorage;
+
+	FParam(const FParam* InOtherParam);
+	FParam(const FParamTypeHandle& InTypeHandle, TArrayView<uint8> InData, EFlags InFlags);
+	FParam(const FParamTypeHandle& InTypeHandle, TConstArrayView<uint8> InData, EFlags InFlags);
+	FParam(const FParamTypeHandle& InTypeHandle, EFlags InFlags);
+
+	// Raw ptr to the data, or the data itself if we have EFlags::Embedded
+	void* Data = nullptr;
+
+	// The type of the param
+	FParamTypeHandle TypeHandle;
+
+	// Size of the data
+	uint16 Size = 0;
+
+	// Internal flags
+	EFlags Flags = EFlags::None;
+
+public:
+	// Internal use, but required for default constructed elements on containers
+	FParam() = default;
+
+	/** Duplicate a parameter into the provided memory */
+	static FParam DuplicateParam(const FParam& InSource, TArrayView<uint8> InTargetMemory);
+};
+
+ENUM_CLASS_FLAGS(FParam::EFlags);
+
+namespace Private
+{
+
+template<typename ValueType>
+static bool CheckParam(const FParam& InParam)
+{
+	return FParamTypeHandle::GetHandle<ValueType>() == InParam.GetTypeHandle();
+}
+
+}
+
+namespace Private
+{
+
+// Concept used to decide whether to use container-based construction of wrapped results
+struct CSizedContainerWithAccessibleDataAsRawPtr
+{
+	template<typename ContainerType>
+	auto Requires(ContainerType& Container) -> decltype(
+		Container.Num(),
+		Container.GetData()
+	);
+};
+
+}
+
+// A typed result which wraps the type-erased underlying param
+template<typename ValueType>
+struct TParam : FParam
+{
+public:
+	FORCEINLINE_DEBUGGABLE bool IsValid() const
+	{
+		return ((Data != nullptr || EnumHasAnyFlags(Flags, FParam::EFlags::Embedded))
+			&& TypeHandle.IsValid());
+	}
+
+	FORCEINLINE_DEBUGGABLE const ValueType& GetDataChecked() const
+	{
+		check(Data != nullptr || EnumHasAnyFlags(Flags, FParam::EFlags::Embedded));
+
+		const ValueType* ValueData = EnumHasAnyFlags(Flags, FParam::EFlags::Embedded)
+			? static_cast<const ValueType*>((void*)&Data)
+			: static_cast<const ValueType*>(Data);
+
+		return *ValueData;
+	}
+
+	FORCEINLINE_DEBUGGABLE ValueType& GetDataChecked()
+	{
+		check(Data != nullptr || EnumHasAnyFlags(Flags, FParam::EFlags::Embedded));
+
+		ValueType* ValueData = EnumHasAnyFlags(Flags, FParam::EFlags::Embedded)
+			? static_cast<ValueType*>((void*)&Data)
+			: static_cast<ValueType*>(Data);
+
+		return *ValueData;
+	}
+
+	FORCEINLINE_DEBUGGABLE operator const ValueType& () const
+	{
+		check((Data != nullptr || EnumHasAnyFlags(Flags, FParam::EFlags::Embedded)));
+
+		const ValueType* ValueData = EnumHasAnyFlags(Flags, FParam::EFlags::Embedded)
+			? static_cast<const ValueType*>((void*)&Data)
+			: static_cast<const ValueType*>(Data);
+
+		return *ValueData;
+	}
+
+	FORCEINLINE_DEBUGGABLE operator ValueType& ()
+	{
+		check((Data != nullptr || EnumHasAnyFlags(Flags, FParam::EFlags::Embedded)));
+
+		ValueType* ValueData = EnumHasAnyFlags(Flags, FParam::EFlags::Embedded)
+			? static_cast<ValueType*>((void*)&Data)
+			: static_cast<ValueType*>(Data);
+
+		return *ValueData;
+	}
+
+	FORCEINLINE_DEBUGGABLE ValueType& operator*() const
+	{
+		check((Data != nullptr || EnumHasAnyFlags(Flags, FParam::EFlags::Embedded)));
+
+		ValueType* ValueData = EnumHasAnyFlags(Flags, FParam::EFlags::Embedded)
+			? static_cast<ValueType*>((void*)&Data)
+			: static_cast<ValueType*>(Data);
+
+		return *ValueData;
+	}
+
+	FORCEINLINE_DEBUGGABLE ValueType* operator->() const
+	{
+		check((Data != nullptr || EnumHasAnyFlags(Flags, FParam::EFlags::Embedded)));
+
+		ValueType* ValueData = EnumHasAnyFlags(Flags, FParam::EFlags::Embedded)
+			? static_cast<ValueType*>((void*)&Data)
+			: static_cast<ValueType*>(Data);
+
+		return ValueData;
+	}
+
+protected:
+	using TValueType = ValueType;
+	
+	friend struct FContext;
+	friend struct FState;
+	friend struct FParamStorage;
+
+	TParam()
+		: FParam(FParamTypeHandle::GetHandle<ValueType>(), GetTypeFlags(EFlags::None))
+	{
+	}
+	
+	TParam(EFlags InFlags)
+		: FParam(FParamTypeHandle::GetHandle<ValueType>(), InFlags)
+	{
+	}
+
+	TParam(TArrayView<uint8> InData, EFlags InFlags)
+		: FParam(FParamTypeHandle::GetHandle<ValueType>(), InData, InFlags)
+	{
+	}
+
+	TParam(TConstArrayView<uint8> InData, EFlags InFlags)
+		: FParam(FParamTypeHandle::GetHandle<ValueType>(), InData, InFlags)
+	{
+	}
+
+	TParam(FParam& InParam, EFlags InAdditionalFlags = EFlags::None)
+		: FParam(&InParam)
+	{
+#if DO_CHECK
+		TStringBuilder<64> Error;
+		checkf(InParam.CanAssignWith(FParamTypeHandle::GetHandle<ValueType>(), GetTypeFlags(InAdditionalFlags), &Error),
+			TEXT("Cannot assign type: %s"), Error.ToString());
+#endif
+	}
+
+	TParam(const FParam& InParam, EFlags InAdditionalFlags = EFlags::None)
+		: FParam(&InParam)
+	{
+#if DO_CHECK
+		TStringBuilder<64> Error;
+		checkf(InParam.CanAssignWith(FParamTypeHandle::GetHandle<ValueType>(), GetTypeFlags(InAdditionalFlags), &Error),
+			TEXT("Cannot assign type: %s"), Error.ToString());
+#endif
+	}
+
+protected:
+	static EFlags GetTypeFlags(EFlags InAdditionalFlags)
+	{
+		EFlags Flags = InAdditionalFlags;
+
+		// Add const flags or not
+		if constexpr (!TIsConst<ValueType>::Value)
+		{
+			Flags |= FParam::EFlags::Mutable;
+		}
+
+		return Flags;
+	}
+};
+
+// A typed param that wraps a user ptr
+template<typename ValueType>
+struct TWrapParam : TParam<ValueType>
+{
+public:
+	TWrapParam(ValueType* InValuePtrToWrap)
+		: TParam<ValueType>(GetDataForConstructor(InValuePtrToWrap), GetFlagsForConstructor(InValuePtrToWrap))
+	{
+	}
+
+	TWrapParam(ValueType& InValueToWrap)
+		: TParam<ValueType>(GetDataForConstructor(&InValueToWrap), GetFlagsForConstructor(&InValueToWrap))
+	{
+	}
+
+protected:
+	TArrayView<uint8> GetDataForConstructor(ValueType* InValuePtrToWrap) const
+	{
+		using MutableValueType = std::remove_const_t<ValueType>;
+
+		return TArrayView<uint8>(reinterpret_cast<uint8*>(const_cast<MutableValueType*>(InValuePtrToWrap)), sizeof(ValueType));
+	}
+
+	FParam::EFlags GetFlagsForConstructor(const ValueType* InValuePtrToWrap) const
+	{
+		FParam::EFlags NewFlags = FParam::EFlags::None;
+
+		// Add const flags or not
+		if constexpr (!TIsConst<ValueType>::Value)
+		{
+			NewFlags |= FParam::EFlags::Mutable;
+		}
+
+		return NewFlags;
+	}
+};
+
+// A typed param that owns it's own memory with size defined at compile time
+template<typename ValueType>
+struct TParamValue : TParam<ValueType>
+{
+public:
+	TParamValue()
+		: TParam<ValueType>(TArrayView<uint8>(reinterpret_cast<uint8*>(&Value), sizeof(ValueType)), GetFlagsForConstructor())
+	{
+	}
+
+	TParamValue(FParam::EFlags InFlags)
+		: TParam<ValueType>(TArrayView<uint8>(reinterpret_cast<uint8*>(&Value), sizeof(ValueType)), InFlags)
+	{
+	}
+
+private:
+	FParam::EFlags GetFlagsForConstructor() const
+	{
+		FParam::EFlags NewFlags = FParam::EFlags::None;
+
+		// Add const flags or not
+		if constexpr (!TIsConst<ValueType>::Value)
+		{
+			NewFlags |= FParam::EFlags::Mutable;
+		}
+
+		return NewFlags;
+	}
+
+	ValueType Value;
+};
+
+template<typename ValueType>
+struct TContextStorageParam : TWrapParam<ValueType>
+{
+public:
+	TContextStorageParam(ValueType* InValuePtrToWrap)
+		: TWrapParam<ValueType>(InValuePtrToWrap)
+	{
+		FParam::Flags |= FParam::EFlags::Stored;
+	}
+
+	TContextStorageParam(ValueType& InValueToWrap)
+		: TWrapParam<ValueType>(&InValueToWrap)
+	{
+		FParam::Flags |= FParam::EFlags::Stored;
+	}
+};
+
+struct ANIMNEXTINTERFACE_API FParamHandle
+{
+	using FInternalHandle = int32;
+	static constexpr FInternalHandle InvalidParamHandle = -1;
+
+	FParamHandle() = default;
+	~FParamHandle();
+
+	FParamHandle(FParamHandle&& Other)
+		: FParamHandle()
+	{
+		Swap(*this, Other);
+	}
+
+	FParamHandle& operator= (FParamHandle&& Other)
+	{
+		Swap(*this, Other);
+		return *this;
+	}
+
+	FParamHandle(const FParamHandle& Other);
+
+	FParamHandle& operator= (const FParamHandle& Other)
+	{
+		FParamHandle Temp(Other); // note that the ref count is incremented here
+
+		Swap(*this, Temp);
+
+		return *this;
+	}
+
+protected:
+	friend struct FContext;
+	friend struct FParamStorage;
+
+	FParamHandle(FParamStorage* InOwnerStorage, FInternalHandle InInternalHandle);
+
+	FParamStorage* OwnerStorage = nullptr;
+	FInternalHandle ParamHandle = InvalidParamHandle;
+};
+
+} // end namespace UE::AnimNext
diff --git a/Engine/Plugins/Experimental/AnimNext/Source/AnimNextInterface/Public/Param/ParamHelpers.h b/Engine/Plugins/Experimental/AnimNext/Source/AnimNextInterface/Public/Param/ParamHelpers.h
new file mode 100644
--- /dev/null
+++ b/Engine/Plugins/Experimental/AnimNext/Source/AnimNextInterface/Public/Param/ParamHelpers.h
@@ -0,0 +1,29 @@
+﻿// Copyright Epic Games, Inc. All Rights Reserved.
+
+#pragma once
+
+#include "CoreMinimal.h"
+
+struct FAnimNextParamType;
+
+namespace UE::AnimNext
+{
+
+struct FParamTypeHandle;
+
+struct ANIMNEXTINTERFACE_API FParamHelpers
+{
+	enum class ECopyResult
+	{
+		Failed,
+		Succeeded,
+	};
+
+	/** Copy a parameter from one type to another, validating type, container and any polymorphism rules if the types differ at all. */
+	static ECopyResult Copy(const FAnimNextParamType& InSourceType, const FAnimNextParamType& InTargetType, TConstArrayView<uint8> InSourceMemory, TArrayView<uint8> InTargetMemory);
+
+	/** Copy a parameter from one type to another, validating type, container and any polymorphism rules if the types differ at all. */
+	static ECopyResult Copy(const FParamTypeHandle& InSourceTypeHandle, const FParamTypeHandle& InTargetTypeHandle, TConstArrayView<uint8> InSourceMemory, TArrayView<uint8> InTargetMemory);
+};
+
+}
\ No newline at end of file
diff --git a/Engine/Plugins/Experimental/AnimNext/Source/AnimNextInterface/Public/Param/ParamStorage.h b/Engine/Plugins/Experimental/AnimNext/Source/AnimNextInterface/Public/Param/ParamStorage.h
new file mode 100644
--- /dev/null
+++ b/Engine/Plugins/Experimental/AnimNext/Source/AnimNextInterface/Public/Param/ParamStorage.h
@@ -0,0 +1,236 @@
+// Copyright Epic Games, Inc. All Rights Reserved.
+
+#pragma once
+
+#include "CoreMinimal.h"
+#include "UObject/Class.h"
+#include "UObject/ScriptInterface.h"
+#include "Param/Param.h"
+#include "Param/ParamType.h"
+#include "Param/ParamHelpers.h"
+
+namespace UE::AnimNext
+{
+
+typedef int32 FParamStorageHandle;
+constexpr FParamStorageHandle InvalidBlockHandle = -1;
+
+struct ANIMNEXTINTERFACE_API FParamStorage
+{
+public:
+	typedef TPair<FParamStorageHandle, uint8*> TBlockDataPair;
+
+	FParamStorage(int32 InMaxParams, int32 InAllocatedMemorySize, int32 InMaxBlocks, bool InAllowGrowing = false);
+	~FParamStorage();
+
+	// *** Interface for param arena ***
+	// This prototype allows storing parameters into the storage using pre-allocated blocks of memory
+	// The owner is responsible for calculating the size and requesting / releasing the memory blocks
+	// Memory blocks are stored in stack LIFO way, last one has to be the first to be released (for simplicity)
+	// If a memoy block is released and it is not the last, it will be held until the ones after it are released
+
+	// Requests a block of RequestedBlockSize size
+	// Returns a pair with the block handle and the pointer to the memory buffer
+	TBlockDataPair RequestBlock(int32 RequestedBlockSize);
+
+	// Releases a block of memory from the arena, using a hanle to identify it
+	// If it is not the last of the stack, it will be held until the blocks after this one are released
+	void ReleaseBlock(FParamStorageHandle BlockHandle);
+
+	// Obtains the block of memory from a block handle
+	uint8* GetBlockMemory(FParamStorageHandle BlockHandle);
+
+
+	// *** Interface for direct Param storage ***
+	// This prototype allows storing parameters directly into the storage
+	// Parameters are stored stack like, in a LIFO approach
+	// If a parameter is released and it is not the last of the stack, it will be held until the parameters over that one are released
+	// This is to have a very simple memory management, might be removed on final implementation, based on requirements
+
+	// Adds a value to the storage (copying the data), returning a HParam
+	// If the total allocation size is <= than size of a void *, the value is stored directly in the FParam Data
+	// If not, additional memory is requested (TODO : better memory management)
+	template<typename ValueType>
+	FParamHandle AddValue(ValueType* Data, FParam::EFlags Flags)
+	{
+		check(ParamFreeIndex < (Parameters.Num())); // TODO : Growing ?
+
+		const int32 ParameterIndex = ParamFreeIndex;
+		ParamFreeIndex++;
+
+		const int32 ParamAllocSize = GetParamAllocSize<ValueType>(1);
+
+		FParam& Param = Parameters[ParameterIndex];
+
+		// Prepare the Param for the copy
+		const FParamTypeHandle TypeHandle = FParamTypeHandle::GetHandle<ValueType>();
+		Param.TypeHandle = TypeHandle;
+		Param.Flags = Flags;
+		Param.Data = Data;
+		Param.Size = ParamAllocSize;
+
+		// If size is <= sizeof(void*), put the value directly at the pointer
+		if (ParamAllocSize <= sizeof(void*)) // TODO : see if we can get this constexpr
+		{
+			ParamMemoryHandles[ParameterIndex] = InvalidBlockHandle;
+
+			// Set target memory directly at Param.Data memory
+			TArrayView<uint8> SourceData(reinterpret_cast<uint8*>(Data), ParamAllocSize);
+			TArrayView<uint8> TargetData(static_cast<uint8*>((void*)&Param.Data), ParamAllocSize);
+			FParamHelpers::Copy(TypeHandle, TypeHandle, SourceData, TargetData);
+
+			Param.Flags |= FParam::EFlags::Embedded; // Added so it can be checked when retrieving (TODO : see if it can be done constexpr by type)
+		}
+		else
+		{
+			TBlockDataPair BlockData = RequestBlock(ParamAllocSize);
+
+			ParamMemoryHandles[ParameterIndex] = BlockData.Key;
+
+			// Set target memory directly at BlockData memory
+			TArrayView<uint8> SourceData(reinterpret_cast<uint8*>(Data), ParamAllocSize);
+			TArrayView<uint8> TargetData(BlockData.Value, ParamAllocSize);
+			FParamHelpers::Copy(TypeHandle, TypeHandle, SourceData, TargetData);
+			// Once the copy is done, update the Param
+			Param.Data = BlockData.Value;
+		}
+
+		RefCounts[ParameterIndex] = 1;
+
+		return FParamHandle(this, static_cast<FParamHandle::FInternalHandle>(ParameterIndex));
+	}
+
+	// Adds a reference (pointer) to the storage, returning a FParamHandle
+	template<typename ValueType>
+	FParamHandle AddReference(ValueType* Data, FParam::EFlags Flags)
+	{
+		check(ParamFreeIndex < (Parameters.Num())); // TODO : Growing ?
+
+		const int32 ParameterIndex = ParamFreeIndex;
+		ParamFreeIndex++;
+
+		FParam& Param = Parameters[ParameterIndex];
+
+		Param.TypeHandle = FParamTypeHandle::GetHandle<ValueType>();
+		Param.Flags = Flags;
+		Param.Data = Data;
+		Param.Size = sizeof(ValueType);
+
+		ParamMemoryHandles[ParameterIndex] = InvalidBlockHandle;
+
+		RefCounts[ParameterIndex] = 1;
+
+		return FParamHandle(this, static_cast<FParamHandle::FInternalHandle>(ParameterIndex));
+	}
+
+	// Obtains the Param associated to the FParamHandle
+	FParam* GetParam(FParamHandle::FInternalHandle InParamHandle)
+	{
+		FParam* RetVal = nullptr;
+
+		const int32 ParameterIndex = InParamHandle;
+
+		if (ParameterIndex < ParamFreeIndex && RefCounts[ParameterIndex] > 0)
+		{
+			RetVal = &Parameters[ParameterIndex];
+		}
+
+		return RetVal;
+	}
+
+protected:
+	friend struct FParamHandle;
+
+	// Adds a reference count to a param, using the passed param handle
+	void IncRefCount(const FParamHandle::FInternalHandle InParamHandle)
+	{
+		const int32 ParameterIndex = InParamHandle;
+		check(ParameterIndex < ParamFreeIndex);
+		
+		check(RefCounts[ParameterIndex] > 0);
+		
+		RefCounts[ParameterIndex]++;
+	}
+
+	// Decrements a reference count to a param, using the passed param handle
+	// If the count reaches 0:
+	// - If it is the last param of the stack, it releases it and all the previous params with ref count == 0
+	// - If it is not the last one, returns without releasing (it will be done later)
+	void DecRefCount(const FParamHandle::FInternalHandle InParamHandle)
+	{
+		int32 ParameterIndex = InParamHandle;
+		check(ParameterIndex < ParamFreeIndex);
+
+		check(RefCounts[ParameterIndex] > 0);
+		
+		if (RefCounts[ParameterIndex] > 0)
+		{
+			RefCounts[ParameterIndex]--;
+
+			if (RefCounts[ParameterIndex] == 0)
+			{
+				// if it is the last parameter, release the memory and deal with the prev params with RefCount == 0, if any
+				if (ParameterIndex == ParamFreeIndex - 1)
+				{
+					// Recover free slots
+					while (ParameterIndex >= 0 && RefCounts[ParameterIndex] == 0)
+					{
+						Parameters[ParameterIndex] = FParam();
+						if (ParamMemoryHandles[ParameterIndex] != InvalidBlockHandle)
+						{
+							ReleaseBlock(ParamMemoryHandles[ParameterIndex]);
+							ParamMemoryHandles[ParameterIndex] = InvalidBlockHandle;
+						}
+
+						ParameterIndex--;
+						ParamFreeIndex--;
+					}
+				}
+			}
+		}
+	}
+
+	// Calculates the required size for a param, including alignment
+	static int32 GetParamAllocSize(const FAnimNextParamType& ParamType, const int32 NumElem)
+	{
+		const int32 ParamAlignment = ParamType.GetAlignment();
+		const int32 ParamSize = ParamType.GetSize();
+		const int32 ParamAllocSize = NumElem * Align(ParamSize, ParamAlignment);
+
+		return ParamAllocSize;
+	}
+
+	// Calculates the required size for a param, including alignment
+	template<typename ValueType>
+	static int32 GetParamAllocSize(const int32 NumElem)
+	{
+		const int32 ParamAlignment = alignof(ValueType);
+		const int32 ParamSize = sizeof(ValueType);
+		const int32 ParamAllocSize = NumElem * Align(ParamSize, ParamAlignment);
+
+		return ParamAllocSize;
+	}
+
+	// --- context arena param storage prototype ---
+	TArray<int32> BlockOffets;
+	TArray<int32> BlockSizes;
+	TBitArray<> BlockFreeFlag;
+	uint8* RawMemory = nullptr;
+	int32 RawMemorySize = 0;
+	int32 CurrentBlocksSize = 0;
+	int32 FreeBlockIndex = 0;
+	
+	// --- shared param storage prototype ---
+	TArray<FParam> Parameters;
+	TArray<FParamStorageHandle> ParamMemoryHandles;
+	TArray<int16> RefCounts;
+	int32 ParamFreeIndex = 0;
+
+	bool AllowGrowing = false;
+
+private:
+	bool CheckBufferSize(int32 RequestedBlockSize);
+	bool CheckNumBlocks(int32 BlockIndex);
+};
+
+} // end namespace UE::AnimNext
diff --git a/Engine/Plugins/Experimental/AnimNext/Source/AnimNextInterface/Public/Param/ParamType.h b/Engine/Plugins/Experimental/AnimNext/Source/AnimNextInterface/Public/Param/ParamType.h
new file mode 100644
--- /dev/null
+++ b/Engine/Plugins/Experimental/AnimNext/Source/AnimNextInterface/Public/Param/ParamType.h
@@ -0,0 +1,277 @@
+﻿// Copyright Epic Games, Inc. All Rights Reserved.
+
+#pragma once
+
+#include "CoreMinimal.h"
+#include "PropertyBag.h"
+#include "Templates/SubclassOf.h"
+#include "UObject/ObjectPtr.h"
+#include "UObject/SoftObjectPtr.h"
+#include "UObject/ReflectedTypeAccessors.h"
+#include "ParamType.generated.h"
+
+namespace UE::AnimNext
+{
+	struct FParamTypeHandle;
+	struct FParamHelpers;
+}
+
+namespace UE::AnimNext::Private
+{
+	struct CHasStaticClass
+	{
+		template<typename ValueType>
+		auto Requires(UClass* Result) -> decltype(
+			Result = ValueType::StaticClass()
+		);
+	};
+
+	struct CHasStaticStruct
+	{
+		template<typename ValueType>
+		auto Requires(UScriptStruct* Result) -> decltype(
+			Result = ValueType::StaticStruct()
+		);
+	};
+
+	template <typename T>
+	struct TIsSoftObjectPtr
+	{
+		enum { Value = false };
+	};
+
+	template <typename T> struct TIsSoftObjectPtr<               TSoftObjectPtr<T>> { enum { Value = true }; };
+	template <typename T> struct TIsSoftObjectPtr<const          TSoftObjectPtr<T>> { enum { Value = true }; };
+	template <typename T> struct TIsSoftObjectPtr<      volatile TSoftObjectPtr<T>> { enum { Value = true }; };
+	template <typename T> struct TIsSoftObjectPtr<const volatile TSoftObjectPtr<T>> { enum { Value = true }; };
+
+	template <typename T>
+	struct TIsSoftClassPtr
+	{
+		enum { Value = false };
+	};
+
+	template <typename T> struct TIsSoftClassPtr<               TSoftClassPtr<T>> { enum { Value = true }; };
+	template <typename T> struct TIsSoftClassPtr<const          TSoftClassPtr<T>> { enum { Value = true }; };
+	template <typename T> struct TIsSoftClassPtr<      volatile TSoftClassPtr<T>> { enum { Value = true }; };
+	template <typename T> struct TIsSoftClassPtr<const volatile TSoftClassPtr<T>> { enum { Value = true }; };
+}
+
+/**
+ * Representation of a parameter's type. Serializable, but fairly heavyweight to pass around and compare.
+ * Faster comparisons and other operations can be performed on UE::AnimNext::FTypeHandle, but they cannot be
+ * serialized as they are not stable across runs.
+ */
+USTRUCT()
+struct ANIMNEXTINTERFACE_API FAnimNextParamType
+{
+public:
+	GENERATED_BODY()
+
+	using EValueType = ::EPropertyBagPropertyType;
+	using EContainerType = ::EPropertyBagContainerType;
+
+	friend struct UE::AnimNext::FParamTypeHandle;
+	friend struct UE::AnimNext::FParamHelpers;
+
+	FAnimNextParamType() = default;
+
+	/** Construct a parameter type from the passed in value, container and object type. */
+	FAnimNextParamType(EValueType InValueType, EContainerType InContainerType = EContainerType::None, UObject* InValueTypeObject = nullptr);
+
+private:
+	/** Pointer to object that defines the Enum, Struct, or Class. */
+	UPROPERTY()
+	TObjectPtr<UObject> ValueTypeObject = nullptr;
+
+	/** Type of the value described by this parameter. */
+	UPROPERTY()
+	EPropertyBagPropertyType ValueType = EPropertyBagPropertyType::None;
+
+	/** Type of the container described by this parameter. */
+	UPROPERTY()
+	EPropertyBagContainerType ContainerType = EPropertyBagContainerType::None;
+	
+private:
+	/** Helper function for GetType */
+	template<typename ParamType>
+	static constexpr void GetTypeInner(FAnimNextParamType& ParameterType)
+	{
+		if constexpr (std::is_same_v<ParamType, bool>)
+		{
+			ParameterType.ValueType = EValueType::Bool;
+		}
+		else if constexpr (std::is_same_v<ParamType, uint8>)
+		{
+			ParameterType.ValueType = EValueType::Byte;
+		}
+		else if constexpr (std::is_same_v<ParamType, int32>)
+		{
+			ParameterType.ValueType = EValueType::Int32;
+		}
+		else if constexpr (std::is_same_v<ParamType, int64>)
+		{
+			ParameterType.ValueType = EValueType::Int64;
+		}
+		else if constexpr (std::is_same_v<ParamType, float>)
+		{
+			ParameterType.ValueType = EValueType::Float;
+		}
+		else if constexpr (std::is_same_v<ParamType, double>)
+		{
+			ParameterType.ValueType = EValueType::Double;
+		}
+		else if constexpr (std::is_same_v<ParamType, FName>)
+		{
+			ParameterType.ValueType = EValueType::Name;
+		}
+		else if constexpr (std::is_same_v<ParamType, FString>)
+		{
+			ParameterType.ValueType = EValueType::String;
+		}
+		else if constexpr (std::is_same_v<ParamType, FText>)
+		{
+			ParameterType.ValueType = EValueType::Text;
+		}
+		else if constexpr (TIsUEnumClass<ParamType>::Value)
+		{
+			ParameterType.ValueType = EValueType::Enum;
+			ParameterType.ValueTypeObject = StaticEnum<ParamType>();
+		}
+		else if constexpr (TModels<UE::AnimNext::Private::CHasStaticStruct, ParamType>::Value)
+		{
+			ParameterType.ValueType = EValueType::Struct;
+			ParameterType.ValueTypeObject = ParamType::StaticStruct();
+		}
+		else if constexpr (TModels<UE::AnimNext::Private::CHasStaticClass, ParamType>::Value)
+		{
+			if constexpr (std::is_same_v<ParamType, UClass>)
+			{
+				ParameterType.ValueType = EValueType::Class;
+				ParameterType.ValueTypeObject = ParamType::StaticClass();
+			}
+			else
+			{
+				ParameterType.ValueType = EValueType::Object;
+				ParameterType.ValueTypeObject = ParamType::StaticClass();
+			}
+		}
+		else if constexpr (TIsTObjectPtr<ParamType>::Value)
+		{
+			if constexpr (std::is_same_v<ParamType, TObjectPtr<UClass>>)
+			{
+				ParameterType.ValueType = EValueType::Class;
+				ParameterType.ValueTypeObject = ParamType::ElementType::StaticClass();
+			}
+			else
+			{
+				ParameterType.ValueType = EValueType::Object;
+				ParameterType.ValueTypeObject = ParamType::ElementType::StaticClass();
+			}
+		}
+		else if constexpr (TIsTSubclassOf<ParamType>::Value)
+		{
+			ParameterType.ValueType = EValueType::Class;
+			ParameterType.ValueTypeObject = ParamType::ElementType::StaticClass();
+		}
+		else if constexpr (UE::AnimNext::Private::TIsSoftObjectPtr<ParamType>::Value)
+		{
+			ParameterType.ValueType = EValueType::SoftObject;
+			ParameterType.ValueTypeObject = ParamType::ElementType::StaticClass();
+		}
+		else if constexpr (UE::AnimNext::Private::TIsSoftClassPtr<ParamType>::Value)
+		{
+			ParameterType.ValueType = EValueType::SoftClass;
+			ParameterType.ValueTypeObject = ParamType::ElementType::StaticClass();
+		}
+	}
+
+	/** Helper function for IsValid */
+	bool IsValidObject() const;
+
+public:
+	/** Get a parameter type based on the passed-in built-in type */
+	template<typename ParamType>
+	static FAnimNextParamType GetType()
+	{
+		using NonConstType = std::remove_const_t<ParamType>;
+
+		FAnimNextParamType ParameterType;
+
+		if constexpr (TIsTArray<NonConstType>::Value)
+		{
+			GetTypeInner<typename NonConstType::ElementType>(ParameterType);
+			ParameterType.ContainerType = EContainerType::Array;
+		}
+		else
+		{
+			GetTypeInner<NonConstType>(ParameterType);
+		}
+
+		return ParameterType;
+	}
+
+	/** Get a parameter type handle that represents this type */
+	UE::AnimNext::FParamTypeHandle GetHandle() const;
+
+	/** 
+	 * Helper function returning the size of the type.
+	 *
+	 * @returns Size in bytes of the type.
+	 */
+	size_t GetSize() const;
+
+	/** 
+	 * Helper function returning the size of the value type.
+	 * This is identical to GetSize for non containers
+	 *
+	 * @returns Size in bytes of the value type.
+	 */	
+	size_t GetValueTypeSize() const;
+	
+	/** 
+	 * Helper function returning the alignment of the type.
+	 *
+	 * @returns Alignment of the type.
+	 */	
+	size_t GetAlignment() const;
+
+	/** 
+	 * Helper function returning the alignment of the value type.
+	 * This is identical to GetAlignment for non containers
+	 *
+	 * @returns Alignment of the value type.
+	 */	
+	size_t GetValueTypeAlignment() const;
+
+	/** Get a string representing this type */
+	FString ToString() const;
+
+	/** Equality operator */
+	friend bool operator==(const FAnimNextParamType& InLHS, const FAnimNextParamType& InRHS)
+	{
+		return InLHS.ValueType == InRHS.ValueType && InLHS.ContainerType == InRHS.ContainerType && InLHS.ValueTypeObject == InRHS.ValueTypeObject;
+	}
+
+	/** Inequality operator */
+	friend bool operator!=(const FAnimNextParamType& InLHS, const FAnimNextParamType& InRHS)
+	{
+		return InLHS.ValueType != InRHS.ValueType || InLHS.ContainerType != InRHS.ContainerType || InLHS.ValueTypeObject != InRHS.ValueTypeObject;
+	}
+
+	/** Type hash for TMap storage */
+	friend uint32 GetTypeHash(const FAnimNextParamType& InType)
+	{
+		return HashCombineFast(GetTypeHash((uint32)InType.ValueType | ((uint32)InType.ContainerType << 8)), GetTypeHash(InType.ValueTypeObject));
+	}
+	
+	/** @return whether this type actually describes a type */
+	bool IsValid() const
+	{
+		const bool bHasValidValueType = ValueType != EValueType::None;
+		const bool bHasValidContainerType = (ContainerType == EContainerType::None) || (bHasValidValueType && ContainerType != EContainerType::None);
+		const bool bHasValidObjectType = (ValueType < EValueType::Enum) || (ValueType >= EValueType::Enum && ValueType <= EValueType::SoftClass && IsValidObject());
+		return bHasValidValueType && bHasValidContainerType && bHasValidObjectType;
+	}
+};
+
diff --git a/Engine/Plugins/Experimental/AnimNext/Source/AnimNextInterface/Public/Param/ParamTypeHandle.h b/Engine/Plugins/Experimental/AnimNext/Source/AnimNextInterface/Public/Param/ParamTypeHandle.h
new file mode 100644
--- /dev/null
+++ b/Engine/Plugins/Experimental/AnimNext/Source/AnimNextInterface/Public/Param/ParamTypeHandle.h
@@ -0,0 +1,370 @@
+// Copyright Epic Games, Inc. All Rights Reserved.
+
+#pragma once
+
+#include "CoreMinimal.h"
+#include "Param/ParamType.h"
+
+class FAnimationAnimNextParametersTest_Types;
+
+namespace UE::AnimNext
+{
+
+/**
+ * Representation of a parameter's type. Lightweight to compare and pass around, but not serializable as not stable
+ * across runs.
+ * Serialization can be performed on FAnimNextParamType. 
+ */
+struct ANIMNEXTINTERFACE_API FParamTypeHandle
+{
+	/**
+	 * Byte enum representing built-in parameter types. Types beyond this are Custom and are accessed via an index into the
+	 * global type array, which is slower
+	 */
+	enum class EParamType : uint8
+	{
+		None = 0,
+
+		Bool,
+		Byte,
+		Int32,
+		Int64,
+		Float,
+		Double,
+		Name,
+		String,
+		Text,
+		Vector,
+		Vector4,
+		Quat,
+		Transform,
+
+		MaxBuiltIn = Transform,
+
+		Custom = 0xff,
+
+		Max = Custom,
+	};
+
+	friend struct ::FAnimNextParamType;
+	friend struct FParamHelpers;
+	friend class ::FAnimationAnimNextParametersTest_Types;
+
+private:
+	union
+	{
+		uint32 Value;
+
+		struct
+		{
+			uint32 BuiltInType : 8;
+			uint32 CustomTypeIndex : 24;
+		} Fields;
+	};
+
+private:
+	/** Get the built-in parameter type */
+	EParamType GetParameterType() const
+	{
+		return (EParamType)Fields.BuiltInType;
+	}
+
+	/** Set the built-in parameter type */
+	void SetParameterType(EParamType InType)
+	{
+		Fields.BuiltInType = (uint32)InType;
+	}
+
+	/** Get the custom type index */
+	uint32 GetCustomTypeIndex() const
+	{
+		return Fields.CustomTypeIndex;
+	}
+
+	/** Get the custom type index */
+	void SetCustomTypeIndex(uint32 InIndex)
+	{
+		checkf(InIndex < (1 << 24), TEXT("FTypeHandle::SetCustomTypeIndex: Type Index ouf of range"));
+		Fields.CustomTypeIndex = InIndex;
+	}
+
+	/** Reset the custom type registry, used in tests */
+	static void ResetCustomTypes();
+
+	/** Get a custom type index for the passed-in type information */
+	static uint32 GetOrAllocateCustomTypeIndex(FAnimNextParamType::EValueType InValueType, FAnimNextParamType::EContainerType InContainerType, UObject* InValueTypeObject);
+
+	/** Validate the passed-in custom type index */
+	static bool ValidateCustomTypeIndex(uint32 InCustomTypeIndex);
+
+	/** Helper for GetHandle */
+	template<typename ParamType>
+	static constexpr void GetHandleInner(FAnimNextParamType::EValueType& OutValueType, UObject*& OutValueTypeObject)
+	{
+		if constexpr (std::is_same_v<ParamType, bool>)
+		{
+			OutValueType = FAnimNextParamType::EValueType::Bool;
+		}
+		else if constexpr (std::is_same_v<ParamType, uint8>)
+		{
+			OutValueType = FAnimNextParamType::EValueType::Byte;
+		}
+		else if constexpr (std::is_same_v<ParamType, int32>)
+		{
+			OutValueType = FAnimNextParamType::EValueType::Int32;
+		}
+		else if constexpr (std::is_same_v<ParamType, int64>)
+		{
+			OutValueType = FAnimNextParamType::EValueType::Int64;
+		}
+		else if constexpr (std::is_same_v<ParamType, float>)
+		{
+			OutValueType = FAnimNextParamType::EValueType::Float;
+		}
+		else if constexpr (std::is_same_v<ParamType, double>)
+		{
+			OutValueType = FAnimNextParamType::EValueType::Double;
+		}
+		else if constexpr (std::is_same_v<ParamType, FName>)
+		{
+			OutValueType = FAnimNextParamType::EValueType::Name;
+		}
+		else if constexpr (std::is_same_v<ParamType, FString>)
+		{
+			OutValueType = FAnimNextParamType::EValueType::String;
+		}
+		else if constexpr (std::is_same_v<ParamType, FText>)
+		{
+			OutValueType = FAnimNextParamType::EValueType::Text;
+		}
+		else if constexpr (TIsUEnumClass<ParamType>::Value)
+		{
+			OutValueType = FAnimNextParamType::EValueType::Enum;
+			OutValueTypeObject = StaticEnum<ParamType>();
+		}
+		else if constexpr (TModels<Private::CHasStaticStruct, ParamType>::Value)
+		{
+			OutValueType = FAnimNextParamType::EValueType::Struct;
+			OutValueTypeObject = ParamType::StaticStruct();
+		}
+		else if constexpr (TModels<Private::CHasStaticClass, ParamType>::Value)
+		{
+			if constexpr (std::is_same_v<ParamType, UClass>)
+			{
+				OutValueType = FAnimNextParamType::EValueType::Class;
+				OutValueTypeObject = ParamType::StaticClass();
+			}
+			else
+			{
+				OutValueType = FAnimNextParamType::EValueType::Object;
+				OutValueTypeObject = ParamType::StaticClass();
+			}
+		}
+		else if constexpr (TIsTObjectPtr<ParamType>::Value)
+		{
+			if constexpr (std::is_same_v<ParamType, TObjectPtr<UClass>>)
+			{
+				OutValueType = FAnimNextParamType::EValueType::Class;
+				OutValueTypeObject = ParamType::ElementType::StaticClass();
+			}
+			else
+			{
+				OutValueType = FAnimNextParamType::EValueType::Object;
+				OutValueTypeObject = ParamType::ElementType::StaticClass();
+			}
+		}
+		else if constexpr (TIsTSubclassOf<ParamType>::Value)
+		{
+			OutValueType = FAnimNextParamType::EValueType::Class;
+			OutValueTypeObject = ParamType::ElementType::StaticClass();
+		}
+		else if constexpr (Private::TIsSoftObjectPtr<ParamType>::Value)
+		{
+			OutValueType = FAnimNextParamType::EValueType::SoftObject;
+			OutValueTypeObject = ParamType::ElementType::StaticClass();
+		}
+		else if constexpr (Private::TIsSoftClassPtr<ParamType>::Value)
+		{
+			OutValueType = FAnimNextParamType::EValueType::SoftClass;
+			OutValueTypeObject = ParamType::ElementType::StaticClass();
+		}
+		else
+		{
+			static_assert(sizeof(ParamType) == 0, "Type is not expressible as a type handle, see FAnimNextParamType for available types.");
+		}
+	}
+
+public:
+	constexpr FParamTypeHandle()
+		: Value(0)
+	{
+	}
+
+	/** Check whether this describes a built-in type (i.e. not a custom type) */
+	bool IsBuiltInType() const
+	{
+		return (EParamType)Fields.BuiltInType <= EParamType::MaxBuiltIn;
+	}
+
+	/** Get a parameter type based on the passed-in built-in type */
+	template<typename ParamType>
+	static constexpr FParamTypeHandle GetHandle()
+	{
+		using NonConstType = std::remove_const_t<ParamType>;
+
+		FParamTypeHandle TypeHandle;
+
+		// Check against built-in types statically
+		if constexpr (std::is_same_v<NonConstType, bool>)
+		{
+			TypeHandle.SetParameterType(EParamType::Bool);
+		}
+		else if constexpr (std::is_same_v<NonConstType, uint8>)
+		{
+			TypeHandle.SetParameterType(EParamType::Byte);
+		}
+		else if constexpr (std::is_same_v<NonConstType, int32>)
+		{
+			TypeHandle.SetParameterType(EParamType::Int32);
+		}
+		else if constexpr (std::is_same_v<NonConstType, int64>)
+		{
+			TypeHandle.SetParameterType(EParamType::Int64);
+		}
+		else if constexpr (std::is_same_v<NonConstType, float>)
+		{
+			TypeHandle.SetParameterType(EParamType::Float);
+		}
+		else if constexpr (std::is_same_v<NonConstType, double>)
+		{
+			TypeHandle.SetParameterType(EParamType::Double);
+		}
+		else if constexpr (std::is_same_v<NonConstType, FName>)
+		{
+			TypeHandle.SetParameterType(EParamType::Name);
+		}
+		else if constexpr (std::is_same_v<NonConstType, FString>)
+		{
+			TypeHandle.SetParameterType(EParamType::String);
+		}
+		else if constexpr (std::is_same_v<NonConstType, FText>)
+		{
+			TypeHandle.SetParameterType(EParamType::Text);
+		}
+		else if constexpr (std::is_same_v<NonConstType, FVector>)
+		{
+			TypeHandle.SetParameterType(EParamType::Vector);
+		}
+		else if constexpr (std::is_same_v<NonConstType, FVector4>)
+		{
+			TypeHandle.SetParameterType(EParamType::Vector4);
+		}
+		else if constexpr (std::is_same_v<NonConstType, FQuat>)
+		{
+			TypeHandle.SetParameterType(EParamType::Quat);
+		}
+		else if constexpr (std::is_same_v<NonConstType, FTransform>)
+		{
+			TypeHandle.SetParameterType(EParamType::Transform);
+		}
+		else
+		{
+			// Not a built-in-type, so we need to do some work
+			FAnimNextParamType::EValueType ValueType = FAnimNextParamType::EValueType::None;
+			FAnimNextParamType::EContainerType ContainerType = FAnimNextParamType::EContainerType::None;
+			UObject* ValueTypeObject = nullptr;
+
+			if constexpr (TIsTArray<NonConstType>::Value)
+			{
+				GetHandleInner<typename NonConstType::ElementType>(ValueType, ValueTypeObject);
+				ContainerType = FAnimNextParamType::EContainerType::Array;
+			}
+			else
+			{
+				GetHandleInner<NonConstType>(ValueType, ValueTypeObject);
+			}
+
+			if(ValueType != FAnimNextParamType::EValueType::None)
+			{
+				TypeHandle.SetCustomTypeIndex(GetOrAllocateCustomTypeIndex(ValueType, ContainerType, ValueTypeObject));
+				TypeHandle.SetParameterType(EParamType::Custom);
+			}
+		}
+
+		check(TypeHandle.IsValid());
+
+		return TypeHandle;
+	}
+
+	/** Get the parameter type that this handle corresponds to */
+	FAnimNextParamType GetType() const;
+
+	/** 
+	 * Helper function returning the size of the type.
+	 *
+	 * @returns Size in bytes of the type.
+	 */
+	size_t GetSize() const;
+
+	/** 
+	 * Helper function returning the size of the value type.
+	 * This is identical to GetSize for non containers
+	 *
+	 * @returns Size in bytes of the value type.
+	 */	
+	size_t GetValueTypeSize() const;
+	
+	/** 
+	 * Helper function returning the alignment of the type.
+	 *
+	 * @returns Alignment of the type.
+	 */	
+	size_t GetAlignment() const;
+
+	/** 
+	 * Helper function returning the alignment of the value type.
+	 * This is identical to GetAlignment for non containers
+	 *
+	 * @returns Alignment of the value type.
+	 */	
+	size_t GetValueTypeAlignment() const;
+	
+	/** Get a string representing this type */
+	FString ToString() const;
+
+	/** Equality operator */
+	friend bool operator==(const FParamTypeHandle& InLHS, const FParamTypeHandle& InRHS)
+	{
+		return InLHS.Value == InRHS.Value;
+	}
+
+	/** Inequality operator */
+	friend bool operator!=(const FParamTypeHandle& InLHS, const FParamTypeHandle& InRHS)
+	{
+		return InLHS.Value != InRHS.Value;
+	}
+
+	/** Type hash for TMap storage */
+	friend uint32 GetTypeHash(const FParamTypeHandle& InHandle)
+	{
+		return GetTypeHash(InHandle.Value);
+	}
+
+	/**
+	 * Check type invariants - type handle is either:
+	 * - Built-in and in known ranges
+	 * - Custom and with a non-zero custom type index
+	 * @return whether this type actually describes a type
+	 */
+	bool IsValid() const
+	{
+		const EParamType ParameterType = GetParameterType();
+		const uint32 CustomTypeIndex = GetCustomTypeIndex();
+		checkSlow(ValidateCustomTypeIndex(CustomTypeIndex));
+		const bool bIsValidCustomType = (ParameterType == EParamType::Custom && CustomTypeIndex != 0);
+		const bool bIsValidBuiltInType = ParameterType > EParamType::None && ParameterType <= EParamType::MaxBuiltIn;
+		return bIsValidBuiltInType || bIsValidCustomType;
+	}
+};
+
+}
\ No newline at end of file
diff --git a/Engine/Plugins/Experimental/AnimNext/Source/AnimNextInterfaceGraph/Internal/AnimGraph/AnimNode_AnimNextInterfaceGraph.h b/Engine/Plugins/Experimental/AnimNext/Source/AnimNextInterfaceGraph/Internal/AnimGraph/AnimNode_AnimNextInterfaceGraph.h
--- a/Engine/Plugins/Experimental/AnimNext/Source/AnimNextInterfaceGraph/Internal/AnimGraph/AnimNode_AnimNextInterfaceGraph.h
+++ b/Engine/Plugins/Experimental/AnimNext/Source/AnimNextInterfaceGraph/Internal/AnimGraph/AnimNode_AnimNextInterfaceGraph.h
@@ -3,11 +3,8 @@
 #pragma once
 
 #include "Animation/AnimNode_CustomProperty.h"
-#include "Animation/InputScaleBias.h"
 #include "AnimNextInterfaceGraph.h"
 #include "AnimNextInterfaceContext.h"
-#include "AnimNextInterfaceParamStorage.h"
-// --- ---
 #include "AnimNode_AnimNextInterfaceGraph.generated.h"
 
 
@@ -62,7 +59,7 @@ private:
 	int32 LODThreshold;
 
 	// This should come from the FPoseContext, but as I need persistence, I store the state here
-	UE::AnimNext::Interface::FState RootState; // TODO : Get rid of num elements ? Let the param itself have it
+	UE::AnimNext::FState RootState; // TODO : Get rid of num elements ? Let the param itself have it
 
 	// Delta time received accumulated in update and used at Evaluate (so we can receive multiple calls to Evaluate)
 	float GraphDeltaTime = 0.f;
diff --git a/Engine/Plugins/Experimental/AnimNext/Source/AnimNextInterfaceGraph/Internal/AnimNextInterfaceExecuteContext.h b/Engine/Plugins/Experimental/AnimNext/Source/AnimNextInterfaceGraph/Internal/AnimNextInterfaceExecuteContext.h
--- a/Engine/Plugins/Experimental/AnimNext/Source/AnimNextInterfaceGraph/Internal/AnimNextInterfaceExecuteContext.h
+++ b/Engine/Plugins/Experimental/AnimNext/Source/AnimNextInterfaceGraph/Internal/AnimNextInterfaceExecuteContext.h
@@ -10,7 +10,7 @@
 
 class IAnimNextInterface;
 
-namespace UE::AnimNext::Interface
+namespace UE::AnimNext
 {
 	struct FContext;
 }
@@ -29,7 +29,7 @@ struct FAnimNextInterfaceExecuteContext : public FRigVMExecuteContext
 	{
 	}
 
-	const UE::AnimNext::Interface::FContext& GetContext() const
+	const UE::AnimNext::FContext& GetContext() const
 	{
 		check(AnimNextInterfaceContext);
 		return *AnimNextInterfaceContext;
@@ -40,7 +40,7 @@ struct FAnimNextInterfaceExecuteContext : public FRigVMExecuteContext
 		return UnitContext;
 	}
 
-	void SetContextData(const IAnimNextInterface* InInterface, const UE::AnimNext::Interface::FContext& InAnimNextInterfaceContext, bool& bInResult)
+	void SetContextData(const IAnimNextInterface* InInterface, const UE::AnimNext::FContext& InAnimNextInterfaceContext, bool& bInResult)
 	{
 		Interface = InInterface;
 		AnimNextInterfaceContext = &InAnimNextInterfaceContext;
@@ -71,7 +71,7 @@ struct FAnimNextInterfaceExecuteContext : public FRigVMExecuteContext
 
 
 private:
-	const UE::AnimNext::Interface::FContext* AnimNextInterfaceContext;
+	const UE::AnimNext::FContext* AnimNextInterfaceContext;
 	const IAnimNextInterface* Interface;
 	bool* ResultPtr;
 	FRigUnitContext UnitContext;
diff --git a/Engine/Plugins/Experimental/AnimNext/Source/AnimNextInterfaceGraph/Internal/AnimNextInterfaceGraph.h b/Engine/Plugins/Experimental/AnimNext/Source/AnimNextInterfaceGraph/Internal/AnimNextInterfaceGraph.h
--- a/Engine/Plugins/Experimental/AnimNext/Source/AnimNextInterfaceGraph/Internal/AnimNextInterfaceGraph.h
+++ b/Engine/Plugins/Experimental/AnimNext/Source/AnimNextInterfaceGraph/Internal/AnimNextInterfaceGraph.h
@@ -5,6 +5,7 @@
 #include "CoreMinimal.h"
 #include "IAnimNextInterface.h"
 #include "RigVMCore/RigVM.h"
+#include "Param/ParamType.h"
 #include "AnimNextInterfaceGraph.generated.h"
 
 class UEdGraph;
@@ -26,7 +27,7 @@ namespace UE::AnimNext::InterfaceGraph
 }
 
 // A user-created graph of logic used to supply data
-UCLASS()
+UCLASS(BlueprintType)
 class ANIMNEXTINTERFACEGRAPH_API UAnimNextInterfaceGraph : public UObject, public IAnimNextInterface
 {
 	GENERATED_BODY()
@@ -36,9 +37,11 @@ class ANIMNEXTINTERFACEGRAPH_API UAnimNextInterfaceGraph : public UObject, publi
 	virtual void GetPreloadDependencies(TArray<UObject*>& OutDeps) override;
 
 	// IAnimAnimNextInterface interface
-	virtual FName GetReturnTypeNameImpl() const final override;
-	virtual const UScriptStruct* GetReturnTypeStructImpl() const final override;
-	virtual bool GetDataImpl(const UE::AnimNext::Interface::FContext& Context) const final override;
+	virtual UE::AnimNext::FParamTypeHandle GetReturnTypeHandleImpl() const final override;
+	virtual bool GetDataImpl(const UE::AnimNext::FContext& Context) const final override;
+
+	/** Set the return type of this graph */
+	void SetReturnTypeHandle(UE::AnimNext::FParamTypeHandle InHandle);
 
 	// Support rig VM execution
 	TArray<FRigVMExternalVariable> GetRigVMExternalVariables();
@@ -55,12 +58,10 @@ class ANIMNEXTINTERFACEGRAPH_API UAnimNextInterfaceGraph : public UObject, publi
 
 	UPROPERTY()
 	FRigVMRuntimeSettings VMRuntimeSettings;
-	
-	UPROPERTY()
-	FName ReturnTypeName;
 
+	/** Serialized return type handle as uint32 */
 	UPROPERTY()
-	TObjectPtr<UScriptStruct> ReturnTypeStruct;
+	FAnimNextParamType ReturnType;
 
 #if WITH_EDITORONLY_DATA
 	UPROPERTY(VisibleAnywhere, Instanced, Category = "Graph", meta = (ShowInnerProperties))
diff --git a/Engine/Plugins/Experimental/AnimNext/Source/AnimNextInterfaceGraph/Internal/AnimNextInterfaceUnitContext.h b/Engine/Plugins/Experimental/AnimNext/Source/AnimNextInterfaceGraph/Internal/AnimNextInterfaceUnitContext.h
--- a/Engine/Plugins/Experimental/AnimNext/Source/AnimNextInterfaceGraph/Internal/AnimNextInterfaceUnitContext.h
+++ b/Engine/Plugins/Experimental/AnimNext/Source/AnimNextInterfaceGraph/Internal/AnimNextInterfaceUnitContext.h
@@ -15,7 +15,7 @@ struct FAnimNextInterfaceUnitContext : public FRigUnitContext
 		, bResult(nullptr)
 	{}
 
-	FAnimNextInterfaceUnitContext(const IAnimNextInterface* InInterface, const UE::AnimNext::Interface::FContext& InAnimNextInterfaceContext, bool& bInResult)
+	FAnimNextInterfaceUnitContext(const IAnimNextInterface* InInterface, const UE::AnimNext::FContext& InAnimNextInterfaceContext, bool& bInResult)
 		: FRigUnitContext()
 		, Interface(InInterface)
 		, AnimNextInterfaceContext(&InAnimNextInterfaceContext)
@@ -23,6 +23,6 @@ struct FAnimNextInterfaceUnitContext : public FRigUnitContext
 	{}
 
 	const IAnimNextInterface* Interface;
-	const UE::AnimNext::Interface::FContext* AnimNextInterfaceContext;
+	const UE::AnimNext::FContext* AnimNextInterfaceContext;
 	bool* bResult;
 };
diff --git a/Engine/Plugins/Experimental/AnimNext/Source/AnimNextInterfaceGraph/Internal/AnimNextInterface_LODPose.h b/Engine/Plugins/Experimental/AnimNext/Source/AnimNextInterfaceGraph/Internal/AnimNextInterface_LODPose.h
--- a/Engine/Plugins/Experimental/AnimNext/Source/AnimNextInterfaceGraph/Internal/AnimNextInterface_LODPose.h
+++ b/Engine/Plugins/Experimental/AnimNext/Source/AnimNextInterfaceGraph/Internal/AnimNextInterface_LODPose.h
@@ -17,14 +17,13 @@ struct FAnimNextGraphReferencePose
 
 	FAnimNextGraphReferencePose() = default;
 
-	explicit FAnimNextGraphReferencePose(const UE::AnimNext::Interface::FAnimationReferencePose* InReferencePose)
+	explicit FAnimNextGraphReferencePose(const UE::AnimNext::FAnimationReferencePose* InReferencePose)
 		: ReferencePose(InReferencePose)
 	{
 	}
 
-	const UE::AnimNext::Interface::FAnimationReferencePose* ReferencePose = nullptr;
+	const UE::AnimNext::FAnimationReferencePose* ReferencePose = nullptr;
 };
-DECLARE_ANIM_NEXT_INTERFACE_PARAM_TYPE(ANIMNEXTINTERFACEGRAPH_API, FAnimNextGraphReferencePose, FAnimNextGraphReferencePose);
 
 USTRUCT(BlueprintType, meta = (DisplayName = "LODPose"))
 struct FAnimNextGraphLODPose
@@ -33,32 +32,29 @@ struct FAnimNextGraphLODPose
 
 	FAnimNextGraphLODPose() = default;
 
-	explicit FAnimNextGraphLODPose(const UE::AnimNext::Interface::FAnimationLODPose& InLODPose)
+	explicit FAnimNextGraphLODPose(const UE::AnimNext::FAnimationLODPose& InLODPose)
 		: LODPose(InLODPose)
 	{
 	}
-	explicit FAnimNextGraphLODPose(UE::AnimNext::Interface::FAnimationLODPose&& InLODPose)
+	explicit FAnimNextGraphLODPose(UE::AnimNext::FAnimationLODPose&& InLODPose)
 		: LODPose(MoveTemp(InLODPose))
 	{
 	}
 
-	UE::AnimNext::Interface::FAnimationLODPose LODPose;
+	UE::AnimNext::FAnimationLODPose LODPose;
 };
-DECLARE_ANIM_NEXT_INTERFACE_PARAM_TYPE(ANIMNEXTINTERFACEGRAPH_API, FAnimNextGraphLODPose, FAnimNextGraphLODPose);
-
 
 UCLASS()
 class UAnimNextInterfaceGraphLODPose : public UObject, public IAnimNextInterface
 {
 	GENERATED_BODY()
 
-	virtual FName GetReturnTypeNameImpl() const final override
+	virtual UE::AnimNext::FParamTypeHandle GetReturnTypeHandleImpl() const final override
 	{
-		static const FName Name = TNameOf<FAnimNextGraphLODPose>::GetName();
-		return Name;
+		return UE::AnimNext::FParamTypeHandle::GetHandle<FAnimNextGraphLODPose>();
 	}
 
-	virtual bool GetDataImpl(const UE::AnimNext::Interface::FContext& Context) const override
+	virtual bool GetDataImpl(const UE::AnimNext::FContext& Context) const override
 	{
 		checkf(false, TEXT("UAnimNextInterfaceLODPose::GetDataImpl must be overridden"));
 		return false;
@@ -71,7 +67,7 @@ class UAnimNextInterface_LODPose_Literal : public UAnimNextInterfaceGraphLODPose
 	GENERATED_BODY()
 
 	// IAnimNextInterface interface
-	virtual bool GetDataImpl(const UE::AnimNext::Interface::FContext& Context) const final override
+	virtual bool GetDataImpl(const UE::AnimNext::FContext& Context) const final override
 	{
 		Context.SetResult(Value);
 		return true;
diff --git a/Engine/Plugins/Experimental/AnimNext/Source/AnimNextInterfaceGraph/Internal/RigUnit_AnimNextInterfaceEndExecution.h b/Engine/Plugins/Experimental/AnimNext/Source/AnimNextInterfaceGraph/Internal/RigUnit_AnimNextInterfaceEndExecution.h
--- a/Engine/Plugins/Experimental/AnimNext/Source/AnimNextInterfaceGraph/Internal/RigUnit_AnimNextInterfaceEndExecution.h
+++ b/Engine/Plugins/Experimental/AnimNext/Source/AnimNextInterfaceGraph/Internal/RigUnit_AnimNextInterfaceEndExecution.h
@@ -5,7 +5,6 @@
 #include "Units/RigUnit.h"
 #include "AnimNextInterfaceExecuteContext.h"
 #include "AnimNextInterfaceUnitContext.h"
-#include "AnimNextInterfaceTypes.h"
 #include "AnimNextInterface_LODPose.h"
 // --- ---
 #include "RigUnit_AnimNextInterfaceEndExecution.generated.h"
diff --git a/Engine/Plugins/Experimental/AnimNext/Source/AnimNextInterfaceGraph/Private/AnimGraph/AnimNode_AnimNextInterfaceGraph.cpp b/Engine/Plugins/Experimental/AnimNext/Source/AnimNextInterfaceGraph/Private/AnimGraph/AnimNode_AnimNextInterfaceGraph.cpp
--- a/Engine/Plugins/Experimental/AnimNext/Source/AnimNextInterfaceGraph/Private/AnimGraph/AnimNode_AnimNextInterfaceGraph.cpp
+++ b/Engine/Plugins/Experimental/AnimNext/Source/AnimNextInterfaceGraph/Private/AnimGraph/AnimNode_AnimNextInterfaceGraph.cpp
@@ -10,7 +10,6 @@
 #include "GameFramework/Actor.h"
 #include "Animation/NodeMappingContainer.h"
 #include "AnimationRuntime.h"
-#include "AnimNextInterfaceTypes.h"
 #include "AnimNextInterface.h"
 #include "AnimationDataRegistry.h"
 #include "AnimationGenerationTools.h"
@@ -18,6 +17,7 @@
 #include "AnimNextInterface_LODPose.h"
 #include "RigUnit_AnimNextAnimSequence.h" // TEST
 #include "Animation/AnimSequence.h"
+#include "Animation/AnimSequence.h" // TEST
 #include "Engine/SkeletalMesh.h"
 #include "BoneContainer.h"
 
@@ -31,7 +31,6 @@ FAnimNode_AnimNextInterfaceGraph::FAnimNode_AnimNextInterfaceGraph()
 	: FAnimNode_CustomProperty()
 	, AnimNextInterfaceGraph(nullptr)
 	, LODThreshold(INDEX_NONE)
-	, RootState(1) // TODO : Get rid of num elements ? Let the param itself have it
 {
 }
 
@@ -97,7 +96,7 @@ void FAnimNode_AnimNextInterfaceGraph::Evaluate_AnyThread(FPoseContext & Output)
 {
 	DECLARE_SCOPE_HIERARCHICAL_COUNTER_FUNC()
 
-	using namespace UE::AnimNext::Interface;
+	using namespace UE::AnimNext;
 
 	FPoseContext SourcePose(Output);
 
@@ -120,11 +119,11 @@ void FAnimNode_AnimNextInterfaceGraph::Evaluate_AnyThread(FPoseContext & Output)
 
 	// TODO : In a normal graph, FState/FParamStorage/FContext would come from FAnimationUpdateContext, here I have to create one the fly for the AnimNext graph
 	//UE::AnimNext::Interface::FState RootState; // moved to .h as I need state persistence in the AnimNext graph
-	UE::AnimNext::Interface::FParamStorage ParamStorage(5, 0, 0);	// TODO : see if we can have state and params together and optional. Also, default size
+	UE::AnimNext::FParamStorage ParamStorage(5, 0, 0);	// TODO : see if we can have state and params together and optional. Also, default size
 	
 	// TODO : Using AnimInstanceProxy->GetDeltaSeconds() makes the preview to advance  multiple times when debug options are activated (i.e. ShowUncompressedAnim)
 	// See where to get / how to calculate the correct delta value
-	UE::AnimNext::Interface::FContext RootContext(GraphDeltaTime, RootState, ParamStorage);
+	UE::AnimNext::FContext RootContext(GraphDeltaTime, RootState, ParamStorage);
 	GraphDeltaTime = 0.f; // Reset for the case that we receive multiple calls to Evaluate (debug options)
 
 	FAnimNextGraphReferencePose GraphReferencePose(&RefPose);
diff --git a/Engine/Plugins/Experimental/AnimNext/Source/AnimNextInterfaceGraph/Private/AnimNextInterfaceGraph.cpp b/Engine/Plugins/Experimental/AnimNext/Source/AnimNextInterfaceGraph/Private/AnimNextInterfaceGraph.cpp
--- a/Engine/Plugins/Experimental/AnimNext/Source/AnimNextInterfaceGraph/Private/AnimNextInterfaceGraph.cpp
+++ b/Engine/Plugins/Experimental/AnimNext/Source/AnimNextInterfaceGraph/Private/AnimNextInterfaceGraph.cpp
@@ -15,17 +15,12 @@ const FName EntryPointName("GetData");
 const FName ResultName("Result");
 }
 
-FName UAnimNextInterfaceGraph::GetReturnTypeNameImpl() const
+UE::AnimNext::FParamTypeHandle UAnimNextInterfaceGraph::GetReturnTypeHandleImpl() const
 {
-	return ReturnTypeName;
+	return ReturnType.GetHandle();
 }
 
-const UScriptStruct* UAnimNextInterfaceGraph::GetReturnTypeStructImpl() const
-{
-	return ReturnTypeStruct;
-}
-
-bool UAnimNextInterfaceGraph::GetDataImpl(const UE::AnimNext::Interface::FContext& Context) const
+bool UAnimNextInterfaceGraph::GetDataImpl(const UE::AnimNext::FContext& Context) const
 {
 	bool bResult = true;
 	
@@ -41,6 +36,11 @@ bool UAnimNextInterfaceGraph::GetDataImpl(const UE::AnimNext::Interface::FContex
 	return bResult;
 }
 
+void UAnimNextInterfaceGraph::SetReturnTypeHandle(UE::AnimNext::FParamTypeHandle InHandle)
+{
+	ReturnType = InHandle.GetType();
+}
+
 TArray<FRigVMExternalVariable> UAnimNextInterfaceGraph::GetRigVMExternalVariables()
 {
 	return TArray<FRigVMExternalVariable>(); 
diff --git a/Engine/Plugins/Experimental/AnimNext/Source/AnimNextInterfaceGraph/Private/AnimNextInterface_Bool.cpp b/Engine/Plugins/Experimental/AnimNext/Source/AnimNextInterfaceGraph/Private/AnimNextInterface_Bool.cpp
--- a/Engine/Plugins/Experimental/AnimNext/Source/AnimNextInterfaceGraph/Private/AnimNextInterface_Bool.cpp
+++ b/Engine/Plugins/Experimental/AnimNext/Source/AnimNextInterfaceGraph/Private/AnimNextInterface_Bool.cpp
@@ -2,51 +2,46 @@
 
 #include "AnimNextInterface_Bool.h"
 #include "AnimNextInterface.h"
-#include "AnimNextInterfaceTypes.h"
-#include "AnimNextInterfaceKernel.h"
 
 #include UE_INLINE_GENERATED_CPP_BY_NAME(AnimNextInterface_Bool)
 
-bool UAnimNextInterface_Bool_And::GetDataImpl(const UE::AnimNext::Interface::FContext& Context) const
+bool UAnimNextInterface_Bool_And::GetDataImpl(const UE::AnimNext::FContext& Context) const
 {
-	using namespace UE::AnimNext::Interface;
+	using namespace UE::AnimNext;
 	
 	check(Inputs.Num() > 0);
 
-	bool bResult = true;
-	TAllocParam<bool> IntermediateParam(Context);
-	TParam<bool> Result = Context.GetResult<bool>();
+	bool bInterfaceResult = true;
+	bool bIntermediateResult = false;
+
+	// Get result to write to
+	bool& bOutResult = Context.GetResult<bool>();
 	
 	for(const TScriptInterface<IAnimNextInterface>& Input : Inputs)
 	{
-		bResult &= UE::AnimNext::Interface::GetDataSafe(Input, Context, IntermediateParam);
-
-		FKernel::Run(Context,
-			[](bool& OutResult, bool InIntermediate)
-			{
-				OutResult = OutResult && InIntermediate;
-			},
-			Result, IntermediateParam);
+		// Get any inputs we may have
+		bInterfaceResult &= Interface::GetDataSafe(Input, Context, bIntermediateResult);
+
+		// AND the inputs
+		bOutResult = bOutResult && bIntermediateResult;
 	}
 
-	return bResult;
+	return bInterfaceResult;
 }
 
-bool UAnimNextInterface_Bool_Not::GetDataImpl(const UE::AnimNext::Interface::FContext& Context) const
+bool UAnimNextInterface_Bool_Not::GetDataImpl(const UE::AnimNext::FContext& Context) const
 {
-	using namespace UE::AnimNext::Interface;
+	using namespace UE::AnimNext;
 
-	TParam<bool> Result = Context.GetResult<bool>();
+	// Get result to write to
+	bool& bOutResult = Context.GetResult<bool>();
 	
-	TAllocParam<bool> CurrentValue(Context);
-	bool bResult = UE::AnimNext::Interface::GetDataSafe(Input, Context, CurrentValue);
+	// Get any input we may have
+	bool bInputResult = false;
+	const bool bInterfaceResult = Interface::GetDataSafe(Input, Context, bInputResult);
 
-	FKernel::Run(Context,
-		[](bool& OutResult, bool InCurrent)
-		{
-			OutResult = !InCurrent;
-		},
-		Result, CurrentValue);
+	// NOT the input, write to output
+	bOutResult = !bInputResult;
 
-	return bResult;
+	return bInterfaceResult;
 }
diff --git a/Engine/Plugins/Experimental/AnimNext/Source/AnimNextInterfaceGraph/Private/AnimNextInterface_Bool.h b/Engine/Plugins/Experimental/AnimNext/Source/AnimNextInterfaceGraph/Private/AnimNextInterface_Bool.h
--- a/Engine/Plugins/Experimental/AnimNext/Source/AnimNextInterfaceGraph/Private/AnimNextInterface_Bool.h
+++ b/Engine/Plugins/Experimental/AnimNext/Source/AnimNextInterfaceGraph/Private/AnimNextInterface_Bool.h
@@ -15,7 +15,7 @@ class UAnimNextInterface_Bool_Literal : public UObject, public IAnimNextInterfac
 	ANIM_NEXT_INTERFACE_RETURN_TYPE(bool)
 
 	// IAnimNextInterface interface
-	virtual bool GetDataImpl(const UE::AnimNext::Interface::FContext& Context) const final override
+	virtual bool GetDataImpl(const UE::AnimNext::FContext& Context) const final override
 	{
 		Context.SetResult(Value);
 		return true;
@@ -32,7 +32,7 @@ class UAnimNextInterface_Bool_And : public UObject, public IAnimNextInterface
 	ANIM_NEXT_INTERFACE_RETURN_TYPE(bool)
 
 	// IAnimNextInterface interface
-	virtual bool GetDataImpl(const UE::AnimNext::Interface::FContext& Context) const final override;
+	virtual bool GetDataImpl(const UE::AnimNext::FContext& Context) const final override;
 
 	UPROPERTY(EditAnywhere, Meta=(AnimNextInterfaceType="bool"), Category = "Parameters")
 	TArray<TScriptInterface<IAnimNextInterface>> Inputs;
@@ -46,7 +46,7 @@ class UAnimNextInterface_Bool_Not : public UObject, public IAnimNextInterface
 	ANIM_NEXT_INTERFACE_RETURN_TYPE(bool)
 
 	// IAnimNextInterface interface
-	virtual bool GetDataImpl(const UE::AnimNext::Interface::FContext& Context) const final override;
+	virtual bool GetDataImpl(const UE::AnimNext::FContext& Context) const final override;
 
 	UPROPERTY(EditAnywhere, Meta=(AnimNextInterfaceType="bool"), Category = "Parameters")
 	TScriptInterface<IAnimNextInterface> Input;
diff --git a/Engine/Plugins/Experimental/AnimNext/Source/AnimNextInterfaceGraph/Private/AnimNextInterface_Float.cpp b/Engine/Plugins/Experimental/AnimNext/Source/AnimNextInterfaceGraph/Private/AnimNextInterface_Float.cpp
--- a/Engine/Plugins/Experimental/AnimNext/Source/AnimNextInterfaceGraph/Private/AnimNextInterface_Float.cpp
+++ b/Engine/Plugins/Experimental/AnimNext/Source/AnimNextInterfaceGraph/Private/AnimNextInterface_Float.cpp
@@ -2,108 +2,79 @@
 
 #include "AnimNextInterface_Float.h"
 #include "AnimNextInterface.h"
-#include "AnimNextInterfaceTypes.h"
-#include "AnimNextInterfaceKernel.h"
 
 #include UE_INLINE_GENERATED_CPP_BY_NAME(AnimNextInterface_Float)
 
 #define LOCTEXT_NAMESPACE "AnimNextInterface"
 
-bool UAnimNextInterface_Float_Multiply::GetDataImpl(const UE::AnimNext::Interface::FContext& Context) const
+bool UAnimNextInterface_Float_Multiply::GetDataImpl(const UE::AnimNext::FContext& Context) const
 {
-	using namespace UE::AnimNext::Interface;
+	using namespace UE::AnimNext;
 	
 	check(Inputs.Num() > 0);
 
 	bool bResult = true;
-	TAllocParam<float> IntermediateParam(Context);
-	TParam<float> Result = Context.GetResult<float>();
+	float Intermediate = 0.0f;
+	float& OutResult = Context.GetResult<float>();
 	
 	for(const TScriptInterface<IAnimNextInterface>& Input : Inputs)
 	{
-		bResult &= UE::AnimNext::Interface::GetDataSafe(Input, Context, IntermediateParam);
-
-		FKernel::Run(Context,
-			[](float& OutResult, float InIntermediate)
-			{
-				OutResult *= InIntermediate;
-			},
-			Result, IntermediateParam);
+		bResult &= Interface::GetDataSafe(Input, Context, Intermediate);
+
+		OutResult *= Intermediate;
 	}
 
 	return bResult;
 }
 
-bool UAnimNextInterface_Float_InterpTo::GetDataImpl(const UE::AnimNext::Interface::FContext& Context) const
+bool UAnimNextInterface_Float_InterpTo::GetDataImpl(const UE::AnimNext::FContext& Context) const
 {
-	using namespace UE::AnimNext::Interface;
+	using namespace UE::AnimNext;
 
-	TParam<float> Result = Context.GetResult<float>();
+	float& OutResult = Context.GetResult<float>();
 	const float DeltaTime = Context.GetDeltaTime();
 	
-	TAllocParam<float> CurrentValue(Context);
-	bool bResult = UE::AnimNext::Interface::GetDataSafe(Current, Context, CurrentValue);
+	float CurrentValue = 0.0f;
+	bool bResult = Interface::GetDataSafe(Current, Context, CurrentValue);
 
-	TAllocParam<float> TargetValue(Context);
-	bResult &= UE::AnimNext::Interface::GetDataSafe(Target, Context, TargetValue);
+	float TargetValue = 0.0f;
+	bResult &= Interface::GetDataSafe(Target, Context, TargetValue);
 
-	TAllocParam<float> SpeedValue(Context);
-	bResult &= UE::AnimNext::Interface::GetDataSafe(Speed, Context, SpeedValue);
+	float SpeedValue = 0.0f;
+	bResult &= Interface::GetDataSafe(Speed, Context, SpeedValue);
 
-	FKernel::Run(Context,
-		[DeltaTime](float& OutResult, float InCurrent, float InTarget, float InSpeed)
-		{
-			OutResult = FMath::FInterpConstantTo(InCurrent, InTarget, DeltaTime, InSpeed);
-		},
-		Result, CurrentValue, TargetValue, SpeedValue);
+	OutResult = FMath::FInterpConstantTo(CurrentValue, TargetValue, DeltaTime, SpeedValue);
 
 	return bResult;
 }
 
-struct FSpringInterpState
-{
-	float Value;
-	float ValueRate;
-};
-
-IMPLEMENT_ANIM_NEXT_INTERFACE_STATE_TYPE(FSpringInterpState, SpringInterpState);
-
-bool UAnimNextInterface_Float_SpringInterp::GetDataImpl(const UE::AnimNext::Interface::FContext& Context) const
+bool UAnimNextInterface_Float_SpringInterp::GetDataImpl(const UE::AnimNext::FContext& Context) const
 {
-	using namespace UE::AnimNext::Interface;
+	using namespace UE::AnimNext;
 	
-	const TParam<FSpringInterpState> State = Context.GetState<FSpringInterpState>(this, 0);
+	FAnimNextInterface_Float_SpringInterpState& State = Context.GetState<FAnimNextInterface_Float_SpringInterpState>(this, 0);
 	const float DeltaTime = Context.GetDeltaTime();
 
-	TAllocParam<float> TargetValueParam(Context);
-	bool bResult = UE::AnimNext::Interface::GetDataSafe(TargetValue, Context, TargetValueParam);
-
-	TAllocParam<float> TargetValueRateParam(Context);
-	bResult &= UE::AnimNext::Interface::GetDataSafe(TargetValueRate, Context, TargetValueRateParam);
-
-	TAllocParam<float> SmoothingTimeParam(Context);
-	bResult &= UE::AnimNext::Interface::GetDataSafe(SmoothingTime, Context, SmoothingTimeParam);
-
-	TAllocParam<float> DampingRatioParam(Context);
-	bResult &= UE::AnimNext::Interface::GetDataSafe(DampingRatio, Context, DampingRatioParam);
-
-	FKernel::Run(Context,
-	[DeltaTime](FSpringInterpState& InOutState,
-					float InTargetValue,
-					float InTargetValueRate,
-					float InSmoothingTime,
-					float InDampingRatio)
-		{
-			FMath::SpringDamperSmoothing(
-				InOutState.Value,
-				InOutState.ValueRate,
-				InTargetValue,
-				InTargetValueRate,
-				DeltaTime,
-				InSmoothingTime,
-				InDampingRatio);
-		},
-		State, TargetValueParam, TargetValueRateParam, SmoothingTimeParam, DampingRatioParam);
+	float TargetValue = 0.0f;
+	bool bResult = Interface::GetDataSafe(Target, Context, TargetValue);
+
+	float TargetRateValue = 0.0f;
+	bResult &= Interface::GetDataSafe(TargetRate, Context, TargetRateValue);
+
+	float SmoothingTimeValue = 0.0f;
+	bResult &= Interface::GetDataSafe(SmoothingTime, Context, SmoothingTimeValue);
+
+	float DampingRatioValue = 0.0f;
+	bResult &= Interface::GetDataSafe(DampingRatio, Context, DampingRatioValue);
+
+	FMath::SpringDamperSmoothing(
+		State.Value,
+		State.ValueRate,
+		TargetValue,
+		TargetRateValue,
+		DeltaTime,
+		SmoothingTimeValue,
+		DampingRatioValue);
 
 	return bResult;
 }
diff --git a/Engine/Plugins/Experimental/AnimNext/Source/AnimNextInterfaceGraph/Private/AnimNextInterface_Float.h b/Engine/Plugins/Experimental/AnimNext/Source/AnimNextInterfaceGraph/Private/AnimNextInterface_Float.h
--- a/Engine/Plugins/Experimental/AnimNext/Source/AnimNextInterfaceGraph/Private/AnimNextInterface_Float.h
+++ b/Engine/Plugins/Experimental/AnimNext/Source/AnimNextInterfaceGraph/Private/AnimNextInterface_Float.h
@@ -13,13 +13,9 @@ class UAnimNextInterfaceFloat : public UObject, public IAnimNextInterface
 {
 	GENERATED_BODY()
 
-	virtual FName GetReturnTypeNameImpl() const final override
-	{
-		static const FName Name = TNameOf<float>::GetName();
-		return Name;
-	}
+	ANIM_NEXT_INTERFACE_RETURN_TYPE(float)
 
-	virtual bool GetDataImpl(const UE::AnimNext::Interface::FContext& Context) const override
+	virtual bool GetDataImpl(const UE::AnimNext::FContext& Context) const override
 	{
 		checkf(false, TEXT("UAnimNextInterfaceFloat::GetDataImpl must be overridden"));
 		return false;
@@ -32,7 +28,7 @@ class UAnimNextInterface_Float_Literal : public UAnimNextInterfaceFloat
 	GENERATED_BODY()
 
 	// IAnimNextInterface interface
-	virtual bool GetDataImpl(const UE::AnimNext::Interface::FContext& Context) const final override
+	virtual bool GetDataImpl(const UE::AnimNext::FContext& Context) const final override
 	{
 		Context.SetResult(Value);
 		return true;
@@ -48,7 +44,7 @@ class UAnimNextInterface_Float_Multiply : public UAnimNextInterfaceFloat
 	GENERATED_BODY()
 
 	// IAnimNextInterface interface
-	virtual bool GetDataImpl(const UE::AnimNext::Interface::FContext& Context) const final override;
+	virtual bool GetDataImpl(const UE::AnimNext::FContext& Context) const final override;
 
 	UPROPERTY(EditAnywhere, Category = "Parameters")
 	TArray<TScriptInterface<IAnimNextInterface>> Inputs;
@@ -60,7 +56,7 @@ class UAnimNextInterface_Float_InterpTo : public UAnimNextInterfaceFloat
 	GENERATED_BODY()
 
 	// IAnimNextInterface interface
-	virtual bool GetDataImpl(const UE::AnimNext::Interface::FContext& Context) const final override;
+	virtual bool GetDataImpl(const UE::AnimNext::FContext& Context) const final override;
 
 	UPROPERTY(EditAnywhere, Category = "Parameters")
 	TScriptInterface<IAnimNextInterface> Current;
@@ -78,26 +74,38 @@ class UAnimNextInterface_Float_DeltaTime : public UAnimNextInterfaceFloat
 	GENERATED_BODY()
 
 	// IAnimNextInterface interface
-	virtual bool GetDataImpl(const UE::AnimNext::Interface::FContext& Context) const final override
+	virtual bool GetDataImpl(const UE::AnimNext::FContext& Context) const final override
 	{
 		Context.SetResult(Context.GetDeltaTime());
 		return true;
 	}
 };
 
+USTRUCT()
+struct FAnimNextInterface_Float_SpringInterpState
+{
+	GENERATED_BODY()
+
+	UPROPERTY()
+	float Value = 0.0f;
+
+	UPROPERTY()
+	float ValueRate = 0.0f;
+};
+
 UCLASS()
 class UAnimNextInterface_Float_SpringInterp : public UAnimNextInterfaceFloat
 {
 	GENERATED_BODY()
 
 	// IAnimNextInterface interface
-	virtual bool GetDataImpl(const UE::AnimNext::Interface::FContext& Context) const final override;
+	virtual bool GetDataImpl(const UE::AnimNext::FContext& Context) const final override;
 
 	UPROPERTY(EditAnywhere, Category = "Parameters")
-	TScriptInterface<IAnimNextInterface> TargetValue;
+	TScriptInterface<IAnimNextInterface> Target;
 
 	UPROPERTY(EditAnywhere, Category = "Parameters")
-	TScriptInterface<IAnimNextInterface> TargetValueRate;
+	TScriptInterface<IAnimNextInterface> TargetRate;
 	
 	UPROPERTY(EditAnywhere, Category = "Parameters")
 	TScriptInterface<IAnimNextInterface> SmoothingTime;
diff --git a/Engine/Plugins/Experimental/AnimNext/Source/AnimNextInterfaceGraph/Private/AnimNextInterface_LODPose.cpp b/Engine/Plugins/Experimental/AnimNext/Source/AnimNextInterfaceGraph/Private/AnimNextInterface_LODPose.cpp
deleted file mode 100644
diff --git a/Engine/Plugins/Experimental/AnimNext/Source/AnimNextInterfaceGraph/Private/AnimNextInterface_Name.h b/Engine/Plugins/Experimental/AnimNext/Source/AnimNextInterfaceGraph/Private/AnimNextInterface_Name.h
--- a/Engine/Plugins/Experimental/AnimNext/Source/AnimNextInterfaceGraph/Private/AnimNextInterface_Name.h
+++ b/Engine/Plugins/Experimental/AnimNext/Source/AnimNextInterfaceGraph/Private/AnimNextInterface_Name.h
@@ -5,6 +5,7 @@
 #include "CoreMinimal.h"
 #include "IAnimNextInterface.h"
 #include "AnimNextInterfaceContext.h"
+#include "AnimNextInterfaceTypes.h"
 #include "AnimNextInterface_Name.generated.h"
 
 UCLASS()
@@ -12,12 +13,9 @@ class UAnimNextInterface_Name_Literal : public UObject, public IAnimNextInterfac
 {
 	GENERATED_BODY()
 
-	virtual FName GetReturnTypeNameImpl() const final override
-	{
-		return NAME_Name;
-	}
+	ANIM_NEXT_INTERFACE_RETURN_TYPE(FName)
 
-	virtual bool GetDataImpl(const UE::AnimNext::Interface::FContext& Context) const final override
+	virtual bool GetDataImpl(const UE::AnimNext::FContext& Context) const final override
 	{
 		Context.SetResult(Value);
 		return true;
diff --git a/Engine/Plugins/Experimental/AnimNext/Source/AnimNextInterfaceGraph/Private/AnimNextInterface_Wrapper.cpp b/Engine/Plugins/Experimental/AnimNext/Source/AnimNextInterfaceGraph/Private/AnimNextInterface_Wrapper.cpp
--- a/Engine/Plugins/Experimental/AnimNext/Source/AnimNextInterfaceGraph/Private/AnimNextInterface_Wrapper.cpp
+++ b/Engine/Plugins/Experimental/AnimNext/Source/AnimNextInterfaceGraph/Private/AnimNextInterface_Wrapper.cpp
@@ -1,22 +1,17 @@
 ﻿// Copyright Epic Games, Inc. All Rights Reserved.
 
 #include "AnimNextInterface_Wrapper.h"
+#include "Param/ParamTypeHandle.h"
 
 #include UE_INLINE_GENERATED_CPP_BY_NAME(AnimNextInterface_Wrapper)
 
-FName UAnimNextInterface_Wrapper::GetReturnTypeNameImpl() const
+UE::AnimNext::FParamTypeHandle UAnimNextInterface_Wrapper::GetReturnTypeHandleImpl() const
 {
 	// We mimic the return type of our output
-	return Output.GetInterface() ? Output.GetInterface()->GetReturnTypeName() : NAME_None;
+	return Output.GetInterface() ? Output.GetInterface()->GetReturnTypeHandle() : UE::AnimNext::FParamTypeHandle();
 }
 
-const UScriptStruct* UAnimNextInterface_Wrapper::GetReturnTypeStructImpl() const
-{
-	// We mimic the return type of our output
-	return Output.GetInterface() ? Output.GetInterface()->GetReturnTypeStruct() : nullptr;
-}
-
-bool UAnimNextInterface_Wrapper::GetDataImpl(const UE::AnimNext::Interface::FContext& Context) const
+bool UAnimNextInterface_Wrapper::GetDataImpl(const UE::AnimNext::FContext& Context) const
 {
 	// Parameterize output with inputs
 	//const UE::AnimNext::Interface::FContext ParameterizedContext = Context.WithParameters(Inputs);
diff --git a/Engine/Plugins/Experimental/AnimNext/Source/AnimNextInterfaceGraph/Private/AnimNextInterface_Wrapper.h b/Engine/Plugins/Experimental/AnimNext/Source/AnimNextInterfaceGraph/Private/AnimNextInterface_Wrapper.h
--- a/Engine/Plugins/Experimental/AnimNext/Source/AnimNextInterfaceGraph/Private/AnimNextInterface_Wrapper.h
+++ b/Engine/Plugins/Experimental/AnimNext/Source/AnimNextInterfaceGraph/Private/AnimNextInterface_Wrapper.h
@@ -13,9 +13,8 @@ class UAnimNextInterface_Wrapper : public UObject, public IAnimNextInterface
 
 private:
 	// IAnimAnimNextInterface interface
-	virtual FName GetReturnTypeNameImpl() const final override;
-	virtual const UScriptStruct* GetReturnTypeStructImpl() const final override;
-	virtual bool GetDataImpl(const UE::AnimNext::Interface::FContext& Context) const final override;
+	virtual UE::AnimNext::FParamTypeHandle GetReturnTypeHandleImpl() const final override;
+	virtual bool GetDataImpl(const UE::AnimNext::FContext& Context) const final override;
 
 //	UPROPERTY(EditAnywhere, Category = "Parameters")
 //	TArray<FAnimNextInterfaceParameter> Inputs;
diff --git a/Engine/Plugins/Experimental/AnimNext/Source/AnimNextInterfaceGraph/Private/RigUnit_AnimNextAnimSequence.cpp b/Engine/Plugins/Experimental/AnimNext/Source/AnimNextInterfaceGraph/Private/RigUnit_AnimNextAnimSequence.cpp
--- a/Engine/Plugins/Experimental/AnimNext/Source/AnimNextInterfaceGraph/Private/RigUnit_AnimNextAnimSequence.cpp
+++ b/Engine/Plugins/Experimental/AnimNext/Source/AnimNextInterfaceGraph/Private/RigUnit_AnimNextAnimSequence.cpp
@@ -5,29 +5,16 @@
 #include "AnimNextInterfaceExecuteContext.h"
 #include "AnimNextInterface.h"
 #include "AnimNextInterfaceTypes.h"
-#include "AnimNextInterfaceKernel.h"
 #include "AnimationDecompressionTools.h"
 #include "Animation/AnimSequence.h"
 
 #include UE_INLINE_GENERATED_CPP_BY_NAME(RigUnit_AnimNextAnimSequence)
 
-
-IMPLEMENT_ANIM_NEXT_INTERFACE_PARAM_TYPE(FAnimNextGraph_AnimSequence, FAnimNextGraph_AnimSequence);
-
-
-struct FAnimSequencePlayerState
-{
-	float InternalTimeAccumulator = 0.0f;
-	float PrevInternalTimeAccumulator = 0.0f;
-};
-
-IMPLEMENT_ANIM_NEXT_INTERFACE_STATE_TYPE(FAnimSequencePlayerState, AnimSequencePlayerState);
-
 // --- FRigUnit_AnimNext_AnimSequenceAsset --- 
 
 FRigUnit_AnimNext_AnimSequenceAsset_Execute()
 {
-	using namespace UE::AnimNext::Interface;
+	using namespace UE::AnimNext;
 
 	//Sequence.AnimSequence = AnimSequence; // TODO : Enable once we have RigVM object support enabled
 }
@@ -36,7 +23,7 @@ FRigUnit_AnimNext_AnimSequenceAsset_Execute()
 
 FRigUnit_AnimNext_AnimSequencePlayer_Initialize()
 {
-	using namespace UE::AnimNext::Interface;
+	using namespace UE::AnimNext;
 
 	const FContext& AnimNextInterfaceContext = ExecuteContext.GetContext();
 
@@ -44,17 +31,17 @@ FRigUnit_AnimNext_AnimSequencePlayer_Initialize()
 	if (AnimSequenceBase != nullptr)
 	{
 		// Get or Create the internal state
-		const TParam<FAnimSequencePlayerState> AnimSequencePlayerState = AnimNextInterfaceContext.GetState<FAnimSequencePlayerState>(ExecuteContext.GetInterface(), 0);
+		FAnimSequencePlayerState& AnimSequencePlayerState = AnimNextInterfaceContext.GetState<FAnimSequencePlayerState>(ExecuteContext.GetInterface(), 0);
 
 		const float SequenceLength = AnimSequenceBase->GetPlayLength();
-		AnimSequencePlayerState->InternalTimeAccumulator = FMath::Clamp(Parameters.StartPosition, 0.f, SequenceLength);
-		AnimSequencePlayerState->PrevInternalTimeAccumulator = AnimSequencePlayerState->InternalTimeAccumulator;
+		AnimSequencePlayerState.InternalTimeAccumulator = FMath::Clamp(Parameters.StartPosition, 0.f, SequenceLength);
+		AnimSequencePlayerState.PrevInternalTimeAccumulator = AnimSequencePlayerState.InternalTimeAccumulator;
 	}
 }
 
 FRigUnit_AnimNext_AnimSequencePlayer_Execute()
 {
-	using namespace UE::AnimNext::Interface;
+	using namespace UE::AnimNext;
 
 	const FContext& AnimNextInterfaceContext = ExecuteContext.GetContext();
 	
@@ -63,7 +50,7 @@ FRigUnit_AnimNext_AnimSequencePlayer_Execute()
 	{
 		// Get or Create the internal state
 		//const TParam<FAnimSequencePlayerState> AnimSequencePlayerState = AnimNextInterfaceContext.GetState<FAnimSequencePlayerState>(ExecuteContext.GetInterface(), 0);
-		FAnimSequencePlayerState& AnimSequencePlayerState = AnimNextInterfaceContext.GetState<FAnimSequencePlayerState, EStatePersistence::Permanent>(ExecuteContext.GetInterface(), 0).GetDataChecked();
+		FAnimSequencePlayerState& AnimSequencePlayerState = AnimNextInterfaceContext.GetState<FAnimSequencePlayerState, EStatePersistence::Permanent>(ExecuteContext.GetInterface(), 0);
 
 		const float DeltaTime = AnimNextInterfaceContext.GetDeltaTime();
 		const float EffectiveDelta = FMath::IsNearlyZero(DeltaTime) || FMath::IsNearlyZero(Parameters.PlayRate) ? 0.f : DeltaTime * Parameters.PlayRate;
diff --git a/Engine/Plugins/Experimental/AnimNext/Source/AnimNextInterfaceGraph/Private/RigUnit_AnimNextAnimSequence.h b/Engine/Plugins/Experimental/AnimNext/Source/AnimNextInterfaceGraph/Private/RigUnit_AnimNextAnimSequence.h
--- a/Engine/Plugins/Experimental/AnimNext/Source/AnimNextInterfaceGraph/Private/RigUnit_AnimNextAnimSequence.h
+++ b/Engine/Plugins/Experimental/AnimNext/Source/AnimNextInterfaceGraph/Private/RigUnit_AnimNextAnimSequence.h
@@ -11,6 +11,18 @@
 
 class UAnimSequenceBase;
 
+USTRUCT()
+struct FAnimSequencePlayerState
+{
+	GENERATED_BODY()
+
+	UPROPERTY()
+	float InternalTimeAccumulator = 0.0f;
+
+	UPROPERTY()
+	float PrevInternalTimeAccumulator = 0.0f;
+};
+
 USTRUCT(BlueprintType, meta = (DisplayName = "AnimSequence"))
 struct FAnimNextGraph_AnimSequence
 {
@@ -26,10 +38,6 @@ struct FAnimNextGraph_AnimSequence
 	const UAnimSequenceBase* AnimSequence = nullptr;
 };
 
-DECLARE_ANIM_NEXT_INTERFACE_PARAM_TYPE(ANIMNEXTINTERFACEGRAPH_API, FAnimNextGraph_AnimSequence, FAnimNextGraph_AnimSequence);
-
-
-
 struct FAnimNextInterfaceUnitContext;
 class UAnimNextInterfaceGraph;
 class UAnimSequenceBase;
diff --git a/Engine/Plugins/Experimental/AnimNext/Source/AnimNextInterfaceGraph/Private/RigUnit_AnimNextInterfaceParameter.cpp b/Engine/Plugins/Experimental/AnimNext/Source/AnimNextInterfaceGraph/Private/RigUnit_AnimNextInterfaceParameter.cpp
--- a/Engine/Plugins/Experimental/AnimNext/Source/AnimNextInterfaceGraph/Private/RigUnit_AnimNextInterfaceParameter.cpp
+++ b/Engine/Plugins/Experimental/AnimNext/Source/AnimNextInterfaceGraph/Private/RigUnit_AnimNextInterfaceParameter.cpp
@@ -5,7 +5,6 @@
 #include "AnimNextInterfaceExecuteContext.h"
 #include "AnimNextInterface.h"
 #include "AnimNextInterfaceTypes.h"
-#include "AnimNextInterfaceKernel.h"
 
 #include UE_INLINE_GENERATED_CPP_BY_NAME(RigUnit_AnimNextInterfaceParameter)
 
@@ -21,7 +20,7 @@ FRigUnit_AnimNextInterfaceParameter_Float_Execute()
 
 FRigUnit_AnimNextInterfaceParameter_GraphLODPose_Execute()
 {
-	const UE::AnimNext::Interface::FContext& AnimNextInterfaceContext = ExecuteContext.GetContext();
+	const UE::AnimNext::FContext& AnimNextInterfaceContext = ExecuteContext.GetContext();
 
 	Result = AnimNextInterfaceContext.GetParameterChecked<FAnimNextGraphLODPose>(Parameter);
 
@@ -30,7 +29,7 @@ FRigUnit_AnimNextInterfaceParameter_GraphLODPose_Execute()
 
 FRigUnit_AnimNextInterfaceParameter_AnimSequence_Execute()
 {
-	const UE::AnimNext::Interface::FContext& AnimNextInterfaceContext = ExecuteContext.GetContext();
+	const UE::AnimNext::FContext& AnimNextInterfaceContext = ExecuteContext.GetContext();
 
 	Result = AnimNextInterfaceContext.GetParameterChecked<FAnimNextGraph_AnimSequence>(Parameter);
 
@@ -45,7 +44,7 @@ FRigUnit_AnimNextInterfaceParameter_AnimNextInterface_Execute()
 
 FRigUnit_AnimNextInterface_Float_Execute()
 {
-	using namespace UE::AnimNext::Interface;
+	using namespace UE::AnimNext;
 
 	// @TODO: ensure that context arg is always present to avoid these checks here
 	const FContext& AnimNextInterfaceContext = ExecuteContext.GetContext();
@@ -59,7 +58,7 @@ FRigUnit_AnimNextInterface_Float_Execute()
 /*
 FRigUnit_AnimNextInterface_Pose_Execute()
 {
-	using namespace UE::AnimNext::Interface;
+	using namespace UE::AnimNext;
 
 	// @TODO: ensure that context arg is always present to avoid these checks here
 	if(Context.State == EControlRigState::Update && ExecuteContext.OpaqueArguments.Num() > 1 && ExecuteContext.OpaqueArguments[1] != nullptr)
@@ -85,41 +84,26 @@ FRigUnit_AnimNextInterface_SequencePlayer_Execute()
 {
 }
 
-struct FSpringDamperState
-{
-	float Value = 0.0f;
-	float ValueRate = 0.0f;
-};
 
-IMPLEMENT_ANIM_NEXT_INTERFACE_STATE_TYPE(FSpringDamperState, SpringDamperState);
 
 FRigUnit_TestFloatState_Execute()
 {
-	using namespace UE::AnimNext::Interface;
+	using namespace UE::AnimNext;
 
 	// @TODO: ensure that context arg is always present to avoid these checks here
 	const FContext& AnimNextInterfaceContext = ExecuteContext.GetContext();
 
-	const TParam<FSpringDamperState> State = AnimNextInterfaceContext.GetState<FSpringDamperState>(ExecuteContext.GetInterface(), 0);
+	FRigUnit_TestFloatState_SpringDamperState& State = AnimNextInterfaceContext.GetState<FRigUnit_TestFloatState_SpringDamperState>(ExecuteContext.GetInterface(), 0);
 	const float DeltaTime = AnimNextInterfaceContext.GetDeltaTime();
 
-	FKernel::Run(AnimNextInterfaceContext,
-		[DeltaTime, &Result](FSpringDamperState& InOutState,
-					float InTargetValue,
-					float InTargetValueRate,
-					float InSmoothingTime,
-					float InDampingRatio)
-			{
-				FMath::SpringDamperSmoothing(
-					InOutState.Value,
-					InOutState.ValueRate,
-					InTargetValue,
-					InTargetValueRate,
-					DeltaTime,
-					InSmoothingTime,
-					InDampingRatio);
-
-				Result = InOutState.Value;
-			},
-			State, TargetValue, TargetValueRate, SmoothingTime, DampingRatio);
+	FMath::SpringDamperSmoothing(
+		State.Value,
+		State.ValueRate,
+		TargetValue,
+		TargetValueRate,
+		DeltaTime,
+		SmoothingTime,
+		DampingRatio);
+
+	Result = State.Value;
 }
diff --git a/Engine/Plugins/Experimental/AnimNext/Source/AnimNextInterfaceGraph/Private/RigUnit_AnimNextInterfaceParameter.h b/Engine/Plugins/Experimental/AnimNext/Source/AnimNextInterfaceGraph/Private/RigUnit_AnimNextInterfaceParameter.h
--- a/Engine/Plugins/Experimental/AnimNext/Source/AnimNextInterfaceGraph/Private/RigUnit_AnimNextInterfaceParameter.h
+++ b/Engine/Plugins/Experimental/AnimNext/Source/AnimNextInterfaceGraph/Private/RigUnit_AnimNextInterfaceParameter.h
@@ -294,6 +294,18 @@ protected:
 	FAnimNextInterfaceExecuteContext Result;
 };
 
+USTRUCT()
+struct FRigUnit_TestFloatState_SpringDamperState
+{
+	GENERATED_BODY()
+
+	UPROPERTY()
+	float Value = 0.0f;
+
+	UPROPERTY()
+	float ValueRate = 0.0f;
+};
+
 /** Unit for getting a pose via an anim interface */
 USTRUCT(meta = (DisplayName = "Test Float State - Spring Damper Smoothing", Varying, Category = "Animation", TitleColor = "1 0 0", NodeColor = "1 1 1"))
 struct FRigUnit_TestFloatState : public FRigUnit_AnimNextInterfaceBase
diff --git a/Engine/Plugins/Experimental/AnimNext/Source/AnimNextInterfaceGraphEditor/Private/AnimNextInterfaceEditorUtils.cpp b/Engine/Plugins/Experimental/AnimNext/Source/AnimNextInterfaceGraphEditor/Private/AnimNextInterfaceEditorUtils.cpp
--- a/Engine/Plugins/Experimental/AnimNext/Source/AnimNextInterfaceGraphEditor/Private/AnimNextInterfaceEditorUtils.cpp
+++ b/Engine/Plugins/Experimental/AnimNext/Source/AnimNextInterfaceGraphEditor/Private/AnimNextInterfaceEditorUtils.cpp
@@ -2,6 +2,7 @@
 
 #include "AnimNextInterfaceEditorUtils.h"
 #include "AnimNextInterfaceGraph_EditorData.h"
+#include "Param/ParamType.h"
 #include "Kismet2/Kismet2NameValidators.h"
 
 namespace UE::AnimNext::InterfaceGraphEditor
@@ -9,7 +10,8 @@ namespace UE::AnimNext::InterfaceGraphEditor
 
 void FUtils::GetAllGraphNames(const UAnimNextInterfaceGraph_EditorData* InEditorData, TSet<FName>& OutNames)
 {
-
+	// TODO
+	ensure(false);
 }
 
 static const int32 MaxNameLength = 100;
@@ -119,4 +121,108 @@ FName FUtils::ValidateName(const UAnimNextInterfaceGraph_EditorData* InEditorDat
 	return *Name;
 }
 
+FAnimNextParamType FUtils::GetParameterTypeFromMetaData(const FStringView& InStringView)
+{
+	if(InStringView.Equals(TEXT("bool")))
+	{
+		return FAnimNextParamType(FAnimNextParamType::EValueType::Bool);
+	}
+	else if(InStringView.Equals(TEXT("uint8")) || InStringView.Equals(TEXT("byte")))
+	{
+		return FAnimNextParamType(FAnimNextParamType::EValueType::Byte);
+	}
+	else if(InStringView.Equals(TEXT("int32")))
+	{
+		return FAnimNextParamType(FAnimNextParamType::EValueType::Int32);
+	}
+	else if(InStringView.Equals(TEXT("int64")))
+	{
+		return FAnimNextParamType(FAnimNextParamType::EValueType::Int64);
+	}
+	else if(InStringView.Equals(TEXT("float")))
+	{
+		return FAnimNextParamType(FAnimNextParamType::EValueType::Float);
+	}
+	else if(InStringView.Equals(TEXT("double")))
+	{
+		return FAnimNextParamType(FAnimNextParamType::EValueType::Double);
+	}
+	else if(InStringView.Equals(TEXT("Name")))
+	{
+		return FAnimNextParamType(FAnimNextParamType::EValueType::Name);
+	}
+	else if(InStringView.Equals(TEXT("String")))
+	{
+		return FAnimNextParamType(FAnimNextParamType::EValueType::String);
+	}
+	else if(InStringView.Equals(TEXT("Text")))
+	{
+		return FAnimNextParamType(FAnimNextParamType::EValueType::Text);
+	}
+	else
+	{
+		int32 SplitIndex = INDEX_NONE;
+		if(InStringView.FindChar(TEXT('\''), SplitIndex))
+		{
+			// Disambiguated by class type: \Path\To\Class'\Path\To\Object, so no need to resolve the object, just the class 
+			const FStringView ClassStringView = InStringView.Left(SplitIndex - 1);
+			const FStringView ObjectStringView = InStringView.RightChop(SplitIndex + 1);
+			
+			const FTopLevelAssetPath ClassPath(ClassStringView);
+			const FTopLevelAssetPath ObjectPath(ObjectStringView);
+
+			if(ClassPath.IsValid() && ObjectPath.IsValid())
+			{
+				const FSoftObjectPath SoftClassPath(ClassPath);
+				const FSoftObjectPath SoftObjectPath(ObjectPath);
+				if(UClass* ResolvedClass = Cast<UClass>(SoftClassPath.ResolveObject()))
+				{
+					if(ResolvedClass == UScriptStruct::StaticClass()) 
+					{
+						return FAnimNextParamType(FAnimNextParamType::EValueType::Struct, FAnimNextParamType::EContainerType::None, ResolvedClass);
+					}
+					else if(ResolvedClass == UEnum::StaticClass()) 
+					{
+						return FAnimNextParamType(FAnimNextParamType::EValueType::Enum, FAnimNextParamType::EContainerType::None, ResolvedClass);
+					}
+					else if(ResolvedClass == UClass::StaticClass()) 
+					{
+						return FAnimNextParamType(FAnimNextParamType::EValueType::Class, FAnimNextParamType::EContainerType::None, ResolvedClass);
+					}
+					else
+					{
+						return FAnimNextParamType(FAnimNextParamType::EValueType::Object, FAnimNextParamType::EContainerType::None, ResolvedClass);
+					}
+				}
+			}
+		}
+		else
+		{
+			// Class must be inferred: \Path\To\Object, so we need to resolve the object
+			const FTopLevelAssetPath ObjectPath(InStringView);
+			if(ObjectPath.IsValid())
+			{
+				const FSoftObjectPath SoftObjectPath(ObjectPath);
+				if(UObject* ResolvedObject = Cast<UObject>(SoftObjectPath.ResolveObject()))
+				{
+					if(UScriptStruct* ResolvedStruct = Cast<UScriptStruct>(ResolvedObject))
+					{
+						return FAnimNextParamType(FAnimNextParamType::EValueType::Struct, FAnimNextParamType::EContainerType::None, ResolvedStruct);
+					}
+					else if(UEnum* ResolvedEnum = Cast<UEnum>(ResolvedObject))
+					{
+						return FAnimNextParamType(FAnimNextParamType::EValueType::Enum, FAnimNextParamType::EContainerType::None, ResolvedEnum);
+					}
+					else if(UClass* ResolvedClass = Cast<UClass>(ResolvedObject))
+					{
+						return FAnimNextParamType(FAnimNextParamType::EValueType::Object, FAnimNextParamType::EContainerType::None, ResolvedClass);
+					}
+				}
+			}
+		}
+	}
+
+	return FAnimNextParamType(); 
+}
+
 }
diff --git a/Engine/Plugins/Experimental/AnimNext/Source/AnimNextInterfaceGraphEditor/Private/AnimNextInterfaceEditorUtils.h b/Engine/Plugins/Experimental/AnimNext/Source/AnimNextInterfaceGraphEditor/Private/AnimNextInterfaceEditorUtils.h
--- a/Engine/Plugins/Experimental/AnimNext/Source/AnimNextInterfaceGraphEditor/Private/AnimNextInterfaceEditorUtils.h
+++ b/Engine/Plugins/Experimental/AnimNext/Source/AnimNextInterfaceGraphEditor/Private/AnimNextInterfaceEditorUtils.h
@@ -6,6 +6,7 @@
 
 class UAnimNextInterfaceGraph;
 class UAnimNextInterfaceGraph_EditorData;
+struct FAnimNextParamType;
 
 namespace UE::AnimNext::InterfaceGraphEditor
 {
@@ -15,6 +16,8 @@ struct FUtils
 	static FName ValidateName(const UAnimNextInterfaceGraph_EditorData* InEditorData, const FString& InName);
 
 	static void GetAllGraphNames(const UAnimNextInterfaceGraph_EditorData* InEditorData, TSet<FName>& OutNames);
+
+	static FAnimNextParamType GetParameterTypeFromMetaData(const FStringView& InStringView);
 };
 
 }
diff --git a/Engine/Plugins/Experimental/AnimNext/Source/AnimNextInterfaceGraphEditor/Private/AnimNextInterfacePropertyTypeCustomization.cpp b/Engine/Plugins/Experimental/AnimNext/Source/AnimNextInterfaceGraphEditor/Private/AnimNextInterfacePropertyTypeCustomization.cpp
--- a/Engine/Plugins/Experimental/AnimNext/Source/AnimNextInterfaceGraphEditor/Private/AnimNextInterfacePropertyTypeCustomization.cpp
+++ b/Engine/Plugins/Experimental/AnimNext/Source/AnimNextInterfaceGraphEditor/Private/AnimNextInterfacePropertyTypeCustomization.cpp
@@ -13,6 +13,8 @@
 #include "PropertyHandle.h"
 #include "SClassViewer.h"
 #include "Widgets/Text/STextBlock.h"
+#include "AnimNextInterfaceEditorUtils.h"
+#include "Param/ParamTypeHandle.h"
 
 #define LOCTEXT_NAMESPACE "AnimNextInterfacePropertyTypeCustomization"
 
@@ -26,45 +28,47 @@ bool FPropertyTypeIdentifier::IsPropertyTypeCustomized(const IPropertyHandle& Pr
 
 void FAnimNextInterfacePropertyTypeCustomization::CustomizeHeader(TSharedRef<IPropertyHandle> PropertyHandle, FDetailWidgetRow& HeaderRow, IPropertyTypeCustomizationUtils& CustomizationUtils)
 {
-	FString TypeName = PropertyHandle->GetMetaData("AnimNextInterfaceType");
-	FName AnimNextInterfaceTypeName(TypeName);
-
-	void* ValuePtr;
-	PropertyHandle->GetValueData(ValuePtr);
-	FScriptInterface* PropertyValue = reinterpret_cast<FScriptInterface*>(ValuePtr);
-	
-	TSharedPtr<SWidget> Widget = FAnimNextInterfaceWidgetFactories::CreateAnimNextInterfaceWidget(AnimNextInterfaceTypeName, PropertyValue->GetObject(),
-			FOnClassPicked::CreateLambda([this, PropertyHandle](UClass* ChosenClass)
-			{
-				TArray<void*> RawData;
-				PropertyHandle->AccessRawData(RawData);
-				TArray<UPackage*> OuterObjects;
-				PropertyHandle->GetOuterPackages(OuterObjects);
-
-				for(int i=0;i<RawData.Num();i++)
+	const FString TypeName = PropertyHandle->GetMetaData("AnimNextInterfaceType");
+	const FAnimNextParamType AnimNextType = FUtils::GetParameterTypeFromMetaData(TypeName);
+	if(AnimNextType.IsValid())
+	{
+		void* ValuePtr;
+		PropertyHandle->GetValueData(ValuePtr);
+		FScriptInterface* PropertyValue = reinterpret_cast<FScriptInterface*>(ValuePtr);
+		
+		TSharedPtr<SWidget> Widget = FWidgetFactories::CreateAnimNextInterfaceWidget(AnimNextType.GetHandle(), PropertyValue->GetObject(),
+				FOnClassPicked::CreateLambda([this, PropertyHandle](UClass* ChosenClass)
 				{
-					UObject* NewValue = NewObject<UObject>(OuterObjects[i], ChosenClass);
-				
-					PropertyHandle->NotifyPreChange();
-					PropertyHandle->SetValue(NewValue);
-					PropertyHandle->NotifyPostChange(EPropertyChangeType::ValueSet);
-					PropertyHandle->GetPropertyNode()->GetParentNode()->RequestRebuildChildren();
-				}
-			}),
-			nullptr
-	);
+					TArray<void*> RawData;
+					PropertyHandle->AccessRawData(RawData);
+					TArray<UPackage*> OuterObjects;
+					PropertyHandle->GetOuterPackages(OuterObjects);
+
+					for(int i=0;i<RawData.Num();i++)
+					{
+						UObject* NewValue = NewObject<UObject>(OuterObjects[i], ChosenClass);
+					
+						PropertyHandle->NotifyPreChange();
+						PropertyHandle->SetValue(NewValue);
+						PropertyHandle->NotifyPostChange(EPropertyChangeType::ValueSet);
+						PropertyHandle->GetPropertyNode()->GetParentNode()->RequestRebuildChildren();
+					}
+				}),
+				nullptr
+		);
 
-	HeaderRow
-	.NameContent()
-	[
-		PropertyHandle->CreatePropertyNameWidget()
-	]
-	.ValueContent()
-	[
-		Widget.ToSharedRef()
-	];
+		HeaderRow
+		.NameContent()
+		[
+			PropertyHandle->CreatePropertyNameWidget()
+		]
+		.ValueContent()
+		[
+			Widget.ToSharedRef()
+		];
+	}
 }
-	
+
 void FAnimNextInterfacePropertyTypeCustomization::CustomizeChildren(TSharedRef<IPropertyHandle> PropertyHandle, IDetailChildrenBuilder& ChildBuilder, IPropertyTypeCustomizationUtils& CustomizationUtils)
 {
 	// todo: make this work correctly for multiselect
diff --git a/Engine/Plugins/Experimental/AnimNext/Source/AnimNextInterfaceGraphEditor/Private/AnimNextInterfaceWidgetFactories.cpp b/Engine/Plugins/Experimental/AnimNext/Source/AnimNextInterfaceGraphEditor/Private/AnimNextInterfaceWidgetFactories.cpp
--- a/Engine/Plugins/Experimental/AnimNext/Source/AnimNextInterfaceGraphEditor/Private/AnimNextInterfaceWidgetFactories.cpp
+++ b/Engine/Plugins/Experimental/AnimNext/Source/AnimNextInterfaceGraphEditor/Private/AnimNextInterfaceWidgetFactories.cpp
@@ -2,12 +2,8 @@
 
 #include "AnimNextInterfaceWidgetFactories.h"
 #include "Modules/ModuleManager.h"
-#include "Widgets/Input/SHyperlink.h"
-#include "Widgets/Input/SNumericEntryBox.h"
 #include "Widgets/Input/SComboButton.h"
 #include "SClassViewer.h"
-#include "DetailCategoryBuilder.h"
-#include "Widgets/Layout/SWidgetSwitcher.h"
 #include "AnimNextInterfaceClassFilter.h"
 
 #define LOCTEXT_NAMESPACE "AnimNextInterfaceEditor"
@@ -20,7 +16,7 @@ void ConvertToText(UObject* Object, FText& OutText)
 	UClass* Class = Object->GetClass();
 	while (Class)
 	{
-		if (auto TextConverter = FAnimNextInterfaceWidgetFactories::AnimNextInterfaceTextConverter.Find(Class))
+		if (auto TextConverter = FWidgetFactories::AnimNextInterfaceTextConverter.Find(Class))
 		{
 			(*TextConverter)(Object, OutText);
 			break;
@@ -29,7 +25,7 @@ void ConvertToText(UObject* Object, FText& OutText)
 	}
 }
 
-TSharedPtr<SWidget> FAnimNextInterfaceWidgetFactories::CreateAnimNextInterfaceWidget(FName AnimNextInterfaceTypeName, UObject* Value, const FOnClassPicked& CreateClassCallback, TSharedPtr<SBorder>* InnerWidget)
+TSharedPtr<SWidget> FWidgetFactories::CreateAnimNextInterfaceWidget(FParamTypeHandle TypeHandle, UObject* Value, const FOnClassPicked& CreateClassCallback, TSharedPtr<SBorder>* InnerWidget)
 {
 	TSharedPtr<SWidget> LeftWidget;
 
@@ -38,7 +34,7 @@ TSharedPtr<SWidget> FAnimNextInterfaceWidgetFactories::CreateAnimNextInterfaceWi
 		UClass* Class = Value->GetClass();
 		while (Class && !LeftWidget.IsValid())
 		{
-			if (auto Creator = FAnimNextInterfaceWidgetFactories::AnimNextInterfaceWidgetCreators.Find(Class))
+			if (auto Creator = FWidgetFactories::AnimNextInterfaceWidgetCreators.Find(Class))
 			{
 				LeftWidget = (*Creator)(Value);
 				break;
@@ -57,10 +53,10 @@ TSharedPtr<SWidget> FAnimNextInterfaceWidgetFactories::CreateAnimNextInterfaceWi
 	TSharedPtr<SComboButton> Button = SNew(SComboButton)
 			.ComboButtonStyle(FAppStyle::Get(), "SimpleComboButton");
 	
-	Button->SetOnGetMenuContent(FOnGetContent::CreateLambda([AnimNextInterfaceTypeName, Button, CreateClassCallback]()
+	Button->SetOnGetMenuContent(FOnGetContent::CreateLambda([TypeHandle, Button, CreateClassCallback]()
 	{
 		FClassViewerInitializationOptions Options;
-		Options.ClassFilters.Add(MakeShared<FAnimNextInterfaceClassFilter>(AnimNextInterfaceTypeName, nullptr));
+		Options.ClassFilters.Add(MakeShared<FAnimNextInterfaceClassFilter>(TypeHandle));
 		
 		// Add class filter for columns here
 		TSharedRef<SWidget>  ClassMenu = FModuleManager::LoadModuleChecked<FClassViewerModule>("ClassViewer").
@@ -104,15 +100,15 @@ TSharedPtr<SWidget> FAnimNextInterfaceWidgetFactories::CreateAnimNextInterfaceWi
 	return Widget;
 }
 
-TMap<const UClass*, TFunction<void (const UObject* Object, FText& OutText)>> FAnimNextInterfaceWidgetFactories::AnimNextInterfaceTextConverter;
-TMap<const UClass*, TFunction<TSharedRef<SWidget> (UObject* Object)>> FAnimNextInterfaceWidgetFactories::AnimNextInterfaceWidgetCreators;
+TMap<const UClass*, TFunction<void (const UObject* Object, FText& OutText)>> FWidgetFactories::AnimNextInterfaceTextConverter;
+TMap<const UClass*, TFunction<TSharedRef<SWidget> (UObject* Object)>> FWidgetFactories::AnimNextInterfaceWidgetCreators;
 
 void ConvertToText_Base(const UObject* Object, FText& OutText)
 {
 	OutText = FText::FromString(Object->GetName());
 }
 
-void FAnimNextInterfaceWidgetFactories::RegisterWidgets()
+void FWidgetFactories::RegisterWidgets()
 {
 	AnimNextInterfaceTextConverter.Add(UObject::StaticClass(), ConvertToText_Base);
 }
diff --git a/Engine/Plugins/Experimental/AnimNext/Source/AnimNextInterfaceGraphEditor/Public/AnimNextInterfaceClassFilter.h b/Engine/Plugins/Experimental/AnimNext/Source/AnimNextInterfaceGraphEditor/Public/AnimNextInterfaceClassFilter.h
--- a/Engine/Plugins/Experimental/AnimNext/Source/AnimNextInterfaceGraphEditor/Public/AnimNextInterfaceClassFilter.h
+++ b/Engine/Plugins/Experimental/AnimNext/Source/AnimNextInterfaceGraphEditor/Public/AnimNextInterfaceClassFilter.h
@@ -1,19 +1,20 @@
 // Copyright Epic Games, Inc. All Rights Reserved.
 #pragma once
 
+#include "Param/ParamTypeHandle.h"
 #include "ClassViewerFilter.h"
 #include "IAnimNextInterface.h"
 
 namespace UE::AnimNext::InterfaceGraphEditor
 {
-	// Filter class for things which implement IChooserColumn
+	// Filter class for things which implement IAnimNextInterface
 	class FAnimNextInterfaceClassFilter : public IClassViewerFilter
 	{
 	public:
-		FAnimNextInterfaceClassFilter(FName InTypeName, const UScriptStruct* InStructType) : TypeName(InTypeName), StructType(InStructType) { };
+		FAnimNextInterfaceClassFilter(FParamTypeHandle InTypeHandle) : TypeHandle(InTypeHandle) {};
 		virtual ~FAnimNextInterfaceClassFilter() override {};
 	
-		virtual bool IsClassAllowed(const FClassViewerInitializationOptions& InInitOptions, const UClass* InClass, TSharedRef< class FClassViewerFilterFuncs > InFilterFuncs ) override
+		virtual bool IsClassAllowed(const FClassViewerInitializationOptions& InInitOptions, const UClass* InClass, TSharedRef<FClassViewerFilterFuncs> InFilterFuncs) override
 		{
 			if (InClass->ImplementsInterface(UAnimNextInterface::StaticClass()))
 			{
@@ -23,28 +24,21 @@ namespace UE::AnimNext::InterfaceGraphEditor
 				// for Object type, check type hierarchy. (todo)
 			
 				IAnimNextInterface* Object = static_cast<IAnimNextInterface*>(InClass->GetDefaultObject()->GetInterfaceAddress(UAnimNextInterface::StaticClass()));
-				if (Object->GetReturnTypeName() == TypeName || Object->GetReturnTypeName() == "" || TypeName == "")
+				if (Object->GetReturnTypeHandle() == TypeHandle || !Object->GetReturnTypeHandle().IsValid() || !TypeHandle.IsValid())
 				{
-					if (const UScriptStruct* ReturnStructType = Object->GetReturnTypeStruct())
-					{
-						return true; // kytodo: compare the struct types to see if they are compatible
-					}
-					else
-					{
-						return true;
-					}
+					// @TODO: deal with inheritance directionality here
+					return true;
 				}
 			}
 
 			return false;
 		}
 
-		virtual bool IsUnloadedClassAllowed(const FClassViewerInitializationOptions& InInitOptions, const TSharedRef< const class IUnloadedBlueprintData > InUnloadedClassData, TSharedRef< class FClassViewerFilterFuncs > InFilterFuncs) override
+		virtual bool IsUnloadedClassAllowed(const FClassViewerInitializationOptions& InInitOptions, const TSharedRef<const IUnloadedBlueprintData> InUnloadedClassData, TSharedRef<FClassViewerFilterFuncs> InFilterFuncs) override
 		{
 			return false;
 		}
 	private:
-		FName TypeName;
-		const UScriptStruct* StructType;
+		FParamTypeHandle TypeHandle;
 	};
 }
\ No newline at end of file
diff --git a/Engine/Plugins/Experimental/AnimNext/Source/AnimNextInterfaceGraphEditor/Public/AnimNextInterfaceWidgetFactories.h b/Engine/Plugins/Experimental/AnimNext/Source/AnimNextInterfaceGraphEditor/Public/AnimNextInterfaceWidgetFactories.h
--- a/Engine/Plugins/Experimental/AnimNext/Source/AnimNextInterfaceGraphEditor/Public/AnimNextInterfaceWidgetFactories.h
+++ b/Engine/Plugins/Experimental/AnimNext/Source/AnimNextInterfaceGraphEditor/Public/AnimNextInterfaceWidgetFactories.h
@@ -7,15 +7,20 @@
 #include "Widgets/Layout/SBorder.h"
 #include "SClassViewer.h"
 
+namespace UE::AnimNext
+{
+	struct FParamTypeHandle;
+}
+
 namespace UE::AnimNext::InterfaceGraphEditor
 {
-	class ANIMNEXTINTERFACEGRAPHEDITOR_API FAnimNextInterfaceWidgetFactories 
+	class ANIMNEXTINTERFACEGRAPHEDITOR_API FWidgetFactories 
 	{
 	public:
 		static TMap<const UClass*, TFunction<void (const UObject* Object, FText& OutText)>> AnimNextInterfaceTextConverter;
 		static TMap<const UClass*, TFunction<TSharedRef<SWidget> (UObject* Object)>> AnimNextInterfaceWidgetCreators;
 
-		static TSharedPtr<SWidget> CreateAnimNextInterfaceWidget(FName AnimNextInterfaceTypeName, UObject* Value, const FOnClassPicked& CreateClassCallback, TSharedPtr<SBorder>* InnerWidget = nullptr);
+		static TSharedPtr<SWidget> CreateAnimNextInterfaceWidget(FParamTypeHandle TypeHandle, UObject* Value, const FOnClassPicked& CreateClassCallback, TSharedPtr<SBorder>* InnerWidget = nullptr);
 		
 		static void RegisterWidgets();
 	};
diff --git a/Engine/Plugins/Experimental/AnimNext/Source/AnimNextInterfaceGraphUncookedOnly/Internal/AnimNextInterfaceGraph_EditorData.h b/Engine/Plugins/Experimental/AnimNext/Source/AnimNextInterfaceGraphUncookedOnly/Internal/AnimNextInterfaceGraph_EditorData.h
--- a/Engine/Plugins/Experimental/AnimNext/Source/AnimNextInterfaceGraphUncookedOnly/Internal/AnimNextInterfaceGraph_EditorData.h
+++ b/Engine/Plugins/Experimental/AnimNext/Source/AnimNextInterfaceGraphUncookedOnly/Internal/AnimNextInterfaceGraph_EditorData.h
@@ -85,8 +85,6 @@ class UAnimNextInterfaceGraph_EditorData : public UObject, public IRigVMClientHo
 
 	bool IsNodeExecConnected(const UEdGraphNode* Node) const;
 
-	FName GetPinTypeName(const FEdGraphPinType& EdGraphPinType);
-
 	UPROPERTY()
 	TObjectPtr<UAnimNextInterfaceGraph_EdGraph> RootGraph;
 
diff --git a/Engine/Plugins/Experimental/AnimNext/Source/AnimNextInterfaceGraphUncookedOnly/Internal/AnimNextInterfaceUncookedOnlyUtils.h b/Engine/Plugins/Experimental/AnimNext/Source/AnimNextInterfaceGraphUncookedOnly/Internal/AnimNextInterfaceUncookedOnlyUtils.h
--- a/Engine/Plugins/Experimental/AnimNext/Source/AnimNextInterfaceGraphUncookedOnly/Internal/AnimNextInterfaceUncookedOnlyUtils.h
+++ b/Engine/Plugins/Experimental/AnimNext/Source/AnimNextInterfaceGraphUncookedOnly/Internal/AnimNextInterfaceUncookedOnlyUtils.h
@@ -3,17 +3,14 @@
 #pragma once
 
 #include "CoreMinimal.h"
+#include "Param/ParamTypeHandle.h"
 
 class UAnimNextInterfaceGraph;
 class UAnimNextInterfaceGraph_EditorData;
 class UAnimNextInterfaceGraph_EdGraph;
 class URigVMController;
 class URigVMGraph;
-
-namespace UE::AnimNext::InterfaceGraphEditor
-{
-	class FGraphEditor;
-}
+struct FEdGraphPinType;
 
 namespace UE::AnimNext::InterfaceGraphUncookedOnly
 {
@@ -25,6 +22,12 @@ struct ANIMNEXTINTERFACEGRAPHUNCOOKEDONLY_API FUtils
 	static UAnimNextInterfaceGraph_EditorData* GetEditorData(const UAnimNextInterfaceGraph* InAnimNextInterfaceGraph);
 	
 	static void RecreateVM(UAnimNextInterfaceGraph* InGraph);
+
+	/**
+	 * Get an AnimNext parameter type handle from an FEdGraphPinType.
+	 * Note that the returned handle may not be valid, so should be checked using IsValid() before use.
+	 **/
+	static FParamTypeHandle GetParameterHandleFromPin(const FEdGraphPinType& InPinType);
 };
 
 }
\ No newline at end of file
diff --git a/Engine/Plugins/Experimental/AnimNext/Source/AnimNextInterfaceGraphUncookedOnly/Private/AnimNextInterfaceGraphUncookedOnlyUtils.cpp b/Engine/Plugins/Experimental/AnimNext/Source/AnimNextInterfaceGraphUncookedOnly/Private/AnimNextInterfaceGraphUncookedOnlyUtils.cpp
--- a/Engine/Plugins/Experimental/AnimNext/Source/AnimNextInterfaceGraphUncookedOnly/Private/AnimNextInterfaceGraphUncookedOnlyUtils.cpp
+++ b/Engine/Plugins/Experimental/AnimNext/Source/AnimNextInterfaceGraphUncookedOnly/Private/AnimNextInterfaceGraphUncookedOnlyUtils.cpp
@@ -88,4 +88,108 @@ UAnimNextInterfaceGraph_EditorData* FUtils::GetEditorData(const UAnimNextInterfa
 	return CastChecked<UAnimNextInterfaceGraph_EditorData>(InAnimNextInterfaceGraph->EditorData);
 }
 
+FParamTypeHandle FUtils::GetParameterHandleFromPin(const FEdGraphPinType& InPinType)
+{
+	FAnimNextParamType::EValueType ValueType = FAnimNextParamType::EValueType::None;
+	FAnimNextParamType::EContainerType ContainerType = FAnimNextParamType::EContainerType::None;
+	UObject* ValueTypeObject = nullptr;
+
+	if (InPinType.PinCategory == UEdGraphSchema_K2::PC_Boolean)
+	{
+		ValueType = FAnimNextParamType::EValueType::Bool;
+	}
+	else if (InPinType.PinCategory == UEdGraphSchema_K2::PC_Byte)
+	{
+		ValueType = FAnimNextParamType::EValueType::Byte;
+	}
+	else if (InPinType.PinCategory == UEdGraphSchema_K2::PC_Int)
+	{
+		ValueType = FAnimNextParamType::EValueType::Int32;
+	}
+	else if (InPinType.PinCategory == UEdGraphSchema_K2::PC_Int64)
+	{
+		ValueType = FAnimNextParamType::EValueType::Int64;
+	}
+	else if (InPinType.PinCategory == UEdGraphSchema_K2::PC_Real)
+	{
+		if (InPinType.PinSubCategory == UEdGraphSchema_K2::PC_Float)
+		{
+			ValueType = FAnimNextParamType::EValueType::Float;
+		}
+		else if (InPinType.PinSubCategory == UEdGraphSchema_K2::PC_Double)
+		{
+			ValueType = FAnimNextParamType::EValueType::Double;
+		}
+		else
+		{
+			ensure(false);	// Reals should be either floats or doubles
+		}
+	}
+	else if (InPinType.PinCategory == UEdGraphSchema_K2::PC_Float)
+	{
+		ValueType = FAnimNextParamType::EValueType::Float;
+	}
+	else if (InPinType.PinCategory == UEdGraphSchema_K2::PC_Double)
+	{
+		ValueType = FAnimNextParamType::EValueType::Double;
+	}
+	else if (InPinType.PinCategory == UEdGraphSchema_K2::PC_Name)
+	{
+		ValueType = FAnimNextParamType::EValueType::Name;
+	}
+	else if (InPinType.PinCategory == UEdGraphSchema_K2::PC_String)
+	{
+		ValueType = FAnimNextParamType::EValueType::String;
+	}
+	else if (InPinType.PinCategory == UEdGraphSchema_K2::PC_Text)
+	{
+		ValueType = FAnimNextParamType::EValueType::Text;
+	}
+	else if (InPinType.PinCategory == UEdGraphSchema_K2::PC_Enum)
+	{
+		ValueType = FAnimNextParamType::EValueType::Enum;
+		ValueTypeObject = InPinType.PinSubCategoryObject.Get();
+		ensure(ValueTypeObject);
+	}
+	else if (InPinType.PinCategory == UEdGraphSchema_K2::PC_Struct)
+	{
+		ValueType = FAnimNextParamType::EValueType::Struct;
+		ValueTypeObject = Cast<UScriptStruct>(InPinType.PinSubCategoryObject.Get());
+		ensure(ValueTypeObject);
+	}
+	else if (InPinType.PinCategory == UEdGraphSchema_K2::PC_Object)
+	{
+		ValueType = FAnimNextParamType::EValueType::Object;
+		ValueTypeObject = Cast<UClass>(InPinType.PinSubCategoryObject.Get());
+		ensure(ValueTypeObject);
+	}
+	else if (InPinType.PinCategory == UEdGraphSchema_K2::PC_SoftObject)
+	{
+		ValueType = FAnimNextParamType::EValueType::SoftObject;
+		ValueTypeObject = Cast<UClass>(InPinType.PinSubCategoryObject.Get());
+		ensure(ValueTypeObject);
+	}
+	else if (InPinType.PinCategory == UEdGraphSchema_K2::PC_SoftClass)
+	{
+		ValueType = FAnimNextParamType::EValueType::SoftClass;
+		ValueTypeObject = Cast<UClass>(InPinType.PinSubCategoryObject.Get());
+		ensure(ValueTypeObject);
+	}
+
+	if(InPinType.ContainerType == EPinContainerType::Array)
+	{
+		ContainerType = FAnimNextParamType::EContainerType::Array;
+	}
+	else if(InPinType.ContainerType == EPinContainerType::Set)
+	{
+		ensureMsgf(false, TEXT("Set pins are not yet supported"));
+	}
+	if(InPinType.ContainerType == EPinContainerType::Map)
+	{
+		ensureMsgf(false, TEXT("Map pins are not yet supported"));
+	}
+	
+	return FAnimNextParamType(ValueType, ContainerType, ValueTypeObject).GetHandle();
+}
+
 }
diff --git a/Engine/Plugins/Experimental/AnimNext/Source/AnimNextInterfaceGraphUncookedOnly/Private/AnimNextInterfaceGraph_EditorData.cpp b/Engine/Plugins/Experimental/AnimNext/Source/AnimNextInterfaceGraphUncookedOnly/Private/AnimNextInterfaceGraph_EditorData.cpp
--- a/Engine/Plugins/Experimental/AnimNext/Source/AnimNextInterfaceGraphUncookedOnly/Private/AnimNextInterfaceGraph_EditorData.cpp
+++ b/Engine/Plugins/Experimental/AnimNext/Source/AnimNextInterfaceGraphUncookedOnly/Private/AnimNextInterfaceGraph_EditorData.cpp
@@ -552,61 +552,11 @@ void UAnimNextInterfaceGraph_EditorData::UpdateGraphReturnType()
 	if (UAnimNextInterfaceGraph* Graph = GetTypedOuter<UAnimNextInterfaceGraph>())
 	{
 		FEdGraphPinType GraphReturnPinType = FindGraphReturnPinType();
-		const FName PinTypeName = GetPinTypeName(GraphReturnPinType);
-		if (Graph->ReturnTypeName != PinTypeName || Graph->ReturnTypeStruct != GraphReturnPinType.PinSubCategoryObject)
+		UE::AnimNext::FParamTypeHandle Handle = UE::AnimNext::InterfaceGraphUncookedOnly::FUtils::GetParameterHandleFromPin(GraphReturnPinType);
+		if (Handle.IsValid() && Graph->GetReturnTypeHandle() != Handle)
 		{
 			Graph->Modify();
-			Graph->ReturnTypeName = PinTypeName;
-			Graph->ReturnTypeStruct = TObjectPtr<UScriptStruct>(Cast<UScriptStruct>(GraphReturnPinType.PinSubCategoryObject));
+			Graph->SetReturnTypeHandle(Handle);
 		}
 	}
 }
-
-FName UAnimNextInterfaceGraph_EditorData::GetPinTypeName(const FEdGraphPinType& EdGraphPinType)
-{
-	FName DataType = EdGraphPinType.PinCategory;
-
-	if (EdGraphPinType.PinCategory == UEdGraphSchema_K2::PC_Real)
-	{
-		if (EdGraphPinType.PinSubCategory == UEdGraphSchema_K2::PC_Float)
-		{
-			DataType = TNameOf<float>::GetName();
-		}
-		else if (EdGraphPinType.PinSubCategory == UEdGraphSchema_K2::PC_Double)
-		{
-			DataType = TNameOf<double>::GetName();
-		}
-		else
-		{
-			ensure(false);
-		}
-	}
-	else if (DataType != NAME_None)
-	{
-		if (DataType == UEdGraphSchema_K2::PC_Struct)
-		{
-			UObject* DataTypeObject = nullptr;
-			DataType = NAME_None;
-			if (UScriptStruct* DataStruct = Cast<UScriptStruct>(EdGraphPinType.PinSubCategoryObject))
-			{
-				DataTypeObject = DataStruct;
-				DataType = *DataStruct->GetStructCPPName();
-			}
-		}
-
-		if (DataType == TEXT("int"))
-		{
-			DataType = TNameOf<int32>::GetName();
-		}
-		else if (DataType == TEXT("name"))
-		{
-			DataType = TNameOf<FName>::GetName();
-		}
-		else if (DataType == TEXT("string"))
-		{
-			DataType = TNameOf<FString>::GetName();;
-		}
-	}
-
-	return DataType;
-}
