AnimNext : Fixed Parameter Handle type check and updated test graph return type

diff --git a/Engine/Plugins/Experimental/AnimNext/Source/AnimNextInterface/Public/Param/ParamTypeHandle.h b/Engine/Plugins/Experimental/AnimNext/Source/AnimNextInterface/Public/Param/ParamTypeHandle.h
--- a/Engine/Plugins/Experimental/AnimNext/Source/AnimNextInterface/Public/Param/ParamTypeHandle.h
+++ b/Engine/Plugins/Experimental/AnimNext/Source/AnimNextInterface/Public/Param/ParamTypeHandle.h
@@ -360,7 +360,7 @@ public:
 	{
 		const EParamType ParameterType = GetParameterType();
 		const uint32 CustomTypeIndex = GetCustomTypeIndex();
-		checkSlow(ValidateCustomTypeIndex(CustomTypeIndex));
+		checkSlow(ParameterType != EParamType::Custom || ValidateCustomTypeIndex(CustomTypeIndex));
 		const bool bIsValidCustomType = (ParameterType == EParamType::Custom && CustomTypeIndex != 0);
 		const bool bIsValidBuiltInType = ParameterType > EParamType::None && ParameterType <= EParamType::MaxBuiltIn;
 		return bIsValidBuiltInType || bIsValidCustomType;
