File/namespace reorg

Removed AnimNext in filenames as much as reasonable
Removed extra 'Interface' in filenames and namespaces.
Trimmed duplicate prefixes (e.g. animation, anim etc.) from namespace-embdeded types

diff --git a/Engine/Plugins/Experimental/AnimNext/AnimNext.uplugin b/Engine/Plugins/Experimental/AnimNext/AnimNext.uplugin
--- a/Engine/Plugins/Experimental/AnimNext/AnimNext.uplugin
+++ b/Engine/Plugins/Experimental/AnimNext/AnimNext.uplugin
@@ -28,22 +28,22 @@
 	"Modules" :
 	[
 		{
-			"Name" : "AnimNextInterface",
+			"Name" : "AnimNext",
 			"Type" : "Runtime",
 			"LoadingPhase" : "PreDefault"
 		},
 		{
-			"Name" : "AnimNextInterfaceGraph",
+			"Name" : "AnimNextGraph",
 			"Type" : "Runtime",
 			"LoadingPhase" : "Default"
 		},
 		{
-			"Name" : "AnimNextInterfaceGraphEditor",
+			"Name" : "AnimNextGraphEditor",
 			"Type" : "Editor",
 			"LoadingPhase" : "Default"
 		},
 		{
-			"Name" : "AnimNextInterfaceGraphUncookedOnly",
+			"Name" : "AnimNextGraphUncookedOnly",
 			"Type" : "UncookedOnly",
 			"LoadingPhase" : "Default"
 		}
diff --git a/Engine/Plugins/Experimental/AnimNext/Source/AnimNext/AnimNext.Build.cs b/Engine/Plugins/Experimental/AnimNext/Source/AnimNext/AnimNext.Build.cs
new file mode 100644
--- /dev/null
+++ b/Engine/Plugins/Experimental/AnimNext/Source/AnimNext/AnimNext.Build.cs
@@ -0,0 +1,26 @@
+// Copyright Epic Games, Inc. All Rights Reserved.
+
+namespace UnrealBuildTool.Rules
+{
+	public class AnimNext : ModuleRules
+	{
+		public AnimNext(ReadOnlyTargetRules Target) : base(Target)
+		{	
+			PublicDependencyModuleNames.AddRange(
+				new string[]
+				{
+					"Core",
+					"CoreUObject",
+					"StructUtils"
+				}
+			);
+			PrivateDependencyModuleNames.AddRange(
+				new string[]
+				{
+					"AssetRegistry",
+					"Engine"
+				}
+			);
+		}
+	}
+}
\ No newline at end of file
diff --git a/Engine/Plugins/Experimental/AnimNext/Source/AnimNext/Private/DecompressionTools.cpp b/Engine/Plugins/Experimental/AnimNext/Source/AnimNext/Private/DecompressionTools.cpp
new file mode 100644
--- /dev/null
+++ b/Engine/Plugins/Experimental/AnimNext/Source/AnimNext/Private/DecompressionTools.cpp
@@ -0,0 +1,663 @@
+// Copyright Epic Games, Inc. All Rights Reserved.
+
+#include "DecompressionTools.h"
+#include "Animation/AnimSequenceDecompressionContext.h"
+#include "Animation/AnimCompressionTypes.h"
+#include "Animation/AnimBoneCompressionCodec.h"
+
+#include "AnimationRuntime.h"
+#include "AnimEncoding.h"
+#include "Animation/AnimBoneDecompressionData.h"
+#include "Animation/SkeletonRemappingRegistry.h"
+#include "Animation/SkeletonRemapping.h"
+#include "Animation/AnimRootMotionProvider.h"
+#include "Animation/AnimSequence.h"
+#include "HAL/ConsoleManager.h"
+
+CSV_DECLARE_CATEGORY_MODULE_EXTERN(ENGINE_API, Animation);
+
+DECLARE_CYCLE_STAT(TEXT("Build Anim Track Pairs"), STAT_BuildAnimTrackPairs, STATGROUP_Anim);
+DECLARE_CYCLE_STAT(TEXT("Extract Pose From Anim Data"), STAT_ExtractPoseFromAnimData, STATGROUP_Anim);
+
+DECLARE_CYCLE_STAT(TEXT("AnimSeq GetBonePose"), STAT_AnimSeq_GetBonePose, STATGROUP_Anim);
+DECLARE_CYCLE_STAT(TEXT("AnimSeq EvalCurveData"), STAT_AnimSeq_EvalCurveData, STATGROUP_Anim);
+
+namespace UE::AnimNext
+{
+
+static IConsoleVariable* CVarForceEvalRawData = IConsoleManager::Get().FindConsoleVariable(TEXT("a.ForceEvalRawData"));
+static bool GetForceRawData()
+{
+	return (CVarForceEvalRawData != nullptr) ? CVarForceEvalRawData->GetBool() : false;
+}
+
+
+//****************************************************************************
+// This file contains code extracted from AnimSequence and AnimationDecompression,
+// in order to be able to decompress an anim sequence using AnimNext format
+//****************************************************************************
+
+struct FGetBonePoseScratchArea : public TThreadSingleton<FGetBonePoseScratchArea>
+{
+	BoneTrackArray RotationScalePairs;
+	BoneTrackArray TranslationPairs;
+	BoneTrackArray AnimScaleRetargetingPairs;
+	BoneTrackArray AnimRelativeRetargetingPairs;
+	BoneTrackArray OrientAndScaleRetargetingPairs;
+};
+
+
+static bool CanEvaluateRawAnimationData(const UAnimSequence* AnimSequence)
+{
+#if WITH_EDITOR
+	return AnimSequence->IsDataModelValid();
+#endif
+
+	return false;
+}
+
+static bool UseRawDataForPoseExtraction(const UAnimSequence* AnimSequence, FLODPose& AnimationPoseData)
+{
+	return CanEvaluateRawAnimationData(AnimSequence) && (
+#if WITH_EDITOR
+		PRAGMA_DISABLE_DEPRECATION_WARNINGS
+		AnimSequence->OnlyUseRawData() ||
+		PRAGMA_ENABLE_DEPRECATION_WARNINGS
+#endif
+		(AnimSequence->GetSkeletonVirtualBoneGuid() != AnimSequence->GetSkeleton()->GetVirtualBoneGuid()) || AnimationPoseData.GetDisableRetargeting() || AnimationPoseData.ShouldUseRawData() ||
+#if WITH_EDITOR
+		GetForceRawData() ||
+#endif // WITH_EDITOR
+		AnimationPoseData.ShouldUseSourceData());
+}
+
+// --- ---
+
+/*static*/ void FDecompressionTools::GetAnimationPose(const UAnimSequence* AnimSequence, FLODPose& OutAnimationPoseData, const FAnimExtractContext& ExtractionContext)
+{
+	if (OutAnimationPoseData.GetRefPose().IsValid() == false)
+	{
+		return;
+	}
+
+	// @todo anim: if compressed and baked in the future, we don't have to do this 
+	if (UseRawDataForPoseExtraction(AnimSequence, OutAnimationPoseData) && AnimSequence->IsValidAdditive())
+	{
+		switch (AnimSequence->GetAdditiveAnimType())
+		{
+			case AAT_LocalSpaceBase:
+			{
+				GetBonePose_Additive(AnimSequence, ExtractionContext, OutAnimationPoseData);
+			}
+			break;
+
+			case AAT_RotationOffsetMeshSpace:
+			{
+				GetBonePose_AdditiveMeshRotationOnly(AnimSequence, ExtractionContext, OutAnimationPoseData);
+			}
+			break;
+
+			default:
+				break;
+		}
+	}
+	else
+	{
+		GetBonePose(AnimSequence, ExtractionContext, OutAnimationPoseData);
+	}
+
+	// If the sequence has root motion enabled, allow sampling of a root motion delta into the custom attribute container of the outgoing pose
+	if (AnimSequence->HasRootMotion())
+	{
+		if (const UE::Anim::IAnimRootMotionProvider* RootMotionProvider = UE::Anim::IAnimRootMotionProvider::Get())
+		{
+			//RootMotionProvider->SampleRootMotion(ExtractionContext.DeltaTimeRecord, *AnimSequence, ExtractionContext.bLooping, OutAnimationPoseData.GetAttributes());
+		}
+	}
+
+	// Check that all bone atoms coming from animation are normalized
+#if DO_CHECK && WITH_EDITORONLY_DATA
+	//check(OutAnimationPoseData.LocalTransforms.IsNormalized());
+#endif
+
+}
+
+void FDecompressionTools::GetBonePose(const UAnimSequence* AnimSequence, const FAnimExtractContext& ExtractionContext, FLODPose& OutAnimationPoseData, bool bForceUseRawData /*= false*/)
+{
+	SCOPE_CYCLE_COUNTER(STAT_AnimSeq_GetBonePose);
+	CSV_SCOPED_TIMING_STAT(Animation, AnimSeq_GetBonePose);
+
+	const TArrayView<const FBoneIndexType> LODBoneIndexes = OutAnimationPoseData.GetLODBoneIndexes();
+
+	check(!bForceUseRawData || CanEvaluateRawAnimationData(AnimSequence));
+	const bool bUseRawDataForPoseExtraction = (CanEvaluateRawAnimationData(AnimSequence) && bForceUseRawData) ||
+#if WITH_EDITOR
+		GetForceRawData() ||
+#endif // WITH_EDITOR
+		UseRawDataForPoseExtraction(AnimSequence, OutAnimationPoseData);
+
+	const bool bIsBakedAdditive = !bUseRawDataForPoseExtraction && AnimSequence->IsValidAdditive();
+
+	const USkeleton* MySkeleton = AnimSequence->GetSkeleton();
+	if (!MySkeleton)
+	{
+		OutAnimationPoseData.SetRefPose(bIsBakedAdditive);
+		return;
+	}
+
+	const bool bDisableRetargeting = OutAnimationPoseData.GetDisableRetargeting();
+
+	// initialize with ref-pose
+	if (bIsBakedAdditive)
+	{
+		//When using baked additive ref pose is identity
+		OutAnimationPoseData.SetRefPose(bIsBakedAdditive);
+	}
+	else
+	{
+		// if retargeting is disabled, we initialize pose with 'Retargeting Source' ref pose.
+		if (bDisableRetargeting)
+		{
+			TArray<FTransform> const& AuthoredOnRefSkeleton = AnimSequence->GetRetargetTransforms();
+
+			const int32 NumBones = LODBoneIndexes.Num();
+			const int32 NumRawBones = AnimSequence->GetSkeleton()->GetReferenceSkeleton().GetRawBoneNum();
+
+			for (int i = 0; i < NumBones; i++)
+			{
+				const int32 SkeletonBoneIndex = LODBoneIndexes[i];
+
+				// Virtual bones are part of the retarget transform pose, so if the pose has not been updated (recently) there might be a mismatch
+				if (SkeletonBoneIndex < NumRawBones || AuthoredOnRefSkeleton.IsValidIndex(SkeletonBoneIndex))
+				{
+					OutAnimationPoseData.LocalTransforms[i] = AuthoredOnRefSkeleton[SkeletonBoneIndex];
+				}
+			}
+		}
+		else
+		{
+			OutAnimationPoseData.SetRefPose();
+		}
+	}
+
+#if WITH_EDITOR
+	const int32 NumTracks = bUseRawDataForPoseExtraction ? AnimSequence->GetDataModelInterface()->GetNumBoneTracks() : AnimSequence->CompressedData.CompressedTrackToSkeletonMapTable.Num();
+#else
+	const int32 NumTracks = AnimSequence->CompressedData.CompressedTrackToSkeletonMapTable.Num();
+#endif 
+	const bool bTreatAnimAsAdditive = (AnimSequence->IsValidAdditive() && !bUseRawDataForPoseExtraction); // Raw data is never additive
+	const FRootMotionReset RootMotionReset(AnimSequence->bEnableRootMotion, AnimSequence->RootMotionRootLock, AnimSequence->bForceRootLock, AnimSequence->ExtractRootTrackTransform(0.0f, /*&RequiredBones*/nullptr), bTreatAnimAsAdditive);
+
+#if WITH_EDITOR
+	// Evaluate raw (source) curve and bone data
+	if (bUseRawDataForPoseExtraction)
+	{
+		// TODO : Curves support
+		//{
+		//	FSkeletonRemappingCurve RemappedCurve(OutAnimationPoseData.GetCurve(), RequiredBones, AnimSequence->GetSkeleton());
+		//	FAnimationPoseData RemappedPoseData(OutAnimationPoseData.GetPose(), RemappedCurve.GetCurve(), OutAnimationPoseData.GetAttributes());
+
+		//	const UE::Anim::DataModel::FEvaluationContext EvaluationContext(ExtractionContext.CurrentTime, DataModelInterface->GetFrameRate(), GetRetargetTransformsSourceName(), GetRetargetTransforms(), Interpolation);
+		//	DataModelInterface->Evaluate(RemappedPoseData, EvaluationContext);
+		//}
+
+		//if ((ExtractionContext.bExtractRootMotion && RootMotionReset.bEnableRootMotion) || RootMotionReset.bForceRootLock)
+		//{
+		//	RootMotionReset.ResetRootBoneForRootMotion(OutPose[FCompactPoseBoneIndex(0)], RequiredBones);
+		//}
+
+		return;
+	}
+	else
+#endif // WITH_EDITOR
+		// Only try and evaluate compressed bone data if the animation contains any bone tracks
+		if (NumTracks != 0)
+		{
+			// Evaluate compressed bone data
+			FAnimSequenceDecompressionContext DecompContext(AnimSequence->GetSamplingFrameRate()
+				, AnimSequence->GetSamplingFrameRate().AsFrameTime(AnimSequence->GetPlayLength()).RoundToFrame().Value
+				, AnimSequence->Interpolation
+				, AnimSequence->GetRetargetTransformsSourceName()
+				, *AnimSequence->CompressedData.CompressedDataStructure
+				, AnimSequence->GetSkeleton()->GetRefLocalPoses()
+				, AnimSequence->CompressedData.CompressedTrackToSkeletonMapTable
+				, AnimSequence->GetSkeleton()
+				, AnimSequence->IsValidAdditive());
+
+			DecompressPose(OutAnimationPoseData, AnimSequence->CompressedData, ExtractionContext, DecompContext, AnimSequence->GetRetargetTransforms(), RootMotionReset);
+		}
+
+//	// TODO : Curves support
+//	// (Always) evaluate compressed curve data
+//	{
+//		// Scoped so that the RemappedCurve destructs and isn't kept alive longer than needed.
+//		FSkeletonRemappingCurve RemappedCurve(OutAnimationPoseData.GetCurve(), RequiredBones, GetSkeleton());
+//#if WITH_EDITOR
+//		// When evaluating from raw animation data, UE::Anim::BuildPoseFromModel will populate the curve data
+//		if (!bUseRawDataForPoseExtraction)
+//#endif // WITH_EDITOR
+//		{
+//			AnimSequence->EvaluateCurveData(RemappedCurve.GetCurve(), ExtractionContext.CurrentTime, bUseRawDataForPoseExtraction);
+//		}
+//	}
+
+	// TODO : attributes support
+	// Evaluate animation attributes (no compressed format yet)
+	//EvaluateAttributes(OutAnimationPoseData, ExtractionContext, false);
+}
+
+void FDecompressionTools::GetBonePose_Additive(const UAnimSequence* AnimSequence, const FAnimExtractContext& ExtractionContext, FLODPose& OutAnimationPoseData)
+{
+}
+
+void FDecompressionTools::GetBonePose_AdditiveMeshRotationOnly(const UAnimSequence* AnimSequence, const FAnimExtractContext& ExtractionContext, FLODPose& OutAnimationPoseData)
+{
+}
+
+
+// --- ---
+
+void FDecompressionTools::DecompressPose(FLODPose& OutAnimationPoseData,
+										const FCompressedAnimSequence& CompressedData,
+										const FAnimExtractContext& ExtractionContext,
+										FAnimSequenceDecompressionContext& DecompressionContext,
+										FName RetargetSource,
+										const FRootMotionReset& RootMotionReset)
+{
+	const TArray<FTransform>& RetargetTransforms = DecompressionContext.GetSourceSkeleton()->GetRefLocalPoses(RetargetSource);
+	DecompressPose(OutAnimationPoseData, CompressedData, ExtractionContext, DecompressionContext, RetargetTransforms, RootMotionReset);
+}
+
+void FDecompressionTools::DecompressPose(FLODPose& OutAnimationPoseData,
+										const FCompressedAnimSequence& CompressedData,
+										const FAnimExtractContext& ExtractionContext,
+										FAnimSequenceDecompressionContext& DecompressionContext,
+										const TArray<FTransform>& RetargetTransforms,
+										const FRootMotionReset& RootMotionReset)
+{
+	const FReferencePose& ReferencePose = OutAnimationPoseData.GetRefPose();
+	const TArrayView<const FBoneIndexType>LODBoneIndexes = OutAnimationPoseData.GetLODBoneIndexes();
+	const TArrayView<const FBoneIndexType> SkeletoonToLODBoneIndexes = ReferencePose.GetSkeletonToLODBoneIndexes(0); // Full list of Skeleton to LOD conversion
+	const int32 NumLODBoneIndexes = LODBoneIndexes.Num();
+
+	const int32 NumTracks = CompressedData.CompressedTrackToSkeletonMapTable.Num();
+
+
+	const USkeleton* TargetSkeleton = OutAnimationPoseData.GetSkeletonAsset();
+	const FSkeletonRemapping& SkeletonRemapping = UE::Anim::FSkeletonRemappingRegistry::Get().GetRemapping(DecompressionContext.GetSourceSkeleton(), TargetSkeleton);
+
+	BoneTrackArray& RotationScalePairs = FGetBonePoseScratchArea::Get().RotationScalePairs;
+	BoneTrackArray& TranslationPairs = FGetBonePoseScratchArea::Get().TranslationPairs;
+	BoneTrackArray& AnimScaleRetargetingPairs = FGetBonePoseScratchArea::Get().AnimScaleRetargetingPairs;
+	BoneTrackArray& AnimRelativeRetargetingPairs = FGetBonePoseScratchArea::Get().AnimRelativeRetargetingPairs;
+	BoneTrackArray& OrientAndScaleRetargetingPairs = FGetBonePoseScratchArea::Get().OrientAndScaleRetargetingPairs;
+
+	// build a list of desired bones
+	RotationScalePairs.Reset();
+	TranslationPairs.Reset();
+	AnimScaleRetargetingPairs.Reset();
+	AnimRelativeRetargetingPairs.Reset();
+	OrientAndScaleRetargetingPairs.Reset();
+
+	// Optimization: assuming first index is root bone. That should always be the case in Skeletons.
+	checkSlow((LODBoneIndexes[0] == FMeshPoseBoneIndex(0).GetInt()));
+	// this is not guaranteed for AnimSequences though... If Root is not animated, Track will not exist.
+	const bool bFirstTrackIsRootBone = (CompressedData.GetSkeletonIndexFromTrackIndex(0) == 0);
+
+	{
+		SCOPE_CYCLE_COUNTER(STAT_BuildAnimTrackPairs);
+
+		// Handle root bone separately if it is track 0. so we start w/ Index 1.
+		for (int32 TrackIndex = (bFirstTrackIsRootBone ? 1 : 0); TrackIndex < NumTracks; TrackIndex++)
+		{
+			const int32 SourceSkeletonBoneIndex = CompressedData.GetSkeletonIndexFromTrackIndex(TrackIndex);
+			const int32 TargetSkeletonBoneIndex = SkeletonRemapping.IsValid() ? SkeletonRemapping.GetTargetSkeletonBoneIndex(SourceSkeletonBoneIndex) : SourceSkeletonBoneIndex;
+
+			if (TargetSkeletonBoneIndex != INDEX_NONE)
+			{
+				const int32 LODBoneIndex = TargetSkeletonBoneIndex < SkeletoonToLODBoneIndexes.Num() ? SkeletoonToLODBoneIndexes[TargetSkeletonBoneIndex] : INDEX_NONE;
+
+				if (LODBoneIndex != INDEX_NONE && LODBoneIndex < NumLODBoneIndexes) // skip bones not in current LOD
+				{
+					RotationScalePairs.Add(BoneTrackPair(LODBoneIndex, TrackIndex));
+
+					// Skip extracting translation component for EBoneTranslationRetargetingMode::Skeleton.
+					switch (TargetSkeleton->GetBoneTranslationRetargetingMode(TargetSkeletonBoneIndex, OutAnimationPoseData.GetDisableRetargeting()))
+					{
+					case EBoneTranslationRetargetingMode::Animation:
+						TranslationPairs.Add(BoneTrackPair(LODBoneIndex, TrackIndex));
+						break;
+					case EBoneTranslationRetargetingMode::AnimationScaled:
+						TranslationPairs.Add(BoneTrackPair(LODBoneIndex, TrackIndex));
+						AnimScaleRetargetingPairs.Add(BoneTrackPair(LODBoneIndex, SourceSkeletonBoneIndex));
+						break;
+					case EBoneTranslationRetargetingMode::AnimationRelative:
+						TranslationPairs.Add(BoneTrackPair(LODBoneIndex, TrackIndex));
+
+						// With baked additives, we can skip 'AnimationRelative' tracks, as the relative transform gets canceled out.
+						// (A1 + Rel) - (A2 + Rel) = A1 - A2.
+						if (!DecompressionContext.IsAdditiveAnimation())
+						{
+							AnimRelativeRetargetingPairs.Add(BoneTrackPair(LODBoneIndex, SourceSkeletonBoneIndex));
+						}
+						break;
+					case EBoneTranslationRetargetingMode::OrientAndScale:
+						TranslationPairs.Add(BoneTrackPair(LODBoneIndex, TrackIndex));
+
+						// Additives remain additives, they're not retargeted.
+						if (!DecompressionContext.IsAdditiveAnimation())
+						{
+							OrientAndScaleRetargetingPairs.Add(BoneTrackPair(LODBoneIndex, SourceSkeletonBoneIndex));
+						}
+						break;
+					}
+				}
+			}
+		}
+	}
+
+	{
+		SCOPE_CYCLE_COUNTER(STAT_ExtractPoseFromAnimData);
+		CSV_SCOPED_TIMING_STAT(Animation, ExtractPoseFromAnimData);
+		CSV_CUSTOM_STAT(Animation, NumberOfExtractedAnimations, 1, ECsvCustomStatOp::Accumulate);
+
+		DecompressionContext.Seek(ExtractionContext.CurrentTime);
+
+		// Handle Root Bone separately
+		if (bFirstTrackIsRootBone)
+		{
+			const int32 TrackIndex = 0;
+			const int32 LODRootBone = 0;
+			FTransform RootAtom = OutAnimationPoseData.LocalTransforms[0];
+
+			CompressedData.BoneCompressionCodec->DecompressBone(DecompressionContext, TrackIndex, RootAtom);
+
+			// Retarget the root onto the target skeleton (correcting for differences in rest poses)
+			if (SkeletonRemapping.RequiresReferencePoseRetarget())
+			{
+				const int32 TargetSkeletonBoneIndex = 0;
+
+				if (DecompressionContext.IsAdditiveAnimation())
+				{
+					RootAtom.SetRotation(SkeletonRemapping.RetargetAdditiveRotationToTargetSkeleton(TargetSkeletonBoneIndex, RootAtom.GetRotation()));
+					if (TargetSkeleton->GetBoneTranslationRetargetingMode(TargetSkeletonBoneIndex, OutAnimationPoseData.GetDisableRetargeting()) != EBoneTranslationRetargetingMode::Skeleton)
+					{
+						RootAtom.SetTranslation(SkeletonRemapping.RetargetAdditiveTranslationToTargetSkeleton(TargetSkeletonBoneIndex, RootAtom.GetTranslation()));
+					}
+				}
+				else
+				{
+					RootAtom.SetRotation(SkeletonRemapping.RetargetBoneRotationToTargetSkeleton(TargetSkeletonBoneIndex, RootAtom.GetRotation()));
+					if (TargetSkeleton->GetBoneTranslationRetargetingMode(TargetSkeletonBoneIndex, OutAnimationPoseData.GetDisableRetargeting()) != EBoneTranslationRetargetingMode::Skeleton)
+					{
+						RootAtom.SetTranslation(SkeletonRemapping.RetargetBoneTranslationToTargetSkeleton(TargetSkeletonBoneIndex, RootAtom.GetTranslation()));
+					}
+				}
+			}
+
+			// @laurent - we should look into splitting rotation and translation tracks, so we don't have to process translation twice.
+			RetargetBoneTransform(OutAnimationPoseData.GetRefPose()
+				, DecompressionContext.GetSourceSkeleton()
+				, OutAnimationPoseData.GetSkeletonAsset()
+				, DecompressionContext.AnimName
+				, RetargetTransforms
+				, RootAtom
+				, 0
+				, LODRootBone
+				, DecompressionContext.IsAdditiveAnimation()
+				, OutAnimationPoseData.GetDisableRetargeting());
+
+			OutAnimationPoseData.LocalTransforms[0] = RootAtom;
+		}
+
+		if (RotationScalePairs.Num() > 0)
+		{
+#if DEFAULT_SOA
+			// get the remaining bone atoms
+			CompressedData.BoneCompressionCodec->DecompressPose(DecompressionContext, UE::Anim::FAnimPoseDecompressionData(RotationScalePairs, TranslationPairs, RotationScalePairs, OutAnimationPoseData.LocalTransforms.Rotations, OutAnimationPoseData.LocalTransforms.Translations, OutAnimationPoseData.LocalTransforms.Scales3D));
+#else
+			// get the remaining bone atoms
+			TArrayView<FTransform> OutPoseBones = OutAnimationPoseData.LocalTransforms.Transforms;
+			CompressedData.BoneCompressionCodec->DecompressPose(DecompressionContext, RotationScalePairs, TranslationPairs, RotationScalePairs, OutPoseBones);
+#endif
+		}
+	}
+
+	// Retarget the pose onto the target skeleton (correcting for differences in rest poses)
+	if (SkeletonRemapping.RequiresReferencePoseRetarget())
+	{
+		const int32 LODNumBones = OutAnimationPoseData.GetNumBones();
+
+		if (DecompressionContext.IsAdditiveAnimation())
+		{
+			for (int32 LODBoneIndex = bFirstTrackIsRootBone ? 1 : 0; LODBoneIndex < LODNumBones; ++LODBoneIndex)
+			{
+				const int32 TargetSkeletonBoneIndex = ReferencePose.GetSkeletonBoneIndexFromLODBoneIndex(LODBoneIndex);
+				OutAnimationPoseData.LocalTransforms[LODBoneIndex].SetRotation(SkeletonRemapping.RetargetAdditiveRotationToTargetSkeleton(TargetSkeletonBoneIndex, OutAnimationPoseData.LocalTransforms[LODBoneIndex].GetRotation()));
+				if (TargetSkeleton->GetBoneTranslationRetargetingMode(TargetSkeletonBoneIndex, OutAnimationPoseData.GetDisableRetargeting()) != EBoneTranslationRetargetingMode::Skeleton)
+				{
+					OutAnimationPoseData.LocalTransforms[LODBoneIndex].SetTranslation(SkeletonRemapping.RetargetAdditiveTranslationToTargetSkeleton(TargetSkeletonBoneIndex, OutAnimationPoseData.LocalTransforms[LODBoneIndex].GetTranslation()));
+				}
+			}
+		}
+		else
+		{
+			for (int32 LODBoneIndex = (bFirstTrackIsRootBone ? 1 : 0); LODBoneIndex < LODNumBones; ++LODBoneIndex)
+			{
+				const int32 TargetSkeletonBoneIndex = LODBoneIndexes[LODBoneIndex]; // ReferencePose.GetSkeletonBoneIndexFromLODBoneIndex(LODBoneIndex);
+				OutAnimationPoseData.LocalTransforms[LODBoneIndex].SetRotation(SkeletonRemapping.RetargetBoneRotationToTargetSkeleton(TargetSkeletonBoneIndex, OutAnimationPoseData.LocalTransforms[LODBoneIndex].GetRotation()));
+				if (TargetSkeleton->GetBoneTranslationRetargetingMode(TargetSkeletonBoneIndex, OutAnimationPoseData.GetDisableRetargeting()) != EBoneTranslationRetargetingMode::Skeleton)
+				{
+					OutAnimationPoseData.LocalTransforms[LODBoneIndex].SetTranslation(SkeletonRemapping.RetargetBoneTranslationToTargetSkeleton(TargetSkeletonBoneIndex, OutAnimationPoseData.LocalTransforms[LODBoneIndex].GetTranslation()));
+				}
+			}
+		}
+	}
+
+	// Once pose has been extracted, snap root bone back to first frame if we are extracting root motion.
+	if ((ExtractionContext.bExtractRootMotion && RootMotionReset.bEnableRootMotion) || RootMotionReset.bForceRootLock)
+	{
+		//RootMotionReset.ResetRootBoneForRootMotion(OutAnimationPoseData.LocalTransforms[FCompactPoseBoneIndex(0)], RequiredBones);
+	}
+
+	// Anim Scale Retargeting
+	int32 const NumBonesToScaleRetarget = AnimScaleRetargetingPairs.Num();
+	if (NumBonesToScaleRetarget > 0)
+	{
+		TArray<FTransform> const& AuthoredOnRefSkeleton = RetargetTransforms;
+
+		for (const BoneTrackPair& BonePair : AnimScaleRetargetingPairs)
+		{
+			const int32 LODBoneIndex(BonePair.AtomIndex);
+			const int32 SourceSkeletonBoneIndex = BonePair.TrackIndex;
+
+			// @todo - precache that in FBoneContainer when we have SkeletonIndex->TrackIndex mapping. So we can just apply scale right away.
+			float const SourceTranslationLength = AuthoredOnRefSkeleton[SourceSkeletonBoneIndex].GetTranslation().Size();
+			if (SourceTranslationLength > UE_KINDA_SMALL_NUMBER)
+			{
+				float const TargetTranslationLength = ReferencePose.GetRefPoseTranslation(LODBoneIndex).Size();
+				OutAnimationPoseData.LocalTransforms[LODBoneIndex].ScaleTranslation(TargetTranslationLength / SourceTranslationLength);
+			}
+		}
+	}
+
+	// Anim Relative Retargeting
+	int32 const NumBonesToRelativeRetarget = AnimRelativeRetargetingPairs.Num();
+	if (NumBonesToRelativeRetarget > 0)
+	{
+		TArray<FTransform> const& AuthoredOnRefSkeleton = RetargetTransforms;
+
+		for (const BoneTrackPair& BonePair : AnimRelativeRetargetingPairs)
+		{
+			const int32 LODBoneIndex(BonePair.AtomIndex);
+			const int32 SourceSkeletonBoneIndex = BonePair.TrackIndex;
+
+			const FTransform& RefPoseTransform = ReferencePose.GetRefPoseTransform(LODBoneIndex);
+
+			// Remap the base pose onto the target skeleton so that we are working entirely in target space
+			FTransform BaseTransform = AuthoredOnRefSkeleton[SourceSkeletonBoneIndex];
+			if (SkeletonRemapping.RequiresReferencePoseRetarget())
+			{
+				const int32 TargetSkeletonBoneIndex = SkeletonRemapping.GetTargetSkeletonBoneIndex(SourceSkeletonBoneIndex);
+				BaseTransform = SkeletonRemapping.RetargetBoneTransformToTargetSkeleton(TargetSkeletonBoneIndex, BaseTransform);
+			}
+
+			// Apply the retargeting as if it were an additive difference between the current skeleton and the retarget skeleton. 
+			OutAnimationPoseData.LocalTransforms[LODBoneIndex].SetRotation(OutAnimationPoseData.LocalTransforms[LODBoneIndex].GetRotation() * BaseTransform.GetRotation().Inverse() * RefPoseTransform.GetRotation());
+			OutAnimationPoseData.LocalTransforms[LODBoneIndex].SetTranslation(OutAnimationPoseData.LocalTransforms[LODBoneIndex].GetTranslation() + (RefPoseTransform.GetTranslation() - BaseTransform.GetTranslation()));
+			OutAnimationPoseData.LocalTransforms[LODBoneIndex].SetScale3D(OutAnimationPoseData.LocalTransforms[LODBoneIndex].GetScale3D() * (RefPoseTransform.GetScale3D() * BaseTransform.GetSafeScaleReciprocal(BaseTransform.GetScale3D())));
+			OutAnimationPoseData.LocalTransforms[LODBoneIndex].NormalizeRotation();
+		}
+	}
+
+	// TODO : Have to recreate GetRetargetSourceCachedData
+	// Translation 'Orient and Scale' Translation Retargeting
+	//const int32 NumBonesToOrientAndScaleRetarget = OrientAndScaleRetargetingPairs.Num();
+	//if (NumBonesToOrientAndScaleRetarget > 0)
+	//{
+	//	const FRetargetSourceCachedData& RetargetSourceCachedData = RequiredBones.GetRetargetSourceCachedData(DecompressionContext.AnimName, SkeletonRemapping, RetargetTransforms);
+	//	const TArray<FOrientAndScaleRetargetingCachedData>& OrientAndScaleDataArray = RetargetSourceCachedData.OrientAndScaleData;
+	//	const TArray<int32>& CompactPoseIndexToOrientAndScaleIndex = RetargetSourceCachedData.CompactPoseIndexToOrientAndScaleIndex;
+
+	//	// If we have any cached retargeting data.
+	//	if ((OrientAndScaleDataArray.Num() > 0) && (CompactPoseIndexToOrientAndScaleIndex.Num() == RequiredBones.GetCompactPoseNumBones()))
+	//	{
+	//		for (int32 Index = 0; Index < NumBonesToOrientAndScaleRetarget; Index++)
+	//		{
+	//			const BoneTrackPair& BonePair = OrientAndScaleRetargetingPairs[Index];
+	//			const FCompactPoseBoneIndex CompactPoseBoneIndex(BonePair.AtomIndex);
+	//			const int32 OrientAndScaleIndex = CompactPoseIndexToOrientAndScaleIndex[CompactPoseBoneIndex.GetInt()];
+	//			if (OrientAndScaleIndex != INDEX_NONE)
+	//			{
+	//				const FOrientAndScaleRetargetingCachedData& OrientAndScaleData = OrientAndScaleDataArray[OrientAndScaleIndex];
+	//				FTransform& BoneTransform = OutAnimationPoseData.LocalTransforms[CompactPoseBoneIndex];
+	//				const FVector AnimatedTranslation = BoneTransform.GetTranslation();
+
+	//				// If Translation is not animated, we can just copy the TargetTranslation. No retargeting needs to be done.
+	//				const FVector NewTranslation = (AnimatedTranslation - OrientAndScaleData.SourceTranslation).IsNearlyZero(BONE_TRANS_RT_ORIENT_AND_SCALE_PRECISION) ?
+	//					OrientAndScaleData.TargetTranslation :
+	//					OrientAndScaleData.TranslationDeltaOrient.RotateVector(AnimatedTranslation) * OrientAndScaleData.TranslationScale;
+
+	//				BoneTransform.SetTranslation(NewTranslation);
+	//			}
+	//		}
+	//	}
+	//}
+}
+
+void FDecompressionTools::RetargetBoneTransform(const FReferencePose& ReferencePose
+	, const USkeleton* SourceSkeleton
+	, const USkeleton* TargetSkeleton
+	, const FName& RetargetSource
+	, FTransform& BoneTransform
+	, const int32 SkeletonBoneIndex
+	, const int32 LODBoneIndex
+	, const bool bIsBakedAdditive
+	, const bool bDisableRetargeting)
+{
+	if (SourceSkeleton)
+	{
+		const TArray<FTransform>& RetargetTransforms = SourceSkeleton->GetRefLocalPoses(RetargetSource);
+		RetargetBoneTransform(ReferencePose, SourceSkeleton, TargetSkeleton, RetargetSource, RetargetTransforms, BoneTransform, SkeletonBoneIndex, LODBoneIndex, bIsBakedAdditive, bDisableRetargeting);
+	}
+}
+
+void FDecompressionTools::RetargetBoneTransform(const FReferencePose& ReferencePose
+	, const USkeleton* SourceSkeleton
+	, const USkeleton* TargetSkeleton
+	, const FName& SourceName
+	, const TArray<FTransform>& RetargetTransforms
+	, FTransform& BoneTransform
+	, const int32 SkeletonBoneIndex
+	, const int32 LODBoneIndex
+	, const bool bIsBakedAdditive
+	, const bool bDisableRetargeting)
+{
+	check(!RetargetTransforms.IsEmpty());
+	if (SourceSkeleton)
+	{
+		//const USkeleton* TargetSkeleton = RequiredBones.GetSkeletonAsset();
+		const FSkeletonRemapping& SkeletonRemapping = UE::Anim::FSkeletonRemappingRegistry::Get().GetRemapping(SourceSkeleton, TargetSkeleton);
+
+		const int32 TargetSkeletonBoneIndex = ReferencePose.GetSkeletonBoneIndexFromLODBoneIndex(LODBoneIndex);
+		const int32 SourceSkeletonBoneIndex = SkeletonRemapping.IsValid() ? SkeletonRemapping.GetSourceSkeletonBoneIndex(TargetSkeletonBoneIndex) : SkeletonBoneIndex;
+
+		switch (TargetSkeleton->GetBoneTranslationRetargetingMode(TargetSkeletonBoneIndex, bDisableRetargeting))
+		{
+		case EBoneTranslationRetargetingMode::AnimationScaled:
+		{
+			// @todo - precache that in FBoneContainer when we have SkeletonIndex->TrackIndex mapping. So we can just apply scale right away.
+			const TArray<FTransform>& SkeletonRefPoseArray = RetargetTransforms;
+			const float SourceTranslationLength = SkeletonRefPoseArray[SourceSkeletonBoneIndex].GetTranslation().Size();
+			if (SourceTranslationLength > UE_KINDA_SMALL_NUMBER)
+			{
+				const float TargetTranslationLength = ReferencePose.GetRefPoseTranslation(LODBoneIndex).Size();
+				BoneTransform.ScaleTranslation(TargetTranslationLength / SourceTranslationLength);
+			}
+			break;
+		}
+
+		case EBoneTranslationRetargetingMode::Skeleton:
+		{
+			BoneTransform.SetTranslation(bIsBakedAdditive ? FVector::ZeroVector : ReferencePose.GetRefPoseTranslation(LODBoneIndex));
+			break;
+		}
+
+		case EBoneTranslationRetargetingMode::AnimationRelative:
+		{
+			// With baked additive animations, Animation Relative delta gets canceled out, so we can skip it.
+			// (A1 + Rel) - (A2 + Rel) = A1 - A2.
+			if (!bIsBakedAdditive)
+			{
+				const TArray<FTransform>& AuthoredOnRefSkeleton = RetargetTransforms;
+				const FTransform& RefPoseTransform = ReferencePose.GetRefPoseTransform(LODBoneIndex);
+
+				// Remap the base pose onto the target skeleton so that we are working entirely in target space
+				FTransform BaseTransform = AuthoredOnRefSkeleton[SourceSkeletonBoneIndex];
+				if (SkeletonRemapping.RequiresReferencePoseRetarget())
+				{
+					BaseTransform = SkeletonRemapping.RetargetBoneTransformToTargetSkeleton(TargetSkeletonBoneIndex, BaseTransform);
+				}
+
+				// Apply the retargeting as if it were an additive difference between the current skeleton and the retarget skeleton. 
+				BoneTransform.SetRotation(BoneTransform.GetRotation() * BaseTransform.GetRotation().Inverse() * RefPoseTransform.GetRotation());
+				BoneTransform.SetTranslation(BoneTransform.GetTranslation() + (RefPoseTransform.GetTranslation() - BaseTransform.GetTranslation()));
+				BoneTransform.SetScale3D(BoneTransform.GetScale3D() * (RefPoseTransform.GetScale3D() * BaseTransform.GetSafeScaleReciprocal(BaseTransform.GetScale3D())));
+				BoneTransform.NormalizeRotation();
+			}
+			break;
+		}
+
+		case EBoneTranslationRetargetingMode::OrientAndScale:
+		{
+			if (!bIsBakedAdditive)
+			{
+				//const FRetargetSourceCachedData& RetargetSourceCachedData = RequiredBones.GetRetargetSourceCachedData(SourceName, SkeletonRemapping, RetargetTransforms);
+				//const TArray<FOrientAndScaleRetargetingCachedData>& OrientAndScaleDataArray = RetargetSourceCachedData.OrientAndScaleData;
+				//const TArray<int32>& CompactPoseIndexToOrientAndScaleIndex = RetargetSourceCachedData.CompactPoseIndexToOrientAndScaleIndex;
+
+				//// If we have any cached retargeting data.
+				//if ((OrientAndScaleDataArray.Num() > 0) && (CompactPoseIndexToOrientAndScaleIndex.Num() == RequiredBones.GetCompactPoseNumBones()))
+				//{
+				//	const int32 OrientAndScaleIndex = CompactPoseIndexToOrientAndScaleIndex[BoneIndex.GetInt()];
+				//	if (OrientAndScaleIndex != INDEX_NONE)
+				//	{
+				//		const FOrientAndScaleRetargetingCachedData& OrientAndScaleData = OrientAndScaleDataArray[OrientAndScaleIndex];
+				//		const FVector AnimatedTranslation = BoneTransform.GetTranslation();
+
+				//		// If Translation is not animated, we can just copy the TargetTranslation. No retargeting needs to be done.
+				//		const FVector NewTranslation = (AnimatedTranslation - OrientAndScaleData.SourceTranslation).IsNearlyZero(BONE_TRANS_RT_ORIENT_AND_SCALE_PRECISION) ?
+				//			OrientAndScaleData.TargetTranslation :
+				//			OrientAndScaleData.TranslationDeltaOrient.RotateVector(AnimatedTranslation) * OrientAndScaleData.TranslationScale;
+
+				//		BoneTransform.SetTranslation(NewTranslation);
+				//	}
+				//}
+			}
+			break;
+		}
+		}
+	}
+}
+
+} // end namespace UE::AnimNext
diff --git a/Engine/Plugins/Experimental/AnimNext/Source/AnimNext/Private/Interface/InterfaceContext.cpp b/Engine/Plugins/Experimental/AnimNext/Source/AnimNext/Private/Interface/InterfaceContext.cpp
new file mode 100644
--- /dev/null
+++ b/Engine/Plugins/Experimental/AnimNext/Source/AnimNext/Private/Interface/InterfaceContext.cpp
@@ -0,0 +1,273 @@
+// Copyright Epic Games, Inc. All Rights Reserved.
+
+#include "Interface/InterfaceContext.h"
+#include "Interface/IAnimNextInterface.h"
+#include "Param/IAnimNextParamInterface.h"
+#include "Param/ParamType.h"
+#include "Param/ParamHelpers.h"
+
+namespace UE::AnimNext
+{
+
+struct FThreadContextData : TThreadSingleton<FThreadContextData>
+{
+	TArray<const FContext*> ContextStack;
+};
+
+// @TODO: we might need to do this for all FContext constructions?
+// As parameters are gathered using the FContext linked list, we always need to access the top of the stack to get 
+// correct parameters. Shortcutting to the bottom of the stack via FThreadContext misses intervening parameters
+FThreadContext::FThreadContext(const FContext& InContext)
+{
+	FThreadContextData& ContextData = FThreadContextData::Get();
+	ContextData.ContextStack.Push(&InContext);
+}
+
+FThreadContext::~FThreadContext()
+{
+	FThreadContextData& ContextData = FThreadContextData::Get();
+	ContextData.ContextStack.Pop();
+}
+
+const FContext& FThreadContext::Get()
+{
+	const FThreadContextData& ContextData = FThreadContextData::Get();
+	return *ContextData.ContextStack.Top();
+}
+
+FContext::FContext(float InDeltaTime, FState& InState, FParamStorage& InParamStorage, IAnimNextParamInterface* InParameters)
+	: State(&InState)
+	, ParamStorage(&InParamStorage)
+	, Parameters(InParameters)
+	, DeltaTime(InDeltaTime)
+{
+}
+
+FContext::~FContext()
+{
+	// Remove any handles before the storage goes out of scope
+	AdditionalParameterHandles.Empty();
+
+	if (BlockHandle != InvalidBlockHandle)
+	{
+		ParamStorage->ReleaseBlock(BlockHandle);
+	}
+}
+
+FContext FContext::WithResult(FParam& InResult) const
+{
+	FContext NewContext(DeltaTime, *State, *ParamStorage, InResult);
+	NewContext.Parent = this;
+	NewContext.Root = (Root == nullptr) ? this : Root;
+	NewContext.CallstackHash = CallstackHash;
+
+	return NewContext;
+}
+
+FContext FContext::WithParameter(FName ParameterId, const FParam& InParameter) const
+{
+	FContext NewContext(DeltaTime, *State, *ParamStorage, *Result);
+	NewContext.Parent = this;
+	NewContext.Root = (Root == nullptr) ? this : Root;
+	NewContext.CallstackHash = CallstackHash;
+	NewContext.AddParameters({ TPairInitializer(ParameterId, InParameter) });
+	return NewContext;
+}
+
+FContext FContext::WithParameters(TArrayView<const TPair<FName, FParam>> InParameters) const
+{
+	FContext NewContext(DeltaTime, *State, *ParamStorage, *Result);
+	NewContext.Parent = this;
+	NewContext.Root = (Root == nullptr) ? this : Root;
+	NewContext.CallstackHash = CallstackHash;
+	NewContext.AddParameters(InParameters);
+
+	return NewContext;
+}
+
+FContext FContext::WithResultAndParameters(FParam& InResult, TArrayView<const TPair<FName, FParam>> InParameters) const
+{
+	FContext NewContext(DeltaTime, *State, *ParamStorage, InResult);
+	NewContext.Parent = this;
+	NewContext.Root = (Root == nullptr) ? this : Root;
+	NewContext.CallstackHash = CallstackHash;
+	NewContext.AddParameters(InParameters);
+
+	return NewContext;
+}
+
+int32 FContext::GetParametersSize(TArrayView<const TPair<FName, FParam>> InParameters, TArray<int32> &ParamAllocSizes) const
+{
+	int32 TotalParamAllocSize = 0;
+
+	const int32 NumParameters = InParameters.Num();
+
+	for (int i = 0; i < NumParameters; i++)
+	{
+		const TPair<FName, FParam>& ParamPair = InParameters[i];
+		const FParam& SourceParam = ParamPair.Value;
+
+		if (EnumHasAnyFlags(SourceParam.GetFlags(), FParam::EFlags::Stored))
+		{
+			const FParamTypeHandle ParamTypeHandle = SourceParam.GetTypeHandle();
+
+			const int32 ParamAlignment = ParamTypeHandle.GetAlignment();
+			const int32 ParamSize = ParamTypeHandle.GetSize();
+
+			const int32 ParamAllocSize = Align(ParamSize, ParamAlignment);
+
+			ParamAllocSizes[i] = ParamAllocSize;
+			TotalParamAllocSize += ParamAllocSize;
+		}
+	}
+
+	return TotalParamAllocSize;
+}
+
+void FContext::AddParameters(TArrayView<const TPair<FName, FParam>> InParameters)
+{
+	int32 TotalParamAllocSize = 0;
+
+	const int32 NumParameters = InParameters.Num();
+
+	TArray<int32> ParamAllocSizes;
+	ParamAllocSizes.Reserve(NumParameters);
+	ParamAllocSizes.AddZeroed(NumParameters);
+
+	TotalParamAllocSize = GetParametersSize(InParameters, ParamAllocSizes);
+
+	uint8* TargetMemory = nullptr;
+	check(BlockHandle == InvalidBlockHandle);
+
+	if (TotalParamAllocSize > 0)
+	{
+		const FParamStorage::TBlockDataPair BlockData = ParamStorage->RequestBlock(TotalParamAllocSize);
+
+		BlockHandle = BlockData.Key;
+		TargetMemory = BlockData.Value;
+	}
+
+	for(int i = 0; i < NumParameters; i++)
+	{
+		const TPair<FName, FParam>& ParamPair = InParameters[i];
+		const FParam& SourceParam = ParamPair.Value;
+		const int32 ParamAllocatedSize = ParamAllocSizes[i];
+
+		if (ParamAllocatedSize > 0 && EnumHasAnyFlags(SourceParam.GetFlags(), FParam::EFlags::Stored))
+		{
+			const FParamTypeHandle ParamTypeHandle = ParamPair.Value.GetTypeHandle();
+
+			TArrayView<uint8> ClonedParamMemory(TargetMemory, ParamTypeHandle.GetSize());
+			FParam ClonedParam = FParam::DuplicateParam(ParamPair.Value, ClonedParamMemory);
+
+			TargetMemory += ParamAllocatedSize;
+
+			AdditionalParameters.Add(ParamPair.Key, ClonedParam);
+		}
+		else
+		{
+			AdditionalParameters.Add(ParamPair.Key, ParamPair.Value);
+		}
+	}
+}
+
+FContext FContext::WithParameters(IAnimNextParamInterface* InParameters) const
+{
+	FContext NewContext(DeltaTime, *State, *ParamStorage, *Result);
+	NewContext.Parameters = InParameters;
+	NewContext.Parent = this;
+	NewContext.Root = (Root == nullptr) ? this : Root;
+	NewContext.CallstackHash = CallstackHash;
+
+	return NewContext;
+}
+
+FContext FContext::CreateSubContext() const
+{
+	FContext NewContext(DeltaTime, *State, *ParamStorage, *Result);
+	NewContext.Parent = this;
+	NewContext.Root = (Root == nullptr) ? this : Root;
+	NewContext.CallstackHash = CallstackHash;
+
+	return NewContext;
+}
+
+bool FContext::GetParameter(FName InKey, FParam& OutParam) const
+{
+	// Check parent contexts in turn. Parameters can be overriden by each scoped context.
+	for(const FContext* CurrentContext = this; CurrentContext != nullptr; CurrentContext = CurrentContext->Parent)
+	{
+		// Check parameter provider
+		if(CurrentContext->Parameters)
+		{
+			if(CurrentContext->Parameters->GetParameter(InKey, OutParam))
+			{
+				return true;
+			}
+		}
+		// Check additional parameters
+		// Note this is currently 'else if' because we dont allow creating new contexts with both a IAnimNextParameters
+		// interface and AdditionalParameters
+		else if(!CurrentContext->AdditionalParameters.IsEmpty())
+		{
+			// Find the parameter
+			if(const FParam* FoundParameter = CurrentContext->AdditionalParameters.Find(InKey))
+			{
+				// Check type compatibility
+				if(FoundParameter->CanAssignTo(OutParam))
+				{
+					OutParam = *FoundParameter;
+					return true;
+				}
+			}
+		}
+		// Check additional parameter handles
+		// Note this is currently 'else if' because we dont allow creating new contexts with both 
+		// AdditionalParameters and AdditionalParametersHandles
+		else if (!CurrentContext->AdditionalParameterHandles.IsEmpty())
+		{
+			// Find the parameter
+			if (const FParamHandle* HParam = CurrentContext->AdditionalParameterHandles.Find(InKey))
+			{
+				const FParam* Param = ParamStorage->GetParam(HParam->ParamHandle);
+				if (Param != nullptr && Param->CanAssignTo(OutParam))
+				{
+					OutParam = Param;
+					return true;
+				}
+			}
+		}
+	}
+
+	return false;
+}
+
+
+TParam<const float> FContext::GetDeltaTimeParam() const
+{
+	return TWrapParam<const float>(&DeltaTime);
+}
+
+FContext FContext::WithCallRaw(const IAnimNextInterface* InAnimNextInterface) const
+{
+	FContext NewContext(DeltaTime, *State, *ParamStorage, *Result);
+	NewContext.Parent = this;
+	NewContext.Root = (Root == nullptr) ? this : Root;
+	NewContext.CallstackHash = HashCombineFast(CallstackHash, GetTypeHash(InAnimNextInterface));
+	return NewContext;
+}
+
+void FContext::FlushRelevancy() const
+{
+	// Flush all relevancy-based allocations if they were not used this update
+	// @TODO: this could be more efficient if we use a more linear iteration here 
+	for(auto It = State->RelevancyValueMap.CreateIterator(); It; ++It)
+	{
+		if(It->Value.UpdateCounter != UpdateCounter)
+		{
+			It.RemoveCurrent();
+		}
+	}
+}
+
+}
diff --git a/Engine/Plugins/Experimental/AnimNext/Source/AnimNext/Private/Interface/InterfaceState.cpp b/Engine/Plugins/Experimental/AnimNext/Source/AnimNext/Private/Interface/InterfaceState.cpp
new file mode 100644
--- /dev/null
+++ b/Engine/Plugins/Experimental/AnimNext/Source/AnimNext/Private/Interface/InterfaceState.cpp
@@ -0,0 +1,56 @@
+﻿// Copyright Epic Games, Inc. All Rights Reserved.
+
+#include "Interface/InterfaceState.h"
+#include "Interface/InterfaceContext.h"
+
+namespace UE::AnimNext
+{
+
+FParam* FState::FindStateRaw(const FInterfaceKeyWithIdAndStack& InKey, const FContext& InContext, EStatePersistence InPersistence)
+{
+	FParam* ExistingValue = nullptr;
+	TMap<FInterfaceKeyWithIdAndStack, FParam>* ValueMap = nullptr; 
+
+	switch(InPersistence)
+	{
+	case EStatePersistence::Relevancy:
+		ExistingValue = RelevancyValueMap.Find(InKey);
+		if(ExistingValue)
+		{
+			// An access to a state param counts as a 'relevant use' for this update
+			static_cast<FRelevancyParam*>(ExistingValue)->UpdateCounter = InContext.UpdateCounter;
+		}
+		break;
+	case EStatePersistence::Permanent:
+		ExistingValue = PermanentValueMap.Find(InKey);
+		break;
+	default:
+		check(false);
+		break;
+	}
+	
+	return ExistingValue;
+}
+
+FParam* FState::AllocateState(const FInterfaceKeyWithIdAndStack& InKey, const FContext& InContext, const FParamTypeHandle& InTypeHandle, EStatePersistence InPersistence)
+{
+	const TArrayView<uint8> Data(static_cast<uint8*>(FMemory::Malloc(InTypeHandle.GetSize(), InTypeHandle.GetAlignment())), InTypeHandle.GetSize());
+
+	switch(InPersistence)
+	{
+	case EStatePersistence::Relevancy:
+		// TODO: chunked allocator for relevancy-based stuff?
+		return &RelevancyValueMap.Add(InKey, FRelevancyParam(InTypeHandle, Data, FParam::EFlags::Mutable, InContext.UpdateCounter));
+		break;
+	case EStatePersistence::Permanent:
+		return &PermanentValueMap.Add(InKey, FParam(InTypeHandle, Data, FParam::EFlags::Mutable));
+		break;
+	default:
+		check(false);
+		break;
+	}
+	
+	return nullptr;
+}
+
+}
diff --git a/Engine/Plugins/Experimental/AnimNext/Source/AnimNext/Private/Module.cpp b/Engine/Plugins/Experimental/AnimNext/Source/AnimNext/Private/Module.cpp
new file mode 100644
--- /dev/null
+++ b/Engine/Plugins/Experimental/AnimNext/Source/AnimNext/Private/Module.cpp
@@ -0,0 +1,32 @@
+// Copyright Epic Games, Inc. All Rights Reserved.
+
+#include "CoreMinimal.h"
+#include "Modules/ModuleInterface.h"
+#include "Modules/ModuleManager.h"
+#include "Misc/CoreDelegates.h"
+#include "DataRegistry.h"
+
+namespace UE::AnimNext
+{
+
+class FModule : public IModuleInterface
+{
+public:
+	virtual void StartupModule() override
+	{
+		FCoreDelegates::OnFEngineLoopInitComplete.AddLambda([]()
+		{
+			FDataRegistry::Init();
+		});
+
+		FCoreDelegates::OnEnginePreExit.AddLambda([]()
+		{
+			FDataRegistry::Destroy();
+		});
+		
+	}
+};
+
+}
+
+IMPLEMENT_MODULE(UE::AnimNext::FModule, AnimNext)
\ No newline at end of file
diff --git a/Engine/Plugins/Experimental/AnimNext/Source/AnimNext/Public/DataRegistry.h b/Engine/Plugins/Experimental/AnimNext/Source/AnimNext/Public/DataRegistry.h
new file mode 100644
--- /dev/null
+++ b/Engine/Plugins/Experimental/AnimNext/Source/AnimNext/Public/DataRegistry.h
@@ -0,0 +1,236 @@
+// Copyright Epic Games, Inc. All Rights Reserved.
+
+#pragma once
+
+#include "CoreMinimal.h"
+#include "DataRegistryTypes.h"
+#include "UObject/WeakObjectPtr.h"
+#include "Param/Param.h"
+
+class USkeletalMeshComponent;
+struct FReferenceSkeleton;
+
+namespace UE::AnimNext
+{
+
+class FModule;
+
+struct FDataHandle;
+
+// Global registry of animation data
+// Holds ref counted data that gets released when the last DataHandle of that element goes out of scope
+// Calling public functions from multiple threads is expected. Data races are guarded by a FRWLock.
+// TODO : Memory management will have to be implemented to avoid fragmentation and performance reasons
+class ANIMNEXT_API FDataRegistry
+{
+public:
+
+	// Access the global registry
+	static FDataRegistry* Get();
+
+
+	// --- Reference Pose Handling ---
+
+	// Generates and registers a reference pose for the SkeletalMesh asset of the SkeletalMeshComponent
+	// and modifies it with the additional required bones 
+	// or the visibility state of the bones of the SkeletalMeshComponent
+	FDataHandle RegisterReferencePose(USkeletalMeshComponent* SkeletalMeshComponent);
+
+	// Returns a ref counted handle to the refence pose of the given skeletal mesh component
+	FDataHandle GetOrGenerateReferencePose(USkeletalMeshComponent* SkeletalMeshComponent);
+
+	// Removes a previously registered reference pose for the given SkeletalMeshComponent
+	void RemoveReferencePose(USkeletalMeshComponent* SkeletalMeshComponent);
+
+	
+	// --- AnimationData Storage / Retrieval  --- 
+
+	// Registers an anim data handle with arbitrary data using an FName
+	// Note that AnimDataHandles are refcounted, so this makes them permanent until unregistered
+	void RegisterData(const FName& Id, const FDataHandle& AnimationDataHandle);
+
+	// Unregisters a previously registered anim data handle 
+	void UnregisterData(const FName& Id);
+
+	// Obtains the data hanle for the passed Id, if it exists.
+	// If there is no anim data handle registered, the handle IsValid will be false
+	FDataHandle GetRegisteredData(const FName& Id) const;
+
+	
+	// --- Supported types registration ---
+
+	// Registers a type and sets the desired preallocation block size
+	// If a type is allocated without registering, a default block size of 32 will be used
+	template<typename DataType>
+	inline void RegisterDataType(int32 AllocationBlockSize)
+	{
+		RegisterDataType_Impl<DataType>(AllocationBlockSize);
+	}
+
+	// --- Persistent Data ---
+
+	// Allocates uninitialized memory for a type (leaving the initialization to the caller)
+	// Returns a refcounted animation data handle
+	// Allocated memory will be released once the refcount reaches 0
+	template<typename DataType>
+	FDataHandle PreAllocateMemory(const int32 NumElements)
+	{
+		FParamTypeHandle ParamTypeHandle = FParamTypeHandle::GetHandle<DataType>();
+
+		const FDataTypeDef* TypeDef = nullptr;
+		{
+			FRWScopeLock Lock(DataTypeDefsLock, SLT_ReadOnly);
+			TypeDef = DataTypeDefs.Find(ParamTypeHandle);
+		}
+
+		if (TypeDef == nullptr)
+		{
+			TypeDef = RegisterDataType_Impl<DataType>(DEFAULT_BLOCK_SIZE);
+			// TODO : Log if we allocate more than DEFAULT_BLOCK_SIZE elements of that type
+		}
+
+		if (ensure(TypeDef != nullptr && TypeDef->ParamTypeHandle.IsValid()))
+		{
+			const int32 ElementSize = TypeDef->ElementSize;
+			const int32 ElementAlign = TypeDef->ElementAlign;
+			const int32 AlignedSize = Align(ElementSize, ElementAlign);
+
+			const int32 BufferSize = NumElements * AlignedSize;
+
+			uint8* Memory = (uint8*)FMemory::Malloc(BufferSize, TypeDef->ElementAlign);    // TODO : This should come from preallocated chunks, use malloc / free for now
+
+			Private::FAllocatedBlock* AllocatedBlock = new Private::FAllocatedBlock(Memory, NumElements, ParamTypeHandle); // TODO : avoid memory fragmentation
+			AllocatedBlock->AddRef();
+			AllocatedBlocks.Add(AllocatedBlock);
+			return FDataHandle(AllocatedBlock);
+		}
+
+		return FDataHandle();
+	}
+
+	// Allocates memory for a type, initialized with the default constructor (with optional passed arguments)
+	// Returns a refcounted animation data handle
+	// Allocated memory will be released once the refcount reaches 0
+	template<typename DataType, typename... ArgTypes>
+	FDataHandle AllocateData(const int32 NumElements, ArgTypes&&... Args)
+	{
+		FDataHandle Handle = PreAllocateMemory<DataType>(NumElements);
+
+		DataType* RetVal = Handle.GetPtr<DataType>();
+		for (int i = 0; i < NumElements; i++)
+		{
+			// perform a placement new per element
+			new ((uint8*)&RetVal[i]) DataType(Forward<ArgTypes>(Args)...);
+		}
+
+		return Handle;
+	}
+
+private:
+	typedef void (*DestroyFnSignature)(uint8* TargetBuffer, int32 NumElem);
+
+	static constexpr int32 DEFAULT_BLOCK_SIZE = 32;
+
+	// structure holding each registered type information
+	struct FDataTypeDef
+	{
+		FParamTypeHandle ParamTypeHandle;
+		DestroyFnSignature DestroyTypeFn = nullptr;
+		int32 ElementSize = 0;
+		int32 ElementAlign = 0;
+		int32 AllocationBlockSize = 0;
+	};
+
+	struct FReferencePoseData
+	{
+		FReferencePoseData() = default;
+		
+		FReferencePoseData(const FDataHandle& InAnimationDataHandle, const FDelegateHandle& InDelegateHandle)
+			: AnimationDataHandle(InAnimationDataHandle)
+			, DelegateHandle(InDelegateHandle)
+		{
+		}
+
+		FReferencePoseData(FDataHandle&& InAnimationDataHandle, FDelegateHandle&& InDelegateHandle)
+			: AnimationDataHandle(MoveTemp(InAnimationDataHandle))
+			, DelegateHandle(MoveTemp(InDelegateHandle))
+		{
+		}
+
+		FDataHandle AnimationDataHandle;
+		FDelegateHandle DelegateHandle;
+	};
+
+	// Map holding registered types
+	TMap<FParamTypeHandle, FDataTypeDef> DataTypeDefs;
+	// Lock for registered types map
+	FRWLock DataTypeDefsLock;
+
+	TSet<Private::FAllocatedBlock*> AllocatedBlocks;
+	// Lock for allocated blocks
+	FRWLock AllocatedBlocksLock;
+
+	// Map holding named data
+	TMap <FName, FDataHandle> StoredData;
+	mutable FRWLock StoredDataLock;
+
+	// Map holding reference poses for SkeletalMeshes
+	TMap <TWeakObjectPtr<USkeletalMeshComponent>, FReferencePoseData> SkeletalMeshReferencePoses;
+	mutable FRWLock SkeletalMeshReferencePosesLock;
+
+	// Registers a type and sets the allocation block size
+	template<typename DataType>
+	FDataTypeDef* RegisterDataType_Impl(int32 AllocationBlockSize)
+	{
+		FParamTypeHandle ParamTypeHandle = FParamTypeHandle::GetHandle<DataType>();
+		check(ParamTypeHandle.IsValid());
+
+		const int32 ElementSize = ParamTypeHandle.GetSize();
+		const int32 ElementAlign = ParamTypeHandle.GetAlignment();
+
+		// If we use raw types, I need a per element destructor
+		DestroyFnSignature DestroyFn = [](uint8* TargetBuffer, int32 NumElem)->void
+		{
+			const DataType* Ptr = (DataType*)TargetBuffer;
+			for (int i = 0; i < NumElem; i++)
+			{
+				Ptr[i].~DataType();
+			}
+		};
+
+		FDataTypeDef* AddedDef = nullptr;
+		{
+			FRWScopeLock WriteLock(DataTypeDefsLock, SLT_Write);
+
+			AddedDef = &DataTypeDefs.FindOrAdd(ParamTypeHandle, { ParamTypeHandle, DestroyFn, ElementSize, ElementAlign, AllocationBlockSize });
+			check(AddedDef->ParamTypeHandle == ParamTypeHandle); // check we have not added two different types with the same ID
+		}
+
+		return AddedDef;
+	}
+
+	void OnLODRequiredBonesUpdate(USkeletalMeshComponent* SkeletalMeshComponent, int32 LODLevel, const TArray<FBoneIndexType>& LODRequiredBones);
+	
+	void FreeAllocatedBlock(Private::FAllocatedBlock * AllocatedBlock);
+
+	void ReleaseReferencePoseData();
+
+// --- ---
+private:
+	friend class FModule;
+	friend struct FDataHandle;
+
+	// Initialize the global registry
+	static void Init();
+
+	// Shutdown the global registry
+	static void Destroy();
+
+	volatile int32 HandleCounter = 0;
+
+
+private:
+	static void HandlePostGarbageCollect();
+};
+
+} // namespace UE::AnimNext
diff --git a/Engine/Plugins/Experimental/AnimNext/Source/AnimNext/Public/DecompressionTools.h b/Engine/Plugins/Experimental/AnimNext/Source/AnimNext/Public/DecompressionTools.h
new file mode 100644
--- /dev/null
+++ b/Engine/Plugins/Experimental/AnimNext/Source/AnimNext/Public/DecompressionTools.h
@@ -0,0 +1,76 @@
+// Copyright Epic Games, Inc. All Rights Reserved.
+
+#pragma once
+
+#include "CoreMinimal.h"
+#include "ReferencePose.h"
+#include "LODPose.h"
+
+struct FCompactPose;
+struct FCompressedAnimSequence;
+struct FAnimExtractContext;
+struct FAnimSequenceDecompressionContext;
+struct FRootMotionReset;
+class UAnimSequence;
+
+namespace UE::AnimNext
+{
+
+class ANIMNEXT_API FDecompressionTools
+{
+public:
+
+	// Extracted from UAnimSequence
+	// Extracts Animation Data from the provided AnimSequence, using AnimExtractContext extraction parameters
+	static void GetAnimationPose(const UAnimSequence* AnimSequence
+		, FLODPose& OutAnimationPoseData
+		, const FAnimExtractContext& ExtractionContext);
+
+	/**
+	* Get Bone Transform of the Time given, relative to Parent for all RequiredBones
+	* This returns different transform based on additive or not. Or what kind of additive.
+	*/
+	static void GetBonePose(const UAnimSequence* AnimSequence, const FAnimExtractContext& ExtractionContext, FLODPose& OutAnimationPoseData, bool bForceUseRawData = false);
+
+	static void GetBonePose_Additive(const UAnimSequence* AnimSequence, const FAnimExtractContext& ExtractionContext, FLODPose& OutAnimationPoseData);
+	static void GetBonePose_AdditiveMeshRotationOnly(const UAnimSequence* AnimSequence, const FAnimExtractContext& ExtractionContext, FLODPose& OutAnimationPoseData);
+
+	// Decompress and retarget animation data using provided RetargetTransforms
+	static void DecompressPose(FLODPose& OutAnimationPoseData,
+		const FCompressedAnimSequence& CompressedData,
+		const FAnimExtractContext& ExtractionContext,
+		FAnimSequenceDecompressionContext& DecompressionContext,
+		const TArray<FTransform>& RetargetTransforms,
+		const FRootMotionReset& RootMotionReset);
+
+	// Decompress and retarget animation data
+	static void DecompressPose(FLODPose& OutAnimationPoseData,
+		const FCompressedAnimSequence& CompressedData,
+		const FAnimExtractContext& ExtractionContext,
+		FAnimSequenceDecompressionContext& DecompressionContext,
+		FName RetargetSource,
+		const FRootMotionReset& RootMotionReset);
+
+	static void RetargetBoneTransform(const FReferencePose& ReferencePose
+		, const USkeleton* SourceSkeleton
+		, const USkeleton* TargetSkeleton
+		, const FName& RetargetSource
+		, FTransform& BoneTransform
+		, const int32 SkeletonBoneIndex
+		, const int32 LODBoneIndex
+		, const bool bIsBakedAdditive
+		, const bool bDisableRetargeting);
+
+	static void RetargetBoneTransform(const FReferencePose& ReferencePose
+		, const USkeleton* SourceSkeleton
+		, const USkeleton* TargetSkeleton
+		, const FName& SourceName
+		, const TArray<FTransform>& RetargetTransforms
+		, FTransform& BoneTransform
+		, const int32 SkeletonBoneIndex
+		, const int32 LODBoneIndex
+		, const bool bIsBakedAdditive
+		, const bool bDisableRetargeting);
+};
+
+} // end namespace UE::AnimNext
diff --git a/Engine/Plugins/Experimental/AnimNext/Source/AnimNext/Public/GenerationTools.h b/Engine/Plugins/Experimental/AnimNext/Source/AnimNext/Public/GenerationTools.h
new file mode 100644
--- /dev/null
+++ b/Engine/Plugins/Experimental/AnimNext/Source/AnimNext/Public/GenerationTools.h
@@ -0,0 +1,102 @@
+// Copyright Epic Games, Inc. All Rights Reserved.
+
+#pragma once
+
+#include "BoneIndices.h"
+#include "Containers/IndirectArray.h"
+#include "ReferencePose.h"
+#include "LODPose.h"
+
+class USkeletalMeshComponent;
+class USkeletalMesh;
+class FSkeletalMeshLODRenderData;
+struct FPoseContext;
+class UPhysicsAsset;
+
+namespace UE::AnimNext
+{
+
+struct FGenerationLODData
+{
+	TArray<FBoneIndexType> RequiredBones;									// All the bones required for the LOD
+	TArray<FBoneIndexType> ExcludedBones;									// List of bones excluded from LOD 0
+	TArray<FBoneIndexType> ExcludedBonesFromPrevLOD;						// list of bones excluded from previous LOD
+};
+
+class ANIMNEXT_API FGenerationTools
+{
+public:
+
+	// Generates the reference pose data from a SkeletalMeshComponent and / or Skeletal Mesh asset
+	// If no SkeletalMeshComponent is passed, the reference pose will not exclude invisible bones and will not include shadow shapes required bones
+	// If no SkeletalMesh asset is passed, there will be no generation
+	static bool GenerateReferencePose(const USkeletalMeshComponent* SkeletalMeshComponent
+		, const USkeletalMesh* SkeletalMesh
+		, UE::AnimNext::FReferencePose& OutAnimationReferencePose);
+
+	// Generates the full list of bones required by a LOD, based on the SkeletalRetrieve RequiredBones
+	// Compute ExcludedBones versus LOD0 and PreviousLOD
+	// Remove ExcludedBonesFromPrevLOD from BonesInAllLODS
+	static void GenerateRawLODData(const USkeletalMeshComponent* SkeletalMeshComponent
+		, const USkeletalMesh* SkeletalMesh
+		, const int32 LODIndex
+		, const TIndirectArray<FSkeletalMeshLODRenderData>& LODRenderData
+		, TArray<FBoneIndexType>& OutRequiredBones
+		, TArray<FBoneIndexType>& OutFillComponentSpaceTransformsRequiredBones);
+
+	// For each LOD > 0:
+	// Retrieve RequiredBones
+	// Compute ExcludedBones versus LOD0 and PreviousLOD
+	static void GenerateLODData(const USkeletalMeshComponent* SkeletalMeshComponent
+		, const USkeletalMesh* SkeletalMesh
+		, const int32 StartLOD
+		, const int32 NumLODs
+		, const TIndirectArray<FSkeletalMeshLODRenderData>& LODRenderData
+		, const TArray<FBoneIndexType>& RequiredBones_LOD0
+		, TArray<FGenerationLODData>& GenerationLODData
+		, TArray<FGenerationLODData>& GenerationComponentSpaceLODData);
+
+	// Calculate the bone indexes difference from LOD0 for LODIndex
+	static void CalculateDifferenceFromParentLOD(int32 LODIndex, TArray<FGenerationLODData>& GenerationLODData);
+
+	// For each LOD :
+	// Check the excluded bones in LOD(N) contain all the bones excluded in LOD(N-1)
+	static bool CheckExcludedBones(const int32 NumLODs
+		, const TArray<FGenerationLODData>& GenerationLODData
+		, const USkeletalMesh* SkeletalMesh);
+
+	// For each LOD :
+	// Generates an unified list of bones, in LOD order (if possible)
+	// Returns true if the unified list could be created, false otherwise
+	static bool GenerateOrderedBoneList(const USkeletalMesh* SkeletalMesh
+		, TArray<FGenerationLODData>& GenerationLODData
+		, TArray<FBoneIndexType>& OrderedBoneList);
+
+	/**
+	 *	Utility for taking two arrays of bone indices, which must be strictly increasing, and finding the A - B.
+	 *	That is - any items left in A, after removing B.
+	 */
+	static void DifferenceBoneIndexArrays(const TArray<FBoneIndexType>& A, const TArray<FBoneIndexType>& B, TArray<FBoneIndexType>& Output);
+
+	// Checks if all sockets of a skeletal mesh are set to always animate, as it is a requirement for generating a single reference pose,
+	// where the local space pose and the component space pose use the same bone indexes
+	static bool CheckSkeletalAllMeshSocketsAlwaysAnimate(const USkeletalMesh* SkeletalMesh);
+
+	// Converts AnimBP pose to AnimNext Pose
+	// This function expects both poses to have the same LOD (number of bones and indexes)
+	// The target pose should be assigned to the correct reference pose prior to this call
+	static void RemapPose(int32 LODLevel
+		, const FPoseContext& SourcePose
+		, const FReferencePose& RefPose
+		, FLODPose& TargetPose); // TODO : see how to deal with heap and stack variants without duplicating functions (and if possible without templating the function)
+
+	// Converts AnimNext pose to AnimBP Pose
+	// This function expects both poses to have the same LOD (number of bones and indexes)
+	// The target pose should be assigned to the correct reference pose prior to this call
+	static void RemapPose(int32 LODLevel
+		, const FReferencePose& RefPose
+		, const FLODPose& SourcePose
+		, FPoseContext& TargetPose); // TODO : see how to deal with heap and stack variants without duplicating functions (and if possible without templating the function)
+};
+
+} // namespace UE::AnimNext
diff --git a/Engine/Plugins/Experimental/AnimNext/Source/AnimNext/Public/Interface/InterfaceContext.h b/Engine/Plugins/Experimental/AnimNext/Source/AnimNext/Public/Interface/InterfaceContext.h
new file mode 100644
--- /dev/null
+++ b/Engine/Plugins/Experimental/AnimNext/Source/AnimNext/Public/Interface/InterfaceContext.h
@@ -0,0 +1,487 @@
+// Copyright Epic Games, Inc. All Rights Reserved.
+
+#pragma once
+
+#include "CoreMinimal.h"
+#include "UObject/Class.h"
+#include "Param/Param.h"
+#include "Interface/InterfaceState.h"
+#include "Param/ParamStorage.h"
+#include "Interface/InterfaceKey.h"
+#include "Param/ParamType.h"
+#include "Misc/MemStack.h"
+
+class IAnimNextInterface;
+class IAnimNextParamInterface;
+
+namespace UE::AnimNext
+{
+
+struct FContext;
+struct FState;
+enum class EStatePersistence : uint8;
+
+// Helper to provide a non-callstack bridge for interfacing with anim interface contexts
+// Declaring one of these on the stack will enable all calls in its scope to access the passed-in context
+// via FThreadContext::Get()
+struct ANIMNEXT_API FThreadContext
+{
+	FThreadContext(const FContext& InContext);
+	~FThreadContext();
+
+	static const FContext& Get();
+};
+
+// Context providing methods for mutating & interrogating the anim interface runtime
+struct ANIMNEXT_API FContext
+{
+private:
+	friend class ::IAnimNextInterface;
+	friend struct FState;
+
+public:
+	// Root public constructor. Constructs a context given a state.
+	FContext(float InDeltaTime, FState& InState, FParamStorage& InParamStorage, IAnimNextParamInterface* InParameters = nullptr);
+	
+	FContext(const FContext& other) = delete;
+	FContext& operator=(const FContext&) = delete;
+
+	FContext(FContext&& Other)
+		: FContext()
+	{
+		Swap(*this, Other);
+	}
+
+	FContext& operator= (FContext&& Other)
+	{
+		Swap(*this, Other);
+		return *this;
+	}
+
+	~FContext();
+
+public:
+	// --- Sub Context Creation --- 
+
+	// Create a sub context from this one that includes the provided Result
+	FContext WithResult(FParam& InResult) const;
+
+	// Create a sub context from this one that includes the provided parameter
+	FContext WithParameter(FName ParameterId, const FParam& InParameter) const;
+
+	// Create a sub context from this one that includes the provided parameters
+	FContext WithParameters(TArrayView<const TPair<FName, FParam>> InParameters) const;
+
+	// Create a sub context from this one that includes the provided Result and parameters
+	FContext WithResultAndParameters(FParam& InResult, TArrayView<const TPair<FName, FParam>> InParameters) const;
+
+	// Create a sub context from this one that includes the provided interface parameter
+	FContext WithParameters(IAnimNextParamInterface* InParameters) const;
+
+public:
+
+	// --- Interface for direct Param storage (prototype) ---
+
+	FContext CreateSubContext() const;
+
+	enum class EParamType : uint8
+	{
+		None		= 0,
+		Input		= 1 << 0,
+		Output		= 1 << 1,
+	};
+
+	// Add an Input parameter by value, copying it to the shared storage
+	template<typename ValueType>
+	FParamHandle AddInputValue(FName ParameterId, ValueType& Value)
+	{
+		// Values are added mutable even if source is const
+		const FParam::EFlags Flags = FParam::EFlags::Value | FParam::EFlags::Mutable;
+
+		return AddParameter(EParamType::Input, ParameterId, Flags, &Value);
+	}
+
+	// Add an Input parameter by value, copying it to the shared storage
+	template<typename ValueType>
+	FParamHandle AddInputValue(FName ParameterId, ValueType&& Value)
+	{
+		// Values are added mutable even if source is const
+		const FParam::EFlags Flags = FParam::EFlags::Value | FParam::EFlags::Mutable;
+
+		return AddParameter(EParamType::Input, ParameterId, Flags, &Value);
+	}
+
+	// Add an Input parameter by value, copying it to the shared storage
+	template<typename ValueType>
+	FParamHandle AddInputValue(FName ParameterId, ValueType* Value)
+	{
+		// Values are added mutable even if source is const
+		const FParam::EFlags Flags = FParam::EFlags::Value | FParam::EFlags::Mutable;
+
+		return AddParameter(EParamType::Input, ParameterId, Flags, Value);
+	}
+
+	// Add an Output parameter by value, copying it to the shared storage
+	template<typename ValueType>
+	FParamHandle AddOutputValue(FName ParameterId, ValueType& Value)
+	{
+		// Values are added mutable even if source is const
+		const FParam::EFlags Flags = FParam::EFlags::Value | FParam::EFlags::Mutable;
+
+		return AddParameter(EParamType::Output, ParameterId, Flags, &Value);
+	}
+
+	// Add an Output parameter by value, copying it to the shared storage
+	template<typename ValueType>
+	FParamHandle AddOutputValue(FName ParameterId, ValueType&& Value)
+	{
+		// Values are added mutable even if source is const
+		const FParam::EFlags Flags = FParam::EFlags::Value | FParam::EFlags::Mutable;
+
+		return AddParameter(EParamType::Output, ParameterId, Flags, &Value);
+	}
+
+	// Add an Output parameter by value, copying it to the shared storage
+	template<typename ValueType>
+	FParamHandle AddOutputValue(FName ParameterId, ValueType* Value)
+	{
+		// Values are added mutable even if source is const
+		const FParam::EFlags Flags = FParam::EFlags::Value | FParam::EFlags::Mutable;
+
+		return AddParameter(EParamType::Output, ParameterId, Flags, Value);
+	}
+
+	// Add a parameter by value, copying it to the shared storage
+	template<typename ValueType>
+	FParamHandle AddValue(EParamType InParamType, FName ParameterId, ValueType &Value)
+	{
+		// Values are added mutable even if source is const
+		const FParam::EFlags Flags = FParam::EFlags::Value | FParam::EFlags::Mutable;
+
+		return AddParameter(InParamType, ParameterId, Flags, &Value);
+	}
+
+	// Add a parameter by value, copying it to the shared storage
+	template<typename ValueType>
+	FParamHandle AddValue(EParamType InParamType, FName ParameterId, ValueType&& Value)
+	{
+		// Values are added mutable even if source is const
+		const FParam::EFlags Flags = FParam::EFlags::Value | FParam::EFlags::Mutable;
+
+		return AddParameter(InParamType, ParameterId, Flags, &Value);
+	}
+
+	// Add a parameter by value, copying it to the shared storage
+	template<typename ValueType>
+	FParamHandle AddValue(EParamType InParamType, FName ParameterId, ValueType* Value)
+	{
+		// Values are added mutable even if source is const
+		const FParam::EFlags Flags = FParam::EFlags::Value | FParam::EFlags::Mutable;
+
+		return AddParameter(InParamType, ParameterId, Flags, Value);
+	}
+
+	// --- Parameters by refecence / pointer ---
+
+	// Add an Input parameter by reference, adding just a pointer to the shared storage
+	template<typename ValueType>
+	FParamHandle AddInputReference(FName ParameterId, ValueType& Value)
+	{
+		const FParam::EFlags Flags = FParam::EFlags::Reference;
+
+		return AddParameter(EParamType::Input, ParameterId, Flags, &Value);
+	}
+
+	// Add an Input parameter by reference, adding just a pointer to the shared storage
+	template<typename ValueType>
+	FParamHandle AddInputReference(FName ParameterId, ValueType* Value)
+	{
+		FParam::EFlags Flags = FParam::EFlags::Reference;
+
+		return AddParameter(EParamType::Input, ParameterId, Flags, Value);
+	}
+
+	// Add an Output parameter by reference, adding just a pointer to the shared storage
+	template<typename ValueType>
+	FParamHandle AddOutputReference(FName ParameterId, ValueType& Value)
+	{
+		const FParam::EFlags Flags = FParam::EFlags::Reference;
+
+		return AddParameter(EParamType::Output, ParameterId, Flags, &Value);
+	}
+
+	// Add an Output parameter by reference, adding just a pointer to the shared storage
+	template<typename ValueType>
+	FParamHandle AddOutputReference(FName ParameterId, ValueType* Value)
+	{
+		FParam::EFlags Flags = FParam::EFlags::Reference;
+
+		return AddParameter(EParamType::Output, ParameterId, Flags, Value);
+	}
+
+	// Add a parameter by reference, adding just a pointer to the shared storage
+	template<typename ValueType>
+	FParamHandle AddReference(EParamType InParamType, FName ParameterId, ValueType& Value)
+	{
+		const FParam::EFlags Flags = FParam::EFlags::Reference;
+
+		return AddParameter(InParamType, ParameterId, Flags, &Value);
+	}
+
+	// Add a parameter by reference, adding just a pointer to the shared storage
+	template<typename ValueType>
+	FParamHandle AddReference(EParamType InParamType, FName ParameterId, ValueType *Value)
+	{
+		FParam::EFlags Flags = FParam::EFlags::Reference;
+
+		return AddParameter(InParamType, ParameterId, Flags, Value);
+	}
+
+	template<typename ValueType>
+	FParamHandle AddParameter(EParamType InParamType, FName ParameterId, FParam::EFlags Flags, ValueType* Value)
+	{
+		FParamHandle ParamHandle;
+
+		check(InParamType == EParamType::Input || InParamType == EParamType::Output);
+		check(EnumHasAnyFlags(Flags, FParam::EFlags::Value) || EnumHasAnyFlags(Flags, FParam::EFlags::Reference));
+
+		// check if it is a simple FParamHandle copy
+		if constexpr (TIsDerivedFrom<ValueType, FParamHandle>::Value)
+		{
+			const FParam* ExistingParam = ParamStorage->GetParam(Value->ParamHandle);
+			check(ExistingParam != nullptr);
+
+			const bool bIsReadOnly = EnumHasAnyFlags(ExistingParam->Flags, FParam::EFlags::Mutable) == false;
+
+			// Verify that we want to create an Input param or if we want an output, the source is not read only
+			check(InParamType == EParamType::Input || bIsReadOnly == false);
+
+			// For now create a copy of the Param (TODO : change ref to value or value to ref, if the caller requests it?)
+			ParamHandle = *Value;
+		}
+		else
+		{
+			using MutableValueType = std::remove_const_t<ValueType>;
+
+			// If non const set the mutable flag
+			if constexpr (TIsConst<ValueType>::Value == false)
+			{
+				EnumAddFlags(Flags, FParam::EFlags::Mutable);
+			}
+
+			ParamHandle = EnumHasAnyFlags(Flags, FParam::EFlags::Value)
+				? ParamStorage->AddValue(const_cast<MutableValueType*>(Value), Flags)
+				: ParamStorage->AddReference(const_cast<MutableValueType*>(Value), Flags);
+		}
+
+		AdditionalParameterHandles.Add(ParameterId, ParamHandle);
+
+		return ParamHandle;
+	}
+
+	// Get a TParam from a handle
+	template<typename ValueType>
+	TParam<ValueType> GetParameterChecked(const FParamHandle &InParamHandle) const
+	{
+		const FParam* Param = ParamStorage->GetParam(InParamHandle.ParamHandle);
+
+		TParam<ValueType> RetVal(Param, Param->GetFlags());
+
+		check(RetVal.IsValid());							// found something
+		checkSlow(Private::CheckParam<ValueType>(RetVal));	// the type found is valid for what was requested
+
+		return RetVal;
+	}
+
+	// Get a parameter from a handle as a specified type
+	template <typename ValueType>
+	ValueType& GetParameterAs(const FParamHandle& InParamHandle)
+	{
+		const FParam* Param = ParamStorage->GetParam(InParamHandle.ParamHandle);
+
+		check(Param != nullptr);
+		check(Param->GetTypeHandle() == FParamTypeHandle::GetHandle<ValueType>());
+
+		// If a non cost is requested, check the param has mutable flag
+		if constexpr (TIsConst<ValueType>::Value == false)
+		{
+			check(Param->IsMutable());
+		}
+
+		ValueType* ValueData = EnumHasAnyFlags(Param->Flags, FParam::EFlags::Embedded)
+			? static_cast<ValueType*>((void*)&Param->Data)
+			: static_cast<ValueType*>(Param->Data);
+
+		return *ValueData;
+	}
+
+	// --- Parameter management --- 
+
+	// Get a parameter if exist, returns false if not
+	bool GetParameter(FName InKey, FParam& OutParam) const;
+
+	// Get a parameter as a specified type, checking it exist and the type
+	template<typename ValueType>
+	TParam<ValueType> GetParameterChecked(FName InKey) const
+	{
+		TParam<ValueType> RetVal;
+
+		GetParameter(InKey, RetVal);
+
+		check(RetVal.IsValid());							// found something
+		checkSlow(Private::CheckParam<ValueType>(RetVal));	// the type found is valid for what was requested
+
+		return RetVal;
+	}
+
+	// Get a parameter as a specified type, returning a TOptional
+	template<typename ValueType>
+	TOptional<TParam<ValueType>> GetParameter(FName InKey) const
+	{
+		TOptional<TParam<ValueType>> RetVal;
+
+		TParam<ValueType> Param(FParam::EFlags::Mutable);
+		if (GetParameter(InKey, Param))
+		{
+			checkSlow(Private::CheckParam<ValueType>(Param));
+
+			RetVal = Param;
+		}
+
+		return RetVal;
+	}
+
+public:
+	// --- Result Management ---
+
+	// Get the ResultParam, checking it has been set
+	FParam& GetResultParam()
+	{
+		check(Result != nullptr);
+		return *Result;
+	}
+
+	// Get the ResultParam, checking it has been set
+	FParam& GetResultParam() const
+	{
+		check(Result != nullptr);
+		return *Result;
+	}
+
+	// Set a HParam as a result (for compatibility reasons)
+	void SetHParamAsResult(const FParamHandle& InHParam)
+	{
+		Result = ParamStorage->GetParam(InHParam.ParamHandle);
+	}
+
+	// Set a result value directly
+	// The receiver Result Param must be set prior to this call
+	template<typename ValueType>
+	void SetResult(const ValueType& InValue) const
+	{
+		check(Result != nullptr);
+		TParam<ValueType> TypedResult(*Result);
+		*TypedResult = InValue;
+	}
+
+	// Get the current result as a TParam
+	template<typename ValueType>
+	TParam<ValueType> GetResultParam() const
+	{
+		check(Result != nullptr);
+		return TParam<ValueType>(*Result);
+	}
+
+	// Get the current result as a mutable reference
+	template<typename ValueType>
+	ValueType& GetResult() const
+	{
+		check(Result != nullptr);
+		TParam<ValueType> TypedResult(*Result);
+		return *TypedResult;
+	}
+
+	// Get the current result as a mutable ptr
+	template<typename ValueType>
+	ValueType* GetResultPtr() const
+	{
+		check(Result != nullptr);
+		TParam<ValueType> TypedResult(*Result);
+		return &*TypedResult;
+	}
+
+public:
+	// --- State management ---
+
+	// Gets (and allocates, if necessary) state for the specified node given this calling context
+	template<typename ValueType, EStatePersistence Persistence = EStatePersistence::Relevancy>
+	TParam<ValueType> GetStateParam(const FInterfaceKeyWithId& InKey) const
+	{
+		return State->GetState<ValueType, Persistence>(InKey, *this, CallstackHash);
+	}
+
+	// Gets (and allocates, if necessary) state for the specified node given this calling context
+	template<typename ValueType, EStatePersistence Persistence = EStatePersistence::Relevancy>
+	TParam<ValueType> GetStateParam(const IAnimNextInterface* InAnimNextInterface, uint32 InId) const
+	{
+		return State->GetState<ValueType, Persistence>(InAnimNextInterface, InId, *this, CallstackHash);
+	}
+
+	// Gets (and allocates, if necessary) state for the specified node given this calling context
+	template<typename ValueType, EStatePersistence Persistence = EStatePersistence::Relevancy>
+	ValueType& GetState(const FInterfaceKeyWithId& InKey) const
+	{
+		TParam<ValueType> Param = State->GetState<ValueType, Persistence>(InKey, *this, CallstackHash); 
+		return *Param;
+	}
+
+	// Gets (and allocates, if necessary) state for the specified node given this calling context
+	template<typename ValueType, EStatePersistence Persistence = EStatePersistence::Relevancy>
+	ValueType& GetState(const IAnimNextInterface* InAnimNextInterface, uint32 InId) const
+	{
+		TParam<ValueType> Param = State->GetState<ValueType, Persistence>(InAnimNextInterface, InId, *this, CallstackHash);
+		return *Param;
+	}
+	
+public:
+	// --- Mix Context Utils ---
+
+	// Access delta time as a param
+	TParam<const float> GetDeltaTimeParam() const;
+
+	// Raw access to delta time
+	float GetDeltaTime() const { return DeltaTime; }
+
+private:
+	FContext(float InDeltaTime, FState& InState, FParamStorage& InParamStorage, FParam& InResult)
+		: State(&InState)
+		, ParamStorage(&InParamStorage)
+		, Result(&InResult)
+		, DeltaTime(InDeltaTime)
+	{}
+
+	FContext() = default;
+
+	FContext WithCallRaw(const IAnimNextInterface* InAnimNextInterface) const;
+
+	void FlushRelevancy() const;
+
+	int32 GetParametersSize(TArrayView<const TPair<FName, FParam>> InParameters, TArray<int32>& ParamAllocSizes) const;
+	void AddParameters(TArrayView<const TPair<FName, FParam>> InParameters);
+
+private:
+	TMap<FName, FParam> AdditionalParameters;
+	TMap<FName, FParamHandle> AdditionalParameterHandles;  // Temp param storage prototype
+	const FContext* Parent = nullptr;
+	const FContext* Root = nullptr;
+	FState* State = nullptr;
+	FParamStorage* ParamStorage = nullptr;
+	FParamStorageHandle BlockHandle = InvalidBlockHandle;
+	FParam* Result = nullptr;
+	IAnimNextParamInterface* Parameters = nullptr;
+	float DeltaTime = 0.0f;
+	uint32 CallstackHash = 0;
+	uint32 UpdateCounter = 0;
+};
+
+}
diff --git a/Engine/Plugins/Experimental/AnimNext/Source/AnimNext/Public/Interface/InterfaceKey.h b/Engine/Plugins/Experimental/AnimNext/Source/AnimNext/Public/Interface/InterfaceKey.h
new file mode 100644
--- /dev/null
+++ b/Engine/Plugins/Experimental/AnimNext/Source/AnimNext/Public/Interface/InterfaceKey.h
@@ -0,0 +1,67 @@
+﻿// Copyright Epic Games, Inc. All Rights Reserved.
+
+#pragma once
+
+#include "CoreMinimal.h"
+#include "UObject/Class.h"
+
+class IAnimNextInterface;
+
+namespace UE::AnimNext
+{
+
+// Key value for cached anim interface state
+struct ANIMNEXT_API FInterfaceKey
+{
+	FInterfaceKey(const IAnimNextInterface* InAnimNextInterface)
+		: AnimNextInterface(InAnimNextInterface)
+		, Hash(::GetTypeHash(InAnimNextInterface))
+	{}
+
+	FInterfaceKey(const FInterfaceKey& InExistingKey, uint32 InHash)
+		: AnimNextInterface(InExistingKey.AnimNextInterface)
+		, Hash(HashCombineFast(InExistingKey.Hash, InHash))
+	{}
+
+	const IAnimNextInterface* AnimNextInterface = nullptr;
+	uint32 Hash = 0;
+};
+
+// Key value for cached anim interface state
+struct ANIMNEXT_API FInterfaceKeyWithId : FInterfaceKey
+{
+	FInterfaceKeyWithId(const FInterfaceKey& InInterfaceKey, uint32 InId)
+		: FInterfaceKey(InInterfaceKey, ::GetTypeHash(InId))
+		, Id(InId)
+	{}
+
+	FInterfaceKeyWithId(const FInterfaceKeyWithId& InInterfaceKeyWithId, uint32 InHash)
+		: FInterfaceKey(InInterfaceKeyWithId, HashCombineFast(InInterfaceKeyWithId.Hash, InHash))
+		, Id(InInterfaceKeyWithId.Id)
+	{}
+	
+	uint32 Id = 0;
+};
+
+// Key value for cached anim interface state
+struct ANIMNEXT_API FInterfaceKeyWithIdAndStack : FInterfaceKeyWithId
+{
+	FInterfaceKeyWithIdAndStack(const FInterfaceKeyWithId& InInterfaceKey, uint32 InCallStackHash)
+		: FInterfaceKeyWithId(InInterfaceKey, InCallStackHash)
+		, CallStackHash(InCallStackHash)
+	{}
+
+	FORCEINLINE friend bool operator==(const FInterfaceKeyWithIdAndStack& A, const FInterfaceKeyWithIdAndStack& B)
+	{
+		return A.AnimNextInterface == B.AnimNextInterface && A.Id == B.Id && A.CallStackHash == B.CallStackHash;
+	}
+
+	FORCEINLINE friend uint32 GetTypeHash(const FInterfaceKeyWithIdAndStack& InKey)
+	{
+		return InKey.Hash;
+	}
+
+	uint32 CallStackHash = 0;
+};
+
+}
\ No newline at end of file
diff --git a/Engine/Plugins/Experimental/AnimNext/Source/AnimNext/Public/Interface/InterfaceState.h b/Engine/Plugins/Experimental/AnimNext/Source/AnimNext/Public/Interface/InterfaceState.h
new file mode 100644
--- /dev/null
+++ b/Engine/Plugins/Experimental/AnimNext/Source/AnimNext/Public/Interface/InterfaceState.h
@@ -0,0 +1,79 @@
+﻿// Copyright Epic Games, Inc. All Rights Reserved.
+
+#pragma once
+
+#include "CoreMinimal.h"
+#include "Param/Param.h"
+#include "Interface/InterfaceKey.h"
+
+namespace UE::AnimNext
+{
+
+struct FContext;
+
+enum class EStatePersistence : uint8
+{
+	None,
+	Relevancy,
+	Permanent,
+};
+
+// Container for anim interface state
+struct ANIMNEXT_API FState
+{
+private:
+	friend struct FContext;
+
+	// Gets (and allocates, if necessary) state for the specified node given this calling context
+	template<typename ValueType, EStatePersistence Persistence = EStatePersistence::Relevancy>
+	TParam<ValueType> GetState(const FInterfaceKeyWithId& InKey, const FContext& InContext, uint32 InHash)
+	{
+		const FInterfaceKeyWithIdAndStack CombinedKey(InKey, InHash);
+
+		if(FParam* ExistingState = FindStateRaw(CombinedKey, InContext, Persistence))
+		{
+			checkSlow(Private::CheckParam<ValueType>(*ExistingState));
+			return TParam<ValueType>(*ExistingState);
+		}
+		else
+		{
+			// Allocate new state
+			FParam* NewState = AllocateState(CombinedKey, InContext, FParamTypeHandle::GetHandle<ValueType>(), Persistence);
+
+			// Construct
+			new (NewState->Data) ValueType();
+
+			// Return as param
+			return TParam<ValueType>(*NewState);
+		}
+	}
+
+	// Gets (and allocates, if necessary) state for the specified node given this calling context
+	template<typename ValueType, EStatePersistence Persistence = EStatePersistence::Relevancy>
+	TParam<ValueType> GetState(const IAnimNextInterface* InAnimNextInterface, uint32 InId, const FContext& InContext, uint32 InContextHash)
+	{
+		const FInterfaceKeyWithId Key(InAnimNextInterface, InId);
+		return GetState<ValueType>(Key, InContext, InContextHash);
+	}
+	
+private:
+	FParam* FindStateRaw(const FInterfaceKeyWithIdAndStack& InKey, const FContext& InContext, EStatePersistence InPersistence);
+
+	FParam* AllocateState(const FInterfaceKeyWithIdAndStack& InKey, const FContext& InContext, const FParamTypeHandle& InTypeHandle, EStatePersistence InPersistence);
+
+private:
+	struct FRelevancyParam : FParam
+	{
+		FRelevancyParam(const FParamTypeHandle& InTypeHandle, TArrayView<uint8> InData, EFlags InFlags, uint32 InUpdateCounter)
+			: FParam(InTypeHandle, InData, InFlags)
+			, UpdateCounter(InUpdateCounter)
+		{}
+	
+		uint32 UpdateCounter = 0;
+	};
+	
+	TMap<FInterfaceKeyWithIdAndStack, FRelevancyParam> RelevancyValueMap;
+	TMap<FInterfaceKeyWithIdAndStack, FParam> PermanentValueMap;
+};
+
+}
\ No newline at end of file
diff --git a/Engine/Plugins/Experimental/AnimNext/Source/AnimNext/Public/Interface/InterfaceTypes.h b/Engine/Plugins/Experimental/AnimNext/Source/AnimNext/Public/Interface/InterfaceTypes.h
new file mode 100644
--- /dev/null
+++ b/Engine/Plugins/Experimental/AnimNext/Source/AnimNext/Public/Interface/InterfaceTypes.h
@@ -0,0 +1,12 @@
+﻿// Copyright Epic Games, Inc. All Rights Reserved.
+
+#pragma once
+
+#include "CoreMinimal.h"
+#include "Param/ParamTypeHandle.h"
+
+#define ANIM_NEXT_INTERFACE_RETURN_TYPE(Typename) \
+virtual UE::AnimNext::FParamTypeHandle GetReturnTypeHandleImpl() const final override\
+{\
+	return UE::AnimNext::FParamTypeHandle::GetHandle<Typename>();\
+}\
diff --git a/Engine/Plugins/Experimental/AnimNext/Source/AnimNext/Public/ReferencePose.h b/Engine/Plugins/Experimental/AnimNext/Source/AnimNext/Public/ReferencePose.h
new file mode 100644
--- /dev/null
+++ b/Engine/Plugins/Experimental/AnimNext/Source/AnimNext/Public/ReferencePose.h
@@ -0,0 +1,190 @@
+// Copyright Epic Games, Inc. All Rights Reserved.
+
+#pragma once
+
+#include "CoreMinimal.h"
+#include "Containers/ContainersFwd.h"
+#include "UObject/WeakObjectPtr.h"
+#include "DataRegistry.h"
+#include <algorithm>
+#include "Animation/AnimTypes.h"
+#include "ReferenceSkeleton.h"
+#include "BoneIndices.h"
+#include "TransformArray.h"
+#include "ReferencePose.generated.h"
+
+class USkeletalMesh;
+class USkeleton;
+
+namespace UE::AnimNext
+{
+
+enum class EReferencePoseGenerationFlags : uint8
+{
+	None = 0,
+	FastPath = 1 << 0
+};
+
+ENUM_CLASS_FLAGS(EReferencePoseGenerationFlags);
+
+template <typename AllocatorType>
+struct TReferencePose
+{
+	TTransformArray<AllocatorType> ReferenceLocalTransforms;
+	TArray<TArray<FBoneIndexType, AllocatorType>, AllocatorType> LODBoneIndexes;
+	TArray<TArray<FBoneIndexType, AllocatorType>, AllocatorType> SkeletonToLODBoneIndexes;
+	TArray<int32, AllocatorType> LODNumBones;
+
+	TWeakObjectPtr<const USkeletalMesh> SkeletalMesh = nullptr;
+	TWeakObjectPtr<const USkeleton> Skeleton = nullptr;
+	EReferencePoseGenerationFlags GenerationFlags = EReferencePoseGenerationFlags::None;
+
+	TReferencePose() = default;
+
+	bool IsValid() const
+	{
+		return ReferenceLocalTransforms.Num() > 0;
+	}
+
+	int32 GetNumBonesForLOD(int32 LODLevel) const
+	{
+		const int32 NumLODS = LODNumBones.Num();
+
+		return (LODLevel < NumLODS) ? LODNumBones[LODLevel] : (NumLODS > 0) ? LODNumBones[0] : 0;
+	}
+
+	bool IsFastPath() const
+	{
+		return (GenerationFlags & EReferencePoseGenerationFlags::FastPath) != EReferencePoseGenerationFlags::None;
+	}
+
+	void Initialize(const FReferenceSkeleton& RefSkeleton
+		, const TArray<TArray<FBoneIndexType>>& InLODBoneIndexes
+		, const TArray<TArray<FBoneIndexType>>& InSkeletonToLODBoneIndexes
+		, const TArray<int32, AllocatorType>& InLODNumBones
+		, bool bFastPath = false)
+	{
+		const int32 NumBonesLOD0 = (InLODNumBones.Num()) > 0 ? InLODNumBones[0] : 0;
+
+		ReferenceLocalTransforms.SetNum(NumBonesLOD0);
+		LODBoneIndexes = InLODBoneIndexes;
+		SkeletonToLODBoneIndexes = InSkeletonToLODBoneIndexes;
+		LODNumBones = InLODNumBones;
+		
+		const TArray<FTransform>& RefBonePose = RefSkeleton.GetRefBonePose();
+		const TArray<FBoneIndexType>& InBoneIndexes = InLODBoneIndexes[0]; // Fill the transforms with the LOD0 indexes
+
+		for (int32 i = 0; i < NumBonesLOD0; ++i)
+		{
+			ReferenceLocalTransforms[i] = RefBonePose[InBoneIndexes[i]]; // TODO : For SoA this is un-optimal, as we are using a TransformAdapter. Evaluate using a specific SoA iterator
+		}
+
+		GenerationFlags = bFastPath ? EReferencePoseGenerationFlags::FastPath : EReferencePoseGenerationFlags::None;
+	}
+
+	const TArrayView<const FBoneIndexType> GetLODBoneIndexes(int32 LODLevel) const
+	{
+		TArrayView<const FBoneIndexType> ArrayView;
+
+		if (LODLevel >= 0 && (IsFastPath() || LODLevel < LODBoneIndexes.Num()))
+		{
+			const int32 NumBonesForLOD = GetNumBonesForLOD(LODLevel);
+
+			if (IsFastPath())
+			{
+				ArrayView = MakeArrayView(LODBoneIndexes[0].GetData(), NumBonesForLOD);
+			}
+			else
+			{
+				ArrayView = MakeArrayView(LODBoneIndexes[LODLevel].GetData(), NumBonesForLOD);
+			}
+		}
+
+		return ArrayView;
+	}
+
+	const TArrayView<const FBoneIndexType> GetSkeletonToLODBoneIndexes(int32 LODLevel) const
+	{
+		TArrayView<const FBoneIndexType> ArrayView;
+
+		if (LODLevel >= 0 && (IsFastPath() || LODLevel < SkeletonToLODBoneIndexes.Num()))
+		{
+			const int32 NumBonesForLOD = GetNumBonesForLOD(LODLevel);
+
+			if (IsFastPath())
+			{
+				return MakeArrayView(SkeletonToLODBoneIndexes[0].GetData(), NumBonesForLOD);
+			}
+			else
+			{
+				return MakeArrayView(SkeletonToLODBoneIndexes[LODLevel].GetData(), NumBonesForLOD);
+			}
+		}
+
+		return ArrayView;
+	}
+
+	int32 GetSkeletonBoneIndexFromLODBoneIndex(int32 LODBoneIndex) const
+	{
+		const int32 NumBonesLOD0 = LODBoneIndexes[0].Num();
+		check(LODBoneIndex < NumBonesLOD0);
+		return LODBoneIndexes[0][LODBoneIndex];
+	}
+
+	int32 GetLODBoneIndexFromSkeletonBoneIndex(int32 SkeletionBoneIndex) const
+	{
+		const int32 NumBonesLOD0 = SkeletonToLODBoneIndexes[0].Num();
+		check(SkeletionBoneIndex < NumBonesLOD0);
+		return SkeletonToLODBoneIndexes[0][SkeletionBoneIndex];
+	}
+
+	FTransform GetRefPoseTransform(int32 LODBoneIndex) const
+	{
+		const int32 NumBonesLOD0 = LODBoneIndexes[0].Num();
+		check(LODBoneIndex < NumBonesLOD0);
+
+		return ReferenceLocalTransforms[LODBoneIndex];
+	}
+
+	const FQuat& GetRefPoseRotation(int32 LODBoneIndex) const
+	{
+		const int32 NumBonesLOD0 = LODBoneIndexes[0].Num();
+		check(LODBoneIndex < NumBonesLOD0);
+
+		return ReferenceLocalTransforms[LODBoneIndex].GetRotation();
+	}
+
+	const FVector& GetRefPoseTranslation(int32 LODBoneIndex) const
+	{
+		const int32 NumBonesLOD0 = LODBoneIndexes[0].Num();
+		check(LODBoneIndex < NumBonesLOD0);
+
+		return ReferenceLocalTransforms[LODBoneIndex].GetTranslation();
+	}
+
+	const FVector& GetRefPoseScale3D(int32 LODBoneIndex) const
+	{
+		const int32 NumBonesLOD0 = LODBoneIndexes[0].Num();
+		check(LODBoneIndex < NumBonesLOD0);
+
+		return ReferenceLocalTransforms[LODBoneIndex].GetScale3D();
+	}
+
+};
+
+/**
+* Hide the allocator usage
+*/
+using FReferencePose = TReferencePose<FDefaultAllocator>;
+
+} // namespace UE::AnimNext
+
+// USTRUCT wrapper for reference pose
+USTRUCT()
+struct FAnimNextReferencePose
+#if CPP
+	: UE::AnimNext::FReferencePose
+#endif
+{
+	GENERATED_BODY()
+};
diff --git a/Engine/Plugins/Experimental/AnimNext/Source/AnimNextGraph/AnimNextGraph.Build.cs b/Engine/Plugins/Experimental/AnimNext/Source/AnimNextGraph/AnimNextGraph.Build.cs
new file mode 100644
--- /dev/null
+++ b/Engine/Plugins/Experimental/AnimNext/Source/AnimNextGraph/AnimNextGraph.Build.cs
@@ -0,0 +1,29 @@
+// Copyright Epic Games, Inc. All Rights Reserved.
+
+namespace UnrealBuildTool.Rules
+{
+	public class AnimNextGraph : ModuleRules
+	{
+		public AnimNextGraph(ReadOnlyTargetRules Target) : base(Target)
+		{
+		
+			PublicDependencyModuleNames.AddRange(
+				new string[]
+				{
+					"RigVM",
+					"ControlRig",
+				}
+			);
+
+			PrivateDependencyModuleNames.AddRange(
+				new string[]
+				{
+					"Core",
+					"CoreUObject",
+					"AnimNext",
+					"Engine"
+				}
+			);
+		}
+	}
+}
\ No newline at end of file
diff --git a/Engine/Plugins/Experimental/AnimNext/Source/AnimNextGraph/Internal/AnimNextGraph.h b/Engine/Plugins/Experimental/AnimNext/Source/AnimNextGraph/Internal/AnimNextGraph.h
new file mode 100644
--- /dev/null
+++ b/Engine/Plugins/Experimental/AnimNext/Source/AnimNextGraph/Internal/AnimNextGraph.h
@@ -0,0 +1,69 @@
+// Copyright Epic Games, Inc. All Rights Reserved.
+
+#pragma once
+
+#include "CoreMinimal.h"
+#include "Interface/IAnimNextInterface.h"
+#include "RigVMCore/RigVM.h"
+#include "Param/ParamType.h"
+#include "AnimNextGraph.generated.h"
+
+class UEdGraph;
+class UAnimNextGraph;
+namespace UE::AnimNext::GraphUncookedOnly
+{
+	struct FUtils;
+}
+
+namespace UE::AnimNext::GraphEditor
+{
+	class FGraphEditor;
+}
+
+namespace UE::AnimNext::Graph
+{
+	extern ANIMNEXTGRAPH_API const FName EntryPointName;
+	extern ANIMNEXTGRAPH_API const FName ResultName;
+}
+
+// A user-created graph of logic used to supply data
+UCLASS(BlueprintType)
+class ANIMNEXTGRAPH_API UAnimNextGraph : public UObject, public IAnimNextInterface
+{
+	GENERATED_BODY()
+
+	// UObject interface
+	virtual void PostRename(UObject* OldOuter, const FName OldName) override;
+	virtual void GetPreloadDependencies(TArray<UObject*>& OutDeps) override;
+
+	// IAnimNextInterface interface
+	virtual UE::AnimNext::FParamTypeHandle GetReturnTypeHandleImpl() const final override;
+	virtual bool GetDataImpl(const UE::AnimNext::FContext& Context) const final override;
+
+	/** Set the return type of this graph */
+	void SetReturnTypeHandle(UE::AnimNext::FParamTypeHandle InHandle);
+
+	// Support rig VM execution
+	TArray<FRigVMExternalVariable> GetRigVMExternalVariables();
+	
+	friend class UAnimNextGraphFactory;
+	friend class UAnimNextGraph_EditorData;
+	friend struct UE::AnimNext::GraphUncookedOnly::FUtils;
+	friend class UE::AnimNext::GraphEditor::FGraphEditor;
+	friend class UAnimNextGraph;
+	friend class UAnimGraphNode_AnimNextGraph;
+	
+	UPROPERTY()
+	TObjectPtr<URigVM> RigVM;
+
+	UPROPERTY()
+	FRigVMRuntimeSettings VMRuntimeSettings;
+
+	UPROPERTY()
+	FAnimNextParamType ReturnType;
+
+#if WITH_EDITORONLY_DATA
+	UPROPERTY(VisibleAnywhere, Instanced, Category = "Graph", meta = (ShowInnerProperties))
+	TObjectPtr<UObject> EditorData;
+#endif
+};
diff --git a/Engine/Plugins/Experimental/AnimNext/Source/AnimNextGraph/Internal/AnimNext_LODPose.h b/Engine/Plugins/Experimental/AnimNext/Source/AnimNextGraph/Internal/AnimNext_LODPose.h
new file mode 100644
--- /dev/null
+++ b/Engine/Plugins/Experimental/AnimNext/Source/AnimNextGraph/Internal/AnimNext_LODPose.h
@@ -0,0 +1,78 @@
+// Copyright Epic Games, Inc. All Rights Reserved.
+
+#pragma once
+
+#include "CoreMinimal.h"
+#include "Interface/InterfaceTypes.h"
+#include "Interface/IAnimNextInterface.h"
+#include "Interface/InterfaceContext.h"
+#include "ReferencePose.h"
+#include "LODPose.h"
+#include "AnimNext_LODPose.generated.h"
+
+USTRUCT(BlueprintType, meta = (DisplayName = "ReferencePose"))
+struct FAnimNextGraphReferencePose
+{
+	GENERATED_BODY()
+
+	FAnimNextGraphReferencePose() = default;
+
+	explicit FAnimNextGraphReferencePose(const UE::AnimNext::FReferencePose* InReferencePose)
+		: ReferencePose(InReferencePose)
+	{
+	}
+
+	const UE::AnimNext::FReferencePose* ReferencePose = nullptr;
+};
+
+USTRUCT(BlueprintType, meta = (DisplayName = "LODPose"))
+struct FAnimNextGraphLODPose
+{
+	GENERATED_BODY()
+
+	FAnimNextGraphLODPose() = default;
+
+	explicit FAnimNextGraphLODPose(const UE::AnimNext::FLODPose& InLODPose)
+		: LODPose(InLODPose)
+	{
+	}
+	explicit FAnimNextGraphLODPose(UE::AnimNext::FLODPose&& InLODPose)
+		: LODPose(MoveTemp(InLODPose))
+	{
+	}
+
+	UE::AnimNext::FLODPose LODPose;
+};
+
+UCLASS()
+class UAnimNextInterfaceGraphLODPose : public UObject, public IAnimNextInterface
+{
+	GENERATED_BODY()
+
+	virtual UE::AnimNext::FParamTypeHandle GetReturnTypeHandleImpl() const final override
+	{
+		return UE::AnimNext::FParamTypeHandle::GetHandle<FAnimNextGraphLODPose>();
+	}
+
+	virtual bool GetDataImpl(const UE::AnimNext::FContext& Context) const override
+	{
+		checkf(false, TEXT("UAnimNextInterfaceLODPose::GetDataImpl must be overridden"));
+		return false;
+	}
+};
+
+UCLASS()
+class UAnimNextInterface_LODPose_Literal : public UAnimNextInterfaceGraphLODPose
+{
+	GENERATED_BODY()
+
+	// IAnimNextInterface interface
+	virtual bool GetDataImpl(const UE::AnimNext::FContext& Context) const final override
+	{
+		Context.SetResult(Value);
+		return true;
+	}
+	
+	UPROPERTY(EditAnywhere, Category = "Parameters")
+	FAnimNextGraphLODPose Value;
+};
diff --git a/Engine/Plugins/Experimental/AnimNext/Source/AnimNextGraph/Internal/ExecuteContext.h b/Engine/Plugins/Experimental/AnimNext/Source/AnimNextGraph/Internal/ExecuteContext.h
new file mode 100644
--- /dev/null
+++ b/Engine/Plugins/Experimental/AnimNext/Source/AnimNextGraph/Internal/ExecuteContext.h
@@ -0,0 +1,84 @@
+﻿// Copyright Epic Games, Inc. All Rights Reserved.
+
+#pragma once
+
+#include "CoreMinimal.h"
+#include "RigVMCore/RigVMExecuteContext.h"
+#include "Units/RigUnit.h"
+#include "UnitContext.h"
+#include "ExecuteContext.generated.h"
+
+class IAnimNextInterface;
+
+namespace UE::AnimNext
+{
+	struct FContext;
+}
+
+USTRUCT(BlueprintType)
+struct FAnimNextExecuteContext : public FRigVMExecuteContext
+{
+	GENERATED_BODY()
+
+	FAnimNextExecuteContext()
+		: FRigVMExecuteContext()
+	, InterfaceContext(nullptr)
+	, Interface(nullptr)
+	, ResultPtr(nullptr)
+	, UnitContext()
+	{
+	}
+
+	const UE::AnimNext::FContext& GetContext() const
+	{
+		check(InterfaceContext);
+		return *InterfaceContext;
+	}
+
+	const FRigUnitContext& GetUnitContext() const
+	{
+		return UnitContext;
+	}
+
+	void SetContextData(const IAnimNextInterface* InInterface, const UE::AnimNext::FContext& InInterfaceContext, bool& bInResult)
+	{
+		Interface = InInterface;
+		InterfaceContext = &InInterfaceContext;
+		ResultPtr = &bInResult;
+	}
+
+	void SetResult(bool bInResult) const
+	{
+		check(ResultPtr);
+		*ResultPtr &= bInResult;
+	}
+
+	const IAnimNextInterface* GetInterface() const
+	{
+		check(Interface);
+		return Interface;
+	}
+	
+	virtual void Copy(const FRigVMExecuteContext* InOtherContext) override
+	{
+		Super::Copy(InOtherContext);
+
+		const FAnimNextExecuteContext* OtherContext = (const FAnimNextExecuteContext*)InOtherContext; 
+		InterfaceContext = OtherContext->InterfaceContext;
+		Interface = OtherContext->Interface;
+		ResultPtr = OtherContext->ResultPtr;
+	}
+
+
+private:
+	const UE::AnimNext::FContext* InterfaceContext;
+	const IAnimNextInterface* Interface;
+	bool* ResultPtr;
+	FRigUnitContext UnitContext;
+};
+
+USTRUCT(meta=(ExecuteContext="FAnimNextExecuteContext"))
+struct FRigUnit_AnimNextBase : public FRigUnit
+{
+	GENERATED_BODY()
+};
