File/namespace reorg

diff --git a/Engine/Plugins/Experimental/AnimNext/Source/AnimNextGraph/Internal/RigUnit_AnimNextEndExecution.h b/Engine/Plugins/Experimental/AnimNext/Source/AnimNextGraph/Internal/RigUnit_AnimNextEndExecution.h
new file mode 100644
--- /dev/null
+++ b/Engine/Plugins/Experimental/AnimNext/Source/AnimNextGraph/Internal/RigUnit_AnimNextEndExecution.h
@@ -0,0 +1,89 @@
+﻿// Copyright Epic Games, Inc. All Rights Reserved.
+
+#pragma once
+
+#include "Units/RigUnit.h"
+#include "ExecuteContext.h"
+#include "UnitContext.h"
+#include "AnimNext_LODPose.h"
+#include "RigUnit_AnimNextEndExecution.generated.h"
+
+struct FRigUnitContext;
+
+/** Event for writing back calculated results to external variables */
+USTRUCT(meta = (DisplayName = "End Execute", Category = "Events", TitleColor = "1 0 0", NodeColor = "1 1 1"))
+struct ANIMNEXTGRAPH_API FRigUnit_AnimNextEndExecution : public FRigUnit_AnimNextBase
+{
+	GENERATED_BODY()
+
+	template<typename ValueType>
+	static void SetResult(const FRigVMExecuteContext& InExecuteContext, const ValueType& InResult)
+	{
+		const FAnimNextExecuteContext& ExecuteContext = static_cast<const FAnimNextExecuteContext&>(InExecuteContext);
+		ExecuteContext.GetContext().SetResult<ValueType>(InResult);
+	}
+
+public:
+	// The execution result
+	UPROPERTY(EditAnywhere, Transient, DisplayName = "End Execute", Category = "EndExecution", meta = (Input))
+	FAnimNextExecuteContext ExecuteContext;
+};
+
+/** Event for writing back a calculated bool */
+USTRUCT(meta = (DisplayName = "End Execute Bool", Category = "Events", TitleColor = "1 0 0", NodeColor = "1 1 1"))
+struct ANIMNEXTGRAPH_API FRigUnit_AnimNextEndExecution_Bool : public FRigUnit_AnimNextEndExecution
+{
+	GENERATED_BODY()
+	
+	FRigUnit_AnimNextEndExecution_Bool()
+		: FRigUnit_AnimNextEndExecution()
+		, Result(false)
+	{}
+	
+	RIGVM_METHOD()
+	virtual void Execute() override;
+
+	virtual bool CanOnlyExistOnce() const override { return true; }
+	
+	UPROPERTY(EditAnywhere, Category = Result, meta = (Input))
+	bool Result;
+};
+
+/** Event for writing back a calculated float */
+USTRUCT(meta = (DisplayName = "End Execute Float", Category = "Events", TitleColor = "1 0 0", NodeColor = "1 1 1"))
+struct ANIMNEXTGRAPH_API FRigUnit_AnimNextEndExecution_Float : public FRigUnit_AnimNextEndExecution
+{
+	GENERATED_BODY()
+
+	FRigUnit_AnimNextEndExecution_Float()
+		: FRigUnit_AnimNextEndExecution()
+		, Result(0.f)
+	{}
+
+	RIGVM_METHOD()
+	virtual void Execute() override;
+
+	virtual bool CanOnlyExistOnce() const override { return true; }
+
+	UPROPERTY(EditAnywhere, Category = Result, meta = (Input))
+	float Result;
+};
+
+/** Event for writing back a calculated pose */
+USTRUCT(meta = (DisplayName = "End Execute LODPose", Category = "Events", TitleColor = "1 0 0", NodeColor = "1 1 1"))
+struct ANIMNEXTGRAPH_API FRigUnit_AnimNextEndExecution_LODPose : public FRigUnit_AnimNextEndExecution
+{
+	GENERATED_BODY()
+
+	FRigUnit_AnimNextEndExecution_LODPose()
+		: FRigUnit_AnimNextEndExecution()
+	{}
+
+	RIGVM_METHOD()
+	virtual void Execute() override;
+
+	virtual bool CanOnlyExistOnce() const override { return true; }
+
+	UPROPERTY(EditAnywhere, Category = Result, meta = (Input))
+	FAnimNextGraphLODPose Result;
+};
diff --git a/Engine/Plugins/Experimental/AnimNext/Source/AnimNextGraph/Internal/UnitContext.h b/Engine/Plugins/Experimental/AnimNext/Source/AnimNextGraph/Internal/UnitContext.h
new file mode 100644
--- /dev/null
+++ b/Engine/Plugins/Experimental/AnimNext/Source/AnimNextGraph/Internal/UnitContext.h
@@ -0,0 +1,28 @@
+// Copyright Epic Games, Inc. All Rights Reserved.
+
+#pragma once
+
+#include "CoreMinimal.h"
+#include "Interface/InterfaceContext.h"
+#include "Units/RigUnitContext.h"
+
+struct FAnimNextUnitContext : public FRigUnitContext
+{
+	FAnimNextUnitContext()
+		: FRigUnitContext()
+		, Interface(nullptr)
+		, InterfaceContext(nullptr)
+		, bResult(nullptr)
+	{}
+
+	FAnimNextUnitContext(const IAnimNextInterface* InInterface, const UE::AnimNext::FContext& InInterfaceContext, bool& bInResult)
+		: FRigUnitContext()
+		, Interface(InInterface)
+		, InterfaceContext(&InInterfaceContext)
+		, bResult(&bInResult)
+	{}
+
+	const IAnimNextInterface* Interface;
+	const UE::AnimNext::FContext* InterfaceContext;
+	bool* bResult;
+};
diff --git a/Engine/Plugins/Experimental/AnimNext/Source/AnimNextGraph/Private/AnimGraph/AnimNode_AnimNextGraph.cpp b/Engine/Plugins/Experimental/AnimNext/Source/AnimNextGraph/Private/AnimGraph/AnimNode_AnimNextGraph.cpp
new file mode 100644
--- /dev/null
+++ b/Engine/Plugins/Experimental/AnimNext/Source/AnimNextGraph/Private/AnimGraph/AnimNode_AnimNextGraph.cpp
@@ -0,0 +1,210 @@
+// Copyright Epic Games, Inc. All Rights Reserved.
+
+#include "AnimGraph/AnimNode_AnimNextGraph.h"
+#include "DataRegistry.h"
+#include "ReferencePose.h"
+#include "Components/SkeletalMeshComponent.h"
+#include "Animation/AnimInstanceProxy.h"
+#include "GameFramework/Actor.h"
+#include "AnimationRuntime.h"
+#include "Interface/AnimNextInterface.h"
+#include "DataRegistry.h"
+#include "GenerationTools.h"
+#include "ReferencePose.h"
+#include "AnimNext_LODPose.h"
+#include "RigUnit_AnimNextAnimSequence.h" // TEST
+#include "Animation/AnimSequence.h"
+#include "Animation/AnimSequence.h" // TEST
+#include "Engine/SkeletalMesh.h"
+#include "BoneContainer.h"
+
+#include UE_INLINE_GENERATED_CPP_BY_NAME(AnimNode_AnimNextGraph)
+
+#if WITH_EDITOR
+#include "Editor.h"
+#endif
+
+FAnimNode_AnimNextGraph::FAnimNode_AnimNextGraph()
+	: FAnimNode_CustomProperty()
+	, AnimNextGraph(nullptr)
+	, LODThreshold(INDEX_NONE)
+{
+}
+
+FAnimNode_AnimNextGraph::~FAnimNode_AnimNextGraph()
+{
+}
+
+void FAnimNode_AnimNextGraph::OnInitializeAnimInstance(const FAnimInstanceProxy* InProxy, const UAnimInstance* InAnimInstance)
+{
+	DECLARE_SCOPE_HIERARCHICAL_COUNTER_FUNC()
+
+	FAnimNode_CustomProperty::OnInitializeAnimInstance(InProxy, InAnimInstance);
+
+	InitializeProperties(InAnimInstance, GetTargetClass());
+}
+
+void FAnimNode_AnimNextGraph::GatherDebugData(FNodeDebugData& DebugData)
+{
+	DECLARE_SCOPE_HIERARCHICAL_COUNTER_FUNC()
+
+	SourceLink.GatherDebugData(DebugData.BranchFlow(1.f));
+}
+
+void FAnimNode_AnimNextGraph::Update_AnyThread(const FAnimationUpdateContext& Context)
+{
+	DECLARE_SCOPE_HIERARCHICAL_COUNTER_FUNC()
+
+	GraphDeltaTime += Context.GetDeltaTime();
+
+	SourceLink.Update(Context);
+
+	if (IsLODEnabled(Context.AnimInstanceProxy))
+	{
+		GetEvaluateGraphExposedInputs().Execute(Context);
+
+		PropagateInputProperties(Context.AnimInstanceProxy->GetAnimInstanceObject());
+	}
+
+	FAnimNode_CustomProperty::Update_AnyThread(Context);
+
+	//TRACE_ANIM_NODE_VALUE(Context, TEXT("Class"), *GetNameSafe(ControlRigClass.Get()));
+}
+
+void FAnimNode_AnimNextGraph::Initialize_AnyThread(const FAnimationInitializeContext& Context)
+{
+	DECLARE_SCOPE_HIERARCHICAL_COUNTER_FUNC()
+
+	SourceLink.Initialize(Context);
+
+	FAnimNode_CustomProperty::Initialize_AnyThread(Context);
+}
+
+void FAnimNode_AnimNextGraph::CacheBones_AnyThread(const FAnimationCacheBonesContext& Context)
+{
+	DECLARE_SCOPE_HIERARCHICAL_COUNTER_FUNC()
+
+	FAnimNode_CustomProperty::CacheBones_AnyThread(Context);
+
+	SourceLink.CacheBones(Context);
+}
+
+void FAnimNode_AnimNextGraph::Evaluate_AnyThread(FPoseContext & Output)
+{
+	DECLARE_SCOPE_HIERARCHICAL_COUNTER_FUNC()
+
+	using namespace UE::AnimNext;
+
+	FPoseContext SourcePose(Output);
+
+	if (SourceLink.GetLinkNode())
+	{
+		SourceLink.Evaluate(SourcePose);
+	}
+	else
+	{
+		SourcePose.ResetToRefPose();
+	}
+
+	USkeletalMeshComponent* SkeletalMeshComponent = Output.AnimInstanceProxy->GetSkelMeshComponent();
+	check(SkeletalMeshComponent != nullptr);
+
+	FDataHandle RefPoseHandle = FDataRegistry::Get()->GetOrGenerateReferencePose(SkeletalMeshComponent);
+	const UE::AnimNext::FReferencePose& RefPose = RefPoseHandle.GetRef<UE::AnimNext::FReferencePose>();
+
+	const int32 LODLevel = Output.AnimInstanceProxy->GetLODLevel();
+
+	// TODO : In a normal graph, FState/FParamStorage/FContext would come from FAnimationUpdateContext, here I have to create one the fly for the AnimNext graph
+	UE::AnimNext::FParamStorage ParamStorage(5, 0, 0);	// TODO : see if we can have state and params together and optional. Also, default size
+	
+	// TODO : Using AnimInstanceProxy->GetDeltaSeconds() makes the preview to advance  multiple times when debug options are activated (i.e. ShowUncompressedAnim)
+	// See where to get / how to calculate the correct delta value
+	UE::AnimNext::FContext RootContext(GraphDeltaTime, RootState, ParamStorage);
+	GraphDeltaTime = 0.f; // Reset for the case that we receive multiple calls to Evaluate (debug options)
+
+	FAnimNextGraphReferencePose GraphReferencePose(&RefPose);
+	FAnimNextGraph_AnimSequence GraphTestSequence(TestSequence);  // TEST
+
+	FAnimNextGraphLODPose GraphSourceLODPose(FLODPose(RefPose, LODLevel, false, Output.ExpectsAdditivePose()));
+	FGenerationTools::RemapPose(LODLevel, SourcePose, RefPose, GraphSourceLODPose.LODPose);
+
+	RootContext.AddInputReference(FName(TEXT("GraphReferencePose")), GraphReferencePose);			// TODO : Should this go into the AnimNextVM Context ?
+	RootContext.AddInputValue(FName(TEXT("GraphLODLevel")), LODLevel);								// TODO : Should this go into the AnimNextVM Context ?
+	RootContext.AddInputValue(FName(TEXT("GraphExpectsAdditive")), Output.ExpectsAdditivePose());	// TODO : Should this go into the AnimNextVM Context ?
+
+	RootContext.AddInputReference(FName(TEXT("SourcePose")), GraphSourceLODPose);		// TODO : Pass this as a external variable maybe ? When we have support for variables in the rigvm graph
+	RootContext.AddInputReference(FName(TEXT("TestSequence")), GraphTestSequence);		// TEST anim decompression - Anim Sequence to decompress
+		
+	TScriptInterface<IAnimNextInterface> ScriptInterface = AnimNextGraph;
+	FAnimNextGraphLODPose GraphResultLODPose(FLODPose(RefPose, LODLevel, true, Output.ExpectsAdditivePose()));
+
+	const bool bOk = UE::AnimNext::Interface::GetDataSafe(ScriptInterface, RootContext, GraphResultLODPose);
+	if (bOk)
+	{
+		FGenerationTools::RemapPose(LODLevel, RefPose, GraphResultLODPose.LODPose, Output);
+	}
+	else
+	{
+		Output.ResetToRefPose();
+	}
+
+
+	// evaluate 
+	FAnimNode_CustomProperty::Evaluate_AnyThread(Output);
+}
+
+void FAnimNode_AnimNextGraph::PostSerialize(const FArchive& Ar)
+{
+	DECLARE_SCOPE_HIERARCHICAL_COUNTER_FUNC()
+
+	// after compile, we have to reinitialize
+	// because it needs new execution code
+	// since memory has changed
+	if (Ar.IsObjectReferenceCollector())
+	{
+		if (AnimNextGraph)
+		{
+			//AnimNextInterfaceGraph->Initialize();
+		}
+	}
+}
+
+void FAnimNode_AnimNextGraph::InitializeProperties(const UObject* InSourceInstance, UClass* InTargetClass)
+{
+	// Build property lists
+	SourceProperties.Reset(SourcePropertyNames.Num());
+	DestProperties.Reset(SourcePropertyNames.Num());
+
+	check(SourcePropertyNames.Num() == DestPropertyNames.Num());
+
+	for (int32 Idx = 0; Idx < SourcePropertyNames.Num(); ++Idx)
+	{
+		FName& SourceName = SourcePropertyNames[Idx];
+		UClass* SourceClass = InSourceInstance->GetClass();
+
+		FProperty* SourceProperty = FindFProperty<FProperty>(SourceClass, SourceName);
+		SourceProperties.Add(SourceProperty);
+		DestProperties.Add(nullptr);
+	}
+}
+
+void FAnimNode_AnimNextGraph::PropagateInputProperties(const UObject* InSourceInstance)
+{
+	if (InSourceInstance)
+	{
+		// Assign value to the properties exposed as pins
+		for (int32 PropIdx = 0; PropIdx < SourceProperties.Num(); ++PropIdx)
+		{
+		}
+	}
+}
+
+
+#if WITH_EDITOR
+
+void FAnimNode_AnimNextGraph::HandleObjectsReinstanced_Impl(UObject* InSourceObject, UObject* InTargetObject, const TMap<UObject*, UObject*>& OldToNewInstanceMap)
+{
+	FAnimNode_CustomProperty::HandleObjectsReinstanced_Impl(InSourceObject, InTargetObject, OldToNewInstanceMap);
+}
+
+#endif
diff --git a/Engine/Plugins/Experimental/AnimNext/Source/AnimNextGraph/Private/AnimNextGraph.cpp b/Engine/Plugins/Experimental/AnimNext/Source/AnimNextGraph/Private/AnimNextGraph.cpp
new file mode 100644
--- /dev/null
+++ b/Engine/Plugins/Experimental/AnimNext/Source/AnimNextGraph/Private/AnimNextGraph.cpp
@@ -0,0 +1,100 @@
+// Copyright Epic Games, Inc. All Rights Reserved.
+
+#include "AnimNextGraph.h"
+#include "RigVMCore/RigVMMemoryStorage.h"
+#include "RigUnit_AnimNextBeginExecution.h"
+#include "UnitContext.h"
+#include "UObject/ObjectSaveContext.h"
+#include "UObject/Package.h"
+
+#include UE_INLINE_GENERATED_CPP_BY_NAME(AnimNextGraph)
+
+namespace UE::AnimNext::Graph
+{
+const FName EntryPointName("GetData");
+const FName ResultName("Result");
+}
+
+UE::AnimNext::FParamTypeHandle UAnimNextGraph::GetReturnTypeHandleImpl() const
+{
+	return ReturnType.GetHandle();
+}
+
+bool UAnimNextGraph::GetDataImpl(const UE::AnimNext::FContext& Context) const
+{
+	bool bResult = true;
+	
+	if(RigVM)
+	{
+		FAnimNextExecuteContext& AnimNextInterfaceContext = RigVM->GetContext().GetPublicDataSafe<FAnimNextExecuteContext>();
+
+		AnimNextInterfaceContext.SetContextData(this, Context, bResult);
+
+		bResult &= (RigVM->Execute(TArray<URigVMMemoryStorage*>(), FRigUnit_AnimNextBeginExecution::EventName) != ERigVMExecuteResult::Failed);
+	}
+	
+	return bResult;
+}
+
+void UAnimNextGraph::SetReturnTypeHandle(UE::AnimNext::FParamTypeHandle InHandle)
+{
+	ReturnType = InHandle.GetType();
+}
+
+TArray<FRigVMExternalVariable> UAnimNextGraph::GetRigVMExternalVariables()
+{
+	return TArray<FRigVMExternalVariable>(); 
+}
+
+static TArray<UClass*> GetClassObjectsInPackage(UPackage* InPackage)
+{
+	TArray<UObject*> Objects;
+	GetObjectsWithOuter(InPackage, Objects, false);
+
+	TArray<UClass*> ClassObjects;
+	for (UObject* Object : Objects)
+	{
+		if (UClass* Class = Cast<UClass>(Object))
+		{
+			ClassObjects.Add(Class);
+		}
+	}
+
+	return ClassObjects;
+}
+
+void UAnimNextGraph::PostRename(UObject* OldOuter, const FName OldName)
+{
+	Super::PostRename(OldOuter, OldName);
+
+	// Whenever the asset is renamed/moved, generated classes parented to the old package
+	// are not moved to the new package automatically (see FAssetRenameManager), so we
+	// have to manually perform the move/rename, to avoid invalid reference to the old package
+
+	// Note: while asset duplication doesn't duplicate the classes either, it is not a problem there
+	// because we always recompile in post duplicate.
+	TArray<UClass*> ClassObjects = GetClassObjectsInPackage(OldOuter->GetPackage());
+
+	for (UClass* ClassObject : ClassObjects)
+	{
+		if (URigVMMemoryStorageGeneratorClass* MemoryClass = Cast<URigVMMemoryStorageGeneratorClass>(ClassObject))
+		{
+			MemoryClass->Rename(nullptr, GetPackage(), REN_ForceNoResetLoaders | REN_DoNotDirty | REN_DontCreateRedirectors | REN_NonTransactional);
+		}
+	}
+}
+
+void UAnimNextGraph::GetPreloadDependencies(TArray<UObject*>& OutDeps)
+{
+	Super::GetPreloadDependencies(OutDeps);
+
+	TArray<UClass*> ClassObjects = GetClassObjectsInPackage(GetPackage());
+
+	for (UClass* ClassObject : ClassObjects)
+	{
+		if (URigVMMemoryStorageGeneratorClass* MemoryClass = Cast<URigVMMemoryStorageGeneratorClass>(ClassObject))
+		{
+			OutDeps.Add(MemoryClass);
+		}
+	}
+}
diff --git a/Engine/Plugins/Experimental/AnimNext/Source/AnimNextGraph/Private/GraphModule.cpp b/Engine/Plugins/Experimental/AnimNext/Source/AnimNextGraph/Private/GraphModule.cpp
new file mode 100644
--- /dev/null
+++ b/Engine/Plugins/Experimental/AnimNext/Source/AnimNextGraph/Private/GraphModule.cpp
@@ -0,0 +1,5 @@
+// Copyright Epic Games, Inc. All Rights Reserved.
+
+#include "Modules/ModuleManager.h"
+
+IMPLEMENT_MODULE(FDefaultModuleImpl, AnimNextGraph);
\ No newline at end of file
diff --git a/Engine/Plugins/Experimental/AnimNext/Source/AnimNextGraph/Private/RigUnit_AnimNextAnimSequence.cpp b/Engine/Plugins/Experimental/AnimNext/Source/AnimNextGraph/Private/RigUnit_AnimNextAnimSequence.cpp
new file mode 100644
--- /dev/null
+++ b/Engine/Plugins/Experimental/AnimNext/Source/AnimNextGraph/Private/RigUnit_AnimNextAnimSequence.cpp
@@ -0,0 +1,94 @@
+// Copyright Epic Games, Inc. All Rights Reserved.
+
+#include "RigUnit_AnimNextAnimSequence.h"
+#include "UnitContext.h"
+#include "ExecuteContext.h"
+#include "Interface/AnimNextInterface.h"
+#include "Interface/InterfaceTypes.h"
+#include "DecompressionTools.h"
+#include "Animation/AnimSequence.h"
+#include "AnimNext_LODPose.h"
+
+#include UE_INLINE_GENERATED_CPP_BY_NAME(RigUnit_AnimNextAnimSequence)
+
+// --- FRigUnit_AnimNext_AnimSequenceAsset --- 
+
+FRigUnit_AnimNext_AnimSequenceAsset_Execute()
+{
+	using namespace UE::AnimNext;
+
+	//Sequence.AnimSequence = AnimSequence; // TODO : Enable once we have RigVM object support enabled
+}
+
+// --- FRigUnit_AnimNext_AnimSequencePlayer --- 
+
+FRigUnit_AnimNext_AnimSequencePlayer_Initialize()
+{
+	using namespace UE::AnimNext;
+
+	const FContext& InterfaceContext = ExecuteContext.GetContext();
+
+	const UAnimSequenceBase* AnimSequenceBase = Sequence.AnimSequence;
+	if (AnimSequenceBase != nullptr)
+	{
+		// Get or Create the internal state
+		FAnimSequencePlayerState& AnimSequencePlayerState = InterfaceContext.GetState<FAnimSequencePlayerState>(ExecuteContext.GetInterface(), 0);
+
+		const float SequenceLength = AnimSequenceBase->GetPlayLength();
+		AnimSequencePlayerState.InternalTimeAccumulator = FMath::Clamp(Parameters.StartPosition, 0.f, SequenceLength);
+		AnimSequencePlayerState.PrevInternalTimeAccumulator = AnimSequencePlayerState.InternalTimeAccumulator;
+	}
+}
+
+FRigUnit_AnimNext_AnimSequencePlayer_Execute()
+{
+	using namespace UE::AnimNext;
+
+	const FContext& InterfaceContext = ExecuteContext.GetContext();
+	
+	const UAnimSequenceBase* AnimSequenceBase = Sequence.AnimSequence;
+	if (AnimSequenceBase != nullptr)
+	{
+		// Get or Create the internal state
+		//const TParam<FAnimSequencePlayerState> AnimSequencePlayerState = InterfaceContext.GetState<FAnimSequencePlayerState>(ExecuteContext.GetInterface(), 0);
+		FAnimSequencePlayerState& AnimSequencePlayerState = InterfaceContext.GetState<FAnimSequencePlayerState, EStatePersistence::Permanent>(ExecuteContext.GetInterface(), 0);
+
+		const float DeltaTime = InterfaceContext.GetDeltaTime();
+		const float EffectiveDelta = FMath::IsNearlyZero(DeltaTime) || FMath::IsNearlyZero(Parameters.PlayRate) ? 0.f : DeltaTime * Parameters.PlayRate;
+
+		const float SequenceLength = AnimSequenceBase->GetPlayLength();
+		float CurrentTime = Parameters.bLoop
+			? FMath::Fmod(AnimSequencePlayerState.InternalTimeAccumulator + EffectiveDelta, SequenceLength)
+			: FMath::Clamp(AnimSequencePlayerState.InternalTimeAccumulator + EffectiveDelta, 0.f, SequenceLength);
+
+		if (Parameters.bLoop && CurrentTime < 0.f)
+		{
+			CurrentTime += SequenceLength;
+		}
+
+		AnimSequencePlayerState.PrevInternalTimeAccumulator = AnimSequencePlayerState.InternalTimeAccumulator;
+		AnimSequencePlayerState.InternalTimeAccumulator = CurrentTime;
+
+		FDeltaTimeRecord DeltaTimeRecord;
+		DeltaTimeRecord.Set(AnimSequencePlayerState.PrevInternalTimeAccumulator, AnimSequencePlayerState.InternalTimeAccumulator);
+
+		const FAnimExtractContext ExtractionContext(static_cast<double>(AnimSequencePlayerState.InternalTimeAccumulator)
+			, false /*Output.AnimInstanceProxy->ShouldExtractRootMotion()*/
+			, DeltaTimeRecord
+			, Parameters.bLoop);
+
+		const FAnimNextGraphReferencePose& GraphReferencePose = InterfaceContext.GetParameterChecked<FAnimNextGraphReferencePose>(TEXT("GraphReferencePose")).GetDataChecked();
+		const int32 GraphLODLevel = InterfaceContext.GetParameterChecked<int32>(TEXT("GraphLODLevel")).GetDataChecked();
+		const bool bGraphExpectsAdditive = InterfaceContext.GetParameterChecked<bool>(TEXT("GraphExpectsAdditive")).GetDataChecked();
+
+		LODPose.LODPose.PrepareForLOD(*GraphReferencePose.ReferencePose, GraphLODLevel, true, bGraphExpectsAdditive);
+		
+		// Note : calling FDecompressionTools instead of UAnimSequence / UAnimSequenceBase, as I can not modify the engine code (so I extracted the function)
+		// TODO : Revisit this, so we can plug other sequence types
+		const UAnimSequence* AnimSequence = Cast<UAnimSequence>(AnimSequenceBase);
+		if (AnimSequence != nullptr)
+		{
+			FDecompressionTools::GetAnimationPose(AnimSequence, LODPose.LODPose, ExtractionContext);
+		}
+	}
+}
diff --git a/Engine/Plugins/Experimental/AnimNext/Source/AnimNextGraph/Private/RigUnit_AnimNextBeginExecution.cpp b/Engine/Plugins/Experimental/AnimNext/Source/AnimNextGraph/Private/RigUnit_AnimNextBeginExecution.cpp
new file mode 100644
--- /dev/null
+++ b/Engine/Plugins/Experimental/AnimNext/Source/AnimNextGraph/Private/RigUnit_AnimNextBeginExecution.cpp
@@ -0,0 +1,12 @@
+﻿// Copyright Epic Games, Inc. All Rights Reserved.
+
+#include "RigUnit_AnimNextBeginExecution.h"
+#include "Units/RigUnitContext.h"
+
+#include UE_INLINE_GENERATED_CPP_BY_NAME(RigUnit_AnimNextBeginExecution)
+
+FName FRigUnit_AnimNextBeginExecution::EventName = TEXT("Update");
+
+FRigUnit_AnimNextBeginExecution_Execute()
+{
+}
diff --git a/Engine/Plugins/Experimental/AnimNext/Source/AnimNextGraph/Private/RigUnit_AnimNextEndExecution.cpp b/Engine/Plugins/Experimental/AnimNext/Source/AnimNextGraph/Private/RigUnit_AnimNextEndExecution.cpp
new file mode 100644
--- /dev/null
+++ b/Engine/Plugins/Experimental/AnimNext/Source/AnimNextGraph/Private/RigUnit_AnimNextEndExecution.cpp
@@ -0,0 +1,21 @@
+﻿// Copyright Epic Games, Inc. All Rights Reserved.
+
+#include "RigUnit_AnimNextEndExecution.h"
+#include "Units/RigUnitContext.h"
+
+#include UE_INLINE_GENERATED_CPP_BY_NAME(RigUnit_AnimNextEndExecution)
+
+FRigUnit_AnimNextEndExecution_Bool_Execute()
+{
+	SetResult(ExecuteContext, Result);
+}
+
+FRigUnit_AnimNextEndExecution_Float_Execute()
+{
+	SetResult(ExecuteContext, Result);
+}
+
+FRigUnit_AnimNextEndExecution_LODPose_Execute()
+{
+	SetResult(ExecuteContext, Result);
+}
\ No newline at end of file
diff --git a/Engine/Plugins/Experimental/AnimNext/Source/AnimNextGraph/Private/RigUnit_AnimNextParameter.cpp b/Engine/Plugins/Experimental/AnimNext/Source/AnimNextGraph/Private/RigUnit_AnimNextParameter.cpp
new file mode 100644
--- /dev/null
+++ b/Engine/Plugins/Experimental/AnimNext/Source/AnimNextGraph/Private/RigUnit_AnimNextParameter.cpp
@@ -0,0 +1,109 @@
+// Copyright Epic Games, Inc. All Rights Reserved.
+
+#include "RigUnit_AnimNextParameter.h"
+#include "UnitContext.h"
+#include "ExecuteContext.h"
+#include "Interface/AnimNextInterface.h"
+#include "Interface/InterfaceTypes.h"
+
+#include UE_INLINE_GENERATED_CPP_BY_NAME(RigUnit_AnimNextParameter)
+
+bool FRigUnit_AnimNextParameter::GetParameterInternal(FName InName, const FAnimNextExecuteContext& InContext, void* OutResult)
+{
+	return true;
+}
+
+FRigUnit_AnimNextParameter_Float_Execute()
+{
+	GetParameterInternal(Parameter, ExecuteContext, &Result);
+}
+
+FRigUnit_AnimNextParameter_GraphLODPose_Execute()
+{
+	const UE::AnimNext::FContext& InterfaceContext = ExecuteContext.GetContext();
+
+	Result = InterfaceContext.GetParameterChecked<FAnimNextGraphLODPose>(Parameter);
+
+	//GetParameterInternal(Parameter, ExecuteContext, &Result);
+}
+
+FRigUnit_AnimNextParameter_AnimSequence_Execute()
+{
+	const UE::AnimNext::FContext& InterfaceContext = ExecuteContext.GetContext();
+
+	Result = InterfaceContext.GetParameterChecked<FAnimNextGraph_AnimSequence>(Parameter);
+
+	//GetParameterInternal(Parameter, ExecuteContext, &Result);
+}
+
+
+FRigUnit_AnimNextParameter_AnimNextInterface_Execute()
+{
+	GetParameterInternal(Parameter, ExecuteContext, &Result);
+}
+
+FRigUnit_AnimNext_Float_Execute()
+{
+	using namespace UE::AnimNext;
+
+	// @TODO: ensure that context arg is always present to avoid these checks here
+	const FContext& InterfaceContext = ExecuteContext.GetContext();
+
+	// Wrap the internal result we are going to be writing to 
+	TWrapParam<float> CallResult(&Result);
+
+	// Call the interface
+	ExecuteContext.SetResult(UE::AnimNext::Interface::GetDataSafe(AnimNextInterface, InterfaceContext, CallResult));
+}
+/*
+FRigUnit_AnimNext_Pose_Execute()
+{
+	using namespace UE::AnimNext;
+
+	// @TODO: ensure that context arg is always present to avoid these checks here
+	if(Context.State == EControlRigState::Update && ExecuteContext.OpaqueArguments.Num() > 1 && ExecuteContext.OpaqueArguments[1] != nullptr)
+	{
+		const FAnimNextUnitContext& UnitContext = *static_cast<const FAnimNextUnitContext*>(ExecuteContext.OpaqueArguments[1]);
+		const FContext& AnimNextContext = AnimNextUnitContext.InterfaceContext;
+
+		// Wrap the internal result we are going to be writing to 
+		TWrapParam<FAnimNextExecuteContext> CallResult(&Result);
+
+		// Call the interface
+		AnimNextInterfaceUnitContext.bResult &= UE::AnimNext::Interface::GetDataSafe(PoseInterface, AnimNextInterfaceContext, CallResult);
+	}
+}
+*/
+FRigUnit_FloatOperator_Execute()
+{}
+
+FRigUnit_PoseOperator_Execute()
+{}
+
+FRigUnit_AnimNext_SequencePlayer_Execute()
+{
+}
+
+
+
+FRigUnit_TestFloatState_Execute()
+{
+	using namespace UE::AnimNext;
+
+	// @TODO: ensure that context arg is always present to avoid these checks here
+	const FContext& InterfaceContext = ExecuteContext.GetContext();
+
+	FRigUnit_TestFloatState_SpringDamperState& State = InterfaceContext.GetState<FRigUnit_TestFloatState_SpringDamperState>(ExecuteContext.GetInterface(), 0);
+	const float DeltaTime = InterfaceContext.GetDeltaTime();
+
+	FMath::SpringDamperSmoothing(
+		State.Value,
+		State.ValueRate,
+		TargetValue,
+		TargetValueRate,
+		DeltaTime,
+		SmoothingTime,
+		DampingRatio);
+
+	Result = State.Value;
+}
diff --git a/Engine/Plugins/Experimental/AnimNext/Source/AnimNextGraph/Private/RigUnit_AnimNextParameter.h b/Engine/Plugins/Experimental/AnimNext/Source/AnimNextGraph/Private/RigUnit_AnimNextParameter.h
new file mode 100644
--- /dev/null
+++ b/Engine/Plugins/Experimental/AnimNext/Source/AnimNextGraph/Private/RigUnit_AnimNextParameter.h
@@ -0,0 +1,332 @@
+// Copyright Epic Games, Inc. All Rights Reserved.
+
+#pragma once
+
+#include "ExecuteContext.h"
+#include "Interface/IAnimNextInterface.h"
+#include "Units/RigUnit.h"
+#include "AnimNext_LODPose.h"
+#include "RigUnit_AnimNextAnimSequence.h"
+// --- ---
+#include "RigUnit_AnimNextParameter.generated.h"
+
+struct FAnimNextUnitContext;
+class UAnimNextGraph;
+
+/** Unit for reading parameters from context */
+USTRUCT(meta = (DisplayName = "Parameter", Category = "Parameters", TitleColor = "1 0 0", NodeColor = "1 1 1"))
+struct FRigUnit_AnimNextParameter : public FRigUnit_AnimNextBase
+{
+	GENERATED_BODY()
+
+protected:
+	static bool GetParameterInternal(FName InName, const FAnimNextExecuteContext& InContext, void* OutResult);
+
+	UPROPERTY(EditAnywhere, Category = "Parameter", meta = (Input))
+	FName Parameter = NAME_None;
+};
+
+/** Unit for reading float parameter from context */
+USTRUCT(meta = (DisplayName = "Float Parameter", Category = "Parameters", TitleColor = "1 0 0", NodeColor = "1 1 1"))
+struct FRigUnit_AnimNextParameter_Float : public FRigUnit_AnimNextParameter
+{
+	GENERATED_BODY()
+
+	RIGVM_METHOD()
+	void Execute();
+	
+protected:
+	UPROPERTY(EditAnywhere, Category = "Result", meta = (Output))
+	float Result = 0.0f;
+};
+
+/** Unit for reading a GraphLODPose parameter from context */
+USTRUCT(meta = (DisplayName = "LODPose Parameter", Category = "Parameters", TitleColor = "1 0 0", NodeColor = "1 1 1"))
+struct FRigUnit_AnimNextParameter_GraphLODPose : public FRigUnit_AnimNextParameter
+{
+	GENERATED_BODY()
+
+	RIGVM_METHOD()
+	void Execute();
+	
+protected:
+	UPROPERTY(EditAnywhere, Category = "Result", meta = (Output))
+	FAnimNextGraphLODPose Result;
+};
+
+/** Unit for reading an AnimSequence parameter from context */
+USTRUCT(meta = (DisplayName = "AnimSequence Parameter", Category = "Parameters", TitleColor = "1 0 0", NodeColor = "1 1 1"))
+struct FRigUnit_AnimNextParameter_AnimSequence : public FRigUnit_AnimNextParameter
+{
+	GENERATED_BODY()
+
+	RIGVM_METHOD()
+	void Execute();
+
+protected:
+	UPROPERTY(EditAnywhere, Category = "Result", meta = (Output))
+	FAnimNextGraph_AnimSequence Result;
+};
+
+/** Unit for reading anim interface graph parameter from context */
+USTRUCT(meta = (DisplayName = "Anim Interface Parameter", Category = "Parameters", TitleColor = "1 0 0", NodeColor = "1 1 1"))
+struct FRigUnit_AnimNextParameter_AnimNextInterface : public FRigUnit_AnimNextParameter
+{
+	GENERATED_BODY()
+
+	RIGVM_METHOD()
+	void Execute();
+	
+protected:
+	UPROPERTY(EditAnywhere, Category = "Result", meta = (Output))
+	TScriptInterface<IAnimNextInterface> Result = nullptr;
+};
+
+USTRUCT(BlueprintType)
+struct FAnimNextParameter
+{
+	GENERATED_BODY()
+
+	FAnimNextParameter()
+	: Name(NAME_None)
+	{}
+
+	UPROPERTY(EditAnywhere, Category = "Anim Interface")
+	FName Name;
+};
+
+USTRUCT(BlueprintType)
+struct FAnimNextParameter_Float
+{
+	GENERATED_BODY()
+	
+	FAnimNextParameter_Float()
+		: Name(NAME_None)
+		, Value(0.f)
+	{}
+
+	UPROPERTY(EditAnywhere, Category = "Anim Interface")
+	FName Name;
+
+	UPROPERTY(EditAnywhere, Category = "Anim Interface")
+	float Value;
+};
+
+USTRUCT(BlueprintType)
+struct FAnimNextParameter_Int
+{
+	GENERATED_BODY()
+
+	FAnimNextParameter_Int()
+		: Name(NAME_None)
+		, Value(0)
+	{}
+
+	UPROPERTY(EditAnywhere, Category = "Anim Interface")
+	FName Name;
+
+	UPROPERTY(EditAnywhere, Category = "Anim Interface")
+	int32 Value;
+};
+
+USTRUCT(BlueprintType)
+struct FAnimNextParameter_Bool
+{
+	GENERATED_BODY()
+
+	FAnimNextParameter_Bool()
+		: Name(NAME_None)
+		, bValue(false)
+	{}
+
+	UPROPERTY(EditAnywhere, Category = "Anim Interface")
+	FName Name;
+
+	UPROPERTY(EditAnywhere, Category = "Anim Interface")
+	bool bValue;
+};
+
+USTRUCT(BlueprintType)
+struct FAnimNextParameters1
+{
+	GENERATED_BODY()
+
+	FAnimNextParameters1()
+		: Param0()
+		, FloatParam(0.f)
+	{}
+
+	UPROPERTY(EditAnywhere, Category = "Anim Interface")
+	FAnimNextParameter_Int Param0;
+
+	UPROPERTY(EditAnywhere, Category = "Anim Interface")
+	float FloatParam;
+};
+
+USTRUCT(BlueprintType)
+struct FAnimNextParameter_AnimNextInterface : public FAnimNextParameter
+{
+	GENERATED_BODY()
+
+	UPROPERTY(EditAnywhere, Category = "Anim Interface")
+	TScriptInterface<IAnimNextInterface> Value = nullptr;
+};
+
+/** Base unit for calling anim interfaces from graphs */
+USTRUCT(meta = (DisplayName = "Anim Interface", Category = "Execution", TitleColor = "1 0 0", NodeColor = "1 1 1"))
+struct FRigUnit_AnimNextInterface : public FRigUnit_AnimNextBase
+{
+	GENERATED_BODY()
+
+protected:
+	UPROPERTY(EditAnywhere, Category = "Anim Interface", meta = (Input))
+	TScriptInterface<IAnimNextInterface> AnimNextInterface = nullptr;
+};
+
+/** Unit for getting a float via an anim interface */
+USTRUCT(meta = (DisplayName = "Anim Interface Float", Category = "Parameters", TitleColor = "1 0 0", NodeColor = "1 1 1"))
+struct FRigUnit_AnimNext_Float : public FRigUnit_AnimNextInterface
+{
+	GENERATED_BODY()
+
+	RIGVM_METHOD()
+	void Execute();
+	
+protected:
+	UPROPERTY(EditAnywhere, Category = "Anim Interface", meta = (Input))
+	TArray<FAnimNextParameter_AnimNextInterface> Parameters;
+
+	UPROPERTY(EditAnywhere, Category = "Result", meta = (Output))
+	float Result = 0.0f;
+};
+
+/** Unit for getting a pose via an anim interface */
+/*
+USTRUCT(meta=(DisplayName="Get Pose"))
+struct FRigUnit_AnimNext_Pose : public FRigUnit
+{
+	GENERATED_BODY()
+
+	RIGVM_METHOD()
+	void Execute();
+	
+protected:
+	UPROPERTY(EditAnywhere, Category = "Pose Interface", meta = (Input))
+	TScriptInterface<IAnimNextInterface> PoseInterface = nullptr;
+
+	UPROPERTY(EditAnywhere, Category = "Pose Interface", meta = (Input))
+	TArray<FAnimNextParameter_AnimNextInterface> Parameters;
+
+	UPROPERTY(EditAnywhere, Transient, DisplayName = "Result", Category = "BeginExecution", meta = (Output))
+	FAnimNextInterfaceExecuteContext Result;
+};
+*/
+/** Unit for getting a pose via an anim interface */
+USTRUCT(meta = (DisplayName = "Float Operator", Category = "Operators", TitleColor = "1 0 0", NodeColor = "1 1 1"))
+struct FRigUnit_FloatOperator : public FRigUnit_AnimNextBase
+{
+	GENERATED_BODY()
+	
+	FRigUnit_FloatOperator()
+		: ParamA(0.f)
+		, ParamB(0.f)
+		, Result(0.f)
+	{}
+
+	RIGVM_METHOD()
+	void Execute();
+	
+protected:
+	UPROPERTY(EditAnywhere, Category = "Pose Interface", meta = (Input))
+	TScriptInterface<IAnimNextInterface> Operator = nullptr;
+
+	UPROPERTY(EditAnywhere, Category = "Pose Interface", meta = (Input))
+	float ParamA;
+
+	UPROPERTY(EditAnywhere, Category = "Pose Interface", meta = (Input))
+	float ParamB;
+
+	UPROPERTY(EditAnywhere, Transient, DisplayName = "Result", Category = "BeginExecution", meta = (Output))
+	float Result;
+};
+
+/** Unit for getting a pose via an anim interface */
+USTRUCT(meta = (DisplayName = "Pose Operator", Category = "Operators", TitleColor = "1 0 0", NodeColor = "1 1 1"))
+struct FRigUnit_PoseOperator : public FRigUnit_AnimNextBase
+{
+	GENERATED_BODY()
+
+	RIGVM_METHOD()
+	void Execute();
+	
+protected:
+	UPROPERTY(EditAnywhere, Category = "Pose Interface", meta = (Input))
+	TScriptInterface<IAnimNextInterface> Operator = nullptr;
+
+	UPROPERTY(EditAnywhere, Category = "Pose Interface", meta = (Input))
+	TArray<FRigVMExecuteContext> InputPoses;
+
+	UPROPERTY(EditAnywhere, Category = "Pose Interface", meta = (Input))
+	FAnimNextParameters1 Parameters;
+
+	UPROPERTY(EditAnywhere, Transient, DisplayName = "Result", Category = "BeginExecution", meta = (Output))
+	FAnimNextExecuteContext Result;
+};
+
+/** Unit for getting a pose via an anim interface */
+USTRUCT(meta = (DisplayName = "Anim Sequence", Category = "Animation", TitleColor = "1 0 0", NodeColor = "1 1 1"))
+struct FRigUnit_AnimNext_SequencePlayer : public FRigUnit_AnimNextBase
+{
+	GENERATED_BODY()
+
+		RIGVM_METHOD()
+		void Execute();
+
+protected:
+	UPROPERTY(EditAnywhere, Category = "Anim Interface", meta = (Input))
+	FAnimSequenceParameters Parameters;
+
+	UPROPERTY(EditAnywhere, Category = "Anim Interface", meta = (Input))
+	TScriptInterface<IAnimNextInterface> Sequence = nullptr;
+
+	UPROPERTY(EditAnywhere, Transient, DisplayName = "Result", Category = "BeginExecution", meta = (Output))
+	FAnimNextExecuteContext Result;
+};
+
+USTRUCT()
+struct FRigUnit_TestFloatState_SpringDamperState
+{
+	GENERATED_BODY()
+
+	UPROPERTY()
+	float Value = 0.0f;
+
+	UPROPERTY()
+	float ValueRate = 0.0f;
+};
+
+/** Unit for getting a pose via an anim interface */
+USTRUCT(meta = (DisplayName = "Test Float State - Spring Damper Smoothing", Varying, Category = "Animation", TitleColor = "1 0 0", NodeColor = "1 1 1"))
+struct FRigUnit_TestFloatState : public FRigUnit_AnimNextBase
+{
+	GENERATED_BODY()
+
+	RIGVM_METHOD()
+	void Execute();
+	
+protected:
+	UPROPERTY(EditAnywhere, Category = "Parameters", meta=(Input))
+	float TargetValue = 0.0f;
+
+	UPROPERTY(EditAnywhere, Category = "Parameters", meta=(Input))
+	float TargetValueRate = 1.0f;
+	
+	UPROPERTY(EditAnywhere, Category = "Parameters", meta=(Input))
+	float SmoothingTime = 1.0f;
+	
+	UPROPERTY(EditAnywhere, Category = "Parameters", meta=(Input))
+	float DampingRatio = 1.0f;
+
+	UPROPERTY(EditAnywhere, Category = "Parameters", meta=(Output))
+	float Result = 0.0f;
+};
diff --git a/Engine/Plugins/Experimental/AnimNext/Source/AnimNextGraphEditor/AnimNextGraphEditor.Build.cs b/Engine/Plugins/Experimental/AnimNext/Source/AnimNextGraphEditor/AnimNextGraphEditor.Build.cs
new file mode 100644
--- /dev/null
+++ b/Engine/Plugins/Experimental/AnimNext/Source/AnimNextGraphEditor/AnimNextGraphEditor.Build.cs
@@ -0,0 +1,34 @@
+// Copyright Epic Games, Inc. All Rights Reserved.
+
+namespace UnrealBuildTool.Rules
+{
+	public class AnimNextGraphEditor : ModuleRules
+	{
+		public AnimNextGraphEditor(ReadOnlyTargetRules Target) : base(Target)
+		{
+
+			PrivateDependencyModuleNames.AddRange(
+				new string[]
+				{
+					"Core",
+					"CoreUObject",
+					"Engine",
+					"AssetTools",
+					"AnimNext",
+					"AnimNextGraph",
+					"AnimNextGraphUncookedOnly",
+					"UnrealEd",
+					"SlateCore",
+					"Slate",
+					"PropertyEditor",
+					"RigVM",
+					"RigVMDeveloper",
+					"ControlRig",
+					"ControlRigDeveloper",
+					"ControlRigEditor",
+					"GraphEditor",
+				}
+			);
+		}
+	}
+}
\ No newline at end of file
diff --git a/Engine/Plugins/Experimental/AnimNext/Source/AnimNextGraphEditor/Private/AnimNextGraphEditor.cpp b/Engine/Plugins/Experimental/AnimNext/Source/AnimNextGraphEditor/Private/AnimNextGraphEditor.cpp
new file mode 100644
--- /dev/null
+++ b/Engine/Plugins/Experimental/AnimNext/Source/AnimNextGraphEditor/Private/AnimNextGraphEditor.cpp
@@ -0,0 +1,322 @@
+﻿// Copyright Epic Games, Inc. All Rights Reserved.
+
+#include "AnimNextGraphEditor.h"
+#include "GraphEditorMode.h"
+#include "AnimNextGraph.h"
+#include "AnimNextGraph_EditorData.h"
+#include "EdGraphNode_Comment.h"
+#include "EdGraphUtilities.h"
+#include "SActionMenu.h"
+#include "UncookedOnlyUtils.h"
+#include "Framework/Commands/GenericCommands.h"
+#include "RigVMModel/RigVMController.h"
+#include "RigVMModel/RigVMGraph.h"
+#include "Kismet2/BlueprintEditorUtils.h"
+#include "ScopedTransaction.h"
+
+#define LOCTEXT_NAMESPACE "AnimNextGraphEditor"
+
+namespace UE::AnimNext::GraphEditor
+{
+
+namespace Modes
+{
+	const FName GraphEditor("AnimNextInterfaceEditorMode");
+}
+
+namespace Tabs
+{
+	const FName Details("DetailsTab");
+	const FName Document("Document");
+}
+
+const FName AppIdentifier("AnimNextInterfaceGraphEditor");
+
+void FGraphEditor::InitEditor(const EToolkitMode::Type InMode, const TSharedPtr<IToolkitHost>& InInitToolkitHost, UAnimNextGraph* InAnimNextGraph)
+{
+	AnimNextGraph = InAnimNextGraph;
+
+	AnimNextGraph_EditorData = UE::AnimNext::GraphUncookedOnly::FUtils::GetEditorData(AnimNextGraph);
+	AnimNextGraph_EditorData->Initialize(false);
+	
+	DocumentManager = MakeShared<FDocumentTracker>();
+	DocumentManager->Initialize(SharedThis(this));
+
+	TSharedRef<FGraphEditorSummoner> GraphEditorSummoner = MakeShared<FGraphEditorSummoner>(SharedThis(this));
+	GraphEditorSummoner->OnCreateGraphEditorWidget().BindSP(this, &FGraphEditor::CreateGraphEditorWidget);
+	GraphEditorSummoner->OnGraphEditorFocused().BindSP(this, &FGraphEditor::OnGraphEditorFocused);
+	GraphEditorSummoner->OnGraphEditorBackgrounded().BindSP(this, &FGraphEditor::OnGraphEditorBackgrounded);
+	GraphEditorTabFactoryPtr = GraphEditorSummoner;
+	DocumentManager->RegisterDocumentFactory(GraphEditorSummoner);
+
+	const bool bCreateDefaultStandaloneMenu = true;
+	const bool bCreateDefaultToolbar = true;
+	FAssetEditorToolkit::InitAssetEditor(InMode, InInitToolkitHost, AppIdentifier, FTabManager::FLayout::NullLayout, bCreateDefaultStandaloneMenu, bCreateDefaultToolbar, InAnimNextGraph);
+
+	BindCommands();
+
+	AddApplicationMode(Modes::GraphEditor, MakeShared<FGraphEditorMode>(SharedThis(this)));
+	SetCurrentMode(Modes::GraphEditor);
+	
+	ExtendMenu();
+	ExtendToolbar();
+	RegenerateMenusAndToolbars();
+
+	// Open initial document
+	//DocumentManager->OpenDocument(FTabPayload_UObject::Make(AnimNextInterfaceGraph_EditorData->EntryPointGraph), FDocumentTracker::EOpenDocumentCause::OpenNewDocument);
+	DocumentManager->OpenDocument(FTabPayload_UObject::Make(AnimNextGraph_EditorData->RootGraph), FDocumentTracker::EOpenDocumentCause::OpenNewDocument);
+	//DocumentManager->OpenDocument(FTabPayload_UObject::Make(AnimNextInterfaceGraph_EditorData->FunctionLibraryEdGraph), FDocumentTracker::EOpenDocumentCause::OpenNewDocument);
+}
+
+void FGraphEditor::BindCommands()
+{
+}
+
+void FGraphEditor::ExtendMenu()
+{
+	
+}
+
+void FGraphEditor::ExtendToolbar()
+{
+	
+}
+
+void FGraphEditor::RegisterTabSpawners(const TSharedRef<FTabManager>& InTabManager)
+{
+	DocumentManager->SetTabManager(InTabManager);
+
+	FWorkflowCentricApplication::RegisterTabSpawners(InTabManager);
+}
+
+void FGraphEditor::UnregisterTabSpawners(const TSharedRef<FTabManager>& InTabManager)
+{
+	
+}
+
+FName FGraphEditor::GetToolkitFName() const
+{
+	return FName("AnimNextInterfaceGraphEditor");
+}
+
+FText FGraphEditor::GetBaseToolkitName() const
+{
+	return LOCTEXT("AppLabel", "AnimNextInterfaceGraphEditor");
+}
+
+FString FGraphEditor::GetWorldCentricTabPrefix() const
+{
+	return LOCTEXT("WorldCentricTabPrefix", "AnimNextInterfaceGraphEditor ").ToString();
+}
+
+FLinearColor FGraphEditor::GetWorldCentricTabColorScale() const
+{
+	return FLinearColor(0.3f, 0.2f, 0.5f, 0.5f);
+}
+
+void FGraphEditor::InitToolMenuContext(FToolMenuContext& InMenuContext)
+{
+}
+
+void FGraphEditor::OnGraphEditorFocused(TSharedRef<SGraphEditor> InGraphEditor)
+{
+	// Update the graph editor that is currently focused
+	FocusedGraphEdPtr = InGraphEditor;
+}
+
+void FGraphEditor::OnGraphEditorBackgrounded(TSharedRef<SGraphEditor> InGraphEditor)
+{
+	// Update the graph editor that is currently focused
+	FocusedGraphEdPtr = nullptr;
+}
+
+UEdGraph* FGraphEditor::GetFocusedGraph() const
+{
+	if (FocusedGraphEdPtr.IsValid())
+	{
+		if (UEdGraph* Graph = FocusedGraphEdPtr.Pin()->GetCurrentGraph())
+		{
+			return Graph;
+		}
+	}
+	return nullptr;
+}
+
+URigVMGraph* FGraphEditor::GetFocusedVMGraph() const
+{
+	UAnimNextGraph_EdGraph* EdGraph = Cast<UAnimNextGraph_EdGraph>(GetFocusedGraph());
+	return AnimNextGraph_EditorData->GetVMGraphForEdGraph(EdGraph);
+}
+
+URigVMController* FGraphEditor::GetFocusedVMController() const
+{
+	return AnimNextGraph_EditorData->GetRigVMClient()->GetController(GetFocusedVMGraph());
+}
+
+TSharedRef<SGraphEditor> FGraphEditor::CreateGraphEditorWidget(TSharedRef<FTabInfo> InTabInfo, UEdGraph* InGraph)
+{
+	if(!GraphEditorCommands.IsValid())
+	{
+		GraphEditorCommands = MakeShared<FUICommandList>();
+	
+		GraphEditorCommands->MapAction(FGenericCommands::Get().Delete,
+			FExecuteAction::CreateSP(this, &FGraphEditor::DeleteSelectedNodes),
+			FCanExecuteAction::CreateSP(this, &FGraphEditor::CanDeleteSelectedNodes));
+	}
+	
+	SGraphEditor::FGraphEditorEvents Events;
+	Events.OnCreateActionMenu = SGraphEditor::FOnCreateActionMenu::CreateSP(this, &FGraphEditor::OnCreateGraphActionMenu);
+	
+	return SNew(SGraphEditor)
+		.AdditionalCommands(GraphEditorCommands)
+		.IsEditable(this, &FGraphEditor::IsEditable, InGraph)
+		.GraphToEdit(InGraph)
+		.GraphEvents(Events)
+		.AssetEditorToolkit(AsShared());
+}
+
+FActionMenuContent FGraphEditor::OnCreateGraphActionMenu(UEdGraph* InGraph, const FVector2D& InNodePosition, const TArray<UEdGraphPin*>& InDraggedPins, bool bAutoExpand, SGraphEditor::FActionMenuClosed InOnMenuClosed)
+{
+	TSharedRef<SActionMenu> ActionMenu = SNew(SActionMenu)
+		.AutoExpandActionMenu(bAutoExpand)
+		.Graph(InGraph)
+		.NewNodePosition(InNodePosition)
+		.DraggedFromPins(InDraggedPins)
+		.OnClosedCallback(InOnMenuClosed);
+
+	TSharedPtr<SWidget> FilterTextBox = StaticCastSharedRef<SWidget>(ActionMenu->GetFilterTextBox());
+	return FActionMenuContent(StaticCastSharedRef<SWidget>(ActionMenu), FilterTextBox);
+}
+
+void FGraphEditor::OnNodeTitleCommitted(const FText& NewText, ETextCommit::Type CommitInfo, UEdGraphNode* NodeBeingChanged)
+{
+	if (UEdGraphNode_Comment* CommentBeingChanged = Cast<UEdGraphNode_Comment>(NodeBeingChanged))
+	{
+		GetFocusedVMController()->SetCommentTextByName(CommentBeingChanged->GetFName(), NewText.ToString(), CommentBeingChanged->FontSize, CommentBeingChanged->bCommentBubbleVisible, CommentBeingChanged->bColorCommentBubble, true, true);
+	}
+}
+
+void FGraphEditor::CloseDocumentTab(const UObject* DocumentID)
+{
+	TSharedRef<FTabPayload_UObject> Payload = FTabPayload_UObject::Make(DocumentID);
+	DocumentManager->CloseTab(Payload);
+}
+
+bool FGraphEditor::InEditingMode() const
+{
+	// @TODO: disallow editing when debugging when implemented
+
+	return true;
+}
+
+bool FGraphEditor::IsEditable(UEdGraph* InGraph) const
+{
+	return InGraph && InEditingMode() && InGraph->bEditable;
+}
+
+FGraphPanelSelectionSet FGraphEditor::GetSelectedNodes() const
+{
+	FGraphPanelSelectionSet CurrentSelection;
+	TSharedPtr<SGraphEditor> FocusedGraphEd = FocusedGraphEdPtr.Pin();
+	if (FocusedGraphEd.IsValid())
+	{
+		CurrentSelection = FocusedGraphEd->GetSelectedNodes();
+	}
+	return CurrentSelection;
+}
+
+void FGraphEditor::DeleteSelectedNodes()
+{
+	TSharedPtr<SGraphEditor> FocusedGraphEd = FocusedGraphEdPtr.Pin();
+	if (!FocusedGraphEd.IsValid())
+	{
+		return;
+	}
+
+	const FScopedTransaction Transaction(FGenericCommands::Get().Delete->GetDescription());
+	FocusedGraphEd->GetCurrentGraph()->Modify();
+	
+	const FGraphPanelSelectionSet SelectedNodes = GetSelectedNodes();
+
+//	SetUISelectionState(NAME_None);
+
+	if(FocusedGraphEd)
+	{
+		FocusedGraphEd->ClearSelectionSet();
+	}
+
+	// Some nodes have sub-objects that are represented as other tabs.
+	// Close them here as a pre-pass before we remove their nodes. If the documents are left open they
+	// may reference dangling data and function incorrectly in cases such as FindBlueprintforNodeChecked
+	for (FGraphPanelSelectionSet::TConstIterator NodeIt( SelectedNodes ); NodeIt; ++NodeIt)
+	{
+		if (UEdGraphNode* Node = Cast<UEdGraphNode>(*NodeIt))
+		{
+			if (Node->CanUserDeleteNode())
+			{
+				auto CloseAllDocumentsTab = [this](const UEdGraphNode* InNode)
+				{
+					TArray<UObject*> NodesToClose;
+					GetObjectsWithOuter(InNode, NodesToClose);
+					for (UObject* Node : NodesToClose)
+					{
+						UEdGraph* NodeGraph = Cast<UEdGraph>(Node);
+						if (NodeGraph)
+						{
+							CloseDocumentTab(NodeGraph);
+						}
+					}
+				};
+				
+				if (Node->GetSubGraphs().Num() > 0)
+				{
+					CloseAllDocumentsTab(Node);
+				}
+			}
+		}
+	}
+
+	// Now remove the selected nodes
+	for (FGraphPanelSelectionSet::TConstIterator NodeIt( SelectedNodes ); NodeIt; ++NodeIt)
+	{
+		if (UEdGraphNode* Node = Cast<UEdGraphNode>(*NodeIt))
+		{
+			if (Node->CanUserDeleteNode())
+			{
+				if (Node->GetSubGraphs().Num() > 0)
+				{
+					DocumentManager->CleanInvalidTabs();
+				}
+
+				FBlueprintEditorUtils::RemoveNode(nullptr, Node);
+			}
+		}
+	}
+}
+
+bool FGraphEditor::CanDeleteSelectedNodes()
+{
+	const FGraphPanelSelectionSet SelectedNodes = GetSelectedNodes();
+
+	bool bCanUserDeleteNode = false;
+
+	if(IsEditable(GetFocusedGraph()) && SelectedNodes.Num() > 0)
+	{
+		for(UObject* NodeObject : SelectedNodes)
+		{
+			// If any nodes allow deleting, then do not disable the delete option
+			UEdGraphNode* Node = Cast<UEdGraphNode>(NodeObject);
+			if(Node->CanUserDeleteNode())
+			{
+				bCanUserDeleteNode = true;
+				break;
+			}
+		}
+	}
+
+	return bCanUserDeleteNode;
+}
+
+}
+
+#undef LOCTEXT_NAMESPACE
\ No newline at end of file
diff --git a/Engine/Plugins/Experimental/AnimNext/Source/AnimNextGraphEditor/Private/AnimNextGraphEditor.h b/Engine/Plugins/Experimental/AnimNext/Source/AnimNextGraphEditor/Private/AnimNextGraphEditor.h
new file mode 100644
--- /dev/null
+++ b/Engine/Plugins/Experimental/AnimNext/Source/AnimNextGraphEditor/Private/AnimNextGraphEditor.h
@@ -0,0 +1,123 @@
+﻿// Copyright Epic Games, Inc. All Rights Reserved.
+
+#pragma once
+
+#include "CoreMinimal.h"
+#include "WorkflowOrientedApp/WorkflowCentricApplication.h"
+#include "GraphEditor.h"
+
+enum class ERigVMGraphNotifType : uint8;
+class UAnimNextGraph;
+class UAnimNextGraph_EditorData;
+class FDocumentTracker;
+class FDocumentTabFactory;
+class FTabInfo;
+class FTabManager;
+class IToolkitHost;
+class URigVMGraph;
+class URigVMController;
+class UDetailsViewWrapperObject;
+
+namespace UE::AnimNext::GraphEditor
+{
+
+namespace Modes
+{
+	extern const FName GraphEditor;
+}
+
+namespace Tabs
+{
+	extern const FName Details;
+	extern const FName Document;
+}
+
+class FGraphEditor : public FWorkflowCentricApplication
+{
+public:
+	/** Edits the specified asset */
+	void InitEditor(const EToolkitMode::Type InMode, const TSharedPtr<IToolkitHost>& InInitToolkitHost, UAnimNextGraph* InAnimNextInterfaceGraph);
+
+private:
+	friend class FGraphEditorMode;
+	
+	// IToolkit interface
+	virtual void RegisterTabSpawners(const TSharedRef<FTabManager>& InTabManager) override;
+	virtual void UnregisterTabSpawners(const TSharedRef<FTabManager>& InTabManager) override;
+	virtual FName GetToolkitFName() const override;
+	virtual FText GetBaseToolkitName() const override;
+	virtual FString GetWorldCentricTabPrefix() const override;
+	virtual FLinearColor GetWorldCentricTabColorScale() const override;
+	virtual void InitToolMenuContext(FToolMenuContext& InMenuContext) override;
+
+	void BindCommands();
+	
+	void ExtendMenu();
+
+	void ExtendToolbar();
+
+	void OnGraphEditorFocused(TSharedRef<SGraphEditor> InGraphEditor);
+
+	void OnGraphEditorBackgrounded(TSharedRef<SGraphEditor> InGraphEditor);
+	
+	UEdGraph* GetFocusedGraph() const;
+
+	URigVMGraph* GetFocusedVMGraph() const;
+
+	URigVMController* GetFocusedVMController() const;
+	
+	// Create new graph editor widget for the supplied document container
+	TSharedRef<SGraphEditor> CreateGraphEditorWidget(TSharedRef<FTabInfo> InTabInfo, UEdGraph* InGraph);
+
+	// Create the action menu for the specified graph
+	FActionMenuContent OnCreateGraphActionMenu(UEdGraph* InGraph, const FVector2D& InNodePosition, const TArray<UEdGraphPin*>& InDraggedPins, bool bAutoExpand, SGraphEditor::FActionMenuClosed InOnMenuClosed);
+
+	// Handle node titles being changed
+	void OnNodeTitleCommitted(const FText& NewText, ETextCommit::Type CommitInfo, UEdGraphNode* NodeBeingChanged);
+
+	void CloseDocumentTab(const UObject* DocumentID);
+
+	bool InEditingMode() const;
+
+	bool IsEditable(UEdGraph* InGraph) const;
+	
+	FGraphPanelSelectionSet GetSelectedNodes() const;
+	
+	void DeleteSelectedNodes();
+	
+	bool CanDeleteSelectedNodes();
+
+	TSharedPtr<SGraphEditor> GetGraphEditor(UEdGraph* InEdGraph) const;
+	
+	void HandleModifiedEvent(ERigVMGraphNotifType InNotifType, URigVMGraph* InGraph, UObject* InSubject);
+
+	void HandleDetailsViewCreated(TSharedRef<IDetailsView> InDetailsView)
+	{
+		DetailsView = InDetailsView;
+	}
+	
+	// The asset we are editing
+	UAnimNextGraph* AnimNextGraph = nullptr;
+
+	// Cached editor data ptr
+	UAnimNextGraph_EditorData* AnimNextGraph_EditorData = nullptr;
+	
+	// Currently focused graph editor
+	TWeakPtr<SGraphEditor> FocusedGraphEdPtr;
+
+	// Document tracker
+	TSharedPtr<FDocumentTracker> DocumentManager;
+
+	// Factory for document tabs
+	TWeakPtr<FDocumentTabFactory> GraphEditorTabFactoryPtr;
+
+	// Command list for the graph editor
+	TSharedPtr<FUICommandList> GraphEditorCommands;
+
+	// Our details panel
+	TSharedPtr<IDetailsView> DetailsView;
+
+	TArray<TStrongObjectPtr<UDetailsViewWrapperObject>> WrapperObjects;
+};
+
+}
\ No newline at end of file
diff --git a/Engine/Plugins/Experimental/AnimNext/Source/AnimNextGraphEditor/Private/AssetTypeActions.cpp b/Engine/Plugins/Experimental/AnimNext/Source/AnimNextGraphEditor/Private/AssetTypeActions.cpp
new file mode 100644
--- /dev/null
+++ b/Engine/Plugins/Experimental/AnimNext/Source/AnimNextGraphEditor/Private/AssetTypeActions.cpp
@@ -0,0 +1,37 @@
+﻿// Copyright Epic Games, Inc. All Rights Reserved.
+
+#include "AssetTypeActions.h"
+#include "AnimNextGraphEditor.h"
+#include "AnimNextGraph.h"
+
+#define LOCTEXT_NAMESPACE "AssetTypeActions_AnimNextGraph" 
+
+namespace UE::AnimNext::GraphEditor
+{
+
+void FAssetTypeActions_AnimNextGraph::OpenAssetEditor(const TArray<UObject*>& InObjects, TSharedPtr<IToolkitHost> EditWithinLevelEditor)
+{
+	const EToolkitMode::Type Mode = EditWithinLevelEditor.IsValid() ? EToolkitMode::WorldCentric : EToolkitMode::Standalone;
+
+	for (UObject* Object : InObjects)
+	{
+		if(UAnimNextGraph* AnimNextGraph = Cast<UAnimNextGraph>(Object))
+		{
+			TSharedRef<FGraphEditor> GraphEditor = MakeShared<FGraphEditor>();
+			GraphEditor->InitEditor(Mode, EditWithinLevelEditor, AnimNextGraph);
+		}
+	}
+}
+
+const TArray<FText>& FAssetTypeActions_AnimNextGraph::GetSubMenus() const
+{
+	static const TArray<FText> SubMenus
+	{
+		LOCTEXT("AnimNextSubMenu", "AnimNext")
+	};
+	return SubMenus;
+}
+
+}
+
+#undef LOCTEXT_NAMESPACE
\ No newline at end of file
diff --git a/Engine/Plugins/Experimental/AnimNext/Source/AnimNextGraphEditor/Private/AssetTypeActions.h b/Engine/Plugins/Experimental/AnimNext/Source/AnimNextGraphEditor/Private/AssetTypeActions.h
new file mode 100644
--- /dev/null
+++ b/Engine/Plugins/Experimental/AnimNext/Source/AnimNextGraphEditor/Private/AssetTypeActions.h
@@ -0,0 +1,24 @@
+﻿// Copyright Epic Games, Inc. All Rights Reserved.
+
+#pragma once
+
+#include "AssetTypeActions_Base.h"
+#include "AnimNextGraph.h"
+
+namespace UE::AnimNext::GraphEditor
+{
+
+class FAssetTypeActions_AnimNextGraph : public FAssetTypeActions_Base
+{
+public:
+	// IAssetTypeActions Implementation
+	virtual FText GetName() const override { return NSLOCTEXT("AssetTypeActions", "AssetTypeActions_AnimNextGraph", "AnimNext Graph"); }
+	virtual FColor GetTypeColor() const override { return FColor(128,128,64); }
+	virtual UClass* GetSupportedClass() const override { return UAnimNextGraph::StaticClass(); }
+	virtual void OpenAssetEditor(const TArray<UObject*>& InObjects, TSharedPtr<IToolkitHost> EditWithinLevelEditor = TSharedPtr<IToolkitHost>()) override;
+	virtual bool CanFilter() override { return true; }
+	virtual const TArray<FText>& GetSubMenus() const override;
+	virtual uint32 GetCategories() override { return EAssetTypeCategories::Animation; }
+};
+
+}
diff --git a/Engine/Plugins/Experimental/AnimNext/Source/AnimNextGraphEditor/Private/GraphEditorMode.cpp b/Engine/Plugins/Experimental/AnimNext/Source/AnimNextGraphEditor/Private/GraphEditorMode.cpp
new file mode 100644
--- /dev/null
+++ b/Engine/Plugins/Experimental/AnimNext/Source/AnimNextGraphEditor/Private/GraphEditorMode.cpp
@@ -0,0 +1,273 @@
+﻿// Copyright Epic Games, Inc. All Rights Reserved.
+
+#include "GraphEditorMode.h"
+#include "AnimNextGraphEditor.h"
+#include "PropertyEditorModule.h"
+#include "Widgets/Docking/SDockTab.h"
+#include "EdGraph/EdGraph.h"
+#include "EdGraph/EdGraphSchema.h"
+#include "Widgets/Layout/SSpacer.h"
+#include "Modules/ModuleManager.h"
+
+#define LOCTEXT_NAMESPACE "GraphEditorMode"
+
+namespace UE::AnimNext::GraphEditor
+{
+
+struct FGraphTabHistory : public FGenericTabHistory
+{
+public:
+	FGraphTabHistory(TSharedPtr<FDocumentTabFactory> InFactory, TSharedPtr<FTabPayload> InPayload)
+		: FGenericTabHistory(InFactory, InPayload)
+		, SavedLocation(FVector2D::ZeroVector)
+		, SavedZoomAmount(INDEX_NONE)
+	{
+	}
+
+private:
+	// FGenericTabHistory interface
+	virtual void EvokeHistory(TSharedPtr<::FTabInfo> InTabInfo, bool bPrevTabMatches) override
+	{
+		FWorkflowTabSpawnInfo SpawnInfo;
+		SpawnInfo.Payload = Payload;
+		SpawnInfo.TabInfo = InTabInfo;
+
+		if(bPrevTabMatches)
+		{
+			TSharedPtr<SDockTab> DockTab = InTabInfo->GetTab().Pin();
+			GraphEditor = StaticCastSharedRef<SGraphEditor>(DockTab->GetContent());
+		}
+		else
+		{
+			TSharedRef<SGraphEditor> GraphEditorRef = StaticCastSharedRef<SGraphEditor>(FactoryPtr.Pin()->CreateTabBody(SpawnInfo));
+			GraphEditor = GraphEditorRef;
+			FactoryPtr.Pin()->UpdateTab(InTabInfo->GetTab().Pin(), SpawnInfo, GraphEditorRef);
+		}
+	}
+	
+	virtual void SaveHistory() override
+	{
+		if (IsHistoryValid())
+		{
+			check(GraphEditor.IsValid());
+			GraphEditor.Pin()->GetViewLocation(SavedLocation, SavedZoomAmount);
+			GraphEditor.Pin()->GetViewBookmark(SavedBookmarkId);
+		}
+	}
+	
+	virtual void RestoreHistory() override
+	{
+		if (IsHistoryValid())
+		{
+			check(GraphEditor.IsValid());
+			GraphEditor.Pin()->SetViewLocation(SavedLocation, SavedZoomAmount, SavedBookmarkId);
+		}
+	}
+	
+	/** The graph editor represented by this history node. While this node is inactive, the graph editor is invalid */
+	TWeakPtr<SGraphEditor> GraphEditor;
+	
+	/** Saved location the graph editor was at when this history node was last visited */
+	FVector2D SavedLocation;
+	
+	/** Saved zoom the graph editor was at when this history node was last visited */
+	float SavedZoomAmount;
+	
+	/** Saved bookmark ID the graph editor was at when this history node was last visited */
+	FGuid SavedBookmarkId;
+};
+
+FGraphEditorSummoner::FGraphEditorSummoner(TSharedPtr<FWorkflowCentricApplication> InHostingApp)
+	: FDocumentTabFactoryForObjects<UEdGraph>(Tabs::Document, InHostingApp)
+{
+}
+
+void FGraphEditorSummoner::OnTabActivated(TSharedPtr<SDockTab> Tab) const
+{
+	TSharedRef<SGraphEditor> GraphEditor = StaticCastSharedRef<SGraphEditor>(Tab->GetContent());
+	OnGraphEditorFocusedDelegate.ExecuteIfBound(GraphEditor);
+}
+
+void FGraphEditorSummoner::OnTabBackgrounded(TSharedPtr<SDockTab> Tab) const
+{
+	TSharedRef<SGraphEditor> GraphEditor = StaticCastSharedRef<SGraphEditor>(Tab->GetContent());
+	OnGraphEditorBackgroundedDelegate.ExecuteIfBound(GraphEditor);
+}
+
+void FGraphEditorSummoner::OnTabRefreshed(TSharedPtr<SDockTab> Tab) const
+{
+	TSharedRef<SGraphEditor> GraphEditor = StaticCastSharedRef<SGraphEditor>(Tab->GetContent());
+	GraphEditor->NotifyGraphChanged();
+}
+
+void FGraphEditorSummoner::SaveState(TSharedPtr<SDockTab> Tab, TSharedPtr<FTabPayload> Payload) const
+{
+	TSharedRef<SGraphEditor> GraphEditor = StaticCastSharedRef<SGraphEditor>(Tab->GetContent());
+
+	FVector2D ViewLocation;
+	float ZoomAmount;
+	GraphEditor->GetViewLocation(ViewLocation, ZoomAmount);
+
+	UEdGraph* Graph = Payload->IsValid() ? FTabPayload_UObject::CastChecked<UEdGraph>(Payload) : nullptr;
+
+	OnSaveGraphStateDelegate.ExecuteIfBound(Graph, ViewLocation, ZoomAmount);
+}
+
+TAttribute<FText> FGraphEditorSummoner::ConstructTabNameForObject(UEdGraph* DocumentID) const
+{
+	if (DocumentID)
+	{
+		if (const UEdGraphSchema* Schema = DocumentID->GetSchema())
+		{
+			FGraphDisplayInfo Info;
+			Schema->GetGraphDisplayInformation(*DocumentID, /*out*/ Info);
+			return Info.DisplayName;
+		}
+		else
+		{
+			// if we don't have a schema, we're dealing with a malformed (or incomplete graph)...
+			// possibly in the midst of some transaction - here we return the object's outer path 
+			// so we can at least get some context as to which graph we're referring
+			return FText::FromString(DocumentID->GetPathName());
+		}
+	}
+
+	return LOCTEXT("UnknownGraphName", "Unknown");
+}
+
+TSharedRef<SWidget> FGraphEditorSummoner::CreateTabBodyForObject(const FWorkflowTabSpawnInfo& Info, UEdGraph* DocumentID) const
+{
+	check(Info.TabInfo.IsValid());
+	
+	if(OnCreateGraphEditorWidgetDelegate.IsBound())
+	{
+		return OnCreateGraphEditorWidgetDelegate.Execute(Info.TabInfo.ToSharedRef(), DocumentID);
+	}
+	else
+	{
+		return SNew(SSpacer);
+	}
+}
+
+const FSlateBrush* FGraphEditorSummoner::GetTabIconForObject(const FWorkflowTabSpawnInfo& Info, UEdGraph* DocumentID) const
+{
+	if(OnGetTabIconDelegate.IsBound())
+	{
+		return OnGetTabIconDelegate.Execute(DocumentID);
+	}
+	else
+	{
+		return FAppStyle::Get().GetBrush(TEXT("GraphEditor.EventGraph_16x"));
+	}
+}
+
+TSharedRef<FGenericTabHistory> FGraphEditorSummoner::CreateTabHistoryNode(TSharedPtr<FTabPayload> Payload)
+{
+	return MakeShared<FGraphTabHistory>(SharedThis(this), Payload);
+}
+
+class FDetailsTabSummoner : public FWorkflowTabFactory
+{
+public:
+	FDetailsTabSummoner(TSharedPtr<FAssetEditorToolkit> InHostingApp, FOnDetailsViewCreated InOnDetailsViewCreated)
+		: FWorkflowTabFactory(Tabs::Details, InHostingApp)
+	{
+		TabLabel = LOCTEXT("DetailsTabLabel", "Details");
+		TabIcon = FSlateIcon("EditorStyle", "LevelEditor.Tabs.Details");
+		ViewMenuDescription = LOCTEXT("DetailsTabMenuDescription", "Details");
+		ViewMenuTooltip = LOCTEXT("DetailsTabToolTip", "Shows the details tab for selected objects.");
+		bIsSingleton = true;
+
+		FDetailsViewArgs DetailsViewArgs;
+		DetailsViewArgs.NameAreaSettings = FDetailsViewArgs::HideNameArea;
+		DetailsViewArgs.bHideSelectionTip = true;
+
+		FPropertyEditorModule& PropertyEditorModule = FModuleManager::LoadModuleChecked<FPropertyEditorModule>("PropertyEditor");
+		DetailsView = PropertyEditorModule.CreateDetailView(DetailsViewArgs);
+		
+		InOnDetailsViewCreated.ExecuteIfBound(DetailsView.ToSharedRef());
+	}
+
+	TSharedPtr<IDetailsView> GetDetailsView() const
+	{
+		return DetailsView;
+	}
+
+private:
+	virtual TSharedRef<SWidget> CreateTabBody(const FWorkflowTabSpawnInfo& Info) const override
+	{
+		return DetailsView.ToSharedRef();
+	}
+	
+	virtual FText GetTabToolTipText(const FWorkflowTabSpawnInfo& Info) const override
+	{
+		return ViewMenuTooltip;
+	}
+	
+	TSharedPtr<IDetailsView> DetailsView;
+};
+
+FGraphEditorMode::FGraphEditorMode(TSharedRef<FWorkflowCentricApplication> InHostingApp)
+	: FApplicationMode(Modes::GraphEditor)
+	, HostingAppPtr(InHostingApp)
+{
+	HostingAppPtr = InHostingApp;
+
+	TSharedRef<FGraphEditor> GraphEditor = StaticCastSharedRef<FGraphEditor>(InHostingApp);
+	
+	TabFactories.RegisterFactory(MakeShared<FDetailsTabSummoner>(GraphEditor, FOnDetailsViewCreated::CreateSP(&GraphEditor.Get(), &FGraphEditor::HandleDetailsViewCreated)));
+
+	TabLayout = FTabManager::NewLayout("Standalone_AnimNextInterfaceGraphEditor_Layout_v1.0")
+		->AddArea
+		(
+			FTabManager::NewPrimaryArea()
+			->SetOrientation(Orient_Vertical)
+			->Split
+			(	
+				FTabManager::NewSplitter()
+				->SetSizeCoefficient(1.0f)
+				->SetOrientation(Orient_Horizontal)
+				->Split
+				(
+					FTabManager::NewStack()
+					->SetSizeCoefficient(0.8f)
+					->SetHideTabWell(false)
+					->AddTab(Tabs::Document, ETabState::ClosedTab)
+				)
+				->Split
+				(
+					FTabManager::NewStack()
+					->SetSizeCoefficient(0.2f)
+					->SetHideTabWell(false)
+					->AddTab(Tabs::Details, ETabState::OpenedTab)
+				)
+			)
+		);
+
+	GraphEditor->RegisterModeToolbarIfUnregistered(GetModeName());
+}
+
+void FGraphEditorMode::RegisterTabFactories(TSharedPtr<FTabManager> InTabManager)
+{
+	TSharedPtr<FWorkflowCentricApplication> HostingApp = HostingAppPtr.Pin();
+	HostingApp->PushTabFactories(TabFactories);
+
+	FApplicationMode::RegisterTabFactories(InTabManager);
+}
+
+void FGraphEditorMode::AddTabFactory(FCreateWorkflowTabFactory FactoryCreator)
+{
+	if (FactoryCreator.IsBound())
+	{
+		TabFactories.RegisterFactory(FactoryCreator.Execute(HostingAppPtr.Pin()));
+	}
+}
+
+void FGraphEditorMode::RemoveTabFactory(FName TabFactoryID)
+{
+	TabFactories.UnregisterFactory(TabFactoryID);
+}
+
+}
+
+#undef LOCTEXT_NAMESPACE
\ No newline at end of file
diff --git a/Engine/Plugins/Experimental/AnimNext/Source/AnimNextGraphEditor/Private/GraphEditorMode.h b/Engine/Plugins/Experimental/AnimNext/Source/AnimNextGraphEditor/Private/GraphEditorMode.h
new file mode 100644
--- /dev/null
+++ b/Engine/Plugins/Experimental/AnimNext/Source/AnimNextGraphEditor/Private/GraphEditorMode.h
@@ -0,0 +1,96 @@
+﻿// Copyright Epic Games, Inc. All Rights Reserved.
+
+#pragma once
+
+#include "CoreMinimal.h"
+#include "EdGraph/EdGraph.h"
+#include "WorkflowOrientedApp/ApplicationMode.h"
+#include "WorkflowOrientedApp/WorkflowTabManager.h"
+#include "WorkflowOrientedApp/WorkflowUObjectDocuments.h"
+
+class FWorkflowCentricApplication;
+class SGraphEditor;
+
+namespace UE::AnimNext::GraphEditor
+{
+
+DECLARE_DELEGATE_OneParam(FOnDetailsViewCreated, TSharedRef<IDetailsView>);
+
+struct FGraphEditorSummoner : public FDocumentTabFactoryForObjects<UEdGraph>
+{
+public:
+	// Delegate called when a graph editor widget is created
+	DECLARE_DELEGATE_RetVal_TwoParams(TSharedRef<SGraphEditor>, FOnCreateGraphEditorWidget, TSharedRef<FTabInfo>, UEdGraph*);
+
+	// Delegate called when a graph editor is focused
+	DECLARE_DELEGATE_OneParam(FOnGraphEditorFocused, TSharedRef<SGraphEditor>);
+
+	// Delegate called when a graph editor is backgrounded
+	DECLARE_DELEGATE_OneParam(FOnGraphEditorBackgrounded, TSharedRef<SGraphEditor>);
+
+	// Delegate called to save the state of a graph
+	DECLARE_DELEGATE_ThreeParams(FOnSaveGraphState, UEdGraph*, FVector2D, float);
+
+	// Delegate called to save the state of a graph
+	DECLARE_DELEGATE_RetVal_OneParam(const FSlateBrush*, FOnGetTabIcon, UEdGraph*);
+	
+	FGraphEditorSummoner(TSharedPtr<FWorkflowCentricApplication> InHostingApp);
+
+	FOnCreateGraphEditorWidget& OnCreateGraphEditorWidget() { return OnCreateGraphEditorWidgetDelegate; }
+	FOnGraphEditorFocused& OnGraphEditorFocused() { return OnGraphEditorFocusedDelegate; }
+	FOnGraphEditorBackgrounded& OnGraphEditorBackgrounded() { return OnGraphEditorBackgroundedDelegate; }
+	FOnSaveGraphState& OnSaveGraphState() { return OnSaveGraphStateDelegate; }
+	FOnGetTabIcon& OnGetTabIcon() { return OnGetTabIconDelegate; }
+
+private:
+	// FWorkflowTabFactory interface
+	virtual void OnTabActivated(TSharedPtr<SDockTab> Tab) const override;
+	virtual void OnTabBackgrounded(TSharedPtr<SDockTab> Tab) const override;
+	virtual void OnTabRefreshed(TSharedPtr<SDockTab> Tab) const override;
+	virtual void SaveState(TSharedPtr<SDockTab> Tab, TSharedPtr<FTabPayload> Payload) const override;
+	
+	// FDocumentTabFactoryForObjects interface
+	virtual TAttribute<FText> ConstructTabNameForObject(UEdGraph* DocumentID) const override;
+	virtual TSharedRef<SWidget> CreateTabBodyForObject(const FWorkflowTabSpawnInfo& Info, UEdGraph* DocumentID) const override;
+	virtual const FSlateBrush* GetTabIconForObject(const FWorkflowTabSpawnInfo& Info, UEdGraph* DocumentID) const override;
+	virtual TSharedRef<FGenericTabHistory> CreateTabHistoryNode(TSharedPtr<FTabPayload> Payload) override;
+
+	// The hosting app
+	TWeakPtr<FWorkflowCentricApplication> HostingAppPtr;
+
+	// Delegate called when a graph editor widget is created
+	FOnCreateGraphEditorWidget OnCreateGraphEditorWidgetDelegate;
+
+	// Delegate called when a graph editor is focused
+	FOnGraphEditorFocused OnGraphEditorFocusedDelegate;
+
+	// Delegate called when a graph editor is backgrounded
+	FOnGraphEditorBackgrounded OnGraphEditorBackgroundedDelegate;
+
+	// Delegate called to save the state of a graph
+	FOnSaveGraphState OnSaveGraphStateDelegate;
+
+	// Delegate called to get a tab icon for a graph
+	FOnGetTabIcon OnGetTabIconDelegate;
+};
+
+class FGraphEditorMode : public FApplicationMode
+{
+public:
+	FGraphEditorMode(TSharedRef<FWorkflowCentricApplication> InHostingApp);
+
+private:
+	// FApplicationMode interface
+	virtual void RegisterTabFactories(TSharedPtr<FTabManager> InTabManager) override;
+	virtual void AddTabFactory(FCreateWorkflowTabFactory FactoryCreator) override;
+	virtual void RemoveTabFactory(FName TabFactoryID) override;
+	
+	// The hosting app
+	TWeakPtr<FWorkflowCentricApplication> HostingAppPtr;
+
+	// The tab factories we support
+	FWorkflowAllowedTabSet TabFactories;
+};
+
+}
+
diff --git a/Engine/Plugins/Experimental/AnimNext/Source/AnimNextGraphEditor/Private/GraphEditorModule.cpp b/Engine/Plugins/Experimental/AnimNext/Source/AnimNextGraphEditor/Private/GraphEditorModule.cpp
new file mode 100644
--- /dev/null
+++ b/Engine/Plugins/Experimental/AnimNext/Source/AnimNextGraphEditor/Private/GraphEditorModule.cpp
@@ -0,0 +1,42 @@
+// Copyright Epic Games, Inc. All Rights Reserved.
+
+#include "GraphEditorModule.h"
+#include "IAssetTools.h"
+#include "AssetTypeActions.h"
+#include "PropertyTypeCustomization.h"
+
+namespace UE::AnimNext::GraphEditor
+{
+
+void FModule::StartupModule()
+{
+	IAssetTools& AssetTools = FModuleManager::LoadModuleChecked<FAssetToolsModule>("AssetTools").Get();
+	AssetTypeActions_AnimNextGraph = MakeShared<FAssetTypeActions_AnimNextGraph>();
+	AssetTools.RegisterAssetTypeActions(AssetTypeActions_AnimNextGraph.ToSharedRef());
+
+	AnimNextPropertyTypeIdentifier = MakeShared<FPropertyTypeIdentifier>();
+	FPropertyEditorModule& PropertyModule = FModuleManager::LoadModuleChecked<FPropertyEditorModule>("PropertyEditor");
+	PropertyModule.RegisterCustomPropertyTypeLayout(
+		FInterfaceProperty::StaticClass()->GetFName(),
+		FOnGetPropertyTypeCustomizationInstance::CreateLambda([] { return MakeShared<FPropertyTypeCustomization>(); }),
+		AnimNextPropertyTypeIdentifier);
+}
+
+void FModule::ShutdownModule()
+{
+	if(FModuleManager::Get().IsModuleLoaded("AssetTools"))
+	{
+		IAssetTools& AssetTools = FModuleManager::GetModuleChecked<FAssetToolsModule>("AssetTools").Get();
+		AssetTools.UnregisterAssetTypeActions(AssetTypeActions_AnimNextGraph.ToSharedRef());
+	}
+	
+	if(FModuleManager::Get().IsModuleLoaded("PropertyEditor"))
+	{
+		FPropertyEditorModule& PropertyModule = FModuleManager::LoadModuleChecked<FPropertyEditorModule>("PropertyEditor");
+		PropertyModule.UnregisterCustomPropertyTypeLayout(FInterfaceProperty::StaticClass()->GetFName(), AnimNextPropertyTypeIdentifier);
+	}
+}
+
+}
+
+IMPLEMENT_MODULE(UE::AnimNext::GraphEditor::FModule, AnimNextGraphEditor);
\ No newline at end of file
diff --git a/Engine/Plugins/Experimental/AnimNext/Source/AnimNextGraphEditor/Private/GraphEditorModule.h b/Engine/Plugins/Experimental/AnimNext/Source/AnimNextGraphEditor/Private/GraphEditorModule.h
new file mode 100644
--- /dev/null
+++ b/Engine/Plugins/Experimental/AnimNext/Source/AnimNextGraphEditor/Private/GraphEditorModule.h
@@ -0,0 +1,25 @@
+﻿// Copyright Epic Games, Inc. All Rights Reserved.
+
+#pragma once
+
+#include "CoreMinimal.h"
+#include "Modules/ModuleManager.h"
+
+namespace UE::AnimNext::GraphEditor
+{
+
+class FAssetTypeActions_AnimNextGraph;
+class FAnimNextPropertyTypeCustomization;
+class FPropertyTypeIdentifier;
+
+class FModule : public IModuleInterface
+{
+private:
+	virtual void StartupModule() override;
+	virtual void ShutdownModule() override;
+
+	TSharedPtr<FAssetTypeActions_AnimNextGraph> AssetTypeActions_AnimNextGraph;
+	TSharedPtr<FPropertyTypeIdentifier> AnimNextPropertyTypeIdentifier;
+};
+
+}
\ No newline at end of file
diff --git a/Engine/Plugins/Experimental/AnimNext/Source/AnimNextGraphEditor/Private/GraphEditorSchemaActions.cpp b/Engine/Plugins/Experimental/AnimNext/Source/AnimNextGraphEditor/Private/GraphEditorSchemaActions.cpp
new file mode 100644
--- /dev/null
+++ b/Engine/Plugins/Experimental/AnimNext/Source/AnimNextGraphEditor/Private/GraphEditorSchemaActions.cpp
@@ -0,0 +1,61 @@
+﻿// Copyright Epic Games, Inc. All Rights Reserved.
+
+#include "GraphEditorSchemaActions.h"
+#include "GraphEditorUtils.h"
+#include "AnimNextGraph_EdGraph.h"
+#include "AnimNextGraph_EdGraphNode.h"
+#include "SGraphActionMenu.h"
+#include "UncookedOnlyUtils.h"
+#include "AnimNextGraph_EditorData.h"
+#include "RigVMModel/Nodes/RigVMUnitNode.h"
+#include "Settings/ControlRigSettings.h"
+
+UEdGraphNode* FAnimNextGraphSchemaAction_RigUnit::PerformAction(UEdGraph* ParentGraph, TArray<UEdGraphPin*>& FromPins, const FVector2D Location, bool bSelectNewNode)
+{
+	UAnimNextGraph_EditorData* EditorData = ParentGraph->GetTypedOuter<UAnimNextGraph_EditorData>();
+	UAnimNextGraph_EdGraphNode* NewNode = nullptr;
+	UAnimNextGraph_EdGraph* EdGraph = Cast<UAnimNextGraph_EdGraph>(ParentGraph);
+
+	if (EditorData != nullptr && EdGraph != nullptr)
+	{
+		FName Name = UE::AnimNext::GraphEditor::FUtils::ValidateName(EditorData, StructTemplate->GetFName().ToString());
+		URigVMController* Controller = EditorData->GetRigVMClient()->GetController(ParentGraph);
+
+		Controller->OpenUndoBracket(FString::Printf(TEXT("Add '%s' Node"), *Name.ToString()));
+
+		FRigVMUnitNodeCreatedContext& UnitNodeCreatedContext = Controller->GetUnitNodeCreatedContext();
+		FRigVMUnitNodeCreatedContext::FScope ReasonScope(UnitNodeCreatedContext, ERigVMNodeCreatedReason::NodeSpawner);
+
+		if (URigVMUnitNode* ModelNode = Controller->AddUnitNode(StructTemplate, FRigUnit::GetMethodName(), Location, Name.ToString(), true, false))
+		{
+			NewNode = Cast<UAnimNextGraph_EdGraphNode>(EdGraph->FindNodeForModelNodeName(ModelNode->GetFName()));
+			check(NewNode);
+
+			if (NewNode)
+			{
+				Controller->ClearNodeSelection(true);
+				Controller->SelectNode(ModelNode, true, true);
+			}
+
+#if WITH_EDITORONLY_DATA
+
+			// @TODO: switch over to custom settings here
+			const FControlRigSettingsPerPinBool* ExpansionMapPtr = UControlRigEditorSettings::Get()->RigUnitPinExpansion.Find(ModelNode->GetScriptStruct()->GetName());
+			if (ExpansionMapPtr)
+			{
+				const FControlRigSettingsPerPinBool& ExpansionMap = *ExpansionMapPtr;
+
+				for (const TPair<FString, bool>& Pair : ExpansionMap.Values)
+				{
+					FString PinPath = FString::Printf(TEXT("%s.%s"), *ModelNode->GetName(), *Pair.Key);
+					Controller->SetPinExpansion(PinPath, Pair.Value, true);
+				}
+			}
+		}
+#endif
+		
+		Controller->CloseUndoBracket();
+	}
+
+	return NewNode;
+}
diff --git a/Engine/Plugins/Experimental/AnimNext/Source/AnimNextGraphEditor/Private/GraphEditorSchemaActions.h b/Engine/Plugins/Experimental/AnimNext/Source/AnimNextGraphEditor/Private/GraphEditorSchemaActions.h
new file mode 100644
--- /dev/null
+++ b/Engine/Plugins/Experimental/AnimNext/Source/AnimNextGraphEditor/Private/GraphEditorSchemaActions.h
@@ -0,0 +1,55 @@
+﻿// Copyright Epic Games, Inc. All Rights Reserved.
+
+#pragma once
+
+#include "CoreMinimal.h"
+#include "EdGraph/EdGraphSchema.h"
+#include "Styling/AppStyle.h"
+#include "GraphEditorSchemaActions.generated.h"
+
+struct FSlateBrush;
+
+USTRUCT()
+struct FAnimNextGraphSchemaAction : public FEdGraphSchemaAction
+{
+	GENERATED_BODY()
+
+	FAnimNextGraphSchemaAction() = default;
+	
+	FAnimNextGraphSchemaAction(FText InNodeCategory, FText InMenuDesc, FText InToolTip, FText InKeywords = FText::GetEmpty())
+		: FEdGraphSchemaAction(MoveTemp(InNodeCategory), MoveTemp(InMenuDesc), MoveTemp(InToolTip), 0, MoveTemp(InKeywords))
+	{
+	}
+
+	virtual const FSlateBrush* GetIconBrush() const
+	{
+		return FAppStyle::Get().GetBrush("NoBrush");
+	}
+
+	virtual const FLinearColor& GetIconColor() const
+	{
+		static const FLinearColor DefaultColor;
+		return DefaultColor;
+	}
+};
+
+USTRUCT()
+struct FAnimNextGraphSchemaAction_RigUnit : public FAnimNextGraphSchemaAction
+{
+	GENERATED_BODY()
+
+	FAnimNextGraphSchemaAction_RigUnit() = default;
+	
+	FAnimNextGraphSchemaAction_RigUnit(UScriptStruct* InStructTemplate, FText InNodeCategory, FText InMenuDesc, FText InToolTip, FText InKeywords = FText::GetEmpty())
+		: FAnimNextGraphSchemaAction(InNodeCategory, InMenuDesc, InToolTip, InKeywords)
+		, StructTemplate(InStructTemplate)
+	{}
+
+	// FEdGraphSchemaAction Interface
+	virtual UEdGraphNode* PerformAction(UEdGraph* ParentGraph, TArray<UEdGraphPin*>& FromPins, const FVector2D Location, bool bSelectNewNode = true) override;
+	virtual UEdGraphNode* PerformAction(UEdGraph* ParentGraph, UEdGraphPin* FromPin, const FVector2D Location, bool bSelectNewNode = true) { return nullptr; }
+
+private:
+	// The script struct for our rig unit
+	UScriptStruct* StructTemplate = nullptr;
+};
diff --git a/Engine/Plugins/Experimental/AnimNext/Source/AnimNextGraphEditor/Private/GraphEditorUtils.cpp b/Engine/Plugins/Experimental/AnimNext/Source/AnimNextGraphEditor/Private/GraphEditorUtils.cpp
new file mode 100644
--- /dev/null
+++ b/Engine/Plugins/Experimental/AnimNext/Source/AnimNextGraphEditor/Private/GraphEditorUtils.cpp
@@ -0,0 +1,228 @@
+// Copyright Epic Games, Inc. All Rights Reserved.
+
+#include "GraphEditorUtils.h"
+#include "AnimNextGraph_EditorData.h"
+#include "Param/ParamType.h"
+#include "Kismet2/Kismet2NameValidators.h"
+
+namespace UE::AnimNext::GraphEditor
+{
+
+void FUtils::GetAllGraphNames(const UAnimNextGraph_EditorData* InEditorData, TSet<FName>& OutNames)
+{
+	// TODO
+	ensure(false);
+}
+
+static const int32 MaxNameLength = 100;
+
+FName FUtils::ValidateName(const UAnimNextGraph_EditorData* InEditorData, const FString& InName)
+{
+	struct FNameValidator : public INameValidatorInterface
+	{
+		FNameValidator(const UAnimNextGraph_EditorData* InEditorData)
+			: EditorData(InEditorData)
+		{
+			GetAllGraphNames(EditorData, Names);
+		}
+		
+		virtual EValidatorResult IsValid (const FName& Name, bool bOriginal = false) override
+		{
+			EValidatorResult ValidatorResult = EValidatorResult::AlreadyInUse;
+
+			if(Name == NAME_None)
+			{
+				ValidatorResult = EValidatorResult::EmptyName;
+			}
+			else if(Name.ToString().Len() > MaxNameLength)
+			{
+				ValidatorResult = EValidatorResult::TooLong;
+			}
+			else
+			{
+				// If it is in the names list then it is already in use.
+				if(!Names.Contains(Name))
+				{
+					ValidatorResult = EValidatorResult::Ok;
+
+					// Check for collision with an existing object.
+					if (UObject* ExistingObject = StaticFindObject(/*Class=*/ nullptr, const_cast<UAnimNextGraph_EditorData*>(EditorData), *Name.ToString(), true))
+					{
+						ValidatorResult = EValidatorResult::AlreadyInUse;
+					}
+				}
+			}
+			
+			return ValidatorResult;
+		}
+		
+		virtual EValidatorResult IsValid (const FString& Name, bool bOriginal = false) override
+		{
+			// Converting a string that is too large for an FName will cause an assert, so verify the length
+			if(Name.Len() >= NAME_SIZE)
+			{
+				return EValidatorResult::TooLong;
+			}
+			else if (!FName::IsValidXName(Name, UE_BLUEPRINT_INVALID_NAME_CHARACTERS))
+			{
+				return EValidatorResult::ContainsInvalidCharacters;
+			}
+
+			// If not defined in name table, not current graph name
+			return IsValid( FName(*Name) );
+		}
+
+		/** Name set to validate */
+		TSet<FName> Names;
+		/** The editor data to check for validity within */
+		const UAnimNextGraph_EditorData* EditorData;
+	};
+	
+	FString Name = InName;
+	if (Name.StartsWith(TEXT("RigUnit_")))
+	{
+		Name.RightChopInline(8, false);
+	}
+
+	FNameValidator NameValidator(InEditorData);
+
+	// Clean up BaseName to not contain any invalid characters, which will mean we can never find a legal name no matter how many numbers we add
+	if (NameValidator.IsValid(Name) == EValidatorResult::ContainsInvalidCharacters)
+	{
+		for (TCHAR& TestChar : Name)
+		{
+			for (TCHAR BadChar : UE_BLUEPRINT_INVALID_NAME_CHARACTERS)
+			{
+				if (TestChar == BadChar)
+				{
+					TestChar = TEXT('_');
+					break;
+				}
+			}
+		}
+	}
+	
+	int32 Count = 0;
+	FString BaseName = Name;
+	while (NameValidator.IsValid(Name) != EValidatorResult::Ok)
+	{
+		// Calculate the number of digits in the number, adding 2 (1 extra to correctly count digits, another to account for the '_' that will be added to the name
+		int32 CountLength = Count > 0 ? (int32)log((double)Count) + 2 : 2;
+
+		// If the length of the final string will be too long, cut off the end so we can fit the number
+		if (CountLength + BaseName.Len() > MaxNameLength)
+		{
+			BaseName.LeftInline(MaxNameLength - CountLength);
+		}
+		Name = FString::Printf(TEXT("%s_%d"), *BaseName, Count);
+		Count++;
+	}
+
+	return *Name;
+}
+
+FAnimNextParamType FUtils::GetParameterTypeFromMetaData(const FStringView& InStringView)
+{
+	if(InStringView.Equals(TEXT("bool")))
+	{
+		return FAnimNextParamType(FAnimNextParamType::EValueType::Bool);
+	}
+	else if(InStringView.Equals(TEXT("uint8")) || InStringView.Equals(TEXT("byte")))
+	{
+		return FAnimNextParamType(FAnimNextParamType::EValueType::Byte);
+	}
+	else if(InStringView.Equals(TEXT("int32")))
+	{
+		return FAnimNextParamType(FAnimNextParamType::EValueType::Int32);
+	}
+	else if(InStringView.Equals(TEXT("int64")))
+	{
+		return FAnimNextParamType(FAnimNextParamType::EValueType::Int64);
+	}
+	else if(InStringView.Equals(TEXT("float")))
+	{
+		return FAnimNextParamType(FAnimNextParamType::EValueType::Float);
+	}
+	else if(InStringView.Equals(TEXT("double")))
+	{
+		return FAnimNextParamType(FAnimNextParamType::EValueType::Double);
+	}
+	else if(InStringView.Equals(TEXT("Name")))
+	{
+		return FAnimNextParamType(FAnimNextParamType::EValueType::Name);
+	}
+	else if(InStringView.Equals(TEXT("String")))
+	{
+		return FAnimNextParamType(FAnimNextParamType::EValueType::String);
+	}
+	else if(InStringView.Equals(TEXT("Text")))
+	{
+		return FAnimNextParamType(FAnimNextParamType::EValueType::Text);
+	}
+	else
+	{
+		int32 SplitIndex = INDEX_NONE;
+		if(InStringView.FindChar(TEXT('\''), SplitIndex))
+		{
+			// Disambiguated by class type: \Path\To\Class'\Path\To\Object, so no need to resolve the object, just the class 
+			const FStringView ClassStringView = InStringView.Left(SplitIndex - 1);
+			const FStringView ObjectStringView = InStringView.RightChop(SplitIndex + 1);
+			
+			const FTopLevelAssetPath ClassPath(ClassStringView);
+			const FTopLevelAssetPath ObjectPath(ObjectStringView);
+
+			if(ClassPath.IsValid() && ObjectPath.IsValid())
+			{
+				const FSoftObjectPath SoftClassPath(ClassPath);
+				const FSoftObjectPath SoftObjectPath(ObjectPath);
+				if(UClass* ResolvedClass = Cast<UClass>(SoftClassPath.ResolveObject()))
+				{
+					if(ResolvedClass == UScriptStruct::StaticClass()) 
+					{
+						return FAnimNextParamType(FAnimNextParamType::EValueType::Struct, FAnimNextParamType::EContainerType::None, ResolvedClass);
+					}
+					else if(ResolvedClass == UEnum::StaticClass()) 
+					{
+						return FAnimNextParamType(FAnimNextParamType::EValueType::Enum, FAnimNextParamType::EContainerType::None, ResolvedClass);
+					}
+					else if(ResolvedClass == UClass::StaticClass()) 
+					{
+						return FAnimNextParamType(FAnimNextParamType::EValueType::Class, FAnimNextParamType::EContainerType::None, ResolvedClass);
+					}
+					else
+					{
+						return FAnimNextParamType(FAnimNextParamType::EValueType::Object, FAnimNextParamType::EContainerType::None, ResolvedClass);
+					}
+				}
+			}
+		}
+		else
+		{
+			// Class must be inferred: \Path\To\Object, so we need to resolve the object
+			const FTopLevelAssetPath ObjectPath(InStringView);
+			if(ObjectPath.IsValid())
+			{
+				const FSoftObjectPath SoftObjectPath(ObjectPath);
+				if(UObject* ResolvedObject = Cast<UObject>(SoftObjectPath.ResolveObject()))
+				{
+					if(UScriptStruct* ResolvedStruct = Cast<UScriptStruct>(ResolvedObject))
+					{
+						return FAnimNextParamType(FAnimNextParamType::EValueType::Struct, FAnimNextParamType::EContainerType::None, ResolvedStruct);
+					}
+					else if(UEnum* ResolvedEnum = Cast<UEnum>(ResolvedObject))
+					{
+						return FAnimNextParamType(FAnimNextParamType::EValueType::Enum, FAnimNextParamType::EContainerType::None, ResolvedEnum);
+					}
+					else if(UClass* ResolvedClass = Cast<UClass>(ResolvedObject))
+					{
+						return FAnimNextParamType(FAnimNextParamType::EValueType::Object, FAnimNextParamType::EContainerType::None, ResolvedClass);
+					}
+				}
+			}
+		}
+	}
+
+	return FAnimNextParamType(); 
+}
+
+}
diff --git a/Engine/Plugins/Experimental/AnimNext/Source/AnimNextGraphEditor/Private/GraphEditorUtils.h b/Engine/Plugins/Experimental/AnimNext/Source/AnimNextGraphEditor/Private/GraphEditorUtils.h
new file mode 100644
--- /dev/null
+++ b/Engine/Plugins/Experimental/AnimNext/Source/AnimNextGraphEditor/Private/GraphEditorUtils.h
@@ -0,0 +1,23 @@
+// Copyright Epic Games, Inc. All Rights Reserved.
+
+#pragma once
+
+#include "CoreMinimal.h"
+
+class UAnimNextGraph;
+class UAnimNextGraph_EditorData;
+struct FAnimNextParamType;
+
+namespace UE::AnimNext::GraphEditor
+{
+
+struct FUtils
+{
+	static FName ValidateName(const UAnimNextGraph_EditorData* InEditorData, const FString& InName);
+
+	static void GetAllGraphNames(const UAnimNextGraph_EditorData* InEditorData, TSet<FName>& OutNames);
+
+	static FAnimNextParamType GetParameterTypeFromMetaData(const FStringView& InStringView);
+};
+
+}
diff --git a/Engine/Plugins/Experimental/AnimNext/Source/AnimNextGraphEditor/Private/GraphFactory.cpp b/Engine/Plugins/Experimental/AnimNext/Source/AnimNextGraphEditor/Private/GraphFactory.cpp
new file mode 100644
--- /dev/null
+++ b/Engine/Plugins/Experimental/AnimNext/Source/AnimNextGraphEditor/Private/GraphFactory.cpp
@@ -0,0 +1,99 @@
+﻿// Copyright Epic Games, Inc. All Rights Reserved.
+
+#include "GraphFactory.h"
+#include "AnimNextGraph.h"
+#include "AnimNextGraph_EditorData.h"
+#include "AnimNextGraph_EdGraph.h"
+#include "AnimNextGraph_EdGraphSchema.h"
+#include "UncookedOnlyUtils.h"
+#include "RigUnit_AnimNextBeginExecution.h"
+#include "RigUnit_AnimNextEndExecution.h"
+#include "Kismet2/BlueprintEditorUtils.h"
+#include "RigVMModel/RigVMController.h"
+#include "Units/RigUnit.h"
+
+UAnimNextGraphFactory::UAnimNextGraphFactory()
+{
+	bCreateNew = true;
+	bEditAfterNew = true;
+	SupportedClass = UAnimNextGraph::StaticClass();
+}
+
+bool UAnimNextGraphFactory::ConfigureProperties()
+{
+	return true;
+}
+
+UObject* UAnimNextGraphFactory::FactoryCreateNew(UClass* Class, UObject* InParent, FName Name, EObjectFlags Flags, UObject* Context, FFeedbackContext* Warn, FName CallingContext)
+{
+	UAnimNextGraph* NewGraph = NewObject<UAnimNextGraph>(InParent, Class, Name, Flags | RF_Public | RF_Standalone | RF_Transactional | RF_LoadCompleted);
+
+	// Create internal editor data & graphs
+	UAnimNextGraph_EditorData* EditorData = NewObject<UAnimNextGraph_EditorData>(NewGraph, TEXT("EditorData"));
+	NewGraph->EditorData = EditorData;
+	EditorData->Initialize(/*bRecompileVM*/false);
+	EditorData->GetRigVMClient()->SetExecuteContextStruct(FAnimNextExecuteContext::StaticStruct());
+
+	// Add initial execution unit
+	URigVMController* Controller = EditorData->GetRigVMClient()->GetController(EditorData->RootGraph);
+
+	URigVMUnitNode* MainEntryPointNode = Controller->AddUnitNode(FRigUnit_AnimNextBeginExecution::StaticStruct(), FRigUnit::GetMethodName(), FVector2D(-400.0f, 0.0f), FString(), false);
+	URigVMPin* BeginExecutePin = MainEntryPointNode->FindPin(GET_MEMBER_NAME_STRING_CHECKED(FRigUnit_AnimNextBeginExecution, ExecuteContext));
+	check(BeginExecutePin);
+	check(BeginExecutePin->GetDirection() == ERigVMPinDirection::Output);
+
+	//// Add function to function lib
+	//{
+	//	FRigVMControllerGraphGuard LibraryGraphGuard(Controller, EditorData->GetRigVMClient()->GetFunctionLibrary());
+	//	EditorData->EntryPoint = Controller->AddFunctionToLibrary(UE::AnimNext::Graph::EntryPointName, true, FVector2D::ZeroVector, false, false);
+
+	//	// Add exposed pin for result to function
+	//	{
+	//		FRigVMControllerGraphGuard FunctionGraphGuard(Controller, EditorData->EntryPoint->GetContainedGraph());
+
+	//		// TODO: using float for now, but needs to be user-driven
+	//		const FString CPPType = TEXT("float");
+	//		const FName CPPTypeObjectPath = NAME_None;
+	//		const FString DefaultValue = TEXT("0.0");
+
+	//		//FunctionController->AddExposedPin(UE::AnimNext::Graph::ResultName, ERigVMPinDirection::Output, CPPType, CPPTypeObjectPath, DefaultValue, false, false);
+	//		Controller->AddExposedPin(UE::AnimNext::Graph::ResultName, ERigVMPinDirection::Output, CPPType, CPPTypeObjectPath, DefaultValue, false, false);
+	//	}
+	//}
+
+	//// Add function call
+	//URigVMFunctionReferenceNode* FunctionCallNode = Controller->AddFunctionReferenceNode(EditorData->EntryPoint, FVector2D::ZeroVector, UE::AnimNext::Graph::EntryPointName.ToString(), false, false);
+	//URigVMPin* FunctionExecutePin = FunctionCallNode->FindPin(FRigVMStruct::ExecuteContextName.ToString());
+	//check(FunctionExecutePin);
+	//check(FunctionExecutePin->GetDirection() == ERigVMPinDirection::IO);
+
+	//// Link entry point to function call
+	//Controller->AddLink(BeginExecutePin, FunctionExecutePin, false);
+
+	//URigVMPin* FunctionResultPin = FunctionCallNode->FindPin(UE::AnimNext::Graph::ResultName.ToString());
+	//check(FunctionResultPin);
+	//check(FunctionResultPin->GetDirection() == ERigVMPinDirection::Output);
+
+	//// Add end-execution unit of the correct type
+	//// TODO: using a float for now, but need to use a registry to determine correct type
+	//URigVMUnitNode* MainExitPointNode = Controller->AddUnitNode(FRigUnit_AnimNextInterfaceEndExecution_Float::StaticStruct(), FRigUnit::GetMethodName(), FVector2D(400.0f, 0.0f), FString(), false);
+	//URigVMPin* EndExecutePin = MainExitPointNode->FindPin(GET_MEMBER_NAME_STRING_CHECKED(FRigUnit_AnimNextEndExecution, ExecuteContext));
+	//check(EndExecutePin);
+	//check(EndExecutePin->GetDirection() == ERigVMPinDirection::Input);
+
+	//URigVMPin* EndExecuteResultPin = MainExitPointNode->FindPin(GET_MEMBER_NAME_STRING_CHECKED(FRigUnit_AnimNextEndExecution_Float, Result));
+	//check(EndExecuteResultPin);
+	//check(EndExecuteResultPin->GetDirection() == ERigVMPinDirection::Input);
+
+	//// Link function call to exit point 
+	//Controller->AddLink(FunctionExecutePin, EndExecutePin, false);
+
+	//// Link result pins 
+	//Controller->AddLink(FunctionResultPin, EndExecuteResultPin, false);
+
+	// Compile the initial skeleton
+	UE::AnimNext::GraphUncookedOnly::FUtils::Compile(NewGraph);
+	check(!EditorData->bErrorsDuringCompilation);
+	
+	return NewGraph;
+}
\ No newline at end of file
diff --git a/Engine/Plugins/Experimental/AnimNext/Source/AnimNextGraphEditor/Private/PropertyTypeCustomization.cpp b/Engine/Plugins/Experimental/AnimNext/Source/AnimNextGraphEditor/Private/PropertyTypeCustomization.cpp
new file mode 100644
--- /dev/null
+++ b/Engine/Plugins/Experimental/AnimNext/Source/AnimNextGraphEditor/Private/PropertyTypeCustomization.cpp
@@ -0,0 +1,92 @@
+﻿// Copyright Epic Games, Inc. All Rights Reserved.
+
+#include "PropertyTypeCustomization.h"
+
+#include "AnimNextGraph.h"
+#include "WidgetFactories.h"
+#include "DetailWidgetRow.h"
+#include "Editor/PropertyEditor/Private/ObjectPropertyNode.h"
+#include "Editor/PropertyEditor/Private/PropertyNode.h"
+#include "Interface/IAnimNextInterface.h"
+#include "IDetailChildrenBuilder.h"
+#include "PropertyCustomizationHelpers.h"
+#include "PropertyHandle.h"
+#include "SClassViewer.h"
+#include "Widgets/Text/STextBlock.h"
+#include "GraphEditorUtils.h"
+#include "Param/ParamTypeHandle.h"
+
+#define LOCTEXT_NAMESPACE "AnimNextPropertyTypeCustomization"
+
+namespace UE::AnimNext::GraphEditor
+{
+
+bool FPropertyTypeIdentifier::IsPropertyTypeCustomized(const IPropertyHandle& PropertyHandle) const
+{
+	return (PropertyHandle.GetMetaData("AnimNextType") != "");
+}
+
+void FPropertyTypeCustomization::CustomizeHeader(TSharedRef<IPropertyHandle> PropertyHandle, FDetailWidgetRow& HeaderRow, IPropertyTypeCustomizationUtils& CustomizationUtils)
+{
+	const FString TypeName = PropertyHandle->GetMetaData("AnimNextType");
+	const FAnimNextParamType AnimNextType = FUtils::GetParameterTypeFromMetaData(TypeName);
+	if(AnimNextType.IsValid())
+	{
+		void* ValuePtr;
+		PropertyHandle->GetValueData(ValuePtr);
+		FScriptInterface* PropertyValue = reinterpret_cast<FScriptInterface*>(ValuePtr);
+		
+		TSharedPtr<SWidget> Widget = FWidgetFactories::CreateAnimNextInterfaceWidget(AnimNextType.GetHandle(), PropertyValue->GetObject(),
+				FOnClassPicked::CreateLambda([this, PropertyHandle](UClass* ChosenClass)
+				{
+					TArray<void*> RawData;
+					PropertyHandle->AccessRawData(RawData);
+					TArray<UPackage*> OuterObjects;
+					PropertyHandle->GetOuterPackages(OuterObjects);
+
+					for(int i=0;i<RawData.Num();i++)
+					{
+						UObject* NewValue = NewObject<UObject>(OuterObjects[i], ChosenClass);
+					
+						PropertyHandle->NotifyPreChange();
+						PropertyHandle->SetValue(NewValue);
+						PropertyHandle->NotifyPostChange(EPropertyChangeType::ValueSet);
+						PropertyHandle->GetPropertyNode()->GetParentNode()->RequestRebuildChildren();
+					}
+				}),
+				nullptr
+		);
+
+		HeaderRow
+		.NameContent()
+		[
+			PropertyHandle->CreatePropertyNameWidget()
+		]
+		.ValueContent()
+		[
+			Widget.ToSharedRef()
+		];
+	}
+}
+
+void FPropertyTypeCustomization::CustomizeChildren(TSharedRef<IPropertyHandle> PropertyHandle, IDetailChildrenBuilder& ChildBuilder, IPropertyTypeCustomizationUtils& CustomizationUtils)
+{
+	// todo: make this work correctly for multiselect
+	void* ValuePtr;
+	PropertyHandle->GetValueData(ValuePtr);
+	FScriptInterface* PropertyValue = reinterpret_cast<FScriptInterface*>(ValuePtr);
+	if (PropertyValue->GetObject())
+	{
+		TSharedPtr<FStructOnScope> StructData = MakeShareable(new FStructOnScope(PropertyValue->GetObject()->GetClass(), (uint8*)PropertyValue->GetObject()));
+		StructData->SetPackage(PropertyValue->GetObject()->GetPackage());
+		TArray<TSharedPtr<IPropertyHandle>> NewHandles = ChildBuilder.AddAllExternalStructureProperties(StructData.ToSharedRef());
+		for (auto& NewHandle : NewHandles)
+		{
+			ChildBuilder.AddProperty(NewHandle.ToSharedRef());
+		}
+	}
+}
+
+}
+
+#undef LOCTEXT_NAMESPACE
\ No newline at end of file
diff --git a/Engine/Plugins/Experimental/AnimNext/Source/AnimNextGraphEditor/Private/PropertyTypeCustomization.h b/Engine/Plugins/Experimental/AnimNext/Source/AnimNextGraphEditor/Private/PropertyTypeCustomization.h
new file mode 100644
--- /dev/null
+++ b/Engine/Plugins/Experimental/AnimNext/Source/AnimNextGraphEditor/Private/PropertyTypeCustomization.h
@@ -0,0 +1,24 @@
+﻿// Copyright Epic Games, Inc. All Rights Reserved.
+
+#pragma once
+
+#include "IPropertyTypeCustomization.h"
+#include "PropertyEditorModule.h"
+
+namespace UE::AnimNext::GraphEditor
+{
+
+class FPropertyTypeCustomization : public IPropertyTypeCustomization
+{
+private:
+	// IPropertyTypeCustomization interface
+	virtual void CustomizeHeader(TSharedRef<IPropertyHandle> PropertyHandle, FDetailWidgetRow& HeaderRow, IPropertyTypeCustomizationUtils& CustomizationUtils) override;
+	virtual void CustomizeChildren(TSharedRef<IPropertyHandle> PropertyHandle, IDetailChildrenBuilder& ChildBuilder, IPropertyTypeCustomizationUtils& CustomizationUtils) override;
+};
+
+class FPropertyTypeIdentifier : public IPropertyTypeIdentifier
+{
+	virtual bool IsPropertyTypeCustomized(const IPropertyHandle& PropertyHandle) const override;
+};
+
+}
\ No newline at end of file
diff --git a/Engine/Plugins/Experimental/AnimNext/Source/AnimNextGraphEditor/Private/SActionMenu.cpp b/Engine/Plugins/Experimental/AnimNext/Source/AnimNextGraphEditor/Private/SActionMenu.cpp
new file mode 100644
--- /dev/null
+++ b/Engine/Plugins/Experimental/AnimNext/Source/AnimNextGraphEditor/Private/SActionMenu.cpp
@@ -0,0 +1,228 @@
+﻿// Copyright Epic Games, Inc. All Rights Reserved.
+
+#include "SActionMenu.h"
+
+#include "GraphEditorSchemaActions.h"
+#include "Framework/Application/SlateApplication.h"
+#include "IDocumentation.h"
+#include "SGraphActionMenu.h"
+#include "SSubobjectEditor.h"
+#include "AnimNextGraph_EdGraphSchema.h"
+#include "Widgets/Images/SImage.h"
+#include "Widgets/Layout/SBox.h"
+#include "Widgets/SBoxPanel.h"
+#include "Widgets/Text/STextBlock.h"
+#include "SGraphPalette.h"
+#include "RigVMCore/RigVMRegistry.h"
+#include "Units/RigUnit.h"
+#include "Widgets/SToolTip.h"
+#include "ExecuteContext.h"
+
+#define LOCTEXT_NAMESPACE "AnimNextEditor"
+
+namespace UE::AnimNext::GraphEditor
+{
+
+static void CollectAllAnimNextInterfaceActions(FGraphContextMenuBuilder& MenuBuilder)
+{
+	static const TArray<UScriptStruct*> AllowedExecuteContexts =
+	{
+		FRigVMExecuteContext::StaticStruct(),
+		FAnimNextExecuteContext::StaticStruct()
+	};
+
+	const TChunkedArray<FRigVMFunction>& Functions = FRigVMRegistry::Get().GetFunctions();
+	const int32 NumFunctions = Functions.Num();
+	for(int i=0; i < NumFunctions; i++)
+	{
+		const FRigVMFunction& Function = Functions[i];
+
+		const UScriptStruct* FunctionContext = Function.GetExecuteContextStruct();
+		if (FunctionContext == nullptr || !AllowedExecuteContexts.Contains(FunctionContext))
+		{
+			continue;
+		}
+
+		UScriptStruct* Struct = Function.Struct;
+		if (Struct == nullptr)
+		{
+			continue;
+		}
+
+		FString CategoryMetadata, DisplayNameMetadata, MenuDescSuffixMetadata;
+		Struct->GetStringMetaDataHierarchical(FRigVMStruct::CategoryMetaName, &CategoryMetadata);
+		Struct->GetStringMetaDataHierarchical(FRigVMStruct::DisplayNameMetaName, &DisplayNameMetadata);
+		Struct->GetStringMetaDataHierarchical(FRigVMStruct::MenuDescSuffixMetaName, &MenuDescSuffixMetadata);
+		if (!MenuDescSuffixMetadata.IsEmpty())
+		{
+			MenuDescSuffixMetadata = TEXT(" ") + MenuDescSuffixMetadata;
+		}
+		const FText NodeCategory = FText::FromString(CategoryMetadata);
+		const FText MenuDesc = FText::FromString(DisplayNameMetadata + MenuDescSuffixMetadata);
+		const FText ToolTip = Struct->GetToolTipText();
+
+		if (MenuDesc.IsEmpty())
+		{
+			continue;
+		}
+
+		MenuBuilder.AddAction(MakeShared<FAnimNextGraphSchemaAction_RigUnit>(Struct, NodeCategory, MenuDesc, ToolTip));
+	};
+}
+
+SActionMenu::~SActionMenu()
+{
+	OnClosedCallback.ExecuteIfBound();
+	OnCloseReasonCallback.ExecuteIfBound(bActionExecuted, false, !DraggedFromPins.IsEmpty());
+}
+
+void SActionMenu::Construct(const FArguments& InArgs)
+{
+	Graph = InArgs._Graph;
+	DraggedFromPins = InArgs._DraggedFromPins;
+	NewNodePosition = InArgs._NewNodePosition;
+	OnClosedCallback = InArgs._OnClosedCallback;
+	bAutoExpandActionMenu = InArgs._AutoExpandActionMenu;
+	OnCloseReasonCallback = InArgs._OnCloseReason;
+
+	SBorder::Construct(SBorder::FArguments()
+		.BorderImage(FAppStyle::Get().GetBrush("Menu.Background"))
+		.Padding(5)
+		[
+			SNew(SBox)
+			.WidthOverride(400)
+			.HeightOverride(400)
+			[
+				SNew(SVerticalBox)
+				+SVerticalBox::Slot()
+				.AutoHeight()
+				.Padding(2, 2, 2, 5)
+				[
+					SNew(SHorizontalBox)
+					+SHorizontalBox::Slot()
+					.AutoWidth()
+					.VAlign(VAlign_Center)
+					[
+						SNew(STextBlock)
+						.Text(LOCTEXT("ContextText", "All AnimNext Interface Node Classes"))
+						.Font(FAppStyle::Get().GetFontStyle("BlueprintEditor.ActionMenu.ContextDescriptionFont"))
+						.ToolTip(IDocumentation::Get()->CreateToolTip(
+							LOCTEXT("ActionMenuContextTextTooltip", "Describes the current context of the action list"),
+							nullptr,
+							TEXT("Shared/Editors/AnimNextInterfaceEditor"),
+							TEXT("AnimNextInterfaceActionMenuContextText")))
+						.WrapTextAt(280)
+					]
+				]
+				+SVerticalBox::Slot()
+				[
+					SAssignNew(GraphActionMenu, SGraphActionMenu)
+						.OnActionSelected(this, &SActionMenu::OnActionSelected)
+						.OnCreateWidgetForAction(SGraphActionMenu::FOnCreateWidgetForAction::CreateSP(this, &SActionMenu::OnCreateWidgetForAction))
+						.OnCollectAllActions(this, &SActionMenu::CollectAllActions)
+						.DraggedFromPins(DraggedFromPins)
+						.GraphObj(Graph)
+						.AlphaSortItems(true)
+						.bAllowPreselectedItemActivation(true)
+				]
+			]
+		]
+	);
+}
+
+void SActionMenu::CollectAllActions(FGraphActionListBuilderBase& OutAllActions)
+{
+	if (!Graph)
+	{
+		return;
+	}
+
+	FGraphContextMenuBuilder MenuBuilder(Graph);
+	if (!DraggedFromPins.IsEmpty())
+	{
+		MenuBuilder.FromPin = DraggedFromPins[0];
+	}
+
+	// Cannot call GetGraphContextActions() during serialization and GC due to its use of FindObject()
+	if(!GIsSavingPackage && !IsGarbageCollecting())
+	{
+		CollectAllAnimNextInterfaceActions(MenuBuilder);
+	}
+
+	OutAllActions.Append(MenuBuilder);
+}
+
+TSharedRef<SEditableTextBox> SActionMenu::GetFilterTextBox()
+{
+	return GraphActionMenu->GetFilterTextBox();
+}
+
+TSharedRef<SWidget> SActionMenu::OnCreateWidgetForAction(FCreateWidgetForActionData* const InCreateData)
+{
+	check(InCreateData);
+	InCreateData->bHandleMouseButtonDown = false;
+
+	const FSlateBrush* IconBrush = nullptr;
+	FLinearColor IconColor;
+	TSharedPtr<FAnimNextGraphSchemaAction> Action = StaticCastSharedPtr<FAnimNextGraphSchemaAction>(InCreateData->Action);
+	if (Action.IsValid())
+	{
+		IconBrush = Action->GetIconBrush();
+		IconColor = Action->GetIconColor();
+	}
+
+	TSharedPtr<SHorizontalBox> WidgetBox = SNew(SHorizontalBox);
+	if (IconBrush)
+	{
+		WidgetBox->AddSlot()
+			.AutoWidth()
+			.VAlign(VAlign_Center)
+			.Padding(5, 0, 0, 0)
+			[
+				SNew(SImage)
+				.ColorAndOpacity(IconColor)
+				.Image(IconBrush)
+			];
+	}
+
+	WidgetBox->AddSlot()
+		.AutoWidth()
+		.VAlign(VAlign_Center)
+		.Padding(0, 0, 0, 0)
+		[
+			SNew(SGraphPaletteItem, InCreateData)
+		];
+
+	return WidgetBox->AsShared();
+}
+
+void SActionMenu::OnActionSelected(const TArray<TSharedPtr<FEdGraphSchemaAction>>& SelectedAction, ESelectInfo::Type InSelectionType)
+{
+	if (!Graph)
+	{
+		return;
+	}
+
+	if (InSelectionType != ESelectInfo::OnMouseClick  && InSelectionType != ESelectInfo::OnKeyPress && !SelectedAction.IsEmpty())
+	{
+		return;
+	}
+
+	for (const TSharedPtr<FEdGraphSchemaAction>& Action : SelectedAction)
+	{
+		if (Action.IsValid() && Graph)
+		{
+			if (!bActionExecuted && (Action->GetTypeId() != FEdGraphSchemaAction_Dummy::StaticGetTypeId()))
+			{
+				FSlateApplication::Get().DismissAllMenus();
+				bActionExecuted = true;
+			}
+
+			Action->PerformAction(Graph, DraggedFromPins, NewNodePosition);
+		}
+	}
+}
+
+} 
+
+#undef LOCTEXT_NAMESPACE
