File/namespace reorg

diff --git a/Engine/Plugins/Experimental/AnimNext/Source/AnimNextGraphEditor/Private/WidgetFactories.cpp b/Engine/Plugins/Experimental/AnimNext/Source/AnimNextGraphEditor/Private/WidgetFactories.cpp
new file mode 100644
--- /dev/null
+++ b/Engine/Plugins/Experimental/AnimNext/Source/AnimNextGraphEditor/Private/WidgetFactories.cpp
@@ -0,0 +1,118 @@
+// Copyright Epic Games, Inc. All Rights Reserved.
+
+#include "WidgetFactories.h"
+#include "Modules/ModuleManager.h"
+#include "Widgets/Input/SComboButton.h"
+#include "SClassViewer.h"
+#include "ClassFilter.h"
+
+#define LOCTEXT_NAMESPACE "AnimNextEditor"
+
+namespace UE::AnimNext::GraphEditor
+{
+
+void ConvertToText(UObject* Object, FText& OutText)
+{
+	UClass* Class = Object->GetClass();
+	while (Class)
+	{
+		if (auto TextConverter = FWidgetFactories::AnimNextInterfaceTextConverter.Find(Class))
+		{
+			(*TextConverter)(Object, OutText);
+			break;
+		}
+		Class = Class->GetSuperClass();
+	}
+}
+
+TSharedPtr<SWidget> FWidgetFactories::CreateAnimNextInterfaceWidget(FParamTypeHandle TypeHandle, UObject* Value, const FOnClassPicked& CreateClassCallback, TSharedPtr<SBorder>* InnerWidget)
+{
+	TSharedPtr<SWidget> LeftWidget;
+
+	if (Value)
+	{
+		UClass* Class = Value->GetClass();
+		while (Class && !LeftWidget.IsValid())
+		{
+			if (auto Creator = FWidgetFactories::AnimNextInterfaceWidgetCreators.Find(Class))
+			{
+				LeftWidget = (*Creator)(Value);
+				break;
+			}
+			Class = Class->GetSuperClass();
+		}
+	}
+	
+	if (!LeftWidget.IsValid())
+	{
+		LeftWidget = SNew(STextBlock).Text(LOCTEXT("SelectDataType", "Select Data Type..."));
+	}
+
+
+	// button for replacing data with a different AnimNext Interface class
+	TSharedPtr<SComboButton> Button = SNew(SComboButton)
+			.ComboButtonStyle(FAppStyle::Get(), "SimpleComboButton");
+	
+	Button->SetOnGetMenuContent(FOnGetContent::CreateLambda([TypeHandle, Button, CreateClassCallback]()
+	{
+		FClassViewerInitializationOptions Options;
+		Options.ClassFilters.Add(MakeShared<FClassFilter>(TypeHandle));
+		
+		// Add class filter for columns here
+		TSharedRef<SWidget>  ClassMenu = FModuleManager::LoadModuleChecked<FClassViewerModule>("ClassViewer").
+		CreateClassViewer(Options, FOnClassPicked::CreateLambda([Button, CreateClassCallback](UClass* Class)
+		{
+			Button->SetIsOpen(false);
+			CreateClassCallback.Execute(Class);
+		}));
+		return ClassMenu;
+	}));
+
+	TSharedPtr <SBorder> Border;
+	if (InnerWidget && InnerWidget->IsValid())
+	{
+		Border = *InnerWidget;
+	}
+	else
+	{
+		Border = SNew(SBorder);
+	}
+	
+	if (InnerWidget)
+	{
+		*InnerWidget = Border;
+	}
+	
+	Border->SetContent(LeftWidget.ToSharedRef());
+
+	TSharedPtr<SWidget> Widget = SNew(SHorizontalBox)
+		+SHorizontalBox::Slot().FillWidth(75)
+		[
+			Border.ToSharedRef()
+		]
+		+SHorizontalBox::Slot().FillWidth(25)
+		[
+			Button.ToSharedRef()
+		]
+	;
+
+
+	return Widget;
+}
+
+TMap<const UClass*, TFunction<void (const UObject* Object, FText& OutText)>> FWidgetFactories::AnimNextInterfaceTextConverter;
+TMap<const UClass*, TFunction<TSharedRef<SWidget> (UObject* Object)>> FWidgetFactories::AnimNextInterfaceWidgetCreators;
+
+void ConvertToText_Base(const UObject* Object, FText& OutText)
+{
+	OutText = FText::FromString(Object->GetName());
+}
+
+void FWidgetFactories::RegisterWidgets()
+{
+	AnimNextInterfaceTextConverter.Add(UObject::StaticClass(), ConvertToText_Base);
+}
+	
+}
+
+#undef LOCTEXT_NAMESPACE
diff --git a/Engine/Plugins/Experimental/AnimNext/Source/AnimNextGraphEditor/Public/ClassFilter.h b/Engine/Plugins/Experimental/AnimNext/Source/AnimNextGraphEditor/Public/ClassFilter.h
new file mode 100644
--- /dev/null
+++ b/Engine/Plugins/Experimental/AnimNext/Source/AnimNextGraphEditor/Public/ClassFilter.h
@@ -0,0 +1,44 @@
+// Copyright Epic Games, Inc. All Rights Reserved.
+#pragma once
+
+#include "Param/ParamTypeHandle.h"
+#include "ClassViewerFilter.h"
+#include "Interface/IAnimNextInterface.h"
+
+namespace UE::AnimNext::GraphEditor
+{
+	// Filter class for things which implement IAnimNextInterface
+	class FClassFilter : public IClassViewerFilter
+	{
+	public:
+		FClassFilter(FParamTypeHandle InTypeHandle) : TypeHandle(InTypeHandle) {};
+		virtual ~FClassFilter() override {};
+	
+		virtual bool IsClassAllowed(const FClassViewerInitializationOptions& InInitOptions, const UClass* InClass, TSharedRef<FClassViewerFilterFuncs> InFilterFuncs) override
+		{
+			if (InClass->ImplementsInterface(UAnimNextInterface::StaticClass()))
+			{
+				// pass filter if the class has "" return type (for wrappers)
+				// pass all classes if the requested typename is ""
+				// otherwise fail if the TypeNames don't match
+				// for Object type, check type hierarchy. (todo)
+			
+				IAnimNextInterface* Object = static_cast<IAnimNextInterface*>(InClass->GetDefaultObject()->GetInterfaceAddress(UAnimNextInterface::StaticClass()));
+				if (Object->GetReturnTypeHandle() == TypeHandle || !Object->GetReturnTypeHandle().IsValid() || !TypeHandle.IsValid())
+				{
+					// @TODO: deal with inheritance directionality here
+					return true;
+				}
+			}
+
+			return false;
+		}
+
+		virtual bool IsUnloadedClassAllowed(const FClassViewerInitializationOptions& InInitOptions, const TSharedRef<const IUnloadedBlueprintData> InUnloadedClassData, TSharedRef<FClassViewerFilterFuncs> InFilterFuncs) override
+		{
+			return false;
+		}
+	private:
+		FParamTypeHandle TypeHandle;
+	};
+}
\ No newline at end of file
diff --git a/Engine/Plugins/Experimental/AnimNext/Source/AnimNextGraphEditor/Public/WidgetFactories.h b/Engine/Plugins/Experimental/AnimNext/Source/AnimNextGraphEditor/Public/WidgetFactories.h
new file mode 100644
--- /dev/null
+++ b/Engine/Plugins/Experimental/AnimNext/Source/AnimNextGraphEditor/Public/WidgetFactories.h
@@ -0,0 +1,27 @@
+// Copyright Epic Games, Inc. All Rights Reserved.
+
+#pragma once
+
+#include "CoreMinimal.h"
+#include "Widgets/SWidget.h"
+#include "Widgets/Layout/SBorder.h"
+#include "SClassViewer.h"
+
+namespace UE::AnimNext
+{
+	struct FParamTypeHandle;
+}
+
+namespace UE::AnimNext::GraphEditor
+{
+	class ANIMNEXTGRAPHEDITOR_API FWidgetFactories
+	{
+	public:
+		static TMap<const UClass*, TFunction<void (const UObject* Object, FText& OutText)>> AnimNextInterfaceTextConverter;
+		static TMap<const UClass*, TFunction<TSharedRef<SWidget> (UObject* Object)>> AnimNextInterfaceWidgetCreators;
+
+		static TSharedPtr<SWidget> CreateAnimNextInterfaceWidget(FParamTypeHandle TypeHandle, UObject* Value, const FOnClassPicked& CreateClassCallback, TSharedPtr<SBorder>* InnerWidget = nullptr);
+		
+		static void RegisterWidgets();
+	};
+}
diff --git a/Engine/Plugins/Experimental/AnimNext/Source/AnimNextGraphUncookedOnly/AnimNextGraphUncookedOnly.Build.cs b/Engine/Plugins/Experimental/AnimNext/Source/AnimNextGraphUncookedOnly/AnimNextGraphUncookedOnly.Build.cs
new file mode 100644
--- /dev/null
+++ b/Engine/Plugins/Experimental/AnimNext/Source/AnimNextGraphUncookedOnly/AnimNextGraphUncookedOnly.Build.cs
@@ -0,0 +1,49 @@
+// Copyright Epic Games, Inc. All Rights Reserved.
+
+namespace UnrealBuildTool.Rules
+{
+	public class AnimNextGraphUncookedOnly : ModuleRules
+	{
+		public AnimNextGraphUncookedOnly(ReadOnlyTargetRules Target) : base(Target)
+		{
+			
+			PublicDependencyModuleNames.AddRange(
+				new string[]
+				{
+					"RigVM",
+					"RigVMDeveloper",
+					"ControlRig",
+					"ControlRigDeveloper",
+				}
+			);
+
+			PrivateDependencyModuleNames.AddRange(
+				new string[]
+				{
+					"Core",
+					"CoreUObject",
+					"Engine",
+					"AnimNext",
+					"AnimNextGraph",
+					"BlueprintGraph",	// For K2 schema
+					"AnimationCore",
+					"AnimGraph",
+					"Kismet",
+					"Slate",
+					"SlateCore",
+				}
+			);
+
+
+			if (Target.bBuildEditor)
+			{
+				PrivateDependencyModuleNames.AddRange(
+					new string[]
+					{
+						"UnrealEd",
+					}
+				);
+			}
+		}
+	}
+}
\ No newline at end of file
diff --git a/Engine/Plugins/Experimental/AnimNext/Source/AnimNextGraphUncookedOnly/Internal/AnimNextGraph_EdGraph.h b/Engine/Plugins/Experimental/AnimNext/Source/AnimNextGraphUncookedOnly/Internal/AnimNextGraph_EdGraph.h
new file mode 100644
--- /dev/null
+++ b/Engine/Plugins/Experimental/AnimNext/Source/AnimNextGraphUncookedOnly/Internal/AnimNextGraph_EdGraph.h
@@ -0,0 +1,27 @@
+﻿// Copyright Epic Games, Inc. All Rights Reserved.
+
+#pragma once
+
+#include "CoreMinimal.h"
+#include "Graph/ControlRigGraph.h"
+#include "AnimNextGraph_EdGraph.generated.h"
+
+class UAnimNextGraph_EditorData;
+
+namespace UE::AnimNext::GraphUncookedOnly
+{
+	struct FUtils;
+}
+
+UCLASS(MinimalAPI)
+class UAnimNextGraph_EdGraph : public UControlRigGraph
+{
+	GENERATED_BODY()
+
+	friend class UAnimNextGraph_EditorData;
+
+	// UControlRigGraph interface
+	virtual FRigVMClient* GetRigVMClient() const override;
+	
+	void Initialize(UAnimNextGraph_EditorData* InEditorData);
+};
\ No newline at end of file
diff --git a/Engine/Plugins/Experimental/AnimNext/Source/AnimNextGraphUncookedOnly/Internal/AnimNextGraph_EdGraphNode.h b/Engine/Plugins/Experimental/AnimNext/Source/AnimNextGraphUncookedOnly/Internal/AnimNextGraph_EdGraphNode.h
new file mode 100644
--- /dev/null
+++ b/Engine/Plugins/Experimental/AnimNext/Source/AnimNextGraphUncookedOnly/Internal/AnimNextGraph_EdGraphNode.h
@@ -0,0 +1,13 @@
+﻿// Copyright Epic Games, Inc. All Rights Reserved.
+
+#pragma once
+
+#include "CoreMinimal.h"
+#include "Graph/ControlRigGraphNode.h"
+#include "AnimNextGraph_EdGraphNode.generated.h"
+
+UCLASS(MinimalAPI)
+class UAnimNextGraph_EdGraphNode : public UControlRigGraphNode
+{
+	GENERATED_BODY()
+};
\ No newline at end of file
diff --git a/Engine/Plugins/Experimental/AnimNext/Source/AnimNextGraphUncookedOnly/Internal/AnimNextGraph_EdGraphSchema.h b/Engine/Plugins/Experimental/AnimNext/Source/AnimNextGraphUncookedOnly/Internal/AnimNextGraph_EdGraphSchema.h
new file mode 100644
--- /dev/null
+++ b/Engine/Plugins/Experimental/AnimNext/Source/AnimNextGraphUncookedOnly/Internal/AnimNextGraph_EdGraphSchema.h
@@ -0,0 +1,20 @@
+﻿// Copyright Epic Games, Inc. All Rights Reserved.
+
+#pragma once
+
+#include "CoreMinimal.h"
+#include "Graph/ControlRigGraphSchema.h"
+#include "AnimNextGraph_EdGraphNode.h"
+#include "AnimNextGraph_EdGraphSchema.generated.h"
+
+UCLASS(MinimalAPI)
+class UAnimNextGraph_EdGraphSchema : public UControlRigGraphSchema
+{
+	GENERATED_BODY()
+
+	// UControlRigGraphSchema interface
+	virtual TSubclassOf<UControlRigGraphNode> GetGraphNodeClass() const override { return UAnimNextGraph_EdGraphNode::StaticClass(); }
+
+	// UEdGraphSchema interface
+	virtual void TrySetDefaultValue(UEdGraphPin& InPin, const FString& InNewDefaultValue, bool bMarkAsModified) const override;
+};
\ No newline at end of file
diff --git a/Engine/Plugins/Experimental/AnimNext/Source/AnimNextGraphUncookedOnly/Internal/AnimNextGraph_EditorData.h b/Engine/Plugins/Experimental/AnimNext/Source/AnimNextGraphUncookedOnly/Internal/AnimNextGraph_EditorData.h
new file mode 100644
--- /dev/null
+++ b/Engine/Plugins/Experimental/AnimNext/Source/AnimNextGraphUncookedOnly/Internal/AnimNextGraph_EditorData.h
@@ -0,0 +1,141 @@
+﻿// Copyright Epic Games, Inc. All Rights Reserved.
+
+#pragma once
+
+#include "CoreMinimal.h"
+#include "RigVMModel/RigVMGraph.h"
+#include "AnimNextGraph_EdGraph.h"
+#include "RigVMCore/RigVMGraphFunctionHost.h"
+#include "AnimNextGraph_EditorData.generated.h"
+
+class UAnimNextGraph;
+enum class ERigVMGraphNotifType : uint8;
+namespace UE::AnimNext::GraphUncookedOnly
+{
+	struct FUtils;
+}
+
+namespace UE::AnimNext::GraphEditor
+{
+	class FGraphEditor;
+}
+
+enum class EAnimNextGraphLoadType : uint8
+{
+	PostLoad,
+	CheckUserDefinedStructs
+};
+
+UCLASS(MinimalAPI)
+class UAnimNextGraph_EditorData : public UObject, public IRigVMClientHost, public IRigVMGraphFunctionHost
+{
+	GENERATED_BODY()
+
+	UAnimNextGraph_EditorData(const FObjectInitializer& ObjectInitializer);
+	
+	friend class UAnimNextGraphFactory;
+	friend class UAnimNextGraph_EdGraph;
+	friend struct UE::AnimNext::GraphUncookedOnly::FUtils;
+	friend class UE::AnimNext::GraphEditor::FGraphEditor;
+	friend struct FAnimNextGraphSchemaAction_RigUnit;
+
+	// UObject interface
+	virtual void Serialize(FArchive& Ar) override;
+	virtual void PostLoad() override;
+	virtual bool IsEditorOnly() const override { return true; }
+#if WITH_EDITOR
+	void HandlePackageDone(const FEndLoadPackageContext& Context);
+	void HandlePackageDone();
+#endif // WITH_EDITOR
+
+	// IRigVMClientHost interface
+	virtual FRigVMClient* GetRigVMClient() override;
+	virtual const FRigVMClient* GetRigVMClient() const override;
+	virtual IRigVMGraphFunctionHost* GetRigVMGraphFunctionHost() override;
+	virtual const IRigVMGraphFunctionHost* GetRigVMGraphFunctionHost() const override;
+	virtual void HandleRigVMGraphAdded(const FRigVMClient* InClient, const FString& InNodePath) override;
+	virtual void HandleRigVMGraphRemoved(const FRigVMClient* InClient, const FString& InNodePath) override {}
+	virtual void HandleRigVMGraphRenamed(const FRigVMClient* InClient, const FString& InOldNodePath, const FString& InNewNodePath) override {}
+	virtual void HandleConfigureRigVMController(const FRigVMClient* InClient, URigVMController* InControllerToConfigure) override;
+	virtual UObject* GetEditorObjectForRigVMGraph(URigVMGraph* InVMGraph) const override;
+
+	// IRigVMGraphFunctionHost interface
+	virtual FRigVMGraphFunctionStore* GetRigVMGraphFunctionStore() override;
+	virtual const FRigVMGraphFunctionStore* GetRigVMGraphFunctionStore() const override;
+
+	ANIMNEXTGRAPHUNCOOKEDONLY_API void Initialize(bool bRecompileVM);
+
+	void RefreshAllModels(EAnimNextGraphLoadType InLoadType);
+
+	void RecompileVM();
+	
+	void RecompileVMIfRequired();
+
+	void RequestAutoVMRecompilation();
+	
+	void HandleModifiedEvent(ERigVMGraphNotifType InNotifType, URigVMGraph* InGraph, UObject* InSubject);
+
+	ANIMNEXTGRAPHUNCOOKEDONLY_API URigVMGraph* GetVMGraphForEdGraph(const UEdGraph* InGraph) const;
+
+	void CreateEdGraphForCollapseNode(URigVMCollapseNode* InNode);
+
+	void UpdateGraphReturnType();
+
+	FEdGraphPinType FindGraphReturnPinType() const;
+
+	bool IsNodeExecConnected(const UEdGraphNode* Node) const;
+
+	UPROPERTY()
+	TObjectPtr<UAnimNextGraph_EdGraph> RootGraph;
+
+	UPROPERTY()
+	TObjectPtr<UAnimNextGraph_EdGraph> EntryPointGraph;
+
+	UPROPERTY()
+	TObjectPtr<UAnimNextGraph_EdGraph> FunctionLibraryEdGraph;
+
+	UPROPERTY()
+	FRigVMClient RigVMClient;
+
+	UPROPERTY()
+	FRigVMGraphFunctionStore GraphFunctionStore;
+
+	UPROPERTY()
+	TObjectPtr<URigVMLibraryNode> EntryPoint;
+
+	UPROPERTY(transient)
+	TMap<TObjectPtr<URigVMGraph>, TObjectPtr<URigVMController>> Controllers;
+
+	UPROPERTY(EditAnywhere, Category = "User Interface")
+	FRigGraphDisplaySettings RigGraphDisplaySettings;
+
+	UPROPERTY(EditAnywhere, Category = "VM")
+	FRigVMRuntimeSettings VMRuntimeSettings;
+	
+	UPROPERTY(EditAnywhere, BlueprintReadWrite, Category = "VM", meta = (AllowPrivateAccess = "true"))
+	FRigVMCompileSettings VMCompileSettings;
+
+	UPROPERTY(EditAnywhere, Category = "Python Log Settings")
+	FControlRigPythonSettings PythonLogSettings;
+
+	UPROPERTY(transient, DuplicateTransient)
+	TMap<FString, FRigVMOperand> PinToOperandMap;
+
+	UPROPERTY(transient, DuplicateTransient)
+	bool bVMRecompilationRequired = false;
+
+	UPROPERTY(transient, DuplicateTransient)
+	bool bIsCompiling = false;
+	
+	FCompilerResultsLog CompileLog;
+
+	FOnVMCompiledEvent VMCompiledEvent;
+	FRigVMGraphModifiedEvent ModifiedEvent;
+
+	bool bAutoRecompileVM = true;
+	bool bErrorsDuringCompilation = false;
+	bool bSuspendModelNotificationsForSelf = false;
+	bool bSuspendModelNotificationsForOthers = false;
+	bool bSuspendAllNotifications = false;
+	bool bCompileInDebugMode = false;
+};
diff --git a/Engine/Plugins/Experimental/AnimNext/Source/AnimNextGraphUncookedOnly/Internal/UncookedOnlyUtils.h b/Engine/Plugins/Experimental/AnimNext/Source/AnimNextGraphUncookedOnly/Internal/UncookedOnlyUtils.h
new file mode 100644
--- /dev/null
+++ b/Engine/Plugins/Experimental/AnimNext/Source/AnimNextGraphUncookedOnly/Internal/UncookedOnlyUtils.h
@@ -0,0 +1,33 @@
+﻿// Copyright Epic Games, Inc. All Rights Reserved.
+
+#pragma once
+
+#include "CoreMinimal.h"
+#include "Param/ParamTypeHandle.h"
+
+class UAnimNextGraph;
+class UAnimNextGraph_EditorData;
+class UAnimNextGraph_EdGraph;
+class URigVMController;
+class URigVMGraph;
+struct FEdGraphPinType;
+
+namespace UE::AnimNext::GraphUncookedOnly
+{
+
+struct ANIMNEXTGRAPHUNCOOKEDONLY_API FUtils
+{
+	static void Compile(UAnimNextGraph* InGraph);
+	
+	static UAnimNextGraph_EditorData* GetEditorData(const UAnimNextGraph* InAnimNextInterfaceGraph);
+	
+	static void RecreateVM(UAnimNextGraph* InGraph);
+
+	/**
+	 * Get an AnimNext parameter type handle from an FEdGraphPinType.
+	 * Note that the returned handle may not be valid, so should be checked using IsValid() before use.
+	 **/
+	static FParamTypeHandle GetParameterHandleFromPin(const FEdGraphPinType& InPinType);
+};
+
+}
\ No newline at end of file
diff --git a/Engine/Plugins/Experimental/AnimNext/Source/AnimNextGraphUncookedOnly/Private/AnimNextGraph_EdGraph.cpp b/Engine/Plugins/Experimental/AnimNext/Source/AnimNextGraphUncookedOnly/Private/AnimNextGraph_EdGraph.cpp
new file mode 100644
--- /dev/null
+++ b/Engine/Plugins/Experimental/AnimNext/Source/AnimNextGraphUncookedOnly/Private/AnimNextGraph_EdGraph.cpp
@@ -0,0 +1,21 @@
+﻿// Copyright Epic Games, Inc. All Rights Reserved.
+
+#include "AnimNextGraph_EdGraph.h"
+#include "AnimNextGraph_EditorData.h"
+
+void UAnimNextGraph_EdGraph::Initialize(UAnimNextGraph_EditorData* InEditorData)
+{
+	InEditorData->ModifiedEvent.RemoveAll(this);
+	InEditorData->ModifiedEvent.AddUObject(this, &UAnimNextGraph_EdGraph::HandleModifiedEvent);
+	InEditorData->VMCompiledEvent.RemoveAll(this);
+	InEditorData->VMCompiledEvent.AddUObject(this, &UAnimNextGraph_EdGraph::HandleVMCompiledEvent);
+}
+
+FRigVMClient* UAnimNextGraph_EdGraph::GetRigVMClient() const
+{
+	if (const UAnimNextGraph_EditorData* EditorData = GetTypedOuter<UAnimNextGraph_EditorData>())
+	{
+		return (FRigVMClient*)EditorData->GetRigVMClient();
+	}
+	return nullptr;
+}
diff --git a/Engine/Plugins/Experimental/AnimNext/Source/AnimNextGraphUncookedOnly/Private/AnimNextGraph_EdGraphSchema.cpp b/Engine/Plugins/Experimental/AnimNext/Source/AnimNextGraphUncookedOnly/Private/AnimNextGraph_EdGraphSchema.cpp
new file mode 100644
--- /dev/null
+++ b/Engine/Plugins/Experimental/AnimNext/Source/AnimNextGraphUncookedOnly/Private/AnimNextGraph_EdGraphSchema.cpp
@@ -0,0 +1,42 @@
+﻿// Copyright Epic Games, Inc. All Rights Reserved.
+
+#include "AnimNextGraph_EdGraphSchema.h"
+#include "Editor.h"
+
+void UAnimNextGraph_EdGraphSchema::TrySetDefaultValue(UEdGraphPin& InPin, const FString& InNewDefaultValue, bool bMarkAsModified) const
+{
+#if WITH_EDITOR
+	if (GEditor)
+	{
+		GEditor->CancelTransaction(0);
+	}
+#endif
+
+	FString UseDefaultValue;
+	TObjectPtr<UObject> UseDefaultObject = nullptr;
+	FText UseDefaultText;
+
+	GetDefault<UEdGraphSchema_K2>()->GetPinDefaultValuesFromString(InPin.PinType, InPin.GetOwningNodeUnchecked(), InNewDefaultValue, UseDefaultValue, UseDefaultObject, UseDefaultText, /*bPreserveTextIdentity*/false);
+
+	// Check the default value and make it an error if it's bogus
+	if (GetDefault<UEdGraphSchema_K2>()->DefaultValueSimpleValidation(InPin.PinType, InPin.PinName, UseDefaultValue, UseDefaultObject, UseDefaultText))
+	{
+		InPin.DefaultObject = UseDefaultObject;
+		InPin.DefaultValue = UseDefaultValue;
+		InPin.DefaultTextValue = UseDefaultText;
+
+		UEdGraphNode* Node = InPin.GetOwningNode();
+		Node->PinDefaultValueChanged(&InPin);
+
+		// If the default value is manually set then treat it as if the value was reset to default and remove the orphaned InPin
+		if (InPin.bOrphanedPin && InPin.DoesDefaultValueMatchAutogenerated())
+		{
+			Node->PinConnectionListChanged(&InPin);
+		}
+
+		if (bMarkAsModified)
+		{
+			Node->MarkPackageDirty();
+		}
+	}
+}
diff --git a/Engine/Plugins/Experimental/AnimNext/Source/AnimNextGraphUncookedOnly/Private/AnimNextGraph_EditorData.cpp b/Engine/Plugins/Experimental/AnimNext/Source/AnimNextGraphUncookedOnly/Private/AnimNextGraph_EditorData.cpp
new file mode 100644
--- /dev/null
+++ b/Engine/Plugins/Experimental/AnimNext/Source/AnimNextGraphUncookedOnly/Private/AnimNextGraph_EditorData.cpp
@@ -0,0 +1,553 @@
+﻿// Copyright Epic Games, Inc. All Rights Reserved.
+
+#include "AnimNextGraph_EditorData.h"
+
+#include "ControlRigDefines.h"
+#include "AnimNextGraph.h"
+#include "AnimNextGraph_EdGraphSchema.h"
+#include "UncookedOnlyUtils.h"
+#include "ExecuteContext.h"
+#include "Rigs/RigHierarchyPose.h"
+#include "RigVMModel/RigVMFunctionLibrary.h"
+#include "RigVMModel/RigVMNotifications.h"
+#include "RigVMModel/Nodes/RigVMCollapseNode.h"
+#include "Curves/CurveFloat.h"
+#include "UObject/ObjectSaveContext.h"
+#include "Kismet2/BlueprintEditorUtils.h"
+
+UAnimNextGraph_EditorData::UAnimNextGraph_EditorData(const FObjectInitializer& ObjectInitializer)
+{
+	RigVMClient.Reset();
+	RigVMClient.SetOuterClientHost(this, GET_MEMBER_NAME_CHECKED(UAnimNextGraph_EditorData, RigVMClient));
+	{
+		TGuardValue<bool> DisableClientNotifs(RigVMClient.bSuspendNotifications, true);
+		RigVMClient.AddModel(TEXT("RigVMGraph"), false, &ObjectInitializer);
+		RigVMClient.GetOrCreateFunctionLibrary(false, &ObjectInitializer);
+	}
+	RigVMClient.SetExecuteContextStruct(FAnimNextExecuteContext::StaticStruct());
+	
+	auto MakeEdGraph = [this, &ObjectInitializer](FName InName) -> UAnimNextGraph_EdGraph*
+	{
+		UAnimNextGraph_EdGraph* EdGraph = ObjectInitializer.CreateDefaultSubobject<UAnimNextGraph_EdGraph>(this, InName);
+		EdGraph->Schema = UAnimNextGraph_EdGraphSchema::StaticClass();
+		EdGraph->bAllowRenaming = false;
+		EdGraph->bEditable = true;
+		EdGraph->bAllowDeletion = false;
+		EdGraph->bIsFunctionDefinition = false;
+		EdGraph->Initialize(this);
+
+		return EdGraph;
+	};
+
+	RootGraph = MakeEdGraph(TEXT("RootEdGraph"));
+	FunctionLibraryEdGraph = MakeEdGraph(TEXT("RigVMFunctionLibraryEdGraph"));
+}
+
+void UAnimNextGraph_EditorData::Serialize(FArchive& Ar)
+{
+	RigVMClient.SetOuterClientHost(this, GET_MEMBER_NAME_CHECKED(UAnimNextGraph_EditorData, RigVMClient));
+
+	UObject::Serialize(Ar);
+}
+
+void UAnimNextGraph_EditorData::Initialize(bool bRecompileVM)
+{
+	UAnimNextGraph* AnimNextGraph = GetTypedOuter<UAnimNextGraph>();
+
+	if (RigVMClient.GetController(0) == nullptr)
+	{
+		check(RigVMClient.Num() == 1);
+		check(RigVMClient.GetFunctionLibrary());
+		
+		RigVMClient.GetOrCreateController(RigVMClient.GetDefaultModel());
+		RigVMClient.GetOrCreateController(RigVMClient.GetFunctionLibrary());
+
+		// Init function library controllers
+		for(URigVMLibraryNode* LibraryNode : RigVMClient.GetFunctionLibrary()->GetFunctions())
+		{
+			RigVMClient.GetOrCreateController(LibraryNode->GetContainedGraph());
+		}
+		
+		if(bRecompileVM)
+		{
+			UE::AnimNext::GraphUncookedOnly::FUtils::Compile(AnimNextGraph);
+		}
+	}
+
+	RootGraph->Initialize(this);
+	FunctionLibraryEdGraph->Initialize(this);
+	if(EntryPointGraph)
+	{
+		EntryPointGraph->Initialize(this);
+	}
+}
+
+void UAnimNextGraph_EditorData::PostLoad()
+{
+	Super::PostLoad();
+	
+	Initialize(/*bRecompileVM*/false);
+	RefreshAllModels(EAnimNextGraphLoadType::PostLoad);
+
+#if WITH_EDITOR
+	if (GIsEditor)
+	{
+		// delay compilation until the package has been loaded
+		FCoreUObjectDelegates::OnEndLoadPackage.AddUObject(this, &UAnimNextGraph_EditorData::HandlePackageDone);
+	}
+#else // !WITH_EDITOR
+	RecompileVMIfRequired();
+#endif // WITH_EDITOR
+}
+
+#if WITH_EDITOR
+void UAnimNextGraph_EditorData::HandlePackageDone(const FEndLoadPackageContext& Context)
+{
+	if (!Context.LoadedPackages.Contains(GetPackage()))
+	{
+		return;
+	}
+	HandlePackageDone();
+}
+
+void UAnimNextGraph_EditorData::HandlePackageDone()
+{
+	FCoreUObjectDelegates::OnEndLoadPackage.RemoveAll(this);
+
+	RecompileVM();
+}
+
+void UAnimNextGraph_EditorData::RefreshAllModels(EAnimNextGraphLoadType InLoadType)
+{
+	const bool bIsPostLoad = InLoadType == EAnimNextGraphLoadType::PostLoad;
+
+	TGuardValue<bool> IsCompilingGuard(bIsCompiling, true);
+	TGuardValue<bool> ClientIgnoreModificationsGuard(RigVMClient.bIgnoreModelNotifications, true);
+
+	TArray<URigVMGraph*> GraphsToDetach = RigVMClient.GetAllModels(true, false);
+
+	if (ensure(IsInGameThread()))
+	{
+		for (URigVMGraph* GraphToDetach : GraphsToDetach)
+		{
+			URigVMController* Controller = RigVMClient.GetOrCreateController(GraphToDetach);
+			// temporarily disable default value validation during load time, serialized values should always be accepted
+			TGuardValue<bool> PerGraphDisablePinDefaultValueValidation(Controller->bValidatePinDefaults, false);
+			FRigVMControllerNotifGuard NotifGuard(Controller, true);
+			Controller->DetachLinksFromPinObjects();
+			TArray<URigVMNode*> Nodes = GraphToDetach->GetNodes();
+			for (URigVMNode* Node : Nodes)
+			{
+				Controller->RepopulatePinsOnNode(Node, true, true);
+			}
+		}
+		//SetupPinRedirectorsForBackwardsCompatibility();
+	}
+
+	for (URigVMGraph* GraphToDetach : GraphsToDetach)
+	{
+		URigVMController* Controller = RigVMClient.GetOrCreateController(GraphToDetach);
+		// at this stage, allow all links to be reattached,
+		// RecomputeAllTemplateFilteredPermutations() later should break any invalid links
+		FRigVMControllerNotifGuard NotifGuard(Controller, true);
+		Controller->ReattachLinksToPinObjects(true /* follow redirectors */, nullptr, true, true);
+	}
+
+	if (bIsPostLoad)
+	{
+		//PatchTemplateNodesWithPreferredPermutation();
+	}
+
+	TArray<URigVMGraph*> GraphsToClean = RigVMClient.GetAllModels(true, true);
+
+	// Sort from leaf graphs to root
+	TArray<URigVMGraph*> SortedGraphsToClean;
+	SortedGraphsToClean.Reserve(GraphsToClean.Num());
+	while (SortedGraphsToClean.Num() < GraphsToClean.Num())
+	{
+		bool bGraphAdded = false;
+		for (URigVMGraph* Graph : GraphsToClean)
+		{
+			if (SortedGraphsToClean.Contains(Graph))
+			{
+				continue;
+			}
+
+			TArray<URigVMGraph*> ContainedGraphs;
+			for (URigVMNode* Node : Graph->GetNodes())
+			{
+				if (URigVMLibraryNode* LibraryNode = Cast<URigVMLibraryNode>(Node))
+				{
+					if (URigVMFunctionReferenceNode* FunctionReferenceNode = Cast<URigVMFunctionReferenceNode>(LibraryNode))
+					{
+						if (FunctionReferenceNode->GetReferencedFunctionHeader().LibraryPointer.LibraryNode.GetLongPackageName() != GetPackage()->GetPathName())
+						{
+							continue;
+						}
+						if (URigVMLibraryNode* ReferencedNode = FunctionReferenceNode->LoadReferencedNode())
+						{
+							ContainedGraphs.Add(ReferencedNode->GetContainedGraph());
+							continue;
+						}
+					}
+
+					if (URigVMGraph* ContainedGraph = LibraryNode->GetContainedGraph())
+					{
+						ContainedGraphs.Add(ContainedGraph);
+					}
+				}
+			}
+
+			bool bAllContained = true;
+			for (URigVMGraph* Contained : ContainedGraphs)
+			{
+				if (!SortedGraphsToClean.Contains(Contained))
+				{
+					bAllContained = false;
+					break;
+				}
+			}
+			if (bAllContained)
+			{
+				SortedGraphsToClean.Add(Graph);
+				bGraphAdded = true;
+			}
+		}
+		ensure(bGraphAdded);
+	}
+
+	for (int32 GraphIndex = 0; GraphIndex < SortedGraphsToClean.Num(); GraphIndex++)
+	{
+		URigVMGraph* GraphToClean = SortedGraphsToClean[GraphIndex];
+		URigVMController* Controller = RigVMClient.GetOrCreateController(GraphToClean);
+		//TGuardValue<bool> GuardEditGraph(GraphToClean->bEditable, true);
+		FRigVMControllerNotifGuard NotifGuard(Controller, true);
+
+		for (URigVMNode* ModelNode : GraphToClean->GetNodes())
+		{
+			Controller->RemoveUnusedOrphanedPins(ModelNode);
+		}
+	}
+}
+
+#endif // WITH_EDITOR
+
+FRigVMClient* UAnimNextGraph_EditorData::GetRigVMClient()
+{
+	return &RigVMClient;
+}
+
+const FRigVMClient* UAnimNextGraph_EditorData::GetRigVMClient() const
+{
+	return &RigVMClient;
+}
+
+IRigVMGraphFunctionHost* UAnimNextGraph_EditorData::GetRigVMGraphFunctionHost()
+{
+	return this;
+}
+
+const IRigVMGraphFunctionHost* UAnimNextGraph_EditorData::GetRigVMGraphFunctionHost() const
+{
+	return this;
+}
+
+
+void UAnimNextGraph_EditorData::HandleRigVMGraphAdded(const FRigVMClient* InClient, const FString& InNodePath)
+{
+	if(URigVMGraph* RigVMGraph = InClient->GetModel(InNodePath))
+	{
+		RigVMGraph->SetExecuteContextStruct(FAnimNextExecuteContext::StaticStruct());
+	}
+}
+
+void UAnimNextGraph_EditorData::HandleConfigureRigVMController(const FRigVMClient* InClient,
+                                                                    URigVMController* InControllerToConfigure)
+{
+	InControllerToConfigure->OnModified().AddUObject(this, &UAnimNextGraph_EditorData::HandleModifiedEvent);
+
+	InControllerToConfigure->UnfoldStructDelegate.BindLambda([](const UStruct* InStruct) -> bool
+	{
+		if (InStruct == TBaseStructure<FQuat>::Get())
+		{
+			return false;
+		}
+		if (InStruct == FRuntimeFloatCurve::StaticStruct())
+		{
+			return false;
+		}
+		if (InStruct == FRigPose::StaticStruct())
+		{
+			return false;
+		}
+		return true;
+	});
+
+	TWeakObjectPtr<UAnimNextGraph_EditorData> WeakThis(this);
+
+	// this delegate is used by the controller to determine variable validity
+	// during a bind process. the controller itself doesn't own the variables,
+	// so we need a delegate to request them from the owning blueprint
+	InControllerToConfigure->GetExternalVariablesDelegate.BindLambda([](URigVMGraph* InGraph) -> TArray<FRigVMExternalVariable>
+	{
+		if (InGraph)
+		{
+			if(UAnimNextGraph_EditorData* EditorData = InGraph->GetTypedOuter<UAnimNextGraph_EditorData>())
+			{
+				if (UAnimNextGraph* Graph = EditorData->GetTypedOuter<UAnimNextGraph>())
+				{
+					return Graph->GetRigVMExternalVariables();
+				}
+			}
+		}
+		return TArray<FRigVMExternalVariable>();
+	});
+
+	// this delegate is used by the controller to retrieve the current bytecode of the VM
+	InControllerToConfigure->GetCurrentByteCodeDelegate.BindLambda([WeakThis]() -> const FRigVMByteCode*
+	{
+		if (WeakThis.IsValid())
+		{
+			if(UAnimNextGraph* Graph = WeakThis->GetTypedOuter<UAnimNextGraph>())
+			{
+				if (Graph->RigVM)
+				{
+					return &Graph->RigVM->GetByteCode();
+				}
+			}
+		}
+		return nullptr;
+
+	});
+
+#if WITH_EDITOR
+	InControllerToConfigure->SetupDefaultUnitNodeDelegates(TDelegate<FName(FRigVMExternalVariable, FString)>::CreateLambda(
+		[](FRigVMExternalVariable InVariableToCreate, FString InDefaultValue) -> FName
+		{
+			return NAME_None;
+		}
+	));
+#endif
+
+}
+
+UObject* UAnimNextGraph_EditorData::GetEditorObjectForRigVMGraph(URigVMGraph* InVMGraph) const
+{
+	if(InVMGraph)
+	{
+		TArray<UAnimNextGraph_EdGraph*> AllGraphs = {RootGraph, EntryPointGraph, FunctionLibraryEdGraph};
+		for(UAnimNextGraph_EdGraph* EdGraph : AllGraphs)
+		{
+			if(EdGraph)
+			{
+				if(EdGraph->ModelNodePath == InVMGraph->GetNodePath())
+				{
+					return EdGraph;
+				}
+			}
+		}
+	}
+	return nullptr;
+}
+
+FRigVMGraphFunctionStore* UAnimNextGraph_EditorData::GetRigVMGraphFunctionStore()
+{
+	return &GraphFunctionStore;
+}
+
+const FRigVMGraphFunctionStore* UAnimNextGraph_EditorData::GetRigVMGraphFunctionStore() const
+{
+	return &GraphFunctionStore;
+}
+
+
+void UAnimNextGraph_EditorData::RecompileVM()
+{
+	UE::AnimNext::GraphUncookedOnly::FUtils::Compile(GetTypedOuter<UAnimNextGraph>());
+}
+
+void UAnimNextGraph_EditorData::RecompileVMIfRequired()
+{
+	if (bVMRecompilationRequired)
+	{
+		RecompileVM();
+	}
+}
+
+void UAnimNextGraph_EditorData::RequestAutoVMRecompilation()
+{
+	bVMRecompilationRequired = true;
+	if (bAutoRecompileVM)
+	{
+		RecompileVMIfRequired();
+	}
+}
+
+void UAnimNextGraph_EditorData::HandleModifiedEvent(ERigVMGraphNotifType InNotifType, URigVMGraph* InGraph, UObject* InSubject)
+{
+	bool bNotifForOthersPending = true;
+
+	switch(InNotifType)
+	{
+	case ERigVMGraphNotifType::NodeAdded:
+		{
+			if (URigVMCollapseNode* CollapseNode = Cast<URigVMCollapseNode>(InSubject))
+			{
+				CreateEdGraphForCollapseNode(CollapseNode);
+				break;
+			}
+		}
+		// Fall through to the next case
+	case ERigVMGraphNotifType::LinkAdded:
+	case ERigVMGraphNotifType::LinkRemoved:
+	case ERigVMGraphNotifType::PinArraySizeChanged:
+	case ERigVMGraphNotifType::PinDirectionChanged:
+		{
+			UpdateGraphReturnType();
+			RecompileVM();
+			break;
+		}
+
+	case ERigVMGraphNotifType::PinDefaultValueChanged:
+		{
+			if (InGraph->GetRuntimeAST().IsValid())
+			{
+				URigVMPin* RootPin = CastChecked<URigVMPin>(InSubject)->GetRootPin();
+				FRigVMASTProxy RootPinProxy = FRigVMASTProxy::MakeFromUObject(RootPin);
+				const FRigVMExprAST* Expression = InGraph->GetRuntimeAST()->GetExprForSubject(RootPinProxy);
+				if (Expression == nullptr)
+				{
+					InGraph->ClearAST();
+					break;
+				}
+				else if (Expression->NumParents() > 1)
+				{
+					InGraph->ClearAST();
+					break;
+				}
+			}
+			break;
+		}
+	}
+	
+	// if the notification still has to be sent...
+	if (bNotifForOthersPending && !bSuspendModelNotificationsForOthers)
+	{
+		if (ModifiedEvent.IsBound())
+		{
+			ModifiedEvent.Broadcast(InNotifType, InGraph, InSubject);
+		}
+	}
+}
+
+URigVMGraph* UAnimNextGraph_EditorData::GetVMGraphForEdGraph(const UEdGraph* InGraph) const
+{
+	if (InGraph == RootGraph)
+	{
+		return RigVMClient.GetDefaultModel();
+	}
+	else
+	{
+		const UAnimNextGraph_EdGraph* Graph = Cast<UAnimNextGraph_EdGraph>(InGraph);
+		check(Graph);
+
+		if (Graph->bIsFunctionDefinition)
+		{
+			if (URigVMLibraryNode* LibraryNode = RigVMClient.GetFunctionLibrary()->FindFunction(*Graph->ModelNodePath))
+			{
+				return LibraryNode->GetContainedGraph();
+			}
+		}
+	}
+
+	return nullptr;
+}
+
+void UAnimNextGraph_EditorData::CreateEdGraphForCollapseNode(URigVMCollapseNode* InNode)
+{
+	if (InNode->GetGraph()->IsA<URigVMFunctionLibrary>())
+	{
+		if (URigVMGraph* ContainedGraph = InNode->GetContainedGraph())
+		{
+			if (EntryPointGraph == nullptr)
+			{
+				// create a sub graph
+				UAnimNextGraph_EdGraph* RigFunctionGraph = NewObject<UAnimNextGraph_EdGraph>(this, *InNode->GetName(), RF_Transactional);
+				RigFunctionGraph->Schema = UAnimNextGraph_EdGraphSchema::StaticClass();
+				RigFunctionGraph->bAllowRenaming = true;
+				RigFunctionGraph->bEditable = true;
+				RigFunctionGraph->bAllowDeletion = true;
+				RigFunctionGraph->ModelNodePath = ContainedGraph->GetNodePath();
+				RigFunctionGraph->bIsFunctionDefinition = true;
+
+				EntryPointGraph = RigFunctionGraph;
+				RigFunctionGraph->Initialize(this);
+
+				RigVMClient.GetOrCreateController(ContainedGraph)->ResendAllNotifications();
+			}
+		}
+	}
+}
+
+FEdGraphPinType UAnimNextGraph_EditorData::FindGraphReturnPinType() const
+{
+	static const TCHAR* EndExecStr = TEXT("AnimNextEndExec");
+	static const FName ResultName = TEXT(".Result");
+
+	const UAnimNextGraph_EdGraph* AnimNextInterfaceEdGraph = RootGraph.Get();
+	if (AnimNextInterfaceEdGraph != nullptr)
+	{
+		for (const TObjectPtr<UEdGraphNode>& Node : AnimNextInterfaceEdGraph->Nodes)
+		{
+			if (IsNodeExecConnected(Node))
+			{
+				const FString& NodeName = Node->GetName();
+				if (NodeName.Contains(EndExecStr))
+				{
+					for (const UEdGraphPin* Pin : Node->Pins)
+					{
+						const FString PinName = Pin->PinName.ToString();
+
+						if (PinName.Contains(ResultName.ToString()) && PinName.Contains(EndExecStr))
+						{
+							return Pin->PinType;
+						}
+					}
+				}
+			}
+		}
+	}
+
+	return FEdGraphPinType();
+}
+
+bool UAnimNextGraph_EditorData::IsNodeExecConnected(const UEdGraphNode* Node) const
+{
+	static const TCHAR* ExecuteContextStr = TEXT(".ExecuteContext");
+
+	for (const UEdGraphPin* Pin : Node->Pins)
+	{
+		const FString PinName = Pin->PinName.ToString();
+
+		if (PinName.Contains(ExecuteContextStr) && Pin->LinkedTo.Num() > 0)
+		{
+			return true;
+		}
+	}
+
+	return false;
+}
+
+void UAnimNextGraph_EditorData::UpdateGraphReturnType()
+{
+	if (UAnimNextGraph* Graph = GetTypedOuter<UAnimNextGraph>())
+	{
+		FEdGraphPinType GraphReturnPinType = FindGraphReturnPinType();
+		UE::AnimNext::FParamTypeHandle Handle = UE::AnimNext::GraphUncookedOnly::FUtils::GetParameterHandleFromPin(GraphReturnPinType);
+		if (Handle.IsValid() && Graph->GetReturnTypeHandle() != Handle)
+		{
+			Graph->Modify();
+			Graph->SetReturnTypeHandle(Handle);
+		}
+	}
+}
diff --git a/Engine/Plugins/Experimental/AnimNext/Source/AnimNextGraphUncookedOnly/Private/UncookedOnlyModule.cpp b/Engine/Plugins/Experimental/AnimNext/Source/AnimNextGraphUncookedOnly/Private/UncookedOnlyModule.cpp
new file mode 100644
--- /dev/null
+++ b/Engine/Plugins/Experimental/AnimNext/Source/AnimNextGraphUncookedOnly/Private/UncookedOnlyModule.cpp
@@ -0,0 +1,5 @@
+// Copyright Epic Games, Inc. All Rights Reserved.
+
+#include "Modules/ModuleManager.h"
+
+IMPLEMENT_MODULE(FDefaultModuleImpl, AnimNextGraphUncookedOnly);
\ No newline at end of file
diff --git a/Engine/Plugins/Experimental/AnimNext/Source/AnimNextGraphUncookedOnly/Private/UncookedOnlyUtils.cpp b/Engine/Plugins/Experimental/AnimNext/Source/AnimNextGraphUncookedOnly/Private/UncookedOnlyUtils.cpp
new file mode 100644
--- /dev/null
+++ b/Engine/Plugins/Experimental/AnimNext/Source/AnimNextGraphUncookedOnly/Private/UncookedOnlyUtils.cpp
@@ -0,0 +1,195 @@
+﻿// Copyright Epic Games, Inc. All Rights Reserved.
+
+#include "UncookedOnlyUtils.h"
+#include "AnimNextGraph.h"
+#include "AnimNextGraph_EditorData.h"
+#include "AnimNextGraph_EdGraph.h"
+#include "RigVMCompiler/RigVMCompiler.h"
+#include "RigVMCore/RigVM.h"
+#include "ExecuteContext.h"
+
+namespace UE::AnimNext::GraphUncookedOnly
+{
+
+void FUtils::Compile(UAnimNextGraph* InGraph)
+{
+	check(InGraph);
+	
+	UAnimNextGraph_EditorData* EditorData = GetEditorData(InGraph);
+	
+	if(EditorData->bIsCompiling)
+	{
+		return;
+	}
+	
+	TGuardValue<bool> CompilingGuard(EditorData->bIsCompiling, true);
+	
+	EditorData->bErrorsDuringCompilation = false;
+
+	EditorData->RigGraphDisplaySettings.MinMicroSeconds = EditorData->RigGraphDisplaySettings.LastMinMicroSeconds = DBL_MAX;
+	EditorData->RigGraphDisplaySettings.MaxMicroSeconds = EditorData->RigGraphDisplaySettings.LastMaxMicroSeconds = (double)INDEX_NONE;
+	
+	TGuardValue<bool> ReentrantGuardSelf(EditorData->bSuspendModelNotificationsForSelf, true);
+	TGuardValue<bool> ReentrantGuardOthers(EditorData->bSuspendModelNotificationsForOthers, true);
+
+	RecreateVM(InGraph);
+
+	InGraph->VMRuntimeSettings = EditorData->VMRuntimeSettings;
+
+	EditorData->CompileLog.Messages.Reset();
+	EditorData->CompileLog.NumErrors = EditorData->CompileLog.NumWarnings = 0;
+
+	URigVMCompiler* Compiler = URigVMCompiler::StaticClass()->GetDefaultObject<URigVMCompiler>();
+	EditorData->VMCompileSettings.SetExecuteContextStruct(EditorData->RigVMClient.GetExecuteContextStruct());
+	Compiler->Settings = (EditorData->bCompileInDebugMode) ? FRigVMCompileSettings::Fast(EditorData->VMCompileSettings.GetExecuteContextStruct()) : EditorData->VMCompileSettings;
+	URigVMController* RootController = EditorData->GetRigVMClient()->GetOrCreateController(EditorData->GetRigVMClient()->GetDefaultModel());
+	Compiler->Compile(EditorData->GetRigVMClient()->GetAllModels(false, false), RootController, InGraph->RigVM, InGraph->GetRigVMExternalVariables(), &EditorData->PinToOperandMap);
+
+	if (EditorData->bErrorsDuringCompilation)
+	{
+		if(Compiler->Settings.SurpressErrors)
+		{
+			Compiler->Settings.Reportf(EMessageSeverity::Info, InGraph,TEXT("Compilation Errors may be suppressed for AnimNext Interface Graph: %s. See VM Compile Settings for more Details"), *InGraph->GetName());
+		}
+	}
+
+	EditorData->bVMRecompilationRequired = false;
+	if(InGraph->RigVM)
+	{
+		EditorData->VMCompiledEvent.Broadcast(InGraph, InGraph->RigVM);
+	}
+
+#if WITH_EDITOR
+//	RefreshBreakpoints(EditorData);
+#endif
+}
+
+void FUtils::RecreateVM(UAnimNextGraph* InGraph)
+{
+	InGraph->RigVM = NewObject<URigVM>(InGraph, TEXT("VM"), RF_NoFlags);
+	InGraph->RigVM->SetContextPublicDataStruct(FAnimNextExecuteContext::StaticStruct());
+	
+	// Cooked platforms will load these pointers from disk
+	if (!FPlatformProperties::RequiresCookedData())
+	{
+		// We dont support ERigVMMemoryType::Work memory as we dont operate on an instance
+	//	InGraph->RigVM->GetMemoryByType(ERigVMMemoryType::Work, true);
+		InGraph->RigVM->GetMemoryByType(ERigVMMemoryType::Literal, true);
+		InGraph->RigVM->GetMemoryByType(ERigVMMemoryType::Debug, true);
+	}
+
+	InGraph->RigVM->Reset();
+}
+
+UAnimNextGraph_EditorData* FUtils::GetEditorData(const UAnimNextGraph* InAnimNextGraph)
+{
+	check(InAnimNextGraph);
+	
+	return CastChecked<UAnimNextGraph_EditorData>(InAnimNextGraph->EditorData);
+}
+
+FParamTypeHandle FUtils::GetParameterHandleFromPin(const FEdGraphPinType& InPinType)
+{
+	FAnimNextParamType::EValueType ValueType = FAnimNextParamType::EValueType::None;
+	FAnimNextParamType::EContainerType ContainerType = FAnimNextParamType::EContainerType::None;
+	UObject* ValueTypeObject = nullptr;
+
+	if (InPinType.PinCategory == UEdGraphSchema_K2::PC_Boolean)
+	{
+		ValueType = FAnimNextParamType::EValueType::Bool;
+	}
+	else if (InPinType.PinCategory == UEdGraphSchema_K2::PC_Byte)
+	{
+		ValueType = FAnimNextParamType::EValueType::Byte;
+	}
+	else if (InPinType.PinCategory == UEdGraphSchema_K2::PC_Int)
+	{
+		ValueType = FAnimNextParamType::EValueType::Int32;
+	}
+	else if (InPinType.PinCategory == UEdGraphSchema_K2::PC_Int64)
+	{
+		ValueType = FAnimNextParamType::EValueType::Int64;
+	}
+	else if (InPinType.PinCategory == UEdGraphSchema_K2::PC_Real)
+	{
+		if (InPinType.PinSubCategory == UEdGraphSchema_K2::PC_Float)
+		{
+			ValueType = FAnimNextParamType::EValueType::Float;
+		}
+		else if (InPinType.PinSubCategory == UEdGraphSchema_K2::PC_Double)
+		{
+			ValueType = FAnimNextParamType::EValueType::Double;
+		}
+		else
+		{
+			ensure(false);	// Reals should be either floats or doubles
+		}
+	}
+	else if (InPinType.PinCategory == UEdGraphSchema_K2::PC_Float)
+	{
+		ValueType = FAnimNextParamType::EValueType::Float;
+	}
+	else if (InPinType.PinCategory == UEdGraphSchema_K2::PC_Double)
+	{
+		ValueType = FAnimNextParamType::EValueType::Double;
+	}
+	else if (InPinType.PinCategory == UEdGraphSchema_K2::PC_Name)
+	{
+		ValueType = FAnimNextParamType::EValueType::Name;
+	}
+	else if (InPinType.PinCategory == UEdGraphSchema_K2::PC_String)
+	{
+		ValueType = FAnimNextParamType::EValueType::String;
+	}
+	else if (InPinType.PinCategory == UEdGraphSchema_K2::PC_Text)
+	{
+		ValueType = FAnimNextParamType::EValueType::Text;
+	}
+	else if (InPinType.PinCategory == UEdGraphSchema_K2::PC_Enum)
+	{
+		ValueType = FAnimNextParamType::EValueType::Enum;
+		ValueTypeObject = InPinType.PinSubCategoryObject.Get();
+		ensure(ValueTypeObject);
+	}
+	else if (InPinType.PinCategory == UEdGraphSchema_K2::PC_Struct)
+	{
+		ValueType = FAnimNextParamType::EValueType::Struct;
+		ValueTypeObject = Cast<UScriptStruct>(InPinType.PinSubCategoryObject.Get());
+		ensure(ValueTypeObject);
+	}
+	else if (InPinType.PinCategory == UEdGraphSchema_K2::PC_Object)
+	{
+		ValueType = FAnimNextParamType::EValueType::Object;
+		ValueTypeObject = Cast<UClass>(InPinType.PinSubCategoryObject.Get());
+		ensure(ValueTypeObject);
+	}
+	else if (InPinType.PinCategory == UEdGraphSchema_K2::PC_SoftObject)
+	{
+		ValueType = FAnimNextParamType::EValueType::SoftObject;
+		ValueTypeObject = Cast<UClass>(InPinType.PinSubCategoryObject.Get());
+		ensure(ValueTypeObject);
+	}
+	else if (InPinType.PinCategory == UEdGraphSchema_K2::PC_SoftClass)
+	{
+		ValueType = FAnimNextParamType::EValueType::SoftClass;
+		ValueTypeObject = Cast<UClass>(InPinType.PinSubCategoryObject.Get());
+		ensure(ValueTypeObject);
+	}
+
+	if(InPinType.ContainerType == EPinContainerType::Array)
+	{
+		ContainerType = FAnimNextParamType::EContainerType::Array;
+	}
+	else if(InPinType.ContainerType == EPinContainerType::Set)
+	{
+		ensureMsgf(false, TEXT("Set pins are not yet supported"));
+	}
+	if(InPinType.ContainerType == EPinContainerType::Map)
+	{
+		ensureMsgf(false, TEXT("Map pins are not yet supported"));
+	}
+	
+	return FAnimNextParamType(ValueType, ContainerType, ValueTypeObject).GetHandle();
+}
+
+}
