Fix CIS errors at Engine Test - AnimNext graph could not serialize

diff --git a/Engine/Plugins/Experimental/AnimNext/Source/AnimNext/Private/GenerationTools.cpp b/Engine/Plugins/Experimental/AnimNext/Source/AnimNext/Private/GenerationTools.cpp
--- a/Engine/Plugins/Experimental/AnimNext/Source/AnimNext/Private/GenerationTools.cpp
+++ b/Engine/Plugins/Experimental/AnimNext/Source/AnimNext/Private/GenerationTools.cpp
@@ -400,14 +400,17 @@ struct FCompareBoneIndexType
 	const TArrayView<const FBoneIndexType> LODBoneIndexes = RefPose.GetLODBoneIndexes(LODLevel);
 	const int32 NumBones = LODBoneIndexes.Num();
 
-	for (int i = 0; i < NumBones; ++i)
+	if (TargetPose.GetNumBones() == NumBones)
 	{
-		const auto& SkeletonBoneIndex = LODBoneIndexes[i];
-
-		const FCompactPoseBoneIndex CompactBoneIndex = SourcePose.Pose.GetBoneContainer().GetCompactPoseIndexFromSkeletonIndex(SkeletonBoneIndex);
-		if (CompactBoneIndex.GetInt() != INDEX_NONE)
+		for (int i = 0; i < NumBones; ++i)
 		{
-			TargetPose.LocalTransforms[i] = SourcePose.Pose[CompactBoneIndex];
+			const auto& SkeletonBoneIndex = LODBoneIndexes[i];
+
+			const FCompactPoseBoneIndex CompactBoneIndex = SourcePose.Pose.GetBoneContainer().GetCompactPoseIndexFromSkeletonIndex(SkeletonBoneIndex);
+			if (CompactBoneIndex.GetInt() != INDEX_NONE)
+			{
+				TargetPose.LocalTransforms[i] = SourcePose.Pose[CompactBoneIndex];
+			}
 		}
 	}
 }
@@ -420,14 +423,17 @@ struct FCompareBoneIndexType
 	const TArrayView<const FBoneIndexType> LODBoneIndexes = RefPose.GetLODBoneIndexes(LODLevel);
 	const int32 NumBones = LODBoneIndexes.Num();
 
-	for (int i = 0; i < NumBones; ++i)
+	if (SourcePose.GetNumBones() == NumBones)
 	{
-		const auto& SkeletonBoneIndex = LODBoneIndexes[i];
-
-		const FCompactPoseBoneIndex CompactBoneIndex = TargetPose.Pose.GetBoneContainer().GetCompactPoseIndexFromSkeletonIndex(SkeletonBoneIndex);
-		if (CompactBoneIndex.GetInt() != INDEX_NONE)
+		for (int i = 0; i < NumBones; ++i)
 		{
-			TargetPose.Pose[CompactBoneIndex] = SourcePose.LocalTransforms[i];
+			const auto& SkeletonBoneIndex = LODBoneIndexes[i];
+
+			const FCompactPoseBoneIndex CompactBoneIndex = TargetPose.Pose.GetBoneContainer().GetCompactPoseIndexFromSkeletonIndex(SkeletonBoneIndex);
+			if (CompactBoneIndex.GetInt() != INDEX_NONE)
+			{
+				TargetPose.Pose[CompactBoneIndex] = SourcePose.LocalTransforms[i];
+			}
 		}
 	}
 }
