Non-unity build fix

diff --git a/Engine/Plugins/Experimental/AnimNext/Source/AnimNext/Public/TransformArray.h b/Engine/Plugins/Experimental/AnimNext/Source/AnimNext/Public/TransformArray.h
--- a/Engine/Plugins/Experimental/AnimNext/Source/AnimNext/Public/TransformArray.h
+++ b/Engine/Plugins/Experimental/AnimNext/Source/AnimNext/Public/TransformArray.h
@@ -4,6 +4,7 @@
 
 #include "CoreMinimal.h"
 #include "Animation/AnimTypes.h"
+#include <algorithm>
 
 #define ANIM_ENABLE_POINTER_ITERATION 1
 
