Removed check in constexpr that required c++20

diff --git a/Engine/Plugins/Experimental/AnimNext/Source/AnimNext/Public/Param/ParamTypeHandle.h b/Engine/Plugins/Experimental/AnimNext/Source/AnimNext/Public/Param/ParamTypeHandle.h
--- a/Engine/Plugins/Experimental/AnimNext/Source/AnimNext/Public/Param/ParamTypeHandle.h
+++ b/Engine/Plugins/Experimental/AnimNext/Source/AnimNext/Public/Param/ParamTypeHandle.h
@@ -291,7 +291,7 @@ public:
 			}
 		}
 
-		check(TypeHandle.IsValid());
+		//check(TypeHandle.IsValid());
 
 		return TypeHandle;
 	}
