Implementing AnimNext low level runtime pieces

diff --git a/Engine/Plugins/Experimental/AnimNext/AnimNext.uplugin b/Engine/Plugins/Experimental/AnimNext/AnimNext.uplugin
--- a/Engine/Plugins/Experimental/AnimNext/AnimNext.uplugin
+++ b/Engine/Plugins/Experimental/AnimNext/AnimNext.uplugin
@@ -12,7 +12,8 @@
 	"SupportURL" : "",
 	"EnabledByDefault" : false,
 	"CanContainContent" : false,
-	"IsBetaVersion" : true,
+	"IsBetaVersion": true,
+	"IsExperimentalVersion": true,
 	"Installed" : false,
 	"Plugins":
 	[
@@ -43,9 +44,13 @@
 			"LoadingPhase" : "Default"
 		},
 		{
-			"Name" : "AnimNextGraphUncookedOnly",
-			"Type" : "UncookedOnly",
-			"LoadingPhase" : "Default"
+			"Name": "AnimNextGraphUncookedOnly",
+			"Type": "UncookedOnly",
+			"LoadingPhase": "Default"
+		},
+		{
+			"Name": "AnimNextTestSuite",
+			"Type": "UncookedOnly"
 		}
 	]
 }
diff --git a/Engine/Plugins/Experimental/AnimNext/Source/AnimNext/Internal/DecoratorBase/DecoratorRegistry.h b/Engine/Plugins/Experimental/AnimNext/Source/AnimNext/Internal/DecoratorBase/DecoratorRegistry.h
new file mode 100644
--- /dev/null
+++ b/Engine/Plugins/Experimental/AnimNext/Source/AnimNext/Internal/DecoratorBase/DecoratorRegistry.h
@@ -0,0 +1,99 @@
+// Copyright Epic Games, Inc. All Rights Reserved.
+
+#pragma once
+
+#include "CoreMinimal.h"
+#include "DecoratorBase/Decorator.h"
+#include "DecoratorBase/DecoratorUID.h"
+#include "DecoratorBase/DecoratorRegistryHandle.h"
+
+namespace UE::AnimNext
+{
+	/**
+	 * FDecoratorRegistry
+	 * 
+	 * A global registry of all existing decorators that can be used in animation graphs.
+	 * 
+	 * @see FDecorator
+	 */
+	struct ANIMNEXT_API FDecoratorRegistry final
+	{
+		// Access the global registry
+		static FDecoratorRegistry& Get();
+
+		// Finds and returns the decorator handle for the provided decorator UID or an invalid
+		// handle if that decorator hasn't been registered yet.
+		FDecoratorRegistryHandle FindHandle(uint32 DecoratorUID) const;
+
+		// Finds and returns the decorator handle for the provided decorator UID or an invalid
+		// handle if that decorator hasn't been registered yet.
+		FDecoratorRegistryHandle FindHandle(FDecoratorUID DecoratorUID) const;
+
+		// Finds and returns the decorator associated with the provided handle. If the handle
+		// is invalid, nullptr is returned.
+		const FDecorator* Find(FDecoratorRegistryHandle DecoratorHandle) const;
+
+		// Finds and returns the decorator associated with the provided decorator UID. If the decorator
+		// is registered, nullptr is returned.
+		const FDecorator* Find(FDecoratorUID DecoratorUID) const;
+
+		// Registers a decorator dynamically
+		void Register(FDecorator* Decorator);
+
+		// Unregisters a decorator dynamically
+		void Unregister(FDecorator* Decorator);
+
+		// Returns the number of registered decorators
+		uint32 GetNum() const;
+
+	private:
+		FDecoratorRegistry() = default;
+
+		// Static init lifetime functions
+		static void StaticRegister(DecoratorConstructorFunc DecoratorConstructor);
+		static void StaticUnregister(DecoratorConstructorFunc DecoratorConstructor);
+		void AutoRegisterImpl(DecoratorConstructorFunc DecoratorConstructor);
+		void AutoUnregisterImpl(DecoratorConstructorFunc DecoratorConstructor);
+
+		// Module lifetime functions
+		static void Init();
+		static void Destroy();
+
+		// Holds information for each registered decorator
+		struct FRegistryEntry
+		{
+			// A pointer to the decorator
+			// TODO: Do we want a shared ptr here? we don't always own it
+			FDecorator*					Decorator = nullptr;
+
+			// A pointer to the constructor function
+			// Only valid when the decorator has been auto-registered
+			DecoratorConstructorFunc	DecoratorConstructor = nullptr;
+
+			// The decorator handle
+			FDecoratorRegistryHandle	DecoratorHandle;
+		};
+
+		// For performance reasons, we store static decorators that never unload into
+		// a single contiguous memory buffer. However, decorators cannot be guaranteed
+		// to be trivially copyable because they contain virtual functions. As such,
+		// we cannot resize the buffer once they have been allocated. We reserve a fixed
+		// amount of space that should easily cover our needs. Static decorators are
+		// generally stateless and only contain a few v-tables. Their size is usually
+		// less than 32 bytes. Additionally, we will likely only ever load a few hundred
+		// decorators. If we exceed the size of the buffer, additional decorators will
+		// be treated as dynamic. Dynamic decorators are instead allocated on the heap.
+		static constexpr uint32			STATIC_DECORATOR_BUFFER_SIZE = 8 * 1024;
+
+		uint8							StaticDecoratorBuffer[STATIC_DECORATOR_BUFFER_SIZE] = { 0 };
+		int32							StaticDecoratorBufferOffset = 0;
+
+		TArray<uintptr_t>				DynamicDecorators;
+		int32							DynamicDecoratorFreeIndexHead = INDEX_NONE;
+
+		TMap<uint32, FRegistryEntry>	DecoratorUIDToEntryMap;
+
+		friend class FModule;
+		friend struct FDecoratorStaticInitHook;
+	};
+}
diff --git a/Engine/Plugins/Experimental/AnimNext/Source/AnimNext/Internal/DecoratorBase/DecoratorTemplate.h b/Engine/Plugins/Experimental/AnimNext/Source/AnimNext/Internal/DecoratorBase/DecoratorTemplate.h
new file mode 100644
--- /dev/null
+++ b/Engine/Plugins/Experimental/AnimNext/Source/AnimNext/Internal/DecoratorBase/DecoratorTemplate.h
@@ -0,0 +1,109 @@
+// Copyright Epic Games, Inc. All Rights Reserved.
+
+#pragma once
+
+#include "CoreMinimal.h"
+#include "DecoratorBase/DecoratorMode.h"
+#include "DecoratorBase/DecoratorUID.h"
+#include "DecoratorBase/DecoratorRegistryHandle.h"
+
+#include <type_traits>
+
+namespace UE::AnimNext
+{
+	struct FDecoratorDescription;
+	struct FDecoratorInstance;
+	struct FNodeDescription;
+	struct FNodeInstance;
+
+	/**
+	 * Decorator Template
+	 * A decorator template represents a specific decorator within a FNodeTemplate.
+	 * 
+	 * @see FNodeTemplate
+	 */
+	struct FDecoratorTemplate
+	{
+		FDecoratorTemplate(FDecoratorUID UID_, FDecoratorRegistryHandle RegistryHandle_, EDecoratorMode Mode_, uint32 AdditiveIndexOrNumAdditive_, uint32 NodeSharedOffset_, uint32 NodeInstanceOffset_)
+			: UID(UID_.GetUID())
+			, RegistryHandle(RegistryHandle_)
+			, Mode(static_cast<uint8>(Mode_))
+			, AdditiveIndexOrNumAdditive(AdditiveIndexOrNumAdditive_)
+			, NodeSharedOffset(NodeSharedOffset_)
+			, NodeInstanceOffset(NodeInstanceOffset_)
+		{}
+
+		// Returns the globally unique identifier for the decorator
+		uint32 GetUID() const { return UID; }
+
+		// Returns the decorator registry handle
+		FDecoratorRegistryHandle GetRegistryHandle() const { return RegistryHandle; }
+
+		// Returns the decorator mode
+		EDecoratorMode GetMode() const { return static_cast<EDecoratorMode>(Mode); }
+
+		// For base decorators only, returns the number of additive decorators on top in the stack
+		uint32 GetNumAdditiveDecorators() const
+		{
+			check(GetMode() == EDecoratorMode::Base);
+			return AdditiveIndexOrNumAdditive;
+		}
+
+		// For additive decorators only, returns the decorator index relative to the base decorator on the stack
+		uint32 GetAdditiveDecoratorIndex() const
+		{
+			check(GetMode() == EDecoratorMode::Additive);
+			return AdditiveIndexOrNumAdditive;
+		}
+
+		// Returns the offset into the shared data where the descriptor begins, relative to the root of the node's shared data
+		uint32 GetNodeSharedOffset() const { return NodeSharedOffset; }
+
+		// Returns the offset into the instance data where the descriptor begins, relative to the root of the node's instance data
+		uint32 GetNodeInstanceOffset() const { return NodeInstanceOffset; }
+
+		// Returns a pointer to the specified decorator description on the current node
+		FDecoratorDescription* GetDecoratorDescription(FNodeDescription& NodeDescription) const
+		{
+			return reinterpret_cast<FDecoratorDescription*>(reinterpret_cast<uint8*>(&NodeDescription) + GetNodeSharedOffset());
+		}
+
+		// Returns a pointer to the specified decorator description on the current node
+		const FDecoratorDescription* GetDecoratorDescription(const FNodeDescription& NodeDescription) const
+		{
+			return reinterpret_cast<const FDecoratorDescription*>(reinterpret_cast<const uint8*>(&NodeDescription) + GetNodeSharedOffset());
+		}
+
+		// Returns a pointer to the specified decorator instance on the current node
+		FDecoratorInstance* GetDecoratorInstance(FNodeInstance& NodeInstance) const
+		{
+			return reinterpret_cast<FDecoratorInstance*>(reinterpret_cast<uint8*>(&NodeInstance) + GetNodeInstanceOffset());
+		}
+
+		// Returns a pointer to the specified decorator instance on the current node
+		const FDecoratorInstance* GetDecoratorInstance(const FNodeInstance& NodeInstance) const
+		{
+			return reinterpret_cast<const FDecoratorInstance*>(reinterpret_cast<const uint8*>(&NodeInstance) + GetNodeInstanceOffset());
+		}
+
+	private:
+		uint32	UID;							// decorator globally unique identifier
+
+		FDecoratorRegistryHandle	RegistryHandle;	// decorator registry handle
+		uint8						Mode;			// the decorator mode (we only need 1 bit)
+
+		// For base decorators, this contains the number of additive decorators on top
+		// For additive decorators, this contains its index relative to the base decorator
+		// The first additive decorator has index 1, there is no index 0 (we re-purpose it for the base decorator)
+		uint8	AdditiveIndexOrNumAdditive;
+
+		// Offsets into the shared read-only and instance data portions of a node
+		uint16	NodeSharedOffset;				// relative to root of node description data (max 64 KB per node)
+		uint16	NodeInstanceOffset;				// relative to root of node instance data (max 64 KB per node)
+
+		// TODO: We could cache which parent/base decorator handles which interface
+		// This would avoid the need to iterate on every decorator to look up a 'Super'. Perhaps only common interfaces could be cached.
+	};
+
+	static_assert(std::is_trivially_copyable_v<FDecoratorTemplate>, "FDecoratorTemplate needs to be trivially copyable");
+}
diff --git a/Engine/Plugins/Experimental/AnimNext/Source/AnimNext/Internal/DecoratorBase/NodeTemplate.h b/Engine/Plugins/Experimental/AnimNext/Source/AnimNext/Internal/DecoratorBase/NodeTemplate.h
new file mode 100644
--- /dev/null
+++ b/Engine/Plugins/Experimental/AnimNext/Source/AnimNext/Internal/DecoratorBase/NodeTemplate.h
@@ -0,0 +1,63 @@
+// Copyright Epic Games, Inc. All Rights Reserved.
+
+#pragma once
+
+#include "CoreMinimal.h"
+#include "DecoratorBase/DecoratorTemplate.h"
+
+#include <type_traits>
+
+namespace UE::AnimNext
+{
+	/**
+	 * Node Template
+	 * A node template represents a specific combination and ordering of decorators on a node.
+	 * For example, a sequence player + sync group decorator + output caching decorator
+	 * 
+	 * Multiple node descriptions can share the same template each with their own property/pin values.
+	 * 
+	 * Templates can exist in multiple sub-graphs on disk. On load, a single unique copy is retained in
+	 * memory and shared between ALL graphs.
+	 * 
+	 * @see FNodeDescription
+	 */
+	struct FNodeTemplate
+	{
+		FNodeTemplate(uint32 UID_, uint32 InstanceSize_, uint32 NumDecorators_)
+			: UID(UID_)
+			, InstanceSize(InstanceSize_)
+			, NumDecorators(NumDecorators_)
+		{}
+
+		// Returns the globally unique template identifier (a hash of all the decorator UIDs present in this node)
+		uint32 GetUID() const { return UID; }
+
+		// Returns the size in bytes of this node template
+		uint32 GetNodeTemplateSize() const { return sizeof(FNodeTemplate) + NumDecorators * sizeof(FDecoratorTemplate); }
+
+		// Returns the size in bytes of a node instance
+		uint32 GetInstanceSize() const { return InstanceSize; }
+
+		// Returns the number of decorators present in the node template
+		uint32 GetNumDecorators() const { return NumDecorators; }
+
+		// Returns a pointer to the list of decorator template descriptions
+		const FDecoratorTemplate* GetDecorators() const { return reinterpret_cast<const FDecoratorTemplate*>(reinterpret_cast<const uint8*>(this) + sizeof(FNodeTemplate)); }
+
+	private:
+		uint32	UID;				// globally unique template identifier or hash
+
+		uint16	InstanceSize;		// size in bytes of a node instance
+
+		uint8	NumDecorators;
+		uint8	Padding[1];
+
+		// TODO: We could use the padding (and extra space) here to cache which decorator handles which interface
+		// This would avoid the need to iterate on every decorator to look it up. Perhaps only common interfaces could be cached.
+
+		// Followed by a list of [FDecoratorTemplate] instances
+	};
+
+	static_assert(std::is_trivially_copyable_v<FNodeTemplate>, "FNodeTemplate needs to be trivially copyable");
+	static_assert(alignof(FNodeTemplate) == alignof(FDecoratorTemplate), "FNodeTemplate and FDecoratorTemplate must have the same alignment");
+}
diff --git a/Engine/Plugins/Experimental/AnimNext/Source/AnimNext/Internal/DecoratorBase/NodeTemplateRegistry.h b/Engine/Plugins/Experimental/AnimNext/Source/AnimNext/Internal/DecoratorBase/NodeTemplateRegistry.h
new file mode 100644
--- /dev/null
+++ b/Engine/Plugins/Experimental/AnimNext/Source/AnimNext/Internal/DecoratorBase/NodeTemplateRegistry.h
@@ -0,0 +1,55 @@
+// Copyright Epic Games, Inc. All Rights Reserved.
+
+#pragma once
+
+#include "CoreMinimal.h"
+#include "DecoratorBase/NodeTemplateRegistryHandle.h"
+
+namespace UE::AnimNext
+{
+	struct FNodeTemplate;
+
+	/**
+	 * FNodeTemplateRegistry
+	 * 
+	 * A global registry of all existing node templates that can be shared between animation graph instances.
+	 * 
+	 * @see FNodeTemplate
+	 */
+	struct ANIMNEXT_API FNodeTemplateRegistry final
+	{
+		// Access the global registry
+		static FNodeTemplateRegistry& Get();
+
+		// Finds or adds the specified node template and returns its offset
+		FNodeTemplateRegistryHandle FindOrAdd(const FNodeTemplate* NodeTemplate);
+
+		// Removes the specified node template from the registry
+		void Unregister(const FNodeTemplate* NodeTemplate);
+
+		// Finds and returns a node template based on its handle or nullptr if the handle is invalid
+		const FNodeTemplate* Find(FNodeTemplateRegistryHandle TemplateHandle) const;
+
+		// Returns the number of registered node templates
+		uint32 GetNum() const;
+
+	private:
+		FNodeTemplateRegistry() = default;
+
+		// Module lifetime functions
+		static void Init();
+		static void Destroy();
+
+		// We only ever append node templates to this contiguous buffers
+		// This is an optimization, we share these and by being contiguous
+		// we improve cache locality and cache line density
+		// Because node templates are trivially copyable, we could remove
+		// from this buffer when there are holes and coalesce everything.
+		// To do so, we would have to fix-up any outstanding handles within
+		// the shared data of loaded anim graphs.
+		TArray<uint8>								TemplateBuffer;
+		TMap<uint32, FNodeTemplateRegistryHandle>	TemplateUIDToHandleMap;
+
+		friend class FModule;
+	};
+}
diff --git a/Engine/Plugins/Experimental/AnimNext/Source/AnimNext/Private/DecoratorBase/Decorator.cpp b/Engine/Plugins/Experimental/AnimNext/Source/AnimNext/Private/DecoratorBase/Decorator.cpp
new file mode 100644
--- /dev/null
+++ b/Engine/Plugins/Experimental/AnimNext/Source/AnimNext/Private/DecoratorBase/Decorator.cpp
@@ -0,0 +1,19 @@
+// Copyright Epic Games, Inc. All Rights Reserved.
+
+#include "DecoratorBase/Decorator.h"
+
+#include "DecoratorBase/DecoratorRegistry.h"
+
+namespace UE::AnimNext
+{
+	FDecoratorStaticInitHook::FDecoratorStaticInitHook(DecoratorConstructorFunc DecoratorConstructor_)
+		: DecoratorConstructor(DecoratorConstructor_)
+	{
+		FDecoratorRegistry::StaticRegister(DecoratorConstructor_);
+	}
+
+	FDecoratorStaticInitHook::~FDecoratorStaticInitHook()
+	{
+		FDecoratorRegistry::StaticUnregister(DecoratorConstructor);
+	}
+}
diff --git a/Engine/Plugins/Experimental/AnimNext/Source/AnimNext/Private/DecoratorBase/DecoratorBinding.cpp b/Engine/Plugins/Experimental/AnimNext/Source/AnimNext/Private/DecoratorBase/DecoratorBinding.cpp
new file mode 100644
--- /dev/null
+++ b/Engine/Plugins/Experimental/AnimNext/Source/AnimNext/Private/DecoratorBase/DecoratorBinding.cpp
@@ -0,0 +1,13 @@
+// Copyright Epic Games, Inc. All Rights Reserved.
+
+#include "DecoratorBase/DecoratorBinding.h"
+
+#include "DecoratorBase/IDecoratorInterface.h"
+
+namespace UE::AnimNext
+{
+	FDecoratorInterfaceUID FDecoratorBinding::GetInterfaceUID() const
+	{
+		return Interface != nullptr ? Interface->GetInterfaceUID() : FDecoratorInterfaceUID();
+	}
+}
diff --git a/Engine/Plugins/Experimental/AnimNext/Source/AnimNext/Private/DecoratorBase/DecoratorPtr.cpp b/Engine/Plugins/Experimental/AnimNext/Source/AnimNext/Private/DecoratorBase/DecoratorPtr.cpp
new file mode 100644
--- /dev/null
+++ b/Engine/Plugins/Experimental/AnimNext/Source/AnimNext/Private/DecoratorBase/DecoratorPtr.cpp
@@ -0,0 +1,120 @@
+// Copyright Epic Games, Inc. All Rights Reserved.
+
+#include "DecoratorBase/DecoratorPtr.h"
+#include "DecoratorBase/NodeInstance.h"
+
+namespace UE::AnimNext
+{
+	FDecoratorPtr::FDecoratorPtr(FNodeInstance* NodeInstance, uint32 DecoratorIndex_)
+		: PackedPointerAndFlags(NodeInstance != nullptr ? reinterpret_cast<uintptr_t>(NodeInstance) : 0)
+		, DecoratorIndex(static_cast<uint8>(DecoratorIndex_))
+	{
+		check((reinterpret_cast<uintptr_t>(NodeInstance) & FLAGS_MASK) == 0);	// Make sure we have enough alignment
+		check(DecoratorIndex == DecoratorIndex_);	// Make sure we didn't truncate
+
+		if (NodeInstance != nullptr)
+		{
+			NodeInstance->AddReference();
+		}
+	}
+
+	FDecoratorPtr::FDecoratorPtr(const FDecoratorPtr& DecoratorPtr)
+		: PackedPointerAndFlags(DecoratorPtr.PackedPointerAndFlags)
+		, DecoratorIndex(DecoratorPtr.DecoratorIndex)
+	{
+		// Only increment the reference count if we aren't a weak handle
+		// The new handle will remain weak
+		if (!DecoratorPtr.IsWeak())
+		{
+			if (FNodeInstance* Node = DecoratorPtr.GetNodeInstance())
+			{
+				Node->AddReference();
+			}
+		}
+	}
+
+	FDecoratorPtr::FDecoratorPtr(FDecoratorPtr&& DecoratorPtr) noexcept
+		: PackedPointerAndFlags(DecoratorPtr.PackedPointerAndFlags)
+		, DecoratorIndex(DecoratorPtr.DecoratorIndex)
+	{
+		DecoratorPtr.PackedPointerAndFlags = 0;
+		DecoratorPtr.DecoratorIndex = 0;
+	}
+
+	FDecoratorPtr::FDecoratorPtr(FNodeInstance* NodeInstance, EFlags Flags, uint32 DecoratorIndex_)
+		: PackedPointerAndFlags(NodeInstance != nullptr ? (reinterpret_cast<uintptr_t>(NodeInstance) | Flags) : 0)
+		, DecoratorIndex(static_cast<uint8>(DecoratorIndex_))
+	{
+		check((reinterpret_cast<uintptr_t>(NodeInstance) & FLAGS_MASK) == 0);	// Make sure we have enough alignment
+		check(DecoratorIndex == DecoratorIndex_);	// Make sure we didn't truncate
+
+		// Only increment the reference count if we aren't a weak handle
+		if (NodeInstance != nullptr && (Flags & IS_WEAK_BIT) == 0)
+		{
+			NodeInstance->AddReference();
+		}
+	}
+
+	FDecoratorPtr::~FDecoratorPtr()
+	{
+		Reset();
+	}
+
+	FDecoratorPtr& FDecoratorPtr::operator=(const FDecoratorPtr& DecoratorPtr)
+	{
+		// Add our new reference first in case this == DecoratorPtr
+
+		// Only increment the reference count if we aren't a weak handle
+		// The new handle will remain weak
+		if (!DecoratorPtr.IsWeak())
+		{
+			if (FNodeInstance* Node = DecoratorPtr.GetNodeInstance())
+			{
+				Node->AddReference();
+			}
+		}
+
+		// Only decrement the reference count if we aren't a weak handle
+		if (!IsWeak())
+		{
+			if (FNodeInstance* Node = GetNodeInstance())
+			{
+				Node->ReleaseReference();
+			}
+		}
+
+		PackedPointerAndFlags = DecoratorPtr.PackedPointerAndFlags;
+		DecoratorIndex = DecoratorPtr.DecoratorIndex;
+
+		return *this;
+	}
+
+	FDecoratorPtr& FDecoratorPtr::operator=(FDecoratorPtr&& DecoratorPtr) noexcept
+	{
+		Swap(PackedPointerAndFlags, DecoratorPtr.PackedPointerAndFlags);
+		Swap(DecoratorIndex, DecoratorPtr.DecoratorIndex);
+
+		return *this;
+	}
+
+	void FDecoratorPtr::Reset()
+	{
+		// Only decrement the reference count if we aren't a weak handle
+		if (!IsWeak())
+		{
+			if (FNodeInstance* Node = GetNodeInstance())
+			{
+				Node->ReleaseReference();
+			}
+		}
+
+		PackedPointerAndFlags = 0;
+		DecoratorIndex = 0;
+	}
+
+	void FWeakDecoratorPtr::Reset()
+	{
+		NodeInstance = nullptr;
+		DecoratorIndex = 0;
+	}
+}
diff --git a/Engine/Plugins/Experimental/AnimNext/Source/AnimNext/Private/DecoratorBase/DecoratorRegistry.cpp b/Engine/Plugins/Experimental/AnimNext/Source/AnimNext/Private/DecoratorBase/DecoratorRegistry.cpp
new file mode 100644
--- /dev/null
+++ b/Engine/Plugins/Experimental/AnimNext/Source/AnimNext/Private/DecoratorBase/DecoratorRegistry.cpp
@@ -0,0 +1,329 @@
+// Copyright Epic Games, Inc. All Rights Reserved.
+
+#include "DecoratorBase/DecoratorRegistry.h"
+
+#include "DecoratorBase/Decorator.h"
+
+namespace UE::AnimNext
+{
+	namespace Private
+	{
+		static FDecoratorRegistry* GDecoratorRegistry = nullptr;
+
+		static TArray<DecoratorConstructorFunc> GPendingRegistrationQueue;
+	}
+
+	FDecoratorRegistry& FDecoratorRegistry::Get()
+	{
+		checkf(Private::GDecoratorRegistry, TEXT("Decorator Registry is not instanced. It is only valid to access this while the engine module is loaded."));
+		return *Private::GDecoratorRegistry;
+	}
+
+	void FDecoratorRegistry::Init()
+	{
+		if (ensure(Private::GDecoratorRegistry == nullptr))
+		{
+			Private::GDecoratorRegistry = new FDecoratorRegistry();
+
+			// Register all our pending static init decorators
+			for (DecoratorConstructorFunc DecoratorConstructor : Private::GPendingRegistrationQueue)
+			{
+				Private::GDecoratorRegistry->AutoRegisterImpl(DecoratorConstructor);
+			}
+
+			// Reset the registration queue, it won't be used anymore now that the registry is up and ready
+			Private::GPendingRegistrationQueue.Empty(0);
+		}
+	}
+
+	void FDecoratorRegistry::Destroy()
+	{
+		if (ensure(Private::GDecoratorRegistry != nullptr))
+		{
+			TArray<FRegistryEntry> Entries;
+			Private::GDecoratorRegistry->DecoratorUIDToEntryMap.GenerateValueArray(Entries);
+
+			for (const FRegistryEntry& Entry : Entries)
+			{
+				Private::GDecoratorRegistry->Unregister(Entry.Decorator);
+			}
+
+			delete Private::GDecoratorRegistry;
+			Private::GDecoratorRegistry = nullptr;
+		}
+	}
+
+	void FDecoratorRegistry::StaticRegister(DecoratorConstructorFunc DecoratorConstructor)
+	{
+		if (Private::GDecoratorRegistry != nullptr)
+		{
+			// Registry is already up and running, use it
+			Private::GDecoratorRegistry->AutoRegisterImpl(DecoratorConstructor);
+		}
+		else
+		{
+			// Registry isn't ready yet, queue up our decorator
+			// Once Init() is called, our queue will be processed
+			Private::GPendingRegistrationQueue.Add(DecoratorConstructor);
+		}
+	}
+
+	void FDecoratorRegistry::StaticUnregister(DecoratorConstructorFunc DecoratorConstructor)
+	{
+		if (Private::GDecoratorRegistry != nullptr)
+		{
+			// Registry is already up and running, use it
+			Private::GDecoratorRegistry->AutoUnregisterImpl(DecoratorConstructor);
+		}
+		else
+		{
+			// Registry isn't ready yet or it got destroyed before the decorators are unregistering
+			const int32 DecoratorIndex = Private::GPendingRegistrationQueue.IndexOfByKey(DecoratorConstructor);
+			if (DecoratorIndex != INDEX_NONE)
+			{
+				Private::GPendingRegistrationQueue.RemoveAtSwap(DecoratorIndex);
+			}
+		}
+	}
+
+	void FDecoratorRegistry::AutoRegisterImpl(DecoratorConstructorFunc DecoratorConstructor)
+	{
+		// Grab the memory requirements of our decorator
+		FDecoratorMemoryLayout DecoratorMemoryRequirements;
+		DecoratorConstructor(nullptr, DecoratorMemoryRequirements);
+
+		// Align it and reserve space
+		const uint32 OldBufferOffset = StaticDecoratorBufferOffset;
+
+		uint8* DecoratorPtr = Align(&StaticDecoratorBuffer[OldBufferOffset], DecoratorMemoryRequirements.DecoratorAlignment);
+
+		const uint32 NewBufferOffset = (DecoratorPtr + DecoratorMemoryRequirements.DecoratorSize) - &StaticDecoratorBuffer[0];
+		const bool bFitsInStaticBuffer = NewBufferOffset <= STATIC_DECORATOR_BUFFER_SIZE;
+
+		FDecoratorRegistryHandle DecoratorHandle;
+		if (bFitsInStaticBuffer)
+		{
+			// This decorator fits in our buffer, add it
+			StaticDecoratorBufferOffset = NewBufferOffset;
+			DecoratorHandle = FDecoratorRegistryHandle::MakeStatic(DecoratorPtr - &StaticDecoratorBuffer[0]);
+		}
+		else
+		{
+			// We have too many static decorators, we should consider increasing the static buffer
+			// TODO: Warn
+			// Allocate the decorator on the heap instead
+			DecoratorPtr = static_cast<uint8*>(FMemory::Malloc(DecoratorMemoryRequirements.DecoratorSize, DecoratorMemoryRequirements.DecoratorAlignment));
+		}
+
+		FDecorator* Decorator = DecoratorConstructor(DecoratorPtr, DecoratorMemoryRequirements);
+		checkf((void*)DecoratorPtr == (void*)Decorator, TEXT("FDecorator 'this' should be where we specified"));
+
+		const FDecoratorUID DecoratorUID = Decorator->GetDecoratorUID();
+
+		if (ensure(!DecoratorUIDToEntryMap.Contains(DecoratorUID.GetUID())))
+		{
+			// This is a new decorator, we'll keep it
+			if (!bFitsInStaticBuffer)
+			{
+				// Find our dynamic decorator index
+				int32 DecoratorIndex;
+				if (DynamicDecoratorFreeIndexHead != INDEX_NONE)
+				{
+					// We already had a free index, grab it
+					DecoratorIndex = DynamicDecoratorFreeIndexHead;
+					DynamicDecoratorFreeIndexHead = DynamicDecorators[DynamicDecoratorFreeIndexHead];
+				}
+				else
+				{
+					// No free indices, allocate a new one
+					DecoratorIndex = DynamicDecorators.Add(reinterpret_cast<uintptr_t>(Decorator));
+				}
+
+				DecoratorHandle = FDecoratorRegistryHandle::MakeDynamic(DecoratorIndex);
+			}
+
+			DecoratorUIDToEntryMap.Add(DecoratorUID.GetUID(), FRegistryEntry{ Decorator, DecoratorConstructor, DecoratorHandle });
+		}
+		else
+		{
+			// We have already registered this decorator, destroy our temporary instance
+			Decorator->~FDecorator();
+
+			if (bFitsInStaticBuffer)
+			{
+				// We were in the static buffer, clear our entry
+				StaticDecoratorBufferOffset = OldBufferOffset;
+				FMemory::Memzero(&StaticDecoratorBuffer[OldBufferOffset], NewBufferOffset - OldBufferOffset);
+			}
+			else
+			{
+				// It isn't in the static buffer, free it
+				FMemory::Free(Decorator);
+				Decorator = nullptr;
+			}
+		}
+	}
+
+	void FDecoratorRegistry::AutoUnregisterImpl(DecoratorConstructorFunc DecoratorConstructor)
+	{
+		for (auto It = DecoratorUIDToEntryMap.CreateIterator(); It; ++It)
+		{
+			FRegistryEntry& Entry = It.Value();
+
+			if (Entry.DecoratorConstructor == DecoratorConstructor)
+			{
+				check(Entry.DecoratorHandle.IsValid());
+
+				// Destroy and release our decorator
+				// We always own auto-registered decorator instances
+				Entry.Decorator->~FDecorator();
+
+				if (Entry.DecoratorHandle.IsDynamic())
+				{
+					// It has been dynamically registered, free it
+					const int32 DecoratorIndex = Entry.DecoratorHandle.GetDynamicIndex();
+
+					DynamicDecorators[DecoratorIndex] = DynamicDecoratorFreeIndexHead;
+					DynamicDecoratorFreeIndexHead = DecoratorIndex;
+
+					FMemory::Free(Entry.Decorator);
+				}
+				else
+				{
+					// It was in the static buffer, we cannot reclaim the space easily, we'd have to add metadata to
+					// track unused chunks so that we could coalesce
+					if (DecoratorUIDToEntryMap.Num() == 1)
+					{
+						// Last static decorator is being removed, reclaims all the space
+						StaticDecoratorBufferOffset = 0;
+					}
+				}
+
+				DecoratorUIDToEntryMap.Remove(It.Key());
+
+				break;
+			}
+		}
+	}
+
+	FDecoratorRegistryHandle FDecoratorRegistry::FindHandle(uint32 DecoratorUID) const
+	{
+		if (const FRegistryEntry* Entry = DecoratorUIDToEntryMap.Find(DecoratorUID))
+		{
+			return Entry->DecoratorHandle;
+		}
+
+		// Decorator not found
+		return FDecoratorRegistryHandle();
+	}
+
+	FDecoratorRegistryHandle FDecoratorRegistry::FindHandle(FDecoratorUID DecoratorUID) const
+	{
+		return FindHandle(DecoratorUID.GetUID());
+	}
+
+	const FDecorator* FDecoratorRegistry::Find(FDecoratorRegistryHandle DecoratorHandle) const
+	{
+		if (!DecoratorHandle.IsValid())
+		{
+			return nullptr;
+		}
+
+		if (DecoratorHandle.IsStatic())
+		{
+			const int32 DecoratorOffset = DecoratorHandle.GetStaticOffset();
+			return reinterpret_cast<const FDecorator*>(&StaticDecoratorBuffer[DecoratorOffset]);
+		}
+		else
+		{
+			const int32 DecoratorIndex = DecoratorHandle.GetDynamicIndex();
+			return reinterpret_cast<const FDecorator*>(DynamicDecorators[DecoratorIndex]);
+		}
+	}
+
+	const FDecorator* FDecoratorRegistry::Find(FDecoratorUID DecoratorUID) const
+	{
+		const FDecoratorRegistryHandle DecoratorHandle = FindHandle(DecoratorUID);
+		return Find(DecoratorHandle);
+	}
+
+	void FDecoratorRegistry::Register(FDecorator* Decorator)
+	{
+		if (Decorator == nullptr)
+		{
+			return;
+		}
+
+		const FDecoratorUID DecoratorUID = Decorator->GetDecoratorUID();
+
+		if (ensure(!DecoratorUIDToEntryMap.Contains(DecoratorUID.GetUID())))
+		{
+			// This is a new decorator, we'll keep it
+			// Find our dynamic decorator index
+			int32 DecoratorIndex;
+			if (DynamicDecoratorFreeIndexHead != INDEX_NONE)
+			{
+				// We already had a free index, grab it
+				DecoratorIndex = DynamicDecoratorFreeIndexHead;
+				DynamicDecoratorFreeIndexHead = DynamicDecorators[DynamicDecoratorFreeIndexHead];
+
+				DynamicDecorators[DecoratorIndex] = reinterpret_cast<uintptr_t>(Decorator);
+			}
+			else
+			{
+				// No free indices, allocate a new one
+				DecoratorIndex = DynamicDecorators.Add(reinterpret_cast<uintptr_t>(Decorator));
+			}
+
+			FDecoratorRegistryHandle DecoratorHandle = FDecoratorRegistryHandle::MakeDynamic(DecoratorIndex);
+
+			DecoratorUIDToEntryMap.Add(DecoratorUID.GetUID(), FRegistryEntry{ Decorator, nullptr, DecoratorHandle });
+		}
+	}
+
+	void FDecoratorRegistry::Unregister(FDecorator* Decorator)
+	{
+		if (Decorator == nullptr)
+		{
+			return;
+		}
+
+		const FDecoratorUID DecoratorUID = Decorator->GetDecoratorUID();
+
+		if (FRegistryEntry* Entry = DecoratorUIDToEntryMap.Find(DecoratorUID.GetUID()))
+		{
+			check(Entry->DecoratorHandle.IsValid());
+
+			if (Entry->DecoratorHandle.IsDynamic())
+			{
+				// It has been dynamically registered, free it
+				const int32 DecoratorIndex = Entry->DecoratorHandle.GetDynamicIndex();
+
+				DynamicDecorators[DecoratorIndex] = DynamicDecoratorFreeIndexHead;
+				DynamicDecoratorFreeIndexHead = DecoratorIndex;
+			}
+
+			if (Entry->DecoratorConstructor != nullptr)
+			{
+				// We own this decorator instance, destroy it
+				Entry->Decorator->~FDecorator();
+
+				if (Entry->DecoratorHandle.IsDynamic())
+				{
+					FMemory::Free(Entry->Decorator);
+				}
+				else
+				{
+					// It was in the static buffer, we cannot reclaim the space
+				}
+			}
+
+			DecoratorUIDToEntryMap.Remove(DecoratorUID.GetUID());
+		}
+	}
+
+	uint32 FDecoratorRegistry::GetNum() const
+	{
+		return DecoratorUIDToEntryMap.Num();
+	}
+}
diff --git a/Engine/Plugins/Experimental/AnimNext/Source/AnimNext/Private/DecoratorBase/ExecutionContext.cpp b/Engine/Plugins/Experimental/AnimNext/Source/AnimNext/Private/DecoratorBase/ExecutionContext.cpp
new file mode 100644
--- /dev/null
+++ b/Engine/Plugins/Experimental/AnimNext/Source/AnimNext/Private/DecoratorBase/ExecutionContext.cpp
@@ -0,0 +1,313 @@
+// Copyright Epic Games, Inc. All Rights Reserved.
+
+#include "DecoratorBase/ExecutionContext.h"
+
+#include "DecoratorBase/Decorator.h"
+#include "DecoratorBase/DecoratorRegistry.h"
+#include "DecoratorBase/DecoratorTemplate.h"
+#include "DecoratorBase/NodeDescription.h"
+#include "DecoratorBase/NodeInstance.h"
+#include "DecoratorBase/NodeTemplate.h"
+#include "DecoratorBase/NodeTemplateRegistry.h"
+
+namespace UE::AnimNext
+{
+	namespace Private
+	{
+		static thread_local UE::AnimNext::FExecutionContext* GThreadLocalExecutionContext = nullptr;
+	}
+
+	FExecutionContext::FExecutionContext(TArrayView<const uint8> GraphSharedData_)
+		: GraphSharedData(GraphSharedData_)
+	{
+		// There can be only one execution context alive per thread
+		ensure(Private::GThreadLocalExecutionContext == nullptr);
+		Private::GThreadLocalExecutionContext = this;
+	}
+
+	FExecutionContext::~FExecutionContext()
+	{
+		// There can be only one execution context alive per thread
+		ensure(Private::GThreadLocalExecutionContext == this);
+		Private::GThreadLocalExecutionContext = nullptr;
+	}
+
+	FDecoratorPtr FExecutionContext::AllocateNodeInstance(FWeakDecoratorPtr ParentBinding, FDecoratorHandle ChildDecoratorHandle)
+	{
+		ensure(ChildDecoratorHandle.IsValid());
+		if (!ChildDecoratorHandle.IsValid())
+		{
+			return FDecoratorPtr();	// Attempting to allocate a node using an invalid decorator handle
+		}
+
+		const FNodeHandle ChildNodeHandle = ChildDecoratorHandle.GetNodeHandle();
+		const FNodeDescription& NodeDesc = GetNodeDescription(ChildNodeHandle);
+
+		const FNodeTemplate* NodeTemplate = GetNodeTemplate(NodeDesc);
+		ensure(NodeTemplate != nullptr);
+		if (NodeTemplate == nullptr)
+		{
+			return FDecoratorPtr();	// Node template wasn't found, node descriptor is perhaps corrupted
+		}
+
+		const uint32_t ChildDecoratorIndex = ChildDecoratorHandle.GetDecoratorIndex();
+
+		if (ChildDecoratorIndex >= NodeTemplate->GetNumDecorators())
+		{
+			return FDecoratorPtr();	// The requested decorator index doesn't exist on that node descriptor
+		}
+
+		// If the decorator we wish to allocate lives in the parent node, return a weak handle to it
+		// We use a weak handle to avoid issues when multiple base decorators live within the same node
+		// When this happens, a decorator can end up pointing to another within the same node causing
+		// the reference count to never reach zero when all other handles are released
+		if (FNodeInstance* ParentNodeInstance = ParentBinding.GetNodeInstance())
+		{
+			if (ParentNodeInstance->GetNodeHandle() == ChildNodeHandle)
+			{
+				return FDecoratorPtr(ParentNodeInstance, FDecoratorPtr::IS_WEAK_BIT, ChildDecoratorIndex);
+			}
+		}
+
+		// We need to allocate a new node instance
+		const FDecoratorTemplate* DecoratorDescs = NodeTemplate->GetDecorators();
+
+		const uint32 InstanceSize = NodeTemplate->GetInstanceSize();
+		uint8* NodeInstanceBuffer = reinterpret_cast<uint8*>(FMemory::Malloc(InstanceSize, 16));
+		FNodeInstance* NodeInstance = new(NodeInstanceBuffer) FNodeInstance(ChildNodeHandle);
+
+		// Start construction with the bottom decorator
+		const FDecoratorTemplate* StartDesc = DecoratorDescs;
+		const FDecoratorTemplate* EndDesc = DecoratorDescs + NodeTemplate->GetNumDecorators();
+
+		const FDecoratorTemplate* FailedDecoratorDesc = nullptr;
+		for (const FDecoratorTemplate* DecoratorDesc = StartDesc; DecoratorDesc != EndDesc; ++DecoratorDesc)
+		{
+			const FDecorator* Decorator = GetDecorator(*DecoratorDesc);
+			ensure(Decorator != nullptr);
+			if (Decorator == nullptr)
+			{
+				// Failed to find the matching decorator, did it get unregistered or is the decorator descriptor corrupted?
+				FailedDecoratorDesc = DecoratorDesc;
+				break;
+			}
+
+			const uint32 DecoratorIndex = DecoratorDesc - DecoratorDescs;
+
+			FWeakDecoratorPtr DecoratorPtr(NodeInstance, DecoratorIndex);
+
+			const FDecoratorDescription* SharedData = DecoratorDesc->GetDecoratorDescription(NodeDesc);
+			FDecoratorInstance* InstanceData = DecoratorDesc->GetDecoratorInstance(*NodeInstance);
+
+			Decorator->ConstructInstance(*this, DecoratorPtr, *SharedData, *InstanceData);
+		}
+
+		if (FailedDecoratorDesc != nullptr)
+		{
+			// We failed to construct our node instance, destroy it
+			// Start destruction with the decorator prior to the one that failed
+			StartDesc = FailedDecoratorDesc - 1;
+			EndDesc = DecoratorDescs - 1;
+			for (const FDecoratorTemplate* DecoratorDesc = StartDesc; DecoratorDesc != EndDesc; --DecoratorDesc)
+			{
+				const FDecorator* Decorator = GetDecorator(*DecoratorDesc);
+				const uint32 DecoratorIndex = DecoratorDesc - DecoratorDescs;
+
+				FWeakDecoratorPtr DecoratorPtr(NodeInstance, DecoratorIndex);
+
+				const FDecoratorDescription* SharedData = DecoratorDesc->GetDecoratorDescription(NodeDesc);
+				FDecoratorInstance* InstanceData = DecoratorDesc->GetDecoratorInstance(*NodeInstance);
+
+				Decorator->DestructInstance(*this, DecoratorPtr, *SharedData, *InstanceData);
+			}
+
+			FMemory::Free(NodeInstance);
+
+			return FDecoratorPtr();
+		}
+
+		return FDecoratorPtr(NodeInstance, ChildDecoratorIndex);
+	}
+
+	void FExecutionContext::ReleaseNodeInstance(FNodeInstance* NodeInstance)
+	{
+		ensure(NodeInstance != nullptr && NodeInstance->IsValid());
+		if (NodeInstance == nullptr || !NodeInstance->IsValid())
+		{
+			return;	// Invalid node instance provided
+		}
+
+		ensure(NodeInstance->GetReferenceCount() == 0);
+		if (NodeInstance->GetReferenceCount() != 0)
+		{
+			return;	// Node instance still has references, we can't release it
+		}
+
+		const FNodeDescription& NodeDesc = GetNodeDescription(NodeInstance->GetNodeHandle());
+
+		const FNodeTemplate* NodeTemplate = GetNodeTemplate(NodeDesc);
+		ensure(NodeTemplate != nullptr);
+		if (NodeTemplate == nullptr)
+		{
+			return;	// Node template wasn't found, node descriptor is perhaps corrupted
+		}
+
+		const FDecoratorTemplate* DecoratorDescs = NodeTemplate->GetDecorators();
+
+		// Start destruction with the top decorator
+		const FDecoratorTemplate* StartDesc = DecoratorDescs + NodeTemplate->GetNumDecorators() - 1;
+		const FDecoratorTemplate* EndDesc = DecoratorDescs - 1;
+		for (const FDecoratorTemplate* DecoratorDesc = StartDesc; DecoratorDesc != EndDesc; --DecoratorDesc)
+		{
+			const FDecorator* Decorator = GetDecorator(*DecoratorDesc);
+			if (ensure(Decorator != nullptr))
+			{
+				const uint32 DecoratorIndex = DecoratorDesc - DecoratorDescs;
+
+				FWeakDecoratorPtr DecoratorPtr(NodeInstance, DecoratorIndex);
+
+				const FDecoratorDescription* SharedData = DecoratorDesc->GetDecoratorDescription(NodeDesc);
+				FDecoratorInstance* InstanceData = DecoratorDesc->GetDecoratorInstance(*NodeInstance);
+
+				Decorator->DestructInstance(*this, DecoratorPtr, *SharedData, *InstanceData);
+			}
+		}
+
+		NodeInstance->~FNodeInstance();
+		FMemory::Free(NodeInstance);
+	}
+
+	bool FExecutionContext::GetInterfaceImpl(FDecoratorInterfaceUID InterfaceUID, FWeakDecoratorPtr DecoratorPtr, FDecoratorBinding& InterfaceBinding) const
+	{
+		if (!DecoratorPtr.IsValid())
+		{
+			return false;
+		}
+
+		FNodeInstance* NodeInstance = DecoratorPtr.GetNodeInstance();
+
+		const FNodeDescription& NodeDesc = GetNodeDescription(NodeInstance->GetNodeHandle());
+
+		const FNodeTemplate* NodeTemplate = GetNodeTemplate(NodeDesc);
+		ensure(NodeTemplate != nullptr);
+		if (NodeTemplate == nullptr)
+		{
+			return false;	// Node template wasn't found, node descriptor is perhaps corrupted
+		}
+
+		const FDecoratorTemplate* DecoratorDescs = NodeTemplate->GetDecorators();
+
+		// We only search within the partial stack of the provided decorator
+		const FDecoratorTemplate* CurrentDecoratorDesc = DecoratorDescs + DecoratorPtr.GetDecoratorIndex();
+		const FDecoratorTemplate* BaseDecoratorDesc = CurrentDecoratorDesc->GetMode() == EDecoratorMode::Base ? CurrentDecoratorDesc : (CurrentDecoratorDesc - CurrentDecoratorDesc->GetAdditiveDecoratorIndex());
+
+		// Start searching with the top additive decorator towards our base decorator
+		const FDecoratorTemplate* StartDesc = BaseDecoratorDesc + BaseDecoratorDesc->GetNumAdditiveDecorators();
+		const FDecoratorTemplate* EndDesc = BaseDecoratorDesc - 1;
+		for (const FDecoratorTemplate* DecoratorDesc = StartDesc; DecoratorDesc != EndDesc; --DecoratorDesc)
+		{
+			const FDecorator* Decorator = GetDecorator(*DecoratorDesc);
+			ensure(Decorator != nullptr);
+			if (Decorator == nullptr)
+			{
+				return false;	// Failed to find the matching decorator, did it get unregistered or is the decorator descriptor corrupted?
+			}
+
+			if (const IDecoratorInterface* Interface = Decorator->GetInterface(InterfaceUID))
+			{
+				const FDecoratorDescription* SharedData = DecoratorDesc->GetDecoratorDescription(NodeDesc);
+				FDecoratorInstance* InstanceData = DecoratorDesc->GetDecoratorInstance(*NodeInstance);
+
+				const uint32 DecoratorIndex = DecoratorDesc - DecoratorDescs;
+				FWeakDecoratorPtr InterfaceDecoratorPtr(NodeInstance, DecoratorIndex);
+
+				InterfaceBinding = FDecoratorBinding(Interface, SharedData, InstanceData, InterfaceDecoratorPtr);
+				return true;
+			}
+		}
+
+		// We failed to find a decorator that handles this interface
+		return false;
+	}
+
+	bool FExecutionContext::GetInterfaceSuperImpl(FDecoratorInterfaceUID InterfaceUID, FWeakDecoratorPtr DecoratorPtr, FDecoratorBinding& SuperBinding) const
+	{
+		if (!DecoratorPtr.IsValid())
+		{
+			return false;	// Decorator pointer isn't valid, can't find a 'super'
+		}
+
+		FNodeInstance* NodeInstance = DecoratorPtr.GetNodeInstance();
+
+		const FNodeDescription& NodeDesc = GetNodeDescription(NodeInstance->GetNodeHandle());
+
+		const FNodeTemplate* NodeTemplate = GetNodeTemplate(NodeDesc);
+		ensure(NodeTemplate != nullptr);
+		if (NodeTemplate == nullptr)
+		{
+			return false;	// Node template wasn't found, node descriptor is perhaps corrupted
+		}
+
+		const FDecoratorTemplate* DecoratorDescs = NodeTemplate->GetDecorators();
+
+		// We only search within the partial stack of the provided decorator
+		const FDecoratorTemplate* CurrentDecoratorDesc = DecoratorDescs + DecoratorPtr.GetDecoratorIndex();
+		if (CurrentDecoratorDesc->GetMode() == EDecoratorMode::Base)
+		{
+			return false;	// We've reached a base decorator, we don't forward to any parent decorator we might have
+		}
+
+		const FDecoratorTemplate* BaseDecoratorDesc = CurrentDecoratorDesc - CurrentDecoratorDesc->GetAdditiveDecoratorIndex();
+
+		// Start searching with the next additive decorator towards our base decorator
+		const FDecoratorTemplate* StartDesc = CurrentDecoratorDesc - 1;
+		const FDecoratorTemplate* EndDesc = BaseDecoratorDesc - 1;
+		for (const FDecoratorTemplate* DecoratorDesc = StartDesc; DecoratorDesc != EndDesc; --DecoratorDesc)
+		{
+			const FDecorator* Decorator = GetDecorator(*DecoratorDesc);
+			ensure(Decorator != nullptr);
+			if (Decorator == nullptr)
+			{
+				return false;	// Failed to find the matching decorator, did it get unregistered or is the decorator descriptor corrupted?
+			}
+
+			if (const IDecoratorInterface* Interface = Decorator->GetInterface(InterfaceUID))
+			{
+				const FDecoratorDescription* SharedData = DecoratorDesc->GetDecoratorDescription(NodeDesc);
+				FDecoratorInstance* InstanceData = DecoratorDesc->GetDecoratorInstance(*NodeInstance);
+
+				const uint32 DecoratorIndex = DecoratorDesc - DecoratorDescs;
+				FWeakDecoratorPtr SuperPtr(NodeInstance, DecoratorIndex);
+
+				SuperBinding = FDecoratorBinding(Interface, SharedData, InstanceData, SuperPtr);
+				return true;
+			}
+		}
+
+		// We failed to find a decorator that handles this interface
+		return false;
+	}
+
+	const FNodeDescription& FExecutionContext::GetNodeDescription(FNodeHandle NodeHandle) const
+	{
+		check(NodeHandle.IsValid());
+		return *reinterpret_cast<const FNodeDescription*>(&GraphSharedData[NodeHandle.GetSharedOffset()]);
+	}
+
+	const FNodeTemplate* FExecutionContext::GetNodeTemplate(const FNodeDescription& NodeDesc) const
+	{
+		check(NodeDesc.GetTemplateHandle().IsValid());
+		return FNodeTemplateRegistry::Get().Find(NodeDesc.GetTemplateHandle());
+	}
+
+	const FDecorator* FExecutionContext::GetDecorator(const FDecoratorTemplate& Template) const
+	{
+		check(Template.GetRegistryHandle().IsValid());
+		return FDecoratorRegistry::Get().Find(Template.GetRegistryHandle());
+	}
+
+	FExecutionContext* GetThreadExecutionContext()
+	{
+		return Private::GThreadLocalExecutionContext;
+	}
+}
diff --git a/Engine/Plugins/Experimental/AnimNext/Source/AnimNext/Private/DecoratorBase/ITraversalContext.cpp b/Engine/Plugins/Experimental/AnimNext/Source/AnimNext/Private/DecoratorBase/ITraversalContext.cpp
new file mode 100644
--- /dev/null
+++ b/Engine/Plugins/Experimental/AnimNext/Source/AnimNext/Private/DecoratorBase/ITraversalContext.cpp
@@ -0,0 +1,20 @@
+// Copyright Epic Games, Inc. All Rights Reserved.
+
+#include "DecoratorBase/ITraversalContext.h"
+
+#include "DecoratorBase/ExecutionContext.h"
+
+namespace UE::AnimNext
+{
+	FScopedTraversalContext::FScopedTraversalContext(FExecutionContext& ExecutionContext_, ITraversalContext& TraversalContext)
+		: ExecutionContext(ExecutionContext_)
+		, OldTraversalContext(ExecutionContext_.TraversalContext)
+	{
+		ExecutionContext_.TraversalContext = &TraversalContext;
+	}
+
+	FScopedTraversalContext::~FScopedTraversalContext()
+	{
+		ExecutionContext.TraversalContext = OldTraversalContext;
+	}
+}
diff --git a/Engine/Plugins/Experimental/AnimNext/Source/AnimNext/Private/DecoratorBase/NodeInstance.cpp b/Engine/Plugins/Experimental/AnimNext/Source/AnimNext/Private/DecoratorBase/NodeInstance.cpp
new file mode 100644
--- /dev/null
+++ b/Engine/Plugins/Experimental/AnimNext/Source/AnimNext/Private/DecoratorBase/NodeInstance.cpp
@@ -0,0 +1,27 @@
+// Copyright Epic Games, Inc. All Rights Reserved.
+
+#include "DecoratorBase/NodeInstance.h"
+#include "DecoratorBase/ExecutionContext.h"
+
+namespace UE::AnimNext
+{
+	FNodeInstance::FNodeInstance(FNodeHandle NodeHandle)
+		: ReferenceCount(0)
+		, NodeHandle(NodeHandle)
+	{
+	}
+
+	void FNodeInstance::ReleaseReference()
+	{
+		check(ReferenceCount > 0);
+
+		if (ReferenceCount-- == 1)
+		{
+			// The last reference has been released, we can now safely destroy this node
+			FExecutionContext* Context = GetThreadExecutionContext();
+			check(Context != nullptr);
+
+			Context->ReleaseNodeInstance(this);
+		}
+	}
+}
diff --git a/Engine/Plugins/Experimental/AnimNext/Source/AnimNext/Private/DecoratorBase/NodeTemplateRegistry.cpp b/Engine/Plugins/Experimental/AnimNext/Source/AnimNext/Private/DecoratorBase/NodeTemplateRegistry.cpp
new file mode 100644
--- /dev/null
+++ b/Engine/Plugins/Experimental/AnimNext/Source/AnimNext/Private/DecoratorBase/NodeTemplateRegistry.cpp
@@ -0,0 +1,92 @@
+// Copyright Epic Games, Inc. All Rights Reserved.
+
+#include "DecoratorBase/NodeTemplateRegistry.h"
+
+#include "DecoratorBase/DecoratorTemplate.h"
+#include "DecoratorBase/NodeTemplate.h"
+
+namespace UE::AnimNext
+{
+	namespace Private
+	{
+		static FNodeTemplateRegistry* GNodeTemplateRegistry = nullptr;
+	}
+
+	FNodeTemplateRegistry& FNodeTemplateRegistry::Get()
+	{
+		checkf(Private::GNodeTemplateRegistry, TEXT("Node Template Registry is not instanced. It is only valid to access this while the engine module is loaded."));
+		return *Private::GNodeTemplateRegistry;
+	}
+
+	void FNodeTemplateRegistry::Init()
+	{
+		if (ensure(Private::GNodeTemplateRegistry == nullptr))
+		{
+			Private::GNodeTemplateRegistry = new FNodeTemplateRegistry();
+		}
+	}
+
+	void FNodeTemplateRegistry::Destroy()
+	{
+		if (ensure(Private::GNodeTemplateRegistry != nullptr))
+		{
+			delete Private::GNodeTemplateRegistry;
+			Private::GNodeTemplateRegistry = nullptr;
+		}
+	}
+
+	FNodeTemplateRegistryHandle FNodeTemplateRegistry::FindOrAdd(const FNodeTemplate* NodeTemplate)
+	{
+		if (NodeTemplate == nullptr)
+		{
+			return FNodeTemplateRegistryHandle();
+		}
+
+		const uint32 TemplateUID = NodeTemplate->GetUID();
+
+		if (FNodeTemplateRegistryHandle* TemplateHandle = TemplateUIDToHandleMap.Find(TemplateUID))
+		{
+			// We have already seen this template before, return its handle
+			return *TemplateHandle;
+		}
+
+		// This is a new template
+		const uint32 TemplateSize = NodeTemplate->GetNodeTemplateSize();
+
+		// Allocate and copy the template in our global buffer
+		const uint32 TemplateOffset = TemplateBuffer.AddUninitialized(TemplateSize);
+		FMemory::Memcpy(&TemplateBuffer[TemplateOffset], NodeTemplate, TemplateSize);
+
+		// Update our mapping
+		FNodeTemplateRegistryHandle TemplateHandle = FNodeTemplateRegistryHandle::MakeHandle(TemplateOffset);
+		TemplateUIDToHandleMap.Add(TemplateUID, TemplateHandle);
+
+		return TemplateHandle;
+	}
+
+	void FNodeTemplateRegistry::Unregister(const FNodeTemplate* NodeTemplate)
+	{
+		if (NodeTemplate == nullptr)
+		{
+			return;
+		}
+
+		const uint32 TemplateUID = NodeTemplate->GetUID();
+		TemplateUIDToHandleMap.Remove(TemplateUID);
+	}
+
+	const FNodeTemplate* FNodeTemplateRegistry::Find(FNodeTemplateRegistryHandle TemplateHandle) const
+	{
+		if (!TemplateHandle.IsValid())
+		{
+			return nullptr;
+		}
+
+		return reinterpret_cast<const FNodeTemplate*>(&TemplateBuffer[TemplateHandle.GetTemplateOffset()]);
+	}
+
+	uint32 FNodeTemplateRegistry::GetNum() const
+	{
+		return TemplateUIDToHandleMap.Num();
+	}
+}
diff --git a/Engine/Plugins/Experimental/AnimNext/Source/AnimNext/Private/DecoratorInterfaces/IEvaluate.cpp b/Engine/Plugins/Experimental/AnimNext/Source/AnimNext/Private/DecoratorInterfaces/IEvaluate.cpp
new file mode 100644
--- /dev/null
+++ b/Engine/Plugins/Experimental/AnimNext/Source/AnimNext/Private/DecoratorInterfaces/IEvaluate.cpp
@@ -0,0 +1,120 @@
+// Copyright Epic Games, Inc. All Rights Reserved.
+
+#include "DecoratorInterfaces/IEvaluate.h"
+
+#include "DecoratorBase/ExecutionContext.h"
+#include "DecoratorInterfaces/IHierarchy.h"
+
+namespace UE::AnimNext
+{
+	FEvaluateTraversalContext::FEvaluateTraversalContext(EEvaluationFlags EvaluationFlags_)
+		: EvaluationFlags(EvaluationFlags_)
+	{
+	}
+
+	void IEvaluate::PreEvaluate(FExecutionContext& Context, const TDecoratorBinding<IEvaluate>& Binding) const
+	{
+		TDecoratorBinding<IEvaluate> SuperBinding;
+		if (Context.GetInterfaceSuper(Binding, SuperBinding))
+		{
+			SuperBinding.PreEvaluate(Context);
+		}
+	}
+
+	void IEvaluate::PostEvaluate(FExecutionContext& Context, const TDecoratorBinding<IEvaluate>& Binding) const
+	{
+		TDecoratorBinding<IEvaluate> SuperBinding;
+		if (Context.GetInterfaceSuper(Binding, SuperBinding))
+		{
+			SuperBinding.PostEvaluate(Context);
+		}
+	}
+
+	enum class EEvaluateStep
+	{
+		PreEvaluate,
+		PostEvaluate,
+	};
+
+	struct FEvaluateEntry
+	{
+		FWeakDecoratorPtr	DecoratorPtr;
+		EEvaluateStep		DesiredStep = EEvaluateStep::PreEvaluate;
+	};
+
+	void EvaluateGraph(FExecutionContext& Context, FWeakDecoratorPtr GraphRootPtr, EEvaluationFlags EvaluationFlags, FPoseContainer& OutContainer)
+	{
+		// TODO: Reset container
+
+		if (!GraphRootPtr.IsValid())
+		{
+			return;	// Nothing to update
+		}
+
+		FMemMark Mark(FMemStack::Get());
+
+		TArray<FEvaluateEntry, TMemStackAllocator<>> NodesPendingUpdate;
+		NodesPendingUpdate.Reserve(64);
+
+		FChildrenArray Children;
+		Children.Reserve(64);
+
+		FEvaluateTraversalContext TraversalContext(EvaluationFlags);
+
+		FScopedTraversalContext ScopedTraversalContext(Context, TraversalContext);
+
+		// Add the graph root to kick start the evaluation process
+		NodesPendingUpdate.Push({ GraphRootPtr, EEvaluateStep::PreEvaluate });
+
+		// Process every node twice: pre-evaluate and post-evaluate
+		while (!NodesPendingUpdate.IsEmpty())
+		{
+			// Grab the top most entry
+			const bool bAllowShrinking = false;	// Don't allow shrinking to avoid churn
+			const FEvaluateEntry Entry = NodesPendingUpdate.Pop(bAllowShrinking);
+
+			if (Entry.DesiredStep == EEvaluateStep::PreEvaluate)
+			{
+				// This is the first time we visit this node, time to pre-evaluate
+				// Queue our node again so that post-evaluate is called afterwards
+				NodesPendingUpdate.Push({ Entry.DecoratorPtr, EEvaluateStep::PostEvaluate });
+
+				TDecoratorBinding<IEvaluate> EvaluateDecorator;
+				if (Context.GetInterface(Entry.DecoratorPtr, EvaluateDecorator))
+				{
+					EvaluateDecorator.PreEvaluate(Context);
+				}
+
+				TDecoratorBinding<IHierarchy> HierarchyDecorator;
+				if (Context.GetInterface(Entry.DecoratorPtr, HierarchyDecorator))
+				{
+					HierarchyDecorator.GetChildren(Context, Children);
+
+					// Append our children in reserve order so that they are visited in the same order they were added
+					for (int32 ChildIndex = Children.Num() - 1; ChildIndex >= 0; --ChildIndex)
+					{
+						NodesPendingUpdate.Push({ Children[ChildIndex], EEvaluateStep::PreEvaluate });
+					}
+
+					// Reset our container for the next time we need it
+					Children.Reset();
+				}
+
+				// Break and continue to the next top-most entry
+				// It is either a child ready for its pre-evaluate or our current entry ready for its post-evaluate (leaf)
+			}
+			else
+			{
+				// We've already visited this node once, time to post-evaluate
+				TDecoratorBinding<IEvaluate> EvaluateDecorator;
+				if (Context.GetInterface(Entry.DecoratorPtr, EvaluateDecorator))
+				{
+					EvaluateDecorator.PostEvaluate(Context);
+				}
+
+				// Break and continue to the next top-most entry
+				// It is either a sibling read for its post-evaluate or our parent entry ready for its post-evaluate
+			}
+		}
+	}
+}
diff --git a/Engine/Plugins/Experimental/AnimNext/Source/AnimNext/Private/DecoratorInterfaces/IHierarchy.cpp b/Engine/Plugins/Experimental/AnimNext/Source/AnimNext/Private/DecoratorInterfaces/IHierarchy.cpp
new file mode 100644
--- /dev/null
+++ b/Engine/Plugins/Experimental/AnimNext/Source/AnimNext/Private/DecoratorInterfaces/IHierarchy.cpp
@@ -0,0 +1,17 @@
+// Copyright Epic Games, Inc. All Rights Reserved.
+
+#include "DecoratorInterfaces/IHierarchy.h"
+
+#include "DecoratorBase/ExecutionContext.h"
+
+namespace UE::AnimNext
+{
+	void IHierarchy::GetChildren(FExecutionContext& Context, const TDecoratorBinding<IHierarchy>& Binding, FChildrenArray& Children) const
+	{
+		TDecoratorBinding<IHierarchy> SuperBinding;
+		if (Context.GetInterfaceSuper(Binding, SuperBinding))
+		{
+			SuperBinding.GetChildren(Context, Children);
+		}
+	}
+}
diff --git a/Engine/Plugins/Experimental/AnimNext/Source/AnimNext/Private/DecoratorInterfaces/ITimeline.cpp b/Engine/Plugins/Experimental/AnimNext/Source/AnimNext/Private/DecoratorInterfaces/ITimeline.cpp
new file mode 100644
--- /dev/null
+++ b/Engine/Plugins/Experimental/AnimNext/Source/AnimNext/Private/DecoratorInterfaces/ITimeline.cpp
@@ -0,0 +1,19 @@
+// Copyright Epic Games, Inc. All Rights Reserved.
+
+#include "DecoratorInterfaces/ITimeline.h"
+
+#include "DecoratorBase/ExecutionContext.h"
+
+namespace UE::AnimNext
+{
+	double ITimeline::GetPlayRate(FExecutionContext& Context, const TDecoratorBinding<ITimeline>& Binding) const
+	{
+		TDecoratorBinding<ITimeline> SuperBinding;
+		if (Context.GetInterfaceSuper(Binding, SuperBinding))
+		{
+			return SuperBinding.GetPlayRate(Context);
+		}
+
+		return 1.0;
+	}
+}
diff --git a/Engine/Plugins/Experimental/AnimNext/Source/AnimNext/Private/DecoratorInterfaces/IUpdate.cpp b/Engine/Plugins/Experimental/AnimNext/Source/AnimNext/Private/DecoratorInterfaces/IUpdate.cpp
new file mode 100644
--- /dev/null
+++ b/Engine/Plugins/Experimental/AnimNext/Source/AnimNext/Private/DecoratorInterfaces/IUpdate.cpp
@@ -0,0 +1,113 @@
+// Copyright Epic Games, Inc. All Rights Reserved.
+
+#include "DecoratorInterfaces/IUpdate.h"
+
+#include "DecoratorBase/ExecutionContext.h"
+#include "DecoratorInterfaces/IHierarchy.h"
+
+namespace UE::AnimNext
+{
+	void IUpdate::PreUpdate(FExecutionContext& Context, const TDecoratorBinding<IUpdate>& Binding) const
+	{
+		TDecoratorBinding<IUpdate> SuperBinding;
+		if (Context.GetInterfaceSuper(Binding, SuperBinding))
+		{
+			SuperBinding.PreUpdate(Context);
+		}
+	}
+
+	void IUpdate::PostUpdate(FExecutionContext& Context, const TDecoratorBinding<IUpdate>& Binding) const
+	{
+		TDecoratorBinding<IUpdate> SuperBinding;
+		if (Context.GetInterfaceSuper(Binding, SuperBinding))
+		{
+			SuperBinding.PostUpdate(Context);
+		}
+	}
+
+	enum class EUpdateStep
+	{
+		PreUpdate,
+		PostUpdate,
+	};
+
+	struct FUpdateEntry
+	{
+		FWeakDecoratorPtr	DecoratorPtr;
+		EUpdateStep			DesiredStep = EUpdateStep::PreUpdate;
+	};
+
+	void UpdateGraph(FExecutionContext& Context, FWeakDecoratorPtr GraphRootPtr)
+	{
+		if (!GraphRootPtr.IsValid())
+		{
+			return;	// Nothing to update
+		}
+
+		FMemMark Mark(FMemStack::Get());
+
+		TArray<FUpdateEntry, TMemStackAllocator<>> NodesPendingUpdate;
+		NodesPendingUpdate.Reserve(64);
+
+		FChildrenArray Children;
+		Children.Reserve(64);
+
+		FUpdateTraversalContext TraversalContext;
+
+		FScopedTraversalContext ScopedTraversalContext(Context, TraversalContext);
+
+		// Add the graph root to kick start the update process
+		NodesPendingUpdate.Push({ GraphRootPtr, EUpdateStep::PreUpdate });
+
+		// Process every node twice: pre-update and post-update
+		while (!NodesPendingUpdate.IsEmpty())
+		{
+			// Grab the top most entry
+			const bool bAllowShrinking = false;	// Don't allow shrinking to avoid churn
+			const FUpdateEntry Entry = NodesPendingUpdate.Pop(bAllowShrinking);
+
+			if (Entry.DesiredStep == EUpdateStep::PreUpdate)
+			{
+				// This is the first time we visit this node, time to pre-update
+				// Queue our node again so that post-update is called afterwards
+				NodesPendingUpdate.Push({ Entry.DecoratorPtr, EUpdateStep::PostUpdate });
+
+				TDecoratorBinding<IUpdate> UpdateDecorator;
+				if (Context.GetInterface(Entry.DecoratorPtr, UpdateDecorator))
+				{
+					UpdateDecorator.PreUpdate(Context);
+				}
+
+				TDecoratorBinding<IHierarchy> HierarchyDecorator;
+				if (Context.GetInterface(Entry.DecoratorPtr, HierarchyDecorator))
+				{
+					HierarchyDecorator.GetChildren(Context, Children);
+
+					// Append our children in reserve order so that they are visited in the same order they were added
+					for (int32 ChildIndex = Children.Num() - 1; ChildIndex >= 0; --ChildIndex)
+					{
+						NodesPendingUpdate.Push({ Children[ChildIndex], EUpdateStep::PreUpdate });
+					}
+
+					// Reset our container for the next time we need it
+					Children.Reset();
+				}
+
+				// Break and continue to the next top-most entry
+				// It is either a child ready for its pre-update or our current entry ready for its post-update (leaf)
+			}
+			else
+			{
+				// We've already visited this node once, time to post-update
+				TDecoratorBinding<IUpdate> UpdateDecorator;
+				if (Context.GetInterface(Entry.DecoratorPtr, UpdateDecorator))
+				{
+					UpdateDecorator.PostUpdate(Context);
+				}
+
+				// Break and continue to the next top-most entry
+				// It is either a sibling read for its post-update or our parent entry ready for its post-update
+			}
+		}
+	}
+}
diff --git a/Engine/Plugins/Experimental/AnimNext/Source/AnimNext/Private/Decorators/BlendTwoWay.cpp b/Engine/Plugins/Experimental/AnimNext/Source/AnimNext/Private/Decorators/BlendTwoWay.cpp
new file mode 100644
--- /dev/null
+++ b/Engine/Plugins/Experimental/AnimNext/Source/AnimNext/Private/Decorators/BlendTwoWay.cpp
@@ -0,0 +1,72 @@
+// Copyright Epic Games, Inc. All Rights Reserved.
+
+#include "Decorators/BlendTwoWay.h"
+
+#include "DecoratorBase/ExecutionContext.h"
+
+namespace UE::AnimNext
+{
+	DEFINE_ANIM_DECORATOR_BEGIN(FBlendTwoWayDecorator)
+		DEFINE_ANIM_DECORATOR_IMPLEMENTS_INTERFACE(IEvaluate)
+		DEFINE_ANIM_DECORATOR_IMPLEMENTS_INTERFACE(IUpdate)
+	DEFINE_ANIM_DECORATOR_END(FBlendTwoWayDecorator)
+
+	void FBlendTwoWayDecorator::PostEvaluate(FExecutionContext& Context, const TDecoratorBinding<IEvaluate>& Binding) const
+	{
+		// TODO:
+		// Now that our children have finished evaluating, we can pop their two poses and blend it
+		// Children execute in depth first order, as such the poses need to be popped in reverse order
+		// FPose&& pose1 = Context.PopPose();
+		// FPose&& pose0 = Context.PopPose();
+		// FPose resultPose = BlendTwoWay(pose0, pose1, blendWeight);
+		// Context.PushPose(resultPose);
+		// If we have full weight on a child, there is no work to do and we can early exit instead
+	}
+
+	void FBlendTwoWayDecorator::PreUpdate(FExecutionContext& Context, const TDecoratorBinding<IUpdate>& Binding) const
+	{
+		const FSharedData* SharedData = Binding.GetSharedData<FSharedData>();
+		FInstanceData* InstanceData = Binding.GetInstanceData<FInstanceData>();
+
+		// TODO: hook up blend weight to a pin, it could change every frame
+
+		if (SharedData->BlendWeight < 1.0)
+		{
+			if (!InstanceData->Children[0].IsValid())
+			{
+				// We need to blend a child that isn't instanced yet, allocate it
+				InstanceData->Children[0] = Context.AllocateNodeInstance(Binding, SharedData->Children[0]);
+			}
+
+			if (SharedData->BlendWeight == 0.0)
+			{
+				// We no longer need this child, release it
+				InstanceData->Children[1].Reset();
+			}
+		}
+
+		if (SharedData->BlendWeight > 0.0)
+		{
+			if (!InstanceData->Children[1].IsValid())
+			{
+				// We need to blend a child that isn't instanced yet, allocate it
+				InstanceData->Children[1] = Context.AllocateNodeInstance(Binding, SharedData->Children[1]);
+			}
+
+			if (SharedData->BlendWeight == 1.0)
+			{
+				// We no longer need this child, release it
+				InstanceData->Children[0].Reset();
+			}
+		}
+	}
+
+	void FBlendTwoWayDecorator::GetChildren(FExecutionContext& Context, const TDecoratorBinding<IHierarchy>& Binding, FChildrenArray& Children) const
+	{
+		const FInstanceData* InstanceData = Binding.GetInstanceData<FInstanceData>();
+
+		// Add the two child handles, even if they are empty
+		Children.Add(InstanceData->Children[0]);
+		Children.Add(InstanceData->Children[1]);
+	}
+}
diff --git a/Engine/Plugins/Experimental/AnimNext/Source/AnimNext/Private/Decorators/BlendTwoWay.h b/Engine/Plugins/Experimental/AnimNext/Source/AnimNext/Private/Decorators/BlendTwoWay.h
new file mode 100644
--- /dev/null
+++ b/Engine/Plugins/Experimental/AnimNext/Source/AnimNext/Private/Decorators/BlendTwoWay.h
@@ -0,0 +1,45 @@
+// Copyright Epic Games, Inc. All Rights Reserved.
+
+#pragma once
+
+#include "CoreMinimal.h"
+#include "DecoratorBase/Decorator.h"
+#include "DecoratorInterfaces/IEvaluate.h"
+#include "DecoratorInterfaces/IHierarchy.h"
+#include "DecoratorInterfaces/IUpdate.h"
+
+namespace UE::AnimNext
+{
+	/**
+	 * FBlendTwoWayDecorator
+	 * 
+	 * A decorator that can blend two inputs.
+	 */
+	struct FBlendTwoWayDecorator : FDecorator, IEvaluate, IUpdate, IHierarchy
+	{
+		DECLARE_ANIM_DECORATOR(FBlendTwoWayDecorator, 0x96a81d1e, FDecorator)
+
+		struct FSharedData : FDecorator::FSharedData
+		{
+			FDecoratorHandle Children[2];
+			double BlendWeight = 0.0;
+		};
+
+		struct FInstanceData : FDecorator::FInstanceData
+		{
+			FDecoratorPtr Children[2];
+		};
+
+		// FDecorator impl
+		virtual EDecoratorMode GetMode() const override { return EDecoratorMode::Base; }
+
+		// IEvaluate impl
+		virtual void PostEvaluate(FExecutionContext& Context, const TDecoratorBinding<IEvaluate>& Binding) const override;
+
+		// IUpdate impl
+		virtual void PreUpdate(FExecutionContext& Context, const TDecoratorBinding<IUpdate>& Binding) const override;
+
+		// IHierarchy impl
+		virtual void GetChildren(FExecutionContext& Context, const TDecoratorBinding<IHierarchy>& Binding, FChildrenArray& Children) const override;
+	};
+}
diff --git a/Engine/Plugins/Experimental/AnimNext/Source/AnimNext/Private/Decorators/SequencePlayer.cpp b/Engine/Plugins/Experimental/AnimNext/Source/AnimNext/Private/Decorators/SequencePlayer.cpp
new file mode 100644
--- /dev/null
+++ b/Engine/Plugins/Experimental/AnimNext/Source/AnimNext/Private/Decorators/SequencePlayer.cpp
@@ -0,0 +1,46 @@
+// Copyright Epic Games, Inc. All Rights Reserved.
+
+#include "Decorators/SequencePlayer.h"
+
+#include "DecoratorBase/ExecutionContext.h"
+
+namespace UE::AnimNext
+{
+	DEFINE_ANIM_DECORATOR_BEGIN(FSequencePlayerDecorator)
+		DEFINE_ANIM_DECORATOR_IMPLEMENTS_INTERFACE(IEvaluate)
+		DEFINE_ANIM_DECORATOR_IMPLEMENTS_INTERFACE(IUpdate)
+		DEFINE_ANIM_DECORATOR_IMPLEMENTS_INTERFACE(ITimeline)
+	DEFINE_ANIM_DECORATOR_END(FSequencePlayerDecorator)
+
+	void FSequencePlayerDecorator::PreEvaluate(FExecutionContext& Context, const TDecoratorBinding<IEvaluate>& Binding) const
+	{
+		// TODO: Sample pose
+		// FPose foo = ...
+		// Context.PushPose(foo);
+	}
+
+	double FSequencePlayerDecorator::GetPlayRate(FExecutionContext& Context, const TDecoratorBinding<ITimeline>& Binding) const
+	{
+		const FSharedData* SharedData = Binding.GetSharedData<FSharedData>();
+		return SharedData->PlayRate;
+	}
+
+	void FSequencePlayerDecorator::PreUpdate(FExecutionContext& Context, const TDecoratorBinding<IUpdate>& Binding) const
+	{
+		const FSharedData* SharedData = Binding.GetSharedData<FSharedData>();
+		FInstanceData* InstanceData = Binding.GetInstanceData<FInstanceData>();
+
+		TDecoratorBinding<ITimeline> TimelineDecorator;
+		Context.GetInterface(Binding, TimelineDecorator);
+
+		FUpdateTraversalContext& TraversalContext = Context.GetTraversalContext<FUpdateTraversalContext>();
+
+		double DeltaTime = TraversalContext.GetDeltaTime();
+
+		// Combine the execution play rate along with the timeline play rate, possibly overriden by another decorator
+		double PlayRate = TraversalContext.GetPlayRate() * TimelineDecorator.GetPlayRate(Context);
+
+		// TODO: need an actual impl, for illustrative purposes
+		InstanceData->CurrentTime += DeltaTime * PlayRate;
+	}
+}
diff --git a/Engine/Plugins/Experimental/AnimNext/Source/AnimNext/Private/Decorators/SequencePlayer.h b/Engine/Plugins/Experimental/AnimNext/Source/AnimNext/Private/Decorators/SequencePlayer.h
new file mode 100644
--- /dev/null
+++ b/Engine/Plugins/Experimental/AnimNext/Source/AnimNext/Private/Decorators/SequencePlayer.h
@@ -0,0 +1,49 @@
+// Copyright Epic Games, Inc. All Rights Reserved.
+
+#pragma once
+
+#include "CoreMinimal.h"
+#include "UObject/ObjectPtr.h"
+#include "DecoratorBase/Decorator.h"
+#include "DecoratorInterfaces/IEvaluate.h"
+#include "DecoratorInterfaces/ITimeline.h"
+#include "DecoratorInterfaces/IUpdate.h"
+
+namespace UE::AnimNext
+{
+	class UAnimSequence;
+
+	/**
+	 * FSequencePlayerDecorator
+	 * 
+	 * A decorator that can play an animation sequence.
+	 */
+	struct FSequencePlayerDecorator : FDecorator, IEvaluate, ITimeline, IUpdate
+	{
+		DECLARE_ANIM_DECORATOR(FSequencePlayerDecorator, 0xa628ad12, FDecorator)
+
+		struct FSharedData : FDecorator::FSharedData
+		{
+			TObjectPtr<UAnimSequence> AnimSeq;
+
+			double PlayRate = 1.0;
+		};
+
+		struct FInstanceData : FDecorator::FInstanceData
+		{
+			double CurrentTime = 0.0;
+		};
+
+		// FDecorator impl
+		virtual EDecoratorMode GetMode() const override { return EDecoratorMode::Base; }
+
+		// IEvaluate impl
+		virtual void PreEvaluate(FExecutionContext& Context, const TDecoratorBinding<IEvaluate>& Binding) const override;
+
+		// ITimeline impl
+		virtual double GetPlayRate(FExecutionContext& Context, const TDecoratorBinding<ITimeline>& Binding) const override;
+
+		// IUpdate impl
+		virtual void PreUpdate(FExecutionContext& Context, const TDecoratorBinding<IUpdate>& Binding) const override;
+	};
+}
diff --git a/Engine/Plugins/Experimental/AnimNext/Source/AnimNext/Private/Module.cpp b/Engine/Plugins/Experimental/AnimNext/Source/AnimNext/Private/Module.cpp
--- a/Engine/Plugins/Experimental/AnimNext/Source/AnimNext/Private/Module.cpp
+++ b/Engine/Plugins/Experimental/AnimNext/Source/AnimNext/Private/Module.cpp
@@ -5,6 +5,8 @@
 #include "Modules/ModuleManager.h"
 #include "Misc/CoreDelegates.h"
 #include "DataRegistry.h"
+#include "DecoratorBase/DecoratorRegistry.h"
+#include "DecoratorBase/NodeTemplateRegistry.h"
 
 namespace UE::AnimNext
 {
@@ -17,11 +19,15 @@ public:
 		FCoreDelegates::OnFEngineLoopInitComplete.AddLambda([]()
 		{
 			FDataRegistry::Init();
+			FDecoratorRegistry::Init();
+			FNodeTemplateRegistry::Init();
 		});
 
 		FCoreDelegates::OnEnginePreExit.AddLambda([]()
 		{
 			FDataRegistry::Destroy();
+			FDecoratorRegistry::Destroy();
+			FNodeTemplateRegistry::Destroy();
 		});
 		
 	}
diff --git a/Engine/Plugins/Experimental/AnimNext/Source/AnimNext/Public/DecoratorBase/Decorator.h b/Engine/Plugins/Experimental/AnimNext/Source/AnimNext/Public/DecoratorBase/Decorator.h
new file mode 100644
--- /dev/null
+++ b/Engine/Plugins/Experimental/AnimNext/Source/AnimNext/Public/DecoratorBase/Decorator.h
@@ -0,0 +1,199 @@
+// Copyright Epic Games, Inc. All Rights Reserved.
+
+#pragma once
+
+#include "CoreMinimal.h"
+#include "DecoratorBase/DecoratorDescription.h"
+#include "DecoratorBase/DecoratorHandle.h"			// Derived types are likely to refer to other decorators as children
+#include "DecoratorBase/DecoratorInstance.h"
+#include "DecoratorBase/DecoratorMode.h"
+#include "DecoratorBase/DecoratorPtr.h"
+#include "DecoratorBase/DecoratorUID.h"
+#include "DecoratorBase/IDecoratorInterface.h"
+
+// Helper macros
+// In the decorator class declaration, this macro declares the Super alias and base functions we override
+#define DECLARE_ANIM_DECORATOR(DecoratorName, DecoratorNameHash, SuperDecoratorName) \
+	using Super = SuperDecoratorName; \
+	/* FDecorator impl */ \
+	static constexpr UE::AnimNext::FDecoratorUID DecoratorUID = UE::AnimNext::FDecoratorUID(TEXT(#DecoratorName), DecoratorNameHash); \
+	virtual UE::AnimNext::FDecoratorUID GetDecoratorUID() const override { return DecoratorUID; } \
+	static const UE::AnimNext::FDecoratorMemoryLayout DecoratorMemoryDescription; \
+	virtual UE::AnimNext::FDecoratorMemoryLayout GetDecoratorMemoryDescription() const override { return DecoratorMemoryDescription; } \
+	virtual void ConstructInstance(UE::AnimNext::FExecutionContext& Context, UE::AnimNext::FWeakDecoratorPtr DecoratorPtr, const UE::AnimNext::FDecoratorDescription& DecoratorDesc, UE::AnimNext::FDecoratorInstance& DecoratorInstance) const override; \
+	virtual void DestructInstance(UE::AnimNext::FExecutionContext& Context, UE::AnimNext::FWeakDecoratorPtr DecoratorPtr, const UE::AnimNext::FDecoratorDescription& DecoratorDesc, UE::AnimNext::FDecoratorInstance& DecoratorInstance) const override; \
+	virtual const UE::AnimNext::IDecoratorInterface* GetInterface(UE::AnimNext::FDecoratorInterfaceUID InterfaceUID) const override;
+
+// In the decorator cpp, these three macros implement the base functionality
+// 
+// Usage is as follow:
+// DEFINE_ANIM_DECORATOR_BEGIN(FSequencePlayerDecorator)
+//     DEFINE_ANIM_DECORATOR_IMPLEMENTS_INTERFACE(IEvaluate)
+//     DEFINE_ANIM_DECORATOR_IMPLEMENTS_INTERFACE(IUpdate)
+//     DEFINE_ANIM_DECORATOR_IMPLEMENTS_INTERFACE(ITimeline)
+// DEFINE_ANIM_DECORATOR_END(FSequencePlayerDecorator)
+
+#define DEFINE_ANIM_DECORATOR_BEGIN(DecoratorName) \
+	const UE::AnimNext::FDecoratorMemoryLayout DecoratorName::DecoratorMemoryDescription = \
+		UE::AnimNext::FDecoratorMemoryLayout{ sizeof(DecoratorName), alignof(DecoratorName), sizeof(DecoratorName::FSharedData), alignof(DecoratorName::FSharedData), sizeof(DecoratorName::FInstanceData), alignof(DecoratorName::FInstanceData) }; \
+	void DecoratorName::ConstructInstance(UE::AnimNext::FExecutionContext& Context, UE::AnimNext::FWeakDecoratorPtr DecoratorPtr, const UE::AnimNext::FDecoratorDescription& DecoratorDesc, UE::AnimNext::FDecoratorInstance& DecoratorInstance) const \
+	{ \
+		FInstanceData* Data = new(&DecoratorInstance) FInstanceData(); \
+		Data->Construct(Context, DecoratorPtr, static_cast<const FSharedData&>(DecoratorDesc)); \
+	} \
+	void DecoratorName::DestructInstance(UE::AnimNext::FExecutionContext& Context, UE::AnimNext::FWeakDecoratorPtr DecoratorPtr, const UE::AnimNext::FDecoratorDescription& DecoratorDesc, UE::AnimNext::FDecoratorInstance& DecoratorInstance) const \
+	{ \
+		FInstanceData& Data = static_cast<FInstanceData&>(DecoratorInstance); \
+		Data.Destruct(Context, DecoratorPtr, DecoratorDesc); \
+		Data.~FInstanceData(); \
+	} \
+	const UE::AnimNext::IDecoratorInterface* DecoratorName::GetInterface(UE::AnimNext::FDecoratorInterfaceUID InterfaceUID_) const \
+	{
+
+#define DEFINE_ANIM_DECORATOR_END(DecoratorName) \
+		/* Forward to base implementation */ \
+		return Super::GetInterface(InterfaceUID_); \
+	}
+
+#define DEFINE_ANIM_DECORATOR_IMPLEMENTS_INTERFACE(InterfaceName) \
+	if (InterfaceUID_ == InterfaceName::InterfaceUID) \
+	{ \
+		return static_cast<const InterfaceName*>(this); \
+	}
+
+// Allows a decorator to auto-register and unregister within the current execution scope
+#define AUTO_REGISTER_ANIM_DECORATOR(DecoratorName) \
+	UE::AnimNext::FDecoratorStaticInitHook DecoratorName##Hook( \
+		[](void* DestPtr, FDecoratorMemoryLayout& MemoryDesc) -> FDecorator* \
+		{ \
+			MemoryDesc = DecoratorName::DecoratorMemoryDescription; \
+			return DestPtr != nullptr ? new(DestPtr) DecoratorName() : nullptr; \
+		});
+
+namespace UE::AnimNext
+{
+	struct FDecorator;
+	struct FDecoratorMemoryLayout;
+	struct FExecutionContext;
+
+	// A function pointer to a shim to construct a decorator into the desired memory location
+	// When called with a nullptr DestPtr, the function returns nullptr and only populates the
+	// memory description output argument. This allows the caller to determine how much space
+	// to reserve and how to properly align it. This is similar in spirit to various Windows SDK functions.
+	using DecoratorConstructorFunc = FDecorator* (*)(void* DestPtr, FDecoratorMemoryLayout& MemoryDesc);
+
+	/**
+	 * FDecoratorMemoryLayout
+	 * 
+	 * Encapsulates size/alignment details for a decorator.
+	 */
+	struct FDecoratorMemoryLayout
+	{
+		// The size in bytes of an instance of the decorator class which derives from FDecorator
+		uint32 DecoratorSize = 0;
+
+		// The alignment in bytes of an instance of the decorator class which derives from FDecorator
+		uint32 DecoratorAlignment = 0;
+
+		// The size in bytes of the shared data for the decorator
+		uint32 SharedDataSize = 0;
+
+		// The alignment in bytes of the shared data for the decorator
+		uint32 SharedDataAlignment = 0;
+
+		// The size in bytes of the instance data for the decorator
+		uint32 InstanceDataSize = 0;
+
+		// The alignment in bytes of the instance data for the decorator
+		uint32 InstanceDataAlignment = 0;
+	};
+
+	/**
+	 * FDecorator
+	 * 
+	 * Base class for all decorators.
+	 * A decorator can implement any number of interfaces based on IDecoratorInterface.
+	 * A decorator may derive from another decorator.
+	 * A decorator should implement GetInterface(..) and test against the interfaces that it supports.
+	 * 
+	 * Decorators should NOT have any internal state, hence why all API functions are 'const'.
+	 * The reason for this is that at runtime, a single instance of every decorator exists.
+	 * That single instance is used by all instances of a decorator on a node and concurrently
+	 * on all worker threads.
+	 * 
+	 * Decorators can have shared read-only data that all instances of a graph can use (e.g. hard-coded properties).
+	 * Shared data must derive from FDecoratorDescription.
+	 * Decorators can have instance data (e.g. blend weight).
+	 * Instance data must derive from FDecoratorInstance.
+	 */
+	struct ANIMNEXT_API FDecorator
+	{
+		// Empty shared/instance data types
+		// Derived types must define an alias for these
+		using FSharedData = FDecoratorDescription;
+		using FInstanceData = FDecoratorInstance;
+
+		// The globally unique UID for this decorator
+		// Derived types will have their own DecoratorUID member that hides/aliases/shadows this one
+		// @see DECLARE_ANIM_DECORATOR
+		static constexpr FDecoratorUID DecoratorUID = FDecoratorUID(TEXT("FDecorator"), 0x4fc735a2);
+
+		// Returns the globally unique UID for this decorator
+		virtual FDecoratorUID GetDecoratorUID() const { return DecoratorUID; };
+
+		// Returns the memory requirements of the derived decorator instance
+		virtual FDecoratorMemoryLayout GetDecoratorMemoryDescription() const { return { sizeof(FDecorator), alignof(FDecorator) }; }
+
+		// Called when a new instance of the decorator is created or destroyed
+		// Derived types must override this and forward to the instance data constructor/destructor
+		virtual void ConstructInstance(FExecutionContext& Context, FWeakDecoratorPtr DecoratorPtr, const FDecoratorDescription& DecoratorDesc, FDecoratorInstance& DecoratorInstance) const = 0;
+		virtual void DestructInstance(FExecutionContext& Context, FWeakDecoratorPtr DecoratorPtr, const FDecoratorDescription& DecoratorDesc, FDecoratorInstance& DecoratorInstance) const = 0;
+
+		// Returns the decorator mode for this derived type.
+		// TODO: This belongs in the editor side of the decorators since we need this information for our cook, here for now
+		virtual EDecoratorMode GetMode() const = 0;
+
+		// Returns a pointer to the specified interface if it is supported.
+		// Derived types must override this.
+		virtual const IDecoratorInterface* GetInterface(FDecoratorInterfaceUID InterfaceUID) const
+		{
+			// TODO:
+			// if/else sequence with static_casts to get the right v-table
+			// could be implemented with two tables: one of UIDs, another with matching offsets to 'this'
+			// we could scan the first table with SIMD, 4x UIDs at a time with 'cmpeq' to generate a mask
+			// we can mode the mask into a general register, if non-zero, we have a match
+			// using the mask, we can easily compute the UID offset in our 4x entry by counting leading/trailing zeroes
+			// using the UID offset, we can load and apply the correct offset
+			// may or may not be faster, but it shifts the burden from code cache to data cache and we can better control locality
+			// we could store the tables contiguous with one another, offsets could be 16 bit or maybe 8 bit (multiple of pointer size)
+			// we could store the tables contiguous with the tables of other decorators for better cache locality
+			// by using tables, it means the lookup code can live in a single place and remain hot
+			// it means we can test 4x UIDs at a time, or interleave and test 8x or 16x
+			// it means we can quickly early out if none of the interfaces match (common case?) since we don't need to test
+			// all of them one by one
+			// SIMD code path also opens the door for cheap bulk interface queries where we query up to 4x interface UIDs and
+			// return 4x interface offsets (caller can generate pointers easily)
+
+			// Base class doesn't implement any interfaces
+			// Derived types must implement this
+			return nullptr;
+		}
+
+		virtual ~FDecorator() {}
+	};
+
+	/**
+	 * FDecoratorStaticInitHook
+	 *
+	 * Allows decorators to automatically register/unregister within the current scope.
+	 * This can be used during static init.
+	 */
+	struct ANIMNEXT_API FDecoratorStaticInitHook final
+	{
+		explicit FDecoratorStaticInitHook(DecoratorConstructorFunc DecoratorConstructor_);
+		~FDecoratorStaticInitHook();
+
+	private:
+		DecoratorConstructorFunc DecoratorConstructor;
+	};
+}
diff --git a/Engine/Plugins/Experimental/AnimNext/Source/AnimNext/Public/DecoratorBase/DecoratorBinding.h b/Engine/Plugins/Experimental/AnimNext/Source/AnimNext/Public/DecoratorBase/DecoratorBinding.h
new file mode 100644
--- /dev/null
+++ b/Engine/Plugins/Experimental/AnimNext/Source/AnimNext/Public/DecoratorBase/DecoratorBinding.h
@@ -0,0 +1,89 @@
+// Copyright Epic Games, Inc. All Rights Reserved.
+
+#pragma once
+
+#include "CoreMinimal.h"
+#include "DecoratorBase/DecoratorPtr.h"
+#include "DecoratorBase/DecoratorInterfaceUID.h"
+
+namespace UE::AnimNext
+{
+	struct FDecoratorDescription;
+	struct FDecoratorInstance;
+	struct IDecoratorInterface;
+
+	/**
+	 * FDecoratorBinding
+	 * 
+	 * Base class for all decorator bindings.
+	 * A decorator binding contains untyped data about a specific decorator instance.
+	 */
+	struct ANIMNEXT_API FDecoratorBinding
+	{
+		// Creates an empty binding.
+		FDecoratorBinding() = default;
+
+		// Returns whether or not this binding is valid and set to an interface.
+		bool IsValid() const { return Interface != nullptr; }
+
+		// Queries a node for a pointer to its decorator shared data.
+		// If the decorator handle is invalid, a null pointer is returned.
+		template<class SharedDataType>
+		const SharedDataType* GetSharedData() const { return static_cast<const SharedDataType*>(SharedData); }
+
+		// Queries a node for a pointer to its decorator instance data.
+		// If the decorator handle is invalid, a null pointer is returned.
+		template<class InstanceDataType>
+		InstanceDataType* GetInstanceData() const { return static_cast<InstanceDataType*>(InstanceData); }
+
+		// Returns the decorator pointer we are bound to.
+		FWeakDecoratorPtr GetDecoratorPtr() const { return DecoratorPtr; }
+
+		// Returns the decorator interface UID when bound, an invalid UID otherwise.
+		FDecoratorInterfaceUID GetInterfaceUID() const;
+
+		// Equality and inequality tests
+		bool operator==(const FDecoratorBinding& RHS) const { return DecoratorPtr == RHS.DecoratorPtr && Interface == RHS.Interface; }
+		bool operator!=(const FDecoratorBinding& RHS) const { return DecoratorPtr != RHS.DecoratorPtr || Interface != RHS.Interface; }
+
+	protected:
+		// Creates a valid binding
+		FDecoratorBinding(const IDecoratorInterface* Interface_, const FDecoratorDescription* SharedData_, FDecoratorInstance* InstanceData_, FWeakDecoratorPtr DecoratorPtr_)
+			: Interface(Interface_)
+			, SharedData(SharedData_)
+			, InstanceData(InstanceData_)
+			, DecoratorPtr(DecoratorPtr_)
+		{}
+
+		// Performs a naked cast to the desired interface type
+		template<class DecoratorInterfaceType>
+		const DecoratorInterfaceType* GetInterfaceTyped() const { return static_cast<const DecoratorInterfaceType*>(Interface); }
+
+		const IDecoratorInterface*		Interface = nullptr;
+
+		const FDecoratorDescription*	SharedData = nullptr;
+		FDecoratorInstance*				InstanceData = nullptr;
+
+		FWeakDecoratorPtr				DecoratorPtr;
+
+		friend struct FExecutionContext;
+	};
+
+	/**
+	 * TDecoratorBinding
+	 * 
+	 * A templated proxy for decorator interfaces. It is meant to be specialized per interface
+	 * in order to allow a clean API and avoid human error. It wraps the necessary information
+	 * to bind a decorator to a specific interface. See existing interfaces for examples.
+	 * 
+	 * Here, we forward declare the template which every interface must specialize. Because we
+	 * rely on specialization, it must be defined within the UE::AnimNext namespace where the
+	 * declaration exists.
+	 * 
+	 * Specializations must derive from FDecoratorBinding to provide the necessary machinery.
+	 * 
+	 * @see IUpdate, IEvaluate, IHierarchy
+	 */
+	template<class DecoratorInterfaceType>
+	struct TDecoratorBinding;
+}
diff --git a/Engine/Plugins/Experimental/AnimNext/Source/AnimNext/Public/DecoratorBase/DecoratorDescription.h b/Engine/Plugins/Experimental/AnimNext/Source/AnimNext/Public/DecoratorBase/DecoratorDescription.h
new file mode 100644
--- /dev/null
+++ b/Engine/Plugins/Experimental/AnimNext/Source/AnimNext/Public/DecoratorBase/DecoratorDescription.h
@@ -0,0 +1,19 @@
+// Copyright Epic Games, Inc. All Rights Reserved.
+
+#pragma once
+
+#include "CoreMinimal.h"
+
+namespace UE::AnimNext
+{
+	/**
+	 * Decorator Description
+	 * A decorator description represents a unique instance in the authored static graph.
+	 * @see FNodeDescription
+	 * 
+	 * A decorator description is the base type that decorator shared data derives from.
+	 */
+	struct FDecoratorDescription
+	{
+	};
+}
diff --git a/Engine/Plugins/Experimental/AnimNext/Source/AnimNext/Public/DecoratorBase/DecoratorHandle.h b/Engine/Plugins/Experimental/AnimNext/Source/AnimNext/Public/DecoratorBase/DecoratorHandle.h
new file mode 100644
--- /dev/null
+++ b/Engine/Plugins/Experimental/AnimNext/Source/AnimNext/Public/DecoratorBase/DecoratorHandle.h
@@ -0,0 +1,52 @@
+// Copyright Epic Games, Inc. All Rights Reserved.
+
+#pragma once
+
+#include "CoreMinimal.h"
+#include "DecoratorBase/NodeHandle.h"
+
+namespace UE::AnimNext
+{
+	/**
+	 * Decorator Handle
+	 * A decorator handle represents a reference to a specific decorator instance in the shared/read-only portion
+	 * of a sub-graph. It points to a FNodeDescription when resolved.
+	 * @see FNodeDescription
+	 */
+	struct FDecoratorHandle
+	{
+		// Creates an invalid decorator handle
+		FDecoratorHandle()
+			: DecoratorIndex(0)
+			, SharedOffset(INVALID_SHARED_OFFSET_VALUE)
+		{}
+
+		// Creates a decorator handle pointing to the first decorator of the specified node
+		explicit FDecoratorHandle(FNodeHandle NodeHandle)
+			: DecoratorIndex(0)
+			, SharedOffset(NodeHandle.GetSharedOffset())
+		{}
+
+		// Creates a decorator handle pointing to the specified decorator on the specified node
+		FDecoratorHandle(FNodeHandle NodeHandle, uint32 DecoratorIndex_)
+			: DecoratorIndex(DecoratorIndex_)
+			, SharedOffset(NodeHandle.GetSharedOffset())
+		{}
+
+		// Returns true if this decorator handle is valid, false otherwise
+		bool IsValid() const { return SharedOffset != INVALID_SHARED_OFFSET_VALUE; }
+
+		// Returns the decorator index
+		uint32 GetDecoratorIndex() const { return DecoratorIndex; }
+
+		// Returns a handle to the node in the shared data segment
+		FNodeHandle GetNodeHandle() const { return FNodeHandle(SharedOffset); }
+
+	private:
+		// We pack the shared offset on 24 bits, the top 8 bits are truncated
+		static constexpr uint32 INVALID_SHARED_OFFSET_VALUE = ~0u >> (32 - 24);
+
+		uint32	DecoratorIndex : 8;
+		uint32	SharedOffset : 24;		// relative to root of sub-graph
+	};
+}
diff --git a/Engine/Plugins/Experimental/AnimNext/Source/AnimNext/Public/DecoratorBase/DecoratorInstance.h b/Engine/Plugins/Experimental/AnimNext/Source/AnimNext/Public/DecoratorBase/DecoratorInstance.h
new file mode 100644
--- /dev/null
+++ b/Engine/Plugins/Experimental/AnimNext/Source/AnimNext/Public/DecoratorBase/DecoratorInstance.h
@@ -0,0 +1,42 @@
+// Copyright Epic Games, Inc. All Rights Reserved.
+
+#pragma once
+
+#include "CoreMinimal.h"
+#include "DecoratorBase/DecoratorPtr.h"
+
+namespace UE::AnimNext
+{
+	struct FDecoratorDescription;
+	struct FExecutionContext;
+
+	/**
+	 * Decorator Instance
+	 * A decorator instance represents allocated data for specific decorator.
+	 * @see FNodeInstance
+	 * 
+	 * A decorator instance is the base type that decorator instance data derives from.
+	 */
+	struct FDecoratorInstance
+	{
+		// Called after the constructor has been called when a new instance is created.
+		// This is called after the default constructor.
+		// You can override this function by adding a new one with the same name on your
+		// derived type. You can also specialize the FDecoratorDescription with the derived
+		// version as long as the reference can coerce to 'const FDecoratorDescription&'.
+		// Decorators are constructed from the bottom to the top.
+		void Construct(FExecutionContext& Context, FWeakDecoratorPtr DecoratorPtr, const FDecoratorDescription& DecoratorDesc)
+		{
+		}
+
+		// Called before the destructor has been called when an instance is destroyed.
+		// This is called before the default destructor.
+		// You can override this function by adding a new one with the same name on your
+		// derived type. You can also specialize the FDecoratorDescription with the derived
+		// version as long as the reference can coerce to 'const FDecoratorDescription&'.
+		// Decorators are destructed from the top to the bottom.
+		void Destruct(FExecutionContext& Context, FWeakDecoratorPtr DecoratorPtr, const FDecoratorDescription& DecoratorDesc)
+		{
+		}
+	};
+}
diff --git a/Engine/Plugins/Experimental/AnimNext/Source/AnimNext/Public/DecoratorBase/DecoratorInterfaceUID.h b/Engine/Plugins/Experimental/AnimNext/Source/AnimNext/Public/DecoratorBase/DecoratorInterfaceUID.h
new file mode 100644
--- /dev/null
+++ b/Engine/Plugins/Experimental/AnimNext/Source/AnimNext/Public/DecoratorBase/DecoratorInterfaceUID.h
@@ -0,0 +1,59 @@
+// Copyright Epic Games, Inc. All Rights Reserved.
+
+#pragma once
+
+#include "CoreMinimal.h"
+
+namespace UE::AnimNext
+{
+	/**
+	 * FDecoratorInterfaceUID
+	 *
+	 * Encapsulates an interface global UID.
+	 * The string is exposed in non-shipping builds for logging and debugging purposes.
+	 * The UID should be generated from the provided string using FNV1a with 32 bits.
+	 *
+	 * The whole struct is meant to be 'constexpr' to allow inlining in interface queries.
+	 */
+	struct FDecoratorInterfaceUID final
+	{
+		// Constructs an invalid UID
+		constexpr FDecoratorInterfaceUID()
+			: UID(0)
+#if !UE_BUILD_SHIPPING
+			, InterfaceName(TEXT("<Invalid Interface UID>"))
+#endif
+		{}
+
+		// Constructs an interface UID
+		constexpr FDecoratorInterfaceUID(const TCHAR* InterfaceName_, uint32 UID_)
+			: UID(UID_)
+#if !UE_BUILD_SHIPPING
+			, InterfaceName(InterfaceName_)
+#endif
+		{
+		}
+
+#if !UE_BUILD_SHIPPING
+		// Returns a literal string to the interface name
+		constexpr const TCHAR* GetInterfaceName() const { return InterfaceName; }
+#endif
+
+		// Returns the interface global UID
+		constexpr uint32 GetUID() const { return UID; }
+
+		// Returns whether this UID is valid or not
+		constexpr bool IsValid() const { return UID != 0; }
+
+		// Compares for equality and inequality
+		constexpr bool operator==(FDecoratorInterfaceUID RHS) const { return UID == RHS.UID; }
+		constexpr bool operator!=(FDecoratorInterfaceUID RHS) const { return UID != RHS.UID; }
+
+	private:
+		uint32	UID;
+
+#if !UE_BUILD_SHIPPING
+		const TCHAR* InterfaceName;
+#endif
+	};
+}
diff --git a/Engine/Plugins/Experimental/AnimNext/Source/AnimNext/Public/DecoratorBase/DecoratorMode.h b/Engine/Plugins/Experimental/AnimNext/Source/AnimNext/Public/DecoratorBase/DecoratorMode.h
new file mode 100644
--- /dev/null
+++ b/Engine/Plugins/Experimental/AnimNext/Source/AnimNext/Public/DecoratorBase/DecoratorMode.h
@@ -0,0 +1,32 @@
+// Copyright Epic Games, Inc. All Rights Reserved.
+
+#pragma once
+
+#include "CoreMinimal.h"
+
+namespace UE::AnimNext
+{
+	/**
+	 * EDecoratorMode
+	 *
+	 * Describes how a decorator behaves once attached to an animation node.
+	 */
+	enum class EDecoratorMode
+	{
+		/**
+		 * Base decorators can live on their own in an animation node and have no 'Super'.
+		 * As a result, calls to GetInterface() do not forward to other decorators below
+		 * them on the node stack. Multiple base decorators can exist in a single animation
+		 * node but they behave as independent nodes (functionally speaking).
+		 */
+		Base,
+
+		/**
+		 * Additive decorators override or augment behavior on prior decorators on the node stack.
+		 * At least one base decorator must be present for an additive decorator to be added
+		 * on top. Calls to GetInterface() will pass-through to other decorators below on the stack
+		 * until a base decorator is found.
+		 */
+		Additive,
+	};
+}
diff --git a/Engine/Plugins/Experimental/AnimNext/Source/AnimNext/Public/DecoratorBase/DecoratorPtr.h b/Engine/Plugins/Experimental/AnimNext/Source/AnimNext/Public/DecoratorBase/DecoratorPtr.h
new file mode 100644
--- /dev/null
+++ b/Engine/Plugins/Experimental/AnimNext/Source/AnimNext/Public/DecoratorBase/DecoratorPtr.h
@@ -0,0 +1,137 @@
+// Copyright Epic Games, Inc. All Rights Reserved.
+
+#pragma once
+
+#include "CoreMinimal.h"
+
+namespace UE::AnimNext
+{
+	struct FNodeInstance;
+	struct FWeakDecoratorPtr;
+
+	/**
+	 * Decorator Pointer
+	 * A decorator pointer represents a shared pointer to allocated instance data.
+	 * It manages reference counting.
+	 * It points to a FDecoratorInstance or FNodeInstance when resolved.
+	 *
+	 * A node pointer can also be weak, meaning that it does not update reference counting.
+	 * Note that weak pointers will point to garbage if the decorator instance is destroyed.
+	 *
+	 * Weak decorator pointers should be used carefully!
+	 * 
+	 * @see FNodeInstance, FDecoratorInstance
+	 */
+	struct ANIMNEXT_API FDecoratorPtr
+	{
+		// Constructs an invalid pointer handle
+		constexpr FDecoratorPtr() noexcept = default;
+
+		// Constructs a pointer handle to the provided instance
+		FDecoratorPtr(FNodeInstance* NodeInstance, uint32 DecoratorIndex_);
+
+		FDecoratorPtr(const FDecoratorPtr& DecoratorPtr);
+		FDecoratorPtr(FDecoratorPtr&& DecoratorPtr) noexcept;
+		~FDecoratorPtr();
+
+		FDecoratorPtr& operator=(const FDecoratorPtr& DecoratorPtr);
+		FDecoratorPtr& operator=(FDecoratorPtr&& DecoratorPtr) noexcept;
+
+		// Returns a pointer to the node instance
+		FNodeInstance* GetNodeInstance() const noexcept { return reinterpret_cast<FNodeInstance*>(PackedPointerAndFlags & ~FLAGS_MASK); }
+
+		// Returns the decorator index this pointer handle references
+		constexpr uint32 GetDecoratorIndex() const noexcept { return DecoratorIndex; }
+
+		// Returns true when the pointer is valid, false otherwise
+		constexpr bool IsValid() const noexcept { return PackedPointerAndFlags != 0; }
+
+		// Returns true if this pointer handle is weak, false otherwise
+		constexpr bool IsWeak() const noexcept { return (PackedPointerAndFlags & IS_WEAK_BIT) != 0; }
+
+		// Clears the handle and renders it invalid
+		void Reset();
+
+		// Equality and inequality tests
+		bool operator==(const FDecoratorPtr& RHS) const { return GetNodeInstance() == RHS.GetNodeInstance() && DecoratorIndex == RHS.DecoratorIndex; }
+		bool operator!=(const FDecoratorPtr& RHS) const { return GetNodeInstance() != RHS.GetNodeInstance() || DecoratorIndex != RHS.DecoratorIndex; }
+		bool operator==(const FWeakDecoratorPtr& RHS) const;
+		bool operator!=(const FWeakDecoratorPtr& RHS) const;
+
+	private:
+		// Various flags stored in the pointer alignment bits, assumes an alignment of at least 4 bytes
+		enum EFlags : uintptr_t
+		{
+			// When true, this handle is weak and it will not update the instance reference count
+			IS_WEAK_BIT = 0x01,
+
+			// Bit mask to clear all flags from the packed pointer value
+			FLAGS_MASK = 0x03,
+		};
+
+		// Constructs a pointer handle to the provided instance
+		FDecoratorPtr(FNodeInstance* NodeInstance, EFlags Flags, uint32 DecoratorIndex_);
+
+		// Packed pointer value contains a pointer along with flags in the alignment bits
+		// The pointer part points to a node instance
+		uintptr_t	PackedPointerAndFlags = 0;
+		uint32		DecoratorIndex = 0;				// Only need 8 bits, but use 32 since we have ample padding, avoids truncation
+
+		friend struct FWeakDecoratorPtr;
+		friend struct FExecutionContext;
+	};
+
+	/**
+	 * Weak Decorator Pointer
+	 * Same as a FDecoratorPtr but is strongly typed to be weak.
+	 *
+	 * Weak decorator pointers should be used carefully!
+	 * 
+	 * @see FDecoratorPtr
+	 */
+	struct ANIMNEXT_API FWeakDecoratorPtr
+	{
+		// Constructs an invalid weak pointer handle
+		constexpr FWeakDecoratorPtr() noexcept = default;
+
+		// Constructs a weak pointer handle from a shared pointer handle
+		FWeakDecoratorPtr(const FDecoratorPtr& DecoratorPtr) noexcept
+			: NodeInstance(DecoratorPtr.GetNodeInstance())
+			, DecoratorIndex(DecoratorPtr.DecoratorIndex)
+		{}
+
+		FWeakDecoratorPtr(FNodeInstance* NodeInstance_, uint32 DecoratorIndex_)
+			: NodeInstance(NodeInstance_)
+			, DecoratorIndex(static_cast<uint8>(DecoratorIndex_))
+		{
+			check(DecoratorIndex == DecoratorIndex_);	// Make sure we didn't truncate
+		}
+
+		// Returns a pointer to the node instance
+		constexpr FNodeInstance* GetNodeInstance() const noexcept { return NodeInstance; }
+
+		// Returns the decorator index this pointer handle references
+		constexpr uint32 GetDecoratorIndex() const noexcept { return DecoratorIndex; }
+
+		// Returns true when the pointer is valid, false otherwise
+		constexpr bool IsValid() const noexcept { return NodeInstance != nullptr; }
+
+		// Clears the handle and renders it invalid
+		void Reset();
+
+		// Equality and inequality tests
+		bool operator==(const FWeakDecoratorPtr& RHS) const { return NodeInstance == RHS.NodeInstance && DecoratorIndex == RHS.DecoratorIndex; }
+		bool operator!=(const FWeakDecoratorPtr& RHS) const { return NodeInstance != RHS.NodeInstance || DecoratorIndex != RHS.DecoratorIndex; }
+		bool operator==(const FDecoratorPtr& RHS) const { return NodeInstance == RHS.GetNodeInstance() && DecoratorIndex == RHS.GetDecoratorIndex(); }
+		bool operator!=(const FDecoratorPtr& RHS) const { return NodeInstance != RHS.GetNodeInstance() || DecoratorIndex != RHS.GetDecoratorIndex(); }
+
+	private:
+		FNodeInstance*	NodeInstance = nullptr;
+		uint8			DecoratorIndex = 0;
+	};
+
+	//////////////////////////////////////////////////////////////////////////
+
+	inline bool FDecoratorPtr::operator==(const FWeakDecoratorPtr& RHS) const { return GetNodeInstance() == RHS.GetNodeInstance() && DecoratorIndex == RHS.GetDecoratorIndex(); }
+	inline bool FDecoratorPtr::operator!=(const FWeakDecoratorPtr& RHS) const { return GetNodeInstance() != RHS.GetNodeInstance() || DecoratorIndex != RHS.GetDecoratorIndex(); }
+}
diff --git a/Engine/Plugins/Experimental/AnimNext/Source/AnimNext/Public/DecoratorBase/DecoratorRegistryHandle.h b/Engine/Plugins/Experimental/AnimNext/Source/AnimNext/Public/DecoratorBase/DecoratorRegistryHandle.h
new file mode 100644
--- /dev/null
+++ b/Engine/Plugins/Experimental/AnimNext/Source/AnimNext/Public/DecoratorBase/DecoratorRegistryHandle.h
@@ -0,0 +1,66 @@
+// Copyright Epic Games, Inc. All Rights Reserved.
+
+#pragma once
+
+#include "CoreMinimal.h"
+
+namespace UE::AnimNext
+{
+	/**
+	 * FDecoratorRegistryHandle
+	 *
+	 * Encapsulates a value used as a handle in the decorator registry.
+	 * When valid, it can be used to retrieve a pointer to the corresponding decorator.
+	 */
+	struct FDecoratorRegistryHandle final
+	{
+		// Default constructed handles are invalid
+		FDecoratorRegistryHandle() = default;
+
+		// Returns whether or not this handle points to a valid decorator
+		bool IsValid() const { return HandleValue != 0; }
+
+		// Returns whether or not this handle is valid and points to a static decorator
+		bool IsStatic() const { return IsValid() && HandleValue > 0; }
+
+		// Returns whether or not this handle is valid and points to a dynamic decorator
+		bool IsDynamic() const { return IsValid() && HandleValue < 0; }
+
+		// Returns the static buffer offset for this handle when valid, otherwise INDEX_NONE
+		int32 GetStaticOffset() const { return IsStatic() ? (HandleValue - 1) : INDEX_NONE; }
+
+		// Returns the dynamic array index for this handle when valid, otherwise INDEX_NONE
+		int32 GetDynamicIndex() const { return IsDynamic() ? (-HandleValue - 1) : INDEX_NONE; }
+
+		// Compares for equality and inequality
+		bool operator==(FDecoratorRegistryHandle RHS) const { return HandleValue == RHS.HandleValue; }
+		bool operator!=(FDecoratorRegistryHandle RHS) const { return HandleValue != RHS.HandleValue; }
+
+	private:
+		explicit FDecoratorRegistryHandle(int16 HandleValue_)
+			: HandleValue(HandleValue_)
+		{}
+
+		// Creates a static handle based on a decorator offset in the static buffer
+		static FDecoratorRegistryHandle MakeStatic(int32 DecoratorOffset)
+		{
+			check(DecoratorOffset >= 0 && DecoratorOffset < (1 << 15));
+			return FDecoratorRegistryHandle(DecoratorOffset + 1);
+		}
+
+		// Creates a dynamic handle based on a decorator index in the dynamic array
+		static FDecoratorRegistryHandle MakeDynamic(int32 DecoratorIndex)
+		{
+			// Convert the index from 0-based to 1-based since we reserve 0 for our invalid handle
+			check(DecoratorIndex >= 0 && DecoratorIndex < (1 << 15));
+			return FDecoratorRegistryHandle(-DecoratorIndex - 1);
+		}
+
+		// When 0, the handle is invalid
+		// When positive, it is a 1-based offset in the registry's static buffer
+		// When negative, it is a 1-based index in the registry's dynamic array
+		int16 HandleValue = 0;
+
+		friend struct FDecoratorRegistry;
+	};
+}
diff --git a/Engine/Plugins/Experimental/AnimNext/Source/AnimNext/Public/DecoratorBase/DecoratorUID.h b/Engine/Plugins/Experimental/AnimNext/Source/AnimNext/Public/DecoratorBase/DecoratorUID.h
new file mode 100644
--- /dev/null
+++ b/Engine/Plugins/Experimental/AnimNext/Source/AnimNext/Public/DecoratorBase/DecoratorUID.h
@@ -0,0 +1,59 @@
+// Copyright Epic Games, Inc. All Rights Reserved.
+
+#pragma once
+
+#include "CoreMinimal.h"
+
+namespace UE::AnimNext
+{
+	/**
+	 * FDecoratorUID
+	 *
+	 * Encapsulates an decorator global UID.
+	 * The string is exposed in non-shipping builds for logging and debugging purposes.
+	 * The UID should be generated from the provided string using FNV1a with 32 bits.
+	 *
+	 * The whole struct is meant to be 'constexpr' to allow inlining.
+	 */
+	struct FDecoratorUID final
+	{
+		// Constructs an invalid UID
+		constexpr FDecoratorUID()
+			: UID(0)
+#if !UE_BUILD_SHIPPING
+			, DecoratorName(TEXT("<Invalid decorator UID>"))
+#endif
+		{}
+
+		// Constructs a decorator UID
+		constexpr FDecoratorUID(const TCHAR* DecoratorName_, uint32 UID_)
+			: UID(UID_)
+#if !UE_BUILD_SHIPPING
+			, DecoratorName(DecoratorName_)
+#endif
+		{
+		}
+
+#if !UE_BUILD_SHIPPING
+		// Returns a literal string to the interface name
+		constexpr const TCHAR* GetDecoratorName() const { return DecoratorName; }
+#endif
+
+		// Returns the decorator global UID
+		constexpr uint32 GetUID() const { return UID; }
+
+		// Returns whether this UID is valid or not
+		constexpr bool IsValid() const { return UID != 0; }
+
+		// Compares for equality and inequality
+		constexpr bool operator==(FDecoratorUID RHS) const { return UID == RHS.UID; }
+		constexpr bool operator!=(FDecoratorUID RHS) const { return UID != RHS.UID; }
+
+	private:
+		uint32	UID;
+
+#if !UE_BUILD_SHIPPING
+		const TCHAR* DecoratorName;
+#endif
+	};
+}
diff --git a/Engine/Plugins/Experimental/AnimNext/Source/AnimNext/Public/DecoratorBase/EvaluationFlags.h b/Engine/Plugins/Experimental/AnimNext/Source/AnimNext/Public/DecoratorBase/EvaluationFlags.h
new file mode 100644
--- /dev/null
+++ b/Engine/Plugins/Experimental/AnimNext/Source/AnimNext/Public/DecoratorBase/EvaluationFlags.h
@@ -0,0 +1,32 @@
+// Copyright Epic Games, Inc. All Rights Reserved.
+
+#pragma once
+
+#include "CoreMinimal.h"
+
+namespace UE::AnimNext
+{
+	// Flags to control what to compute when evaluating an animation node
+	enum class EEvaluationFlags
+	{
+		// Empty flags
+		None = 0x00,
+
+		// Evaluates the trajectory of the root bone
+		Trajectory = 0x01,
+
+		// Evaluates the animation pose of all bones
+		Pose = 0x02,
+
+		// Evaluates the animation curves
+		Curves = 0x04,
+
+		// Evaluates the animation attributes
+		Attributes = 0x08,
+
+		// Evaluates everything
+		All = 0xff,
+	};
+
+	ENUM_CLASS_FLAGS(EEvaluationFlags)
+}
diff --git a/Engine/Plugins/Experimental/AnimNext/Source/AnimNext/Public/DecoratorBase/ExecutionContext.h b/Engine/Plugins/Experimental/AnimNext/Source/AnimNext/Public/DecoratorBase/ExecutionContext.h
new file mode 100644
--- /dev/null
+++ b/Engine/Plugins/Experimental/AnimNext/Source/AnimNext/Public/DecoratorBase/ExecutionContext.h
@@ -0,0 +1,122 @@
+// Copyright Epic Games, Inc. All Rights Reserved.
+
+#pragma once
+
+#include "CoreMinimal.h"
+#include "DecoratorBase/DecoratorBinding.h"
+#include "DecoratorBase/DecoratorHandle.h"
+#include "DecoratorBase/DecoratorPtr.h"
+#include "DecoratorBase/DecoratorInterfaceUID.h"
+#include "DecoratorBase/NodeHandle.h"
+
+namespace UE::AnimNext
+{
+	struct FNodeDescription;
+	struct FNodeInstance;
+	struct FNodeTemplate;
+	struct FDecorator;
+	struct FDecoratorTemplate;
+	struct ITraversalContext;
+
+	/**
+	 * Execution Context
+	 * 
+	 * The execution context holds internal state during traversals of the animation graph.
+	 */
+	struct ANIMNEXT_API FExecutionContext
+	{
+		// Creates an execution context for the specified graph
+		explicit FExecutionContext(TArrayView<const uint8> GraphSharedData);
+
+		// Destroys the execution context
+		~FExecutionContext();
+
+		// Queries a node for a decorator that implements the specified interface.
+		// If no such decorator exists, nullptr is returned.
+		template<class DecoratorInterface>
+		bool GetInterface(FWeakDecoratorPtr DecoratorPtr, TDecoratorBinding<DecoratorInterface>& InterfaceBinding) const;
+
+		// Queries a node for a decorator that implements the specified interface.
+		// If no such decorator exists, nullptr is returned.
+		template<class DecoratorInterface>
+		bool GetInterface(const FDecoratorBinding& Binding, TDecoratorBinding<DecoratorInterface>& InterfaceBinding) const;
+
+		// Queries a node for a decorator lower on the stack that implements the specified interface.
+		// If no such decorator exists, nullptr is returned.
+		template<class DecoratorInterface>
+		bool GetInterfaceSuper(FWeakDecoratorPtr DecoratorPtr, TDecoratorBinding<DecoratorInterface>& SuperBinding) const;
+
+		// Queries a node for a decorator lower on the stack that implements the specified interface.
+		// If no such decorator exists, nullptr is returned.
+		template<class DecoratorInterface>
+		bool GetInterfaceSuper(const FDecoratorBinding& Binding, TDecoratorBinding<DecoratorInterface>& SuperBinding) const;
+
+		// Allocates a new node instance from a decorator handle
+		// If the desired decorator lives in the current parent, a weak handle to it will be returned
+		FDecoratorPtr AllocateNodeInstance(const FDecoratorBinding& ParentBinding, FDecoratorHandle ChildDecoratorHandle);
+
+		// Allocates a new node instance from a decorator handle
+		// If the desired decorator lives in the current parent, a weak handle to it will be returned
+		FDecoratorPtr AllocateNodeInstance(FWeakDecoratorPtr ParentBinding, FDecoratorHandle ChildDecoratorHandle);
+
+		// Releases a node instance that is no longer referenced
+		void ReleaseNodeInstance(FNodeInstance* Node);
+
+		// Returns the current strongly typed traversal context or nullptr if not in a traversal
+		template<class TraversalContextType>
+		TraversalContextType& GetTraversalContext() const { return *static_cast<TraversalContextType*>(TraversalContext); }	// TODO: Add a casting check for safety
+
+	private:
+		FExecutionContext(const FExecutionContext&) = delete;
+		FExecutionContext(FExecutionContext&&) = delete;
+
+		bool GetInterfaceImpl(FDecoratorInterfaceUID InterfaceUID, FWeakDecoratorPtr DecoratorPtr, FDecoratorBinding& InterfaceBinding) const;
+		bool GetInterfaceSuperImpl(FDecoratorInterfaceUID InterfaceUID, FWeakDecoratorPtr DecoratorPtr, FDecoratorBinding& SuperBinding) const;
+
+		const FNodeDescription& GetNodeDescription(FNodeHandle NodeHandle) const;
+		const FNodeTemplate* GetNodeTemplate(const FNodeDescription& NodeDesc) const;
+		const FDecorator* GetDecorator(const FDecoratorTemplate& DecoratorDesc) const;
+
+		ITraversalContext* TraversalContext = nullptr;
+
+		TArrayView<const uint8> GraphSharedData;
+
+		friend struct FScopedTraversalContext;
+	};
+
+	// Returns a pointer to the current execution context if present, nullptr otherwise.
+	FExecutionContext* GetThreadExecutionContext();
+
+	//////////////////////////////////////////////////////////////////////////
+
+	template<class DecoratorInterface>
+	bool FExecutionContext::GetInterface(FWeakDecoratorPtr DecoratorPtr, TDecoratorBinding<DecoratorInterface>& InterfaceBinding) const
+	{
+		constexpr FDecoratorInterfaceUID InterfaceUID = DecoratorInterface::InterfaceUID;
+		return GetInterfaceImpl(InterfaceUID, DecoratorPtr, InterfaceBinding);
+	}
+
+	template<class DecoratorInterface>
+	bool FExecutionContext::GetInterface(const FDecoratorBinding& Binding, TDecoratorBinding<DecoratorInterface>& InterfaceBinding) const
+	{
+		return GetInterface<DecoratorInterface>(Binding.GetDecoratorPtr(), InterfaceBinding);
+	}
+
+	template<class DecoratorInterface>
+	bool FExecutionContext::GetInterfaceSuper(FWeakDecoratorPtr DecoratorPtr, TDecoratorBinding<DecoratorInterface>& SuperBinding) const
+	{
+		constexpr FDecoratorInterfaceUID InterfaceUID = DecoratorInterface::InterfaceUID;
+		return GetInterfaceSuperImpl(InterfaceUID, DecoratorPtr, SuperBinding);
+	}
+
+	template<class DecoratorInterface>
+	bool FExecutionContext::GetInterfaceSuper(const FDecoratorBinding& Binding, TDecoratorBinding<DecoratorInterface>& SuperBinding) const
+	{
+		return GetInterfaceSuper<DecoratorInterface>(Binding.GetDecoratorPtr(), SuperBinding);
+	}
+
+	inline FDecoratorPtr FExecutionContext::AllocateNodeInstance(const FDecoratorBinding& ParentBinding, FDecoratorHandle ChildDecoratorHandle)
+	{
+		return AllocateNodeInstance(ParentBinding.GetDecoratorPtr(), ChildDecoratorHandle);
+	}
+}
diff --git a/Engine/Plugins/Experimental/AnimNext/Source/AnimNext/Public/DecoratorBase/IDecoratorInterface.h b/Engine/Plugins/Experimental/AnimNext/Source/AnimNext/Public/DecoratorBase/IDecoratorInterface.h
new file mode 100644
--- /dev/null
+++ b/Engine/Plugins/Experimental/AnimNext/Source/AnimNext/Public/DecoratorBase/IDecoratorInterface.h
@@ -0,0 +1,36 @@
+// Copyright Epic Games, Inc. All Rights Reserved.
+
+#pragma once
+
+#include "CoreMinimal.h"
+#include "DecoratorBase/DecoratorBinding.h"			// Derived types will need to implement the binding specialization
+#include "DecoratorBase/DecoratorInterfaceUID.h"
+
+// Helper macros
+#define DECLARE_ANIM_DECORATOR_INTERFACE(InterfaceName, InterfaceNameHash) \
+	/* Globally unique UID for this interface */ \
+	static constexpr UE::AnimNext::FDecoratorInterfaceUID InterfaceUID = UE::AnimNext::FDecoratorInterfaceUID(TEXT(#InterfaceName), InterfaceNameHash); \
+	virtual UE::AnimNext::FDecoratorInterfaceUID GetInterfaceUID() const override { return InterfaceUID; }
+
+namespace UE::AnimNext
+{
+	struct ExecutionContext;	// Derived types will have functions that accept the execution context
+
+	/**
+	 * IDecoratorInterface
+	 * 
+	 * Base type for all decorator interfaces. Used for type safety.
+	 */
+	struct ANIMNEXT_API IDecoratorInterface
+	{
+		virtual ~IDecoratorInterface() {}
+
+		// The globally unique UID for this interface
+		// Derived types will have their own InterfaceUID member that hides/aliases/shadows this one
+		// @see DECLARE_ANIM_DECORATOR_INTERFACE
+		static constexpr FDecoratorInterfaceUID InterfaceUID = FDecoratorInterfaceUID(TEXT("IDecoratorInterface"), 0x402f6df4);
+
+		// Returns the globally unique UID for this interface
+		virtual FDecoratorInterfaceUID GetInterfaceUID() const { return InterfaceUID; };
+	};
+}
diff --git a/Engine/Plugins/Experimental/AnimNext/Source/AnimNext/Public/DecoratorBase/ITraversalContext.h b/Engine/Plugins/Experimental/AnimNext/Source/AnimNext/Public/DecoratorBase/ITraversalContext.h
new file mode 100644
--- /dev/null
+++ b/Engine/Plugins/Experimental/AnimNext/Source/AnimNext/Public/DecoratorBase/ITraversalContext.h
@@ -0,0 +1,47 @@
+// Copyright Epic Games, Inc. All Rights Reserved.
+
+#pragma once
+
+#include "CoreMinimal.h"
+
+namespace UE::AnimNext
+{
+	struct FExecutionContext;
+
+	/**
+	 * ITraversalContext
+	 * 
+	 * Base type for all traversal context derived types.
+	 * A traversal context is a transient object that lives only during a particular
+	 * traversal. It can be access through the FExecutionContext.
+	 * 
+	 * @see FExecutionContext
+	 */
+	struct ITraversalContext
+	{
+		virtual ~ITraversalContext() {}
+	};
+
+	/**
+	 * FScopedTraversalContext
+	 * 
+	 * Pushes and pops a traversal context onto the execution context.
+	 * During the lifetime of an instance of this object, the execution context will
+	 * return the provided traversal context as the current traversal context.
+	 * 
+	 * @see FExecutionContext
+	 */
+	struct ANIMNEXT_API FScopedTraversalContext final
+	{
+		FScopedTraversalContext(FExecutionContext& ExecutionContext, ITraversalContext& TraversalContext);
+
+		FScopedTraversalContext(const FScopedTraversalContext&) = delete;
+		FScopedTraversalContext& operator=(const FScopedTraversalContext&) = delete;
+
+		~FScopedTraversalContext();
+
+	private:
+		FExecutionContext& ExecutionContext;
+		ITraversalContext* OldTraversalContext;
+	};
+}
diff --git a/Engine/Plugins/Experimental/AnimNext/Source/AnimNext/Public/DecoratorBase/NodeDescription.h b/Engine/Plugins/Experimental/AnimNext/Source/AnimNext/Public/DecoratorBase/NodeDescription.h
new file mode 100644
--- /dev/null
+++ b/Engine/Plugins/Experimental/AnimNext/Source/AnimNext/Public/DecoratorBase/NodeDescription.h
@@ -0,0 +1,45 @@
+// Copyright Epic Games, Inc. All Rights Reserved.
+
+#pragma once
+
+#include "CoreMinimal.h"
+#include "DecoratorBase/DecoratorTemplate.h"
+#include "DecoratorBase/NodeTemplateRegistryHandle.h"
+
+namespace UE::AnimNext
+{
+	struct FDecoratorDescription;
+
+	/**
+	 * Node Description
+	 * A node description represents a unique instance in the authored static graph.
+	 * A node description may have any number of runtime instances in the dynamically executed graph.
+	 * As such, a node description is read-only at runtime while a node instance is read/write.
+	 * 
+	 * A node description is followed in memory by the decorator descriptions (their shared read-only data)
+	 * that live within it. Decorator descriptions include things like hard-coded/inline properties, pin links, etc.
+	 * 
+	 * A node description is itself an instance of a node template.
+	 * 
+	 * @see FNodeTemplate
+	 */
+	struct alignas(alignof(uint32)) FNodeDescription
+	{
+		FNodeDescription(uint16 UID_, FNodeTemplateRegistryHandle TemplateHandle_)
+			: UID(UID_)
+			, TemplateHandle(TemplateHandle_)
+		{}
+
+		// Returns the node UID, unique to the owning sub-graph
+		uint32 GetUID() const { return UID; }
+
+		// Returns the handle of the node's template in the node template registry
+		FNodeTemplateRegistryHandle GetTemplateHandle() const { return TemplateHandle; }
+
+	private:
+		uint16							UID;				// assigned during export/cook, unique to current sub-graph
+		FNodeTemplateRegistryHandle		TemplateHandle;		// offset of the node template within the global list
+
+		// Followed by a list of [FDecoratorDescription] instances and optional padding
+	};
+}
diff --git a/Engine/Plugins/Experimental/AnimNext/Source/AnimNext/Public/DecoratorBase/NodeHandle.h b/Engine/Plugins/Experimental/AnimNext/Source/AnimNext/Public/DecoratorBase/NodeHandle.h
new file mode 100644
--- /dev/null
+++ b/Engine/Plugins/Experimental/AnimNext/Source/AnimNext/Public/DecoratorBase/NodeHandle.h
@@ -0,0 +1,47 @@
+// Copyright Epic Games, Inc. All Rights Reserved.
+
+#pragma once
+
+#include "CoreMinimal.h"
+
+namespace UE::AnimNext
+{
+	/**
+	 * Node Handle
+	 * A node handle represents a reference to a specific node instance in the shared/read-only portion
+	 * of a sub-graph. It points to a FNodeDescription when resolved.
+	 * 
+	 * @see FNodeDescription
+	 */
+	struct FNodeHandle
+	{
+		// Creates an invalid node handle
+		FNodeHandle()
+			: SharedOffset(INVALID_SHARED_OFFSET_VALUE)
+		{}
+
+		// Creates a node handle for the specified shared offset
+		explicit FNodeHandle(uint32 SharedOffset_)
+			: SharedOffset(SharedOffset_)
+		{}
+
+		// Returns true if this node handle is valid, false otherwise
+		bool IsValid() const { return SharedOffset != INVALID_SHARED_OFFSET_VALUE; }
+
+		// Returns the offset into the shared data relative to the root of the owning sub-graph
+		uint32 GetSharedOffset() const { return SharedOffset; }
+
+		// Equality and inequality tests
+		bool operator==(FNodeHandle RHS) const { return SharedOffset == RHS.SharedOffset; }
+		bool operator!=(FNodeHandle RHS) const { return SharedOffset != RHS.SharedOffset; }
+
+	private:
+		// We pack the shared offset on 24 bits, the top 8 bits are truncated
+		static constexpr uint32 INVALID_SHARED_OFFSET_VALUE = ~0u >> (32 - 24);
+
+		uint32	Padding : 8;			// unused padding
+		uint32	SharedOffset : 24;		// relative to root of sub-graph
+
+		friend struct FDecoratorHandle;
+	};
+}
diff --git a/Engine/Plugins/Experimental/AnimNext/Source/AnimNext/Public/DecoratorBase/NodeInstance.h b/Engine/Plugins/Experimental/AnimNext/Source/AnimNext/Public/DecoratorBase/NodeInstance.h
new file mode 100644
--- /dev/null
+++ b/Engine/Plugins/Experimental/AnimNext/Source/AnimNext/Public/DecoratorBase/NodeInstance.h
@@ -0,0 +1,51 @@
+// Copyright Epic Games, Inc. All Rights Reserved.
+
+#pragma once
+
+#include "CoreMinimal.h"
+#include "DecoratorBase/DecoratorTemplate.h"
+#include "DecoratorBase/NodeHandle.h"
+
+namespace UE::AnimNext
+{
+	struct FDecoratorInstance;
+
+	/**
+	 * Node Instance
+	 * A node instance represents allocated data for specific node.
+	 *
+	 * In order to access the decorator instance data, the offsets need to be looked up in the node template.
+	 * 
+	 * @see FNodeTemplate, FDecoratorTemplate
+	 */
+	struct ANIMNEXT_API FNodeInstance
+	{
+		FNodeInstance(const FNodeInstance&) = delete;
+		FNodeInstance& operator=(const FNodeInstance&) = delete;
+
+		// Returns whether the node instance is valid or not
+		bool IsValid() const { return NodeHandle.IsValid(); }
+
+		// Returns a handle to the shared data for this node
+		FNodeHandle GetNodeHandle() const { return NodeHandle; }
+
+		// Returns the number of live references to this node instance, does not include weak handles
+		uint32 GetReferenceCount() const { return ReferenceCount; }
+
+	private:
+		explicit FNodeInstance(FNodeHandle NodeHandle);
+		~FNodeInstance() { check(ReferenceCount == 0); }
+
+		void AddReference() { ReferenceCount++; }
+		void ReleaseReference();
+
+		uint32		ReferenceCount;		// how many non-weak FNodePtr and FDecoratorPtr handles point to us, not thread safe
+		FNodeHandle	NodeHandle;			// relative to root of sub-graph, should this be a pointer?
+
+		// Followed by a list of [FDecoratorInstance] instances and optional padding
+
+		friend struct FDecoratorPtr;
+		friend struct FExecutionContext;
+		friend struct FNodePtr;
+	};
+}
diff --git a/Engine/Plugins/Experimental/AnimNext/Source/AnimNext/Public/DecoratorBase/NodeTemplateRegistryHandle.h b/Engine/Plugins/Experimental/AnimNext/Source/AnimNext/Public/DecoratorBase/NodeTemplateRegistryHandle.h
new file mode 100644
--- /dev/null
+++ b/Engine/Plugins/Experimental/AnimNext/Source/AnimNext/Public/DecoratorBase/NodeTemplateRegistryHandle.h
@@ -0,0 +1,48 @@
+// Copyright Epic Games, Inc. All Rights Reserved.
+
+#pragma once
+
+#include "CoreMinimal.h"
+
+namespace UE::AnimNext
+{
+	/**
+	 * FNodeTemplateRegistryHandle
+	 *
+	 * Encapsulates a value used as a handle in the node template registry.
+	 * When valid, it can be used to retrieve a pointer to the corresponding node template.
+	 */
+	struct FNodeTemplateRegistryHandle final
+	{
+		// Default constructed handles are invalid
+		FNodeTemplateRegistryHandle() = default;
+
+		// Returns whether or not this handle points to a valid node template
+		bool IsValid() const { return TemplateOffset != 0; }
+
+		// Returns the template offset for this handle when valid, otherwise INDEX_NONE
+		int32 GetTemplateOffset() const { return IsValid() ? (TemplateOffset - 1) : INDEX_NONE; }
+
+		// Compares for equality and inequality
+		bool operator==(FNodeTemplateRegistryHandle RHS) const { return TemplateOffset == RHS.TemplateOffset; }
+		bool operator!=(FNodeTemplateRegistryHandle RHS) const { return TemplateOffset != RHS.TemplateOffset; }
+
+	private:
+		explicit FNodeTemplateRegistryHandle(int16 TemplateOffset_)
+			: TemplateOffset(TemplateOffset_)
+		{}
+
+		// Creates a handle based on a note template offset in the static buffer
+		static FNodeTemplateRegistryHandle MakeHandle(int32 TemplateOffset_)
+		{
+			check(TemplateOffset_ >= 0 && TemplateOffset_ < ((1 << 16) - 2));
+			return FNodeTemplateRegistryHandle(TemplateOffset_ + 1);
+		}
+
+		// When 0, the handle is invalid
+		// Otherwise it is a 1-based offset in the registry's static buffer
+		int16 TemplateOffset = 0;
+
+		friend struct FNodeTemplateRegistry;
+	};
+}
diff --git a/Engine/Plugins/Experimental/AnimNext/Source/AnimNext/Public/DecoratorInterfaces/IEvaluate.h b/Engine/Plugins/Experimental/AnimNext/Source/AnimNext/Public/DecoratorInterfaces/IEvaluate.h
new file mode 100644
--- /dev/null
+++ b/Engine/Plugins/Experimental/AnimNext/Source/AnimNext/Public/DecoratorInterfaces/IEvaluate.h
@@ -0,0 +1,94 @@
+// Copyright Epic Games, Inc. All Rights Reserved.
+
+#pragma once
+
+#include "CoreMinimal.h"
+#include "DecoratorBase/EvaluationFlags.h"
+#include "DecoratorBase/IDecoratorInterface.h"
+#include "DecoratorBase/ITraversalContext.h"
+
+namespace UE::AnimNext
+{
+	// TODO: Hook up our real pose container
+	struct FPoseContainer final {};
+
+	/**
+	 * FEvaluateTraversalContext
+	 *
+	 * Contains all relevant transient data for an evaluate traversal.
+	 */
+	struct ANIMNEXT_API FEvaluateTraversalContext : ITraversalContext
+	{
+		explicit FEvaluateTraversalContext(EEvaluationFlags EvaluationFlags_);
+
+		EEvaluationFlags GetEvaluationFlags() const { return EvaluationFlags; }
+
+		//void PushPose(FPose&& Pose);
+		//FPose&& PopPose();
+
+	private:
+		EEvaluationFlags EvaluationFlags;
+	};
+
+	/**
+	 * IEvaluate
+	 * 
+	 * This interface is called during the evaluation traversal. It aims to produce some
+	 * of the following: an output pose, its velocity, the root trajectory, scalar curves, attributes, etc.
+	 * 
+	 * When a node is decorator is visited, PreEvaluate is first called on it. It is responsible for forwarding
+	 * the call to the next decorator that implements this interface on the decorator stack of the node. Once
+	 * all decorators have had the chance to PreEvaluate, the children of the decorator are queried through
+	 * the IHierarchy interface. The children will then evaluate and PostEvaluate will then be called afterwards
+	 * on the original decorator.
+	 * 
+	 * The execution context contains what to evaluate.
+	 * @see EEvaluationFlags
+	 */
+	struct ANIMNEXT_API IEvaluate : IDecoratorInterface
+	{
+		DECLARE_ANIM_DECORATOR_INTERFACE(IEvaluate, 0xa303e9e7)
+
+		// Called before a decorator's children are evaluated
+		virtual void PreEvaluate(FExecutionContext& Context, const TDecoratorBinding<IEvaluate>& Binding) const;
+
+		// Called after a decorator's children have been evaluated
+		virtual void PostEvaluate(FExecutionContext& Context, const TDecoratorBinding<IEvaluate>& Binding) const;
+	};
+
+	/**
+	 * Specialization for decorator binding.
+	 */
+	template<>
+	struct TDecoratorBinding<IEvaluate> : FDecoratorBinding
+	{
+		// @see IEvaluate::PreEvaluate
+		void PreEvaluate(FExecutionContext& Context) const
+		{
+			GetInterface()->PreEvaluate(Context, *this);
+		}
+
+		// @see IEvaluate::PostEvaluate
+		void PostEvaluate(FExecutionContext& Context) const
+		{
+			GetInterface()->PostEvaluate(Context, *this);
+		}
+
+	protected:
+		const IEvaluate* GetInterface() const { return GetInterfaceTyped<IEvaluate>(); }
+	};
+
+	/**
+	 * Evaluates a sub-graph starting at the graph root and produces its result.
+	 * Evaluation starts at the top of the stack that includes the graph root decorator.
+	 *
+	 * For each node:
+	 *     - We call PreEvaluate on all its decorators
+	 *     - We call GetChildren on all its decorators
+	 *     - We evaluate all children found
+	 *     - We call PostEvaluate on all its decorators
+	 *
+	 * @see IEvaluate::PreEvaluate, IEvaluate::PostEvaluate, IHierarchy::GetChildren
+	 */
+	ANIMNEXT_API void EvaluateGraph(FExecutionContext& Context, FWeakDecoratorPtr GraphRootPtr, EEvaluationFlags EvaluationFlags, FPoseContainer& OutContainer);
+}
diff --git a/Engine/Plugins/Experimental/AnimNext/Source/AnimNext/Public/DecoratorInterfaces/IHierarchy.h b/Engine/Plugins/Experimental/AnimNext/Source/AnimNext/Public/DecoratorInterfaces/IHierarchy.h
new file mode 100644
--- /dev/null
+++ b/Engine/Plugins/Experimental/AnimNext/Source/AnimNext/Public/DecoratorInterfaces/IHierarchy.h
@@ -0,0 +1,42 @@
+// Copyright Epic Games, Inc. All Rights Reserved.
+
+#pragma once
+
+#include "CoreMinimal.h"
+#include "Misc/MemStack.h"
+#include "DecoratorBase/IDecoratorInterface.h"
+
+namespace UE::AnimNext
+{
+	using FChildrenArray = TArray<FWeakDecoratorPtr, TMemStackAllocator<>>;
+
+	/**
+	 * IHierarchy
+	 * 
+	 * This interface exposes hierarchy traversal information to navigate the graph.
+	 */
+	struct ANIMNEXT_API IHierarchy : IDecoratorInterface
+	{
+		DECLARE_ANIM_DECORATOR_INTERFACE(IHierarchy, 0x846d8a37)
+
+		// Appends weak handles to any children we wish to traverse.
+		// Decorators are responsible for allocating and releasing child instance data.
+		virtual void GetChildren(FExecutionContext& Context, const TDecoratorBinding<IHierarchy>& Binding, FChildrenArray& Children) const;
+	};
+
+	/**
+	 * Specialization for decorator binding.
+	 */
+	template<>
+	struct TDecoratorBinding<IHierarchy> : FDecoratorBinding
+	{
+		// @see IHierarchy::GetChildren
+		void GetChildren(FExecutionContext& Context, FChildrenArray& Children) const
+		{
+			GetInterface()->GetChildren(Context, *this, Children);
+		}
+
+	protected:
+		const IHierarchy* GetInterface() const { return GetInterfaceTyped<IHierarchy>(); }
+	};
+}
diff --git a/Engine/Plugins/Experimental/AnimNext/Source/AnimNext/Public/DecoratorInterfaces/ITimeline.h b/Engine/Plugins/Experimental/AnimNext/Source/AnimNext/Public/DecoratorInterfaces/ITimeline.h
new file mode 100644
--- /dev/null
+++ b/Engine/Plugins/Experimental/AnimNext/Source/AnimNext/Public/DecoratorInterfaces/ITimeline.h
@@ -0,0 +1,38 @@
+// Copyright Epic Games, Inc. All Rights Reserved.
+
+#pragma once
+
+#include "CoreMinimal.h"
+#include "DecoratorBase/IDecoratorInterface.h"
+
+namespace UE::AnimNext
+{
+	/**
+	 * ITimeline
+	 *
+	 * This interface exposes timeline related information.
+	 */
+	struct ANIMNEXT_API ITimeline : IDecoratorInterface
+	{
+		DECLARE_ANIM_DECORATOR_INTERFACE(ITimeline, 0x53760727)
+
+		// Returns the play rate of this timeline
+		virtual double GetPlayRate(FExecutionContext& Context, const TDecoratorBinding<ITimeline>& Binding) const;
+	};
+
+	/**
+	 * Specialization for decorator binding.
+	 */
+	template<>
+	struct TDecoratorBinding<ITimeline> : FDecoratorBinding
+	{
+		// @see ITimeline::GetPlayRate
+		double GetPlayRate(FExecutionContext& Context) const
+		{
+			return GetInterface()->GetPlayRate(Context, *this);
+		}
+
+	protected:
+		const ITimeline* GetInterface() const { return GetInterfaceTyped<ITimeline>(); }
+	};
+}
diff --git a/Engine/Plugins/Experimental/AnimNext/Source/AnimNext/Public/DecoratorInterfaces/IUpdate.h b/Engine/Plugins/Experimental/AnimNext/Source/AnimNext/Public/DecoratorInterfaces/IUpdate.h
new file mode 100644
--- /dev/null
+++ b/Engine/Plugins/Experimental/AnimNext/Source/AnimNext/Public/DecoratorInterfaces/IUpdate.h
@@ -0,0 +1,79 @@
+// Copyright Epic Games, Inc. All Rights Reserved.
+
+#pragma once
+
+#include "CoreMinimal.h"
+#include "DecoratorBase/IDecoratorInterface.h"
+#include "DecoratorBase/ITraversalContext.h"
+
+namespace UE::AnimNext
+{
+	/**
+	 * FUpdateTraversalContext
+	 * 
+	 * Contains all relevant transient data for an update traversal.
+	 */
+	struct FUpdateTraversalContext : ITraversalContext
+	{
+		double GetDeltaTime() const { return 1.0 / 30.0; }
+		double GetPlayRate() const { return 1.0; }
+	};
+
+	/**
+	 * IUpdate
+	 *
+	 * This interface is called during the update traversal.
+	 *
+	 * When a node is decorator is visited, PreUpdate is first called on it. It is responsible for forwarding
+	 * the call to the next decorator that implements this interface on the decorator stack of the node. Once
+	 * all decorators have had the chance to PreUpdate, the children of the decorator are queried through
+	 * the IHierarchy interface. The children will then evaluate and PostUpdate will then be called afterwards
+	 * on the original decorator.
+	 */
+	struct ANIMNEXT_API IUpdate : IDecoratorInterface
+	{
+		DECLARE_ANIM_DECORATOR_INTERFACE(IUpdate, 0x59d24dc5)
+
+		// Called before a decorator's children are updated
+		virtual void PreUpdate(FExecutionContext& Context, const TDecoratorBinding<IUpdate>& Binding) const;
+
+		// Called after a decorator's children have been updated
+		virtual void PostUpdate(FExecutionContext& Context, const TDecoratorBinding<IUpdate>& Binding) const;
+	};
+
+	/**
+	 * Specialization for decorator binding.
+	 */
+	template<>
+	struct TDecoratorBinding<IUpdate> : FDecoratorBinding
+	{
+		// @see IUpdate::PreUpdate
+		void PreUpdate(FExecutionContext& Context) const
+		{
+			GetInterface()->PreUpdate(Context, *this);
+		}
+
+		// @see IUpdate::PostUpdate
+		void PostUpdate(FExecutionContext& Context) const
+		{
+			GetInterface()->PostUpdate(Context, *this);
+		}
+
+	protected:
+		const IUpdate* GetInterface() const { return GetInterfaceTyped<IUpdate>(); }
+	};
+
+	/**
+	 * Updates a sub-graph starting at the graph root.
+	 * Update starts at the top of the stack that includes the graph root decorator.
+	 *
+	 * For each node:
+	 *     - We call PreUpdate on all its decorators
+	 *     - We call GetChildren on all its decorators
+	 *     - We update all children found
+	 *     - We call PostUpdate on all its decorators
+	 *
+	 * @see IUpdate::PreUpdate, IUpdate::PostUpdate, IHierarchy::GetChildren
+	 */
+	ANIMNEXT_API void UpdateGraph(FExecutionContext& Context, FWeakDecoratorPtr GraphRootPtr);
+}
diff --git a/Engine/Plugins/Experimental/AnimNext/Source/AnimNextTestSuite/AnimNextTestSuite.Build.cs b/Engine/Plugins/Experimental/AnimNext/Source/AnimNextTestSuite/AnimNextTestSuite.Build.cs
new file mode 100644
--- /dev/null
+++ b/Engine/Plugins/Experimental/AnimNext/Source/AnimNextTestSuite/AnimNextTestSuite.Build.cs
@@ -0,0 +1,27 @@
+// Copyright Epic Games, Inc. All Rights Reserved.
+
+namespace UnrealBuildTool.Rules
+{
+	public class AnimNextTestSuite : ModuleRules
+	{
+		public AnimNextTestSuite(ReadOnlyTargetRules Target) : base(Target)
+		{
+			PCHUsage = ModuleRules.PCHUsageMode.UseExplicitOrSharedPCHs;
+
+			PublicDependencyModuleNames.AddRange(
+				new string[]
+				{
+					"Core",
+					"CoreUObject",
+					"Engine",
+					"AnimNext",
+				}
+			);
+
+			if (Target.bBuildEditor == true)
+			{
+				PrivateDependencyModuleNames.Add("UnrealEd");
+			}
+		}
+	}
+}
\ No newline at end of file
diff --git a/Engine/Plugins/Experimental/AnimNext/Source/AnimNextTestSuite/Private/AnimNextDecoratorBaseTest.cpp b/Engine/Plugins/Experimental/AnimNext/Source/AnimNextTestSuite/Private/AnimNextDecoratorBaseTest.cpp
new file mode 100644
--- /dev/null
+++ b/Engine/Plugins/Experimental/AnimNext/Source/AnimNextTestSuite/Private/AnimNextDecoratorBaseTest.cpp
@@ -0,0 +1,1016 @@
+// Copyright Epic Games, Inc. All Rights Reserved.
+
+
+#include "AnimNextRuntimeTest.h"
+
+#include "Misc/AutomationTest.h"
+
+#if WITH_DEV_AUTOMATION_TESTS
+
+#include "DecoratorBase/Decorator.h"
+#include "DecoratorBase/DecoratorRegistry.h"
+#include "DecoratorBase/ExecutionContext.h"
+#include "DecoratorBase/IDecoratorInterface.h"
+#include "DecoratorBase/NodeInstance.h"
+#include "DecoratorBase/NodeTemplateRegistry.h"
+
+//****************************************************************************
+// AnimNext Runtime DecoratorBase Tests
+//****************************************************************************
+
+namespace UE::AnimNext
+{
+	namespace Private
+	{
+		static TArray<FDecoratorUID>* ConstructedDecorators = nullptr;
+		static TArray<FDecoratorUID>* DestructedDecorators = nullptr;
+	}
+
+	struct IInterfaceA : IDecoratorInterface
+	{
+		DECLARE_ANIM_DECORATOR_INTERFACE(IInterfaceA, 0x34cb8e62)
+
+		virtual void FuncA(FExecutionContext& Context, const TDecoratorBinding<IInterfaceA>& Binding) const;
+	};
+
+	template<>
+	struct TDecoratorBinding<IInterfaceA> : FDecoratorBinding
+	{
+		void FuncA(FExecutionContext& Context) const
+		{
+			GetInterface()->FuncA(Context, *this);
+		}
+
+	protected:
+		const IInterfaceA* GetInterface() const { return GetInterfaceTyped<IInterfaceA>(); }
+	};
+
+	void IInterfaceA::FuncA(FExecutionContext& Context, const TDecoratorBinding<IInterfaceA>& Binding) const
+	{
+		TDecoratorBinding<IInterfaceA> SuperBinding;
+		if (Context.GetInterfaceSuper(Binding, SuperBinding))
+		{
+			SuperBinding.FuncA(Context);
+		}
+	}
+
+	//////////////////////////////////////////////////////////////////////////
+
+	struct IInterfaceB : IDecoratorInterface
+	{
+		DECLARE_ANIM_DECORATOR_INTERFACE(IInterfaceB, 0x33cb8ccf)
+
+		virtual void FuncB(FExecutionContext& Context, const TDecoratorBinding<IInterfaceB>& Binding) const;
+	};
+
+	template<>
+	struct TDecoratorBinding<IInterfaceB> : FDecoratorBinding
+	{
+		void FuncB(FExecutionContext& Context) const
+		{
+			GetInterface()->FuncB(Context, *this);
+		}
+
+	protected:
+		const IInterfaceB* GetInterface() const { return GetInterfaceTyped<IInterfaceB>(); }
+	};
+
+	void IInterfaceB::FuncB(FExecutionContext& Context, const TDecoratorBinding<IInterfaceB>& Binding) const
+	{
+		TDecoratorBinding<IInterfaceB> SuperBinding;
+		if (Context.GetInterfaceSuper(Binding, SuperBinding))
+		{
+			SuperBinding.FuncB(Context);
+		}
+	}
+
+	//////////////////////////////////////////////////////////////////////////
+
+	struct IInterfaceC : IDecoratorInterface
+	{
+		DECLARE_ANIM_DECORATOR_INTERFACE(IInterfaceC, 0x32cb8b3c)
+
+		virtual void FuncC(FExecutionContext& Context, const TDecoratorBinding<IInterfaceC>& Binding) const;
+	};
+
+	template<>
+	struct TDecoratorBinding<IInterfaceC> : FDecoratorBinding
+	{
+		void FuncC(FExecutionContext& Context) const
+		{
+			GetInterface()->FuncC(Context, *this);
+		}
+
+	protected:
+		const IInterfaceC* GetInterface() const { return GetInterfaceTyped<IInterfaceC>(); }
+	};
+
+	void IInterfaceC::FuncC(FExecutionContext& Context, const TDecoratorBinding<IInterfaceC>& Binding) const
+	{
+		TDecoratorBinding<IInterfaceC> SuperBinding;
+		if (Context.GetInterfaceSuper(Binding, SuperBinding))
+		{
+			SuperBinding.FuncC(Context);
+		}
+	}
+
+	//////////////////////////////////////////////////////////////////////////
+
+	struct FDecoratorA_Base : FDecorator, IInterfaceA
+	{
+		DECLARE_ANIM_DECORATOR(FDecoratorA_Base, 0x3a1861cf, FDecorator)
+
+		struct FSharedData : FDecorator::FSharedData
+		{
+			FDecoratorUID DecoratorUID = FDecoratorA_Base::DecoratorUID;
+		};
+
+		struct FInstanceData : FDecorator::FInstanceData
+		{
+			FDecoratorUID DecoratorUID = FDecoratorA_Base::DecoratorUID;
+
+			FInstanceData()
+			{
+				if (Private::ConstructedDecorators != nullptr)
+				{
+					Private::ConstructedDecorators->Add(FDecoratorA_Base::DecoratorUID);
+				}
+			}
+
+			~FInstanceData()
+			{
+				if (Private::DestructedDecorators != nullptr)
+				{
+					Private::DestructedDecorators->Add(FDecoratorA_Base::DecoratorUID);
+				}
+			}
+		};
+
+		// FDecorator impl
+		virtual EDecoratorMode GetMode() const override { return EDecoratorMode::Base; }
+
+		// IInterfaceA impl
+		virtual void FuncA(FExecutionContext& Context, const TDecoratorBinding<IInterfaceA>& Binding) const override
+		{
+		}
+	};
+
+	DEFINE_ANIM_DECORATOR_BEGIN(FDecoratorA_Base)
+		DEFINE_ANIM_DECORATOR_IMPLEMENTS_INTERFACE(IInterfaceA)
+	DEFINE_ANIM_DECORATOR_END(FDecoratorA_Base)
+
+	//////////////////////////////////////////////////////////////////////////
+
+	struct FDecoratorAB_Add : FDecorator, IInterfaceA, IInterfaceB
+	{
+		DECLARE_ANIM_DECORATOR(FDecoratorAB_Add, 0xe205a0e1, FDecorator)
+
+		struct FSharedData : FDecorator::FSharedData
+		{
+			FDecoratorUID DecoratorUID = FDecoratorAB_Add::DecoratorUID;
+		};
+
+		struct FInstanceData : FDecorator::FInstanceData
+		{
+			FDecoratorUID DecoratorUID = FDecoratorAB_Add::DecoratorUID;
+
+			FInstanceData()
+			{
+				if (Private::ConstructedDecorators != nullptr)
+				{
+					Private::ConstructedDecorators->Add(FDecoratorAB_Add::DecoratorUID);
+				}
+			}
+
+			~FInstanceData()
+			{
+				if (Private::DestructedDecorators != nullptr)
+				{
+					Private::DestructedDecorators->Add(FDecoratorAB_Add::DecoratorUID);
+				}
+			}
+		};
+
+		// FDecorator impl
+		virtual EDecoratorMode GetMode() const override { return EDecoratorMode::Additive; }
+
+		// IInterfaceA impl
+		virtual void FuncA(FExecutionContext& Context, const TDecoratorBinding<IInterfaceA>& Binding) const override
+		{
+		}
+
+		// IInterfaceB impl
+		virtual void FuncB(FExecutionContext& Context, const TDecoratorBinding<IInterfaceB>& Binding) const override
+		{
+		}
+	};
+
+	DEFINE_ANIM_DECORATOR_BEGIN(FDecoratorAB_Add)
+		DEFINE_ANIM_DECORATOR_IMPLEMENTS_INTERFACE(IInterfaceA)
+		DEFINE_ANIM_DECORATOR_IMPLEMENTS_INTERFACE(IInterfaceB)
+	DEFINE_ANIM_DECORATOR_END(FDecoratorAB_Add)
+
+	//////////////////////////////////////////////////////////////////////////
+
+	struct FDecoratorAC_Add : FDecorator, IInterfaceA, IInterfaceC
+	{
+		DECLARE_ANIM_DECORATOR(FDecoratorAC_Add, 0x26d83846, FDecorator)
+
+		struct FSharedData : FDecorator::FSharedData
+		{
+			FDecoratorUID DecoratorUID = FDecoratorAC_Add::DecoratorUID;
+		};
+
+		struct FInstanceData : FDecorator::FInstanceData
+		{
+			FDecoratorUID DecoratorUID = FDecoratorAC_Add::DecoratorUID;
+
+			FInstanceData()
+			{
+				if (Private::ConstructedDecorators != nullptr)
+				{
+					Private::ConstructedDecorators->Add(FDecoratorAC_Add::DecoratorUID);
+				}
+			}
+
+			~FInstanceData()
+			{
+				if (Private::DestructedDecorators != nullptr)
+				{
+					Private::DestructedDecorators->Add(FDecoratorAC_Add::DecoratorUID);
+				}
+			}
+		};
+
+		// FDecorator impl
+		virtual EDecoratorMode GetMode() const override { return EDecoratorMode::Additive; }
+
+		// IInterfaceA impl
+		virtual void FuncA(FExecutionContext& Context, const TDecoratorBinding<IInterfaceA>& Binding) const override
+		{
+		}
+
+		// IInterfaceC impl
+		virtual void FuncC(FExecutionContext& Context, const TDecoratorBinding<IInterfaceC>& Binding) const override
+		{
+		}
+	};
+
+	DEFINE_ANIM_DECORATOR_BEGIN(FDecoratorAC_Add)
+		DEFINE_ANIM_DECORATOR_IMPLEMENTS_INTERFACE(IInterfaceA)
+		DEFINE_ANIM_DECORATOR_IMPLEMENTS_INTERFACE(IInterfaceC)
+	DEFINE_ANIM_DECORATOR_END(FDecoratorAC_Add)
+}
+
+IMPLEMENT_SIMPLE_AUTOMATION_TEST(FAnimationAnimNextRuntimeTest_DecoratorRegistry, "Animation.AnimNext.Runtime.DecoratorRegistry", EAutomationTestFlags::EditorContext | EAutomationTestFlags::EngineFilter)
+
+bool FAnimationAnimNextRuntimeTest_DecoratorRegistry::RunTest(const FString& InParameters)
+{
+	using namespace UE::AnimNext;
+
+	FDecoratorRegistry& Registry = FDecoratorRegistry::Get();
+
+	AddErrorIfFalse(Registry.GetNum() == 0, "FAnimationAnimNextRuntimeTest_DecoratorRegistry -> Registry should be empty");
+	AddErrorIfFalse(!Registry.FindHandle(FDecoratorA_Base::DecoratorUID).IsValid(), "FAnimationAnimNextRuntimeTest_DecoratorRegistry -> Registry should not contain our decorator");
+	AddErrorIfFalse(!Registry.FindHandle(FDecoratorAB_Add::DecoratorUID).IsValid(), "FAnimationAnimNextRuntimeTest_DecoratorRegistry -> Registry should not contain our decorator");
+	AddErrorIfFalse(!Registry.FindHandle(FDecoratorAC_Add::DecoratorUID).IsValid(), "FAnimationAnimNextRuntimeTest_DecoratorRegistry -> Registry should not contain our decorator");
+
+	{
+		// Auto register a decorator
+		AUTO_REGISTER_ANIM_DECORATOR(FDecoratorA_Base)
+
+		AddErrorIfFalse(Registry.GetNum() == 1, "FAnimationAnimNextRuntimeTest_DecoratorRegistry -> Registry should contain 1 decorator");
+
+		FDecoratorRegistryHandle HandleA = Registry.FindHandle(FDecoratorA_Base::DecoratorUID);
+		AddErrorIfFalse(HandleA.IsValid(), "FAnimationAnimNextRuntimeTest_DecoratorRegistry -> Decorator should have registered automatically");
+		AddErrorIfFalse(HandleA.IsStatic(), "FAnimationAnimNextRuntimeTest_DecoratorRegistry -> Decorator should have been statically allocated");
+
+		const FDecorator* DecoratorA = Registry.Find(HandleA);
+		AddErrorIfFalse(DecoratorA != nullptr, "FAnimationAnimNextRuntimeTest_DecoratorRegistry -> Decorator should exist");
+		AddErrorIfFalse(DecoratorA->GetDecoratorUID() == FDecoratorA_Base::DecoratorUID, "FAnimationAnimNextRuntimeTest_DecoratorRegistry -> Unexpected decorator instance type");
+
+		{
+			// Auto register another decorator
+			AUTO_REGISTER_ANIM_DECORATOR(FDecoratorAB_Add)
+
+			AddErrorIfFalse(Registry.GetNum() == 2, "FAnimationAnimNextRuntimeTest_DecoratorRegistry -> Registry should contain 2 decorators");
+
+			FDecoratorRegistryHandle HandleAB = Registry.FindHandle(FDecoratorAB_Add::DecoratorUID);
+			AddErrorIfFalse(HandleAB.IsValid(), "FAnimationAnimNextRuntimeTest_DecoratorRegistry -> Decorator should have registered automatically");
+			AddErrorIfFalse(HandleAB.IsStatic(), "FAnimationAnimNextRuntimeTest_DecoratorRegistry -> Decorator should have been statically allocated");
+			AddErrorIfFalse(HandleA != HandleAB, "FAnimationAnimNextRuntimeTest_DecoratorRegistry -> Decorator handles should be different");
+
+			const FDecorator* DecoratorAB = Registry.Find(HandleAB);
+			AddErrorIfFalse(DecoratorAB != nullptr, "FAnimationAnimNextRuntimeTest_DecoratorRegistry -> Decorator should exist");
+			AddErrorIfFalse(DecoratorAB->GetDecoratorUID() == FDecoratorAB_Add::DecoratorUID, "FAnimationAnimNextRuntimeTest_DecoratorRegistry -> Unexpected decorator instance type");
+
+			FDecoratorRegistryHandle HandleAC_0;
+			{
+				// Dynamically register a decorator
+				FDecoratorAC_Add DecoratorAC_0;
+				Registry.Register(&DecoratorAC_0);
+
+				AddErrorIfFalse(Registry.GetNum() == 3, "FAnimationAnimNextRuntimeTest_DecoratorRegistry -> Registry should contain 3 decorators");
+
+				HandleAC_0 = Registry.FindHandle(FDecoratorAC_Add::DecoratorUID);
+				AddErrorIfFalse(HandleAC_0.IsValid(), "FAnimationAnimNextRuntimeTest_DecoratorRegistry -> Decorator should have registered automatically");
+				AddErrorIfFalse(HandleAC_0.IsDynamic(), "FAnimationAnimNextRuntimeTest_DecoratorRegistry -> Decorator should have been dynamically allocated");
+				AddErrorIfFalse(HandleA != HandleAC_0, "FAnimationAnimNextRuntimeTest_DecoratorRegistry -> Decorator handles should be different");
+
+				const FDecorator* DecoratorAC_0Ptr = Registry.Find(HandleAC_0);
+				AddErrorIfFalse(DecoratorAC_0Ptr != nullptr, "FAnimationAnimNextRuntimeTest_DecoratorRegistry -> Decorator should exist");
+				AddErrorIfFalse(DecoratorAC_0Ptr->GetDecoratorUID() == FDecoratorAC_Add::DecoratorUID, "FAnimationAnimNextRuntimeTest_DecoratorRegistry -> Unexpected decorator instance type");
+				AddErrorIfFalse(&DecoratorAC_0 == DecoratorAC_0Ptr, "FAnimationAnimNextRuntimeTest_DecoratorRegistry -> Unexpected decorator instance pointer");
+
+				// Unregister our instances
+				Registry.Unregister(&DecoratorAC_0);
+
+				AddErrorIfFalse(Registry.GetNum() == 2, "FAnimationAnimNextRuntimeTest_DecoratorRegistry -> Registry should contain 2 decorators");
+				AddErrorIfFalse(!Registry.FindHandle(FDecoratorAC_Add::DecoratorUID).IsValid(), "FAnimationAnimNextRuntimeTest_DecoratorRegistry -> Decorator should have unregistered");
+			}
+
+			{
+				// Dynamically register another decorator, re-using the previous dynamic index
+				FDecoratorAC_Add DecoratorAC_1;
+				Registry.Register(&DecoratorAC_1);
+
+				AddErrorIfFalse(Registry.GetNum() == 3, "FAnimationAnimNextRuntimeTest_DecoratorRegistry -> Registry should contain 3 decorators");
+
+				FDecoratorRegistryHandle HandleAC_1 = Registry.FindHandle(FDecoratorAC_Add::DecoratorUID);
+				AddErrorIfFalse(HandleAC_1.IsValid(), "FAnimationAnimNextRuntimeTest_DecoratorRegistry -> Decorator should have registered automatically");
+				AddErrorIfFalse(HandleAC_1.IsDynamic(), "FAnimationAnimNextRuntimeTest_DecoratorRegistry -> Decorator should have been dynamically allocated");
+				AddErrorIfFalse(HandleA != HandleAC_1, "FAnimationAnimNextRuntimeTest_DecoratorRegistry -> Decorator handles should be different");
+				AddErrorIfFalse(HandleAC_0 == HandleAC_1, "FAnimationAnimNextRuntimeTest_DecoratorRegistry -> Decorator handles should be identical");
+
+				const FDecorator* DecoratorAC_1Ptr = Registry.Find(HandleAC_1);
+				AddErrorIfFalse(DecoratorAC_1Ptr != nullptr, "FAnimationAnimNextRuntimeTest_DecoratorRegistry -> Decorator should exist");
+				AddErrorIfFalse(DecoratorAC_1Ptr->GetDecoratorUID() == FDecoratorAC_Add::DecoratorUID, "FAnimationAnimNextRuntimeTest_DecoratorRegistry -> Unexpected decorator instance type");
+				AddErrorIfFalse(&DecoratorAC_1 == DecoratorAC_1Ptr, "FAnimationAnimNextRuntimeTest_DecoratorRegistry -> Unexpected decorator instance pointer");
+
+				// Unregister our instances
+				Registry.Unregister(&DecoratorAC_1);
+
+				AddErrorIfFalse(Registry.GetNum() == 2, "FAnimationAnimNextRuntimeTest_DecoratorRegistry -> Registry should contain 2 decorators");
+				AddErrorIfFalse(!Registry.FindHandle(FDecoratorAC_Add::DecoratorUID).IsValid(), "FAnimationAnimNextRuntimeTest_DecoratorRegistry -> Decorator should have unregistered");
+			}
+		}
+
+		AddErrorIfFalse(Registry.GetNum() == 1, "FAnimationAnimNextRuntimeTest_DecoratorRegistry -> Registry should contain 1 decorator");
+		AddErrorIfFalse(!Registry.FindHandle(FDecoratorAB_Add::DecoratorUID).IsValid(), "FAnimationAnimNextRuntimeTest_DecoratorRegistry -> Decorator should have unregistered automatically");
+		AddErrorIfFalse(HandleA == Registry.FindHandle(FDecoratorA_Base::DecoratorUID), "FAnimationAnimNextRuntimeTest_DecoratorRegistry -> Decorator handle should not have changed");
+	}
+
+	AddErrorIfFalse(Registry.GetNum() == 0, "FAnimationAnimNextRuntimeTest_DecoratorRegistry -> All decorators should have unregistered");
+	AddErrorIfFalse(!Registry.FindHandle(FDecoratorA_Base::DecoratorUID).IsValid(), "FAnimationAnimNextRuntimeTest_DecoratorRegistry -> Decorator should have unregistered automatically");
+	AddErrorIfFalse(!Registry.FindHandle(FDecoratorAB_Add::DecoratorUID).IsValid(), "FAnimationAnimNextRuntimeTest_DecoratorRegistry -> Decorator should have unregistered automatically");
+	AddErrorIfFalse(!Registry.FindHandle(FDecoratorAC_Add::DecoratorUID).IsValid(), "FAnimationAnimNextRuntimeTest_DecoratorRegistry -> Decorator should have unregistered automatically");
+
+	return true;
+}
+
+IMPLEMENT_SIMPLE_AUTOMATION_TEST(FAnimationAnimNextRuntimeTest_NodeTemplateRegistry, "Animation.AnimNext.Runtime.NodeTemplateRegistry", EAutomationTestFlags::EditorContext | EAutomationTestFlags::EngineFilter)
+
+bool FAnimationAnimNextRuntimeTest_NodeTemplateRegistry::RunTest(const FString& InParameters)
+{
+	using namespace UE::AnimNext;
+
+	AUTO_REGISTER_ANIM_DECORATOR(FDecoratorA_Base)
+	AUTO_REGISTER_ANIM_DECORATOR(FDecoratorAB_Add)
+	AUTO_REGISTER_ANIM_DECORATOR(FDecoratorAC_Add)
+
+	FNodeTemplateRegistry& Registry = FNodeTemplateRegistry::Get();
+
+	TArray<FDecoratorUID> NodeTemplateDecoratorList;
+	NodeTemplateDecoratorList.Add(FDecoratorA_Base::DecoratorUID);
+	NodeTemplateDecoratorList.Add(FDecoratorAB_Add::DecoratorUID);
+	NodeTemplateDecoratorList.Add(FDecoratorAC_Add::DecoratorUID);
+	NodeTemplateDecoratorList.Add(FDecoratorA_Base::DecoratorUID);
+	NodeTemplateDecoratorList.Add(FDecoratorAC_Add::DecoratorUID);
+
+	// Populate our node template registry
+	TArray<uint8> NodeTemplateBuffer0;
+	const FNodeTemplate* NodeTemplate0 = BuildNodeTemplate(NodeTemplateDecoratorList, NodeTemplateBuffer0);
+
+	AddErrorIfFalse(Registry.GetNum() == 0, "FAnimationAnimNextRuntimeTest_NodeTemplateRegistry -> Registry should contain any templates");
+
+	FNodeTemplateRegistryHandle TemplateHandle0 = Registry.FindOrAdd(NodeTemplate0);
+	AddErrorIfFalse(Registry.GetNum() == 1, "FAnimationAnimNextRuntimeTest_NodeTemplateRegistry -> Registry should contain 1 template");
+	AddErrorIfFalse(TemplateHandle0.IsValid(), "FAnimationAnimNextRuntimeTest_NodeTemplateRegistry -> Registry should contain our template");
+
+	const uint32 TemplateSize0 = NodeTemplate0->GetNodeTemplateSize();
+	const FNodeTemplate* NodeTemplate0_ = Registry.Find(TemplateHandle0);
+	AddErrorIfFalse(NodeTemplate0_ != nullptr, "FAnimationAnimNextRuntimeTest_NodeTemplateRegistry -> Registry should contain our template");
+	AddErrorIfFalse(NodeTemplate0_ != NodeTemplate0, "FAnimationAnimNextRuntimeTest_NodeTemplateRegistry -> Template pointers should be different");
+	AddErrorIfFalse(FMemory::Memcmp(NodeTemplate0, NodeTemplate0_, TemplateSize0) == 0, "FAnimationAnimNextRuntimeTest_NodeTemplateRegistry -> Templates should be identical");
+
+	// Try and register a duplicate template
+	TArray<uint8> NodeTemplateBuffer1;
+	const FNodeTemplate* NodeTemplate1 = BuildNodeTemplate(NodeTemplateDecoratorList, NodeTemplateBuffer1);
+	AddErrorIfFalse(NodeTemplate0 != NodeTemplate1, "FAnimationAnimNextRuntimeTest_NodeTemplateRegistry -> Node template pointers should be different");
+	AddErrorIfFalse(NodeTemplate0->GetUID() == NodeTemplate1->GetUID(), "FAnimationAnimNextRuntimeTest_NodeTemplateRegistry -> Node template UIDs should be identical");
+
+	FNodeTemplateRegistryHandle TemplateHandle1 = Registry.FindOrAdd(NodeTemplate1);
+	AddErrorIfFalse(Registry.GetNum() == 1, "FAnimationAnimNextRuntimeTest_NodeTemplateRegistry -> Registry should contain 1 template");
+	AddErrorIfFalse(TemplateHandle0 == TemplateHandle1, "FAnimationAnimNextRuntimeTest_NodeTemplateRegistry -> Node template handles should be identical");
+
+	// Try and register a new template
+	TArray<FDecoratorUID> NodeTemplateDecoratorList2;
+	NodeTemplateDecoratorList2.Add(FDecoratorA_Base::DecoratorUID);
+	NodeTemplateDecoratorList2.Add(FDecoratorAB_Add::DecoratorUID);
+	NodeTemplateDecoratorList2.Add(FDecoratorAC_Add::DecoratorUID);
+	NodeTemplateDecoratorList2.Add(FDecoratorAC_Add::DecoratorUID);
+
+	TArray<uint8> NodeTemplateBuffer2;
+	const FNodeTemplate* NodeTemplate2 = BuildNodeTemplate(NodeTemplateDecoratorList2, NodeTemplateBuffer2);
+	AddErrorIfFalse(NodeTemplate0->GetUID() != NodeTemplate2->GetUID(), "FAnimationAnimNextRuntimeTest_NodeTemplateRegistry -> Node template UIDs should be different");
+
+	FNodeTemplateRegistryHandle TemplateHandle2 = Registry.FindOrAdd(NodeTemplate2);
+	AddErrorIfFalse(Registry.GetNum() == 2, "FAnimationAnimNextRuntimeTest_NodeTemplateRegistry -> Registry should contain 2 templates");
+	AddErrorIfFalse(TemplateHandle0 != TemplateHandle2, "FAnimationAnimNextRuntimeTest_NodeTemplateRegistry -> Node template handles should be identical");
+	AddErrorIfFalse(TemplateHandle2.IsValid(), "FAnimationAnimNextRuntimeTest_NodeTemplateRegistry -> Registry should contain our template");
+
+	// Unregister our templates
+	Registry.Unregister(NodeTemplate0);
+	Registry.Unregister(NodeTemplate2);
+
+	AddErrorIfFalse(Registry.GetNum() == 0, "FAnimationAnimNextRuntimeTest_NodeTemplateRegistry -> Registry should contain 0 templates");
+
+	return true;
+}
+
+IMPLEMENT_SIMPLE_AUTOMATION_TEST(FAnimationAnimNextRuntimeTest_NodeLifetime, "Animation.AnimNext.Runtime.NodeLifetime", EAutomationTestFlags::EditorContext | EAutomationTestFlags::EngineFilter)
+
+bool FAnimationAnimNextRuntimeTest_NodeLifetime::RunTest(const FString& InParameters)
+{
+	using namespace UE::AnimNext;
+
+	AUTO_REGISTER_ANIM_DECORATOR(FDecoratorA_Base)
+	AUTO_REGISTER_ANIM_DECORATOR(FDecoratorAB_Add)
+	AUTO_REGISTER_ANIM_DECORATOR(FDecoratorAC_Add)
+
+	FNodeTemplateRegistry& Registry = FNodeTemplateRegistry::Get();
+
+	TArray<FDecoratorUID> NodeTemplateDecoratorList;
+	NodeTemplateDecoratorList.Add(FDecoratorA_Base::DecoratorUID);
+	NodeTemplateDecoratorList.Add(FDecoratorAB_Add::DecoratorUID);
+	NodeTemplateDecoratorList.Add(FDecoratorAC_Add::DecoratorUID);
+	NodeTemplateDecoratorList.Add(FDecoratorA_Base::DecoratorUID);
+	NodeTemplateDecoratorList.Add(FDecoratorAC_Add::DecoratorUID);
+
+	// Populate our node template registry
+	TArray<uint8> NodeTemplateBuffer0;
+	const FNodeTemplate* NodeTemplate0 = BuildNodeTemplate(NodeTemplateDecoratorList, NodeTemplateBuffer0);
+
+	FNodeTemplateRegistryHandle TemplateHandle0 = Registry.FindOrAdd(NodeTemplate0);
+	AddErrorIfFalse(Registry.GetNum() == 1, "FAnimationAnimNextRuntimeTest_NodeLifetime -> Registry should contain 1 template");
+	AddErrorIfFalse(TemplateHandle0.IsValid(), "FAnimationAnimNextRuntimeTest_NodeLifetime -> Registry should contain our template");
+
+	// Build our graph
+	uint16 NodeUID = 0;
+	TArray<uint8> GraphSharedDataBuffer;
+	const FNodeHandle Node0 = AppendNodeToGraph(*NodeTemplate0, NodeUID, GraphSharedDataBuffer);
+	const FNodeHandle Node1 = AppendNodeToGraph(*NodeTemplate0, NodeUID, GraphSharedDataBuffer);
+
+	FExecutionContext Context(GraphSharedDataBuffer);
+
+	// Validate handle bookkeeping
+	{
+		FDecoratorBinding RootBinding;					// Empty, no parent
+		FDecoratorHandle DecoratorHandle00(Node0, 0);	// Point to first node, first base decorator
+
+		// Allocate a node
+		FDecoratorPtr DecoratorPtr00 = Context.AllocateNodeInstance(RootBinding, DecoratorHandle00);
+		AddErrorIfFalse(DecoratorPtr00.IsValid(), "FAnimationAnimNextRuntimeTest_NodeLifetime -> Failed to allocate a node instance");
+		AddErrorIfFalse(DecoratorPtr00.GetDecoratorIndex() == 0, "FAnimationAnimNextRuntimeTest_NodeLifetime -> Allocated decorator pointer should point to root decorator");
+		AddErrorIfFalse(!DecoratorPtr00.IsWeak(), "FAnimationAnimNextRuntimeTest_NodeLifetime -> Allocated decorator pointer should not be weak, we have no parent");
+		AddErrorIfFalse(DecoratorPtr00.GetNodeInstance()->GetNodeHandle() == Node0, "FAnimationAnimNextRuntimeTest_NodeLifetime -> Allocated node should point to the provided node handle");
+		AddErrorIfFalse(DecoratorPtr00.GetNodeInstance()->GetReferenceCount() == 1, "FAnimationAnimNextRuntimeTest_NodeLifetime -> Allocated node should have a single reference");
+
+		{
+			FWeakDecoratorPtr WeakDecoratorPtr00(DecoratorPtr00);
+			AddErrorIfFalse(WeakDecoratorPtr00.GetNodeInstance() == DecoratorPtr00.GetNodeInstance(), "FAnimationAnimNextRuntimeTest_NodeLifetime -> Weak reference should point to the same node instance");
+			AddErrorIfFalse(WeakDecoratorPtr00.GetDecoratorIndex() == DecoratorPtr00.GetDecoratorIndex(), "FAnimationAnimNextRuntimeTest_NodeLifetime -> Weak reference should point to the same decorator index");
+			AddErrorIfFalse(DecoratorPtr00.GetNodeInstance()->GetReferenceCount() == 1, "FAnimationAnimNextRuntimeTest_NodeLifetime -> Weak references shouldn't increase ref count");
+		}
+
+		AddErrorIfFalse(DecoratorPtr00.GetNodeInstance()->GetReferenceCount() == 1, "FAnimationAnimNextRuntimeTest_NodeLifetime -> Weak references shouldn't decrease ref count");
+
+		{
+			FWeakDecoratorPtr WeakDecoratorPtr00 = DecoratorPtr00;
+			AddErrorIfFalse(WeakDecoratorPtr00.GetNodeInstance() == DecoratorPtr00.GetNodeInstance(), "FAnimationAnimNextRuntimeTest_NodeLifetime -> Weak reference should point to the same node instance");
+			AddErrorIfFalse(WeakDecoratorPtr00.GetDecoratorIndex() == DecoratorPtr00.GetDecoratorIndex(), "FAnimationAnimNextRuntimeTest_NodeLifetime -> Weak reference should point to the same decorator index");
+			AddErrorIfFalse(DecoratorPtr00.GetNodeInstance()->GetReferenceCount() == 1, "FAnimationAnimNextRuntimeTest_NodeLifetime -> Weak references shouldn't increase ref count");
+		}
+
+		AddErrorIfFalse(DecoratorPtr00.GetNodeInstance()->GetReferenceCount() == 1, "FAnimationAnimNextRuntimeTest_NodeLifetime -> Weak references shouldn't decrease ref count");
+
+		{
+			FDecoratorPtr DecoratorPtr00_1(DecoratorPtr00);
+			AddErrorIfFalse(DecoratorPtr00_1.GetNodeInstance() == DecoratorPtr00.GetNodeInstance(), "FAnimationAnimNextRuntimeTest_NodeLifetime -> Strong reference should point to the same node instance");
+			AddErrorIfFalse(DecoratorPtr00_1.GetDecoratorIndex() == DecoratorPtr00.GetDecoratorIndex(), "FAnimationAnimNextRuntimeTest_NodeLifetime -> Strong reference should point to the same decorator index");
+			AddErrorIfFalse(DecoratorPtr00.GetNodeInstance()->GetReferenceCount() == 2, "FAnimationAnimNextRuntimeTest_NodeLifetime -> Strong references should increase ref count");
+		}
+
+		AddErrorIfFalse(DecoratorPtr00.GetNodeInstance()->GetReferenceCount() == 1, "FAnimationAnimNextRuntimeTest_NodeLifetime -> Strong references should decrease ref count");
+
+		{
+			FDecoratorPtr DecoratorPtr00_1 = DecoratorPtr00;
+			AddErrorIfFalse(DecoratorPtr00_1.GetNodeInstance() == DecoratorPtr00.GetNodeInstance(), "FAnimationAnimNextRuntimeTest_NodeLifetime -> Strong reference should point to the same node instance");
+			AddErrorIfFalse(DecoratorPtr00_1.GetDecoratorIndex() == DecoratorPtr00.GetDecoratorIndex(), "FAnimationAnimNextRuntimeTest_NodeLifetime -> Strong reference should point to the same decorator index");
+			AddErrorIfFalse(DecoratorPtr00.GetNodeInstance()->GetReferenceCount() == 2, "FAnimationAnimNextRuntimeTest_NodeLifetime -> Strong references should increase ref count");
+		}
+
+		AddErrorIfFalse(DecoratorPtr00.GetNodeInstance()->GetReferenceCount() == 1, "FAnimationAnimNextRuntimeTest_NodeLifetime -> Strong references should decrease ref count");
+	}
+
+	// Validate parent support
+	{
+		FDecoratorBinding RootBinding;					// Empty, no parent
+		FDecoratorHandle DecoratorHandle00(Node0, 0);	// Point to first node, first base decorator
+		FDecoratorHandle DecoratorHandle03(Node0, 3);	// Point to first node, second base decorator
+		FDecoratorHandle DecoratorHandle10(Node1, 0);	// Point to second node, first base decorator
+
+		// Allocate our first node
+		FDecoratorPtr DecoratorPtr00 = Context.AllocateNodeInstance(RootBinding, DecoratorHandle00);
+		AddErrorIfFalse(DecoratorPtr00.IsValid(), "FAnimationAnimNextRuntimeTest_NodeLifetime -> Failed to allocate a node instance");
+
+		// Allocate a new node, using the first as a parent
+		// Both decorators live on the same node, the returned handle should be weak on the parent
+		FDecoratorPtr DecoratorPtr03 = Context.AllocateNodeInstance(DecoratorPtr00, DecoratorHandle03);
+		AddErrorIfFalse(DecoratorPtr03.IsValid(), "FAnimationAnimNextRuntimeTest_NodeLifetime -> Failed to allocate a node instance");
+		AddErrorIfFalse(DecoratorPtr03.GetDecoratorIndex() == 3, "FAnimationAnimNextRuntimeTest_NodeLifetime -> Allocated decorator pointer should point to fourth decorator");
+		AddErrorIfFalse(DecoratorPtr03.IsWeak(), "FAnimationAnimNextRuntimeTest_NodeLifetime -> Allocated decorator pointer should be weak, we have the same parent");
+		AddErrorIfFalse(DecoratorPtr03.GetNodeInstance()->GetNodeHandle() == Node0, "FAnimationAnimNextRuntimeTest_NodeLifetime -> Allocated node should point to the provided node handle");
+		AddErrorIfFalse(DecoratorPtr03.GetNodeInstance() == DecoratorPtr00.GetNodeInstance(), "FAnimationAnimNextRuntimeTest_NodeLifetime -> Handles should point to the same node instance");
+		AddErrorIfFalse(DecoratorPtr03.GetNodeInstance()->GetReferenceCount() == 1, "FAnimationAnimNextRuntimeTest_NodeLifetime -> Allocated node should have one reference");
+
+		// Allocate a new node, using the first as a parent
+		// The second decorator lives on a new node, a new node instance will be allocated
+		FDecoratorPtr DecoratorPtr10 = Context.AllocateNodeInstance(DecoratorPtr00, DecoratorHandle10);
+		AddErrorIfFalse(DecoratorPtr10.IsValid(), "FAnimationAnimNextRuntimeTest_NodeLifetime -> Failed to allocate a node instance");
+		AddErrorIfFalse(DecoratorPtr10.GetDecoratorIndex() == 0, "FAnimationAnimNextRuntimeTest_NodeLifetime -> Allocated decorator pointer should point to first decorator");
+		AddErrorIfFalse(!DecoratorPtr10.IsWeak(), "FAnimationAnimNextRuntimeTest_NodeLifetime -> Allocated decorator pointer should not be weak, we have the same parent but a different node handle");
+		AddErrorIfFalse(DecoratorPtr10.GetNodeInstance()->GetNodeHandle() == Node1, "FAnimationAnimNextRuntimeTest_NodeLifetime -> Allocated node should point to the provided node handle");
+		AddErrorIfFalse(DecoratorPtr10.GetNodeInstance() != DecoratorPtr00.GetNodeInstance(), "FAnimationAnimNextRuntimeTest_NodeLifetime -> Handles should not point to the same node instance");
+		AddErrorIfFalse(DecoratorPtr10.GetNodeInstance()->GetReferenceCount() == 1, "FAnimationAnimNextRuntimeTest_NodeLifetime -> Allocated node should have one reference");
+	}
+
+	// Validate constructors and destructors
+	{
+		TArray<FDecoratorUID> ConstructedDecorators;
+		TArray<FDecoratorUID> DestructedDecorators;
+
+		Private::ConstructedDecorators = &ConstructedDecorators;
+		Private::DestructedDecorators = &DestructedDecorators;
+
+		{
+			FDecoratorBinding RootBinding;					// Empty, no parent
+			FDecoratorHandle DecoratorHandle00(Node0, 0);	// Point to first node, first base decorator
+
+			// Allocate our node instance
+			FDecoratorPtr DecoratorPtr00 = Context.AllocateNodeInstance(RootBinding, DecoratorHandle00);
+			AddErrorIfFalse(DecoratorPtr00.IsValid(), "FAnimationAnimNextRuntimeTest_NodeLifetime -> Failed to allocate a node instance");
+
+			// Validate instance constructors
+			AddErrorIfFalse(ConstructedDecorators.Num() == 5, "FAnimationAnimNextRuntimeTest_NodeLifetime -> Expected all 5 decorators to have been constructed");
+			AddErrorIfFalse(DestructedDecorators.Num() == 0, "FAnimationAnimNextRuntimeTest_NodeLifetime -> Expected no decorators to have been destructed");
+			AddErrorIfFalse(ConstructedDecorators[0] == NodeTemplateDecoratorList[0], "FAnimationAnimNextRuntimeTest_NodeLifetime -> Unexpected constructor order");
+			AddErrorIfFalse(ConstructedDecorators[1] == NodeTemplateDecoratorList[1], "FAnimationAnimNextRuntimeTest_NodeLifetime -> Unexpected constructor order");
+			AddErrorIfFalse(ConstructedDecorators[2] == NodeTemplateDecoratorList[2], "FAnimationAnimNextRuntimeTest_NodeLifetime -> Unexpected constructor order");
+			AddErrorIfFalse(ConstructedDecorators[3] == NodeTemplateDecoratorList[3], "FAnimationAnimNextRuntimeTest_NodeLifetime -> Unexpected constructor order");
+			AddErrorIfFalse(ConstructedDecorators[4] == NodeTemplateDecoratorList[4], "FAnimationAnimNextRuntimeTest_NodeLifetime -> Unexpected constructor order");
+
+			// Destruct our node instance
+		}
+
+		// Validate instance destructors
+		AddErrorIfFalse(ConstructedDecorators.Num() == 5, "FAnimationAnimNextRuntimeTest_NodeLifetime -> Expected no decorators to have been constructed");
+		AddErrorIfFalse(DestructedDecorators.Num() == 5, "FAnimationAnimNextRuntimeTest_NodeLifetime -> Expected all 5 decorators to have been destructed");
+		AddErrorIfFalse(DestructedDecorators[0] == NodeTemplateDecoratorList[4], "FAnimationAnimNextRuntimeTest_NodeLifetime -> Unexpected destructor order");
+		AddErrorIfFalse(DestructedDecorators[1] == NodeTemplateDecoratorList[3], "FAnimationAnimNextRuntimeTest_NodeLifetime -> Unexpected destructor order");
+		AddErrorIfFalse(DestructedDecorators[2] == NodeTemplateDecoratorList[2], "FAnimationAnimNextRuntimeTest_NodeLifetime -> Unexpected destructor order");
+		AddErrorIfFalse(DestructedDecorators[3] == NodeTemplateDecoratorList[1], "FAnimationAnimNextRuntimeTest_NodeLifetime -> Unexpected destructor order");
+		AddErrorIfFalse(DestructedDecorators[4] == NodeTemplateDecoratorList[0], "FAnimationAnimNextRuntimeTest_NodeLifetime -> Unexpected destructor order");
+
+		Private::ConstructedDecorators = nullptr;
+		Private::DestructedDecorators = nullptr;
+	}
+
+	// Unregister our templates
+	Registry.Unregister(NodeTemplate0);
+
+	AddErrorIfFalse(Registry.GetNum() == 0, "FAnimationAnimNextRuntimeTest_NodeLifetime -> Registry should contain 0 templates");
+
+	return true;
+}
+
+IMPLEMENT_SIMPLE_AUTOMATION_TEST(FAnimationAnimNextRuntimeTest_GetInterface, "Animation.AnimNext.Runtime.GetInterface", EAutomationTestFlags::EditorContext | EAutomationTestFlags::EngineFilter)
+
+bool FAnimationAnimNextRuntimeTest_GetInterface::RunTest(const FString& InParameters)
+{
+	using namespace UE::AnimNext;
+
+	AUTO_REGISTER_ANIM_DECORATOR(FDecoratorA_Base)
+	AUTO_REGISTER_ANIM_DECORATOR(FDecoratorAB_Add)
+	AUTO_REGISTER_ANIM_DECORATOR(FDecoratorAC_Add)
+
+	FNodeTemplateRegistry& Registry = FNodeTemplateRegistry::Get();
+
+	TArray<FDecoratorUID> NodeTemplateDecoratorList;
+	NodeTemplateDecoratorList.Add(FDecoratorA_Base::DecoratorUID);
+	NodeTemplateDecoratorList.Add(FDecoratorAB_Add::DecoratorUID);
+	NodeTemplateDecoratorList.Add(FDecoratorAC_Add::DecoratorUID);
+	NodeTemplateDecoratorList.Add(FDecoratorA_Base::DecoratorUID);
+	NodeTemplateDecoratorList.Add(FDecoratorAC_Add::DecoratorUID);
+
+	// Populate our node template registry
+	TArray<uint8> NodeTemplateBuffer0;
+	const FNodeTemplate* NodeTemplate0 = BuildNodeTemplate(NodeTemplateDecoratorList, NodeTemplateBuffer0);
+
+	FNodeTemplateRegistryHandle TemplateHandle0 = Registry.FindOrAdd(NodeTemplate0);
+	AddErrorIfFalse(Registry.GetNum() == 1, "FAnimationAnimNextRuntimeTest_GetInterface -> Registry should contain 1 template");
+	AddErrorIfFalse(TemplateHandle0.IsValid(), "FAnimationAnimNextRuntimeTest_GetInterface -> Registry should contain our template");
+
+	// Build our graph
+	uint16 NodeUID = 0;
+	TArray<uint8> GraphSharedDataBuffer;
+
+	const FNodeHandle Node0 = AppendNodeToGraph(*NodeTemplate0, NodeUID, GraphSharedDataBuffer);
+	InitNodeDecorator(*NodeTemplate0, Node0, GraphSharedDataBuffer, 0, FDecoratorA_Base::FSharedData());
+	InitNodeDecorator(*NodeTemplate0, Node0, GraphSharedDataBuffer, 1, FDecoratorAB_Add::FSharedData());
+	InitNodeDecorator(*NodeTemplate0, Node0, GraphSharedDataBuffer, 2, FDecoratorAC_Add::FSharedData());
+	InitNodeDecorator(*NodeTemplate0, Node0, GraphSharedDataBuffer, 3, FDecoratorA_Base::FSharedData());
+	InitNodeDecorator(*NodeTemplate0, Node0, GraphSharedDataBuffer, 4, FDecoratorAC_Add::FSharedData());
+
+	FExecutionContext Context(GraphSharedDataBuffer);
+
+	// Validate from the first base decorator
+	{
+		FDecoratorBinding ParentBinding;				// Empty, no parent
+		FDecoratorHandle DecoratorHandle00(Node0, 0);	// Point to first node, first base decorator
+
+		FDecoratorPtr DecoratorPtr00 = Context.AllocateNodeInstance(ParentBinding, DecoratorHandle00);
+		AddErrorIfFalse(DecoratorPtr00.IsValid(), "FAnimationAnimNextRuntimeTest_GetInterface -> Failed to allocate a node instance");
+
+		// Validate GetInterface from a decorator handle
+		TDecoratorBinding<IInterfaceC> Binding00C;
+		AddErrorIfFalse(Context.GetInterface(DecoratorPtr00, Binding00C), "FAnimationAnimNextRuntimeTest_GetInterface -> InterfaceC not found");
+		AddErrorIfFalse(Binding00C.IsValid(), "FAnimationAnimNextRuntimeTest_GetInterface -> InterfaceC binding not valid");
+		AddErrorIfFalse(Binding00C.GetInterfaceUID() == IInterfaceC::InterfaceUID, "FAnimationAnimNextRuntimeTest_GetInterface -> Unexpected interface UID found in decorator binding");
+		AddErrorIfFalse(Binding00C.GetDecoratorPtr().GetDecoratorIndex() == 2, "FAnimationAnimNextRuntimeTest_GetInterface -> InterfaceC not found on expected decorator");
+		AddErrorIfFalse(Binding00C.GetDecoratorPtr().GetNodeInstance()->GetNodeHandle() == Node0, "FAnimationAnimNextRuntimeTest_GetInterface -> InterfaceC not found on expected node");
+		AddErrorIfFalse(Binding00C.GetSharedData<FDecoratorAC_Add::FSharedData>()->DecoratorUID == FDecoratorAC_Add::DecoratorUID, "FAnimationAnimNextRuntimeTest_GetInterface -> Unexpected shared data in decorator binding");
+		AddErrorIfFalse(Binding00C.GetInstanceData<FDecoratorAC_Add::FInstanceData>()->DecoratorUID == FDecoratorAC_Add::DecoratorUID, "FAnimationAnimNextRuntimeTest_GetInterface -> Unexpected instance data in decorator binding");
+
+		TDecoratorBinding<IInterfaceB> Binding00B;
+		AddErrorIfFalse(Context.GetInterface(DecoratorPtr00, Binding00B), "FAnimationAnimNextRuntimeTest_GetInterface -> InterfaceB not found");
+		AddErrorIfFalse(Binding00B.IsValid(), "FAnimationAnimNextRuntimeTest_GetInterface -> InterfaceB binding not valid");
+		AddErrorIfFalse(Binding00B.GetInterfaceUID() == IInterfaceB::InterfaceUID, "FAnimationAnimNextRuntimeTest_GetInterface -> Unexpected interface UID found in decorator binding");
+		AddErrorIfFalse(Binding00B.GetDecoratorPtr().GetDecoratorIndex() == 1, "FAnimationAnimNextRuntimeTest_GetInterface -> InterfaceB not found on expected decorator");
+		AddErrorIfFalse(Binding00B.GetDecoratorPtr().GetNodeInstance()->GetNodeHandle() == Node0, "FAnimationAnimNextRuntimeTest_GetInterface -> InterfaceB not found on expected node");
+		AddErrorIfFalse(Binding00B.GetSharedData<FDecoratorAB_Add::FSharedData>()->DecoratorUID == FDecoratorAB_Add::DecoratorUID, "FAnimationAnimNextRuntimeTest_GetInterface -> Unexpected shared data in decorator binding");
+		AddErrorIfFalse(Binding00B.GetInstanceData<FDecoratorAB_Add::FInstanceData>()->DecoratorUID == FDecoratorAB_Add::DecoratorUID, "FAnimationAnimNextRuntimeTest_GetInterface -> Unexpected instance data in decorator binding");
+
+		TDecoratorBinding<IInterfaceA> Binding00A;
+		AddErrorIfFalse(Context.GetInterface(DecoratorPtr00, Binding00A), "FAnimationAnimNextRuntimeTest_GetInterface -> InterfaceA not found");
+		AddErrorIfFalse(Binding00A.IsValid(), "FAnimationAnimNextRuntimeTest_GetInterface -> InterfaceA binding not valid");
+		AddErrorIfFalse(Binding00A.GetInterfaceUID() == IInterfaceA::InterfaceUID, "FAnimationAnimNextRuntimeTest_GetInterface -> Unexpected interface UID found in decorator binding");
+		AddErrorIfFalse(Binding00A.GetDecoratorPtr().GetDecoratorIndex() == 2, "FAnimationAnimNextRuntimeTest_GetInterface -> InterfaceA not found on expected decorator");
+		AddErrorIfFalse(Binding00A.GetDecoratorPtr().GetNodeInstance()->GetNodeHandle() == Node0, "FAnimationAnimNextRuntimeTest_GetInterface -> InterfaceA not found on expected node");
+		AddErrorIfFalse(Binding00A.GetSharedData<FDecoratorAC_Add::FSharedData>()->DecoratorUID == FDecoratorAC_Add::DecoratorUID, "FAnimationAnimNextRuntimeTest_GetInterface -> Unexpected shared data in decorator binding");
+		AddErrorIfFalse(Binding00A.GetInstanceData<FDecoratorAC_Add::FInstanceData>()->DecoratorUID == FDecoratorAC_Add::DecoratorUID, "FAnimationAnimNextRuntimeTest_GetInterface -> Unexpected instance data in decorator binding");
+
+		// Validate GetInterface from a decorator binding
+		{
+			{
+				{
+					TDecoratorBinding<IInterfaceC> Binding00C_;
+					AddErrorIfFalse(Context.GetInterface(Binding00C, Binding00C_), "FAnimationAnimNextRuntimeTest_GetInterface -> InterfaceC not found");
+					AddErrorIfFalse(Binding00C == Binding00C_, "FAnimationAnimNextRuntimeTest_GetInterface -> GetInterface methods should return the same result");
+				}
+
+				{
+					TDecoratorBinding<IInterfaceC> Binding00C_;
+					AddErrorIfFalse(Context.GetInterface(Binding00B, Binding00C_), "FAnimationAnimNextRuntimeTest_GetInterface -> InterfaceC not found");
+					AddErrorIfFalse(Binding00C == Binding00C_, "FAnimationAnimNextRuntimeTest_GetInterface -> GetInterface methods should return the same result");
+				}
+
+				{
+					TDecoratorBinding<IInterfaceC> Binding00C_;
+					AddErrorIfFalse(Context.GetInterface(Binding00A, Binding00C_), "FAnimationAnimNextRuntimeTest_GetInterface -> InterfaceC not found");
+					AddErrorIfFalse(Binding00C == Binding00C_, "FAnimationAnimNextRuntimeTest_GetInterface -> GetInterface methods should return the same result");
+				}
+			}
+
+			{
+				{
+					TDecoratorBinding<IInterfaceB> Binding00B_;
+					AddErrorIfFalse(Context.GetInterface(Binding00C, Binding00B_), "FAnimationAnimNextRuntimeTest_GetInterface -> InterfaceB not found");
+					AddErrorIfFalse(Binding00B == Binding00B_, "FAnimationAnimNextRuntimeTest_GetInterface -> GetInterface methods should return the same result");
+				}
+
+				{
+					TDecoratorBinding<IInterfaceB> Binding00B_;
+					AddErrorIfFalse(Context.GetInterface(Binding00B, Binding00B_), "FAnimationAnimNextRuntimeTest_GetInterface -> InterfaceB not found");
+					AddErrorIfFalse(Binding00B == Binding00B_, "FAnimationAnimNextRuntimeTest_GetInterface -> GetInterface methods should return the same result");
+				}
+
+				{
+					TDecoratorBinding<IInterfaceB> Binding00B_;
+					AddErrorIfFalse(Context.GetInterface(Binding00A, Binding00B_), "FAnimationAnimNextRuntimeTest_GetInterface -> InterfaceB not found");
+					AddErrorIfFalse(Binding00B == Binding00B_, "FAnimationAnimNextRuntimeTest_GetInterface -> GetInterface methods should return the same result");
+				}
+			}
+
+			{
+				{
+					TDecoratorBinding<IInterfaceA> Binding00A_;
+					AddErrorIfFalse(Context.GetInterface(Binding00C, Binding00A_), "FAnimationAnimNextRuntimeTest_GetInterface -> InterfaceA not found");
+					AddErrorIfFalse(Binding00A == Binding00A_, "FAnimationAnimNextRuntimeTest_GetInterface -> GetInterface methods should return the same result");
+				}
+
+				{
+					TDecoratorBinding<IInterfaceA> Binding00A_;
+					AddErrorIfFalse(Context.GetInterface(Binding00B, Binding00A_), "FAnimationAnimNextRuntimeTest_GetInterface -> InterfaceA not found");
+					AddErrorIfFalse(Binding00A == Binding00A_, "FAnimationAnimNextRuntimeTest_GetInterface -> GetInterface methods should return the same result");
+				}
+
+				{
+					TDecoratorBinding<IInterfaceA> Binding00A_;
+					AddErrorIfFalse(Context.GetInterface(Binding00A, Binding00A_), "FAnimationAnimNextRuntimeTest_GetInterface -> InterfaceA not found");
+					AddErrorIfFalse(Binding00A == Binding00A_, "FAnimationAnimNextRuntimeTest_GetInterface -> GetInterface methods should return the same result");
+				}
+			}
+		}
+	}
+
+	// Validate from the second base decorator
+	{
+		FDecoratorBinding ParentBinding;				// Empty, no parent
+		FDecoratorHandle DecoratorHandle03(Node0, 3);	// Point to first node, second base decorator
+
+		FDecoratorPtr DecoratorPtr03 = Context.AllocateNodeInstance(ParentBinding, DecoratorHandle03);
+		AddErrorIfFalse(DecoratorPtr03.IsValid(), "FAnimationAnimNextRuntimeTest_GetInterface -> Failed to allocate a node instance");
+
+		// Validate GetInterface from a decorator handle
+		TDecoratorBinding<IInterfaceC> Binding03C;
+		AddErrorIfFalse(Context.GetInterface(DecoratorPtr03, Binding03C), "FAnimationAnimNextRuntimeTest_GetInterface -> InterfaceC not found");
+		AddErrorIfFalse(Binding03C.IsValid(), "FAnimationAnimNextRuntimeTest_GetInterface -> InterfaceC binding not valid");
+		AddErrorIfFalse(Binding03C.GetInterfaceUID() == IInterfaceC::InterfaceUID, "FAnimationAnimNextRuntimeTest_GetInterface -> Unexpected interface UID found in decorator binding");
+		AddErrorIfFalse(Binding03C.GetDecoratorPtr().GetDecoratorIndex() == 4, "FAnimationAnimNextRuntimeTest_GetInterface -> InterfaceC not found on expected decorator");
+		AddErrorIfFalse(Binding03C.GetDecoratorPtr().GetNodeInstance()->GetNodeHandle() == Node0, "FAnimationAnimNextRuntimeTest_GetInterface -> InterfaceC not found on expected node");
+		AddErrorIfFalse(Binding03C.GetSharedData<FDecoratorAC_Add::FSharedData>()->DecoratorUID == FDecoratorAC_Add::DecoratorUID, "FAnimationAnimNextRuntimeTest_GetInterface -> Unexpected shared data in decorator binding");
+		AddErrorIfFalse(Binding03C.GetInstanceData<FDecoratorAC_Add::FInstanceData>()->DecoratorUID == FDecoratorAC_Add::DecoratorUID, "FAnimationAnimNextRuntimeTest_GetInterface -> Unexpected instance data in decorator binding");
+
+		TDecoratorBinding<IInterfaceB> Binding03B;
+		AddErrorIfFalse(!Context.GetInterface(DecoratorPtr03, Binding03B), "FAnimationAnimNextRuntimeTest_GetInterface -> InterfaceB found");
+		AddErrorIfFalse(!Binding03B.IsValid(), "FAnimationAnimNextRuntimeTest_GetInterface -> InterfaceB binding should not be valid");
+
+		TDecoratorBinding<IInterfaceA> Binding03A;
+		AddErrorIfFalse(Context.GetInterface(DecoratorPtr03, Binding03A), "FAnimationAnimNextRuntimeTest_GetInterface -> InterfaceA not found");
+		AddErrorIfFalse(Binding03A.IsValid(), "FAnimationAnimNextRuntimeTest_GetInterface -> InterfaceA binding not valid");
+		AddErrorIfFalse(Binding03A.GetInterfaceUID() == IInterfaceA::InterfaceUID, "FAnimationAnimNextRuntimeTest_GetInterface -> Unexpected interface UID found in decorator binding");
+		AddErrorIfFalse(Binding03A.GetDecoratorPtr().GetDecoratorIndex() == 4, "FAnimationAnimNextRuntimeTest_GetInterface -> InterfaceA not found on expected decorator");
+		AddErrorIfFalse(Binding03A.GetDecoratorPtr().GetNodeInstance()->GetNodeHandle() == Node0, "FAnimationAnimNextRuntimeTest_GetInterface -> InterfaceA not found on expected node");
+		AddErrorIfFalse(Binding03A.GetSharedData<FDecoratorAC_Add::FSharedData>()->DecoratorUID == FDecoratorAC_Add::DecoratorUID, "FAnimationAnimNextRuntimeTest_GetInterface -> Unexpected shared data in decorator binding");
+		AddErrorIfFalse(Binding03A.GetInstanceData<FDecoratorAC_Add::FInstanceData>()->DecoratorUID == FDecoratorAC_Add::DecoratorUID, "FAnimationAnimNextRuntimeTest_GetInterface -> Unexpected instance data in decorator binding");
+
+		// Validate GetInterface from a decorator binding
+		{
+			{
+				{
+					TDecoratorBinding<IInterfaceC> Binding03C_;
+					AddErrorIfFalse(Context.GetInterface(Binding03C, Binding03C_), "FAnimationAnimNextRuntimeTest_GetInterface -> InterfaceC not found");
+					AddErrorIfFalse(Binding03C == Binding03C_, "FAnimationAnimNextRuntimeTest_GetInterface -> GetInterface methods should return the same result");
+				}
+
+				{
+					TDecoratorBinding<IInterfaceC> Binding03C_;
+					AddErrorIfFalse(!Context.GetInterface(Binding03B, Binding03C_), "FAnimationAnimNextRuntimeTest_GetInterface -> InterfaceC found");
+				}
+
+				{
+					TDecoratorBinding<IInterfaceC> Binding03C_;
+					AddErrorIfFalse(Context.GetInterface(Binding03A, Binding03C_), "FAnimationAnimNextRuntimeTest_GetInterface -> InterfaceC not found");
+					AddErrorIfFalse(Binding03C == Binding03C_, "FAnimationAnimNextRuntimeTest_GetInterface -> GetInterface methods should return the same result");
+				}
+			}
+
+			{
+				{
+					TDecoratorBinding<IInterfaceB> Binding03B_;
+					AddErrorIfFalse(!Context.GetInterface(Binding03C, Binding03B_), "FAnimationAnimNextRuntimeTest_GetInterface -> InterfaceB found");
+					AddErrorIfFalse(Binding03B == Binding03B_, "FAnimationAnimNextRuntimeTest_GetInterface -> GetInterface methods should return the same result");
+				}
+
+				{
+					TDecoratorBinding<IInterfaceB> Binding03B_;
+					AddErrorIfFalse(!Context.GetInterface(Binding03B, Binding03B_), "FAnimationAnimNextRuntimeTest_GetInterface -> InterfaceB found");
+					AddErrorIfFalse(Binding03B == Binding03B_, "FAnimationAnimNextRuntimeTest_GetInterface -> GetInterface methods should return the same result");
+				}
+
+				{
+					TDecoratorBinding<IInterfaceB> Binding03B_;
+					AddErrorIfFalse(!Context.GetInterface(Binding03A, Binding03B_), "FAnimationAnimNextRuntimeTest_GetInterface -> InterfaceB found");
+					AddErrorIfFalse(Binding03B == Binding03B_, "FAnimationAnimNextRuntimeTest_GetInterface -> GetInterface methods should return the same result");
+				}
+			}
+
+			{
+				{
+					TDecoratorBinding<IInterfaceA> Binding03A_;
+					AddErrorIfFalse(Context.GetInterface(Binding03C, Binding03A_), "FAnimationAnimNextRuntimeTest_GetInterface -> InterfaceA not found");
+					AddErrorIfFalse(Binding03A == Binding03A_, "FAnimationAnimNextRuntimeTest_GetInterface -> GetInterface methods should return the same result");
+				}
+
+				{
+					TDecoratorBinding<IInterfaceA> Binding03A_;
+					AddErrorIfFalse(!Context.GetInterface(Binding03B, Binding03A_), "FAnimationAnimNextRuntimeTest_GetInterface -> InterfaceA found");
+				}
+
+				{
+					TDecoratorBinding<IInterfaceA> Binding03A_;
+					AddErrorIfFalse(Context.GetInterface(Binding03A, Binding03A_), "FAnimationAnimNextRuntimeTest_GetInterface -> InterfaceA not found");
+					AddErrorIfFalse(Binding03A == Binding03A_, "FAnimationAnimNextRuntimeTest_GetInterface -> GetInterface methods should return the same result");
+				}
+			}
+		}
+	}
+
+	Registry.Unregister(NodeTemplate0);
+
+	AddErrorIfFalse(Registry.GetNum() == 0, "FAnimationAnimNextRuntimeTest_GetInterface -> Registry should contain 0 templates");
+
+	return true;
+}
+
+IMPLEMENT_SIMPLE_AUTOMATION_TEST(FAnimationAnimNextRuntimeTest_GetInterfaceSuper, "Animation.AnimNext.Runtime.GetInterfaceSuper", EAutomationTestFlags::EditorContext | EAutomationTestFlags::EngineFilter)
+
+bool FAnimationAnimNextRuntimeTest_GetInterfaceSuper::RunTest(const FString& InParameters)
+{
+	using namespace UE::AnimNext;
+
+	AUTO_REGISTER_ANIM_DECORATOR(FDecoratorA_Base)
+	AUTO_REGISTER_ANIM_DECORATOR(FDecoratorAB_Add)
+	AUTO_REGISTER_ANIM_DECORATOR(FDecoratorAC_Add)
+
+	FNodeTemplateRegistry& Registry = FNodeTemplateRegistry::Get();
+
+	TArray<FDecoratorUID> NodeTemplateDecoratorList;
+	NodeTemplateDecoratorList.Add(FDecoratorA_Base::DecoratorUID);
+	NodeTemplateDecoratorList.Add(FDecoratorAB_Add::DecoratorUID);
+	NodeTemplateDecoratorList.Add(FDecoratorAC_Add::DecoratorUID);
+	NodeTemplateDecoratorList.Add(FDecoratorA_Base::DecoratorUID);
+	NodeTemplateDecoratorList.Add(FDecoratorAC_Add::DecoratorUID);
+
+	// Populate our node template registry
+	TArray<uint8> NodeTemplateBuffer0;
+	const FNodeTemplate* NodeTemplate0 = BuildNodeTemplate(NodeTemplateDecoratorList, NodeTemplateBuffer0);
+
+	FNodeTemplateRegistryHandle TemplateHandle0 = Registry.FindOrAdd(NodeTemplate0);
+	AddErrorIfFalse(Registry.GetNum() == 1, "FAnimationAnimNextRuntimeTest_GetInterfaceSuper -> Registry should contain 1 template");
+	AddErrorIfFalse(TemplateHandle0.IsValid(), "FAnimationAnimNextRuntimeTest_GetInterfaceSuper -> Registry should contain our template");
+
+	// Build our graph
+	uint16 NodeUID = 0;
+	TArray<uint8> GraphSharedDataBuffer;
+
+	const FNodeHandle Node0 = AppendNodeToGraph(*NodeTemplate0, NodeUID, GraphSharedDataBuffer);
+	InitNodeDecorator(*NodeTemplate0, Node0, GraphSharedDataBuffer, 0, FDecoratorA_Base::FSharedData());
+	InitNodeDecorator(*NodeTemplate0, Node0, GraphSharedDataBuffer, 1, FDecoratorAB_Add::FSharedData());
+	InitNodeDecorator(*NodeTemplate0, Node0, GraphSharedDataBuffer, 2, FDecoratorAC_Add::FSharedData());
+	InitNodeDecorator(*NodeTemplate0, Node0, GraphSharedDataBuffer, 3, FDecoratorA_Base::FSharedData());
+	InitNodeDecorator(*NodeTemplate0, Node0, GraphSharedDataBuffer, 4, FDecoratorAC_Add::FSharedData());
+
+	FExecutionContext Context(GraphSharedDataBuffer);
+
+	// Validate from the first base decorator
+	{
+		FDecoratorBinding ParentBinding;				// Empty, no parent
+		FDecoratorHandle DecoratorHandle00(Node0, 0);	// Point to first node, first base decorator
+
+		FDecoratorPtr DecoratorPtr00 = Context.AllocateNodeInstance(ParentBinding, DecoratorHandle00);
+		AddErrorIfFalse(DecoratorPtr00.IsValid(), "FAnimationAnimNextRuntimeTest_GetInterfaceSuper -> Failed to allocate a node instance");
+
+		{
+			// Get a valid decorator binding: FDecoratorAC_Add
+			TDecoratorBinding<IInterfaceC> Binding02C;
+			AddErrorIfFalse(Context.GetInterface(DecoratorPtr00, Binding02C), "FAnimationAnimNextRuntimeTest_GetInterfaceSuper -> InterfaceC not found");
+
+			// Validate GetInterfaceSuper from a decorator handle
+			TDecoratorBinding<IInterfaceC> SuperBinding02C;
+			AddErrorIfFalse(!Context.GetInterfaceSuper(Binding02C.GetDecoratorPtr(), SuperBinding02C), "FAnimationAnimNextRuntimeTest_GetInterfaceSuper -> InterfaceC found");
+
+			// Validate GetInterfaceSuper from a decorator binding
+			TDecoratorBinding<IInterfaceC> SuperBinding02C_;
+			AddErrorIfFalse(!Context.GetInterfaceSuper(Binding02C, SuperBinding02C_), "FAnimationAnimNextRuntimeTest_GetInterfaceSuper -> InterfaceC found");
+			AddErrorIfFalse(SuperBinding02C == SuperBinding02C_, "FAnimationAnimNextRuntimeTest_GetInterfaceSuper -> GetInterfaceSuper methods should return the same result");
+		}
+
+		{
+			// Get a valid decorator binding: FDecoratorAC_Add
+			TDecoratorBinding<IInterfaceA> Binding02A;
+			AddErrorIfFalse(Context.GetInterface(DecoratorPtr00, Binding02A), "FAnimationAnimNextRuntimeTest_GetInterfaceSuper -> InterfaceA not found");
+
+			// Validate GetInterfaceSuper from a decorator handle, FDecoratorAB_Add
+			TDecoratorBinding<IInterfaceA> SuperBinding02A;
+			AddErrorIfFalse(Context.GetInterfaceSuper(Binding02A.GetDecoratorPtr(), SuperBinding02A), "FAnimationAnimNextRuntimeTest_GetInterfaceSuper -> InterfaceA not found");
+			AddErrorIfFalse(SuperBinding02A.IsValid(), "FAnimationAnimNextRuntimeTest_GetInterfaceSuper -> InterfaceA binding not valid");
+			AddErrorIfFalse(SuperBinding02A.GetInterfaceUID() == IInterfaceA::InterfaceUID, "FAnimationAnimNextRuntimeTest_GetInterfaceSuper -> Unexpected interface UID found in decorator binding");
+			AddErrorIfFalse(SuperBinding02A.GetDecoratorPtr().GetDecoratorIndex() == 1, "FAnimationAnimNextRuntimeTest_GetInterfaceSuper -> InterfaceA not found on expected decorator");
+			AddErrorIfFalse(SuperBinding02A.GetDecoratorPtr().GetNodeInstance()->GetNodeHandle() == Node0, "FAnimationAnimNextRuntimeTest_GetInterfaceSuper -> InterfaceA not found on expected node");
+			AddErrorIfFalse(SuperBinding02A.GetSharedData<FDecoratorAB_Add::FSharedData>()->DecoratorUID == FDecoratorAB_Add::DecoratorUID, "FAnimationAnimNextRuntimeTest_GetInterfaceSuper -> Unexpected shared data in decorator binding");
+			AddErrorIfFalse(SuperBinding02A.GetInstanceData<FDecoratorAB_Add::FInstanceData>()->DecoratorUID == FDecoratorAB_Add::DecoratorUID, "FAnimationAnimNextRuntimeTest_GetInterfaceSuper -> Unexpected instance data in decorator binding");
+
+			// Validate GetInterfaceSuper from a decorator binding, FDecoratorAB_Add
+			TDecoratorBinding<IInterfaceA> SuperBinding02A_;
+			AddErrorIfFalse(Context.GetInterfaceSuper(Binding02A, SuperBinding02A_), "FAnimationAnimNextRuntimeTest_GetInterfaceSuper -> InterfaceA not found");
+			AddErrorIfFalse(SuperBinding02A == SuperBinding02A_, "FAnimationAnimNextRuntimeTest_GetInterfaceSuper -> GetInterfaceSuper methods should return the same result");
+
+			// Validate GetInterfaceSuper from a decorator handle, FDecoratorA_Base
+			TDecoratorBinding<IInterfaceA> SuperBinding01A;
+			AddErrorIfFalse(Context.GetInterfaceSuper(SuperBinding02A.GetDecoratorPtr(), SuperBinding01A), "FAnimationAnimNextRuntimeTest_GetInterfaceSuper -> InterfaceA not found");
+			AddErrorIfFalse(SuperBinding01A.IsValid(), "FAnimationAnimNextRuntimeTest_GetInterfaceSuper -> InterfaceA binding not valid");
+			AddErrorIfFalse(SuperBinding01A.GetInterfaceUID() == IInterfaceA::InterfaceUID, "FAnimationAnimNextRuntimeTest_GetInterfaceSuper -> Unexpected interface UID found in decorator binding");
+			AddErrorIfFalse(SuperBinding01A.GetDecoratorPtr().GetDecoratorIndex() == 0, "FAnimationAnimNextRuntimeTest_GetInterfaceSuper -> InterfaceA not found on expected decorator");
+			AddErrorIfFalse(SuperBinding01A.GetDecoratorPtr().GetNodeInstance()->GetNodeHandle() == Node0, "FAnimationAnimNextRuntimeTest_GetInterfaceSuper -> InterfaceA not found on expected node");
+			AddErrorIfFalse(SuperBinding01A.GetSharedData<FDecoratorA_Base::FSharedData>()->DecoratorUID == FDecoratorA_Base::DecoratorUID, "FAnimationAnimNextRuntimeTest_GetInterfaceSuper -> Unexpected shared data in decorator binding");
+			AddErrorIfFalse(SuperBinding01A.GetInstanceData<FDecoratorA_Base::FInstanceData>()->DecoratorUID == FDecoratorA_Base::DecoratorUID, "FAnimationAnimNextRuntimeTest_GetInterfaceSuper -> Unexpected instance data in decorator binding");
+
+			// Validate GetInterfaceSuper from a decorator binding, FDecoratorA_Base
+			TDecoratorBinding<IInterfaceA> SuperBinding01A_;
+			AddErrorIfFalse(Context.GetInterfaceSuper(SuperBinding02A, SuperBinding01A_), "FAnimationAnimNextRuntimeTest_GetInterfaceSuper -> InterfaceA not found");
+			AddErrorIfFalse(SuperBinding01A == SuperBinding01A_, "FAnimationAnimNextRuntimeTest_GetInterfaceSuper -> GetInterfaceSuper methods should return the same result");
+
+			// Validate GetInterfaceSuper from a decorator handle
+			TDecoratorBinding<IInterfaceA> SuperBinding00A;
+			AddErrorIfFalse(!Context.GetInterfaceSuper(SuperBinding01A.GetDecoratorPtr(), SuperBinding00A), "FAnimationAnimNextRuntimeTest_GetInterfaceSuper -> InterfaceA found");
+
+			// Validate GetInterfaceSuper from a decorator binding
+			TDecoratorBinding<IInterfaceA> SuperBinding00A_;
+			AddErrorIfFalse(!Context.GetInterfaceSuper(SuperBinding01A, SuperBinding00A_), "FAnimationAnimNextRuntimeTest_GetInterfaceSuper -> InterfaceA found");
+			AddErrorIfFalse(SuperBinding00A == SuperBinding00A_, "FAnimationAnimNextRuntimeTest_GetInterfaceSuper -> GetInterfaceSuper methods should return the same result");
+		}
+	}
+
+	// Validate from the second base decorator
+	{
+		FDecoratorBinding ParentBinding;				// Empty, no parent
+		FDecoratorHandle DecoratorHandle03(Node0, 3);	// Point to first node, second base decorator
+
+		FDecoratorPtr DecoratorPtr03 = Context.AllocateNodeInstance(ParentBinding, DecoratorHandle03);
+		AddErrorIfFalse(DecoratorPtr03.IsValid(), "FAnimationAnimNextRuntimeTest_GetInterfaceSuper -> Failed to allocate a node instance");
+
+		{
+			// Get a valid decorator binding: FDecoratorAC_Add
+			TDecoratorBinding<IInterfaceC> Binding04C;
+			AddErrorIfFalse(Context.GetInterface(DecoratorPtr03, Binding04C), "FAnimationAnimNextRuntimeTest_GetInterfaceSuper -> InterfaceC not found");
+
+			// Validate GetInterfaceSuper from a decorator handle
+			TDecoratorBinding<IInterfaceC> SuperBinding04C;
+			AddErrorIfFalse(!Context.GetInterfaceSuper(Binding04C.GetDecoratorPtr(), SuperBinding04C), "FAnimationAnimNextRuntimeTest_GetInterfaceSuper -> InterfaceC found");
+
+			// Validate GetInterfaceSuper from a decorator binding
+			TDecoratorBinding<IInterfaceC> SuperBinding04C_;
+			AddErrorIfFalse(!Context.GetInterfaceSuper(Binding04C, SuperBinding04C_), "FAnimationAnimNextRuntimeTest_GetInterfaceSuper -> InterfaceC found");
+			AddErrorIfFalse(SuperBinding04C == SuperBinding04C_, "FAnimationAnimNextRuntimeTest_GetInterfaceSuper -> GetInterfaceSuper methods should return the same result");
+		}
+
+		{
+			// Get a valid decorator binding: FDecoratorAC_Add
+			TDecoratorBinding<IInterfaceA> Binding04A;
+			AddErrorIfFalse(Context.GetInterface(DecoratorPtr03, Binding04A), "FAnimationAnimNextRuntimeTest_GetInterfaceSuper -> InterfaceA not found");
+
+			// Validate GetInterfaceSuper from a decorator handle, FDecoratorA_Base
+			TDecoratorBinding<IInterfaceA> SuperBinding04A;
+			AddErrorIfFalse(Context.GetInterfaceSuper(Binding04A.GetDecoratorPtr(), SuperBinding04A), "FAnimationAnimNextRuntimeTest_GetInterfaceSuper -> InterfaceA not found");
+			AddErrorIfFalse(SuperBinding04A.IsValid(), "FAnimationAnimNextRuntimeTest_GetInterfaceSuper -> InterfaceA binding not valid");
+			AddErrorIfFalse(SuperBinding04A.GetInterfaceUID() == IInterfaceA::InterfaceUID, "FAnimationAnimNextRuntimeTest_GetInterfaceSuper -> Unexpected interface UID found in decorator binding");
+			AddErrorIfFalse(SuperBinding04A.GetDecoratorPtr().GetDecoratorIndex() == 3, "FAnimationAnimNextRuntimeTest_GetInterfaceSuper -> InterfaceA not found on expected decorator");
+			AddErrorIfFalse(SuperBinding04A.GetDecoratorPtr().GetNodeInstance()->GetNodeHandle() == Node0, "FAnimationAnimNextRuntimeTest_GetInterfaceSuper -> InterfaceA not found on expected node");
+			AddErrorIfFalse(SuperBinding04A.GetSharedData<FDecoratorA_Base::FSharedData>()->DecoratorUID == FDecoratorA_Base::DecoratorUID, "FAnimationAnimNextRuntimeTest_GetInterfaceSuper -> Unexpected shared data in decorator binding");
+			AddErrorIfFalse(SuperBinding04A.GetInstanceData<FDecoratorA_Base::FInstanceData>()->DecoratorUID == FDecoratorA_Base::DecoratorUID, "FAnimationAnimNextRuntimeTest_GetInterfaceSuper -> Unexpected instance data in decorator binding");
+
+			// Validate GetInterfaceSuper from a decorator binding, FDecoratorA_Base
+			TDecoratorBinding<IInterfaceA> SuperBinding04A_;
+			AddErrorIfFalse(Context.GetInterfaceSuper(Binding04A, SuperBinding04A_), "FAnimationAnimNextRuntimeTest_GetInterfaceSuper -> InterfaceA not found");
+			AddErrorIfFalse(SuperBinding04A == SuperBinding04A_, "FAnimationAnimNextRuntimeTest_GetInterfaceSuper -> GetInterfaceSuper methods should return the same result");
+
+			// Validate GetInterfaceSuper from a decorator handle
+			TDecoratorBinding<IInterfaceA> SuperBinding03A;
+			AddErrorIfFalse(!Context.GetInterfaceSuper(SuperBinding04A.GetDecoratorPtr(), SuperBinding03A), "FAnimationAnimNextRuntimeTest_GetInterfaceSuper -> InterfaceA found");
+
+			// Validate GetInterfaceSuper from a decorator binding
+			TDecoratorBinding<IInterfaceA> SuperBinding03A_;
+			AddErrorIfFalse(!Context.GetInterfaceSuper(SuperBinding04A, SuperBinding03A_), "FAnimationAnimNextRuntimeTest_GetInterfaceSuper -> InterfaceA found");
+			AddErrorIfFalse(SuperBinding03A == SuperBinding03A_, "FAnimationAnimNextRuntimeTest_GetInterfaceSuper -> GetInterfaceSuper methods should return the same result");
+		}
+	}
+
+	Registry.Unregister(NodeTemplate0);
+
+	AddErrorIfFalse(Registry.GetNum() == 0, "FAnimationAnimNextRuntimeTest_GetInterface -> Registry should contain 0 templates");
+
+	return true;
+}
+#endif
diff --git a/Engine/Plugins/Experimental/AnimNext/Source/AnimNextTestSuite/Private/AnimNextDecoratorInterfacesTest.cpp b/Engine/Plugins/Experimental/AnimNext/Source/AnimNextTestSuite/Private/AnimNextDecoratorInterfacesTest.cpp
new file mode 100644
--- /dev/null
+++ b/Engine/Plugins/Experimental/AnimNext/Source/AnimNextTestSuite/Private/AnimNextDecoratorInterfacesTest.cpp
@@ -0,0 +1,527 @@
+// Copyright Epic Games, Inc. All Rights Reserved.
+
+
+#include "AnimNextRuntimeTest.h"
+
+#include "Misc/AutomationTest.h"
+
+#if WITH_DEV_AUTOMATION_TESTS
+
+#include "DecoratorBase/Decorator.h"
+#include "DecoratorBase/DecoratorRegistry.h"
+#include "DecoratorBase/ExecutionContext.h"
+#include "DecoratorBase/IDecoratorInterface.h"
+#include "DecoratorBase/NodeInstance.h"
+#include "DecoratorBase/NodeTemplateRegistry.h"
+#include "DecoratorInterfaces/IEvaluate.h"
+#include "DecoratorInterfaces/IHierarchy.h"
+#include "DecoratorInterfaces/IUpdate.h"
+
+//****************************************************************************
+// AnimNext Runtime DecoratorInterfaces Tests
+//****************************************************************************
+
+namespace UE::AnimNext
+{
+	namespace Private
+	{
+		static TArray<FDecoratorUID>* UpdatedDecorators = nullptr;
+		static TArray<FDecoratorUID>* EvaluatedDecorators = nullptr;
+	}
+
+	struct FDecoratorWithNoChildren : FDecorator, IHierarchy, IUpdate, IEvaluate
+	{
+		DECLARE_ANIM_DECORATOR(FDecoratorWithNoChildren, 0xe2400d2a, FDecorator)
+
+		// FDecorator impl
+		virtual EDecoratorMode GetMode() const override { return EDecoratorMode::Base; }
+
+		// IHierarchy impl
+		virtual void GetChildren(FExecutionContext& Context, const TDecoratorBinding<IHierarchy>& Binding, FChildrenArray& Children) const override
+		{
+			IHierarchy::GetChildren(Context, Binding, Children);
+		}
+
+		// IUpdate impl
+		virtual void PreUpdate(FExecutionContext& Context, const TDecoratorBinding<IUpdate>& Binding) const override
+		{
+			if (Private::UpdatedDecorators != nullptr)
+			{
+				Private::UpdatedDecorators->Add(FDecoratorWithNoChildren::DecoratorUID);
+			}
+
+			IUpdate::PreUpdate(Context, Binding);
+		}
+
+		virtual void PostUpdate(FExecutionContext& Context, const TDecoratorBinding<IUpdate>& Binding) const override
+		{
+			if (Private::UpdatedDecorators != nullptr)
+			{
+				Private::UpdatedDecorators->Add(FDecoratorWithNoChildren::DecoratorUID);
+			}
+
+			IUpdate::PostUpdate(Context, Binding);
+		}
+
+		// IEvaluate impl
+		virtual void PreEvaluate(FExecutionContext& Context, const TDecoratorBinding<IEvaluate>& Binding) const override
+		{
+			if (Private::EvaluatedDecorators != nullptr)
+			{
+				Private::EvaluatedDecorators->Add(FDecoratorWithNoChildren::DecoratorUID);
+			}
+
+			IEvaluate::PreEvaluate(Context, Binding);
+		}
+
+		virtual void PostEvaluate(FExecutionContext& Context, const TDecoratorBinding<IEvaluate>& Binding) const override
+		{
+			if (Private::EvaluatedDecorators != nullptr)
+			{
+				Private::EvaluatedDecorators->Add(FDecoratorWithNoChildren::DecoratorUID);
+			}
+
+			IEvaluate::PostEvaluate(Context, Binding);
+		}
+	};
+
+	DEFINE_ANIM_DECORATOR_BEGIN(FDecoratorWithNoChildren)
+		DEFINE_ANIM_DECORATOR_IMPLEMENTS_INTERFACE(IHierarchy)
+		DEFINE_ANIM_DECORATOR_IMPLEMENTS_INTERFACE(IUpdate)
+		DEFINE_ANIM_DECORATOR_IMPLEMENTS_INTERFACE(IEvaluate)
+	DEFINE_ANIM_DECORATOR_END(FDecoratorWithNoChildren)
+
+	// This decorator does not update or evaluate
+	struct FDecoratorWithOneChild : FDecorator, IHierarchy
+	{
+		DECLARE_ANIM_DECORATOR(FDecoratorWithOneChild, 0xba24d224, FDecorator)
+
+		struct FSharedData : FDecorator::FSharedData
+		{
+			FDecoratorHandle Child;
+
+			FSharedData(FDecoratorHandle Child_)
+				: Child(Child_)
+			{}
+		};
+
+		struct FInstanceData : FDecorator::FInstanceData
+		{
+			FDecoratorPtr Child;
+
+			void Construct(FExecutionContext& Context, FWeakDecoratorPtr DecoratorPtr, const FSharedData& SharedData)
+			{
+				Child = Context.AllocateNodeInstance(DecoratorPtr, SharedData.Child);
+			}
+		};
+
+		// FDecorator impl
+		virtual EDecoratorMode GetMode() const override { return EDecoratorMode::Base; }
+
+		// IHierarchy impl
+		virtual void GetChildren(FExecutionContext& Context, const TDecoratorBinding<IHierarchy>& Binding, FChildrenArray& Children) const override
+		{
+			const FInstanceData* InstanceData = Binding.GetInstanceData<FInstanceData>();
+
+			Children.Add(InstanceData->Child);
+
+			IHierarchy::GetChildren(Context, Binding, Children);
+		}
+	};
+
+	DEFINE_ANIM_DECORATOR_BEGIN(FDecoratorWithOneChild)
+		DEFINE_ANIM_DECORATOR_IMPLEMENTS_INTERFACE(IHierarchy)
+	DEFINE_ANIM_DECORATOR_END(FDecoratorWithOneChild)
+
+	struct FDecoratorWithChildren : FDecorator, IHierarchy, IUpdate, IEvaluate
+	{
+		DECLARE_ANIM_DECORATOR(FDecoratorWithChildren, 0xa3ad93b9, FDecorator)
+
+		struct FSharedData : FDecorator::FSharedData
+		{
+			FDecoratorHandle Children[2];
+
+			FSharedData(FDecoratorHandle Child0, FDecoratorHandle Child1)
+				: Children{ Child0, Child1 }
+			{}
+		};
+
+		struct FInstanceData : FDecorator::FInstanceData
+		{
+			FDecoratorPtr Children[2];
+
+			void Construct(FExecutionContext& Context, FWeakDecoratorPtr DecoratorPtr, const FSharedData& SharedData)
+			{
+				Children[0] = Context.AllocateNodeInstance(DecoratorPtr, SharedData.Children[0]);
+				Children[1] = Context.AllocateNodeInstance(DecoratorPtr, SharedData.Children[1]);
+			}
+		};
+
+		// FDecorator impl
+		virtual EDecoratorMode GetMode() const override { return EDecoratorMode::Base; }
+
+		// IHierarchy impl
+		virtual void GetChildren(FExecutionContext& Context, const TDecoratorBinding<IHierarchy>& Binding, FChildrenArray& Children) const override
+		{
+			const FInstanceData* InstanceData = Binding.GetInstanceData<FInstanceData>();
+
+			Children.Add(InstanceData->Children[0]);
+			Children.Add(InstanceData->Children[1]);
+
+			IHierarchy::GetChildren(Context, Binding, Children);
+		}
+
+		// IUpdate impl
+		virtual void PreUpdate(FExecutionContext& Context, const TDecoratorBinding<IUpdate>& Binding) const override
+		{
+			if (Private::UpdatedDecorators != nullptr)
+			{
+				Private::UpdatedDecorators->Add(FDecoratorWithChildren::DecoratorUID);
+			}
+
+			IUpdate::PreUpdate(Context, Binding);
+		}
+
+		virtual void PostUpdate(FExecutionContext& Context, const TDecoratorBinding<IUpdate>& Binding) const override
+		{
+			if (Private::UpdatedDecorators != nullptr)
+			{
+				Private::UpdatedDecorators->Add(FDecoratorWithChildren::DecoratorUID);
+			}
+
+			IUpdate::PostUpdate(Context, Binding);
+		}
+
+		// IEvaluate impl
+		virtual void PreEvaluate(FExecutionContext& Context, const TDecoratorBinding<IEvaluate>& Binding) const override
+		{
+			if (Private::EvaluatedDecorators != nullptr)
+			{
+				Private::EvaluatedDecorators->Add(FDecoratorWithChildren::DecoratorUID);
+			}
+
+			IEvaluate::PreEvaluate(Context, Binding);
+		}
+
+		virtual void PostEvaluate(FExecutionContext& Context, const TDecoratorBinding<IEvaluate>& Binding) const override
+		{
+			if (Private::EvaluatedDecorators != nullptr)
+			{
+				Private::EvaluatedDecorators->Add(FDecoratorWithChildren::DecoratorUID);
+			}
+
+			IEvaluate::PostEvaluate(Context, Binding);
+		}
+	};
+
+	DEFINE_ANIM_DECORATOR_BEGIN(FDecoratorWithChildren)
+		DEFINE_ANIM_DECORATOR_IMPLEMENTS_INTERFACE(IHierarchy)
+		DEFINE_ANIM_DECORATOR_IMPLEMENTS_INTERFACE(IUpdate)
+		DEFINE_ANIM_DECORATOR_IMPLEMENTS_INTERFACE(IEvaluate)
+	DEFINE_ANIM_DECORATOR_END(FDecoratorWithChildren)
+}
+
+IMPLEMENT_SIMPLE_AUTOMATION_TEST(FAnimationAnimNextRuntimeTest_IHierarchy, "Animation.AnimNext.Runtime.IHierarchy", EAutomationTestFlags::EditorContext | EAutomationTestFlags::EngineFilter)
+
+bool FAnimationAnimNextRuntimeTest_IHierarchy::RunTest(const FString& InParameters)
+{
+	using namespace UE::AnimNext;
+
+	AUTO_REGISTER_ANIM_DECORATOR(FDecoratorWithNoChildren)
+	AUTO_REGISTER_ANIM_DECORATOR(FDecoratorWithOneChild)
+	AUTO_REGISTER_ANIM_DECORATOR(FDecoratorWithChildren)
+
+	FNodeTemplateRegistry& Registry = FNodeTemplateRegistry::Get();
+
+	// We create a few node templates
+	// Template A has a single decorator with no children
+	TArray<FDecoratorUID> NodeTemplateDecoratorListA;
+	NodeTemplateDecoratorListA.Add(FDecoratorWithNoChildren::DecoratorUID);
+
+	// Template B has a single decorator with one child
+	TArray<FDecoratorUID> NodeTemplateDecoratorListB;
+	NodeTemplateDecoratorListB.Add(FDecoratorWithOneChild::DecoratorUID);
+
+	// Template C has two decorators, each with one child
+	TArray<FDecoratorUID> NodeTemplateDecoratorListC;
+	NodeTemplateDecoratorListC.Add(FDecoratorWithOneChild::DecoratorUID);
+	NodeTemplateDecoratorListC.Add(FDecoratorWithOneChild::DecoratorUID);
+
+	// Template D has a single decorator with children
+	TArray<FDecoratorUID> NodeTemplateDecoratorListD;
+	NodeTemplateDecoratorListD.Add(FDecoratorWithChildren::DecoratorUID);
+
+	// Populate our node template registry
+	TArray<uint8> NodeTemplateBufferA, NodeTemplateBufferB, NodeTemplateBufferC, NodeTemplateBufferD;
+	const FNodeTemplate* NodeTemplateA = BuildNodeTemplate(NodeTemplateDecoratorListA, NodeTemplateBufferA);
+	const FNodeTemplate* NodeTemplateB = BuildNodeTemplate(NodeTemplateDecoratorListB, NodeTemplateBufferB);
+	const FNodeTemplate* NodeTemplateC = BuildNodeTemplate(NodeTemplateDecoratorListC, NodeTemplateBufferC);
+	const FNodeTemplate* NodeTemplateD = BuildNodeTemplate(NodeTemplateDecoratorListD, NodeTemplateBufferD);
+
+	// Build our graph, it as follow (each node template has a single node instance):
+	// NodeA has no children
+	// NodeB has one child: NodeA
+	// NodeC has two children: NodeA and NodeB (but both decorators are base, only NodeB will be referenced)
+	// NodeD has two children: NodeA and NodeC
+
+	uint16 NodeUID = 0;
+	TArray<uint8> GraphSharedDataBuffer;
+
+	const FNodeHandle NodeA = AppendNodeToGraph(*NodeTemplateA, NodeUID, GraphSharedDataBuffer);
+	InitNodeDecorator(*NodeTemplateA, NodeA, GraphSharedDataBuffer, 0, FDecoratorWithNoChildren::FSharedData());
+
+	const FNodeHandle NodeB = AppendNodeToGraph(*NodeTemplateB, NodeUID, GraphSharedDataBuffer);
+	InitNodeDecorator(*NodeTemplateB, NodeB, GraphSharedDataBuffer, 0, FDecoratorWithOneChild::FSharedData(FDecoratorHandle(NodeA)));
+
+	const FNodeHandle NodeC = AppendNodeToGraph(*NodeTemplateC, NodeUID, GraphSharedDataBuffer);
+	InitNodeDecorator(*NodeTemplateC, NodeC, GraphSharedDataBuffer, 0, FDecoratorWithOneChild::FSharedData(FDecoratorHandle(NodeA)));
+	InitNodeDecorator(*NodeTemplateC, NodeC, GraphSharedDataBuffer, 1, FDecoratorWithOneChild::FSharedData(FDecoratorHandle(NodeB)));
+
+	const FNodeHandle NodeD = AppendNodeToGraph(*NodeTemplateD, NodeUID, GraphSharedDataBuffer);
+	InitNodeDecorator(*NodeTemplateD, NodeD, GraphSharedDataBuffer, 0, FDecoratorWithChildren::FSharedData(FDecoratorHandle(NodeA), FDecoratorHandle(NodeC, 1)));
+
+	FExecutionContext Context(GraphSharedDataBuffer);
+
+	{
+		FMemMark Mark(FMemStack::Get());
+
+		FWeakDecoratorPtr NullPtr;				// Empty, no parent
+		FDecoratorHandle RootHandle(NodeD);		// Point to NodeD, first base decorator
+
+		FDecoratorPtr NodeDPtr = Context.AllocateNodeInstance(NullPtr, RootHandle);
+		AddErrorIfFalse(NodeDPtr.IsValid(), "FAnimationAnimNextRuntimeTest_IHierarchy -> Failed to allocate root node instance");
+
+		TDecoratorBinding<IHierarchy> HierarchyBindingNodeD;
+		AddErrorIfFalse(Context.GetInterface(NodeDPtr, HierarchyBindingNodeD), "FAnimationAnimNextRuntimeTest_IHierarchy -> IHierarchy not found");
+
+		FChildrenArray ChildrenNodeD;
+		HierarchyBindingNodeD.GetChildren(Context, ChildrenNodeD);
+
+		AddErrorIfFalse(ChildrenNodeD.Num() == 2, "FAnimationAnimNextRuntimeTest_IHierarchy -> Expected 2 children");
+		AddErrorIfFalse(ChildrenNodeD[0].IsValid() && ChildrenNodeD[0].GetNodeInstance()->GetNodeHandle() == NodeA, "FAnimationAnimNextRuntimeTest_IHierarchy -> Expected child: NodeA");
+		AddErrorIfFalse(ChildrenNodeD[1].IsValid() && ChildrenNodeD[1].GetNodeInstance()->GetNodeHandle() == NodeC, "FAnimationAnimNextRuntimeTest_IHierarchy -> Expected child: NodeC");
+
+		{
+			TDecoratorBinding<IHierarchy> HierarchyBindingNodeA;
+			AddErrorIfFalse(Context.GetInterface(ChildrenNodeD[0], HierarchyBindingNodeA), "FAnimationAnimNextRuntimeTest_IHierarchy -> IHierarchy not found");
+
+			FChildrenArray ChildrenNodeA;
+			HierarchyBindingNodeA.GetChildren(Context, ChildrenNodeA);
+
+			AddErrorIfFalse(ChildrenNodeA.Num() == 0, "FAnimationAnimNextRuntimeTest_IHierarchy -> Expected 0 children");
+		}
+
+		{
+			TDecoratorBinding<IHierarchy> HierarchyBindingNodeC;
+			AddErrorIfFalse(Context.GetInterface(ChildrenNodeD[1], HierarchyBindingNodeC), "FAnimationAnimNextRuntimeTest_IHierarchy -> IHierarchy not found");
+
+			FChildrenArray ChildrenNodeC;
+			HierarchyBindingNodeC.GetChildren(Context, ChildrenNodeC);
+
+			AddErrorIfFalse(ChildrenNodeC.Num() == 1, "FAnimationAnimNextRuntimeTest_IHierarchy -> Expected 2 children");
+			AddErrorIfFalse(ChildrenNodeC[0].IsValid() && ChildrenNodeC[0].GetNodeInstance()->GetNodeHandle() == NodeB, "FAnimationAnimNextRuntimeTest_IHierarchy -> Expected child: NodeB");
+
+			{
+				TDecoratorBinding<IHierarchy> HierarchyBindingNodeB;
+				AddErrorIfFalse(Context.GetInterface(ChildrenNodeC[0], HierarchyBindingNodeB), "FAnimationAnimNextRuntimeTest_IHierarchy -> IHierarchy not found");
+
+				FChildrenArray ChildrenNodeB;
+				HierarchyBindingNodeB.GetChildren(Context, ChildrenNodeB);
+
+				AddErrorIfFalse(ChildrenNodeB.Num() == 1, "FAnimationAnimNextRuntimeTest_IHierarchy -> Expected 1 child");
+				AddErrorIfFalse(ChildrenNodeB[0].IsValid() && ChildrenNodeB[0].GetNodeInstance()->GetNodeHandle() == NodeA, "FAnimationAnimNextRuntimeTest_IHierarchy -> Expected child: NodeA");
+
+				{
+					TDecoratorBinding<IHierarchy> HierarchyBindingNodeA;
+					AddErrorIfFalse(Context.GetInterface(ChildrenNodeB[0], HierarchyBindingNodeA), "FAnimationAnimNextRuntimeTest_IHierarchy -> IHierarchy not found");
+
+					FChildrenArray ChildrenNodeA;
+					HierarchyBindingNodeA.GetChildren(Context, ChildrenNodeA);
+
+					AddErrorIfFalse(ChildrenNodeA.Num() == 0, "FAnimationAnimNextRuntimeTest_IHierarchy -> Expected 0 children");
+				}
+			}
+		}
+	}
+
+	Registry.Unregister(NodeTemplateA);
+	Registry.Unregister(NodeTemplateB);
+	Registry.Unregister(NodeTemplateC);
+	Registry.Unregister(NodeTemplateD);
+
+	AddErrorIfFalse(Registry.GetNum() == 0, "FAnimationAnimNextRuntimeTest_IHierarchy -> Registry should contain 0 templates");
+
+	return true;
+}
+
+IMPLEMENT_SIMPLE_AUTOMATION_TEST(FAnimationAnimNextRuntimeTest_IUpdate, "Animation.AnimNext.Runtime.IUpdate", EAutomationTestFlags::EditorContext | EAutomationTestFlags::EngineFilter)
+
+bool FAnimationAnimNextRuntimeTest_IUpdate::RunTest(const FString& InParameters)
+{
+	using namespace UE::AnimNext;
+
+	AUTO_REGISTER_ANIM_DECORATOR(FDecoratorWithNoChildren)
+	AUTO_REGISTER_ANIM_DECORATOR(FDecoratorWithOneChild)
+	AUTO_REGISTER_ANIM_DECORATOR(FDecoratorWithChildren)
+
+	FNodeTemplateRegistry& Registry = FNodeTemplateRegistry::Get();
+
+	// We create a few node templates
+	// Template A has a single decorator with no children
+	TArray<FDecoratorUID> NodeTemplateDecoratorListA;
+	NodeTemplateDecoratorListA.Add(FDecoratorWithNoChildren::DecoratorUID);
+
+	// Template B has a single decorator with one child, it doesn't update
+	TArray<FDecoratorUID> NodeTemplateDecoratorListB;
+	NodeTemplateDecoratorListB.Add(FDecoratorWithOneChild::DecoratorUID);
+
+	// Template C has a single decorator with children
+	TArray<FDecoratorUID> NodeTemplateDecoratorListC;
+	NodeTemplateDecoratorListC.Add(FDecoratorWithChildren::DecoratorUID);
+
+	// Populate our node template registry
+	TArray<uint8> NodeTemplateBufferA, NodeTemplateBufferB, NodeTemplateBufferC;
+	const FNodeTemplate* NodeTemplateA = BuildNodeTemplate(NodeTemplateDecoratorListA, NodeTemplateBufferA);
+	const FNodeTemplate* NodeTemplateB = BuildNodeTemplate(NodeTemplateDecoratorListB, NodeTemplateBufferB);
+	const FNodeTemplate* NodeTemplateC = BuildNodeTemplate(NodeTemplateDecoratorListC, NodeTemplateBufferC);
+
+	// Build our graph, it as follow (each node template has a single node instance):
+	// NodeA has no children
+	// NodeB has one child: NodeA (it doesn't update)
+	// NodeD has two children: NodeA and NodeB
+
+	uint16 NodeUID = 0;
+	TArray<uint8> GraphSharedDataBuffer;
+
+	const FNodeHandle NodeA = AppendNodeToGraph(*NodeTemplateA, NodeUID, GraphSharedDataBuffer);
+	InitNodeDecorator(*NodeTemplateA, NodeA, GraphSharedDataBuffer, 0, FDecoratorWithNoChildren::FSharedData());
+
+	const FNodeHandle NodeB = AppendNodeToGraph(*NodeTemplateB, NodeUID, GraphSharedDataBuffer);
+	InitNodeDecorator(*NodeTemplateB, NodeB, GraphSharedDataBuffer, 0, FDecoratorWithOneChild::FSharedData(FDecoratorHandle(NodeA)));
+
+	const FNodeHandle NodeC = AppendNodeToGraph(*NodeTemplateC, NodeUID, GraphSharedDataBuffer);
+	InitNodeDecorator(*NodeTemplateC, NodeC, GraphSharedDataBuffer, 0, FDecoratorWithChildren::FSharedData(FDecoratorHandle(NodeA), FDecoratorHandle(NodeB)));
+
+	FExecutionContext Context(GraphSharedDataBuffer);
+
+	{
+		TArray<FDecoratorUID> UpdatedDecorators;
+
+		Private::UpdatedDecorators = &UpdatedDecorators;
+
+		FWeakDecoratorPtr NullPtr;				// Empty, no parent
+		FDecoratorHandle RootHandle(NodeC);		// Point to NodeC, first base decorator
+
+		FDecoratorPtr NodeCPtr = Context.AllocateNodeInstance(NullPtr, RootHandle);
+		AddErrorIfFalse(NodeCPtr.IsValid(), "FAnimationAnimNextRuntimeTest_IUpdate -> Failed to allocate root node instance");
+
+		// Call pre/post update on our graph
+		UpdateGraph(Context, NodeCPtr);
+
+		AddErrorIfFalse(UpdatedDecorators.Num() == 6, "FAnimationAnimNextRuntimeTest_IUpdate -> Expected 6 nodes to have been visited during the update traversal");
+		AddErrorIfFalse(UpdatedDecorators[0] == FDecoratorWithChildren::DecoratorUID, "FAnimationAnimNextRuntimeTest_IUpdate -> Unexpected update order");		// NodeC
+		AddErrorIfFalse(UpdatedDecorators[1] == FDecoratorWithNoChildren::DecoratorUID, "FAnimationAnimNextRuntimeTest_IUpdate -> Unexpected update order");		// NodeA
+		AddErrorIfFalse(UpdatedDecorators[2] == FDecoratorWithNoChildren::DecoratorUID, "FAnimationAnimNextRuntimeTest_IUpdate -> Unexpected update order");		// NodeA
+		AddErrorIfFalse(UpdatedDecorators[3] == FDecoratorWithNoChildren::DecoratorUID, "FAnimationAnimNextRuntimeTest_IUpdate -> Unexpected update order");		// NodeB -> NodeA
+		AddErrorIfFalse(UpdatedDecorators[4] == FDecoratorWithNoChildren::DecoratorUID, "FAnimationAnimNextRuntimeTest_IUpdate -> Unexpected update order");		// NodeB -> NodeA
+		AddErrorIfFalse(UpdatedDecorators[5] == FDecoratorWithChildren::DecoratorUID, "FAnimationAnimNextRuntimeTest_IUpdate -> Unexpected update order");		// NodeC
+
+		Private::UpdatedDecorators = nullptr;
+	}
+
+	Registry.Unregister(NodeTemplateA);
+	Registry.Unregister(NodeTemplateB);
+	Registry.Unregister(NodeTemplateC);
+
+	AddErrorIfFalse(Registry.GetNum() == 0, "FAnimationAnimNextRuntimeTest_IUpdate -> Registry should contain 0 templates");
+
+	return true;
+}
+
+IMPLEMENT_SIMPLE_AUTOMATION_TEST(FAnimationAnimNextRuntimeTest_IEvaluate, "Animation.AnimNext.Runtime.IEvaluate", EAutomationTestFlags::EditorContext | EAutomationTestFlags::EngineFilter)
+
+bool FAnimationAnimNextRuntimeTest_IEvaluate::RunTest(const FString& InParameters)
+{
+	using namespace UE::AnimNext;
+
+	AUTO_REGISTER_ANIM_DECORATOR(FDecoratorWithNoChildren)
+	AUTO_REGISTER_ANIM_DECORATOR(FDecoratorWithOneChild)
+	AUTO_REGISTER_ANIM_DECORATOR(FDecoratorWithChildren)
+
+	FNodeTemplateRegistry& Registry = FNodeTemplateRegistry::Get();
+
+	// We create a few node templates
+	// Template A has a single decorator with no children
+	TArray<FDecoratorUID> NodeTemplateDecoratorListA;
+	NodeTemplateDecoratorListA.Add(FDecoratorWithNoChildren::DecoratorUID);
+
+	// Template B has a single decorator with one child, it doesn't evaluate
+	TArray<FDecoratorUID> NodeTemplateDecoratorListB;
+	NodeTemplateDecoratorListB.Add(FDecoratorWithOneChild::DecoratorUID);
+
+	// Template C has a single decorator with children
+	TArray<FDecoratorUID> NodeTemplateDecoratorListC;
+	NodeTemplateDecoratorListC.Add(FDecoratorWithChildren::DecoratorUID);
+
+	// Populate our node template registry
+	TArray<uint8> NodeTemplateBufferA, NodeTemplateBufferB, NodeTemplateBufferC;
+	const FNodeTemplate* NodeTemplateA = BuildNodeTemplate(NodeTemplateDecoratorListA, NodeTemplateBufferA);
+	const FNodeTemplate* NodeTemplateB = BuildNodeTemplate(NodeTemplateDecoratorListB, NodeTemplateBufferB);
+	const FNodeTemplate* NodeTemplateC = BuildNodeTemplate(NodeTemplateDecoratorListC, NodeTemplateBufferC);
+
+	// Build our graph, it as follow (each node template has a single node instance):
+	// NodeA has no children
+	// NodeB has one child: NodeA (it doesn't evaluate)
+	// NodeD has two children: NodeA and NodeB
+
+	uint16 NodeUID = 0;
+	TArray<uint8> GraphSharedDataBuffer;
+
+	const FNodeHandle NodeA = AppendNodeToGraph(*NodeTemplateA, NodeUID, GraphSharedDataBuffer);
+	InitNodeDecorator(*NodeTemplateA, NodeA, GraphSharedDataBuffer, 0, FDecoratorWithNoChildren::FSharedData());
+
+	const FNodeHandle NodeB = AppendNodeToGraph(*NodeTemplateB, NodeUID, GraphSharedDataBuffer);
+	InitNodeDecorator(*NodeTemplateB, NodeB, GraphSharedDataBuffer, 0, FDecoratorWithOneChild::FSharedData(FDecoratorHandle(NodeA)));
+
+	const FNodeHandle NodeC = AppendNodeToGraph(*NodeTemplateC, NodeUID, GraphSharedDataBuffer);
+	InitNodeDecorator(*NodeTemplateC, NodeC, GraphSharedDataBuffer, 0, FDecoratorWithChildren::FSharedData(FDecoratorHandle(NodeA), FDecoratorHandle(NodeB)));
+
+	FExecutionContext Context(GraphSharedDataBuffer);
+
+	{
+		TArray<FDecoratorUID> EvaluatedDecorators;
+
+		Private::EvaluatedDecorators = &EvaluatedDecorators;
+
+		FWeakDecoratorPtr NullPtr;				// Empty, no parent
+		FDecoratorHandle RootHandle(NodeC);		// Point to NodeC, first base decorator
+
+		FDecoratorPtr NodeCPtr = Context.AllocateNodeInstance(NullPtr, RootHandle);
+		AddErrorIfFalse(NodeCPtr.IsValid(), "FAnimationAnimNextRuntimeTest_IEvaluate -> Failed to allocate root node instance");
+
+		// Call pre/post evaluate on our graph
+		const EEvaluationFlags EvaluationFlags = EEvaluationFlags::All;
+		FPoseContainer PoseContainer;
+		EvaluateGraph(Context, NodeCPtr, EvaluationFlags, PoseContainer);
+
+		AddErrorIfFalse(EvaluatedDecorators.Num() == 6, "FAnimationAnimNextRuntimeTest_IEvaluate -> Expected 6 nodes to have been visited during the evaluate traversal");
+		AddErrorIfFalse(EvaluatedDecorators[0] == FDecoratorWithChildren::DecoratorUID, "FAnimationAnimNextRuntimeTest_IEvaluate -> Unexpected evaluate order");		// NodeC
+		AddErrorIfFalse(EvaluatedDecorators[1] == FDecoratorWithNoChildren::DecoratorUID, "FAnimationAnimNextRuntimeTest_IEvaluate -> Unexpected evaluate order");			// NodeA
+		AddErrorIfFalse(EvaluatedDecorators[2] == FDecoratorWithNoChildren::DecoratorUID, "FAnimationAnimNextRuntimeTest_IEvaluate -> Unexpected evaluate order");			// NodeA
+		AddErrorIfFalse(EvaluatedDecorators[3] == FDecoratorWithNoChildren::DecoratorUID, "FAnimationAnimNextRuntimeTest_IEvaluate -> Unexpected evaluate order");			// NodeB -> NodeA
+		AddErrorIfFalse(EvaluatedDecorators[4] == FDecoratorWithNoChildren::DecoratorUID, "FAnimationAnimNextRuntimeTest_IEvaluate -> Unexpected evaluate order");			// NodeB -> NodeA
+		AddErrorIfFalse(EvaluatedDecorators[5] == FDecoratorWithChildren::DecoratorUID, "FAnimationAnimNextRuntimeTest_IEvaluate -> Unexpected evaluate order");		// NodeC
+
+		Private::EvaluatedDecorators = nullptr;
+	}
+
+	Registry.Unregister(NodeTemplateA);
+	Registry.Unregister(NodeTemplateB);
+	Registry.Unregister(NodeTemplateC);
+
+	AddErrorIfFalse(Registry.GetNum() == 0, "FAnimationAnimNextRuntimeTest_IEvaluate -> Registry should contain 0 templates");
+
+	return true;
+}
+
+#endif
diff --git a/Engine/Plugins/Experimental/AnimNext/Source/AnimNextTestSuite/Private/AnimNextRuntimeTest.cpp b/Engine/Plugins/Experimental/AnimNext/Source/AnimNextTestSuite/Private/AnimNextRuntimeTest.cpp
new file mode 100644
--- /dev/null
+++ b/Engine/Plugins/Experimental/AnimNext/Source/AnimNextTestSuite/Private/AnimNextRuntimeTest.cpp
@@ -0,0 +1,147 @@
+// Copyright Epic Games, Inc. All Rights Reserved.
+
+
+#include "AnimNextRuntimeTest.h"
+
+#include "Misc/AutomationTest.h"
+
+#if WITH_DEV_AUTOMATION_TESTS
+
+#include "DecoratorBase/Decorator.h"
+#include "DecoratorBase/DecoratorRegistry.h"
+#include "DecoratorBase/NodeInstance.h"
+#include "DecoratorBase/NodeTemplateRegistry.h"
+
+namespace UE::AnimNext
+{
+	uint32 GetNodeTemplateUID(const TArray<FDecoratorUID>& NodeTemplateDecoratorList)
+	{
+		uint32 NodeTemplateUID = 0;
+		for (const FDecoratorUID DecoratorUID : NodeTemplateDecoratorList)
+		{
+			NodeTemplateUID = HashCombineFast(NodeTemplateUID, DecoratorUID.GetUID());
+		}
+
+		return NodeTemplateUID;
+	}
+
+	uint32 GetNodeInstanceSize(const TArray<FDecoratorUID>& NodeTemplateDecoratorList)
+	{
+		const FDecoratorRegistry& Registry = FDecoratorRegistry::Get();
+
+		uint32 NodeInstanceSize = sizeof(FNodeInstance);
+		for (const FDecoratorUID DecoratorUID : NodeTemplateDecoratorList)
+		{
+			const FDecorator* Decorator = Registry.Find(DecoratorUID);
+
+			const FDecoratorMemoryLayout MemoryLayout = Decorator->GetDecoratorMemoryDescription();
+			NodeInstanceSize = Align(NodeInstanceSize, MemoryLayout.InstanceDataAlignment);
+			NodeInstanceSize += MemoryLayout.InstanceDataSize;
+		}
+
+		return NodeInstanceSize;
+	}
+
+	void AppendTemplateDecorator(const TArray<FDecoratorUID>& NodeTemplateDecoratorList, int32 DecoratorIndex, TArray<uint8>& NodeTemplateBuffer, uint32& SharedDataOffset, uint32& InstanceDataOffset)
+	{
+		const FDecoratorRegistry& Registry = FDecoratorRegistry::Get();
+
+		const FDecoratorUID DecoratorUID = NodeTemplateDecoratorList[DecoratorIndex];
+		const FDecoratorRegistryHandle DecoratorHandle = Registry.FindHandle(DecoratorUID);
+		const FDecorator* Decorator = Registry.Find(DecoratorHandle);
+		const EDecoratorMode DecoratorMode = Decorator->GetMode();
+
+		uint32 AdditiveIndexOrNumAdditive;
+		if (DecoratorMode == EDecoratorMode::Base)
+		{
+			// Find out how many additive decorators we have
+			AdditiveIndexOrNumAdditive = 0;
+			for (int32 Index = DecoratorIndex + 1; Index < NodeTemplateDecoratorList.Num(); ++Index)	// Skip ourself
+			{
+				const FDecorator* ChildDecorator = Registry.Find(NodeTemplateDecoratorList[Index]);
+				if (ChildDecorator->GetMode() == EDecoratorMode::Base)
+				{
+					break;	// Found another base decorator, we are done
+				}
+
+				// We are additive
+				AdditiveIndexOrNumAdditive++;
+			}
+		}
+		else
+		{
+			// Find out our additive index
+			AdditiveIndexOrNumAdditive = 1;	// Skip ourself
+			for (int32 Index = DecoratorIndex - 1; Index >= 0; --Index)
+			{
+				const FDecorator* ParentDecorator = Registry.Find(NodeTemplateDecoratorList[Index]);
+				if (ParentDecorator->GetMode() == EDecoratorMode::Base)
+				{
+					break;	// Found our base decorator, we are done
+				}
+
+				// We are additive
+				AdditiveIndexOrNumAdditive++;
+			}
+		}
+
+		// Append and update our offsets
+		const int32 BufferIndex = NodeTemplateBuffer.AddUninitialized(sizeof(FDecoratorTemplate));
+		new(&NodeTemplateBuffer[BufferIndex]) FDecoratorTemplate(DecoratorUID, DecoratorHandle, DecoratorMode, AdditiveIndexOrNumAdditive, SharedDataOffset, InstanceDataOffset);
+
+		const FDecoratorMemoryLayout MemoryLayout = Decorator->GetDecoratorMemoryDescription();
+		SharedDataOffset = Align(SharedDataOffset, MemoryLayout.SharedDataAlignment);
+		SharedDataOffset += MemoryLayout.SharedDataSize;
+		InstanceDataOffset = Align(InstanceDataOffset, MemoryLayout.InstanceDataAlignment);
+		InstanceDataOffset += MemoryLayout.InstanceDataSize;
+	}
+
+	const FNodeTemplate* BuildNodeTemplate(const TArray<FDecoratorUID>& NodeTemplateDecoratorList, TArray<uint8>& NodeTemplateBuffer)
+	{
+		const int32 BufferIndex = NodeTemplateBuffer.AddUninitialized(sizeof(FNodeTemplate));
+		new(&NodeTemplateBuffer[BufferIndex]) FNodeTemplate(GetNodeTemplateUID(NodeTemplateDecoratorList), GetNodeInstanceSize(NodeTemplateDecoratorList), NodeTemplateDecoratorList.Num());
+
+		uint32 SharedDataOffset = sizeof(FNodeDescription);
+		uint32 InstanceDataOffset = sizeof(FNodeInstance);
+
+		for (int32 DecoratorIndex = 0; DecoratorIndex < NodeTemplateDecoratorList.Num(); ++DecoratorIndex)
+		{
+			AppendTemplateDecorator(NodeTemplateDecoratorList, DecoratorIndex, NodeTemplateBuffer, SharedDataOffset, InstanceDataOffset);
+		}
+
+		return reinterpret_cast<const FNodeTemplate*>(&NodeTemplateBuffer[0]);
+	}
+
+	uint32 GetNodeSharedSize(const FNodeTemplate& NodeTemplate)
+	{
+		const FDecoratorRegistry& Registry = FDecoratorRegistry::Get();
+
+		const uint32 NumDecorators = NodeTemplate.GetNumDecorators();
+		const FDecoratorTemplate* DecoratorTemplates = NodeTemplate.GetDecorators();
+
+		uint32 NodeSharedSize = sizeof(FNodeDescription);
+		for (uint32_t DecoratorIndex = 0; DecoratorIndex < NumDecorators; ++DecoratorIndex)
+		{
+			const FDecorator* Decorator = Registry.Find(DecoratorTemplates[DecoratorIndex].GetRegistryHandle());
+
+			const FDecoratorMemoryLayout MemoryLayout = Decorator->GetDecoratorMemoryDescription();
+			NodeSharedSize = Align(NodeSharedSize, MemoryLayout.SharedDataAlignment);
+			NodeSharedSize += MemoryLayout.SharedDataSize;
+		}
+
+		return NodeSharedSize;
+	}
+
+	// Adds a node template to a shared data graph and returns a handle to it
+	FNodeHandle AppendNodeToGraph(const FNodeTemplate& NodeTemplate, uint16& NodeUID, TArray<uint8>& GraphSharedDataBuffer)
+	{
+		FNodeTemplateRegistry& Registry = FNodeTemplateRegistry::Get();
+
+		const uint32 NodeSharedDataSize = GetNodeSharedSize(NodeTemplate);
+		const uint32 BufferIndex = GraphSharedDataBuffer.AddZeroed(NodeSharedDataSize);
+		new(&GraphSharedDataBuffer[BufferIndex]) FNodeDescription(NodeUID++, Registry.FindOrAdd(&NodeTemplate));
+
+		return FNodeHandle(BufferIndex);
+	}
+}
+#endif
diff --git a/Engine/Plugins/Experimental/AnimNext/Source/AnimNextTestSuite/Private/AnimNextRuntimeTest.h b/Engine/Plugins/Experimental/AnimNext/Source/AnimNextTestSuite/Private/AnimNextRuntimeTest.h
new file mode 100644
--- /dev/null
+++ b/Engine/Plugins/Experimental/AnimNext/Source/AnimNextTestSuite/Private/AnimNextRuntimeTest.h
@@ -0,0 +1,35 @@
+// Copyright Epic Games, Inc. All Rights Reserved.
+
+#pragma once
+
+#include "CoreMinimal.h"
+//#include "AnimNextRuntimeTest.generated.h"
+
+#if WITH_DEV_AUTOMATION_TESTS
+
+#include "DecoratorBase/DecoratorUID.h"
+#include "DecoratorBase/DecoratorTemplate.h"
+#include "DecoratorBase/NodeDescription.h"
+#include "DecoratorBase/NodeHandle.h"
+#include "DecoratorBase/NodeTemplate.h"
+
+// TODO: Add USTRUCT and other similar types here, if required for tests
+
+namespace UE::AnimNext
+{
+	const FNodeTemplate* BuildNodeTemplate(const TArray<FDecoratorUID>& NodeTemplateDecoratorList, TArray<uint8>& NodeTemplateBuffer);
+
+	// Adds a node template to a shared data graph and returns a handle to it
+	FNodeHandle AppendNodeToGraph(const FNodeTemplate& NodeTemplate, uint16& NodeUID, TArray<uint8>& GraphSharedDataBuffer);
+
+	template<class SharedDataType>
+	void InitNodeDecorator(const FNodeTemplate& NodeTemplate, const FNodeHandle NodeHandle, TArray<uint8>& GraphSharedDataBuffer, uint32 DecoratorIndex, const SharedDataType& SharedData)
+	{
+		FNodeDescription& NodeDesc = *reinterpret_cast<FNodeDescription*>(&GraphSharedDataBuffer[NodeHandle.GetSharedOffset()]);
+
+		const FDecoratorTemplate& DecoratorTemplate = NodeTemplate.GetDecorators()[DecoratorIndex];
+		*static_cast<SharedDataType*>(DecoratorTemplate.GetDecoratorDescription(NodeDesc)) = SharedData;
+	}
+}
+
+#endif
diff --git a/Engine/Plugins/Experimental/AnimNext/Source/AnimNextTestSuite/Private/Module.cpp b/Engine/Plugins/Experimental/AnimNext/Source/AnimNextTestSuite/Private/Module.cpp
new file mode 100644
--- /dev/null
+++ b/Engine/Plugins/Experimental/AnimNext/Source/AnimNextTestSuite/Private/Module.cpp
@@ -0,0 +1,17 @@
+// Copyright Epic Games, Inc. All Rights Reserved.
+
+#include "CoreMinimal.h"
+#include "Modules/ModuleInterface.h"
+#include "Modules/ModuleManager.h"
+
+namespace UE::AnimNext
+{
+
+class FModule : public IModuleInterface
+{
+public:
+};
+
+}
+
+IMPLEMENT_MODULE(UE::AnimNext::FModule, AnimNextTestSuite)
