Fix non-unity build by adding the correct forward declaration

diff --git a/Engine/Plugins/Experimental/AnimNext/Source/AnimNext/Public/DecoratorBase/IDecoratorInterface.h b/Engine/Plugins/Experimental/AnimNext/Source/AnimNext/Public/DecoratorBase/IDecoratorInterface.h
--- a/Engine/Plugins/Experimental/AnimNext/Source/AnimNext/Public/DecoratorBase/IDecoratorInterface.h
+++ b/Engine/Plugins/Experimental/AnimNext/Source/AnimNext/Public/DecoratorBase/IDecoratorInterface.h
@@ -14,7 +14,7 @@
 
 namespace UE::AnimNext
 {
-	struct ExecutionContext;	// Derived types will have functions that accept the execution context
+	struct FExecutionContext;	// Derived types will have functions that accept the execution context
 
 	/**
 	 * IDecoratorInterface
