Fix static analysis issues due to potential nullptr access in unit tests

diff --git a/Engine/Plugins/Experimental/AnimNext/Source/AnimNextTestSuite/Private/AnimNextDecoratorBaseTest.cpp b/Engine/Plugins/Experimental/AnimNext/Source/AnimNextTestSuite/Private/AnimNextDecoratorBaseTest.cpp
--- a/Engine/Plugins/Experimental/AnimNext/Source/AnimNextTestSuite/Private/AnimNextDecoratorBaseTest.cpp
+++ b/Engine/Plugins/Experimental/AnimNext/Source/AnimNextTestSuite/Private/AnimNextDecoratorBaseTest.cpp
@@ -287,77 +287,89 @@ bool FAnimationAnimNextRuntimeTest_DecoratorRegistry::RunTest(const FString& InP
 
 		const FDecorator* DecoratorA = Registry.Find(HandleA);
 		AddErrorIfFalse(DecoratorA != nullptr, "FAnimationAnimNextRuntimeTest_DecoratorRegistry -> Decorator should exist");
-		AddErrorIfFalse(DecoratorA->GetDecoratorUID() == FDecoratorA_Base::DecoratorUID, "FAnimationAnimNextRuntimeTest_DecoratorRegistry -> Unexpected decorator instance type");
-
+		if (DecoratorA != nullptr)
 		{
-			// Auto register another decorator
-			AUTO_REGISTER_ANIM_DECORATOR(FDecoratorAB_Add)
-
-			AddErrorIfFalse(Registry.GetNum() == 2, "FAnimationAnimNextRuntimeTest_DecoratorRegistry -> Registry should contain 2 decorators");
-
-			FDecoratorRegistryHandle HandleAB = Registry.FindHandle(FDecoratorAB_Add::DecoratorUID);
-			AddErrorIfFalse(HandleAB.IsValid(), "FAnimationAnimNextRuntimeTest_DecoratorRegistry -> Decorator should have registered automatically");
-			AddErrorIfFalse(HandleAB.IsStatic(), "FAnimationAnimNextRuntimeTest_DecoratorRegistry -> Decorator should have been statically allocated");
-			AddErrorIfFalse(HandleA != HandleAB, "FAnimationAnimNextRuntimeTest_DecoratorRegistry -> Decorator handles should be different");
-
-			const FDecorator* DecoratorAB = Registry.Find(HandleAB);
-			AddErrorIfFalse(DecoratorAB != nullptr, "FAnimationAnimNextRuntimeTest_DecoratorRegistry -> Decorator should exist");
-			AddErrorIfFalse(DecoratorAB->GetDecoratorUID() == FDecoratorAB_Add::DecoratorUID, "FAnimationAnimNextRuntimeTest_DecoratorRegistry -> Unexpected decorator instance type");
+			AddErrorIfFalse(DecoratorA->GetDecoratorUID() == FDecoratorA_Base::DecoratorUID, "FAnimationAnimNextRuntimeTest_DecoratorRegistry -> Unexpected decorator instance type");
 
-			FDecoratorRegistryHandle HandleAC_0;
 			{
-				// Dynamically register a decorator
-				FDecoratorAC_Add DecoratorAC_0;
-				Registry.Register(&DecoratorAC_0);
+				// Auto register another decorator
+				AUTO_REGISTER_ANIM_DECORATOR(FDecoratorAB_Add)
 
-				AddErrorIfFalse(Registry.GetNum() == 3, "FAnimationAnimNextRuntimeTest_DecoratorRegistry -> Registry should contain 3 decorators");
+					AddErrorIfFalse(Registry.GetNum() == 2, "FAnimationAnimNextRuntimeTest_DecoratorRegistry -> Registry should contain 2 decorators");
 
-				HandleAC_0 = Registry.FindHandle(FDecoratorAC_Add::DecoratorUID);
-				AddErrorIfFalse(HandleAC_0.IsValid(), "FAnimationAnimNextRuntimeTest_DecoratorRegistry -> Decorator should have registered automatically");
-				AddErrorIfFalse(HandleAC_0.IsDynamic(), "FAnimationAnimNextRuntimeTest_DecoratorRegistry -> Decorator should have been dynamically allocated");
-				AddErrorIfFalse(HandleA != HandleAC_0, "FAnimationAnimNextRuntimeTest_DecoratorRegistry -> Decorator handles should be different");
+				FDecoratorRegistryHandle HandleAB = Registry.FindHandle(FDecoratorAB_Add::DecoratorUID);
+				AddErrorIfFalse(HandleAB.IsValid(), "FAnimationAnimNextRuntimeTest_DecoratorRegistry -> Decorator should have registered automatically");
+				AddErrorIfFalse(HandleAB.IsStatic(), "FAnimationAnimNextRuntimeTest_DecoratorRegistry -> Decorator should have been statically allocated");
+				AddErrorIfFalse(HandleA != HandleAB, "FAnimationAnimNextRuntimeTest_DecoratorRegistry -> Decorator handles should be different");
 
-				const FDecorator* DecoratorAC_0Ptr = Registry.Find(HandleAC_0);
-				AddErrorIfFalse(DecoratorAC_0Ptr != nullptr, "FAnimationAnimNextRuntimeTest_DecoratorRegistry -> Decorator should exist");
-				AddErrorIfFalse(DecoratorAC_0Ptr->GetDecoratorUID() == FDecoratorAC_Add::DecoratorUID, "FAnimationAnimNextRuntimeTest_DecoratorRegistry -> Unexpected decorator instance type");
-				AddErrorIfFalse(&DecoratorAC_0 == DecoratorAC_0Ptr, "FAnimationAnimNextRuntimeTest_DecoratorRegistry -> Unexpected decorator instance pointer");
-
-				// Unregister our instances
-				Registry.Unregister(&DecoratorAC_0);
-
-				AddErrorIfFalse(Registry.GetNum() == 2, "FAnimationAnimNextRuntimeTest_DecoratorRegistry -> Registry should contain 2 decorators");
-				AddErrorIfFalse(!Registry.FindHandle(FDecoratorAC_Add::DecoratorUID).IsValid(), "FAnimationAnimNextRuntimeTest_DecoratorRegistry -> Decorator should have unregistered");
+				const FDecorator* DecoratorAB = Registry.Find(HandleAB);
+				AddErrorIfFalse(DecoratorAB != nullptr, "FAnimationAnimNextRuntimeTest_DecoratorRegistry -> Decorator should exist");
+				if (DecoratorAB != nullptr)
+				{
+					AddErrorIfFalse(DecoratorAB->GetDecoratorUID() == FDecoratorAB_Add::DecoratorUID, "FAnimationAnimNextRuntimeTest_DecoratorRegistry -> Unexpected decorator instance type");
+
+					FDecoratorRegistryHandle HandleAC_0;
+					{
+						// Dynamically register a decorator
+						FDecoratorAC_Add DecoratorAC_0;
+						Registry.Register(&DecoratorAC_0);
+
+						AddErrorIfFalse(Registry.GetNum() == 3, "FAnimationAnimNextRuntimeTest_DecoratorRegistry -> Registry should contain 3 decorators");
+
+						HandleAC_0 = Registry.FindHandle(FDecoratorAC_Add::DecoratorUID);
+						AddErrorIfFalse(HandleAC_0.IsValid(), "FAnimationAnimNextRuntimeTest_DecoratorRegistry -> Decorator should have registered automatically");
+						AddErrorIfFalse(HandleAC_0.IsDynamic(), "FAnimationAnimNextRuntimeTest_DecoratorRegistry -> Decorator should have been dynamically allocated");
+						AddErrorIfFalse(HandleA != HandleAC_0, "FAnimationAnimNextRuntimeTest_DecoratorRegistry -> Decorator handles should be different");
+
+						const FDecorator* DecoratorAC_0Ptr = Registry.Find(HandleAC_0);
+						AddErrorIfFalse(DecoratorAC_0Ptr != nullptr, "FAnimationAnimNextRuntimeTest_DecoratorRegistry -> Decorator should exist");
+						if (DecoratorAC_0Ptr != nullptr)
+						{
+							AddErrorIfFalse(DecoratorAC_0Ptr->GetDecoratorUID() == FDecoratorAC_Add::DecoratorUID, "FAnimationAnimNextRuntimeTest_DecoratorRegistry -> Unexpected decorator instance type");
+							AddErrorIfFalse(&DecoratorAC_0 == DecoratorAC_0Ptr, "FAnimationAnimNextRuntimeTest_DecoratorRegistry -> Unexpected decorator instance pointer");
+
+							// Unregister our instances
+							Registry.Unregister(&DecoratorAC_0);
+
+							AddErrorIfFalse(Registry.GetNum() == 2, "FAnimationAnimNextRuntimeTest_DecoratorRegistry -> Registry should contain 2 decorators");
+							AddErrorIfFalse(!Registry.FindHandle(FDecoratorAC_Add::DecoratorUID).IsValid(), "FAnimationAnimNextRuntimeTest_DecoratorRegistry -> Decorator should have unregistered");
+						}
+					}
+
+					{
+						// Dynamically register another decorator, re-using the previous dynamic index
+						FDecoratorAC_Add DecoratorAC_1;
+						Registry.Register(&DecoratorAC_1);
+
+						AddErrorIfFalse(Registry.GetNum() == 3, "FAnimationAnimNextRuntimeTest_DecoratorRegistry -> Registry should contain 3 decorators");
+
+						FDecoratorRegistryHandle HandleAC_1 = Registry.FindHandle(FDecoratorAC_Add::DecoratorUID);
+						AddErrorIfFalse(HandleAC_1.IsValid(), "FAnimationAnimNextRuntimeTest_DecoratorRegistry -> Decorator should have registered automatically");
+						AddErrorIfFalse(HandleAC_1.IsDynamic(), "FAnimationAnimNextRuntimeTest_DecoratorRegistry -> Decorator should have been dynamically allocated");
+						AddErrorIfFalse(HandleA != HandleAC_1, "FAnimationAnimNextRuntimeTest_DecoratorRegistry -> Decorator handles should be different");
+						AddErrorIfFalse(HandleAC_0 == HandleAC_1, "FAnimationAnimNextRuntimeTest_DecoratorRegistry -> Decorator handles should be identical");
+
+						const FDecorator* DecoratorAC_1Ptr = Registry.Find(HandleAC_1);
+						AddErrorIfFalse(DecoratorAC_1Ptr != nullptr, "FAnimationAnimNextRuntimeTest_DecoratorRegistry -> Decorator should exist");
+						if (DecoratorAC_1Ptr != nullptr)
+						{
+							AddErrorIfFalse(DecoratorAC_1Ptr->GetDecoratorUID() == FDecoratorAC_Add::DecoratorUID, "FAnimationAnimNextRuntimeTest_DecoratorRegistry -> Unexpected decorator instance type");
+							AddErrorIfFalse(&DecoratorAC_1 == DecoratorAC_1Ptr, "FAnimationAnimNextRuntimeTest_DecoratorRegistry -> Unexpected decorator instance pointer");
+
+							// Unregister our instances
+							Registry.Unregister(&DecoratorAC_1);
+
+							AddErrorIfFalse(Registry.GetNum() == 2, "FAnimationAnimNextRuntimeTest_DecoratorRegistry -> Registry should contain 2 decorators");
+							AddErrorIfFalse(!Registry.FindHandle(FDecoratorAC_Add::DecoratorUID).IsValid(), "FAnimationAnimNextRuntimeTest_DecoratorRegistry -> Decorator should have unregistered");
+						}
+					}
+				}
 			}
 
-			{
-				// Dynamically register another decorator, re-using the previous dynamic index
-				FDecoratorAC_Add DecoratorAC_1;
-				Registry.Register(&DecoratorAC_1);
-
-				AddErrorIfFalse(Registry.GetNum() == 3, "FAnimationAnimNextRuntimeTest_DecoratorRegistry -> Registry should contain 3 decorators");
-
-				FDecoratorRegistryHandle HandleAC_1 = Registry.FindHandle(FDecoratorAC_Add::DecoratorUID);
-				AddErrorIfFalse(HandleAC_1.IsValid(), "FAnimationAnimNextRuntimeTest_DecoratorRegistry -> Decorator should have registered automatically");
-				AddErrorIfFalse(HandleAC_1.IsDynamic(), "FAnimationAnimNextRuntimeTest_DecoratorRegistry -> Decorator should have been dynamically allocated");
-				AddErrorIfFalse(HandleA != HandleAC_1, "FAnimationAnimNextRuntimeTest_DecoratorRegistry -> Decorator handles should be different");
-				AddErrorIfFalse(HandleAC_0 == HandleAC_1, "FAnimationAnimNextRuntimeTest_DecoratorRegistry -> Decorator handles should be identical");
-
-				const FDecorator* DecoratorAC_1Ptr = Registry.Find(HandleAC_1);
-				AddErrorIfFalse(DecoratorAC_1Ptr != nullptr, "FAnimationAnimNextRuntimeTest_DecoratorRegistry -> Decorator should exist");
-				AddErrorIfFalse(DecoratorAC_1Ptr->GetDecoratorUID() == FDecoratorAC_Add::DecoratorUID, "FAnimationAnimNextRuntimeTest_DecoratorRegistry -> Unexpected decorator instance type");
-				AddErrorIfFalse(&DecoratorAC_1 == DecoratorAC_1Ptr, "FAnimationAnimNextRuntimeTest_DecoratorRegistry -> Unexpected decorator instance pointer");
-
-				// Unregister our instances
-				Registry.Unregister(&DecoratorAC_1);
-
-				AddErrorIfFalse(Registry.GetNum() == 2, "FAnimationAnimNextRuntimeTest_DecoratorRegistry -> Registry should contain 2 decorators");
-				AddErrorIfFalse(!Registry.FindHandle(FDecoratorAC_Add::DecoratorUID).IsValid(), "FAnimationAnimNextRuntimeTest_DecoratorRegistry -> Decorator should have unregistered");
-			}
+			AddErrorIfFalse(Registry.GetNum() == 1, "FAnimationAnimNextRuntimeTest_DecoratorRegistry -> Registry should contain 1 decorator");
+			AddErrorIfFalse(!Registry.FindHandle(FDecoratorAB_Add::DecoratorUID).IsValid(), "FAnimationAnimNextRuntimeTest_DecoratorRegistry -> Decorator should have unregistered automatically");
+			AddErrorIfFalse(HandleA == Registry.FindHandle(FDecoratorA_Base::DecoratorUID), "FAnimationAnimNextRuntimeTest_DecoratorRegistry -> Decorator handle should not have changed");
 		}
-
-		AddErrorIfFalse(Registry.GetNum() == 1, "FAnimationAnimNextRuntimeTest_DecoratorRegistry -> Registry should contain 1 decorator");
-		AddErrorIfFalse(!Registry.FindHandle(FDecoratorAB_Add::DecoratorUID).IsValid(), "FAnimationAnimNextRuntimeTest_DecoratorRegistry -> Decorator should have unregistered automatically");
-		AddErrorIfFalse(HandleA == Registry.FindHandle(FDecoratorA_Base::DecoratorUID), "FAnimationAnimNextRuntimeTest_DecoratorRegistry -> Decorator handle should not have changed");
 	}
 
 	AddErrorIfFalse(Registry.GetNum() == 0, "FAnimationAnimNextRuntimeTest_DecoratorRegistry -> All decorators should have unregistered");
@@ -400,38 +412,48 @@ bool FAnimationAnimNextRuntimeTest_NodeTemplateRegistry::RunTest(const FString&
 	const uint32 TemplateSize0 = NodeTemplate0->GetNodeTemplateSize();
 	const FNodeTemplate* NodeTemplate0_ = Registry.Find(TemplateHandle0);
 	AddErrorIfFalse(NodeTemplate0_ != nullptr, "FAnimationAnimNextRuntimeTest_NodeTemplateRegistry -> Registry should contain our template");
-	AddErrorIfFalse(NodeTemplate0_ != NodeTemplate0, "FAnimationAnimNextRuntimeTest_NodeTemplateRegistry -> Template pointers should be different");
-	AddErrorIfFalse(FMemory::Memcmp(NodeTemplate0, NodeTemplate0_, TemplateSize0) == 0, "FAnimationAnimNextRuntimeTest_NodeTemplateRegistry -> Templates should be identical");
-
-	// Try and register a duplicate template
-	TArray<uint8> NodeTemplateBuffer1;
-	const FNodeTemplate* NodeTemplate1 = BuildNodeTemplate(NodeTemplateDecoratorList, NodeTemplateBuffer1);
-	AddErrorIfFalse(NodeTemplate0 != NodeTemplate1, "FAnimationAnimNextRuntimeTest_NodeTemplateRegistry -> Node template pointers should be different");
-	AddErrorIfFalse(NodeTemplate0->GetUID() == NodeTemplate1->GetUID(), "FAnimationAnimNextRuntimeTest_NodeTemplateRegistry -> Node template UIDs should be identical");
-
-	FNodeTemplateRegistryHandle TemplateHandle1 = Registry.FindOrAdd(NodeTemplate1);
-	AddErrorIfFalse(Registry.GetNum() == 1, "FAnimationAnimNextRuntimeTest_NodeTemplateRegistry -> Registry should contain 1 template");
-	AddErrorIfFalse(TemplateHandle0 == TemplateHandle1, "FAnimationAnimNextRuntimeTest_NodeTemplateRegistry -> Node template handles should be identical");
+	if (NodeTemplate0_ != nullptr)
+	{
+		AddErrorIfFalse(NodeTemplate0_ != NodeTemplate0, "FAnimationAnimNextRuntimeTest_NodeTemplateRegistry -> Template pointers should be different");
+		AddErrorIfFalse(FMemory::Memcmp(NodeTemplate0, NodeTemplate0_, TemplateSize0) == 0, "FAnimationAnimNextRuntimeTest_NodeTemplateRegistry -> Templates should be identical");
 
-	// Try and register a new template
-	TArray<FDecoratorUID> NodeTemplateDecoratorList2;
-	NodeTemplateDecoratorList2.Add(FDecoratorA_Base::DecoratorUID);
-	NodeTemplateDecoratorList2.Add(FDecoratorAB_Add::DecoratorUID);
-	NodeTemplateDecoratorList2.Add(FDecoratorAC_Add::DecoratorUID);
-	NodeTemplateDecoratorList2.Add(FDecoratorAC_Add::DecoratorUID);
+		// Try and register a duplicate template
+		TArray<uint8> NodeTemplateBuffer1;
+		const FNodeTemplate* NodeTemplate1 = BuildNodeTemplate(NodeTemplateDecoratorList, NodeTemplateBuffer1);
+		if (NodeTemplate1 != nullptr)
+		{
+			AddErrorIfFalse(NodeTemplate0 != NodeTemplate1, "FAnimationAnimNextRuntimeTest_NodeTemplateRegistry -> Node template pointers should be different");
+			AddErrorIfFalse(NodeTemplate0->GetUID() == NodeTemplate1->GetUID(), "FAnimationAnimNextRuntimeTest_NodeTemplateRegistry -> Node template UIDs should be identical");
+
+			FNodeTemplateRegistryHandle TemplateHandle1 = Registry.FindOrAdd(NodeTemplate1);
+			AddErrorIfFalse(Registry.GetNum() == 1, "FAnimationAnimNextRuntimeTest_NodeTemplateRegistry -> Registry should contain 1 template");
+			AddErrorIfFalse(TemplateHandle0 == TemplateHandle1, "FAnimationAnimNextRuntimeTest_NodeTemplateRegistry -> Node template handles should be identical");
+
+			// Try and register a new template
+			TArray<FDecoratorUID> NodeTemplateDecoratorList2;
+			NodeTemplateDecoratorList2.Add(FDecoratorA_Base::DecoratorUID);
+			NodeTemplateDecoratorList2.Add(FDecoratorAB_Add::DecoratorUID);
+			NodeTemplateDecoratorList2.Add(FDecoratorAC_Add::DecoratorUID);
+			NodeTemplateDecoratorList2.Add(FDecoratorAC_Add::DecoratorUID);
+
+			TArray<uint8> NodeTemplateBuffer2;
+			const FNodeTemplate* NodeTemplate2 = BuildNodeTemplate(NodeTemplateDecoratorList2, NodeTemplateBuffer2);
+			if (NodeTemplate2 != nullptr)
+			{
+				AddErrorIfFalse(NodeTemplate0->GetUID() != NodeTemplate2->GetUID(), "FAnimationAnimNextRuntimeTest_NodeTemplateRegistry -> Node template UIDs should be different");
 
-	TArray<uint8> NodeTemplateBuffer2;
-	const FNodeTemplate* NodeTemplate2 = BuildNodeTemplate(NodeTemplateDecoratorList2, NodeTemplateBuffer2);
-	AddErrorIfFalse(NodeTemplate0->GetUID() != NodeTemplate2->GetUID(), "FAnimationAnimNextRuntimeTest_NodeTemplateRegistry -> Node template UIDs should be different");
+				FNodeTemplateRegistryHandle TemplateHandle2 = Registry.FindOrAdd(NodeTemplate2);
+				AddErrorIfFalse(Registry.GetNum() == 2, "FAnimationAnimNextRuntimeTest_NodeTemplateRegistry -> Registry should contain 2 templates");
+				AddErrorIfFalse(TemplateHandle0 != TemplateHandle2, "FAnimationAnimNextRuntimeTest_NodeTemplateRegistry -> Node template handles should be identical");
+				AddErrorIfFalse(TemplateHandle2.IsValid(), "FAnimationAnimNextRuntimeTest_NodeTemplateRegistry -> Registry should contain our template");
 
-	FNodeTemplateRegistryHandle TemplateHandle2 = Registry.FindOrAdd(NodeTemplate2);
-	AddErrorIfFalse(Registry.GetNum() == 2, "FAnimationAnimNextRuntimeTest_NodeTemplateRegistry -> Registry should contain 2 templates");
-	AddErrorIfFalse(TemplateHandle0 != TemplateHandle2, "FAnimationAnimNextRuntimeTest_NodeTemplateRegistry -> Node template handles should be identical");
-	AddErrorIfFalse(TemplateHandle2.IsValid(), "FAnimationAnimNextRuntimeTest_NodeTemplateRegistry -> Registry should contain our template");
+				// Unregister our templates
+				Registry.Unregister(NodeTemplate2);
+			}
+		}
 
-	// Unregister our templates
-	Registry.Unregister(NodeTemplate0);
-	Registry.Unregister(NodeTemplate2);
+		Registry.Unregister(NodeTemplate0);
+	}
 
 	AddErrorIfFalse(Registry.GetNum() == 0, "FAnimationAnimNextRuntimeTest_NodeTemplateRegistry -> Registry should contain 0 templates");
 
