Fix AnimNext module init to always happen on load and ensure we remove our engine exit lambda when the module is unloaded

diff --git a/Engine/Plugins/Experimental/AnimNext/Source/AnimNext/Private/Module.cpp b/Engine/Plugins/Experimental/AnimNext/Source/AnimNext/Private/Module.cpp
--- a/Engine/Plugins/Experimental/AnimNext/Source/AnimNext/Private/Module.cpp
+++ b/Engine/Plugins/Experimental/AnimNext/Source/AnimNext/Private/Module.cpp
@@ -16,23 +16,26 @@ class FModule : public IModuleInterface
 public:
 	virtual void StartupModule() override
 	{
-		FCoreDelegates::OnFEngineLoopInitComplete.AddLambda([]()
-		{
-			FDataRegistry::Init();
-			FDecoratorRegistry::Init();
-			FNodeTemplateRegistry::Init();
-		});
+		FDataRegistry::Init();
+		FDecoratorRegistry::Init();
+		FNodeTemplateRegistry::Init();
 
-		FCoreDelegates::OnEnginePreExit.AddLambda([]()
+		EnginePreExitHandle = FCoreDelegates::OnEnginePreExit.AddLambda([]()
 		{
 			FDataRegistry::Destroy();
 			FDecoratorRegistry::Destroy();
 			FNodeTemplateRegistry::Destroy();
 		});
-		
 	}
+
+	virtual void ShutdownModule() override
+	{
+		FCoreDelegates::OnEnginePreExit.Remove(EnginePreExitHandle);
+	}
+
+	FDelegateHandle EnginePreExitHandle;
 };
 
 }
 
-IMPLEMENT_MODULE(UE::AnimNext::FModule, AnimNext)
\ No newline at end of file
+IMPLEMENT_MODULE(UE::AnimNext::FModule, AnimNext)
