AnimNext: Initial editor-side parameter blocks & libraries

Adds two new assets, both enabled with OFPA-style external package handling for subobjects

Parameter Library:
Defines a set of parameter names, their types and some metadata (just comment text at the moment).
Parameter libraries are entirely editor-only data and are not cooked.

Parameter Block:
Defines a set of bindings to parameters. At present, this is pretty limited and not really functional, but (currently) in the end it is intended that the following will be possible:
- Setting a literal value
- Calculating a value based around some inputs/other parameters using RigVM graphs
- Performing calculations and setting multiple parameters based on the results (event-graph style)

At runtime, a parameter block bakes down to a structure of named parameters (using UPropertyBag) and some RigVM bytecode to perform the logic to fill those parameters in.

Moved a ton of files, reworked namespaces and includes to consolidate into 4 main modules - AnimNext (runtime), AnimNextEditor (editor widgets etc), AnimNextUncooked (stuff that is needed for cooking) and AnimNextTestSuite (all tests).

diff --git a/Engine/Plugins/Experimental/AnimNext/AnimNext.uplugin b/Engine/Plugins/Experimental/AnimNext/AnimNext.uplugin
--- a/Engine/Plugins/Experimental/AnimNext/AnimNext.uplugin
+++ b/Engine/Plugins/Experimental/AnimNext/AnimNext.uplugin
@@ -24,6 +24,14 @@
 		{
 			"Name": "StructUtils",
 			"Enabled": true
+		},
+		{
+			"Name": "ContentBrowserFileDataSource",
+			"Enabled": true
+		},
+		{
+			"Name": "PythonScriptPlugin",
+			"Enabled": true
 		}
 	],
 	"Modules" :
@@ -34,23 +42,19 @@
 			"LoadingPhase" : "PreDefault"
 		},
 		{
-			"Name" : "AnimNextGraph",
-			"Type" : "Runtime",
-			"LoadingPhase" : "Default"
-		},
-		{
-			"Name" : "AnimNextGraphEditor",
+			"Name" : "AnimNextEditor",
 			"Type" : "Editor",
 			"LoadingPhase" : "Default"
 		},
 		{
-			"Name": "AnimNextGraphUncookedOnly",
-			"Type": "UncookedOnly",
-			"LoadingPhase": "Default"
+			"Name" : "AnimNextUncookedOnly",
+			"Type" : "UncookedOnly",
+			"LoadingPhase" : "Default"
 		},
 		{
-			"Name": "AnimNextTestSuite",
-			"Type": "UncookedOnly"
+			"Name" : "AnimNextTestSuite",
+			"Type" : "UncookedOnly",
+			"LoadingPhase" : "Default"
 		}
 	]
 }
diff --git a/Engine/Plugins/Experimental/AnimNext/Source/AnimNext/AnimNext.Build.cs b/Engine/Plugins/Experimental/AnimNext/Source/AnimNext/AnimNext.Build.cs
--- a/Engine/Plugins/Experimental/AnimNext/Source/AnimNext/AnimNext.Build.cs
+++ b/Engine/Plugins/Experimental/AnimNext/Source/AnimNext/AnimNext.Build.cs
@@ -11,14 +11,16 @@ namespace UnrealBuildTool.Rules
 				{
 					"Core",
 					"CoreUObject",
-					"StructUtils"
+					"StructUtils",
+					"RigVM",
+					"ControlRig",
+					"Engine"
 				}
 			);
 			PrivateDependencyModuleNames.AddRange(
 				new string[]
 				{
-					"AssetRegistry",
-					"Engine"
+					"AssetRegistry"
 				}
 			);
 		}
diff --git a/Engine/Plugins/Experimental/AnimNext/Source/AnimNext/Internal/Graph/AnimNextGraph.h b/Engine/Plugins/Experimental/AnimNext/Source/AnimNext/Internal/Graph/AnimNextGraph.h
new file mode 100644
--- /dev/null
+++ b/Engine/Plugins/Experimental/AnimNext/Source/AnimNext/Internal/Graph/AnimNextGraph.h
@@ -0,0 +1,69 @@
+// Copyright Epic Games, Inc. All Rights Reserved.
+
+#pragma once
+
+#include "CoreMinimal.h"
+#include "Interface/IAnimNextInterface.h"
+#include "RigVMCore/RigVM.h"
+#include "Param/ParamType.h"
+#include "AnimNextGraph.generated.h"
+
+class UEdGraph;
+class UAnimNextGraph;
+namespace UE::AnimNext::UncookedOnly
+{
+	struct FUtils;
+}
+
+namespace UE::AnimNext::Editor
+{
+	class FGraphEditor;
+}
+
+namespace UE::AnimNext::Graph
+{
+	extern ANIMNEXT_API const FName EntryPointName;
+	extern ANIMNEXT_API const FName ResultName;
+}
+
+// A user-created graph of logic used to supply data
+UCLASS(BlueprintType)
+class ANIMNEXT_API UAnimNextGraph : public UObject, public IAnimNextInterface
+{
+	GENERATED_BODY()
+
+	// UObject interface
+	virtual void PostRename(UObject* OldOuter, const FName OldName) override;
+	virtual void GetPreloadDependencies(TArray<UObject*>& OutDeps) override;
+
+	// IAnimNextInterface interface
+	virtual UE::AnimNext::FParamTypeHandle GetReturnTypeHandleImpl() const final override;
+	virtual bool GetDataImpl(const UE::AnimNext::FContext& Context) const final override;
+
+	/** Set the return type of this graph */
+	void SetReturnTypeHandle(UE::AnimNext::FParamTypeHandle InHandle);
+
+	// Support rig VM execution
+	TArray<FRigVMExternalVariable> GetRigVMExternalVariables();
+	
+	friend class UAnimNextGraphFactory;
+	friend class UAnimNextGraph_EditorData;
+	friend struct UE::AnimNext::UncookedOnly::FUtils;
+	friend class UE::AnimNext::Editor::FGraphEditor;
+	friend class UAnimNextGraph;
+	friend class UAnimGraphNode_AnimNextGraph;
+	
+	UPROPERTY()
+	TObjectPtr<URigVM> RigVM;
+
+	UPROPERTY()
+	FRigVMRuntimeSettings VMRuntimeSettings;
+
+	UPROPERTY()
+	FAnimNextParamType ReturnType;
+
+#if WITH_EDITORONLY_DATA
+	UPROPERTY(VisibleAnywhere, Instanced, Category = "Graph", meta = (ShowInnerProperties))
+	TObjectPtr<UObject> EditorData;
+#endif
+};
diff --git a/Engine/Plugins/Experimental/AnimNext/Source/AnimNext/Internal/Graph/AnimNext_LODPose.h b/Engine/Plugins/Experimental/AnimNext/Source/AnimNext/Internal/Graph/AnimNext_LODPose.h
new file mode 100644
--- /dev/null
+++ b/Engine/Plugins/Experimental/AnimNext/Source/AnimNext/Internal/Graph/AnimNext_LODPose.h
@@ -0,0 +1,78 @@
+// Copyright Epic Games, Inc. All Rights Reserved.
+
+#pragma once
+
+#include "CoreMinimal.h"
+#include "Interface/InterfaceTypes.h"
+#include "Interface/IAnimNextInterface.h"
+#include "Interface/InterfaceContext.h"
+#include "ReferencePose.h"
+#include "LODPose.h"
+#include "AnimNext_LODPose.generated.h"
+
+USTRUCT(BlueprintType, meta = (DisplayName = "ReferencePose"))
+struct FAnimNextGraphReferencePose
+{
+	GENERATED_BODY()
+
+	FAnimNextGraphReferencePose() = default;
+
+	explicit FAnimNextGraphReferencePose(const UE::AnimNext::FReferencePose* InReferencePose)
+		: ReferencePose(InReferencePose)
+	{
+	}
+
+	const UE::AnimNext::FReferencePose* ReferencePose = nullptr;
+};
+
+USTRUCT(BlueprintType, meta = (DisplayName = "LODPose"))
+struct FAnimNextGraphLODPose
+{
+	GENERATED_BODY()
+
+	FAnimNextGraphLODPose() = default;
+
+	explicit FAnimNextGraphLODPose(const UE::AnimNext::FLODPose& InLODPose)
+		: LODPose(InLODPose)
+	{
+	}
+	explicit FAnimNextGraphLODPose(UE::AnimNext::FLODPose&& InLODPose)
+		: LODPose(MoveTemp(InLODPose))
+	{
+	}
+
+	UE::AnimNext::FLODPose LODPose;
+};
+
+UCLASS()
+class UAnimNextInterfaceGraphLODPose : public UObject, public IAnimNextInterface
+{
+	GENERATED_BODY()
+
+	virtual UE::AnimNext::FParamTypeHandle GetReturnTypeHandleImpl() const final override
+	{
+		return UE::AnimNext::FParamTypeHandle::GetHandle<FAnimNextGraphLODPose>();
+	}
+
+	virtual bool GetDataImpl(const UE::AnimNext::FContext& Context) const override
+	{
+		checkf(false, TEXT("UAnimNextInterfaceLODPose::GetDataImpl must be overridden"));
+		return false;
+	}
+};
+
+UCLASS()
+class UAnimNextInterface_LODPose_Literal : public UAnimNextInterfaceGraphLODPose
+{
+	GENERATED_BODY()
+
+	// IAnimNextInterface interface
+	virtual bool GetDataImpl(const UE::AnimNext::FContext& Context) const final override
+	{
+		Context.SetResult(Value);
+		return true;
+	}
+	
+	UPROPERTY(EditAnywhere, Category = "Parameters")
+	FAnimNextGraphLODPose Value;
+};
diff --git a/Engine/Plugins/Experimental/AnimNext/Source/AnimNext/Internal/Graph/GraphExecuteContext.h b/Engine/Plugins/Experimental/AnimNext/Source/AnimNext/Internal/Graph/GraphExecuteContext.h
new file mode 100644
--- /dev/null
+++ b/Engine/Plugins/Experimental/AnimNext/Source/AnimNext/Internal/Graph/GraphExecuteContext.h
@@ -0,0 +1,76 @@
+﻿// Copyright Epic Games, Inc. All Rights Reserved.
+
+#pragma once
+
+#include "CoreMinimal.h"
+#include "RigVMCore/RigVMExecuteContext.h"
+#include "Units/RigUnit.h"
+#include "GraphExecuteContext.generated.h"
+
+class IAnimNextInterface;
+
+namespace UE::AnimNext
+{
+	struct FContext;
+}
+
+USTRUCT(BlueprintType)
+struct FAnimNextGraphExecuteContext : public FRigVMExecuteContext
+{
+	GENERATED_BODY()
+
+	FAnimNextGraphExecuteContext()
+		: FRigVMExecuteContext()
+		, InterfaceContext(nullptr)
+		, Interface(nullptr)
+		, ResultPtr(nullptr)
+	{
+	}
+
+	const UE::AnimNext::FContext& GetContext() const
+	{
+		check(InterfaceContext);
+		return *InterfaceContext;
+	}
+
+	void SetContextData(const IAnimNextInterface* InInterface, const UE::AnimNext::FContext& InInterfaceContext, bool& bInResult)
+	{
+		Interface = InInterface;
+		InterfaceContext = &InInterfaceContext;
+		ResultPtr = &bInResult;
+	}
+
+	void SetResult(bool bInResult) const
+	{
+		check(ResultPtr);
+		*ResultPtr &= bInResult;
+	}
+
+	const IAnimNextInterface* GetInterface() const
+	{
+		check(Interface);
+		return Interface;
+	}
+	
+	virtual void Copy(const FRigVMExecuteContext* InOtherContext) override
+	{
+		Super::Copy(InOtherContext);
+
+		const FAnimNextGraphExecuteContext* OtherContext = (const FAnimNextGraphExecuteContext*)InOtherContext; 
+		InterfaceContext = OtherContext->InterfaceContext;
+		Interface = OtherContext->Interface;
+		ResultPtr = OtherContext->ResultPtr;
+	}
+
+
+private:
+	const UE::AnimNext::FContext* InterfaceContext;
+	const IAnimNextInterface* Interface;
+	bool* ResultPtr;
+};
+
+USTRUCT(meta=(ExecuteContext="FAnimNextGraphExecuteContext"))
+struct FRigUnit_AnimNextBase : public FRigUnit
+{
+	GENERATED_BODY()
+};
diff --git a/Engine/Plugins/Experimental/AnimNext/Source/AnimNext/Internal/Graph/RigUnit_AnimNextBeginExecution.h b/Engine/Plugins/Experimental/AnimNext/Source/AnimNext/Internal/Graph/RigUnit_AnimNextBeginExecution.h
new file mode 100644
--- /dev/null
+++ b/Engine/Plugins/Experimental/AnimNext/Source/AnimNext/Internal/Graph/RigUnit_AnimNextBeginExecution.h
@@ -0,0 +1,28 @@
+﻿// Copyright Epic Games, Inc. All Rights Reserved.
+
+#pragma once
+
+#include "Units/RigUnit.h"
+#include "GraphExecuteContext.h"
+#include "RigUnit_AnimNextBeginExecution.generated.h"
+
+/**
+ * Event for driving the skeleton hierarchy with variables and rig elements
+ */
+USTRUCT(meta=(DisplayName="Execute", Category="Events", NodeColor="1, 0, 0", Keywords="Begin,Update,Tick,Forward,Event"))
+struct ANIMNEXT_API FRigUnit_AnimNextBeginExecution : public FRigUnit_AnimNextBase
+{
+	GENERATED_BODY()
+
+	RIGVM_METHOD()
+	void Execute();
+
+	virtual FName GetEventName() const override { return EventName; }
+	virtual bool CanOnlyExistOnce() const override { return true; }
+
+	// The execution result
+	UPROPERTY(EditAnywhere, DisplayName = "Execute", Category = "BeginExecution", meta = (Output))
+	FAnimNextGraphExecuteContext ExecuteContext;
+
+	static FName EventName;
+};
diff --git a/Engine/Plugins/Experimental/AnimNext/Source/AnimNext/Internal/Graph/RigUnit_AnimNextEndExecution.h b/Engine/Plugins/Experimental/AnimNext/Source/AnimNext/Internal/Graph/RigUnit_AnimNextEndExecution.h
new file mode 100644
--- /dev/null
+++ b/Engine/Plugins/Experimental/AnimNext/Source/AnimNext/Internal/Graph/RigUnit_AnimNextEndExecution.h
@@ -0,0 +1,88 @@
+﻿// Copyright Epic Games, Inc. All Rights Reserved.
+
+#pragma once
+
+#include "Units/RigUnit.h"
+#include "GraphExecuteContext.h"
+#include "AnimNext_LODPose.h"
+#include "RigUnit_AnimNextEndExecution.generated.h"
+
+struct FRigUnitContext;
+
+/** Event for writing back calculated results to external variables */
+USTRUCT(meta = (DisplayName = "End Execute", Category = "Events", TitleColor = "1 0 0", NodeColor = "1 1 1"))
+struct ANIMNEXT_API FRigUnit_AnimNextEndExecution : public FRigUnit_AnimNextBase
+{
+	GENERATED_BODY()
+
+	template<typename ValueType>
+	static void SetResult(const FRigVMExecuteContext& InExecuteContext, const ValueType& InResult)
+	{
+		const FAnimNextGraphExecuteContext& ExecuteContext = static_cast<const FAnimNextGraphExecuteContext&>(InExecuteContext);
+		ExecuteContext.GetContext().SetResult<ValueType>(InResult);
+	}
+
+public:
+	// The execution result
+	UPROPERTY(EditAnywhere, Transient, DisplayName = "End Execute", Category = "EndExecution", meta = (Input))
+	FAnimNextGraphExecuteContext ExecuteContext;
+};
+
+/** Event for writing back a calculated bool */
+USTRUCT(meta = (DisplayName = "End Execute Bool", Category = "Events", TitleColor = "1 0 0", NodeColor = "1 1 1"))
+struct ANIMNEXT_API FRigUnit_AnimNextEndExecution_Bool : public FRigUnit_AnimNextEndExecution
+{
+	GENERATED_BODY()
+	
+	FRigUnit_AnimNextEndExecution_Bool()
+		: FRigUnit_AnimNextEndExecution()
+		, Result(false)
+	{}
+	
+	RIGVM_METHOD()
+	virtual void Execute() override;
+
+	virtual bool CanOnlyExistOnce() const override { return true; }
+	
+	UPROPERTY(EditAnywhere, Category = Result, meta = (Input))
+	bool Result;
+};
+
+/** Event for writing back a calculated float */
+USTRUCT(meta = (DisplayName = "End Execute Float", Category = "Events", TitleColor = "1 0 0", NodeColor = "1 1 1"))
+struct ANIMNEXT_API FRigUnit_AnimNextEndExecution_Float : public FRigUnit_AnimNextEndExecution
+{
+	GENERATED_BODY()
+
+	FRigUnit_AnimNextEndExecution_Float()
+		: FRigUnit_AnimNextEndExecution()
+		, Result(0.f)
+	{}
+
+	RIGVM_METHOD()
+	virtual void Execute() override;
+
+	virtual bool CanOnlyExistOnce() const override { return true; }
+
+	UPROPERTY(EditAnywhere, Category = Result, meta = (Input))
+	float Result;
+};
+
+/** Event for writing back a calculated pose */
+USTRUCT(meta = (DisplayName = "End Execute LODPose", Category = "Events", TitleColor = "1 0 0", NodeColor = "1 1 1"))
+struct ANIMNEXT_API FRigUnit_AnimNextEndExecution_LODPose : public FRigUnit_AnimNextEndExecution
+{
+	GENERATED_BODY()
+
+	FRigUnit_AnimNextEndExecution_LODPose()
+		: FRigUnit_AnimNextEndExecution()
+	{}
+
+	RIGVM_METHOD()
+	virtual void Execute() override;
+
+	virtual bool CanOnlyExistOnce() const override { return true; }
+
+	UPROPERTY(EditAnywhere, Category = Result, meta = (Input))
+	FAnimNextGraphLODPose Result;
+};
diff --git a/Engine/Plugins/Experimental/AnimNext/Source/AnimNext/Internal/Param/AnimNextParameterBlock.h b/Engine/Plugins/Experimental/AnimNext/Source/AnimNext/Internal/Param/AnimNextParameterBlock.h
new file mode 100644
--- /dev/null
+++ b/Engine/Plugins/Experimental/AnimNext/Source/AnimNext/Internal/Param/AnimNextParameterBlock.h
@@ -0,0 +1,88 @@
+﻿// Copyright Epic Games, Inc. All Rights Reserved.
+
+#pragma once
+
+#include "CoreMinimal.h"
+#include "PropertyBag.h"
+#include "RigVMCore/RigVMExecuteContext.h"
+#include "AnimNextParameterBlock.generated.h"
+
+class UEdGraph;
+
+namespace UE::AnimNext::UncookedOnly
+{
+	struct FUtils;
+	struct FUtilsPrivate;
+}
+
+namespace UE::AnimNext::Editor
+{
+	class FParametersEditor;
+	struct FUtils;
+}
+
+// Library entry used to export to asset registry
+USTRUCT()
+struct FAnimNextParameterBlockAssetRegistryExportEntry
+{
+	GENERATED_BODY()
+
+	FAnimNextParameterBlockAssetRegistryExportEntry() = default;
+
+	FAnimNextParameterBlockAssetRegistryExportEntry(FName InName, const FSoftObjectPath& InLibrary)
+		: Name(InName)
+		, Library(InLibrary)
+	{}
+	
+	UPROPERTY()
+	FName Name;
+
+	UPROPERTY()
+	FSoftObjectPath Library;
+};
+
+// Library used to export to asset registry
+USTRUCT()
+struct FAnimNextParameterBlockAssetRegistryExports
+{
+	GENERATED_BODY()
+
+	UPROPERTY()
+	TArray<FAnimNextParameterBlockAssetRegistryExportEntry> Bindings;
+};
+
+/** An asset used to define AnimNext parameters and their bindings */
+UCLASS(MinimalAPI, BlueprintType)
+class UAnimNextParameterBlock : public UObject
+{
+	GENERATED_BODY()
+
+	friend class UAnimNextParameterBlockFactory;
+	friend class UAnimNextParameterBlock_EditorData;
+	friend struct UE::AnimNext::UncookedOnly::FUtils;
+	friend struct UE::AnimNext::UncookedOnly::FUtilsPrivate;
+	friend class UE::AnimNext::Editor::FParametersEditor;
+	friend struct UE::AnimNext::Editor::FUtils;
+
+	// UObject interface
+	virtual void PostRename(UObject* OldOuter, const FName OldName) override;
+	virtual void GetPreloadDependencies(TArray<UObject*>& OutDeps) override;
+	virtual void GetAssetRegistryTags(TArray<FAssetRegistryTag>& OutTags) const override;
+
+	// Support rig VM execution
+	ANIMNEXT_API TArray<FRigVMExternalVariable> GetRigVMExternalVariables();
+
+	UPROPERTY()
+	TObjectPtr<URigVM> RigVM;
+
+	UPROPERTY()
+	FRigVMRuntimeSettings VMRuntimeSettings;
+
+	UPROPERTY()
+	FInstancedPropertyBag PropertyBag;
+
+#if WITH_EDITORONLY_DATA
+	UPROPERTY(VisibleAnywhere, Instanced, Category = "Parameters", meta = (ShowInnerProperties))
+	TObjectPtr<UObject> EditorData;
+#endif
+};
\ No newline at end of file
diff --git a/Engine/Plugins/Experimental/AnimNext/Source/AnimNext/Internal/Param/ParametersExecuteContext.h b/Engine/Plugins/Experimental/AnimNext/Source/AnimNext/Internal/Param/ParametersExecuteContext.h
new file mode 100644
--- /dev/null
+++ b/Engine/Plugins/Experimental/AnimNext/Source/AnimNext/Internal/Param/ParametersExecuteContext.h
@@ -0,0 +1,65 @@
+﻿// Copyright Epic Games, Inc. All Rights Reserved.
+
+#pragma once
+
+#include "CoreMinimal.h"
+#include "RigVMCore/RigVMExecuteContext.h"
+#include "Units/RigUnit.h"
+#include "ParametersExecuteContext.generated.h"
+
+class IAnimNextInterface;
+
+namespace UE::AnimNext
+{
+	struct FContext;
+}
+
+USTRUCT(BlueprintType)
+struct FAnimNextParametersExecuteContext : public FRigVMExecuteContext
+{
+	GENERATED_BODY()
+
+	FAnimNextParametersExecuteContext()
+		: FRigVMExecuteContext()
+	{
+	}
+
+	void SetContextData(TConstArrayView<TArrayView<uint8>> InValues)
+	{
+		Values = InValues;
+	}
+
+	void SetCurrentValueIndex(int32 InIndex)
+	{
+		check(Values.IsValidIndex(InIndex));
+		Index = InIndex;
+	}
+
+	TArrayView<uint8> GetData() const
+	{
+		check(Values.Num() > 0);
+		check(Values.IsValidIndex(Index));
+		return Values[Index]; 
+	}
+	
+	virtual void Copy(const FRigVMExecuteContext* InOtherContext) override
+	{
+		Super::Copy(InOtherContext);
+
+		const FAnimNextParametersExecuteContext* OtherContext = (const FAnimNextParametersExecuteContext*)InOtherContext; 
+		Values = OtherContext->Values;
+	}
+
+private:
+	// The parameter values to set
+	TConstArrayView<TArrayView<uint8>> Values;
+
+	// Current value that is being set
+	int32 Index = 0;
+};
+
+USTRUCT(meta=(ExecuteContext="FAnimNextParametersExecuteContext"))
+struct FRigUnit_AnimNextParametersBase : public FRigUnit
+{
+	GENERATED_BODY()
+};
diff --git a/Engine/Plugins/Experimental/AnimNext/Source/AnimNext/Internal/Param/RigUnit_AnimNextParametersBeginExecution.h b/Engine/Plugins/Experimental/AnimNext/Source/AnimNext/Internal/Param/RigUnit_AnimNextParametersBeginExecution.h
new file mode 100644
--- /dev/null
+++ b/Engine/Plugins/Experimental/AnimNext/Source/AnimNext/Internal/Param/RigUnit_AnimNextParametersBeginExecution.h
@@ -0,0 +1,25 @@
+﻿// Copyright Epic Games, Inc. All Rights Reserved.
+
+#pragma once
+
+#include "Units/RigUnit.h"
+#include "ParametersExecuteContext.h"
+#include "RigUnit_AnimNextParametersBeginExecution.generated.h"
+
+USTRUCT(meta=(DisplayName="Execute", Category="Events", NodeColor="1, 0, 0", Keywords="Begin,Update,Tick,Forward,Event"))
+struct ANIMNEXT_API FRigUnit_AnimNextParametersBeginExecution : public FRigUnit_AnimNextParametersBase
+{
+	GENERATED_BODY()
+
+	RIGVM_METHOD()
+	void Execute();
+
+	virtual FName GetEventName() const override { return EventName; }
+	virtual bool CanOnlyExistOnce() const override { return true; }
+
+	// The execution result
+	UPROPERTY(EditAnywhere, DisplayName = "Execute", Category = "BeginExecution", meta = (Output))
+	FAnimNextParametersExecuteContext ExecuteContext;
+
+	static FName EventName;
+};
diff --git a/Engine/Plugins/Experimental/AnimNext/Source/AnimNext/Internal/Param/RigVMDispatch_SetParameter.h b/Engine/Plugins/Experimental/AnimNext/Source/AnimNext/Internal/Param/RigVMDispatch_SetParameter.h
new file mode 100644
--- /dev/null
+++ b/Engine/Plugins/Experimental/AnimNext/Source/AnimNext/Internal/Param/RigVMDispatch_SetParameter.h
@@ -0,0 +1,40 @@
+﻿// Copyright Epic Games, Inc. All Rights Reserved.
+
+#pragma once
+
+#include "RigVMCore/RigVMDispatchFactory.h"
+#include "RigVMDispatch_SetParameter.generated.h"
+
+namespace UE::AnimNext::UncookedOnly
+{
+struct FUtils;
+}
+
+/*
+ * Sets a parameter's value
+ */
+USTRUCT(meta=(DisplayName = "Set Parameter", NodeColor = "0.8, 0, 0.2, 1"))
+struct ANIMNEXT_API FRigVMDispatch_SetParameter : public FRigVMDispatchFactory
+{
+	GENERATED_BODY()
+
+	FRigVMDispatch_SetParameter();
+
+private:
+	friend struct UE::AnimNext::UncookedOnly::FUtils;
+
+	virtual FName GetArgumentNameForOperandIndex(int32 InOperandIndex, int32 InTotalOperands) const override;
+	virtual TArray<FRigVMTemplateArgument> GetArguments() const override;
+	virtual TArray<FRigVMExecuteArgument> GetExecuteArguments_Impl(const FRigVMDispatchContext& InContext) const override;
+	virtual FRigVMTemplateTypeMap OnNewArgumentType(const FName& InArgumentName, TRigVMTypeIndex InTypeIndex) const override;
+	virtual bool IsSingleton() const override { return true; } 
+
+	virtual FRigVMFunctionPtr GetDispatchFunctionImpl(const FRigVMTemplateTypeMap& InTypes) const override
+	{
+		return &FRigVMDispatch_SetParameter::Execute;
+	}
+	static void Execute(FRigVMExtendedExecuteContext& InContext, FRigVMMemoryHandleArray Handles);
+
+	static const FName ValueName;
+	static const FName ExecuteContextName;
+};
diff --git a/Engine/Plugins/Experimental/AnimNext/Source/AnimNext/Private/Graph/AnimGraph/AnimNode_AnimNextGraph.cpp b/Engine/Plugins/Experimental/AnimNext/Source/AnimNext/Private/Graph/AnimGraph/AnimNode_AnimNextGraph.cpp
new file mode 100644
--- /dev/null
+++ b/Engine/Plugins/Experimental/AnimNext/Source/AnimNext/Private/Graph/AnimGraph/AnimNode_AnimNextGraph.cpp
@@ -0,0 +1,210 @@
+// Copyright Epic Games, Inc. All Rights Reserved.
+
+#include "Graph/AnimGraph/AnimNode_AnimNextGraph.h"
+#include "DataRegistry.h"
+#include "ReferencePose.h"
+#include "Components/SkeletalMeshComponent.h"
+#include "Animation/AnimInstanceProxy.h"
+#include "GameFramework/Actor.h"
+#include "AnimationRuntime.h"
+#include "Interface/AnimNextInterface.h"
+#include "DataRegistry.h"
+#include "GenerationTools.h"
+#include "ReferencePose.h"
+#include "Graph/AnimNext_LODPose.h"
+#include "Graph/RigUnit_AnimNextAnimSequence.h" // TEST
+#include "Animation/AnimSequence.h"
+#include "Animation/AnimSequence.h" // TEST
+#include "Engine/SkeletalMesh.h"
+#include "BoneContainer.h"
+
+#include UE_INLINE_GENERATED_CPP_BY_NAME(AnimNode_AnimNextGraph)
+
+#if WITH_EDITOR
+#include "Editor.h"
+#endif
+
+FAnimNode_AnimNextGraph::FAnimNode_AnimNextGraph()
+	: FAnimNode_CustomProperty()
+	, AnimNextGraph(nullptr)
+	, LODThreshold(INDEX_NONE)
+{
+}
+
+FAnimNode_AnimNextGraph::~FAnimNode_AnimNextGraph()
+{
+}
+
+void FAnimNode_AnimNextGraph::OnInitializeAnimInstance(const FAnimInstanceProxy* InProxy, const UAnimInstance* InAnimInstance)
+{
+	DECLARE_SCOPE_HIERARCHICAL_COUNTER_FUNC()
+
+	FAnimNode_CustomProperty::OnInitializeAnimInstance(InProxy, InAnimInstance);
+
+	InitializeProperties(InAnimInstance, GetTargetClass());
+}
+
+void FAnimNode_AnimNextGraph::GatherDebugData(FNodeDebugData& DebugData)
+{
+	DECLARE_SCOPE_HIERARCHICAL_COUNTER_FUNC()
+
+	SourceLink.GatherDebugData(DebugData.BranchFlow(1.f));
+}
+
+void FAnimNode_AnimNextGraph::Update_AnyThread(const FAnimationUpdateContext& Context)
+{
+	DECLARE_SCOPE_HIERARCHICAL_COUNTER_FUNC()
+
+	GraphDeltaTime += Context.GetDeltaTime();
+
+	SourceLink.Update(Context);
+
+	if (IsLODEnabled(Context.AnimInstanceProxy))
+	{
+		GetEvaluateGraphExposedInputs().Execute(Context);
+
+		PropagateInputProperties(Context.AnimInstanceProxy->GetAnimInstanceObject());
+	}
+
+	FAnimNode_CustomProperty::Update_AnyThread(Context);
+
+	//TRACE_ANIM_NODE_VALUE(Context, TEXT("Class"), *GetNameSafe(ControlRigClass.Get()));
+}
+
+void FAnimNode_AnimNextGraph::Initialize_AnyThread(const FAnimationInitializeContext& Context)
+{
+	DECLARE_SCOPE_HIERARCHICAL_COUNTER_FUNC()
+
+	SourceLink.Initialize(Context);
+
+	FAnimNode_CustomProperty::Initialize_AnyThread(Context);
+}
+
+void FAnimNode_AnimNextGraph::CacheBones_AnyThread(const FAnimationCacheBonesContext& Context)
+{
+	DECLARE_SCOPE_HIERARCHICAL_COUNTER_FUNC()
+
+	FAnimNode_CustomProperty::CacheBones_AnyThread(Context);
+
+	SourceLink.CacheBones(Context);
+}
+
+void FAnimNode_AnimNextGraph::Evaluate_AnyThread(FPoseContext & Output)
+{
+	DECLARE_SCOPE_HIERARCHICAL_COUNTER_FUNC()
+
+	using namespace UE::AnimNext;
+
+	FPoseContext SourcePose(Output);
+
+	if (SourceLink.GetLinkNode())
+	{
+		SourceLink.Evaluate(SourcePose);
+	}
+	else
+	{
+		SourcePose.ResetToRefPose();
+	}
+
+	USkeletalMeshComponent* SkeletalMeshComponent = Output.AnimInstanceProxy->GetSkelMeshComponent();
+	check(SkeletalMeshComponent != nullptr);
+
+	FDataHandle RefPoseHandle = FDataRegistry::Get()->GetOrGenerateReferencePose(SkeletalMeshComponent);
+	const UE::AnimNext::FReferencePose& RefPose = RefPoseHandle.GetRef<UE::AnimNext::FReferencePose>();
+
+	const int32 LODLevel = Output.AnimInstanceProxy->GetLODLevel();
+
+	// TODO : In a normal graph, FState/FParamStorage/FContext would come from FAnimationUpdateContext, here I have to create one the fly for the AnimNext graph
+	UE::AnimNext::FParamStorage ParamStorage(5, 0, 0);	// TODO : see if we can have state and params together and optional. Also, default size
+	
+	// TODO : Using AnimInstanceProxy->GetDeltaSeconds() makes the preview to advance  multiple times when debug options are activated (i.e. ShowUncompressedAnim)
+	// See where to get / how to calculate the correct delta value
+	UE::AnimNext::FContext RootContext(GraphDeltaTime, RootState, ParamStorage);
+	GraphDeltaTime = 0.f; // Reset for the case that we receive multiple calls to Evaluate (debug options)
+
+	FAnimNextGraphReferencePose GraphReferencePose(&RefPose);
+	FAnimNextGraph_AnimSequence GraphTestSequence(TestSequence);  // TEST
+
+	FAnimNextGraphLODPose GraphSourceLODPose(FLODPose(RefPose, LODLevel, false, Output.ExpectsAdditivePose()));
+	FGenerationTools::RemapPose(LODLevel, SourcePose, RefPose, GraphSourceLODPose.LODPose);
+
+	RootContext.AddInputReference(FName(TEXT("GraphReferencePose")), GraphReferencePose);			// TODO : Should this go into the AnimNextVM Context ?
+	RootContext.AddInputValue(FName(TEXT("GraphLODLevel")), LODLevel);								// TODO : Should this go into the AnimNextVM Context ?
+	RootContext.AddInputValue(FName(TEXT("GraphExpectsAdditive")), Output.ExpectsAdditivePose());	// TODO : Should this go into the AnimNextVM Context ?
+
+	RootContext.AddInputReference(FName(TEXT("SourcePose")), GraphSourceLODPose);		// TODO : Pass this as a external variable maybe ? When we have support for variables in the rigvm graph
+	RootContext.AddInputReference(FName(TEXT("TestSequence")), GraphTestSequence);		// TEST anim decompression - Anim Sequence to decompress
+		
+	TScriptInterface<IAnimNextInterface> ScriptInterface = AnimNextGraph;
+	FAnimNextGraphLODPose GraphResultLODPose(FLODPose(RefPose, LODLevel, true, Output.ExpectsAdditivePose()));
+
+	const bool bOk = UE::AnimNext::Interface::GetDataSafe(ScriptInterface, RootContext, GraphResultLODPose);
+	if (bOk)
+	{
+		FGenerationTools::RemapPose(LODLevel, RefPose, GraphResultLODPose.LODPose, Output);
+	}
+	else
+	{
+		Output.ResetToRefPose();
+	}
+
+
+	// evaluate 
+	FAnimNode_CustomProperty::Evaluate_AnyThread(Output);
+}
+
+void FAnimNode_AnimNextGraph::PostSerialize(const FArchive& Ar)
+{
+	DECLARE_SCOPE_HIERARCHICAL_COUNTER_FUNC()
+
+	// after compile, we have to reinitialize
+	// because it needs new execution code
+	// since memory has changed
+	if (Ar.IsObjectReferenceCollector())
+	{
+		if (AnimNextGraph)
+		{
+			//AnimNextInterfaceGraph->Initialize();
+		}
+	}
+}
+
+void FAnimNode_AnimNextGraph::InitializeProperties(const UObject* InSourceInstance, UClass* InTargetClass)
+{
+	// Build property lists
+	SourceProperties.Reset(SourcePropertyNames.Num());
+	DestProperties.Reset(SourcePropertyNames.Num());
+
+	check(SourcePropertyNames.Num() == DestPropertyNames.Num());
+
+	for (int32 Idx = 0; Idx < SourcePropertyNames.Num(); ++Idx)
+	{
+		FName& SourceName = SourcePropertyNames[Idx];
+		UClass* SourceClass = InSourceInstance->GetClass();
+
+		FProperty* SourceProperty = FindFProperty<FProperty>(SourceClass, SourceName);
+		SourceProperties.Add(SourceProperty);
+		DestProperties.Add(nullptr);
+	}
+}
+
+void FAnimNode_AnimNextGraph::PropagateInputProperties(const UObject* InSourceInstance)
+{
+	if (InSourceInstance)
+	{
+		// Assign value to the properties exposed as pins
+		for (int32 PropIdx = 0; PropIdx < SourceProperties.Num(); ++PropIdx)
+		{
+		}
+	}
+}
+
+
+#if WITH_EDITOR
+
+void FAnimNode_AnimNextGraph::HandleObjectsReinstanced_Impl(UObject* InSourceObject, UObject* InTargetObject, const TMap<UObject*, UObject*>& OldToNewInstanceMap)
+{
+	FAnimNode_CustomProperty::HandleObjectsReinstanced_Impl(InSourceObject, InTargetObject, OldToNewInstanceMap);
+}
+
+#endif
diff --git a/Engine/Plugins/Experimental/AnimNext/Source/AnimNext/Private/Graph/AnimNextGraph.cpp b/Engine/Plugins/Experimental/AnimNext/Source/AnimNext/Private/Graph/AnimNextGraph.cpp
new file mode 100644
--- /dev/null
+++ b/Engine/Plugins/Experimental/AnimNext/Source/AnimNext/Private/Graph/AnimNextGraph.cpp
@@ -0,0 +1,102 @@
+// Copyright Epic Games, Inc. All Rights Reserved.
+
+#include "Graph/AnimNextGraph.h"
+#include "RigVMCore/RigVMMemoryStorage.h"
+#include "RigVMCore/RigVMExecuteContext.h"
+#include "Graph/RigUnit_AnimNextBeginExecution.h"
+#include "UObject/ObjectSaveContext.h"
+#include "UObject/Package.h"
+#include "Param/ParamTypeHandle.h"
+
+#include UE_INLINE_GENERATED_CPP_BY_NAME(AnimNextGraph)
+
+namespace UE::AnimNext::Graph
+{
+const FName EntryPointName("GetData");
+const FName ResultName("Result");
+}
+
+UE::AnimNext::FParamTypeHandle UAnimNextGraph::GetReturnTypeHandleImpl() const
+{
+	return ReturnType.GetHandle();
+}
+
+bool UAnimNextGraph::GetDataImpl(const UE::AnimNext::FContext& Context) const
+{
+	bool bResult = true;
+	
+	if(RigVM)
+	{
+		FRigVMExtendedExecuteContext RigVMExtendedExecuteContext; // we might have to set some sensible values to the extended execute context
+		FAnimNextGraphExecuteContext& AnimNextInterfaceContext = RigVMExtendedExecuteContext.GetPublicDataSafe<FAnimNextGraphExecuteContext>();
+
+		AnimNextInterfaceContext.SetContextData(this, Context, bResult);
+
+		bResult &= (RigVM->Execute(RigVMExtendedExecuteContext, TArray<URigVMMemoryStorage*>(), FRigUnit_AnimNextBeginExecution::EventName) != ERigVMExecuteResult::Failed);
+	}
+	
+	return bResult;
+}
+
+void UAnimNextGraph::SetReturnTypeHandle(UE::AnimNext::FParamTypeHandle InHandle)
+{
+	ReturnType = InHandle.GetType();
+}
+
+TArray<FRigVMExternalVariable> UAnimNextGraph::GetRigVMExternalVariables()
+{
+	return TArray<FRigVMExternalVariable>(); 
+}
+
+static TArray<UClass*> GetClassObjectsInPackage(UPackage* InPackage)
+{
+	TArray<UObject*> Objects;
+	GetObjectsWithOuter(InPackage, Objects, false);
+
+	TArray<UClass*> ClassObjects;
+	for (UObject* Object : Objects)
+	{
+		if (UClass* Class = Cast<UClass>(Object))
+		{
+			ClassObjects.Add(Class);
+		}
+	}
+
+	return ClassObjects;
+}
+
+void UAnimNextGraph::PostRename(UObject* OldOuter, const FName OldName)
+{
+	Super::PostRename(OldOuter, OldName);
+
+	// Whenever the asset is renamed/moved, generated classes parented to the old package
+	// are not moved to the new package automatically (see FAssetRenameManager), so we
+	// have to manually perform the move/rename, to avoid invalid reference to the old package
+
+	// Note: while asset duplication doesn't duplicate the classes either, it is not a problem there
+	// because we always recompile in post duplicate.
+	TArray<UClass*> ClassObjects = GetClassObjectsInPackage(OldOuter->GetPackage());
+
+	for (UClass* ClassObject : ClassObjects)
+	{
+		if (URigVMMemoryStorageGeneratorClass* MemoryClass = Cast<URigVMMemoryStorageGeneratorClass>(ClassObject))
+		{
+			MemoryClass->Rename(nullptr, GetPackage(), REN_ForceNoResetLoaders | REN_DoNotDirty | REN_DontCreateRedirectors | REN_NonTransactional);
+		}
+	}
+}
+
+void UAnimNextGraph::GetPreloadDependencies(TArray<UObject*>& OutDeps)
+{
+	Super::GetPreloadDependencies(OutDeps);
+
+	TArray<UClass*> ClassObjects = GetClassObjectsInPackage(GetPackage());
+
+	for (UClass* ClassObject : ClassObjects)
+	{
+		if (URigVMMemoryStorageGeneratorClass* MemoryClass = Cast<URigVMMemoryStorageGeneratorClass>(ClassObject))
+		{
+			OutDeps.Add(MemoryClass);
+		}
+	}
+}
diff --git a/Engine/Plugins/Experimental/AnimNext/Source/AnimNext/Private/Graph/RigUnit_AnimNextAnimSequence.cpp b/Engine/Plugins/Experimental/AnimNext/Source/AnimNext/Private/Graph/RigUnit_AnimNextAnimSequence.cpp
new file mode 100644
--- /dev/null
+++ b/Engine/Plugins/Experimental/AnimNext/Source/AnimNext/Private/Graph/RigUnit_AnimNextAnimSequence.cpp
@@ -0,0 +1,92 @@
+// Copyright Epic Games, Inc. All Rights Reserved.
+
+#include "RigUnit_AnimNextAnimSequence.h"
+#include "Graph/GraphExecuteContext.h"
+#include "Interface/AnimNextInterface.h"
+#include "DecompressionTools.h"
+#include "Animation/AnimSequence.h"
+#include "Graph/AnimNext_LODPose.h"
+
+#include UE_INLINE_GENERATED_CPP_BY_NAME(RigUnit_AnimNextAnimSequence)
+
+// --- FRigUnit_AnimNext_AnimSequenceAsset --- 
+
+FRigUnit_AnimNext_AnimSequenceAsset_Execute()
+{
+	using namespace UE::AnimNext;
+
+	//Sequence.AnimSequence = AnimSequence; // TODO : Enable once we have RigVM object support enabled
+}
+
+// --- FRigUnit_AnimNext_AnimSequencePlayer --- 
+
+FRigUnit_AnimNext_AnimSequencePlayer_Initialize()
+{
+	using namespace UE::AnimNext;
+
+	const FContext& InterfaceContext = ExecuteContext.GetContext();
+
+	const UAnimSequenceBase* AnimSequenceBase = Sequence.AnimSequence;
+	if (AnimSequenceBase != nullptr)
+	{
+		// Get or Create the internal state
+		FAnimSequencePlayerState& AnimSequencePlayerState = InterfaceContext.GetState<FAnimSequencePlayerState>(ExecuteContext.GetInterface(), 0);
+
+		const float SequenceLength = AnimSequenceBase->GetPlayLength();
+		AnimSequencePlayerState.InternalTimeAccumulator = FMath::Clamp(Parameters.StartPosition, 0.f, SequenceLength);
+		AnimSequencePlayerState.PrevInternalTimeAccumulator = AnimSequencePlayerState.InternalTimeAccumulator;
+	}
+}
+
+FRigUnit_AnimNext_AnimSequencePlayer_Execute()
+{
+	using namespace UE::AnimNext;
+
+	const FContext& InterfaceContext = ExecuteContext.GetContext();
+	
+	const UAnimSequenceBase* AnimSequenceBase = Sequence.AnimSequence;
+	if (AnimSequenceBase != nullptr)
+	{
+		// Get or Create the internal state
+		//const TParam<FAnimSequencePlayerState> AnimSequencePlayerState = InterfaceContext.GetState<FAnimSequencePlayerState>(ExecuteContext.GetInterface(), 0);
+		FAnimSequencePlayerState& AnimSequencePlayerState = InterfaceContext.GetState<FAnimSequencePlayerState, EStatePersistence::Permanent>(ExecuteContext.GetInterface(), 0);
+
+		const float DeltaTime = InterfaceContext.GetDeltaTime();
+		const float EffectiveDelta = FMath::IsNearlyZero(DeltaTime) || FMath::IsNearlyZero(Parameters.PlayRate) ? 0.f : DeltaTime * Parameters.PlayRate;
+
+		const float SequenceLength = AnimSequenceBase->GetPlayLength();
+		float CurrentTime = Parameters.bLoop
+			? FMath::Fmod(AnimSequencePlayerState.InternalTimeAccumulator + EffectiveDelta, SequenceLength)
+			: FMath::Clamp(AnimSequencePlayerState.InternalTimeAccumulator + EffectiveDelta, 0.f, SequenceLength);
+
+		if (Parameters.bLoop && CurrentTime < 0.f)
+		{
+			CurrentTime += SequenceLength;
+		}
+
+		AnimSequencePlayerState.PrevInternalTimeAccumulator = AnimSequencePlayerState.InternalTimeAccumulator;
+		AnimSequencePlayerState.InternalTimeAccumulator = CurrentTime;
+
+		FDeltaTimeRecord DeltaTimeRecord;
+		DeltaTimeRecord.Set(AnimSequencePlayerState.PrevInternalTimeAccumulator, AnimSequencePlayerState.InternalTimeAccumulator);
+
+		const FAnimExtractContext ExtractionContext(static_cast<double>(AnimSequencePlayerState.InternalTimeAccumulator)
+			, false /*Output.AnimInstanceProxy->ShouldExtractRootMotion()*/
+			, DeltaTimeRecord
+			, Parameters.bLoop);
+
+		const FAnimNextGraphReferencePose& GraphReferencePose = InterfaceContext.GetParameterChecked<FAnimNextGraphReferencePose>(TEXT("GraphReferencePose")).GetDataChecked();
+		const int32 GraphLODLevel = InterfaceContext.GetParameterChecked<int32>(TEXT("GraphLODLevel")).GetDataChecked();
+		const bool bGraphExpectsAdditive = InterfaceContext.GetParameterChecked<bool>(TEXT("GraphExpectsAdditive")).GetDataChecked();
+
+		LODPose.LODPose.PrepareForLOD(*GraphReferencePose.ReferencePose, GraphLODLevel, true, bGraphExpectsAdditive);
+		
+		// Note : calling FDecompressionTools instead of UAnimSequence / UAnimSequenceBase, as I can not modify the engine code (so I extracted the function)
+		// TODO : Revisit this, so we can plug other sequence types
+		const UAnimSequence* AnimSequence = Cast<UAnimSequence>(AnimSequenceBase);
+		if (AnimSequence != nullptr)
+		{
+			FDecompressionTools::GetAnimationPose(AnimSequence, LODPose.LODPose, ExtractionContext);
+		}
+	}
+}
diff --git a/Engine/Plugins/Experimental/AnimNext/Source/AnimNext/Private/Graph/RigUnit_AnimNextEndExecution.cpp b/Engine/Plugins/Experimental/AnimNext/Source/AnimNext/Private/Graph/RigUnit_AnimNextEndExecution.cpp
new file mode 100644
--- /dev/null
+++ b/Engine/Plugins/Experimental/AnimNext/Source/AnimNext/Private/Graph/RigUnit_AnimNextEndExecution.cpp
@@ -0,0 +1,21 @@
+﻿// Copyright Epic Games, Inc. All Rights Reserved.
+
+#include "Graph/RigUnit_AnimNextEndExecution.h"
+#include "Units/RigUnitContext.h"
+
+#include UE_INLINE_GENERATED_CPP_BY_NAME(RigUnit_AnimNextEndExecution)
+
+FRigUnit_AnimNextEndExecution_Bool_Execute()
+{
+	SetResult(ExecuteContext, Result);
+}
+
+FRigUnit_AnimNextEndExecution_Float_Execute()
+{
+	SetResult(ExecuteContext, Result);
+}
+
+FRigUnit_AnimNextEndExecution_LODPose_Execute()
+{
+	SetResult(ExecuteContext, Result);
+}
\ No newline at end of file
diff --git a/Engine/Plugins/Experimental/AnimNext/Source/AnimNext/Private/Graph/RigUnit_AnimNextParameter.cpp b/Engine/Plugins/Experimental/AnimNext/Source/AnimNext/Private/Graph/RigUnit_AnimNextParameter.cpp
new file mode 100644
--- /dev/null
+++ b/Engine/Plugins/Experimental/AnimNext/Source/AnimNext/Private/Graph/RigUnit_AnimNextParameter.cpp
@@ -0,0 +1,56 @@
+// Copyright Epic Games, Inc. All Rights Reserved.
+
+#include "RigUnit_AnimNextParameter.h"
+#include "Graph/GraphExecuteContext.h"
+#include "Interface/AnimNextInterface.h"
+
+#include UE_INLINE_GENERATED_CPP_BY_NAME(RigUnit_AnimNextParameter)
+
+bool FRigUnit_AnimNextParameter::GetParameterInternal(FName InName, const FAnimNextGraphExecuteContext& InContext, void* OutResult)
+{
+	return true;
+}
+
+FRigUnit_AnimNextParameter_AnimSequence_Execute()
+{
+	const UE::AnimNext::FContext& InterfaceContext = ExecuteContext.GetContext();
+
+	Result = InterfaceContext.GetParameterChecked<FAnimNextGraph_AnimSequence>(Parameter);
+
+	//GetParameterInternal(Parameter, ExecuteContext, &Result);
+}
+
+
+FRigUnit_AnimNextParameter_AnimNextInterface_Execute()
+{
+	GetParameterInternal(Parameter, ExecuteContext, &Result);
+}
+
+FRigUnit_FloatOperator_Execute()
+{}
+
+FRigUnit_AnimNext_SequencePlayer_Execute()
+{
+}
+
+FRigUnit_TestFloatState_Execute()
+{
+	using namespace UE::AnimNext;
+
+	// @TODO: ensure that context arg is always present to avoid these checks here
+	const FContext& InterfaceContext = ExecuteContext.GetContext();
+
+	FRigUnit_TestFloatState_SpringDamperState& State = InterfaceContext.GetState<FRigUnit_TestFloatState_SpringDamperState>(ExecuteContext.GetInterface(), 0);
+	const float DeltaTime = InterfaceContext.GetDeltaTime();
+
+	FMath::SpringDamperSmoothing(
+		State.Value,
+		State.ValueRate,
+		TargetValue,
+		TargetValueRate,
+		DeltaTime,
+		SmoothingTime,
+		DampingRatio);
+
+	Result = State.Value;
+}
diff --git a/Engine/Plugins/Experimental/AnimNext/Source/AnimNext/Private/Graph/RigUnit_AnimNextParameter.h b/Engine/Plugins/Experimental/AnimNext/Source/AnimNext/Private/Graph/RigUnit_AnimNextParameter.h
new file mode 100644
--- /dev/null
+++ b/Engine/Plugins/Experimental/AnimNext/Source/AnimNext/Private/Graph/RigUnit_AnimNextParameter.h
@@ -0,0 +1,153 @@
+// Copyright Epic Games, Inc. All Rights Reserved.
+
+#pragma once
+
+#include "Graph/GraphExecuteContext.h"
+#include "Interface/IAnimNextInterface.h"
+#include "Units/RigUnit.h"
+#include "Graph/AnimNext_LODPose.h"
+#include "RigUnit_AnimNextAnimSequence.h"
+// --- ---
+#include "RigUnit_AnimNextParameter.generated.h"
+
+struct FAnimNextUnitContext;
+class UAnimNextGraph;
+
+/** Unit for reading parameters from context */
+USTRUCT(meta = (DisplayName = "Parameter", Category = "Parameters", TitleColor = "1 0 0", NodeColor = "1 1 1"))
+struct FRigUnit_AnimNextParameter : public FRigUnit_AnimNextBase
+{
+	GENERATED_BODY()
+
+protected:
+	static bool GetParameterInternal(FName InName, const FAnimNextGraphExecuteContext& InContext, void* OutResult);
+
+	UPROPERTY(EditAnywhere, Category = "Parameter", meta = (Input))
+	FName Parameter = NAME_None;
+};
+
+/** Unit for reading an AnimSequence parameter from context */
+USTRUCT(meta = (DisplayName = "AnimSequence Parameter", Category = "Parameters", TitleColor = "1 0 0", NodeColor = "1 1 1"))
+struct FRigUnit_AnimNextParameter_AnimSequence : public FRigUnit_AnimNextParameter
+{
+	GENERATED_BODY()
+
+	RIGVM_METHOD()
+	void Execute();
+
+protected:
+	UPROPERTY(EditAnywhere, Category = "Result", meta = (Output))
+	FAnimNextGraph_AnimSequence Result;
+};
+
+/** Unit for reading anim interface graph parameter from context */
+USTRUCT(meta = (DisplayName = "Anim Interface Parameter", Category = "Parameters", TitleColor = "1 0 0", NodeColor = "1 1 1"))
+struct FRigUnit_AnimNextParameter_AnimNextInterface : public FRigUnit_AnimNextParameter
+{
+	GENERATED_BODY()
+
+	RIGVM_METHOD()
+	void Execute();
+	
+protected:
+	UPROPERTY(EditAnywhere, Category = "Result", meta = (Output))
+	TScriptInterface<IAnimNextInterface> Result = nullptr;
+};
+
+/** Base unit for calling anim interfaces from graphs */
+USTRUCT(meta = (DisplayName = "Anim Interface", Category = "Execution", TitleColor = "1 0 0", NodeColor = "1 1 1"))
+struct FRigUnit_AnimNextInterface : public FRigUnit_AnimNextBase
+{
+	GENERATED_BODY()
+
+protected:
+	UPROPERTY(EditAnywhere, Category = "Anim Interface", meta = (Input))
+	TScriptInterface<IAnimNextInterface> AnimNextInterface = nullptr;
+};
+
+/** Unit for getting a pose via an anim interface */
+USTRUCT(meta = (DisplayName = "Float Operator", Category = "Operators", TitleColor = "1 0 0", NodeColor = "1 1 1"))
+struct FRigUnit_FloatOperator : public FRigUnit_AnimNextBase
+{
+	GENERATED_BODY()
+	
+	FRigUnit_FloatOperator()
+		: ParamA(0.f)
+		, ParamB(0.f)
+		, Result(0.f)
+	{}
+
+	RIGVM_METHOD()
+	void Execute();
+	
+protected:
+	UPROPERTY(EditAnywhere, Category = "Pose Interface", meta = (Input))
+	TScriptInterface<IAnimNextInterface> Operator = nullptr;
+
+	UPROPERTY(EditAnywhere, Category = "Pose Interface", meta = (Input))
+	float ParamA;
+
+	UPROPERTY(EditAnywhere, Category = "Pose Interface", meta = (Input))
+	float ParamB;
+
+	UPROPERTY(EditAnywhere, Transient, DisplayName = "Result", Category = "BeginExecution", meta = (Output))
+	float Result;
+};
+
+/** Unit for getting a pose via an anim interface */
+USTRUCT(meta = (DisplayName = "Anim Sequence", Category = "Animation", TitleColor = "1 0 0", NodeColor = "1 1 1"))
+struct FRigUnit_AnimNext_SequencePlayer : public FRigUnit_AnimNextBase
+{
+	GENERATED_BODY()
+
+		RIGVM_METHOD()
+		void Execute();
+
+protected:
+	UPROPERTY(EditAnywhere, Category = "Anim Interface", meta = (Input))
+	FAnimSequenceParameters Parameters;
+
+	UPROPERTY(EditAnywhere, Category = "Anim Interface", meta = (Input))
+	TScriptInterface<IAnimNextInterface> Sequence = nullptr;
+
+	UPROPERTY(EditAnywhere, Transient, DisplayName = "Result", Category = "BeginExecution", meta = (Output))
+	FAnimNextGraphExecuteContext Result;
+};
+
+USTRUCT()
+struct FRigUnit_TestFloatState_SpringDamperState
+{
+	GENERATED_BODY()
+
+	UPROPERTY()
+	float Value = 0.0f;
+
+	UPROPERTY()
+	float ValueRate = 0.0f;
+};
+
+/** Unit for getting a pose via an anim interface */
+USTRUCT(meta = (DisplayName = "Test Float State - Spring Damper Smoothing", Varying, Category = "Animation", TitleColor = "1 0 0", NodeColor = "1 1 1"))
+struct FRigUnit_TestFloatState : public FRigUnit_AnimNextBase
+{
+	GENERATED_BODY()
+
+	RIGVM_METHOD()
+	void Execute();
+	
+protected:
+	UPROPERTY(EditAnywhere, Category = "Parameters", meta=(Input))
+	float TargetValue = 0.0f;
+
+	UPROPERTY(EditAnywhere, Category = "Parameters", meta=(Input))
+	float TargetValueRate = 1.0f;
+	
+	UPROPERTY(EditAnywhere, Category = "Parameters", meta=(Input))
+	float SmoothingTime = 1.0f;
+	
+	UPROPERTY(EditAnywhere, Category = "Parameters", meta=(Input))
+	float DampingRatio = 1.0f;
+
+	UPROPERTY(EditAnywhere, Category = "Parameters", meta=(Output))
+	float Result = 0.0f;
+};
diff --git a/Engine/Plugins/Experimental/AnimNext/Source/AnimNext/Private/Param/AnimNextParameterBlock.cpp b/Engine/Plugins/Experimental/AnimNext/Source/AnimNext/Private/Param/AnimNextParameterBlock.cpp
new file mode 100644
--- /dev/null
+++ b/Engine/Plugins/Experimental/AnimNext/Source/AnimNext/Private/Param/AnimNextParameterBlock.cpp
@@ -0,0 +1,83 @@
+﻿// Copyright Epic Games, Inc. All Rights Reserved.
+
+#include "Param/AnimNextParameterBlock.h"
+#include "RigVMCore/RigVMMemoryStorage.h"
+#include "UObject/ObjectSaveContext.h"
+#include "UObject/Package.h"
+
+#include UE_INLINE_GENERATED_CPP_BY_NAME(AnimNextParameterBlock)
+
+TArray<FRigVMExternalVariable> UAnimNextParameterBlock::GetRigVMExternalVariables()
+{
+	return TArray<FRigVMExternalVariable>(); 
+}
+
+namespace UE::AnimNext::Private
+{
+
+static TArray<UClass*> GetClassObjectsInPackage(UPackage* InPackage)
+{
+	TArray<UObject*> Objects;
+	GetObjectsWithOuter(InPackage, Objects, false);
+
+	TArray<UClass*> ClassObjects;
+	for (UObject* Object : Objects)
+	{
+		if (UClass* Class = Cast<UClass>(Object))
+		{
+			ClassObjects.Add(Class);
+		}
+	}
+
+	return ClassObjects;
+}
+
+}
+
+void UAnimNextParameterBlock::PostRename(UObject* OldOuter, const FName OldName)
+{
+	Super::PostRename(OldOuter, OldName);
+
+	// Whenever the asset is renamed/moved, generated classes parented to the old package
+	// are not moved to the new package automatically (see FAssetRenameManager), so we
+	// have to manually perform the move/rename, to avoid invalid reference to the old package
+
+	// Note: while asset duplication doesn't duplicate the classes either, it is not a problem there
+	// because we always recompile in post duplicate.
+	TArray<UClass*> ClassObjects = UE::AnimNext::Private::GetClassObjectsInPackage(OldOuter->GetPackage());
+
+	for (UClass* ClassObject : ClassObjects)
+	{
+		if (URigVMMemoryStorageGeneratorClass* MemoryClass = Cast<URigVMMemoryStorageGeneratorClass>(ClassObject))
+		{
+			MemoryClass->Rename(nullptr, GetPackage(), REN_ForceNoResetLoaders | REN_DoNotDirty | REN_DontCreateRedirectors | REN_NonTransactional);
+		}
+	}
+}
+
+void UAnimNextParameterBlock::GetPreloadDependencies(TArray<UObject*>& OutDeps)
+{
+	Super::GetPreloadDependencies(OutDeps);
+
+	TArray<UClass*> ClassObjects = UE::AnimNext::Private::GetClassObjectsInPackage(GetPackage());
+
+	for (UClass* ClassObject : ClassObjects)
+	{
+		if (URigVMMemoryStorageGeneratorClass* MemoryClass = Cast<URigVMMemoryStorageGeneratorClass>(ClassObject))
+		{
+			OutDeps.Add(MemoryClass);
+		}
+	}
+}
+
+void UAnimNextParameterBlock::GetAssetRegistryTags(TArray<FAssetRegistryTag>& OutTags) const
+{
+	Super::GetAssetRegistryTags(OutTags);
+	
+#if WITH_EDITORONLY_DATA
+	if(EditorData)
+	{
+		EditorData->GetAssetRegistryTags(OutTags);
+	}
+#endif
+}
\ No newline at end of file
diff --git a/Engine/Plugins/Experimental/AnimNext/Source/AnimNext/Private/Param/RigUnit_AnimNextParametersBeginExecution.cpp b/Engine/Plugins/Experimental/AnimNext/Source/AnimNext/Private/Param/RigUnit_AnimNextParametersBeginExecution.cpp
new file mode 100644
--- /dev/null
+++ b/Engine/Plugins/Experimental/AnimNext/Source/AnimNext/Private/Param/RigUnit_AnimNextParametersBeginExecution.cpp
@@ -0,0 +1,9 @@
+﻿// Copyright Epic Games, Inc. All Rights Reserved.
+
+#include "Param/RigUnit_AnimNextParametersBeginExecution.h"
+
+FName FRigUnit_AnimNextParametersBeginExecution::EventName = TEXT("Update");
+
+FRigUnit_AnimNextParametersBeginExecution_Execute()
+{
+}
diff --git a/Engine/Plugins/Experimental/AnimNext/Source/AnimNext/Private/Param/RigVMDispatch_SetParameter.cpp b/Engine/Plugins/Experimental/AnimNext/Source/AnimNext/Private/Param/RigVMDispatch_SetParameter.cpp
new file mode 100644
--- /dev/null
+++ b/Engine/Plugins/Experimental/AnimNext/Source/AnimNext/Private/Param/RigVMDispatch_SetParameter.cpp
@@ -0,0 +1,61 @@
+﻿// Copyright Epic Games, Inc. All Rights Reserved.
+
+#include "Param/RigVMDispatch_SetParameter.h"
+
+#include "Param/ParametersExecuteContext.h"
+#include "RigVMCore/RigVMStruct.h"
+#include "RigVMCore/RigVM.h"
+
+class UAnimNextParameter;
+const FName FRigVMDispatch_SetParameter::ValueName = TEXT("Value");
+const FName FRigVMDispatch_SetParameter::ExecuteContextName = TEXT("ExecuteContext");
+
+FRigVMDispatch_SetParameter::FRigVMDispatch_SetParameter()
+{
+	FactoryScriptStruct = StaticStruct();
+}
+
+FName FRigVMDispatch_SetParameter::GetArgumentNameForOperandIndex(int32 InOperandIndex, int32 InTotalOperands) const
+{
+	static const FName ArgumentNames[] = 
+	{
+		ValueName
+	};
+	check(InTotalOperands == UE_ARRAY_COUNT(ArgumentNames));
+	return ArgumentNames[InOperandIndex];
+}
+
+TArray<FRigVMTemplateArgument> FRigVMDispatch_SetParameter::GetArguments() const
+{
+	const TArray<FRigVMTemplateArgument::ETypeCategory> ValueCategories =
+	{
+		FRigVMTemplateArgument::ETypeCategory_SingleAnyValue,
+		FRigVMTemplateArgument::ETypeCategory_ArrayAnyValue
+	};
+	return { FRigVMTemplateArgument(ValueName, ERigVMPinDirection::Input, ValueCategories) };
+}
+
+TArray<FRigVMExecuteArgument> FRigVMDispatch_SetParameter::GetExecuteArguments_Impl(const FRigVMDispatchContext& InContext) const
+{
+	return { { ExecuteContextName, ERigVMPinDirection::Input } };
+}
+
+FRigVMTemplateTypeMap FRigVMDispatch_SetParameter::OnNewArgumentType(const FName& InArgumentName, TRigVMTypeIndex InTypeIndex) const
+{
+	FRigVMTemplateTypeMap Types;
+	Types.Add(ValueName, InTypeIndex);
+	return Types;
+}
+
+void FRigVMDispatch_SetParameter::Execute(FRigVMExtendedExecuteContext& InContext, FRigVMMemoryHandleArray Handles)
+{
+	const FProperty* ValueProperty = Handles[0].GetResolvedProperty(); 
+	check(ValueProperty);
+	const uint8* SourceData =  Handles[0].GetData();
+
+	FAnimNextParametersExecuteContext& Context = InContext.GetPublicData<FAnimNextParametersExecuteContext>();
+	TArrayView<uint8> TargetData = Context.GetData();
+
+	ValueProperty->CopyCompleteValue(TargetData.GetData(), SourceData);
+}
+
diff --git a/Engine/Plugins/Experimental/AnimNext/Source/AnimNext/Public/Param/ParamType.h b/Engine/Plugins/Experimental/AnimNext/Source/AnimNext/Public/Param/ParamType.h
--- a/Engine/Plugins/Experimental/AnimNext/Source/AnimNext/Public/Param/ParamType.h
+++ b/Engine/Plugins/Experimental/AnimNext/Source/AnimNext/Public/Param/ParamType.h
@@ -8,6 +8,7 @@
 #include "UObject/ObjectPtr.h"
 #include "UObject/SoftObjectPtr.h"
 #include "UObject/ReflectedTypeAccessors.h"
+#include "Concepts/BaseStructureProvider.h"
 #include "ParamType.generated.h"
 
 namespace UE::AnimNext
@@ -16,24 +17,13 @@ namespace UE::AnimNext
 	struct FParamHelpers;
 }
 
-namespace UE::AnimNext::Private
+namespace UE::AnimNext::UncookedOnly
 {
-	struct CHasStaticClass
-	{
-		template<typename ValueType>
-		auto Requires(UClass* Result) -> decltype(
-			Result = ValueType::StaticClass()
-		);
-	};
-
-	struct CHasStaticStruct
-	{
-		template<typename ValueType>
-		auto Requires(UScriptStruct* Result) -> decltype(
-			Result = ValueType::StaticStruct()
-		);
-	};
+	struct FUtils;
+}
 
+namespace UE::AnimNext::Private
+{
 	template <typename T>
 	struct TIsSoftObjectPtr
 	{
@@ -73,6 +63,7 @@ public:
 
 	friend struct UE::AnimNext::FParamTypeHandle;
 	friend struct UE::AnimNext::FParamHelpers;
+	friend struct UE::AnimNext::UncookedOnly::FUtils;
 
 	FAnimNextParamType() = default;
 
@@ -138,12 +129,17 @@ private:
 			ParameterType.ValueType = EValueType::Enum;
 			ParameterType.ValueTypeObject = StaticEnum<ParamType>();
 		}
-		else if constexpr (TModels<UE::AnimNext::Private::CHasStaticStruct, ParamType>::Value)
+		else if constexpr (TModels<CStaticStructProvider, ParamType>::Value)
 		{
 			ParameterType.ValueType = EValueType::Struct;
 			ParameterType.ValueTypeObject = ParamType::StaticStruct();
 		}
-		else if constexpr (TModels<UE::AnimNext::Private::CHasStaticClass, ParamType>::Value)
+		else if constexpr (TModels<CBaseStructureProvider, ParamType>::Value)
+		{
+			ParameterType.ValueType = EValueType::Struct;
+			ParameterType.ValueTypeObject = TBaseStructure<ParamType>::Get();
+		}
+		else if constexpr (TModels<CStaticClassProvider, ParamType>::Value)
 		{
 			if constexpr (std::is_same_v<ParamType, UClass>)
 			{
@@ -184,6 +180,10 @@ private:
 			ParameterType.ValueType = EValueType::SoftClass;
 			ParameterType.ValueTypeObject = ParamType::ElementType::StaticClass();
 		}
+		else
+		{
+			static_assert(sizeof(ParamType) == 0, "Type is not expressible as a FAnimNextParamType for available types.");
+		}
 	}
 
 	/** Helper function for IsValid */
@@ -214,6 +214,24 @@ public:
 	/** Get a parameter type handle that represents this type */
 	UE::AnimNext::FParamTypeHandle GetHandle() const;
 
+	/** Get the pointer to the object that defines the Enum, Struct, or Class. */
+	UObject* GetValueTypeObject() const
+	{
+		return ValueTypeObject;
+	}
+
+	/** Get the type of the value described by this parameter. */
+	EValueType GetValueType() const
+	{
+		return ValueType;
+	}
+
+	/** Get the type of the container described by this parameter. */
+	EContainerType GetContainerType() const
+	{
+		return ContainerType;
+	}
+
 	/** 
 	 * Helper function returning the size of the type.
 	 *
diff --git a/Engine/Plugins/Experimental/AnimNext/Source/AnimNext/Public/Param/ParamTypeHandle.h b/Engine/Plugins/Experimental/AnimNext/Source/AnimNext/Public/Param/ParamTypeHandle.h
--- a/Engine/Plugins/Experimental/AnimNext/Source/AnimNext/Public/Param/ParamTypeHandle.h
+++ b/Engine/Plugins/Experimental/AnimNext/Source/AnimNext/Public/Param/ParamTypeHandle.h
@@ -3,6 +3,7 @@
 #pragma once
 
 #include "CoreMinimal.h"
+#include "Concepts/BaseStructureProvider.h"
 #include "Param/ParamType.h"
 
 class FAnimationAnimNextParametersTest_Types;
@@ -142,12 +143,17 @@ private:
 			OutValueType = FAnimNextParamType::EValueType::Enum;
 			OutValueTypeObject = StaticEnum<ParamType>();
 		}
-		else if constexpr (TModels<Private::CHasStaticStruct, ParamType>::Value)
+		else if constexpr (TModels<CStaticStructProvider, ParamType>::Value)
 		{
 			OutValueType = FAnimNextParamType::EValueType::Struct;
 			OutValueTypeObject = ParamType::StaticStruct();
 		}
-		else if constexpr (TModels<Private::CHasStaticClass, ParamType>::Value)
+		else if constexpr (TModels<CBaseStructureProvider, ParamType>::Value)
+		{
+			OutValueType = FAnimNextParamType::EValueType::Struct;
+			OutValueTypeObject = TBaseStructure<ParamType>::Get();
+		}
+		else if constexpr (TModels<CStaticClassProvider, ParamType>::Value)
 		{
 			if constexpr (std::is_same_v<ParamType, UClass>)
 			{
diff --git a/Engine/Plugins/Experimental/AnimNext/Source/AnimNextEditor/AnimNextEditor.Build.cs b/Engine/Plugins/Experimental/AnimNext/Source/AnimNextEditor/AnimNextEditor.Build.cs
new file mode 100644
--- /dev/null
+++ b/Engine/Plugins/Experimental/AnimNext/Source/AnimNextEditor/AnimNextEditor.Build.cs
@@ -0,0 +1,43 @@
+// Copyright Epic Games, Inc. All Rights Reserved.
+
+namespace UnrealBuildTool.Rules
+{
+	public class AnimNextEditor : ModuleRules
+	{
+		public AnimNextEditor(ReadOnlyTargetRules Target) : base(Target)
+		{
+			PrivateDependencyModuleNames.AddRange(
+				new string[]
+				{
+					"Core",
+					"CoreUObject",
+					"Engine",
+					"AssetTools",
+					"UnrealEd",
+					"AnimNext",
+					"AnimNextUncookedOnly",
+					"UnrealEd",
+					"SlateCore",
+					"Slate",
+					"InputCore",
+					"PropertyEditor",
+					"RigVM",
+					"RigVMDeveloper",
+					"ControlRig",
+					"ControlRigDeveloper",
+					"ControlRigEditor",
+					"GraphEditor",
+					"ToolWidgets",
+					"ToolMenus",
+					"AssetDefinition",
+					"SourceControl", 
+					"KismetWidgets",
+					"StructUtilsEditor",
+					"BlueprintGraph",	// For K2 Schema
+					"DesktopWidgets",
+					"ContentBrowserFileDataSource",
+				}
+			);
+		}
+	}
+}
\ No newline at end of file
diff --git a/Engine/Plugins/Experimental/AnimNext/Source/AnimNextEditor/Internal/AnimNextEditorModule.h b/Engine/Plugins/Experimental/AnimNext/Source/AnimNextEditor/Internal/AnimNextEditorModule.h
new file mode 100644
--- /dev/null
+++ b/Engine/Plugins/Experimental/AnimNext/Source/AnimNextEditor/Internal/AnimNextEditorModule.h
@@ -0,0 +1,30 @@
+﻿// Copyright Epic Games, Inc. All Rights Reserved.
+
+#pragma once
+
+#include "CoreMinimal.h"
+#include "Modules/ModuleManager.h"
+
+namespace UE::AnimNext::Editor
+{
+
+class FAssetTypeActions_AnimNextGraph;
+class FAnimNextPropertyTypeCustomization;
+class FPropertyTypeIdentifier;
+struct FParameterPickerArgs;
+
+class FModule : public IModuleInterface
+{
+public:
+	// Create a parameter picker
+	TSharedPtr<SWidget> CreateParameterPicker(const FParameterPickerArgs& InArgs);
+	
+private:
+	virtual void StartupModule() override;
+	virtual void ShutdownModule() override;
+
+	TSharedPtr<FAssetTypeActions_AnimNextGraph> AssetTypeActions_AnimNextGraph;
+	TSharedPtr<FPropertyTypeIdentifier> AnimNextPropertyTypeIdentifier;
+};
+
+}
\ No newline at end of file
diff --git a/Engine/Plugins/Experimental/AnimNext/Source/AnimNextEditor/Internal/Param/ParameterBlockFactory.h b/Engine/Plugins/Experimental/AnimNext/Source/AnimNextEditor/Internal/Param/ParameterBlockFactory.h
new file mode 100644
--- /dev/null
+++ b/Engine/Plugins/Experimental/AnimNext/Source/AnimNextEditor/Internal/Param/ParameterBlockFactory.h
@@ -0,0 +1,22 @@
+﻿// Copyright Epic Games, Inc. All Rights Reserved.
+
+#pragma once
+
+#include "CoreMinimal.h"
+#include "Factories/Factory.h"
+#include "ParameterBlockFactory.generated.h"
+
+UCLASS(MinimalAPI, BlueprintType)
+class UAnimNextParameterBlockFactory : public UFactory
+{
+	GENERATED_BODY()
+
+	UAnimNextParameterBlockFactory();
+	
+	virtual bool ConfigureProperties() override;
+	virtual UObject* FactoryCreateNew(UClass* Class, UObject* InParent, FName Name, EObjectFlags Flags, UObject* Context, FFeedbackContext* Warn, FName CallingContext) override;
+	virtual UObject* FactoryCreateNew(UClass* Class, UObject* InParent, FName Name, EObjectFlags Flags, UObject* Context, FFeedbackContext* Warn) override
+	{
+		return FactoryCreateNew(Class, InParent, Name, Flags, Context, Warn, NAME_None);
+	}
+};
\ No newline at end of file
diff --git a/Engine/Plugins/Experimental/AnimNext/Source/AnimNextEditor/Internal/Param/ParameterLibraryFactory.h b/Engine/Plugins/Experimental/AnimNext/Source/AnimNextEditor/Internal/Param/ParameterLibraryFactory.h
new file mode 100644
--- /dev/null
+++ b/Engine/Plugins/Experimental/AnimNext/Source/AnimNextEditor/Internal/Param/ParameterLibraryFactory.h
@@ -0,0 +1,27 @@
+﻿// Copyright Epic Games, Inc. All Rights Reserved.
+
+#pragma once
+
+#include "CoreMinimal.h"
+#include "Factories/Factory.h"
+#include "ParameterLibraryFactory.generated.h"
+
+namespace UE::AnimNext::Editor
+{
+struct FUtils;
+}
+
+UCLASS(MinimalAPI, BlueprintType)
+class UAnimNextParameterLibraryFactory : public UFactory
+{
+	GENERATED_BODY()
+	
+	UAnimNextParameterLibraryFactory();
+
+	virtual bool ConfigureProperties() override;
+	virtual UObject* FactoryCreateNew(UClass* Class, UObject* InParent, FName Name, EObjectFlags Flags, UObject* Context, FFeedbackContext* Warn, FName CallingContext) override;
+	virtual UObject* FactoryCreateNew(UClass* Class, UObject* InParent, FName Name, EObjectFlags Flags, UObject* Context, FFeedbackContext* Warn) override
+	{
+		return FactoryCreateNew(Class, InParent, Name, Flags, Context, Warn, NAME_None);
+	}
+};
\ No newline at end of file
diff --git a/Engine/Plugins/Experimental/AnimNext/Source/AnimNextEditor/Internal/Param/ParameterPickerArgs.h b/Engine/Plugins/Experimental/AnimNext/Source/AnimNextEditor/Internal/Param/ParameterPickerArgs.h
new file mode 100644
--- /dev/null
+++ b/Engine/Plugins/Experimental/AnimNext/Source/AnimNextEditor/Internal/Param/ParameterPickerArgs.h
@@ -0,0 +1,68 @@
+// Copyright Epic Games, Inc. All Rights Reserved.
+
+#pragma once
+
+#include "Delegates/DelegateCombinations.h"
+
+struct FAssetData;
+
+namespace UE::AnimNext::Editor
+{
+
+// A parameter asset, optionally bound in a block
+struct FParameterBindingReference
+{
+	FParameterBindingReference() = default;
+
+	FParameterBindingReference(const FName& InParameter, const FAssetData& InLibrary, const FAssetData& InBlock = FAssetData())
+		: Parameter(InParameter)
+		, Library(InLibrary)
+		, Block(InBlock)
+	{
+	}
+
+	// Parameter name
+	FName Parameter;
+
+	// Library asset
+	FAssetData Library;
+
+	// Optional block asset that the parameter is bound in
+	FAssetData Block;
+};
+
+// Delegate called when a parameter has been picked. Block argument is invalid when an unbound parameter is chosen.
+DECLARE_DELEGATE_OneParam(FOnGetParameterBindings, TArray<FParameterBindingReference>& /*OutParameterBindings*/);
+
+// Result of a filter operation via FOnFilterParameter
+enum class EFilterParameterResult : int32
+{
+	Include,
+	Exclude
+};
+
+// Delegate called to filter parameters for display to the user
+DECLARE_DELEGATE_RetVal_OneParam(EFilterParameterResult, FOnFilterParameter, const FParameterBindingReference& /*InParameterBinding*/);
+
+struct FParameterPickerArgs
+{
+	FParameterPickerArgs() = default;
+
+	// Ptr to existing called delegate to which the picker will register a function which returns the selected parameter
+	// bindings
+	FOnGetParameterBindings* OnGetParameterBindings = nullptr;
+
+	// Delegate used to signal whether selection has changed
+	FSimpleDelegate OnSelectionChanged;
+
+	// Delegate called to filter parameters for display to the user
+	FOnFilterParameter OnFilterParameter;
+
+	// Whether we should show parameters that are bound in a parameter block
+	bool bShowBoundParameters = true;
+
+	// Whether we should show parameters that are not bound in a parameter block
+	bool bShowUnboundParameters = true;	
+};
+
+}
diff --git a/Engine/Plugins/Experimental/AnimNext/Source/AnimNextEditor/Private/AnimNextEditorModule.cpp b/Engine/Plugins/Experimental/AnimNext/Source/AnimNextEditor/Private/AnimNextEditorModule.cpp
new file mode 100644
--- /dev/null
+++ b/Engine/Plugins/Experimental/AnimNext/Source/AnimNextEditor/Private/AnimNextEditorModule.cpp
@@ -0,0 +1,58 @@
+// Copyright Epic Games, Inc. All Rights Reserved.
+
+#include "AnimNextEditorModule.h"
+#include "AssetToolsModule.h"
+#include "IAssetTools.h"
+#include "Graph/AssetTypeActions.h"
+#include "Graph/PropertyTypeCustomization.h"
+#include "Param/ParamTypePropertyCustomization.h"
+#include "Param/ParameterPickerArgs.h"
+#include "Param/SParameterPicker.h"
+#include "Param/ParamType.h"
+
+namespace UE::AnimNext::Editor
+{
+
+void FModule::StartupModule()
+{
+	IAssetTools& AssetTools = FModuleManager::LoadModuleChecked<FAssetToolsModule>("AssetTools").Get();
+	AssetTypeActions_AnimNextGraph = MakeShared<FAssetTypeActions_AnimNextGraph>();
+	AssetTools.RegisterAssetTypeActions(AssetTypeActions_AnimNextGraph.ToSharedRef());
+
+	AnimNextPropertyTypeIdentifier = MakeShared<FPropertyTypeIdentifier>();
+	FPropertyEditorModule& PropertyModule = FModuleManager::LoadModuleChecked<FPropertyEditorModule>("PropertyEditor");
+	PropertyModule.RegisterCustomPropertyTypeLayout(
+		FInterfaceProperty::StaticClass()->GetFName(),
+		FOnGetPropertyTypeCustomizationInstance::CreateLambda([] { return MakeShared<FPropertyTypeCustomization>(); }),
+		AnimNextPropertyTypeIdentifier);
+
+	PropertyModule.RegisterCustomPropertyTypeLayout(
+		FAnimNextParamType::StaticStruct()->GetFName(),
+		FOnGetPropertyTypeCustomizationInstance::CreateLambda([] { return MakeShared<FParamTypePropertyTypeCustomization>(); }));
+}
+
+void FModule::ShutdownModule()
+{
+	if(FModuleManager::Get().IsModuleLoaded("AssetTools"))
+	{
+		IAssetTools& AssetTools = FModuleManager::GetModuleChecked<FAssetToolsModule>("AssetTools").Get();
+		AssetTools.UnregisterAssetTypeActions(AssetTypeActions_AnimNextGraph.ToSharedRef());
+	}
+	
+	if(FModuleManager::Get().IsModuleLoaded("PropertyEditor"))
+	{
+		FPropertyEditorModule& PropertyModule = FModuleManager::LoadModuleChecked<FPropertyEditorModule>("PropertyEditor");
+		PropertyModule.UnregisterCustomPropertyTypeLayout(FInterfaceProperty::StaticClass()->GetFName(), AnimNextPropertyTypeIdentifier);
+		PropertyModule.UnregisterCustomPropertyTypeLayout(FAnimNextParamType::StaticStruct()->GetFName());
+	}
+}
+
+TSharedPtr<SWidget> FModule::CreateParameterPicker(const FParameterPickerArgs& InArgs)
+{
+	return SNew(SParameterPicker)
+		.Args(InArgs);
+}
+
+}
+
+IMPLEMENT_MODULE(UE::AnimNext::Editor::FModule, AnimNextEditor);
\ No newline at end of file
diff --git a/Engine/Plugins/Experimental/AnimNext/Source/AnimNextEditor/Private/EditorUtils.h b/Engine/Plugins/Experimental/AnimNext/Source/AnimNextEditor/Private/EditorUtils.h
new file mode 100644
--- /dev/null
+++ b/Engine/Plugins/Experimental/AnimNext/Source/AnimNextEditor/Private/EditorUtils.h
@@ -0,0 +1,48 @@
+// Copyright Epic Games, Inc. All Rights Reserved.
+
+#pragma once
+
+#include "CoreMinimal.h"
+#include "Param/AnimNextParameterBlockEntry.h"
+#include "Param/ParamType.h"
+
+class UAnimNextGraph;
+class UAnimNextGraph_EditorData;
+struct FAnimNextParamType;
+class UAnimNextParameterBlock;
+class UAnimNextParameterBlockBinding;
+class UAnimNextParameterBlock_EditorData;
+class UAnimNextParameter;
+class URigVMController;
+struct FAnimNextParameterLibraryAssetRegistryExports;
+struct FAnimNextParameterBlockAssetRegistryExports;
+
+namespace UE::AnimNext::Editor
+{
+
+struct FUtils
+{
+	static FName ValidateName(const UAnimNextGraph_EditorData* InEditorData, const FString& InName);
+
+	static void GetAllGraphNames(const UAnimNextGraph_EditorData* InEditorData, TSet<FName>& OutNames);
+
+	static FAnimNextParamType GetParameterTypeFromMetaData(const FStringView& InStringView);
+
+	static FName ValidateName(const UAnimNextParameterBlock_EditorData* InEditorData, const FString& InName);
+
+	static void GetAllGraphNames(const UAnimNextParameterBlock_EditorData* InEditorData, TSet<FName>& OutNames);
+
+	static void GetFilteredVariableTypeTree(TArray<TSharedPtr<UEdGraphSchema_K2::FPinTypeTreeInfo>>& TypeTree, ETypeTreeFilter TypeTreeFilter);
+
+	static bool GetExportedParametersForLibrary(const FAssetData& InLibraryAsset, FAnimNextParameterLibraryAssetRegistryExports& OutExports);
+
+	static FName GetNewParameterNameInLibrary(const FAssetData& InLibraryAsset, const TCHAR* InBaseName, TArrayView<FName> InAdditionalExistingNames);
+
+	static bool DoesParameterExistInLibrary(const FAssetData& InLibraryAsset, const FName InParameterName);
+
+	static FAnimNextParamType GetParameterTypeFromLibraryExports(FName InName, const FAnimNextParameterLibraryAssetRegistryExports& InExports);
+
+	static bool GetExportedBindingsForBlock(const FAssetData& InBlockAsset, FAnimNextParameterBlockAssetRegistryExports& OutExports);
+};
+
+}
diff --git a/Engine/Plugins/Experimental/AnimNext/Source/AnimNextEditor/Private/Graph/AnimNextGraphEditor.cpp b/Engine/Plugins/Experimental/AnimNext/Source/AnimNextEditor/Private/Graph/AnimNextGraphEditor.cpp
new file mode 100644
--- /dev/null
+++ b/Engine/Plugins/Experimental/AnimNext/Source/AnimNextEditor/Private/Graph/AnimNextGraphEditor.cpp
@@ -0,0 +1,320 @@
+﻿// Copyright Epic Games, Inc. All Rights Reserved.
+
+#include "AnimNextGraphEditor.h"
+#include "GraphEditorMode.h"
+#include "Graph/AnimNextGraph.h"
+#include "Graph/AnimNextGraph_EditorData.h"
+#include "EdGraphNode_Comment.h"
+#include "SActionMenu.h"
+#include "UncookedOnlyUtils.h"
+#include "Framework/Commands/GenericCommands.h"
+#include "RigVMModel/RigVMController.h"
+#include "Kismet2/BlueprintEditorUtils.h"
+#include "ScopedTransaction.h"
+
+#define LOCTEXT_NAMESPACE "AnimNextGraphEditor"
+
+namespace UE::AnimNext::Editor
+{
+
+namespace GraphModes
+{
+	const FName GraphEditor("AnimNextGraphEditorMode");
+}
+
+namespace GraphTabs
+{
+	const FName Details("DetailsTab");
+	const FName Document("Document");
+}
+
+const FName GraphAppIdentifier("AnimNextGraphEditor");
+
+void FGraphEditor::InitEditor(const EToolkitMode::Type InMode, const TSharedPtr<IToolkitHost>& InInitToolkitHost, UAnimNextGraph* InAnimNextGraph)
+{
+	AnimNextGraph = InAnimNextGraph;
+
+	AnimNextGraph_EditorData = UE::AnimNext::UncookedOnly::FUtils::GetEditorData(AnimNextGraph);
+	AnimNextGraph_EditorData->Initialize(false);
+	
+	DocumentManager = MakeShared<FDocumentTracker>();
+	DocumentManager->Initialize(SharedThis(this));
+
+	TSharedRef<FGraphEditorSummoner> GraphEditorSummoner = MakeShared<FGraphEditorSummoner>(SharedThis(this));
+	GraphEditorSummoner->OnCreateGraphEditorWidget().BindSP(this, &FGraphEditor::CreateGraphEditorWidget);
+	GraphEditorSummoner->OnGraphEditorFocused().BindSP(this, &FGraphEditor::OnGraphEditorFocused);
+	GraphEditorSummoner->OnGraphEditorBackgrounded().BindSP(this, &FGraphEditor::OnGraphEditorBackgrounded);
+	GraphEditorTabFactoryPtr = GraphEditorSummoner;
+	DocumentManager->RegisterDocumentFactory(GraphEditorSummoner);
+
+	const bool bCreateDefaultStandaloneMenu = true;
+	const bool bCreateDefaultToolbar = true;
+	FAssetEditorToolkit::InitAssetEditor(InMode, InInitToolkitHost, GraphAppIdentifier, FTabManager::FLayout::NullLayout, bCreateDefaultStandaloneMenu, bCreateDefaultToolbar, InAnimNextGraph);
+
+	BindCommands();
+
+	AddApplicationMode(GraphModes::GraphEditor, MakeShared<FGraphEditorMode>(SharedThis(this)));
+	SetCurrentMode(GraphModes::GraphEditor);
+	
+	ExtendMenu();
+	ExtendToolbar();
+	RegenerateMenusAndToolbars();
+
+	// Open initial document
+	//DocumentManager->OpenDocument(FTabPayload_UObject::Make(AnimNextInterfaceGraph_EditorData->EntryPointGraph), FDocumentTracker::EOpenDocumentCause::OpenNewDocument);
+	DocumentManager->OpenDocument(FTabPayload_UObject::Make(AnimNextGraph_EditorData->RootGraph), FDocumentTracker::EOpenDocumentCause::OpenNewDocument);
+	//DocumentManager->OpenDocument(FTabPayload_UObject::Make(AnimNextInterfaceGraph_EditorData->FunctionLibraryEdGraph), FDocumentTracker::EOpenDocumentCause::OpenNewDocument);
+}
+
+void FGraphEditor::BindCommands()
+{
+}
+
+void FGraphEditor::ExtendMenu()
+{
+	
+}
+
+void FGraphEditor::ExtendToolbar()
+{
+	
+}
+
+void FGraphEditor::RegisterTabSpawners(const TSharedRef<FTabManager>& InTabManager)
+{
+	DocumentManager->SetTabManager(InTabManager);
+
+	FWorkflowCentricApplication::RegisterTabSpawners(InTabManager);
+}
+
+void FGraphEditor::UnregisterTabSpawners(const TSharedRef<FTabManager>& InTabManager)
+{
+	
+}
+
+FName FGraphEditor::GetToolkitFName() const
+{
+	return FName("AnimNextInterfaceGraphEditor");
+}
+
+FText FGraphEditor::GetBaseToolkitName() const
+{
+	return LOCTEXT("AppLabel", "AnimNextInterfaceGraphEditor");
+}
+
+FString FGraphEditor::GetWorldCentricTabPrefix() const
+{
+	return LOCTEXT("WorldCentricTabPrefix", "AnimNextInterfaceGraphEditor ").ToString();
+}
+
+FLinearColor FGraphEditor::GetWorldCentricTabColorScale() const
+{
+	return FLinearColor(0.3f, 0.2f, 0.5f, 0.5f);
+}
+
+void FGraphEditor::InitToolMenuContext(FToolMenuContext& InMenuContext)
+{
+}
+
+void FGraphEditor::OnGraphEditorFocused(TSharedRef<SGraphEditor> InGraphEditor)
+{
+	// Update the graph editor that is currently focused
+	FocusedGraphEdPtr = InGraphEditor;
+}
+
+void FGraphEditor::OnGraphEditorBackgrounded(TSharedRef<SGraphEditor> InGraphEditor)
+{
+	// Update the graph editor that is currently focused
+	FocusedGraphEdPtr = nullptr;
+}
+
+UEdGraph* FGraphEditor::GetFocusedGraph() const
+{
+	if (FocusedGraphEdPtr.IsValid())
+	{
+		if (UEdGraph* Graph = FocusedGraphEdPtr.Pin()->GetCurrentGraph())
+		{
+			return Graph;
+		}
+	}
+	return nullptr;
+}
+
+URigVMGraph* FGraphEditor::GetFocusedVMGraph() const
+{
+	UAnimNextGraph_EdGraph* EdGraph = Cast<UAnimNextGraph_EdGraph>(GetFocusedGraph());
+	return AnimNextGraph_EditorData->GetVMGraphForEdGraph(EdGraph);
+}
+
+URigVMController* FGraphEditor::GetFocusedVMController() const
+{
+	return AnimNextGraph_EditorData->GetRigVMClient()->GetController(GetFocusedVMGraph());
+}
+
+TSharedRef<SGraphEditor> FGraphEditor::CreateGraphEditorWidget(TSharedRef<FTabInfo> InTabInfo, UEdGraph* InGraph)
+{
+	if(!GraphEditorCommands.IsValid())
+	{
+		GraphEditorCommands = MakeShared<FUICommandList>();
+	
+		GraphEditorCommands->MapAction(FGenericCommands::Get().Delete,
+			FExecuteAction::CreateSP(this, &FGraphEditor::DeleteSelectedNodes),
+			FCanExecuteAction::CreateSP(this, &FGraphEditor::CanDeleteSelectedNodes));
+	}
+	
+	SGraphEditor::FGraphEditorEvents Events;
+	Events.OnCreateActionMenu = SGraphEditor::FOnCreateActionMenu::CreateSP(this, &FGraphEditor::OnCreateGraphActionMenu);
+	
+	return SNew(SGraphEditor)
+		.AdditionalCommands(GraphEditorCommands)
+		.IsEditable(this, &FGraphEditor::IsEditable, InGraph)
+		.GraphToEdit(InGraph)
+		.GraphEvents(Events)
+		.AssetEditorToolkit(AsShared());
+}
+
+FActionMenuContent FGraphEditor::OnCreateGraphActionMenu(UEdGraph* InGraph, const FVector2D& InNodePosition, const TArray<UEdGraphPin*>& InDraggedPins, bool bAutoExpand, SGraphEditor::FActionMenuClosed InOnMenuClosed)
+{
+	TSharedRef<SActionMenu> ActionMenu = SNew(SActionMenu)
+		.AutoExpandActionMenu(bAutoExpand)
+		.Graph(InGraph)
+		.NewNodePosition(InNodePosition)
+		.DraggedFromPins(InDraggedPins)
+		.OnClosedCallback(InOnMenuClosed);
+
+	TSharedPtr<SWidget> FilterTextBox = StaticCastSharedRef<SWidget>(ActionMenu->GetFilterTextBox());
+	return FActionMenuContent(StaticCastSharedRef<SWidget>(ActionMenu), FilterTextBox);
+}
+
+void FGraphEditor::OnNodeTitleCommitted(const FText& NewText, ETextCommit::Type CommitInfo, UEdGraphNode* NodeBeingChanged)
+{
+	if (UEdGraphNode_Comment* CommentBeingChanged = Cast<UEdGraphNode_Comment>(NodeBeingChanged))
+	{
+		GetFocusedVMController()->SetCommentTextByName(CommentBeingChanged->GetFName(), NewText.ToString(), CommentBeingChanged->FontSize, CommentBeingChanged->bCommentBubbleVisible, CommentBeingChanged->bColorCommentBubble, true, true);
+	}
+}
+
+void FGraphEditor::CloseDocumentTab(const UObject* DocumentID)
+{
+	TSharedRef<FTabPayload_UObject> Payload = FTabPayload_UObject::Make(DocumentID);
+	DocumentManager->CloseTab(Payload);
+}
+
+bool FGraphEditor::InEditingMode() const
+{
+	// @TODO: disallow editing when debugging when implemented
+
+	return true;
+}
+
+bool FGraphEditor::IsEditable(UEdGraph* InGraph) const
+{
+	return InGraph && InEditingMode() && InGraph->bEditable;
+}
+
+FGraphPanelSelectionSet FGraphEditor::GetSelectedNodes() const
+{
+	FGraphPanelSelectionSet CurrentSelection;
+	TSharedPtr<SGraphEditor> FocusedGraphEd = FocusedGraphEdPtr.Pin();
+	if (FocusedGraphEd.IsValid())
+	{
+		CurrentSelection = FocusedGraphEd->GetSelectedNodes();
+	}
+	return CurrentSelection;
+}
+
+void FGraphEditor::DeleteSelectedNodes()
+{
+	TSharedPtr<SGraphEditor> FocusedGraphEd = FocusedGraphEdPtr.Pin();
+	if (!FocusedGraphEd.IsValid())
+	{
+		return;
+	}
+
+	const FScopedTransaction Transaction(FGenericCommands::Get().Delete->GetDescription());
+	FocusedGraphEd->GetCurrentGraph()->Modify();
+	
+	const FGraphPanelSelectionSet SelectedNodes = GetSelectedNodes();
+
+//	SetUISelectionState(NAME_None);
+
+	if(FocusedGraphEd)
+	{
+		FocusedGraphEd->ClearSelectionSet();
+	}
+
+	// Some nodes have sub-objects that are represented as other tabs.
+	// Close them here as a pre-pass before we remove their nodes. If the documents are left open they
+	// may reference dangling data and function incorrectly in cases such as FindBlueprintforNodeChecked
+	for (FGraphPanelSelectionSet::TConstIterator NodeIt( SelectedNodes ); NodeIt; ++NodeIt)
+	{
+		if (UEdGraphNode* Node = Cast<UEdGraphNode>(*NodeIt))
+		{
+			if (Node->CanUserDeleteNode())
+			{
+				auto CloseAllDocumentsTab = [this](const UEdGraphNode* InNode)
+				{
+					TArray<UObject*> NodesToClose;
+					GetObjectsWithOuter(InNode, NodesToClose);
+					for (UObject* Node : NodesToClose)
+					{
+						UEdGraph* NodeGraph = Cast<UEdGraph>(Node);
+						if (NodeGraph)
+						{
+							CloseDocumentTab(NodeGraph);
+						}
+					}
+				};
+				
+				if (Node->GetSubGraphs().Num() > 0)
+				{
+					CloseAllDocumentsTab(Node);
+				}
+			}
+		}
+	}
+
+	// Now remove the selected nodes
+	for (FGraphPanelSelectionSet::TConstIterator NodeIt( SelectedNodes ); NodeIt; ++NodeIt)
+	{
+		if (UEdGraphNode* Node = Cast<UEdGraphNode>(*NodeIt))
+		{
+			if (Node->CanUserDeleteNode())
+			{
+				if (Node->GetSubGraphs().Num() > 0)
+				{
+					DocumentManager->CleanInvalidTabs();
+				}
+
+				FBlueprintEditorUtils::RemoveNode(nullptr, Node);
+			}
+		}
+	}
+}
+
+bool FGraphEditor::CanDeleteSelectedNodes()
+{
+	const FGraphPanelSelectionSet SelectedNodes = GetSelectedNodes();
+
+	bool bCanUserDeleteNode = false;
+
+	if(IsEditable(GetFocusedGraph()) && SelectedNodes.Num() > 0)
+	{
+		for(UObject* NodeObject : SelectedNodes)
+		{
+			// If any nodes allow deleting, then do not disable the delete option
+			UEdGraphNode* Node = Cast<UEdGraphNode>(NodeObject);
+			if(Node->CanUserDeleteNode())
+			{
+				bCanUserDeleteNode = true;
+				break;
+			}
+		}
+	}
+
+	return bCanUserDeleteNode;
+}
+
+}
+
+#undef LOCTEXT_NAMESPACE
\ No newline at end of file
diff --git a/Engine/Plugins/Experimental/AnimNext/Source/AnimNextEditor/Private/Graph/AnimNextGraphEditor.h b/Engine/Plugins/Experimental/AnimNext/Source/AnimNextEditor/Private/Graph/AnimNextGraphEditor.h
new file mode 100644
--- /dev/null
+++ b/Engine/Plugins/Experimental/AnimNext/Source/AnimNextEditor/Private/Graph/AnimNextGraphEditor.h
@@ -0,0 +1,123 @@
+﻿// Copyright Epic Games, Inc. All Rights Reserved.
+
+#pragma once
+
+#include "CoreMinimal.h"
+#include "WorkflowOrientedApp/WorkflowCentricApplication.h"
+#include "GraphEditor.h"
+
+enum class ERigVMGraphNotifType : uint8;
+class UAnimNextGraph;
+class UAnimNextGraph_EditorData;
+class FDocumentTracker;
+class FDocumentTabFactory;
+class FTabInfo;
+class FTabManager;
+class IToolkitHost;
+class URigVMGraph;
+class URigVMController;
+class UDetailsViewWrapperObject;
+
+namespace UE::AnimNext::Editor
+{
+
+namespace GraphModes
+{
+	extern const FName GraphEditor;
+}
+
+namespace GraphTabs
+{
+	extern const FName Details;
+	extern const FName Document;
+}
+
+class FGraphEditor : public FWorkflowCentricApplication
+{
+public:
+	/** Edits the specified asset */
+	void InitEditor(const EToolkitMode::Type InMode, const TSharedPtr<IToolkitHost>& InInitToolkitHost, UAnimNextGraph* InAnimNextInterfaceGraph);
+
+private:
+	friend class FGraphEditorMode;
+	
+	// IToolkit interface
+	virtual void RegisterTabSpawners(const TSharedRef<FTabManager>& InTabManager) override;
+	virtual void UnregisterTabSpawners(const TSharedRef<FTabManager>& InTabManager) override;
+	virtual FName GetToolkitFName() const override;
+	virtual FText GetBaseToolkitName() const override;
+	virtual FString GetWorldCentricTabPrefix() const override;
+	virtual FLinearColor GetWorldCentricTabColorScale() const override;
+	virtual void InitToolMenuContext(FToolMenuContext& InMenuContext) override;
+
+	void BindCommands();
+	
+	void ExtendMenu();
+
+	void ExtendToolbar();
+
+	void OnGraphEditorFocused(TSharedRef<SGraphEditor> InGraphEditor);
+
+	void OnGraphEditorBackgrounded(TSharedRef<SGraphEditor> InGraphEditor);
+	
+	UEdGraph* GetFocusedGraph() const;
+
+	URigVMGraph* GetFocusedVMGraph() const;
+
+	URigVMController* GetFocusedVMController() const;
+	
+	// Create new graph editor widget for the supplied document container
+	TSharedRef<SGraphEditor> CreateGraphEditorWidget(TSharedRef<FTabInfo> InTabInfo, UEdGraph* InGraph);
+
+	// Create the action menu for the specified graph
+	FActionMenuContent OnCreateGraphActionMenu(UEdGraph* InGraph, const FVector2D& InNodePosition, const TArray<UEdGraphPin*>& InDraggedPins, bool bAutoExpand, SGraphEditor::FActionMenuClosed InOnMenuClosed);
+
+	// Handle node titles being changed
+	void OnNodeTitleCommitted(const FText& NewText, ETextCommit::Type CommitInfo, UEdGraphNode* NodeBeingChanged);
+
+	void CloseDocumentTab(const UObject* DocumentID);
+
+	bool InEditingMode() const;
+
+	bool IsEditable(UEdGraph* InGraph) const;
+	
+	FGraphPanelSelectionSet GetSelectedNodes() const;
+	
+	void DeleteSelectedNodes();
+	
+	bool CanDeleteSelectedNodes();
+
+	TSharedPtr<SGraphEditor> GetGraphEditor(UEdGraph* InEdGraph) const;
+	
+	void HandleModifiedEvent(ERigVMGraphNotifType InNotifType, URigVMGraph* InGraph, UObject* InSubject);
+
+	void HandleDetailsViewCreated(TSharedRef<IDetailsView> InDetailsView)
+	{
+		DetailsView = InDetailsView;
+	}
+	
+	// The asset we are editing
+	UAnimNextGraph* AnimNextGraph = nullptr;
+
+	// Cached editor data ptr
+	UAnimNextGraph_EditorData* AnimNextGraph_EditorData = nullptr;
+	
+	// Currently focused graph editor
+	TWeakPtr<SGraphEditor> FocusedGraphEdPtr;
+
+	// Document tracker
+	TSharedPtr<FDocumentTracker> DocumentManager;
+
+	// Factory for document tabs
+	TWeakPtr<FDocumentTabFactory> GraphEditorTabFactoryPtr;
+
+	// Command list for the graph editor
+	TSharedPtr<FUICommandList> GraphEditorCommands;
+
+	// Our details panel
+	TSharedPtr<IDetailsView> DetailsView;
+
+	TArray<TStrongObjectPtr<UDetailsViewWrapperObject>> WrapperObjects;
+};
+
+}
\ No newline at end of file
diff --git a/Engine/Plugins/Experimental/AnimNext/Source/AnimNextEditor/Private/Graph/AssetTypeActions.cpp b/Engine/Plugins/Experimental/AnimNext/Source/AnimNextEditor/Private/Graph/AssetTypeActions.cpp
new file mode 100644
--- /dev/null
+++ b/Engine/Plugins/Experimental/AnimNext/Source/AnimNextEditor/Private/Graph/AssetTypeActions.cpp
@@ -0,0 +1,37 @@
+﻿// Copyright Epic Games, Inc. All Rights Reserved.
+
+#include "AssetTypeActions.h"
+#include "AnimNextGraphEditor.h"
+#include "Graph/AnimNextGraph.h"
+
+#define LOCTEXT_NAMESPACE "AssetTypeActions_AnimNextGraph" 
+
+namespace UE::AnimNext::Editor
+{
+
+void FAssetTypeActions_AnimNextGraph::OpenAssetEditor(const TArray<UObject*>& InObjects, TSharedPtr<IToolkitHost> EditWithinLevelEditor)
+{
+	const EToolkitMode::Type Mode = EditWithinLevelEditor.IsValid() ? EToolkitMode::WorldCentric : EToolkitMode::Standalone;
+
+	for (UObject* Object : InObjects)
+	{
+		if(UAnimNextGraph* AnimNextGraph = Cast<UAnimNextGraph>(Object))
+		{
+			TSharedRef<FGraphEditor> GraphEditor = MakeShared<FGraphEditor>();
+			GraphEditor->InitEditor(Mode, EditWithinLevelEditor, AnimNextGraph);
+		}
+	}
+}
+
+const TArray<FText>& FAssetTypeActions_AnimNextGraph::GetSubMenus() const
+{
+	static const TArray<FText> SubMenus
+	{
+		LOCTEXT("AnimNextSubMenu", "AnimNext")
+	};
+	return SubMenus;
+}
+
+}
+
+#undef LOCTEXT_NAMESPACE
\ No newline at end of file
diff --git a/Engine/Plugins/Experimental/AnimNext/Source/AnimNextEditor/Private/Graph/AssetTypeActions.h b/Engine/Plugins/Experimental/AnimNext/Source/AnimNextEditor/Private/Graph/AssetTypeActions.h
new file mode 100644
--- /dev/null
+++ b/Engine/Plugins/Experimental/AnimNext/Source/AnimNextEditor/Private/Graph/AssetTypeActions.h
@@ -0,0 +1,24 @@
+﻿// Copyright Epic Games, Inc. All Rights Reserved.
+
+#pragma once
+
+#include "AssetTypeActions_Base.h"
+#include "Graph/AnimNextGraph.h"
+
+namespace UE::AnimNext::Editor
+{
+
+class FAssetTypeActions_AnimNextGraph : public FAssetTypeActions_Base
+{
+public:
+	// IAssetTypeActions Implementation
+	virtual FText GetName() const override { return NSLOCTEXT("AssetTypeActions", "AssetTypeActions_AnimNextGraph", "AnimNext Graph"); }
+	virtual FColor GetTypeColor() const override { return FColor(128,128,64); }
+	virtual UClass* GetSupportedClass() const override { return UAnimNextGraph::StaticClass(); }
+	virtual void OpenAssetEditor(const TArray<UObject*>& InObjects, TSharedPtr<IToolkitHost> EditWithinLevelEditor = TSharedPtr<IToolkitHost>()) override;
+	virtual bool CanFilter() override { return true; }
+	virtual const TArray<FText>& GetSubMenus() const override;
+	virtual uint32 GetCategories() override { return EAssetTypeCategories::Animation; }
+};
+
+}
diff --git a/Engine/Plugins/Experimental/AnimNext/Source/AnimNextEditor/Private/Graph/ClassFilter.h b/Engine/Plugins/Experimental/AnimNext/Source/AnimNextEditor/Private/Graph/ClassFilter.h
new file mode 100644
--- /dev/null
+++ b/Engine/Plugins/Experimental/AnimNext/Source/AnimNextEditor/Private/Graph/ClassFilter.h
@@ -0,0 +1,44 @@
+// Copyright Epic Games, Inc. All Rights Reserved.
+#pragma once
+
+#include "Param/ParamTypeHandle.h"
+#include "ClassViewerFilter.h"
+#include "Interface/IAnimNextInterface.h"
+
+namespace UE::AnimNext::Editor
+{
+	// Filter class for things which implement IAnimNextInterface
+	class FClassFilter : public IClassViewerFilter
+	{
+	public:
+		FClassFilter(FParamTypeHandle InTypeHandle) : TypeHandle(InTypeHandle) {};
+		virtual ~FClassFilter() override {};
+	
+		virtual bool IsClassAllowed(const FClassViewerInitializationOptions& InInitOptions, const UClass* InClass, TSharedRef<FClassViewerFilterFuncs> InFilterFuncs) override
+		{
+			if (InClass->ImplementsInterface(UAnimNextInterface::StaticClass()))
+			{
+				// pass filter if the class has "" return type (for wrappers)
+				// pass all classes if the requested typename is ""
+				// otherwise fail if the TypeNames don't match
+				// for Object type, check type hierarchy. (todo)
+			
+				IAnimNextInterface* Object = static_cast<IAnimNextInterface*>(InClass->GetDefaultObject()->GetInterfaceAddress(UAnimNextInterface::StaticClass()));
+				if (Object->GetReturnTypeHandle() == TypeHandle || !Object->GetReturnTypeHandle().IsValid() || !TypeHandle.IsValid())
+				{
+					// @TODO: deal with inheritance directionality here
+					return true;
+				}
+			}
+
+			return false;
+		}
+
+		virtual bool IsUnloadedClassAllowed(const FClassViewerInitializationOptions& InInitOptions, const TSharedRef<const IUnloadedBlueprintData> InUnloadedClassData, TSharedRef<FClassViewerFilterFuncs> InFilterFuncs) override
+		{
+			return false;
+		}
+	private:
+		FParamTypeHandle TypeHandle;
+	};
+}
\ No newline at end of file
diff --git a/Engine/Plugins/Experimental/AnimNext/Source/AnimNextEditor/Private/Graph/GraphEditorMode.cpp b/Engine/Plugins/Experimental/AnimNext/Source/AnimNextEditor/Private/Graph/GraphEditorMode.cpp
new file mode 100644
--- /dev/null
+++ b/Engine/Plugins/Experimental/AnimNext/Source/AnimNextEditor/Private/Graph/GraphEditorMode.cpp
@@ -0,0 +1,273 @@
+﻿// Copyright Epic Games, Inc. All Rights Reserved.
+
+#include "GraphEditorMode.h"
+#include "AnimNextGraphEditor.h"
+#include "PropertyEditorModule.h"
+#include "Widgets/Docking/SDockTab.h"
+#include "EdGraph/EdGraph.h"
+#include "EdGraph/EdGraphSchema.h"
+#include "Widgets/Layout/SSpacer.h"
+#include "Modules/ModuleManager.h"
+
+#define LOCTEXT_NAMESPACE "GraphEditorMode"
+
+namespace UE::AnimNext::Editor
+{
+
+struct FGraphTabHistory : public FGenericTabHistory
+{
+public:
+	FGraphTabHistory(TSharedPtr<FDocumentTabFactory> InFactory, TSharedPtr<FTabPayload> InPayload)
+		: FGenericTabHistory(InFactory, InPayload)
+		, SavedLocation(FVector2D::ZeroVector)
+		, SavedZoomAmount(INDEX_NONE)
+	{
+	}
+
+private:
+	// FGenericTabHistory interface
+	virtual void EvokeHistory(TSharedPtr<::FTabInfo> InTabInfo, bool bPrevTabMatches) override
+	{
+		FWorkflowTabSpawnInfo SpawnInfo;
+		SpawnInfo.Payload = Payload;
+		SpawnInfo.TabInfo = InTabInfo;
+
+		if(bPrevTabMatches)
+		{
+			TSharedPtr<SDockTab> DockTab = InTabInfo->GetTab().Pin();
+			GraphEditor = StaticCastSharedRef<SGraphEditor>(DockTab->GetContent());
+		}
+		else
+		{
+			TSharedRef<SGraphEditor> GraphEditorRef = StaticCastSharedRef<SGraphEditor>(FactoryPtr.Pin()->CreateTabBody(SpawnInfo));
+			GraphEditor = GraphEditorRef;
+			FactoryPtr.Pin()->UpdateTab(InTabInfo->GetTab().Pin(), SpawnInfo, GraphEditorRef);
+		}
+	}
+	
+	virtual void SaveHistory() override
+	{
+		if (IsHistoryValid())
+		{
+			check(GraphEditor.IsValid());
+			GraphEditor.Pin()->GetViewLocation(SavedLocation, SavedZoomAmount);
+			GraphEditor.Pin()->GetViewBookmark(SavedBookmarkId);
+		}
+	}
+	
+	virtual void RestoreHistory() override
+	{
+		if (IsHistoryValid())
+		{
+			check(GraphEditor.IsValid());
+			GraphEditor.Pin()->SetViewLocation(SavedLocation, SavedZoomAmount, SavedBookmarkId);
+		}
+	}
+	
+	/** The graph editor represented by this history node. While this node is inactive, the graph editor is invalid */
+	TWeakPtr<SGraphEditor> GraphEditor;
+	
+	/** Saved location the graph editor was at when this history node was last visited */
+	FVector2D SavedLocation;
+	
+	/** Saved zoom the graph editor was at when this history node was last visited */
+	float SavedZoomAmount;
+	
+	/** Saved bookmark ID the graph editor was at when this history node was last visited */
+	FGuid SavedBookmarkId;
+};
+
+FGraphEditorSummoner::FGraphEditorSummoner(TSharedPtr<FWorkflowCentricApplication> InHostingApp)
+	: FDocumentTabFactoryForObjects<UEdGraph>(GraphTabs::Document, InHostingApp)
+{
+}
+
+void FGraphEditorSummoner::OnTabActivated(TSharedPtr<SDockTab> Tab) const
+{
+	TSharedRef<SGraphEditor> GraphEditor = StaticCastSharedRef<SGraphEditor>(Tab->GetContent());
+	OnGraphEditorFocusedDelegate.ExecuteIfBound(GraphEditor);
+}
+
+void FGraphEditorSummoner::OnTabBackgrounded(TSharedPtr<SDockTab> Tab) const
+{
+	TSharedRef<SGraphEditor> GraphEditor = StaticCastSharedRef<SGraphEditor>(Tab->GetContent());
+	OnGraphEditorBackgroundedDelegate.ExecuteIfBound(GraphEditor);
+}
+
+void FGraphEditorSummoner::OnTabRefreshed(TSharedPtr<SDockTab> Tab) const
+{
+	TSharedRef<SGraphEditor> GraphEditor = StaticCastSharedRef<SGraphEditor>(Tab->GetContent());
+	GraphEditor->NotifyGraphChanged();
+}
+
+void FGraphEditorSummoner::SaveState(TSharedPtr<SDockTab> Tab, TSharedPtr<FTabPayload> Payload) const
+{
+	TSharedRef<SGraphEditor> GraphEditor = StaticCastSharedRef<SGraphEditor>(Tab->GetContent());
+
+	FVector2D ViewLocation;
+	float ZoomAmount;
+	GraphEditor->GetViewLocation(ViewLocation, ZoomAmount);
+
+	UEdGraph* Graph = Payload->IsValid() ? FTabPayload_UObject::CastChecked<UEdGraph>(Payload) : nullptr;
+
+	OnSaveGraphStateDelegate.ExecuteIfBound(Graph, ViewLocation, ZoomAmount);
+}
+
+TAttribute<FText> FGraphEditorSummoner::ConstructTabNameForObject(UEdGraph* DocumentID) const
+{
+	if (DocumentID)
+	{
+		if (const UEdGraphSchema* Schema = DocumentID->GetSchema())
+		{
+			FGraphDisplayInfo Info;
+			Schema->GetGraphDisplayInformation(*DocumentID, /*out*/ Info);
+			return Info.DisplayName;
+		}
+		else
+		{
+			// if we don't have a schema, we're dealing with a malformed (or incomplete graph)...
+			// possibly in the midst of some transaction - here we return the object's outer path 
+			// so we can at least get some context as to which graph we're referring
+			return FText::FromString(DocumentID->GetPathName());
+		}
+	}
+
+	return LOCTEXT("UnknownGraphName", "Unknown");
+}
+
+TSharedRef<SWidget> FGraphEditorSummoner::CreateTabBodyForObject(const FWorkflowTabSpawnInfo& Info, UEdGraph* DocumentID) const
+{
+	check(Info.TabInfo.IsValid());
+	
+	if(OnCreateGraphEditorWidgetDelegate.IsBound())
+	{
+		return OnCreateGraphEditorWidgetDelegate.Execute(Info.TabInfo.ToSharedRef(), DocumentID);
+	}
+	else
+	{
+		return SNew(SSpacer);
+	}
+}
+
+const FSlateBrush* FGraphEditorSummoner::GetTabIconForObject(const FWorkflowTabSpawnInfo& Info, UEdGraph* DocumentID) const
+{
+	if(OnGetTabIconDelegate.IsBound())
+	{
+		return OnGetTabIconDelegate.Execute(DocumentID);
+	}
+	else
+	{
+		return FAppStyle::Get().GetBrush(TEXT("GraphEditor.EventGraph_16x"));
+	}
+}
+
+TSharedRef<FGenericTabHistory> FGraphEditorSummoner::CreateTabHistoryNode(TSharedPtr<FTabPayload> Payload)
+{
+	return MakeShared<FGraphTabHistory>(SharedThis(this), Payload);
+}
+
+class FGraphDetailsTabSummoner : public FWorkflowTabFactory
+{
+public:
+	FGraphDetailsTabSummoner(TSharedPtr<FAssetEditorToolkit> InHostingApp, FOnDetailsViewCreated InOnDetailsViewCreated)
+		: FWorkflowTabFactory(GraphTabs::Details, InHostingApp)
+	{
+		TabLabel = LOCTEXT("DetailsTabLabel", "Details");
+		TabIcon = FSlateIcon("EditorStyle", "LevelEditor.Tabs.Details");
+		ViewMenuDescription = LOCTEXT("DetailsTabMenuDescription", "Details");
+		ViewMenuTooltip = LOCTEXT("DetailsTabToolTip", "Shows the details tab for selected objects.");
+		bIsSingleton = true;
+
+		FDetailsViewArgs DetailsViewArgs;
+		DetailsViewArgs.NameAreaSettings = FDetailsViewArgs::HideNameArea;
+		DetailsViewArgs.bHideSelectionTip = true;
+
+		FPropertyEditorModule& PropertyEditorModule = FModuleManager::LoadModuleChecked<FPropertyEditorModule>("PropertyEditor");
+		DetailsView = PropertyEditorModule.CreateDetailView(DetailsViewArgs);
+		
+		InOnDetailsViewCreated.ExecuteIfBound(DetailsView.ToSharedRef());
+	}
+
+	TSharedPtr<IDetailsView> GetDetailsView() const
+	{
+		return DetailsView;
+	}
+
+private:
+	virtual TSharedRef<SWidget> CreateTabBody(const FWorkflowTabSpawnInfo& Info) const override
+	{
+		return DetailsView.ToSharedRef();
+	}
+	
+	virtual FText GetTabToolTipText(const FWorkflowTabSpawnInfo& Info) const override
+	{
+		return ViewMenuTooltip;
+	}
+	
+	TSharedPtr<IDetailsView> DetailsView;
+};
+
+FGraphEditorMode::FGraphEditorMode(TSharedRef<FWorkflowCentricApplication> InHostingApp)
+	: FApplicationMode(GraphModes::GraphEditor)
+	, HostingAppPtr(InHostingApp)
+{
+	HostingAppPtr = InHostingApp;
+
+	TSharedRef<FGraphEditor> GraphEditor = StaticCastSharedRef<FGraphEditor>(InHostingApp);
+	
+	TabFactories.RegisterFactory(MakeShared<FGraphDetailsTabSummoner>(GraphEditor, FOnDetailsViewCreated::CreateSP(&GraphEditor.Get(), &FGraphEditor::HandleDetailsViewCreated)));
+
+	TabLayout = FTabManager::NewLayout("Standalone_AnimNextGraphEditor_Layout_v1.0")
+		->AddArea
+		(
+			FTabManager::NewPrimaryArea()
+			->SetOrientation(Orient_Vertical)
+			->Split
+			(	
+				FTabManager::NewSplitter()
+				->SetSizeCoefficient(1.0f)
+				->SetOrientation(Orient_Horizontal)
+				->Split
+				(
+					FTabManager::NewStack()
+					->SetSizeCoefficient(0.8f)
+					->SetHideTabWell(false)
+					->AddTab(GraphTabs::Document, ETabState::ClosedTab)
+				)
+				->Split
+				(
+					FTabManager::NewStack()
+					->SetSizeCoefficient(0.2f)
+					->SetHideTabWell(false)
+					->AddTab(GraphTabs::Details, ETabState::OpenedTab)
+				)
+			)
+		);
+
+	GraphEditor->RegisterModeToolbarIfUnregistered(GetModeName());
+}
+
+void FGraphEditorMode::RegisterTabFactories(TSharedPtr<FTabManager> InTabManager)
+{
+	TSharedPtr<FWorkflowCentricApplication> HostingApp = HostingAppPtr.Pin();
+	HostingApp->PushTabFactories(TabFactories);
+
+	FApplicationMode::RegisterTabFactories(InTabManager);
+}
+
+void FGraphEditorMode::AddTabFactory(FCreateWorkflowTabFactory FactoryCreator)
+{
+	if (FactoryCreator.IsBound())
+	{
+		TabFactories.RegisterFactory(FactoryCreator.Execute(HostingAppPtr.Pin()));
+	}
+}
+
+void FGraphEditorMode::RemoveTabFactory(FName TabFactoryID)
+{
+	TabFactories.UnregisterFactory(TabFactoryID);
+}
+
+}
+
+#undef LOCTEXT_NAMESPACE
\ No newline at end of file
diff --git a/Engine/Plugins/Experimental/AnimNext/Source/AnimNextEditor/Private/Graph/GraphEditorMode.h b/Engine/Plugins/Experimental/AnimNext/Source/AnimNextEditor/Private/Graph/GraphEditorMode.h
new file mode 100644
--- /dev/null
+++ b/Engine/Plugins/Experimental/AnimNext/Source/AnimNextEditor/Private/Graph/GraphEditorMode.h
@@ -0,0 +1,96 @@
+﻿// Copyright Epic Games, Inc. All Rights Reserved.
+
+#pragma once
+
+#include "CoreMinimal.h"
+#include "EdGraph/EdGraph.h"
+#include "WorkflowOrientedApp/ApplicationMode.h"
+#include "WorkflowOrientedApp/WorkflowTabManager.h"
+#include "WorkflowOrientedApp/WorkflowUObjectDocuments.h"
+
+class FWorkflowCentricApplication;
+class SGraphEditor;
+
+namespace UE::AnimNext::Editor
+{
+
+DECLARE_DELEGATE_OneParam(FOnDetailsViewCreated, TSharedRef<IDetailsView>);
+
+struct FGraphEditorSummoner : public FDocumentTabFactoryForObjects<UEdGraph>
+{
+public:
+	// Delegate called when a graph editor widget is created
+	DECLARE_DELEGATE_RetVal_TwoParams(TSharedRef<SGraphEditor>, FOnCreateGraphEditorWidget, TSharedRef<FTabInfo>, UEdGraph*);
+
+	// Delegate called when a graph editor is focused
+	DECLARE_DELEGATE_OneParam(FOnGraphEditorFocused, TSharedRef<SGraphEditor>);
+
+	// Delegate called when a graph editor is backgrounded
+	DECLARE_DELEGATE_OneParam(FOnGraphEditorBackgrounded, TSharedRef<SGraphEditor>);
+
+	// Delegate called to save the state of a graph
+	DECLARE_DELEGATE_ThreeParams(FOnSaveGraphState, UEdGraph*, FVector2D, float);
+
+	// Delegate called to save the state of a graph
+	DECLARE_DELEGATE_RetVal_OneParam(const FSlateBrush*, FOnGetTabIcon, UEdGraph*);
+	
+	FGraphEditorSummoner(TSharedPtr<FWorkflowCentricApplication> InHostingApp);
+
+	FOnCreateGraphEditorWidget& OnCreateGraphEditorWidget() { return OnCreateGraphEditorWidgetDelegate; }
+	FOnGraphEditorFocused& OnGraphEditorFocused() { return OnGraphEditorFocusedDelegate; }
+	FOnGraphEditorBackgrounded& OnGraphEditorBackgrounded() { return OnGraphEditorBackgroundedDelegate; }
+	FOnSaveGraphState& OnSaveGraphState() { return OnSaveGraphStateDelegate; }
+	FOnGetTabIcon& OnGetTabIcon() { return OnGetTabIconDelegate; }
+
+private:
+	// FWorkflowTabFactory interface
+	virtual void OnTabActivated(TSharedPtr<SDockTab> Tab) const override;
+	virtual void OnTabBackgrounded(TSharedPtr<SDockTab> Tab) const override;
+	virtual void OnTabRefreshed(TSharedPtr<SDockTab> Tab) const override;
+	virtual void SaveState(TSharedPtr<SDockTab> Tab, TSharedPtr<FTabPayload> Payload) const override;
+	
+	// FDocumentTabFactoryForObjects interface
+	virtual TAttribute<FText> ConstructTabNameForObject(UEdGraph* DocumentID) const override;
+	virtual TSharedRef<SWidget> CreateTabBodyForObject(const FWorkflowTabSpawnInfo& Info, UEdGraph* DocumentID) const override;
+	virtual const FSlateBrush* GetTabIconForObject(const FWorkflowTabSpawnInfo& Info, UEdGraph* DocumentID) const override;
+	virtual TSharedRef<FGenericTabHistory> CreateTabHistoryNode(TSharedPtr<FTabPayload> Payload) override;
+
+	// The hosting app
+	TWeakPtr<FWorkflowCentricApplication> HostingAppPtr;
+
+	// Delegate called when a graph editor widget is created
+	FOnCreateGraphEditorWidget OnCreateGraphEditorWidgetDelegate;
+
+	// Delegate called when a graph editor is focused
+	FOnGraphEditorFocused OnGraphEditorFocusedDelegate;
+
+	// Delegate called when a graph editor is backgrounded
+	FOnGraphEditorBackgrounded OnGraphEditorBackgroundedDelegate;
+
+	// Delegate called to save the state of a graph
+	FOnSaveGraphState OnSaveGraphStateDelegate;
+
+	// Delegate called to get a tab icon for a graph
+	FOnGetTabIcon OnGetTabIconDelegate;
+};
+
+class FGraphEditorMode : public FApplicationMode
+{
+public:
+	FGraphEditorMode(TSharedRef<FWorkflowCentricApplication> InHostingApp);
+
+private:
+	// FApplicationMode interface
+	virtual void RegisterTabFactories(TSharedPtr<FTabManager> InTabManager) override;
+	virtual void AddTabFactory(FCreateWorkflowTabFactory FactoryCreator) override;
+	virtual void RemoveTabFactory(FName TabFactoryID) override;
+	
+	// The hosting app
+	TWeakPtr<FWorkflowCentricApplication> HostingAppPtr;
+
+	// The tab factories we support
+	FWorkflowAllowedTabSet TabFactories;
+};
+
+}
+
