AnimNext: Initial editor-side parameter blocks & libraries

diff --git a/Engine/Plugins/Experimental/AnimNext/Source/AnimNextEditor/Private/Graph/GraphEditorSchemaActions.h b/Engine/Plugins/Experimental/AnimNext/Source/AnimNextEditor/Private/Graph/GraphEditorSchemaActions.h
new file mode 100644
--- /dev/null
+++ b/Engine/Plugins/Experimental/AnimNext/Source/AnimNextEditor/Private/Graph/GraphEditorSchemaActions.h
@@ -0,0 +1,55 @@
+﻿// Copyright Epic Games, Inc. All Rights Reserved.
+
+#pragma once
+
+#include "CoreMinimal.h"
+#include "EdGraph/EdGraphSchema.h"
+#include "Styling/AppStyle.h"
+#include "GraphEditorSchemaActions.generated.h"
+
+struct FSlateBrush;
+
+USTRUCT()
+struct FAnimNextGraphSchemaAction : public FEdGraphSchemaAction
+{
+	GENERATED_BODY()
+
+	FAnimNextGraphSchemaAction() = default;
+	
+	FAnimNextGraphSchemaAction(FText InNodeCategory, FText InMenuDesc, FText InToolTip, FText InKeywords = FText::GetEmpty())
+		: FEdGraphSchemaAction(MoveTemp(InNodeCategory), MoveTemp(InMenuDesc), MoveTemp(InToolTip), 0, MoveTemp(InKeywords))
+	{
+	}
+
+	virtual const FSlateBrush* GetIconBrush() const
+	{
+		return FAppStyle::Get().GetBrush("NoBrush");
+	}
+
+	virtual const FLinearColor& GetIconColor() const
+	{
+		static const FLinearColor DefaultColor;
+		return DefaultColor;
+	}
+};
+
+USTRUCT()
+struct FAnimNextGraphSchemaAction_RigUnit : public FAnimNextGraphSchemaAction
+{
+	GENERATED_BODY()
+
+	FAnimNextGraphSchemaAction_RigUnit() = default;
+	
+	FAnimNextGraphSchemaAction_RigUnit(UScriptStruct* InStructTemplate, FText InNodeCategory, FText InMenuDesc, FText InToolTip, FText InKeywords = FText::GetEmpty())
+		: FAnimNextGraphSchemaAction(InNodeCategory, InMenuDesc, InToolTip, InKeywords)
+		, StructTemplate(InStructTemplate)
+	{}
+
+	// FEdGraphSchemaAction Interface
+	virtual UEdGraphNode* PerformAction(UEdGraph* ParentGraph, TArray<UEdGraphPin*>& FromPins, const FVector2D Location, bool bSelectNewNode = true) override;
+	virtual UEdGraphNode* PerformAction(UEdGraph* ParentGraph, UEdGraphPin* FromPin, const FVector2D Location, bool bSelectNewNode = true) { return nullptr; }
+
+private:
+	// The script struct for our rig unit
+	UScriptStruct* StructTemplate = nullptr;
+};
diff --git a/Engine/Plugins/Experimental/AnimNext/Source/AnimNextEditor/Private/Graph/GraphFactory.cpp b/Engine/Plugins/Experimental/AnimNext/Source/AnimNextEditor/Private/Graph/GraphFactory.cpp
new file mode 100644
--- /dev/null
+++ b/Engine/Plugins/Experimental/AnimNext/Source/AnimNextEditor/Private/Graph/GraphFactory.cpp
@@ -0,0 +1,97 @@
+﻿// Copyright Epic Games, Inc. All Rights Reserved.
+
+#include "GraphFactory.h"
+#include "Graph/AnimNextGraph.h"
+#include "Graph/AnimNextGraph_EditorData.h"
+#include "Graph/AnimNextGraph_EdGraph.h"
+#include "UncookedOnlyUtils.h"
+#include "Graph/RigUnit_AnimNextBeginExecution.h"
+#include "Graph/RigUnit_AnimNextEndExecution.h"
+#include "RigVMModel/RigVMController.h"
+#include "Units/RigUnit.h"
+
+UAnimNextGraphFactory::UAnimNextGraphFactory()
+{
+	bCreateNew = true;
+	bEditAfterNew = true;
+	SupportedClass = UAnimNextGraph::StaticClass();
+}
+
+bool UAnimNextGraphFactory::ConfigureProperties()
+{
+	return true;
+}
+
+UObject* UAnimNextGraphFactory::FactoryCreateNew(UClass* Class, UObject* InParent, FName Name, EObjectFlags Flags, UObject* Context, FFeedbackContext* Warn, FName CallingContext)
+{
+	UAnimNextGraph* NewGraph = NewObject<UAnimNextGraph>(InParent, Class, Name, Flags | RF_Public | RF_Standalone | RF_Transactional | RF_LoadCompleted);
+
+	// Create internal editor data & graphs
+	UAnimNextGraph_EditorData* EditorData = NewObject<UAnimNextGraph_EditorData>(NewGraph, TEXT("EditorData"));
+	NewGraph->EditorData = EditorData;
+	EditorData->Initialize(/*bRecompileVM*/false);
+	EditorData->GetRigVMClient()->SetExecuteContextStruct(FAnimNextGraphExecuteContext::StaticStruct());
+
+	// Add initial execution unit
+	URigVMController* Controller = EditorData->GetRigVMClient()->GetController(EditorData->RootGraph);
+
+	URigVMUnitNode* MainEntryPointNode = Controller->AddUnitNode(FRigUnit_AnimNextBeginExecution::StaticStruct(), FRigUnit::GetMethodName(), FVector2D(-400.0f, 0.0f), FString(), false);
+	URigVMPin* BeginExecutePin = MainEntryPointNode->FindPin(GET_MEMBER_NAME_STRING_CHECKED(FRigUnit_AnimNextBeginExecution, ExecuteContext));
+	check(BeginExecutePin);
+	check(BeginExecutePin->GetDirection() == ERigVMPinDirection::Output);
+
+	//// Add function to function lib
+	//{
+	//	FRigVMControllerGraphGuard LibraryGraphGuard(Controller, EditorData->GetRigVMClient()->GetFunctionLibrary());
+	//	EditorData->EntryPoint = Controller->AddFunctionToLibrary(UE::AnimNext::Graph::EntryPointName, true, FVector2D::ZeroVector, false, false);
+
+	//	// Add exposed pin for result to function
+	//	{
+	//		FRigVMControllerGraphGuard FunctionGraphGuard(Controller, EditorData->EntryPoint->GetContainedGraph());
+
+	//		// TODO: using float for now, but needs to be user-driven
+	//		const FString CPPType = TEXT("float");
+	//		const FName CPPTypeObjectPath = NAME_None;
+	//		const FString DefaultValue = TEXT("0.0");
+
+	//		//FunctionController->AddExposedPin(UE::AnimNext::Graph::ResultName, ERigVMPinDirection::Output, CPPType, CPPTypeObjectPath, DefaultValue, false, false);
+	//		Controller->AddExposedPin(UE::AnimNext::Graph::ResultName, ERigVMPinDirection::Output, CPPType, CPPTypeObjectPath, DefaultValue, false, false);
+	//	}
+	//}
+
+	//// Add function call
+	//URigVMFunctionReferenceNode* FunctionCallNode = Controller->AddFunctionReferenceNode(EditorData->EntryPoint, FVector2D::ZeroVector, UE::AnimNext::Graph::EntryPointName.ToString(), false, false);
+	//URigVMPin* FunctionExecutePin = FunctionCallNode->FindPin(FRigVMStruct::ExecuteContextName.ToString());
+	//check(FunctionExecutePin);
+	//check(FunctionExecutePin->GetDirection() == ERigVMPinDirection::IO);
+
+	//// Link entry point to function call
+	//Controller->AddLink(BeginExecutePin, FunctionExecutePin, false);
+
+	//URigVMPin* FunctionResultPin = FunctionCallNode->FindPin(UE::AnimNext::Graph::ResultName.ToString());
+	//check(FunctionResultPin);
+	//check(FunctionResultPin->GetDirection() == ERigVMPinDirection::Output);
+
+	//// Add end-execution unit of the correct type
+	//// TODO: using a float for now, but need to use a registry to determine correct type
+	//URigVMUnitNode* MainExitPointNode = Controller->AddUnitNode(FRigUnit_AnimNextInterfaceEndExecution_Float::StaticStruct(), FRigUnit::GetMethodName(), FVector2D(400.0f, 0.0f), FString(), false);
+	//URigVMPin* EndExecutePin = MainExitPointNode->FindPin(GET_MEMBER_NAME_STRING_CHECKED(FRigUnit_AnimNextEndExecution, ExecuteContext));
+	//check(EndExecutePin);
+	//check(EndExecutePin->GetDirection() == ERigVMPinDirection::Input);
+
+	//URigVMPin* EndExecuteResultPin = MainExitPointNode->FindPin(GET_MEMBER_NAME_STRING_CHECKED(FRigUnit_AnimNextEndExecution_Float, Result));
+	//check(EndExecuteResultPin);
+	//check(EndExecuteResultPin->GetDirection() == ERigVMPinDirection::Input);
+
+	//// Link function call to exit point 
+	//Controller->AddLink(FunctionExecutePin, EndExecutePin, false);
+
+	//// Link result pins 
+	//Controller->AddLink(FunctionResultPin, EndExecuteResultPin, false);
+
+	// Compile the initial skeleton
+	UE::AnimNext::UncookedOnly::FUtils::Compile(NewGraph);
+	check(!EditorData->bErrorsDuringCompilation);
+	
+	return NewGraph;
+}
\ No newline at end of file
diff --git a/Engine/Plugins/Experimental/AnimNext/Source/AnimNextEditor/Private/Graph/PropertyTypeCustomization.cpp b/Engine/Plugins/Experimental/AnimNext/Source/AnimNextEditor/Private/Graph/PropertyTypeCustomization.cpp
new file mode 100644
--- /dev/null
+++ b/Engine/Plugins/Experimental/AnimNext/Source/AnimNextEditor/Private/Graph/PropertyTypeCustomization.cpp
@@ -0,0 +1,88 @@
+﻿// Copyright Epic Games, Inc. All Rights Reserved.
+
+#include "PropertyTypeCustomization.h"
+
+#include "Graph/WidgetFactories.h"
+#include "DetailWidgetRow.h"
+#include "Editor/PropertyEditor/Private/PropertyNode.h"
+#include "IDetailChildrenBuilder.h"
+#include "PropertyHandle.h"
+#include "SClassViewer.h"
+#include "Widgets/Text/STextBlock.h"
+#include "EditorUtils.h"
+#include "Param/ParamTypeHandle.h"
+
+#define LOCTEXT_NAMESPACE "AnimNextPropertyTypeCustomization"
+
+namespace UE::AnimNext::Editor
+{
+
+bool FPropertyTypeIdentifier::IsPropertyTypeCustomized(const IPropertyHandle& PropertyHandle) const
+{
+	return (PropertyHandle.GetMetaData("AnimNextType") != "");
+}
+
+void FPropertyTypeCustomization::CustomizeHeader(TSharedRef<IPropertyHandle> PropertyHandle, FDetailWidgetRow& HeaderRow, IPropertyTypeCustomizationUtils& CustomizationUtils)
+{
+	const FString TypeName = PropertyHandle->GetMetaData("AnimNextType");
+	const FAnimNextParamType AnimNextType = FUtils::GetParameterTypeFromMetaData(TypeName);
+	if(AnimNextType.IsValid())
+	{
+		void* ValuePtr;
+		PropertyHandle->GetValueData(ValuePtr);
+		FScriptInterface* PropertyValue = reinterpret_cast<FScriptInterface*>(ValuePtr);
+		
+		TSharedPtr<SWidget> Widget = FWidgetFactories::CreateAnimNextInterfaceWidget(AnimNextType.GetHandle(), PropertyValue->GetObject(),
+				FOnClassPicked::CreateLambda([this, PropertyHandle](UClass* ChosenClass)
+				{
+					TArray<void*> RawData;
+					PropertyHandle->AccessRawData(RawData);
+					TArray<UPackage*> OuterObjects;
+					PropertyHandle->GetOuterPackages(OuterObjects);
+
+					for(int i=0;i<RawData.Num();i++)
+					{
+						UObject* NewValue = NewObject<UObject>(OuterObjects[i], ChosenClass);
+					
+						PropertyHandle->NotifyPreChange();
+						PropertyHandle->SetValue(NewValue);
+						PropertyHandle->NotifyPostChange(EPropertyChangeType::ValueSet);
+						PropertyHandle->GetPropertyNode()->GetParentNode()->RequestRebuildChildren();
+					}
+				}),
+				nullptr
+		);
+
+		HeaderRow
+		.NameContent()
+		[
+			PropertyHandle->CreatePropertyNameWidget()
+		]
+		.ValueContent()
+		[
+			Widget.ToSharedRef()
+		];
+	}
+}
+
+void FPropertyTypeCustomization::CustomizeChildren(TSharedRef<IPropertyHandle> PropertyHandle, IDetailChildrenBuilder& ChildBuilder, IPropertyTypeCustomizationUtils& CustomizationUtils)
+{
+	// todo: make this work correctly for multiselect
+	void* ValuePtr;
+	PropertyHandle->GetValueData(ValuePtr);
+	FScriptInterface* PropertyValue = reinterpret_cast<FScriptInterface*>(ValuePtr);
+	if (PropertyValue->GetObject())
+	{
+		TSharedPtr<FStructOnScope> StructData = MakeShareable(new FStructOnScope(PropertyValue->GetObject()->GetClass(), (uint8*)PropertyValue->GetObject()));
+		StructData->SetPackage(PropertyValue->GetObject()->GetPackage());
+		TArray<TSharedPtr<IPropertyHandle>> NewHandles = ChildBuilder.AddAllExternalStructureProperties(StructData.ToSharedRef());
+		for (auto& NewHandle : NewHandles)
+		{
+			ChildBuilder.AddProperty(NewHandle.ToSharedRef());
+		}
+	}
+}
+
+}
+
+#undef LOCTEXT_NAMESPACE
\ No newline at end of file
diff --git a/Engine/Plugins/Experimental/AnimNext/Source/AnimNextEditor/Private/Graph/PropertyTypeCustomization.h b/Engine/Plugins/Experimental/AnimNext/Source/AnimNextEditor/Private/Graph/PropertyTypeCustomization.h
new file mode 100644
--- /dev/null
+++ b/Engine/Plugins/Experimental/AnimNext/Source/AnimNextEditor/Private/Graph/PropertyTypeCustomization.h
@@ -0,0 +1,24 @@
+﻿// Copyright Epic Games, Inc. All Rights Reserved.
+
+#pragma once
+
+#include "IPropertyTypeCustomization.h"
+#include "PropertyEditorModule.h"
+
+namespace UE::AnimNext::Editor
+{
+
+class FPropertyTypeCustomization : public IPropertyTypeCustomization
+{
+private:
+	// IPropertyTypeCustomization interface
+	virtual void CustomizeHeader(TSharedRef<IPropertyHandle> PropertyHandle, FDetailWidgetRow& HeaderRow, IPropertyTypeCustomizationUtils& CustomizationUtils) override;
+	virtual void CustomizeChildren(TSharedRef<IPropertyHandle> PropertyHandle, IDetailChildrenBuilder& ChildBuilder, IPropertyTypeCustomizationUtils& CustomizationUtils) override;
+};
+
+class FPropertyTypeIdentifier : public IPropertyTypeIdentifier
+{
+	virtual bool IsPropertyTypeCustomized(const IPropertyHandle& PropertyHandle) const override;
+};
+
+}
\ No newline at end of file
diff --git a/Engine/Plugins/Experimental/AnimNext/Source/AnimNextEditor/Private/Graph/SActionMenu.cpp b/Engine/Plugins/Experimental/AnimNext/Source/AnimNextEditor/Private/Graph/SActionMenu.cpp
new file mode 100644
--- /dev/null
+++ b/Engine/Plugins/Experimental/AnimNext/Source/AnimNextEditor/Private/Graph/SActionMenu.cpp
@@ -0,0 +1,227 @@
+﻿// Copyright Epic Games, Inc. All Rights Reserved.
+
+#include "SActionMenu.h"
+
+#include "GraphEditorSchemaActions.h"
+#include "Framework/Application/SlateApplication.h"
+#include "IDocumentation.h"
+#include "SSubobjectEditor.h"
+#include "Graph/AnimNextGraph_EdGraphSchema.h"
+#include "Widgets/Images/SImage.h"
+#include "Widgets/Layout/SBox.h"
+#include "Widgets/SBoxPanel.h"
+#include "Widgets/Text/STextBlock.h"
+#include "SGraphPalette.h"
+#include "RigVMCore/RigVMRegistry.h"
+#include "Units/RigUnit.h"
+#include "Widgets/SToolTip.h"
+#include "Graph/GraphExecuteContext.h"
+
+#define LOCTEXT_NAMESPACE "AnimNextEditor"
+
+namespace UE::AnimNext::Editor
+{
+
+static void CollectAllAnimNextGraphActions(FGraphContextMenuBuilder& MenuBuilder)
+{
+	static const TArray<UScriptStruct*> AllowedExecuteContexts =
+	{
+		FRigVMExecuteContext::StaticStruct(),
+		FAnimNextGraphExecuteContext::StaticStruct()
+	};
+
+	const TChunkedArray<FRigVMFunction>& Functions = FRigVMRegistry::Get().GetFunctions();
+	const int32 NumFunctions = Functions.Num();
+	for(int i=0; i < NumFunctions; i++)
+	{
+		const FRigVMFunction& Function = Functions[i];
+
+		const UScriptStruct* FunctionContext = Function.GetExecuteContextStruct();
+		if (FunctionContext == nullptr || !AllowedExecuteContexts.Contains(FunctionContext))
+		{
+			continue;
+		}
+
+		UScriptStruct* Struct = Function.Struct;
+		if (Struct == nullptr)
+		{
+			continue;
+		}
+
+		FString CategoryMetadata, DisplayNameMetadata, MenuDescSuffixMetadata;
+		Struct->GetStringMetaDataHierarchical(FRigVMStruct::CategoryMetaName, &CategoryMetadata);
+		Struct->GetStringMetaDataHierarchical(FRigVMStruct::DisplayNameMetaName, &DisplayNameMetadata);
+		Struct->GetStringMetaDataHierarchical(FRigVMStruct::MenuDescSuffixMetaName, &MenuDescSuffixMetadata);
+		if (!MenuDescSuffixMetadata.IsEmpty())
+		{
+			MenuDescSuffixMetadata = TEXT(" ") + MenuDescSuffixMetadata;
+		}
+		const FText NodeCategory = FText::FromString(CategoryMetadata);
+		const FText MenuDesc = FText::FromString(DisplayNameMetadata + MenuDescSuffixMetadata);
+		const FText ToolTip = Struct->GetToolTipText();
+
+		if (MenuDesc.IsEmpty())
+		{
+			continue;
+		}
+
+		MenuBuilder.AddAction(MakeShared<FAnimNextGraphSchemaAction_RigUnit>(Struct, NodeCategory, MenuDesc, ToolTip));
+	};
+}
+
+SActionMenu::~SActionMenu()
+{
+	OnClosedCallback.ExecuteIfBound();
+	OnCloseReasonCallback.ExecuteIfBound(bActionExecuted, false, !DraggedFromPins.IsEmpty());
+}
+
+void SActionMenu::Construct(const FArguments& InArgs)
+{
+	Graph = InArgs._Graph;
+	DraggedFromPins = InArgs._DraggedFromPins;
+	NewNodePosition = InArgs._NewNodePosition;
+	OnClosedCallback = InArgs._OnClosedCallback;
+	bAutoExpandActionMenu = InArgs._AutoExpandActionMenu;
+	OnCloseReasonCallback = InArgs._OnCloseReason;
+
+	SBorder::Construct(SBorder::FArguments()
+		.BorderImage(FAppStyle::Get().GetBrush("Menu.Background"))
+		.Padding(5)
+		[
+			SNew(SBox)
+			.WidthOverride(400)
+			.HeightOverride(400)
+			[
+				SNew(SVerticalBox)
+				+SVerticalBox::Slot()
+				.AutoHeight()
+				.Padding(2, 2, 2, 5)
+				[
+					SNew(SHorizontalBox)
+					+SHorizontalBox::Slot()
+					.AutoWidth()
+					.VAlign(VAlign_Center)
+					[
+						SNew(STextBlock)
+						.Text(LOCTEXT("ContextText", "All AnimNext Interface Node Classes"))
+						.Font(FAppStyle::Get().GetFontStyle("BlueprintEditor.ActionMenu.ContextDescriptionFont"))
+						.ToolTip(IDocumentation::Get()->CreateToolTip(
+							LOCTEXT("ActionMenuContextTextTooltip", "Describes the current context of the action list"),
+							nullptr,
+							TEXT("Shared/Editors/AnimNextInterfaceEditor"),
+							TEXT("AnimNextInterfaceActionMenuContextText")))
+						.WrapTextAt(280)
+					]
+				]
+				+SVerticalBox::Slot()
+				[
+					SAssignNew(GraphActionMenu, SGraphActionMenu)
+						.OnActionSelected(this, &SActionMenu::OnActionSelected)
+						.OnCreateWidgetForAction(SGraphActionMenu::FOnCreateWidgetForAction::CreateSP(this, &SActionMenu::OnCreateWidgetForAction))
+						.OnCollectAllActions(this, &SActionMenu::CollectAllActions)
+						.DraggedFromPins(DraggedFromPins)
+						.GraphObj(Graph)
+						.AlphaSortItems(true)
+						.bAllowPreselectedItemActivation(true)
+				]
+			]
+		]
+	);
+}
+
+void SActionMenu::CollectAllActions(FGraphActionListBuilderBase& OutAllActions)
+{
+	if (!Graph)
+	{
+		return;
+	}
+
+	FGraphContextMenuBuilder MenuBuilder(Graph);
+	if (!DraggedFromPins.IsEmpty())
+	{
+		MenuBuilder.FromPin = DraggedFromPins[0];
+	}
+
+	// Cannot call GetGraphContextActions() during serialization and GC due to its use of FindObject()
+	if(!GIsSavingPackage && !IsGarbageCollecting())
+	{
+		CollectAllAnimNextGraphActions(MenuBuilder);
+	}
+
+	OutAllActions.Append(MenuBuilder);
+}
+
+TSharedRef<SEditableTextBox> SActionMenu::GetFilterTextBox()
+{
+	return GraphActionMenu->GetFilterTextBox();
+}
+
+TSharedRef<SWidget> SActionMenu::OnCreateWidgetForAction(FCreateWidgetForActionData* const InCreateData)
+{
+	check(InCreateData);
+	InCreateData->bHandleMouseButtonDown = false;
+
+	const FSlateBrush* IconBrush = nullptr;
+	FLinearColor IconColor;
+	TSharedPtr<FAnimNextGraphSchemaAction> Action = StaticCastSharedPtr<FAnimNextGraphSchemaAction>(InCreateData->Action);
+	if (Action.IsValid())
+	{
+		IconBrush = Action->GetIconBrush();
+		IconColor = Action->GetIconColor();
+	}
+
+	TSharedPtr<SHorizontalBox> WidgetBox = SNew(SHorizontalBox);
+	if (IconBrush)
+	{
+		WidgetBox->AddSlot()
+			.AutoWidth()
+			.VAlign(VAlign_Center)
+			.Padding(5, 0, 0, 0)
+			[
+				SNew(SImage)
+				.ColorAndOpacity(IconColor)
+				.Image(IconBrush)
+			];
+	}
+
+	WidgetBox->AddSlot()
+		.AutoWidth()
+		.VAlign(VAlign_Center)
+		.Padding(0, 0, 0, 0)
+		[
+			SNew(SGraphPaletteItem, InCreateData)
+		];
+
+	return WidgetBox->AsShared();
+}
+
+void SActionMenu::OnActionSelected(const TArray<TSharedPtr<FEdGraphSchemaAction>>& SelectedAction, ESelectInfo::Type InSelectionType)
+{
+	if (!Graph)
+	{
+		return;
+	}
+
+	if (InSelectionType != ESelectInfo::OnMouseClick  && InSelectionType != ESelectInfo::OnKeyPress && !SelectedAction.IsEmpty())
+	{
+		return;
+	}
+
+	for (const TSharedPtr<FEdGraphSchemaAction>& Action : SelectedAction)
+	{
+		if (Action.IsValid() && Graph)
+		{
+			if (!bActionExecuted && (Action->GetTypeId() != FEdGraphSchemaAction_Dummy::StaticGetTypeId()))
+			{
+				FSlateApplication::Get().DismissAllMenus();
+				bActionExecuted = true;
+			}
+
+			Action->PerformAction(Graph, DraggedFromPins, NewNodePosition);
+		}
+	}
+}
+
+} 
+
+#undef LOCTEXT_NAMESPACE
diff --git a/Engine/Plugins/Experimental/AnimNext/Source/AnimNextEditor/Private/Graph/WidgetFactories.cpp b/Engine/Plugins/Experimental/AnimNext/Source/AnimNextEditor/Private/Graph/WidgetFactories.cpp
new file mode 100644
--- /dev/null
+++ b/Engine/Plugins/Experimental/AnimNext/Source/AnimNextEditor/Private/Graph/WidgetFactories.cpp
@@ -0,0 +1,118 @@
+// Copyright Epic Games, Inc. All Rights Reserved.
+
+#include "Graph/WidgetFactories.h"
+#include "Modules/ModuleManager.h"
+#include "Widgets/Input/SComboButton.h"
+#include "SClassViewer.h"
+#include "Graph/ClassFilter.h"
+
+#define LOCTEXT_NAMESPACE "AnimNextEditor"
+
+namespace UE::AnimNext::Editor
+{
+
+void ConvertToText(UObject* Object, FText& OutText)
+{
+	UClass* Class = Object->GetClass();
+	while (Class)
+	{
+		if (auto TextConverter = FWidgetFactories::AnimNextInterfaceTextConverter.Find(Class))
+		{
+			(*TextConverter)(Object, OutText);
+			break;
+		}
+		Class = Class->GetSuperClass();
+	}
+}
+
+TSharedPtr<SWidget> FWidgetFactories::CreateAnimNextInterfaceWidget(FParamTypeHandle TypeHandle, UObject* Value, const FOnClassPicked& CreateClassCallback, TSharedPtr<SBorder>* InnerWidget)
+{
+	TSharedPtr<SWidget> LeftWidget;
+
+	if (Value)
+	{
+		UClass* Class = Value->GetClass();
+		while (Class && !LeftWidget.IsValid())
+		{
+			if (auto Creator = FWidgetFactories::AnimNextInterfaceWidgetCreators.Find(Class))
+			{
+				LeftWidget = (*Creator)(Value);
+				break;
+			}
+			Class = Class->GetSuperClass();
+		}
+	}
+	
+	if (!LeftWidget.IsValid())
+	{
+		LeftWidget = SNew(STextBlock).Text(LOCTEXT("SelectDataType", "Select Data Type..."));
+	}
+
+
+	// button for replacing data with a different AnimNext Interface class
+	TSharedPtr<SComboButton> Button = SNew(SComboButton)
+			.ComboButtonStyle(FAppStyle::Get(), "SimpleComboButton");
+	
+	Button->SetOnGetMenuContent(FOnGetContent::CreateLambda([TypeHandle, Button, CreateClassCallback]()
+	{
+		FClassViewerInitializationOptions Options;
+		Options.ClassFilters.Add(MakeShared<FClassFilter>(TypeHandle));
+		
+		// Add class filter for columns here
+		TSharedRef<SWidget>  ClassMenu = FModuleManager::LoadModuleChecked<FClassViewerModule>("ClassViewer").
+		CreateClassViewer(Options, FOnClassPicked::CreateLambda([Button, CreateClassCallback](UClass* Class)
+		{
+			Button->SetIsOpen(false);
+			CreateClassCallback.Execute(Class);
+		}));
+		return ClassMenu;
+	}));
+
+	TSharedPtr <SBorder> Border;
+	if (InnerWidget && InnerWidget->IsValid())
+	{
+		Border = *InnerWidget;
+	}
+	else
+	{
+		Border = SNew(SBorder);
+	}
+	
+	if (InnerWidget)
+	{
+		*InnerWidget = Border;
+	}
+	
+	Border->SetContent(LeftWidget.ToSharedRef());
+
+	TSharedPtr<SWidget> Widget = SNew(SHorizontalBox)
+		+SHorizontalBox::Slot().FillWidth(75)
+		[
+			Border.ToSharedRef()
+		]
+		+SHorizontalBox::Slot().FillWidth(25)
+		[
+			Button.ToSharedRef()
+		]
+	;
+
+
+	return Widget;
+}
+
+TMap<const UClass*, TFunction<void (const UObject* Object, FText& OutText)>> FWidgetFactories::AnimNextInterfaceTextConverter;
+TMap<const UClass*, TFunction<TSharedRef<SWidget> (UObject* Object)>> FWidgetFactories::AnimNextInterfaceWidgetCreators;
+
+void ConvertToText_Base(const UObject* Object, FText& OutText)
+{
+	OutText = FText::FromString(Object->GetName());
+}
+
+void FWidgetFactories::RegisterWidgets()
+{
+	AnimNextInterfaceTextConverter.Add(UObject::StaticClass(), ConvertToText_Base);
+}
+	
+}
+
+#undef LOCTEXT_NAMESPACE
diff --git a/Engine/Plugins/Experimental/AnimNext/Source/AnimNextEditor/Private/Graph/WidgetFactories.h b/Engine/Plugins/Experimental/AnimNext/Source/AnimNextEditor/Private/Graph/WidgetFactories.h
new file mode 100644
--- /dev/null
+++ b/Engine/Plugins/Experimental/AnimNext/Source/AnimNextEditor/Private/Graph/WidgetFactories.h
@@ -0,0 +1,27 @@
+// Copyright Epic Games, Inc. All Rights Reserved.
+
+#pragma once
+
+#include "CoreMinimal.h"
+#include "Widgets/SWidget.h"
+#include "Widgets/Layout/SBorder.h"
+#include "SClassViewer.h"
+
+namespace UE::AnimNext
+{
+	struct FParamTypeHandle;
+}
+
+namespace UE::AnimNext::Editor
+{
+	class FWidgetFactories
+	{
+	public:
+		static TMap<const UClass*, TFunction<void (const UObject* Object, FText& OutText)>> AnimNextInterfaceTextConverter;
+		static TMap<const UClass*, TFunction<TSharedRef<SWidget> (UObject* Object)>> AnimNextInterfaceWidgetCreators;
+
+		static TSharedPtr<SWidget> CreateAnimNextInterfaceWidget(FParamTypeHandle TypeHandle, UObject* Value, const FOnClassPicked& CreateClassCallback, TSharedPtr<SBorder>* InnerWidget = nullptr);
+		
+		static void RegisterWidgets();
+	};
+}
diff --git a/Engine/Plugins/Experimental/AnimNext/Source/AnimNextEditor/Private/Param/AddParameterDialogMenuContext.h b/Engine/Plugins/Experimental/AnimNext/Source/AnimNextEditor/Private/Param/AddParameterDialogMenuContext.h
new file mode 100644
--- /dev/null
+++ b/Engine/Plugins/Experimental/AnimNext/Source/AnimNextEditor/Private/Param/AddParameterDialogMenuContext.h
@@ -0,0 +1,25 @@
+﻿// Copyright Epic Games, Inc. All Rights Reserved.
+
+#pragma once
+
+#include "AddParameterDialogMenuContext.generated.h"
+
+namespace UE::AnimNext::Editor
+{
+	class SAddParametersDialog;
+	struct FParameterToAdd;
+}
+
+UCLASS()
+class UAddParameterDialogMenuContext : public UObject
+{
+	GENERATED_BODY()
+
+	friend class UE::AnimNext::Editor::SAddParametersDialog;
+
+	// The add parameter dialog that we are editing
+	TWeakPtr<UE::AnimNext::Editor::SAddParametersDialog> AddParametersDialog;
+
+	// The entry that we are editing
+	TWeakPtr<UE::AnimNext::Editor::FParameterToAdd> Entry;
+};
\ No newline at end of file
diff --git a/Engine/Plugins/Experimental/AnimNext/Source/AnimNextEditor/Private/Param/AnimNextParameterBlockEditor.cpp b/Engine/Plugins/Experimental/AnimNext/Source/AnimNextEditor/Private/Param/AnimNextParameterBlockEditor.cpp
new file mode 100644
--- /dev/null
+++ b/Engine/Plugins/Experimental/AnimNext/Source/AnimNextEditor/Private/Param/AnimNextParameterBlockEditor.cpp
@@ -0,0 +1,378 @@
+﻿// Copyright Epic Games, Inc. All Rights Reserved.
+
+#include "AnimNextParameterBlockEditor.h"
+#include "ParameterBlockEditorMode.h"
+#include "Param/AnimNextParameterBlock.h"
+#include "Param/AnimNextParameterBlock_EditorData.h"
+#include "EdGraphNode_Comment.h"
+#include "ExternalPackageHelper.h"
+#include "Graph/SActionMenu.h"
+#include "UncookedOnlyUtils.h"
+#include "Framework/Commands/GenericCommands.h"
+#include "RigVMModel/RigVMController.h"
+#include "Kismet2/BlueprintEditorUtils.h"
+#include "ScopedTransaction.h"
+#include "Param/AnimNextParameterBlockEntry.h"
+#include "Param/AnimNextParameterSettings.h"
+#include "ParameterBlockEditorMode.h"
+
+#define LOCTEXT_NAMESPACE "AnimNextParameterBlockEditor"
+
+namespace UE::AnimNext::Editor
+{
+
+namespace ParameterBlockModes
+{
+	const FName ParameterBlockEditor("AnimNextParameterBlockEditorMode");
+}
+
+namespace ParameterBlockTabs
+{
+	const FName Details("DetailsTab");
+	const FName Parameters("ParametersTab");
+	const FName Document("Document");
+}
+
+const FName ParameterBlockAppIdentifier("AnimNextParameterBlockEditor");
+
+FParameterBlockEditor::FParameterBlockEditor()
+{
+	UAnimNextParameterSettings* Settings = GetMutableDefault<UAnimNextParameterSettings>();
+	Settings->LoadConfig();
+}
+
+FParameterBlockEditor::~FParameterBlockEditor()
+{
+	UAnimNextParameterSettings* Settings = GetMutableDefault<UAnimNextParameterSettings>();
+	Settings->SaveConfig();
+}
+
+void FParameterBlockEditor::InitEditor(const EToolkitMode::Type InMode, const TSharedPtr<IToolkitHost>& InInitToolkitHost, UAnimNextParameterBlock* InAnimNextParameterBlock)
+{
+	ParameterBlock = InAnimNextParameterBlock;
+
+	EditorData = UncookedOnly::FUtils::GetEditorData(ParameterBlock);
+	EditorData->Initialize(false);
+
+	DocumentManager = MakeShared<FDocumentTracker>();
+	DocumentManager->Initialize(SharedThis(this));
+
+	TSharedRef<FParameterBlockEditorDocumentSummoner> DocumentSummoner = MakeShared<FParameterBlockEditorDocumentSummoner>(SharedThis(this));
+	DocumentSummoner->OnCreateGraphEditorWidget().BindSP(this, &FParameterBlockEditor::CreateGraphEditorWidget);
+	DocumentSummoner->OnGraphEditorFocused().BindSP(this, &FParameterBlockEditor::OnGraphEditorFocused);
+	DocumentSummoner->OnGraphEditorBackgrounded().BindSP(this, &FParameterBlockEditor::OnGraphEditorBackgrounded);
+	GraphEditorTabFactoryPtr = DocumentSummoner;
+	DocumentManager->RegisterDocumentFactory(DocumentSummoner);
+
+	const bool bCreateDefaultStandaloneMenu = true;
+	const bool bCreateDefaultToolbar = true;
+	InitAssetEditor(InMode, InInitToolkitHost, ParameterBlockAppIdentifier, FTabManager::FLayout::NullLayout, bCreateDefaultStandaloneMenu, bCreateDefaultToolbar, InAnimNextParameterBlock);
+
+	BindCommands();
+
+	AddApplicationMode(ParameterBlockModes::ParameterBlockEditor, MakeShared<FParameterBlockEditorMode>(SharedThis(this)));
+	SetCurrentMode(ParameterBlockModes::ParameterBlockEditor);
+	
+	ExtendMenu();
+	ExtendToolbar();
+	RegenerateMenusAndToolbars();
+	RestoreLastEditedState();
+}
+
+void FParameterBlockEditor::RestoreLastEditedState()
+{
+	for (const FEditedDocumentInfo& Document : EditorData->LastEditedDocuments)
+	{
+		if (UObject* Obj = Document.EditedObjectPath.ResolveObject())
+		{
+			OpenDocument(Obj, FDocumentTracker::RestorePreviousDocument);
+		}
+	}
+}
+
+TSharedPtr<SDockTab> FParameterBlockEditor::OpenDocument(const UObject* InForObject, FDocumentTracker::EOpenDocumentCause InCause)
+{
+	if(InCause != FDocumentTracker::RestorePreviousDocument)
+	{
+		EditorData->LastEditedDocuments.AddUnique(const_cast<UObject*>(InForObject));
+	}
+
+	TSharedRef<FTabPayload_UObject> Payload = FTabPayload_UObject::Make(InForObject);
+	return DocumentManager->OpenDocument(Payload, InCause);
+}
+
+void FParameterBlockEditor::BindCommands()
+{
+}
+
+void FParameterBlockEditor::ExtendMenu()
+{
+	
+}
+
+void FParameterBlockEditor::ExtendToolbar()
+{
+	
+}
+
+void FParameterBlockEditor::RegisterTabSpawners(const TSharedRef<FTabManager>& InTabManager)
+{
+	DocumentManager->SetTabManager(InTabManager);
+
+	FWorkflowCentricApplication::RegisterTabSpawners(InTabManager);
+}
+
+void FParameterBlockEditor::UnregisterTabSpawners(const TSharedRef<FTabManager>& InTabManager)
+{
+	
+}
+
+FName FParameterBlockEditor::GetToolkitFName() const
+{
+	return FName("AnimNextParameterBlockEditor");
+}
+
+FText FParameterBlockEditor::GetBaseToolkitName() const
+{
+	return LOCTEXT("AppLabel", "AnimNextParameterBlockEditor");
+}
+
+FString FParameterBlockEditor::GetWorldCentricTabPrefix() const
+{
+	return LOCTEXT("WorldCentricTabPrefix", "AnimNextParameterBlockEditor ").ToString();
+}
+
+FLinearColor FParameterBlockEditor::GetWorldCentricTabColorScale() const
+{
+	return FLinearColor(0.3f, 0.2f, 0.5f, 0.5f);
+}
+
+void FParameterBlockEditor::InitToolMenuContext(FToolMenuContext& InMenuContext)
+{
+}
+
+void FParameterBlockEditor::OnGraphEditorFocused(TSharedRef<SGraphEditor> InGraphEditor)
+{
+	// Update the graph editor that is currently focused
+	FocusedGraphEdPtr = InGraphEditor;
+}
+
+void FParameterBlockEditor::OnGraphEditorBackgrounded(TSharedRef<SGraphEditor> InGraphEditor)
+{
+	// Update the graph editor that is currently focused
+	FocusedGraphEdPtr = nullptr;
+}
+
+UEdGraph* FParameterBlockEditor::GetFocusedGraph() const
+{
+	if (FocusedGraphEdPtr.IsValid())
+	{
+		if (UEdGraph* Graph = FocusedGraphEdPtr.Pin()->GetCurrentGraph())
+		{
+			return Graph;
+		}
+	}
+	return nullptr;
+}
+
+URigVMGraph* FParameterBlockEditor::GetFocusedVMGraph() const
+{
+	UAnimNextParameterBlock_EdGraph* EdGraph = Cast<UAnimNextParameterBlock_EdGraph>(GetFocusedGraph());
+	return EditorData->GetVMGraphForEdGraph(EdGraph);
+}
+
+URigVMController* FParameterBlockEditor::GetFocusedVMController() const
+{
+	return EditorData->GetRigVMClient()->GetController(GetFocusedVMGraph());
+}
+
+TSharedRef<SGraphEditor> FParameterBlockEditor::CreateGraphEditorWidget(TSharedRef<FTabInfo> InTabInfo, UEdGraph* InGraph)
+{
+	if(!GraphEditorCommands.IsValid())
+	{
+		GraphEditorCommands = MakeShared<FUICommandList>();
+	
+		GraphEditorCommands->MapAction(FGenericCommands::Get().Delete,
+			FExecuteAction::CreateSP(this, &FParameterBlockEditor::DeleteSelectedNodes),
+			FCanExecuteAction::CreateSP(this, &FParameterBlockEditor::CanDeleteSelectedNodes));
+	}
+	
+	SGraphEditor::FGraphEditorEvents Events;
+	Events.OnCreateActionMenu = SGraphEditor::FOnCreateActionMenu::CreateSP(this, &FParameterBlockEditor::OnCreateGraphActionMenu);
+	
+	return SNew(SGraphEditor)
+		.AdditionalCommands(GraphEditorCommands)
+		.IsEditable(this, &FParameterBlockEditor::IsEditable, InGraph)
+		.GraphToEdit(InGraph)
+		.GraphEvents(Events)
+		.AssetEditorToolkit(AsShared());
+}
+
+FActionMenuContent FParameterBlockEditor::OnCreateGraphActionMenu(UEdGraph* InGraph, const FVector2D& InNodePosition, const TArray<UEdGraphPin*>& InDraggedPins, bool bAutoExpand, SGraphEditor::FActionMenuClosed InOnMenuClosed)
+{
+	TSharedRef<SActionMenu> ActionMenu = SNew(SActionMenu)
+		.AutoExpandActionMenu(bAutoExpand)
+		.Graph(InGraph)
+		.NewNodePosition(InNodePosition)
+		.DraggedFromPins(InDraggedPins)
+		.OnClosedCallback(InOnMenuClosed);
+
+	TSharedPtr<SWidget> FilterTextBox = StaticCastSharedRef<SWidget>(ActionMenu->GetFilterTextBox());
+	return FActionMenuContent(StaticCastSharedRef<SWidget>(ActionMenu), FilterTextBox);
+}
+
+void FParameterBlockEditor::OnNodeTitleCommitted(const FText& NewText, ETextCommit::Type CommitInfo, UEdGraphNode* NodeBeingChanged)
+{
+	if (UEdGraphNode_Comment* CommentBeingChanged = Cast<UEdGraphNode_Comment>(NodeBeingChanged))
+	{
+		GetFocusedVMController()->SetCommentTextByName(CommentBeingChanged->GetFName(), NewText.ToString(), CommentBeingChanged->FontSize, CommentBeingChanged->bCommentBubbleVisible, CommentBeingChanged->bColorCommentBubble, true, true);
+	}
+}
+
+void FParameterBlockEditor::CloseDocumentTab(const UObject* DocumentID)
+{
+	TSharedRef<FTabPayload_UObject> Payload = FTabPayload_UObject::Make(DocumentID);
+	DocumentManager->CloseTab(Payload);
+}
+
+bool FParameterBlockEditor::InEditingMode() const
+{
+	// @TODO: disallow editing when debugging when implemented
+
+	return true;
+}
+
+bool FParameterBlockEditor::IsEditable(UEdGraph* InGraph) const
+{
+	return InGraph && InEditingMode() && InGraph->bEditable;
+}
+
+FGraphPanelSelectionSet FParameterBlockEditor::GetSelectedNodes() const
+{
+	FGraphPanelSelectionSet CurrentSelection;
+	TSharedPtr<SGraphEditor> FocusedGraphEd = FocusedGraphEdPtr.Pin();
+	if (FocusedGraphEd.IsValid())
+	{
+		CurrentSelection = FocusedGraphEd->GetSelectedNodes();
+	}
+	return CurrentSelection;
+}
+
+void FParameterBlockEditor::DeleteSelectedNodes()
+{
+	TSharedPtr<SGraphEditor> FocusedGraphEd = FocusedGraphEdPtr.Pin();
+	if (!FocusedGraphEd.IsValid())
+	{
+		return;
+	}
+
+	const FScopedTransaction Transaction(FGenericCommands::Get().Delete->GetDescription());
+	FocusedGraphEd->GetCurrentGraph()->Modify();
+	
+	const FGraphPanelSelectionSet SelectedNodes = GetSelectedNodes();
+
+//	SetUISelectionState(NAME_None);
+
+	if(FocusedGraphEd)
+	{
+		FocusedGraphEd->ClearSelectionSet();
+	}
+
+	// Some nodes have sub-objects that are represented as other tabs.
+	// Close them here as a pre-pass before we remove their nodes. If the documents are left open they
+	// may reference dangling data and function incorrectly in cases such as FindBlueprintforNodeChecked
+	for (FGraphPanelSelectionSet::TConstIterator NodeIt( SelectedNodes ); NodeIt; ++NodeIt)
+	{
+		if (UEdGraphNode* Node = Cast<UEdGraphNode>(*NodeIt))
+		{
+			if (Node->CanUserDeleteNode())
+			{
+				auto CloseAllDocumentsTab = [this](const UEdGraphNode* InNode)
+				{
+					TArray<UObject*> NodesToClose;
+					GetObjectsWithOuter(InNode, NodesToClose);
+					for (UObject* Node : NodesToClose)
+					{
+						UEdGraph* NodeGraph = Cast<UEdGraph>(Node);
+						if (NodeGraph)
+						{
+							CloseDocumentTab(NodeGraph);
+						}
+					}
+				};
+				
+				if (Node->GetSubGraphs().Num() > 0)
+				{
+					CloseAllDocumentsTab(Node);
+				}
+			}
+		}
+	}
+
+	// Now remove the selected nodes
+	for (FGraphPanelSelectionSet::TConstIterator NodeIt( SelectedNodes ); NodeIt; ++NodeIt)
+	{
+		if (UEdGraphNode* Node = Cast<UEdGraphNode>(*NodeIt))
+		{
+			if (Node->CanUserDeleteNode())
+			{
+				if (Node->GetSubGraphs().Num() > 0)
+				{
+					DocumentManager->CleanInvalidTabs();
+				}
+
+				FBlueprintEditorUtils::RemoveNode(nullptr, Node);
+			}
+		}
+	}
+}
+
+bool FParameterBlockEditor::CanDeleteSelectedNodes()
+{
+	const FGraphPanelSelectionSet SelectedNodes = GetSelectedNodes();
+
+	bool bCanUserDeleteNode = false;
+
+	if(IsEditable(GetFocusedGraph()) && SelectedNodes.Num() > 0)
+	{
+		for(UObject* NodeObject : SelectedNodes)
+		{
+			// If any nodes allow deleting, then do not disable the delete option
+			UEdGraphNode* Node = Cast<UEdGraphNode>(NodeObject);
+			if(Node->CanUserDeleteNode())
+			{
+				bCanUserDeleteNode = true;
+				break;
+			}
+		}
+	}
+
+	return bCanUserDeleteNode;
+}
+
+void FParameterBlockEditor::SetSelectedObjects(TArray<UObject*> InObjects)
+{
+	if(DetailsView.IsValid())
+	{
+		DetailsView->SetObjects(InObjects);
+	}
+}
+
+void FParameterBlockEditor::GetSaveableObjects(TArray<UObject*>& OutObjects) const
+{
+	// Base class will pick up edited object
+	FWorkflowCentricApplication::GetSaveableObjects(OutObjects);
+
+	// Get external objects too
+	FExternalPackageHelper::GetExternalSaveableObjects(EditorData, OutObjects);
+
+	// Get any referenced objects too for convenience 
+	for(UAnimNextParameterBlockEntry* Entry : EditorData->Entries)
+	{
+		Entry->GetEditedObjects(OutObjects);
+	}
+}
+
+}
+
+#undef LOCTEXT_NAMESPACE
\ No newline at end of file
diff --git a/Engine/Plugins/Experimental/AnimNext/Source/AnimNextEditor/Private/Param/AnimNextParameterBlockEditor.h b/Engine/Plugins/Experimental/AnimNext/Source/AnimNextEditor/Private/Param/AnimNextParameterBlockEditor.h
new file mode 100644
--- /dev/null
+++ b/Engine/Plugins/Experimental/AnimNext/Source/AnimNextEditor/Private/Param/AnimNextParameterBlockEditor.h
@@ -0,0 +1,137 @@
+﻿// Copyright Epic Games, Inc. All Rights Reserved.
+
+#pragma once
+
+#include "CoreMinimal.h"
+#include "WorkflowOrientedApp/WorkflowCentricApplication.h"
+#include "GraphEditor.h"
+#include "WorkflowOrientedApp/WorkflowTabManager.h"
+
+enum class ERigVMGraphNotifType : uint8;
+class UAnimNextParameterBlock;
+class UAnimNextParameterBlock_EditorData;
+class FDocumentTracker;
+class FDocumentTabFactory;
+class FTabInfo;
+class FTabManager;
+class IToolkitHost;
+class URigVMGraph;
+class URigVMController;
+class UDetailsViewWrapperObject;
+
+namespace UE::AnimNext::Editor
+{
+
+class FParametersTabSummoner;
+class FParametersEditorMode;
+
+namespace ParameterBlockModes
+{
+	extern const FName ParameterBlockEditor;
+}
+
+namespace ParameterBlockTabs
+{
+	extern const FName Details;
+	extern const FName Parameters;
+	extern const FName Document;
+}
+
+class FParameterBlockEditor : public FWorkflowCentricApplication
+{
+public:
+	FParameterBlockEditor();
+	virtual ~FParameterBlockEditor() override;
+
+	/** Edits the specified asset */
+	void InitEditor(const EToolkitMode::Type InMode, const TSharedPtr<IToolkitHost>& InInitToolkitHost, UAnimNextParameterBlock* InAnimNextParameterBlock);
+
+private:
+	friend class FParameterBlockEditorMode;
+	friend class FParameterBlockTabSummoner;
+
+	// IToolkit interface
+	virtual void RegisterTabSpawners(const TSharedRef<FTabManager>& InTabManager) override;
+	virtual void UnregisterTabSpawners(const TSharedRef<FTabManager>& InTabManager) override;
+	virtual FName GetToolkitFName() const override;
+	virtual FText GetBaseToolkitName() const override;
+	virtual FString GetWorldCentricTabPrefix() const override;
+	virtual FLinearColor GetWorldCentricTabColorScale() const override;
+	virtual void InitToolMenuContext(FToolMenuContext& InMenuContext) override;
+
+	// FAssetEditorToolkit interface
+	virtual void GetSaveableObjects(TArray<UObject*>& OutObjects) const override;
+
+	void BindCommands();
+	
+	void ExtendMenu();
+
+	void ExtendToolbar();
+
+	void OnGraphEditorFocused(TSharedRef<SGraphEditor> InGraphEditor);
+
+	void OnGraphEditorBackgrounded(TSharedRef<SGraphEditor> InGraphEditor);
+	
+	UEdGraph* GetFocusedGraph() const;
+
+	URigVMGraph* GetFocusedVMGraph() const;
+
+	URigVMController* GetFocusedVMController() const;
+
+	// Create new graph editor widget for the supplied document container
+	TSharedRef<SGraphEditor> CreateGraphEditorWidget(TSharedRef<FTabInfo> InTabInfo, UEdGraph* InGraph);
+
+	// Create the action menu for the specified graph
+	FActionMenuContent OnCreateGraphActionMenu(UEdGraph* InGraph, const FVector2D& InNodePosition, const TArray<UEdGraphPin*>& InDraggedPins, bool bAutoExpand, SGraphEditor::FActionMenuClosed InOnMenuClosed);
+
+	// Handle node titles being changed
+	void OnNodeTitleCommitted(const FText& NewText, ETextCommit::Type CommitInfo, UEdGraphNode* NodeBeingChanged);
+
+	void CloseDocumentTab(const UObject* DocumentID);
+
+	bool InEditingMode() const;
+
+	bool IsEditable(UEdGraph* InGraph) const;
+	
+	FGraphPanelSelectionSet GetSelectedNodes() const;
+	
+	void DeleteSelectedNodes();
+	
+	bool CanDeleteSelectedNodes();
+
+	void SetSelectedObjects(TArray<UObject*> InObjects);
+
+	void HandleDetailsViewCreated(TSharedRef<IDetailsView> InDetailsView)
+	{
+		DetailsView = InDetailsView;
+	}
+
+	void RestoreLastEditedState();
+
+	TSharedPtr<SDockTab> OpenDocument(const UObject* InForObject, FDocumentTracker::EOpenDocumentCause InCause);
+
+	// The asset we are editing
+	UAnimNextParameterBlock* ParameterBlock = nullptr;
+
+	// Cached editor data ptr
+	UAnimNextParameterBlock_EditorData* EditorData = nullptr;
+
+	// Currently focused graph editor
+	TWeakPtr<SGraphEditor> FocusedGraphEdPtr;
+
+	// Document tracker
+	TSharedPtr<FDocumentTracker> DocumentManager;
+
+	// Factory for document tabs
+	TWeakPtr<FDocumentTabFactory> GraphEditorTabFactoryPtr;
+
+	// Command list for the graph editor
+	TSharedPtr<FUICommandList> GraphEditorCommands;
+
+	// Our details panel
+	TSharedPtr<IDetailsView> DetailsView;
+
+	TArray<TStrongObjectPtr<UDetailsViewWrapperObject>> WrapperObjects;
+};
+
+}
\ No newline at end of file
diff --git a/Engine/Plugins/Experimental/AnimNext/Source/AnimNextEditor/Private/Param/AnimNextParameterLibraryEditor.cpp b/Engine/Plugins/Experimental/AnimNext/Source/AnimNextEditor/Private/Param/AnimNextParameterLibraryEditor.cpp
new file mode 100644
--- /dev/null
+++ b/Engine/Plugins/Experimental/AnimNext/Source/AnimNextEditor/Private/Param/AnimNextParameterLibraryEditor.cpp
@@ -0,0 +1,118 @@
+﻿// Copyright Epic Games, Inc. All Rights Reserved.
+
+#include "AnimNextParameterLibraryEditor.h"
+#include "ParameterLibraryEditorMode.h"
+#include "Param/AnimNextParameterLibrary.h"
+#include "ExternalPackageHelper.h"
+#include "Graph/SActionMenu.h"
+#include "UncookedOnlyUtils.h"
+#include "Framework/Commands/GenericCommands.h"
+#include "ScopedTransaction.h"
+#include "Param/AnimNextParameterSettings.h"
+
+#define LOCTEXT_NAMESPACE "AnimNextParameterLibraryEditor"
+
+namespace UE::AnimNext::Editor
+{
+
+namespace ParameterLibraryModes
+{
+	const FName ParameterLibraryEditor("AnimNextParameterLibraryEditorMode");
+}
+
+namespace ParameterLibraryTabs
+{
+	const FName Details("DetailsTab");
+	const FName Parameters("ParametersTab");
+}
+
+const FName ParameterLibraryAppIdentifier("AnimNextParameterLibraryEditor");
+
+FParameterLibraryEditor::FParameterLibraryEditor()
+{
+	UAnimNextParameterSettings* Settings = GetMutableDefault<UAnimNextParameterSettings>();
+	Settings->LoadConfig();
+}
+
+FParameterLibraryEditor::~FParameterLibraryEditor()
+{
+	UAnimNextParameterSettings* Settings = GetMutableDefault<UAnimNextParameterSettings>();
+	Settings->SaveConfig();
+}
+
+void FParameterLibraryEditor::InitEditor(const EToolkitMode::Type InMode, const TSharedPtr<IToolkitHost>& InInitToolkitHost, UAnimNextParameterLibrary* InAnimNextParameterLibrary)
+{
+	ParameterLibrary = InAnimNextParameterLibrary;
+
+	const bool bCreateDefaultStandaloneMenu = true;
+	const bool bCreateDefaultToolbar = true;
+	InitAssetEditor(InMode, InInitToolkitHost, ParameterLibraryAppIdentifier, FTabManager::FLayout::NullLayout, bCreateDefaultStandaloneMenu, bCreateDefaultToolbar, InAnimNextParameterLibrary);
+
+	BindCommands();
+
+	AddApplicationMode(ParameterLibraryModes::ParameterLibraryEditor, MakeShared<FParameterLibraryEditorMode>(SharedThis(this)));
+	SetCurrentMode(ParameterLibraryModes::ParameterLibraryEditor);
+
+	ExtendMenu();
+	ExtendToolbar();
+	RegenerateMenusAndToolbars();
+}
+
+void FParameterLibraryEditor::BindCommands()
+{
+}
+
+void FParameterLibraryEditor::ExtendMenu()
+{
+	
+}
+
+void FParameterLibraryEditor::ExtendToolbar()
+{
+	
+}
+
+FName FParameterLibraryEditor::GetToolkitFName() const
+{
+	return FName("AnimNextParameterLibraryEditor");
+}
+
+FText FParameterLibraryEditor::GetBaseToolkitName() const
+{
+	return LOCTEXT("AppLabel", "AnimNextParameterLibraryEditor");
+}
+
+FString FParameterLibraryEditor::GetWorldCentricTabPrefix() const
+{
+	return LOCTEXT("WorldCentricTabPrefix", "AnimNextParameterLibraryEditor ").ToString();
+}
+
+FLinearColor FParameterLibraryEditor::GetWorldCentricTabColorScale() const
+{
+	return FLinearColor(0.3f, 0.2f, 0.5f, 0.5f);
+}
+
+void FParameterLibraryEditor::InitToolMenuContext(FToolMenuContext& InMenuContext)
+{
+}
+
+void FParameterLibraryEditor::SetSelectedObjects(TArray<UObject*> InObjects)
+{
+	if(DetailsView.IsValid())
+	{
+		DetailsView->SetObjects(InObjects);
+	}
+}
+
+void FParameterLibraryEditor::GetSaveableObjects(TArray<UObject*>& OutObjects) const
+{
+	// Base class will pick up edited object
+	FWorkflowCentricApplication::GetSaveableObjects(OutObjects);
+
+	// Get external objects too
+	FExternalPackageHelper::GetExternalSaveableObjects(ParameterLibrary, OutObjects);
+}
+
+}
+
+#undef LOCTEXT_NAMESPACE
\ No newline at end of file
diff --git a/Engine/Plugins/Experimental/AnimNext/Source/AnimNextEditor/Private/Param/AnimNextParameterLibraryEditor.h b/Engine/Plugins/Experimental/AnimNext/Source/AnimNextEditor/Private/Param/AnimNextParameterLibraryEditor.h
new file mode 100644
--- /dev/null
+++ b/Engine/Plugins/Experimental/AnimNext/Source/AnimNextEditor/Private/Param/AnimNextParameterLibraryEditor.h
@@ -0,0 +1,81 @@
+﻿// Copyright Epic Games, Inc. All Rights Reserved.
+
+#pragma once
+
+#include "CoreMinimal.h"
+#include "WorkflowOrientedApp/WorkflowCentricApplication.h"
+#include "GraphEditor.h"
+#include "WorkflowOrientedApp/WorkflowTabManager.h"
+
+enum class ERigVMGraphNotifType : uint8;
+class UAnimNextParameterLibrary;
+class FDocumentTracker;
+class FDocumentTabFactory;
+class FTabInfo;
+class FTabManager;
+class IToolkitHost;
+class URigVMGraph;
+class URigVMController;
+class UDetailsViewWrapperObject;
+
+namespace UE::AnimNext::Editor
+{
+
+class FParametersTabSummoner;
+class FParametersEditorMode;
+
+namespace ParameterLibraryModes
+{
+	extern const FName ParameterLibraryEditor;
+}
+
+namespace ParameterLibraryTabs
+{
+	extern const FName Details;
+	extern const FName Parameters;
+}
+
+class FParameterLibraryEditor : public FWorkflowCentricApplication
+{
+public:
+	FParameterLibraryEditor();
+	virtual ~FParameterLibraryEditor() override;
+
+	/** Edits the specified asset */
+	void InitEditor(const EToolkitMode::Type InMode, const TSharedPtr<IToolkitHost>& InInitToolkitHost, UAnimNextParameterLibrary* InAnimNextParameterLibrary);
+
+private:
+	friend class FParameterLibraryEditorMode;
+	friend class FParameterLibraryTabSummoner;
+
+	// IToolkit interface
+	virtual FName GetToolkitFName() const override;
+	virtual FText GetBaseToolkitName() const override;
+	virtual FString GetWorldCentricTabPrefix() const override;
+	virtual FLinearColor GetWorldCentricTabColorScale() const override;
+	virtual void InitToolMenuContext(FToolMenuContext& InMenuContext) override;
+
+	// FAssetEditorToolkit interface
+	virtual void GetSaveableObjects(TArray<UObject*>& OutObjects) const override;
+
+	void BindCommands();
+
+	void ExtendMenu();
+
+	void ExtendToolbar();
+
+	void SetSelectedObjects(TArray<UObject*> InObjects);
+
+	void HandleDetailsViewCreated(TSharedRef<IDetailsView> InDetailsView)
+	{
+		DetailsView = InDetailsView;
+	}
+
+	// The asset we are editing
+	UAnimNextParameterLibrary* ParameterLibrary = nullptr;
+
+	// Our details panel
+	TSharedPtr<IDetailsView> DetailsView;
+};
+
+}
\ No newline at end of file
diff --git a/Engine/Plugins/Experimental/AnimNext/Source/AnimNextEditor/Private/Param/AnimNextParameterSettings.cpp b/Engine/Plugins/Experimental/AnimNext/Source/AnimNextEditor/Private/Param/AnimNextParameterSettings.cpp
new file mode 100644
--- /dev/null
+++ b/Engine/Plugins/Experimental/AnimNext/Source/AnimNextEditor/Private/Param/AnimNextParameterSettings.cpp
@@ -0,0 +1,35 @@
+﻿// Copyright Epic Games, Inc. All Rights Reserved.
+
+#include "AnimNextParameterSettings.h"
+
+#include "AssetRegistry/AssetRegistryModule.h"
+#include "AssetRegistry/IAssetRegistry.h"
+#include "Misc/PackageName.h"
+#include "Misc/Paths.h"
+
+UAnimNextParameterSettings::UAnimNextParameterSettings()
+{
+}
+
+const FAnimNextParamType& UAnimNextParameterSettings::GetLastParameterType() const
+{
+	return LastParameterType;
+}
+
+void UAnimNextParameterSettings::SetLastParameterType(const FAnimNextParamType& InLastParameterType)
+{
+	LastParameterType = InLastParameterType;
+}
+
+FAssetData UAnimNextParameterSettings::GetLastLibrary() const
+{
+	IAssetRegistry& AssetRegistryModule = FModuleManager::Get().LoadModuleChecked<FAssetRegistryModule>(TEXT("AssetRegistry")).Get();
+	return AssetRegistryModule.GetAssetByObjectPath(LastLibrary);
+}
+
+void UAnimNextParameterSettings::SetLastLibrary(const FAssetData& InLastLibrary)
+{
+	LastLibrary = InLastLibrary.GetSoftObjectPath();
+
+	SaveConfig();
+}
\ No newline at end of file
diff --git a/Engine/Plugins/Experimental/AnimNext/Source/AnimNextEditor/Private/Param/AnimNextParameterSettings.h b/Engine/Plugins/Experimental/AnimNext/Source/AnimNextEditor/Private/Param/AnimNextParameterSettings.h
new file mode 100644
--- /dev/null
+++ b/Engine/Plugins/Experimental/AnimNext/Source/AnimNextEditor/Private/Param/AnimNextParameterSettings.h
@@ -0,0 +1,36 @@
+﻿// Copyright Epic Games, Inc. All Rights Reserved.
+
+#pragma once
+
+#include "Containers/UnrealString.h"
+#include "UObject/Object.h"
+#include "Param/ParamType.h"
+#include "AnimNextParameterSettings.generated.h"
+
+UCLASS(Config = EditorPerProjectUserSettings)
+class UAnimNextParameterSettings : public UObject
+{
+	GENERATED_BODY()
+
+	UAnimNextParameterSettings();
+
+public:
+	// Get the type of the last parameter type that we created
+	const FAnimNextParamType& GetLastParameterType() const;
+
+	// Set the type of the last parameter type that we created
+	void SetLastParameterType(const FAnimNextParamType& InLastParameterType);
+
+	// Get the library of the last parameter type that we created
+	FAssetData GetLastLibrary() const;
+
+	// Set the library of the last parameter type that we created
+	void SetLastLibrary(const FAssetData& InLastLibrary);
+
+private:
+	UPROPERTY(config)
+	FSoftObjectPath LastLibrary;
+
+	UPROPERTY(Transient)
+	FAnimNextParamType LastParameterType = FAnimNextParamType::GetType<bool>();
+};
\ No newline at end of file
diff --git a/Engine/Plugins/Experimental/AnimNext/Source/AnimNextEditor/Private/Param/ParameterBlockEditorMode.cpp b/Engine/Plugins/Experimental/AnimNext/Source/AnimNextEditor/Private/Param/ParameterBlockEditorMode.cpp
new file mode 100644
--- /dev/null
+++ b/Engine/Plugins/Experimental/AnimNext/Source/AnimNextEditor/Private/Param/ParameterBlockEditorMode.cpp
@@ -0,0 +1,327 @@
+﻿// Copyright Epic Games, Inc. All Rights Reserved.
+
+#include "ParameterBlockEditorMode.h"
+#include "AnimNextParameterBlockEditor.h"
+#include "PropertyEditorModule.h"
+#include "Widgets/Docking/SDockTab.h"
+#include "EdGraph/EdGraph.h"
+#include "EdGraph/EdGraphSchema.h"
+#include "Widgets/Layout/SSpacer.h"
+#include "Modules/ModuleManager.h"
+#include "SParameterBlockView.h"
+#include "Param/AnimNextParameterBlock_EditorData.h"
+#include "RigVMModel/RigVMGraph.h"
+
+#define LOCTEXT_NAMESPACE "ParameterBlockEditorMode"
+
+namespace UE::AnimNext::Editor
+{
+
+struct FParameterBlockTabHistory : public FGenericTabHistory
+{
+public:
+	FParameterBlockTabHistory(TSharedPtr<FDocumentTabFactory> InFactory, TSharedPtr<FTabPayload> InPayload)
+		: FGenericTabHistory(InFactory, InPayload)
+		, SavedLocation(FVector2D::ZeroVector)
+		, SavedZoomAmount(INDEX_NONE)
+	{
+	}
+
+private:
+	// FGenericTabHistory interface
+	virtual void EvokeHistory(TSharedPtr<::FTabInfo> InTabInfo, bool bPrevTabMatches) override
+	{
+		FWorkflowTabSpawnInfo SpawnInfo;
+		SpawnInfo.Payload = Payload;
+		SpawnInfo.TabInfo = InTabInfo;
+
+		if(bPrevTabMatches)
+		{
+			TSharedPtr<SDockTab> DockTab = InTabInfo->GetTab().Pin();
+			GraphEditor = StaticCastSharedRef<SGraphEditor>(DockTab->GetContent());
+		}
+		else
+		{
+			TSharedRef<SGraphEditor> GraphEditorRef = StaticCastSharedRef<SGraphEditor>(FactoryPtr.Pin()->CreateTabBody(SpawnInfo));
+			GraphEditor = GraphEditorRef;
+			FactoryPtr.Pin()->UpdateTab(InTabInfo->GetTab().Pin(), SpawnInfo, GraphEditorRef);
+		}
+	}
+	
+	virtual void SaveHistory() override
+	{
+		if (IsHistoryValid())
+		{
+			check(GraphEditor.IsValid());
+			GraphEditor.Pin()->GetViewLocation(SavedLocation, SavedZoomAmount);
+			GraphEditor.Pin()->GetViewBookmark(SavedBookmarkId);
+		}
+	}
+	
+	virtual void RestoreHistory() override
+	{
+		if (IsHistoryValid())
+		{
+			check(GraphEditor.IsValid());
+			GraphEditor.Pin()->SetViewLocation(SavedLocation, SavedZoomAmount, SavedBookmarkId);
+		}
+	}
+	
+	/** The graph editor represented by this history node. While this node is inactive, the graph editor is invalid */
+	TWeakPtr<SGraphEditor> GraphEditor;
+	
+	/** Saved location the graph editor was at when this history node was last visited */
+	FVector2D SavedLocation;
+	
+	/** Saved zoom the graph editor was at when this history node was last visited */
+	float SavedZoomAmount;
+	
+	/** Saved bookmark ID the graph editor was at when this history node was last visited */
+	FGuid SavedBookmarkId;
+};
+
+FParameterBlockEditorDocumentSummoner::FParameterBlockEditorDocumentSummoner(TSharedPtr<FWorkflowCentricApplication> InHostingApp)
+	: FDocumentTabFactoryForObjects<UEdGraph>(ParameterBlockTabs::Document, InHostingApp)
+{
+}
+
+void FParameterBlockEditorDocumentSummoner::OnTabActivated(TSharedPtr<SDockTab> Tab) const
+{
+	TSharedRef<SGraphEditor> GraphEditor = StaticCastSharedRef<SGraphEditor>(Tab->GetContent());
+	OnGraphEditorFocusedDelegate.ExecuteIfBound(GraphEditor);
+}
+
+void FParameterBlockEditorDocumentSummoner::OnTabBackgrounded(TSharedPtr<SDockTab> Tab) const
+{
+	TSharedRef<SGraphEditor> GraphEditor = StaticCastSharedRef<SGraphEditor>(Tab->GetContent());
+	OnGraphEditorBackgroundedDelegate.ExecuteIfBound(GraphEditor);
+}
+
+void FParameterBlockEditorDocumentSummoner::OnTabRefreshed(TSharedPtr<SDockTab> Tab) const
+{
+	TSharedRef<SGraphEditor> GraphEditor = StaticCastSharedRef<SGraphEditor>(Tab->GetContent());
+	GraphEditor->NotifyGraphChanged();
+}
+
+void FParameterBlockEditorDocumentSummoner::SaveState(TSharedPtr<SDockTab> Tab, TSharedPtr<FTabPayload> Payload) const
+{
+	TSharedRef<SGraphEditor> GraphEditor = StaticCastSharedRef<SGraphEditor>(Tab->GetContent());
+
+	FVector2D ViewLocation;
+	float ZoomAmount;
+	GraphEditor->GetViewLocation(ViewLocation, ZoomAmount);
+
+	UEdGraph* Graph = Payload->IsValid() ? FTabPayload_UObject::CastChecked<UEdGraph>(Payload) : nullptr;
+
+	OnSaveGraphStateDelegate.ExecuteIfBound(Graph, ViewLocation, ZoomAmount);
+}
+
+TAttribute<FText> FParameterBlockEditorDocumentSummoner::ConstructTabNameForObject(UEdGraph* DocumentID) const
+{
+	if (DocumentID)
+	{
+		if (const UEdGraphSchema* Schema = DocumentID->GetSchema())
+		{
+			FGraphDisplayInfo Info;
+			Schema->GetGraphDisplayInformation(*DocumentID, /*out*/ Info);
+			return Info.DisplayName;
+		}
+		else
+		{
+			// if we don't have a schema, we're dealing with a malformed (or incomplete graph)...
+			// possibly in the midst of some transaction - here we return the object's outer path 
+			// so we can at least get some context as to which graph we're referring
+			return FText::FromString(DocumentID->GetPathName());
+		}
+	}
+
+	return LOCTEXT("UnknownGraphName", "Unknown");
+}
+
+TSharedRef<SWidget> FParameterBlockEditorDocumentSummoner::CreateTabBodyForObject(const FWorkflowTabSpawnInfo& Info, UEdGraph* DocumentID) const
+{
+	check(Info.TabInfo.IsValid());
+	
+	if(OnCreateGraphEditorWidgetDelegate.IsBound())
+	{
+		return OnCreateGraphEditorWidgetDelegate.Execute(Info.TabInfo.ToSharedRef(), DocumentID);
+	}
+	else
+	{
+		return SNew(SSpacer);
+	}
+}
+
+const FSlateBrush* FParameterBlockEditorDocumentSummoner::GetTabIconForObject(const FWorkflowTabSpawnInfo& Info, UEdGraph* DocumentID) const
+{
+	if(OnGetTabIconDelegate.IsBound())
+	{
+		return OnGetTabIconDelegate.Execute(DocumentID);
+	}
+	else
+	{
+		return FAppStyle::Get().GetBrush(TEXT("GraphEditor.EventGraph_16x"));
+	}
+}
+
+TSharedRef<FGenericTabHistory> FParameterBlockEditorDocumentSummoner::CreateTabHistoryNode(TSharedPtr<FTabPayload> Payload)
+{
+	return MakeShared<FParameterBlockTabHistory>(SharedThis(this), Payload);
+}
+
+class FParameterBlockDetailsTabSummoner : public FWorkflowTabFactory
+{
+public:
+	FParameterBlockDetailsTabSummoner(TSharedPtr<FParameterBlockEditor> InHostingApp, FOnDetailsViewCreated InOnDetailsViewCreated)
+		: FWorkflowTabFactory(ParameterBlockTabs::Details, InHostingApp)
+	{
+		TabLabel = LOCTEXT("DetailsTabLabel", "Details");
+		TabIcon = FSlateIcon("EditorStyle", "LevelEditor.Tabs.Details");
+		ViewMenuDescription = LOCTEXT("DetailsTabMenuDescription", "Details");
+		ViewMenuTooltip = LOCTEXT("DetailsTabToolTip", "Shows the details tab for selected objects.");
+		bIsSingleton = true;
+
+		FDetailsViewArgs DetailsViewArgs;
+		DetailsViewArgs.NameAreaSettings = FDetailsViewArgs::HideNameArea;
+		DetailsViewArgs.bHideSelectionTip = true;
+
+		FPropertyEditorModule& PropertyEditorModule = FModuleManager::LoadModuleChecked<FPropertyEditorModule>("PropertyEditor");
+		DetailsView = PropertyEditorModule.CreateDetailView(DetailsViewArgs);
+		
+		InOnDetailsViewCreated.ExecuteIfBound(DetailsView.ToSharedRef());
+	}
+
+	TSharedPtr<IDetailsView> GetDetailsView() const
+	{
+		return DetailsView;
+	}
+
+private:
+	virtual TSharedRef<SWidget> CreateTabBody(const FWorkflowTabSpawnInfo& Info) const override
+	{
+		return DetailsView.ToSharedRef();
+	}
+	
+	virtual FText GetTabToolTipText(const FWorkflowTabSpawnInfo& Info) const override
+	{
+		return ViewMenuTooltip;
+	}
+	
+	TSharedPtr<IDetailsView> DetailsView;
+};
+
+class FParameterBlockTabSummoner : public FWorkflowTabFactory
+{
+public:
+	FParameterBlockTabSummoner(TSharedPtr<FParameterBlockEditor> InHostingApp)
+		: FWorkflowTabFactory(ParameterBlockTabs::Parameters, InHostingApp)
+	{
+		TabLabel = LOCTEXT("ParameterBlockTabLabel", "Parameter Block");
+		TabIcon = FSlateIcon("EditorStyle", "LevelEditor.Tabs.Details");
+		ViewMenuDescription = LOCTEXT("ParameterBlockTabMenuDescription", "Parameter Block");
+		ViewMenuTooltip = LOCTEXT("ParameterBlockTabToolTip", "Shows all parameters in this parameter block.");
+		bIsSingleton = true;
+	}
+
+private:
+	virtual TSharedRef<SWidget> CreateTabBody(const FWorkflowTabSpawnInfo& Info) const override
+	{
+		return SNew(SParameterBlockView, StaticCastSharedPtr<FParameterBlockEditor>(HostingApp.Pin())->EditorData)
+			.OnSelectionChanged(this, &FParameterBlockTabSummoner::HandleSelectionChanged)
+			.OnOpenGraph(this, &FParameterBlockTabSummoner::HandleOpenGraph);
+	}
+	
+	virtual FText GetTabToolTipText(const FWorkflowTabSpawnInfo& Info) const override
+	{
+		return ViewMenuTooltip;
+	}
+
+	void HandleSelectionChanged(const TArray<UObject*>& InObjects) const
+	{
+		TSharedPtr<FParameterBlockEditor> ParametersEditor = StaticCastSharedPtr<FParameterBlockEditor>(HostingApp.Pin());
+		ParametersEditor->SetSelectedObjects(InObjects);
+	}
+
+	void HandleOpenGraph(URigVMGraph* InGraph) const
+	{
+		TSharedPtr<FParameterBlockEditor> ParameterBlockEditor = StaticCastSharedPtr<FParameterBlockEditor>(HostingApp.Pin());
+		UAnimNextParameterBlock_EditorData* EditorData = StaticCastSharedPtr<FParameterBlockEditor>(HostingApp.Pin())->EditorData;
+		if(UObject* EditorObject = EditorData->GetEditorObjectForRigVMGraph(InGraph))
+		{
+			ParameterBlockEditor->OpenDocument(EditorObject, FDocumentTracker::EOpenDocumentCause::OpenNewDocument);
+		}
+	}
+};
+
+FParameterBlockEditorMode::FParameterBlockEditorMode(TSharedRef<FWorkflowCentricApplication> InHostingApp)
+	: FApplicationMode(ParameterBlockModes::ParameterBlockEditor)
+	, HostingAppPtr(InHostingApp)
+{
+	HostingAppPtr = InHostingApp;
+
+	TSharedRef<FParameterBlockEditor> ParametersEditor = StaticCastSharedRef<FParameterBlockEditor>(InHostingApp);
+	
+	TabFactories.RegisterFactory(MakeShared<FParameterBlockDetailsTabSummoner>(ParametersEditor, FOnDetailsViewCreated::CreateSP(&ParametersEditor.Get(), &FParameterBlockEditor::HandleDetailsViewCreated)));
+	TabFactories.RegisterFactory(MakeShared<FParameterBlockTabSummoner>(ParametersEditor));
+
+	TabLayout = FTabManager::NewLayout("Standalone_AnimNextParameterBlockEditor_Layout_v1.0")
+		->AddArea
+		(
+			FTabManager::NewPrimaryArea()
+			->SetOrientation(Orient_Vertical)
+			->Split
+			(
+				FTabManager::NewSplitter()
+				->SetSizeCoefficient(1.0f)
+				->SetOrientation(Orient_Horizontal)
+				->Split
+				(
+					FTabManager::NewStack()
+					->SetSizeCoefficient(0.25f)
+					->SetHideTabWell(false)
+					->AddTab(ParameterBlockTabs::Parameters, ETabState::OpenedTab)
+				)
+				->Split
+				(
+					FTabManager::NewStack()
+					->SetSizeCoefficient(0.5f)
+					->SetHideTabWell(false)
+					->AddTab(ParameterBlockTabs::Document, ETabState::ClosedTab)
+				)
+				->Split
+				(
+					FTabManager::NewStack()
+					->SetSizeCoefficient(0.25f)
+					->SetHideTabWell(false)
+					->AddTab(ParameterBlockTabs::Details, ETabState::OpenedTab)
+				)
+			)
+		);
+
+	ParametersEditor->RegisterModeToolbarIfUnregistered(GetModeName());
+}
+
+void FParameterBlockEditorMode::RegisterTabFactories(TSharedPtr<FTabManager> InTabManager)
+{
+	TSharedPtr<FWorkflowCentricApplication> HostingApp = HostingAppPtr.Pin();
+	HostingApp->PushTabFactories(TabFactories);
+
+	FApplicationMode::RegisterTabFactories(InTabManager);
+}
+
+void FParameterBlockEditorMode::AddTabFactory(FCreateWorkflowTabFactory FactoryCreator)
+{
+	if (FactoryCreator.IsBound())
+	{
+		TabFactories.RegisterFactory(FactoryCreator.Execute(HostingAppPtr.Pin()));
+	}
+}
+
+void FParameterBlockEditorMode::RemoveTabFactory(FName TabFactoryID)
+{
+	TabFactories.UnregisterFactory(TabFactoryID);
+}
+
+}
+
+#undef LOCTEXT_NAMESPACE
\ No newline at end of file
diff --git a/Engine/Plugins/Experimental/AnimNext/Source/AnimNextEditor/Private/Param/ParameterBlockEditorMode.h b/Engine/Plugins/Experimental/AnimNext/Source/AnimNextEditor/Private/Param/ParameterBlockEditorMode.h
new file mode 100644
--- /dev/null
+++ b/Engine/Plugins/Experimental/AnimNext/Source/AnimNextEditor/Private/Param/ParameterBlockEditorMode.h
@@ -0,0 +1,96 @@
+﻿// Copyright Epic Games, Inc. All Rights Reserved.
+
+#pragma once
+
+#include "CoreMinimal.h"
+#include "EdGraph/EdGraph.h"
+#include "WorkflowOrientedApp/ApplicationMode.h"
+#include "WorkflowOrientedApp/WorkflowTabManager.h"
+#include "WorkflowOrientedApp/WorkflowUObjectDocuments.h"
+
+class FWorkflowCentricApplication;
+class SGraphEditor;
+
+namespace UE::AnimNext::Editor
+{
+
+DECLARE_DELEGATE_OneParam(FOnDetailsViewCreated, TSharedRef<IDetailsView>);
+
+struct FParameterBlockEditorDocumentSummoner : public FDocumentTabFactoryForObjects<UEdGraph>
+{
+public:
+	// Delegate called when a graph editor widget is created
+	DECLARE_DELEGATE_RetVal_TwoParams(TSharedRef<SGraphEditor>, FOnCreateGraphEditorWidget, TSharedRef<FTabInfo>, UEdGraph*);
+
+	// Delegate called when a graph editor is focused
+	DECLARE_DELEGATE_OneParam(FOnGraphEditorFocused, TSharedRef<SGraphEditor>);
+
+	// Delegate called when a graph editor is backgrounded
+	DECLARE_DELEGATE_OneParam(FOnGraphEditorBackgrounded, TSharedRef<SGraphEditor>);
+
+	// Delegate called to save the state of a graph
+	DECLARE_DELEGATE_ThreeParams(FOnSaveGraphState, UEdGraph*, FVector2D, float);
+
+	// Delegate called to save the state of a graph
+	DECLARE_DELEGATE_RetVal_OneParam(const FSlateBrush*, FOnGetTabIcon, UEdGraph*);
+	
+	FParameterBlockEditorDocumentSummoner(TSharedPtr<FWorkflowCentricApplication> InHostingApp);
+
+	FOnCreateGraphEditorWidget& OnCreateGraphEditorWidget() { return OnCreateGraphEditorWidgetDelegate; }
+	FOnGraphEditorFocused& OnGraphEditorFocused() { return OnGraphEditorFocusedDelegate; }
+	FOnGraphEditorBackgrounded& OnGraphEditorBackgrounded() { return OnGraphEditorBackgroundedDelegate; }
+	FOnSaveGraphState& OnSaveGraphState() { return OnSaveGraphStateDelegate; }
+	FOnGetTabIcon& OnGetTabIcon() { return OnGetTabIconDelegate; }
+
+private:
+	// FWorkflowTabFactory interface
+	virtual void OnTabActivated(TSharedPtr<SDockTab> Tab) const override;
+	virtual void OnTabBackgrounded(TSharedPtr<SDockTab> Tab) const override;
+	virtual void OnTabRefreshed(TSharedPtr<SDockTab> Tab) const override;
+	virtual void SaveState(TSharedPtr<SDockTab> Tab, TSharedPtr<FTabPayload> Payload) const override;
+	
+	// FDocumentTabFactoryForObjects interface
+	virtual TAttribute<FText> ConstructTabNameForObject(UEdGraph* DocumentID) const override;
+	virtual TSharedRef<SWidget> CreateTabBodyForObject(const FWorkflowTabSpawnInfo& Info, UEdGraph* DocumentID) const override;
+	virtual const FSlateBrush* GetTabIconForObject(const FWorkflowTabSpawnInfo& Info, UEdGraph* DocumentID) const override;
+	virtual TSharedRef<FGenericTabHistory> CreateTabHistoryNode(TSharedPtr<FTabPayload> Payload) override;
+
+	// The hosting app
+	TWeakPtr<FWorkflowCentricApplication> HostingAppPtr;
+
+	// Delegate called when a graph editor widget is created
+	FOnCreateGraphEditorWidget OnCreateGraphEditorWidgetDelegate;
+
+	// Delegate called when a graph editor is focused
+	FOnGraphEditorFocused OnGraphEditorFocusedDelegate;
+
+	// Delegate called when a graph editor is backgrounded
+	FOnGraphEditorBackgrounded OnGraphEditorBackgroundedDelegate;
+
+	// Delegate called to save the state of a graph
+	FOnSaveGraphState OnSaveGraphStateDelegate;
+
+	// Delegate called to get a tab icon for a graph
+	FOnGetTabIcon OnGetTabIconDelegate;
+};
+
+class FParameterBlockEditorMode : public FApplicationMode
+{
+public:
+	FParameterBlockEditorMode(TSharedRef<FWorkflowCentricApplication> InHostingApp);
+
+private:
+	// FApplicationMode interface
+	virtual void RegisterTabFactories(TSharedPtr<FTabManager> InTabManager) override;
+	virtual void AddTabFactory(FCreateWorkflowTabFactory FactoryCreator) override;
+	virtual void RemoveTabFactory(FName TabFactoryID) override;
+	
+	// The hosting app
+	TWeakPtr<FWorkflowCentricApplication> HostingAppPtr;
+
+	// The tab factories we support
+	FWorkflowAllowedTabSet TabFactories;
+};
+
+}
+
diff --git a/Engine/Plugins/Experimental/AnimNext/Source/AnimNextEditor/Private/Param/ParameterBlockFactory.cpp b/Engine/Plugins/Experimental/AnimNext/Source/AnimNextEditor/Private/Param/ParameterBlockFactory.cpp
new file mode 100644
--- /dev/null
+++ b/Engine/Plugins/Experimental/AnimNext/Source/AnimNextEditor/Private/Param/ParameterBlockFactory.cpp
@@ -0,0 +1,35 @@
+﻿// Copyright Epic Games, Inc. All Rights Reserved.
+
+#include "Param/ParameterBlockFactory.h"
+#include "Param/AnimNextParameterBlock.h"
+#include "Param/AnimNextParameterBlock_EditorData.h"
+#include "Param/ParametersExecuteContext.h"
+#include "UncookedOnlyUtils.h"
+
+UAnimNextParameterBlockFactory::UAnimNextParameterBlockFactory()
+{
+	bCreateNew = true;
+	bEditAfterNew = true;
+	SupportedClass = UAnimNextParameterBlock::StaticClass();
+}
+
+bool UAnimNextParameterBlockFactory::ConfigureProperties()
+{
+	return true;
+}
+
+UObject* UAnimNextParameterBlockFactory::FactoryCreateNew(UClass* Class, UObject* InParent, FName Name, EObjectFlags Flags, UObject* Context, FFeedbackContext* Warn, FName CallingContext)
+{
+	UAnimNextParameterBlock* NewBlock = NewObject<UAnimNextParameterBlock>(InParent, Class, Name, Flags | RF_Public | RF_Standalone | RF_Transactional | RF_LoadCompleted);
+
+	UAnimNextParameterBlock_EditorData* EditorData = NewObject<UAnimNextParameterBlock_EditorData>(NewBlock, TEXT("EditorData"), RF_Transactional);
+	NewBlock->EditorData = EditorData;
+	EditorData->Initialize(/*bRecompileVM*/false);
+	EditorData->GetRigVMClient()->SetExecuteContextStruct(FAnimNextParametersExecuteContext::StaticStruct());
+
+	// Compile the initial skeleton
+	UE::AnimNext::UncookedOnly::FUtils::Compile(NewBlock);
+	check(!EditorData->bErrorsDuringCompilation);
+	
+	return NewBlock;
+}
\ No newline at end of file
diff --git a/Engine/Plugins/Experimental/AnimNext/Source/AnimNextEditor/Private/Param/ParameterBlockViewMenuContext.h b/Engine/Plugins/Experimental/AnimNext/Source/AnimNextEditor/Private/Param/ParameterBlockViewMenuContext.h
new file mode 100644
--- /dev/null
+++ b/Engine/Plugins/Experimental/AnimNext/Source/AnimNextEditor/Private/Param/ParameterBlockViewMenuContext.h
@@ -0,0 +1,21 @@
+﻿// Copyright Epic Games, Inc. All Rights Reserved.
+
+#pragma once
+
+#include "ParameterBlockViewMenuContext.generated.h"
+
+namespace UE::AnimNext::Editor
+{
+	class SParameterBlockView;
+}
+
+UCLASS()
+class UParameterBlockViewMenuContext : public UObject
+{
+	GENERATED_BODY()
+
+	friend class UE::AnimNext::Editor::SParameterBlockView;
+	
+	// The parameter block view that we are editing
+	TWeakPtr<UE::AnimNext::Editor::SParameterBlockView> ParameterBlockView;
+};
\ No newline at end of file
diff --git a/Engine/Plugins/Experimental/AnimNext/Source/AnimNextEditor/Private/Param/ParameterLibraryEditorMode.cpp b/Engine/Plugins/Experimental/AnimNext/Source/AnimNextEditor/Private/Param/ParameterLibraryEditorMode.cpp
new file mode 100644
--- /dev/null
+++ b/Engine/Plugins/Experimental/AnimNext/Source/AnimNextEditor/Private/Param/ParameterLibraryEditorMode.cpp
@@ -0,0 +1,156 @@
+﻿// Copyright Epic Games, Inc. All Rights Reserved.
+
+#include "ParameterLibraryEditorMode.h"
+#include "AnimNextParameterLibraryEditor.h"
+#include "PropertyEditorModule.h"
+#include "Widgets/Docking/SDockTab.h"
+#include "EdGraph/EdGraph.h"
+#include "EdGraph/EdGraphSchema.h"
+#include "Widgets/Layout/SSpacer.h"
+#include "Modules/ModuleManager.h"
+#include "SParameterLibraryView.h"
+#include "Param/AnimNextParameterLibrary.h"
+
+#define LOCTEXT_NAMESPACE "ParameterLibraryEditorMode"
+
+namespace UE::AnimNext::Editor
+{
+
+class FParameterLibraryDetailsTabSummoner : public FWorkflowTabFactory
+{
+public:
+	FParameterLibraryDetailsTabSummoner(TSharedPtr<FParameterLibraryEditor> InHostingApp, FOnDetailsViewCreated InOnDetailsViewCreated)
+		: FWorkflowTabFactory(ParameterLibraryTabs::Details, InHostingApp)
+	{
+		TabLabel = LOCTEXT("DetailsTabLabel", "Details");
+		TabIcon = FSlateIcon("EditorStyle", "LevelEditor.Tabs.Details");
+		ViewMenuDescription = LOCTEXT("DetailsTabMenuDescription", "Details");
+		ViewMenuTooltip = LOCTEXT("DetailsTabToolTip", "Shows the details tab for selected objects.");
+		bIsSingleton = true;
+
+		FDetailsViewArgs DetailsViewArgs;
+		DetailsViewArgs.NameAreaSettings = FDetailsViewArgs::HideNameArea;
+		DetailsViewArgs.bHideSelectionTip = true;
+
+		FPropertyEditorModule& PropertyEditorModule = FModuleManager::LoadModuleChecked<FPropertyEditorModule>("PropertyEditor");
+		DetailsView = PropertyEditorModule.CreateDetailView(DetailsViewArgs);
+
+		InOnDetailsViewCreated.ExecuteIfBound(DetailsView.ToSharedRef());
+	}
+
+	TSharedPtr<IDetailsView> GetDetailsView() const
+	{
+		return DetailsView;
+	}
+
+private:
+	virtual TSharedRef<SWidget> CreateTabBody(const FWorkflowTabSpawnInfo& Info) const override
+	{
+		return DetailsView.ToSharedRef();
+	}
+	
+	virtual FText GetTabToolTipText(const FWorkflowTabSpawnInfo& Info) const override
+	{
+		return ViewMenuTooltip;
+	}
+	
+	TSharedPtr<IDetailsView> DetailsView;
+};
+
+class FParameterLibraryTabSummoner : public FWorkflowTabFactory
+{
+public:
+	FParameterLibraryTabSummoner(TSharedPtr<FParameterLibraryEditor> InHostingApp)
+		: FWorkflowTabFactory(ParameterLibraryTabs::Parameters, InHostingApp)
+	{
+		TabLabel = LOCTEXT("ParameterLibraryTabLabel", "Parameter Library");
+		TabIcon = FSlateIcon("EditorStyle", "LevelEditor.Tabs.Details");
+		ViewMenuDescription = LOCTEXT("ParameterLibraryTabMenuDescription", "Parameter Library");
+		ViewMenuTooltip = LOCTEXT("ParameterLibraryTabToolTip", "Shows all parameters in this parameter library.");
+		bIsSingleton = true;
+	}
+
+private:
+	virtual TSharedRef<SWidget> CreateTabBody(const FWorkflowTabSpawnInfo& Info) const override
+	{
+		return SNew(SParameterLibraryView, StaticCastSharedPtr<FParameterLibraryEditor>(HostingApp.Pin())->ParameterLibrary)
+			.OnSelectionChanged(this, &FParameterLibraryTabSummoner::HandleSelectionChanged);
+	}
+	
+	virtual FText GetTabToolTipText(const FWorkflowTabSpawnInfo& Info) const override
+	{
+		return ViewMenuTooltip;
+	}
+
+	void HandleSelectionChanged(const TArray<UObject*>& InObjects) const
+	{
+		TSharedPtr<FParameterLibraryEditor> ParameterLibraryEditor = StaticCastSharedPtr<FParameterLibraryEditor>(HostingApp.Pin());
+		ParameterLibraryEditor->SetSelectedObjects(InObjects);
+	}
+};
+
+FParameterLibraryEditorMode::FParameterLibraryEditorMode(TSharedRef<FWorkflowCentricApplication> InHostingApp)
+	: FApplicationMode(ParameterLibraryModes::ParameterLibraryEditor)
+	, HostingAppPtr(InHostingApp)
+{
+	HostingAppPtr = InHostingApp;
+
+	TSharedRef<FParameterLibraryEditor> ParameterLibraryEditor = StaticCastSharedRef<FParameterLibraryEditor>(InHostingApp);
+	
+	TabFactories.RegisterFactory(MakeShared<FParameterLibraryDetailsTabSummoner>(ParameterLibraryEditor, FOnDetailsViewCreated::CreateSP(&ParameterLibraryEditor.Get(), &FParameterLibraryEditor::HandleDetailsViewCreated)));
+	TabFactories.RegisterFactory(MakeShared<FParameterLibraryTabSummoner>(ParameterLibraryEditor));
+
+	TabLayout = FTabManager::NewLayout("Standalone_AnimNextParameterLibraryEditor_Layout_v1.0")
+		->AddArea
+		(
+			FTabManager::NewPrimaryArea()
+			->SetOrientation(Orient_Vertical)
+			->Split
+			(
+				FTabManager::NewSplitter()
+				->SetSizeCoefficient(1.0f)
+				->SetOrientation(Orient_Horizontal)
+				->Split
+				(
+					FTabManager::NewStack()
+					->SetSizeCoefficient(0.5f)
+					->SetHideTabWell(false)
+					->AddTab(ParameterLibraryTabs::Parameters, ETabState::OpenedTab)
+				)
+				->Split
+				(
+					FTabManager::NewStack()
+					->SetSizeCoefficient(0.5f)
+					->SetHideTabWell(false)
+					->AddTab(ParameterLibraryTabs::Details, ETabState::OpenedTab)
+				)
+			)
+		);
+
+	ParameterLibraryEditor->RegisterModeToolbarIfUnregistered(GetModeName());
+}
+
+void FParameterLibraryEditorMode::RegisterTabFactories(TSharedPtr<FTabManager> InTabManager)
+{
+	TSharedPtr<FWorkflowCentricApplication> HostingApp = HostingAppPtr.Pin();
+	HostingApp->PushTabFactories(TabFactories);
+
+	FApplicationMode::RegisterTabFactories(InTabManager);
+}
+
+void FParameterLibraryEditorMode::AddTabFactory(FCreateWorkflowTabFactory FactoryCreator)
+{
+	if (FactoryCreator.IsBound())
+	{
+		TabFactories.RegisterFactory(FactoryCreator.Execute(HostingAppPtr.Pin()));
+	}
+}
+
+void FParameterLibraryEditorMode::RemoveTabFactory(FName TabFactoryID)
+{
+	TabFactories.UnregisterFactory(TabFactoryID);
+}
+
+}
+
+#undef LOCTEXT_NAMESPACE
\ No newline at end of file
diff --git a/Engine/Plugins/Experimental/AnimNext/Source/AnimNextEditor/Private/Param/ParameterLibraryEditorMode.h b/Engine/Plugins/Experimental/AnimNext/Source/AnimNextEditor/Private/Param/ParameterLibraryEditorMode.h
new file mode 100644
--- /dev/null
+++ b/Engine/Plugins/Experimental/AnimNext/Source/AnimNextEditor/Private/Param/ParameterLibraryEditorMode.h
@@ -0,0 +1,37 @@
+﻿// Copyright Epic Games, Inc. All Rights Reserved.
+
+#pragma once
+
+#include "CoreMinimal.h"
+#include "EdGraph/EdGraph.h"
+#include "WorkflowOrientedApp/ApplicationMode.h"
+#include "WorkflowOrientedApp/WorkflowTabManager.h"
+
+class FWorkflowCentricApplication;
+class SGraphEditor;
+
+namespace UE::AnimNext::Editor
+{
+
+DECLARE_DELEGATE_OneParam(FOnDetailsViewCreated, TSharedRef<IDetailsView>);
+
+class FParameterLibraryEditorMode : public FApplicationMode
+{
+public:
+	FParameterLibraryEditorMode(TSharedRef<FWorkflowCentricApplication> InHostingApp);
+
+private:
+	// FApplicationMode interface
+	virtual void RegisterTabFactories(TSharedPtr<FTabManager> InTabManager) override;
+	virtual void AddTabFactory(FCreateWorkflowTabFactory FactoryCreator) override;
+	virtual void RemoveTabFactory(FName TabFactoryID) override;
+	
+	// The hosting app
+	TWeakPtr<FWorkflowCentricApplication> HostingAppPtr;
+
+	// The tab factories we support
+	FWorkflowAllowedTabSet TabFactories;
+};
+
+}
+
diff --git a/Engine/Plugins/Experimental/AnimNext/Source/AnimNextEditor/Private/Param/ParameterLibraryFactory.cpp b/Engine/Plugins/Experimental/AnimNext/Source/AnimNextEditor/Private/Param/ParameterLibraryFactory.cpp
new file mode 100644
--- /dev/null
+++ b/Engine/Plugins/Experimental/AnimNext/Source/AnimNextEditor/Private/Param/ParameterLibraryFactory.cpp
@@ -0,0 +1,27 @@
+﻿// Copyright Epic Games, Inc. All Rights Reserved.
+
+#include "Param/ParameterLibraryFactory.h"
+#include "Param/AnimNextParameterLibrary.h"
+
+UAnimNextParameterLibraryFactory::UAnimNextParameterLibraryFactory()
+{
+	bCreateNew = true;
+	bEditAfterNew = true;
+	SupportedClass = UAnimNextParameterLibrary::StaticClass();
+}
+
+bool UAnimNextParameterLibraryFactory::ConfigureProperties()
+{
+	return true;
+}
+
+UObject* UAnimNextParameterLibraryFactory::FactoryCreateNew(UClass* Class, UObject* InParent, FName Name, EObjectFlags Flags, UObject* Context, FFeedbackContext* Warn, FName CallingContext)
+{
+	UAnimNextParameterLibrary* NewParameter = NewObject<UAnimNextParameterLibrary>(InParent, Class, Name, Flags | RF_Public | RF_Standalone | RF_Transactional | RF_LoadCompleted);
+
+	// make sure the package is never cooked.
+	UPackage* Package = NewParameter->GetOutermost();
+	Package->SetPackageFlags(Package->GetPackageFlags() | PKG_EditorOnly);
+
+	return NewParameter;
+}
\ No newline at end of file
diff --git a/Engine/Plugins/Experimental/AnimNext/Source/AnimNextEditor/Private/Param/ParameterLibraryViewMenuContext.h b/Engine/Plugins/Experimental/AnimNext/Source/AnimNextEditor/Private/Param/ParameterLibraryViewMenuContext.h
new file mode 100644
--- /dev/null
+++ b/Engine/Plugins/Experimental/AnimNext/Source/AnimNextEditor/Private/Param/ParameterLibraryViewMenuContext.h
@@ -0,0 +1,21 @@
+﻿// Copyright Epic Games, Inc. All Rights Reserved.
+
+#pragma once
+
+#include "ParameterLibraryViewMenuContext.generated.h"
+
+namespace UE::AnimNext::Editor
+{
+class SParameterLibraryView;
+}
+
+UCLASS()
+class UParameterLibraryViewMenuContext : public UObject
+{
+	GENERATED_BODY()
+
+	friend class UE::AnimNext::Editor::SParameterLibraryView;
+	
+	// The parameter library view that we are editing
+	TWeakPtr<UE::AnimNext::Editor::SParameterLibraryView> ParameterLibraryView;
+};
\ No newline at end of file
diff --git a/Engine/Plugins/Experimental/AnimNext/Source/AnimNextEditor/Private/Param/ParametersAssetDefinitions.cpp b/Engine/Plugins/Experimental/AnimNext/Source/AnimNextEditor/Private/Param/ParametersAssetDefinitions.cpp
new file mode 100644
--- /dev/null
+++ b/Engine/Plugins/Experimental/AnimNext/Source/AnimNextEditor/Private/Param/ParametersAssetDefinitions.cpp
@@ -0,0 +1,50 @@
+﻿// Copyright Epic Games, Inc. All Rights Reserved.
+
+#include "ParametersAssetDefinitions.h"
+#include "AnimNextParameterBlockEditor.h"
+#include "AnimNextParameterLibraryEditor.h"
+#include "Param/AnimNextParameterBlock.h"
+#include "Param/AnimNextParameterLibrary.h"
+
+#define LOCTEXT_NAMESPACE "AnimNextAssetDefinitions"
+
+EAssetCommandResult UAssetDefinition_AnimNextParameterLibrary::OpenAssets(const FAssetOpenArgs& OpenArgs) const
+{
+	for (UAnimNextParameterLibrary* Asset : OpenArgs.LoadObjects<UAnimNextParameterLibrary>())
+	{
+		TSharedRef<UE::AnimNext::Editor::FParameterLibraryEditor> Editor = MakeShared<UE::AnimNext::Editor::FParameterLibraryEditor>();
+		Editor->InitEditor(OpenArgs.GetToolkitMode(), OpenArgs.ToolkitHost, Asset);
+	}
+
+	return EAssetCommandResult::Handled;
+}
+
+FText UAssetDefinition_AnimNextParameter::GetObjectDisplayNameText(UObject* Object) const
+{
+	UAnimNextParameter* Parameter = CastChecked<UAnimNextParameter>(Object);
+	if(UAnimNextParameterLibrary* Library = CastChecked<UAnimNextParameterLibrary>(Object->GetOuter()))
+	{
+		return FText::Format(LOCTEXT("LibraryParameterDisplayFormat", "{0}.{1}"), FText::FromString(Library->GetName()), FText::FromString(Parameter->GetName()));
+	}
+
+	return FText::FromString(Parameter->GetName());
+}
+
+EAssetCommandResult UAssetDefinition_AnimNextParameterBlock::OpenAssets(const FAssetOpenArgs& OpenArgs) const
+{
+	for (UAnimNextParameterBlock* Asset : OpenArgs.LoadObjects<UAnimNextParameterBlock>())
+	{
+		TSharedRef<UE::AnimNext::Editor::FParameterBlockEditor> Editor = MakeShared<UE::AnimNext::Editor::FParameterBlockEditor>();
+		Editor->InitEditor(OpenArgs.GetToolkitMode(), OpenArgs.ToolkitHost, Asset);
+	}
+
+	return EAssetCommandResult::Handled;
+}
+
+FText UAssetDefinition_AnimNextParameterBlockBinding::GetObjectDisplayNameText(UObject* Object) const
+{
+	UAnimNextParameterBlockBinding* Binding = CastChecked<UAnimNextParameterBlockBinding>(Object);
+	return FText::Format(LOCTEXT("BindingDisplayFormat", "{0} Binding"), FText::FromName(Binding->ParameterName));
+}
+
+#undef LOCTEXT_NAMESPACE
\ No newline at end of file
diff --git a/Engine/Plugins/Experimental/AnimNext/Source/AnimNextEditor/Private/Param/ParametersAssetDefinitions.h b/Engine/Plugins/Experimental/AnimNext/Source/AnimNextEditor/Private/Param/ParametersAssetDefinitions.h
new file mode 100644
--- /dev/null
+++ b/Engine/Plugins/Experimental/AnimNext/Source/AnimNextEditor/Private/Param/ParametersAssetDefinitions.h
@@ -0,0 +1,76 @@
+﻿// Copyright Epic Games, Inc. All Rights Reserved.
+
+#pragma once
+
+#include "AssetDefinitionDefault.h"
+#include "Param/AnimNextParameter.h"
+#include "Param/AnimNextParameterBlockBinding.h"
+#include "Param/AnimNextParameterBlock.h"
+#include "Param/AnimNextParameterLibrary.h"
+#include "ParametersAssetDefinitions.generated.h"
+
+#define LOCTEXT_NAMESPACE "AnimNextAssetDefinitions"
+
+UCLASS()
+class UAssetDefinition_AnimNextParameterLibrary : public UAssetDefinitionDefault
+{
+	GENERATED_BODY()
+
+public:
+	// UAssetDefinition interface
+	virtual FText GetAssetDisplayName() const override { return LOCTEXT("AnimNextParameterLibrary", "AnimNext Parameter Library"); }
+	virtual FLinearColor GetAssetColor() const override { return FLinearColor(FColor(128,64,64)); }
+	virtual TSoftClassPtr<UObject> GetAssetClass() const override { return UAnimNextParameterLibrary::StaticClass(); }
+	virtual EAssetCommandResult OpenAssets(const FAssetOpenArgs& OpenArgs) const override;
+	virtual TConstArrayView<FAssetCategoryPath> GetAssetCategories() const override
+	{
+		static const auto Categories = { FAssetCategoryPath(EAssetCategoryPaths::Animation, LOCTEXT("AnimNextSubMenu", "AnimNext")) };
+		return Categories;
+	}
+};
+
+UCLASS()
+class UAssetDefinition_AnimNextParameter : public UAssetDefinitionDefault
+{
+	GENERATED_BODY()
+
+public:
+	// UAssetDefinition interface
+	virtual FText GetAssetDisplayName() const override { return LOCTEXT("AnimNextParameter", "AnimNext Parameter"); }
+	virtual FLinearColor GetAssetColor() const override { return FLinearColor(FColor(128,64,64)); }
+	virtual TSoftClassPtr<UObject> GetAssetClass() const override { return UAnimNextParameter::StaticClass(); }
+	virtual FText GetObjectDisplayNameText(UObject* Object) const override;
+};
+
+UCLASS()
+class UAssetDefinition_AnimNextParameterBlock : public UAssetDefinitionDefault
+{
+	GENERATED_BODY()
+
+public:
+	// UAssetDefinition interface
+	virtual FText GetAssetDisplayName() const override { return LOCTEXT("AnimNextParameterBlock", "AnimNext Parameter Block"); }
+	virtual FLinearColor GetAssetColor() const override { return FLinearColor(FColor(128,64,64)); }
+	virtual TSoftClassPtr<UObject> GetAssetClass() const override { return UAnimNextParameterBlock::StaticClass(); }
+	virtual EAssetCommandResult OpenAssets(const FAssetOpenArgs& OpenArgs) const override;
+	virtual TConstArrayView<FAssetCategoryPath> GetAssetCategories() const override
+	{
+		static const auto Categories = { FAssetCategoryPath(EAssetCategoryPaths::Animation, LOCTEXT("AnimNextSubMenu", "AnimNext")) };
+		return Categories;
+	}
+};
+
+UCLASS()
+class UAssetDefinition_AnimNextParameterBlockBinding : public UAssetDefinitionDefault
+{
+	GENERATED_BODY()
+
+public:
+	// UAssetDefinition interface
+	virtual FText GetAssetDisplayName() const override { return LOCTEXT("AnimNextParameterBinding", "AnimNext Parameter Binding"); }
+	virtual FLinearColor GetAssetColor() const override { return FLinearColor(FColor(128,64,64)); }
+	virtual TSoftClassPtr<UObject> GetAssetClass() const override { return UAnimNextParameterBlockBinding::StaticClass(); }
+	virtual FText GetObjectDisplayNameText(UObject* Object) const override;
+};
+
+#undef LOCTEXT_NAMESPACE
\ No newline at end of file
diff --git a/Engine/Plugins/Experimental/AnimNext/Source/AnimNextEditor/Private/Param/ParametersEditorSchemaActions.cpp b/Engine/Plugins/Experimental/AnimNext/Source/AnimNextEditor/Private/Param/ParametersEditorSchemaActions.cpp
new file mode 100644
--- /dev/null
+++ b/Engine/Plugins/Experimental/AnimNext/Source/AnimNextEditor/Private/Param/ParametersEditorSchemaActions.cpp
@@ -0,0 +1,59 @@
+﻿// Copyright Epic Games, Inc. All Rights Reserved.
+
+#include "ParametersEditorSchemaActions.h"
+#include "Param/AnimNextParameterBlock_EdGraph.h"
+#include "Param/AnimNextParameterBlock_EdGraphNode.h"
+#include "Param/AnimNextParameterBlock_EditorData.h"
+#include "RigVMModel/Nodes/RigVMUnitNode.h"
+#include "Settings/ControlRigSettings.h"
+#include "EditorUtils.h"
+
+UEdGraphNode* FAnimNextParameterSchemaAction_RigUnit::PerformAction(UEdGraph* ParentGraph, TArray<UEdGraphPin*>& FromPins, const FVector2D Location, bool bSelectNewNode)
+{
+	UAnimNextParameterBlock_EditorData* EditorData = ParentGraph->GetTypedOuter<UAnimNextParameterBlock_EditorData>();
+	UAnimNextParameterBlock_EdGraphNode* NewNode = nullptr;
+	UAnimNextParameterBlock_EdGraph* EdGraph = Cast<UAnimNextParameterBlock_EdGraph>(ParentGraph);
+
+	if (EditorData != nullptr && EdGraph != nullptr)
+	{
+		FName Name = UE::AnimNext::Editor::FUtils::ValidateName(EditorData, StructTemplate->GetFName().ToString());
+		URigVMController* Controller = EditorData->GetRigVMClient()->GetController(ParentGraph);
+
+		Controller->OpenUndoBracket(FString::Printf(TEXT("Add '%s' Node"), *Name.ToString()));
+
+		FRigVMUnitNodeCreatedContext& UnitNodeCreatedContext = Controller->GetUnitNodeCreatedContext();
+		FRigVMUnitNodeCreatedContext::FScope ReasonScope(UnitNodeCreatedContext, ERigVMNodeCreatedReason::NodeSpawner);
+
+		if (URigVMUnitNode* ModelNode = Controller->AddUnitNode(StructTemplate, FRigUnit::GetMethodName(), Location, Name.ToString(), true, false))
+		{
+			NewNode = Cast<UAnimNextParameterBlock_EdGraphNode>(EdGraph->FindNodeForModelNodeName(ModelNode->GetFName()));
+			check(NewNode);
+
+			if (NewNode)
+			{
+				Controller->ClearNodeSelection(true);
+				Controller->SelectNode(ModelNode, true, true);
+			}
+
+#if WITH_EDITORONLY_DATA
+
+			// @TODO: switch over to custom settings here
+			const FControlRigSettingsPerPinBool* ExpansionMapPtr = UControlRigEditorSettings::Get()->RigUnitPinExpansion.Find(ModelNode->GetScriptStruct()->GetName());
+			if (ExpansionMapPtr)
+			{
+				const FControlRigSettingsPerPinBool& ExpansionMap = *ExpansionMapPtr;
+
+				for (const TPair<FString, bool>& Pair : ExpansionMap.Values)
+				{
+					FString PinPath = FString::Printf(TEXT("%s.%s"), *ModelNode->GetName(), *Pair.Key);
+					Controller->SetPinExpansion(PinPath, Pair.Value, true);
+				}
+			}
+		}
+#endif
+		
+		Controller->CloseUndoBracket();
+	}
+
+	return NewNode;
+}
diff --git a/Engine/Plugins/Experimental/AnimNext/Source/AnimNextEditor/Private/Param/ParametersEditorSchemaActions.h b/Engine/Plugins/Experimental/AnimNext/Source/AnimNextEditor/Private/Param/ParametersEditorSchemaActions.h
new file mode 100644
--- /dev/null
+++ b/Engine/Plugins/Experimental/AnimNext/Source/AnimNextEditor/Private/Param/ParametersEditorSchemaActions.h
@@ -0,0 +1,55 @@
+﻿// Copyright Epic Games, Inc. All Rights Reserved.
+
+#pragma once
+
+#include "CoreMinimal.h"
+#include "EdGraph/EdGraphSchema.h"
+#include "Styling/AppStyle.h"
+#include "ParametersEditorSchemaActions.generated.h"
+
+struct FSlateBrush;
+
+USTRUCT()
+struct FAnimNextParameterSchemaAction : public FEdGraphSchemaAction
+{
+	GENERATED_BODY()
+
+	FAnimNextParameterSchemaAction() = default;
+	
+	FAnimNextParameterSchemaAction(FText InNodeCategory, FText InMenuDesc, FText InToolTip, FText InKeywords = FText::GetEmpty())
+		: FEdGraphSchemaAction(MoveTemp(InNodeCategory), MoveTemp(InMenuDesc), MoveTemp(InToolTip), 0, MoveTemp(InKeywords))
+	{
+	}
+
+	virtual const FSlateBrush* GetIconBrush() const
+	{
+		return FAppStyle::Get().GetBrush("NoBrush");
+	}
+
+	virtual const FLinearColor& GetIconColor() const
+	{
+		static const FLinearColor DefaultColor;
+		return DefaultColor;
+	}
+};
+
+USTRUCT()
+struct FAnimNextParameterSchemaAction_RigUnit : public FAnimNextParameterSchemaAction
+{
+	GENERATED_BODY()
+
+	FAnimNextParameterSchemaAction_RigUnit() = default;
+	
+	FAnimNextParameterSchemaAction_RigUnit(UScriptStruct* InStructTemplate, FText InNodeCategory, FText InMenuDesc, FText InToolTip, FText InKeywords = FText::GetEmpty())
+		: FAnimNextParameterSchemaAction(InNodeCategory, InMenuDesc, InToolTip, InKeywords)
+		, StructTemplate(InStructTemplate)
+	{}
+
+	// FEdGraphSchemaAction Interface
+	virtual UEdGraphNode* PerformAction(UEdGraph* ParentGraph, TArray<UEdGraphPin*>& FromPins, const FVector2D Location, bool bSelectNewNode = true) override;
+	virtual UEdGraphNode* PerformAction(UEdGraph* ParentGraph, UEdGraphPin* FromPin, const FVector2D Location, bool bSelectNewNode = true) { return nullptr; }
+
+private:
+	// The script struct for our rig unit
+	UScriptStruct* StructTemplate = nullptr;
+};
