AnimNext: Initial editor-side parameter blocks & libraries

diff --git a/Engine/Plugins/Experimental/AnimNext/Source/AnimNextEditor/Private/Param/SAddParametersDialog.cpp b/Engine/Plugins/Experimental/AnimNext/Source/AnimNextEditor/Private/Param/SAddParametersDialog.cpp
new file mode 100644
--- /dev/null
+++ b/Engine/Plugins/Experimental/AnimNext/Source/AnimNextEditor/Private/Param/SAddParametersDialog.cpp
@@ -0,0 +1,420 @@
+// Copyright Epic Games, Inc. All Rights Reserved.
+
+#include "SAddParametersDialog.h"
+
+#include "AddParameterDialogMenuContext.h"
+#include "AnimNextParameterSettings.h"
+#include "AssetToolsModule.h"
+#include "ContentBrowserModule.h"
+#include "DetailLayoutBuilder.h"
+#include "IContentBrowserSingleton.h"
+#include "EditorUtils.h"
+#include "UncookedOnlyUtils.h"
+#include "PropertyBagDetails.h"
+#include "SPinTypeSelector.h"
+#include "Framework/Application/SlateApplication.h"
+#include "Framework/Docking/TabManager.h"
+#include "Widgets/Layout/SUniformGridPanel.h"
+#include "Widgets/Layout/SBorder.h"
+#include "Widgets/Input/SButton.h"
+#include "Widgets/Layout/SBox.h"
+#include "Widgets/Text/SInlineEditableTextBlock.h"
+#include "SSimpleButton.h"
+#include "SSimpleComboButton.h"
+#include "Param/AnimNextParameterLibrary.h"
+#include "Param/ParameterLibraryFactory.h"
+
+#define LOCTEXT_NAMESPACE "SAddParametersDialog"
+
+namespace UE::AnimNext::Editor
+{
+
+namespace AddParametersDialog
+{
+static FName Column_Name(TEXT("Name"));
+static FName Column_Type(TEXT("Type"));
+static FName Column_Library(TEXT("Library"));
+static FName SelectLibraryMenuName(TEXT("AnimNext.AddParametersDialog.SelectedLibraryMenu"));
+}
+
+void SAddParametersDialog::Construct(const FArguments& InArgs)
+{
+	using namespace AddParametersDialog;
+
+	Library = InArgs._Library;
+
+	SWindow::Construct(SWindow::FArguments()
+		.Title(LOCTEXT("WindowTitle", "Add Parameters"))
+		.SizingRule(ESizingRule::UserSized)
+		.ClientSize(FVector2D(500.f, 500.f))
+		.SupportsMaximize(false)
+		.SupportsMinimize(false)
+		[
+			SNew(SBox)
+			.Padding(5.0f)
+			[
+				SNew(SVerticalBox)
+				+SVerticalBox::Slot()
+				.AutoHeight()
+				.HAlign(HAlign_Left)
+				.Padding(0.0f, 5.0f)
+				[
+					SNew(SSimpleButton)
+					.Text(LOCTEXT("AddButton", "Add"))
+					.ToolTipText(LOCTEXT("AddButtonTooltip", "Queue a new parameter for adding. New parameters will re-use the settings from the last queued parameter."))
+					.Icon(FAppStyle::Get().GetBrush("Icons.Plus"))
+					.OnClicked_Lambda([this]()
+					{
+						AddEntry();
+						return FReply::Handled();
+					})
+				]
+				+SVerticalBox::Slot()
+				.FillHeight(1.0f)
+				[
+					SAssignNew(EntriesList, SListView<TSharedRef<FParameterToAdd>>)
+					.ListItemsSource(&Entries)
+					.OnGenerateRow(this, &SAddParametersDialog::HandleGenerateRow)
+					.ItemHeight(20.0f)
+					.HeaderRow(
+						SNew(SHeaderRow)
+						+SHeaderRow::Column(Column_Name)
+						.DefaultLabel(LOCTEXT("NameColumnHeader", "Name"))
+						.ToolTipText(LOCTEXT("NameColumnHeaderTooltip", "The name of the new parameter"))
+						.FillWidth(0.25f)
+
+						+SHeaderRow::Column(Column_Type)
+						.DefaultLabel(LOCTEXT("TypeColumnHeader", "Type"))
+						.ToolTipText(LOCTEXT("TypeColumnHeaderTooltip", "The type of the new parameter"))
+						.FillWidth(0.25f)
+
+						+SHeaderRow::Column(Column_Library)
+						.DefaultLabel(LOCTEXT("LibraryColumnHeader", "Library"))
+						.ToolTipText(LOCTEXT("LibraryColumnHeaderTooltip", "The library the new parameter will be created in"))
+						.FillWidth(0.5f)
+					)
+				]
+				+SVerticalBox::Slot()
+				.AutoHeight()
+				.HAlign(HAlign_Right)
+				[
+					SNew(SUniformGridPanel)
+					.SlotPadding(FAppStyle::Get().GetMargin("StandardDialog.SlotPadding"))
+					.MinDesiredSlotWidth(FAppStyle::GetFloat("StandardDialog.MinDesiredSlotWidth"))
+					.MinDesiredSlotHeight(FAppStyle::GetFloat("StandardDialog.MinDesiredSlotHeight"))
+					+SUniformGridPanel::Slot(0, 0)
+					[
+						SNew(SButton)
+						.HAlign(HAlign_Center)
+						.ButtonStyle(&FAppStyle::Get().GetWidgetStyle<FButtonStyle>("PrimaryButton"))
+						.Text_Lambda([this]()
+						{
+							return FText::Format(LOCTEXT("AddButtonFormat", "Add {0} {0}|plural(one=Parameter,other=Parameters)"), FText::AsNumber(Entries.Num()));
+						})
+						.ToolTipText(LOCTEXT("AddButtonTooltip", "Add the selected parameters to the current parameter block"))
+						.OnClicked_Lambda([this]()
+						{
+							RequestDestroyWindow();
+							return FReply::Handled();
+						})
+					]
+					+SUniformGridPanel::Slot(1, 0)
+					[
+						SNew(SButton)
+						.HAlign(HAlign_Center)
+						.ButtonStyle(&FAppStyle::Get().GetWidgetStyle<FButtonStyle>("Button"))
+						.Text(LOCTEXT("CancelButton", "Cancel"))
+						.ToolTipText(LOCTEXT("CancelButtonTooltip", "Cancel adding new parameters"))
+						.OnClicked_Lambda([this]()
+						{
+							bCancelPressed = true;
+							RequestDestroyWindow();
+							return FReply::Handled();
+						})
+					]
+				]
+			]
+		]);
+
+	// Add an initial item
+	AddEntry();
+}
+
+void SAddParametersDialog::AddEntry()
+{
+	const UAnimNextParameterSettings* Settings = GetDefault<UAnimNextParameterSettings>();
+	FAssetData LibraryAsset = Settings->GetLastLibrary();
+	if(Library != nullptr)
+	{
+		LibraryAsset = FAssetData(Library);
+	}
+	
+	TArray<FName> PendingNames;
+	PendingNames.Reserve(Entries.Num());
+	for(const TSharedRef<FParameterToAdd>& QueuedAdd : Entries)
+	{
+		if(QueuedAdd->Library == LibraryAsset)
+		{
+			PendingNames.Add(QueuedAdd->Name);
+		}
+	}
+	FName ParameterName = FUtils::GetNewParameterNameInLibrary(LibraryAsset, TEXT("NewParameter"), PendingNames);
+	Entries.Add(MakeShared<FParameterToAdd>(Settings->GetLastParameterType(), ParameterName, Settings->GetLastLibrary()));
+
+	RefreshEntries();
+}
+
+void SAddParametersDialog::RefreshEntries()
+{
+	EntriesList->RequestListRefresh();
+}
+
+class SParameterToAdd : public SMultiColumnTableRow<TSharedRef<FParameterToAdd>>
+{
+	SLATE_BEGIN_ARGS(SParameterToAdd) {}
+	SLATE_END_ARGS()
+
+	void Construct(const FArguments& InArgs, const TSharedRef<STableViewBase>& InOwnerTableView, TSharedRef<FParameterToAdd> InEntry, TSharedRef<SAddParametersDialog> InDialog)
+	{
+		Entry = InEntry;
+		WeakDialog = InDialog;
+		
+		SMultiColumnTableRow<TSharedRef<FParameterToAdd>>::Construct( SMultiColumnTableRow<TSharedRef<FParameterToAdd>>::FArguments(), InOwnerTableView);
+	}
+
+	virtual TSharedRef<SWidget> GenerateWidgetForColumn(const FName& InColumnName) override
+	{
+		using namespace AddParametersDialog;
+
+		if(InColumnName == Column_Name)
+		{
+			return
+				SNew(SBox)
+				.HAlign(HAlign_Left)
+				.VAlign(VAlign_Center)
+				[
+					SNew(SInlineEditableTextBlock)
+					.Font(IDetailLayoutBuilder::GetDetailFont())
+					.IsSelected(this, &SParameterToAdd::IsSelectedExclusively)
+					.ToolTipText(LOCTEXT("NameTooltip", "The name of the new parameter"))
+					.Text_Lambda([this]()
+					{
+						return FText::FromName(Entry->Name);
+					})
+					.OnTextCommitted_Lambda([this](const FText& InText, ETextCommit::Type InCommitType)
+					{
+						if(InCommitType == ETextCommit::OnEnter)
+						{
+							Entry->Name = *InText.ToString();
+						}
+					})
+					.OnVerifyTextChanged_Lambda([this](const FText& InNewText, FText& OutErrorText)
+					{
+						const FString NewString = InNewText.ToString();
+
+						// Make sure the new name only contains valid characters
+						if (!FName::IsValidXName(NewString, INVALID_OBJECTNAME_CHARACTERS INVALID_LONGPACKAGE_CHARACTERS, &OutErrorText))
+						{
+							return false;
+						}
+
+						FName Name(*NewString);
+
+						if(FUtils::DoesParameterExistInLibrary(Entry->Library, Name))
+						{
+							OutErrorText = LOCTEXT("Error_NameExists", "This name already exists in the specified library");
+							return false;
+						}
+
+						return true;
+					})
+				];
+		}
+		else if(InColumnName == Column_Type)
+		{
+			auto GetPinInfo = [this]()
+			{
+				return UncookedOnly::FUtils::GetPinTypeFromParamType(Entry->Type);
+			};
+
+			auto PinInfoChanged = [this](const FEdGraphPinType& PinType)
+			{
+				Entry->Type = UncookedOnly::FUtils::GetParamTypeFromPinType(PinType);
+
+				UAnimNextParameterSettings* Settings = GetMutableDefault<UAnimNextParameterSettings>();
+				Settings->SetLastParameterType(Entry->Type);
+			};
+
+			return
+				SNew(SBox)
+				.HAlign(HAlign_Left)
+				.VAlign(VAlign_Center)
+				[
+					SNew(SPinTypeSelector, FGetPinTypeTree::CreateStatic(&Editor::FUtils::GetFilteredVariableTypeTree))
+						.TargetPinType_Lambda(GetPinInfo)
+						.OnPinTypeChanged_Lambda(PinInfoChanged)
+						.Schema(GetDefault<UPropertyBagSchema>())
+						.bAllowArrays(true)
+						.TypeTreeFilter(ETypeTreeFilter::None)
+						.Font(IDetailLayoutBuilder::GetDetailFont())
+				];
+		}
+		else if(InColumnName == Column_Library)
+		{
+			return
+				SNew(SBox)
+				.HAlign(HAlign_Left)
+				.VAlign(VAlign_Center)
+				.IsEnabled(WeakDialog.Pin()->Library == nullptr)
+				[
+					SNew(SSimpleComboButton)
+					.UsesSmallText(true)
+					.HasDownArrow(true)
+					.Text_Lambda([this]()
+					{
+						return FText::FromName(Entry->Library.PackageName);
+					})
+					.ToolTipText_Lambda([this]()
+					{
+						return FText::Format(LOCTEXT("LibraryTooltip", "The library of the new parameter.\n{0}"), FText::FromName(Entry->Library.PackageName));
+					})
+					.OnGetMenuContent(WeakDialog.Pin().Get(), &SAddParametersDialog::HandleGetAddParameterMenuContent, Entry)
+				];
+		}
+
+		return SNullWidget::NullWidget;
+	}
+
+	TSharedPtr<FParameterToAdd> Entry;
+	TWeakPtr<SAddParametersDialog> WeakDialog;
+};
+
+TSharedRef<ITableRow> SAddParametersDialog::HandleGenerateRow(TSharedRef<FParameterToAdd> InEntry, const TSharedRef<STableViewBase>& InOwnerTable)
+{
+	return SNew(SParameterToAdd, InOwnerTable, InEntry, SharedThis(this));
+}
+
+bool SAddParametersDialog::ShowModal(TArray<FParameterToAdd>& OutParameters)
+{
+	FSlateApplication::Get().AddModalWindow(SharedThis(this), FGlobalTabmanager::Get()->GetRootWindow());
+
+	if(!bCancelPressed)
+	{
+		bool bHasValid = false;
+		for(TSharedRef<FParameterToAdd>& Entry : Entries)
+		{
+			if(Entry->IsValid())
+			{
+				OutParameters.Add(*Entry);
+				bHasValid = true;
+			}
+		}
+		return bHasValid;
+	}
+	return false;
+}
+
+TSharedRef<SWidget> SAddParametersDialog::HandleGetAddParameterMenuContent(TSharedPtr<FParameterToAdd> InEntry)
+{
+	using namespace AddParametersDialog;
+
+	UToolMenus* ToolMenus = UToolMenus::Get();
+
+	if(!ToolMenus->IsMenuRegistered(SelectLibraryMenuName))
+	{
+		UToolMenu* Menu = ToolMenus->RegisterMenu(SelectLibraryMenuName);
+
+		{
+			FToolMenuSection& Section = Menu->AddSection("AddNewLibrary", LOCTEXT("AddNewLibraryMenuSection", "Add New Library"));
+			Section.AddSubMenu("AddNewLibrary",
+				LOCTEXT("AddNewLibraryLabel", "Add New Library"),
+				LOCTEXT("AddNewLibraryTooltip", "Add a new parameter library"),
+				FNewToolMenuDelegate::CreateLambda([](UToolMenu* InMenu)
+				{
+					if(UAddParameterDialogMenuContext* MenuContext = InMenu->FindContext<UAddParameterDialogMenuContext>())
+					{
+						FPathPickerConfig PathPickerConfig;
+						PathPickerConfig.bShowFavorites = false;
+						PathPickerConfig.OnPathSelected = FOnPathSelected::CreateLambda([MenuContext](const FString& InPath)
+						{
+							FSlateApplication::Get().DismissAllMenus();
+
+							if(TSharedPtr<FParameterToAdd> Entry = MenuContext->Entry.Pin())
+							{
+								IAssetTools& AssetTools = FModuleManager::LoadModuleChecked<FAssetToolsModule>("AssetTools").Get();
+								if(UObject* NewAsset = AssetTools.CreateAsset(TEXT("NewParameterLibrary"), InPath, UAnimNextParameterLibrary::StaticClass(), NewObject<UAnimNextParameterLibraryFactory>()))
+								{
+									Entry->Library = FAssetData(NewAsset);
+
+									UAnimNextParameterSettings* Settings = GetMutableDefault<UAnimNextParameterSettings>();
+									Settings->SetLastLibrary(Entry->Library);
+								}
+							}
+						});
+
+						FContentBrowserModule& ContentBrowserModule = FModuleManager::LoadModuleChecked<FContentBrowserModule>("ContentBrowser");
+						FToolMenuSection& Section = InMenu->AddSection("ChooseLibrary", LOCTEXT("ChooseLibraryMenuSection", "Choose Existing Library"));
+						Section.AddEntry(FToolMenuEntry::InitWidget(
+								"LibraryPicker",
+								SNew(SBox)
+								.WidthOverride(300.0f)
+								.HeightOverride(400.0f)
+								[
+									ContentBrowserModule.Get().CreatePathPicker(PathPickerConfig)
+								],
+								FText::GetEmpty(),
+								true));
+					}
+				}));
+		}
+
+		{
+			
+
+			FToolMenuSection& Section = Menu->AddDynamicSection("ChooseLibrary",  FNewToolMenuDelegate::CreateLambda([](UToolMenu* InMenu)
+			{
+				if(UAddParameterDialogMenuContext* MenuContext = InMenu->FindContext<UAddParameterDialogMenuContext>())
+				{
+					FAssetPickerConfig AssetPickerConfig;
+					AssetPickerConfig.SelectionMode = ESelectionMode::Single;
+					AssetPickerConfig.InitialAssetViewType = EAssetViewType::List;
+					AssetPickerConfig.Filter.ClassPaths = { UAnimNextParameterLibrary::StaticClass()->GetClassPathName() };
+					AssetPickerConfig.OnAssetSelected = FOnAssetSelected::CreateLambda([MenuContext](const FAssetData& InAssetData)
+					{
+						FSlateApplication::Get().DismissAllMenus();
+
+						if(TSharedPtr<FParameterToAdd> Entry = MenuContext->Entry.Pin())
+						{
+							Entry->Library = InAssetData;
+
+							UAnimNextParameterSettings* Settings = GetMutableDefault<UAnimNextParameterSettings>();
+							Settings->SetLastLibrary(InAssetData);
+						}
+					});
+
+					FContentBrowserModule& ContentBrowserModule = FModuleManager::LoadModuleChecked<FContentBrowserModule>("ContentBrowser");
+					FToolMenuSection& Section = InMenu->AddSection("ChooseLibrary", LOCTEXT("ChooseLibraryMenuSection", "Choose Existing Library"));
+					Section.AddEntry(FToolMenuEntry::InitWidget(
+							"LibraryPicker",
+							SNew(SBox)
+							.WidthOverride(300.0f)
+							.HeightOverride(400.0f)
+							[
+								ContentBrowserModule.Get().CreateAssetPicker(AssetPickerConfig)
+							],
+							FText::GetEmpty(),
+							true));
+				}
+			}));
+		}
+	}
+
+	UAddParameterDialogMenuContext* MenuContext = NewObject<UAddParameterDialogMenuContext>();
+	MenuContext->AddParametersDialog = SharedThis(this);
+	MenuContext->Entry = InEntry;
+	return ToolMenus->GenerateWidget(SelectLibraryMenuName, FToolMenuContext(MenuContext));
+}
+
+}
+
+#undef LOCTEXT_NAMESPACE
\ No newline at end of file
diff --git a/Engine/Plugins/Experimental/AnimNext/Source/AnimNextEditor/Private/Param/SAddParametersDialog.h b/Engine/Plugins/Experimental/AnimNext/Source/AnimNextEditor/Private/Param/SAddParametersDialog.h
new file mode 100644
--- /dev/null
+++ b/Engine/Plugins/Experimental/AnimNext/Source/AnimNextEditor/Private/Param/SAddParametersDialog.h
@@ -0,0 +1,79 @@
+// Copyright Epic Games, Inc. All Rights Reserved.
+
+#pragma once
+
+#include "Param/ParamType.h"
+#include "Widgets/SWindow.h"
+#include "AssetRegistry/AssetData.h"
+
+class SWrapBox;
+class UAnimNextParameterLibrary;
+
+namespace UE::AnimNext::Editor
+{
+
+struct FParameterToAdd
+{
+	FParameterToAdd() = default;
+
+	FParameterToAdd(const FAnimNextParamType& InType, FName InName, const FAssetData& InLibrary)
+		: Type(InType)
+		, Name(InName)
+		, Library(InLibrary)
+	{}
+
+	bool IsValid() const
+	{
+		return Name != NAME_None && Type.IsValid() && Library.IsValid(); 
+	}
+	
+	// Type
+	FAnimNextParamType Type;
+
+	// Name for parameter
+	FName Name;
+
+	// Parameter library
+	FAssetData Library;
+};
+
+class SAddParametersDialog : public SWindow
+{
+public:
+	SLATE_BEGIN_ARGS(SAddParametersDialog)
+		: _Library(nullptr)
+	{}
+
+	SLATE_ARGUMENT(UAnimNextParameterLibrary*, Library)
+	
+	SLATE_END_ARGS()
+
+	void Construct(const FArguments& InArgs);
+
+	bool ShowModal(TArray<FParameterToAdd>& OutParameters);
+
+private:
+	void AddEntry();
+
+	void RefreshEntries();
+
+	TSharedRef<ITableRow> HandleGenerateRow(TSharedRef<FParameterToAdd> InEntry, const TSharedRef<STableViewBase>& InOwnerTable);
+
+	TSharedRef<SWidget> HandleGetAddParameterMenuContent(TSharedPtr<FParameterToAdd> InEntry);
+	
+private:
+	friend class SParameterToAdd;
+
+	TSharedPtr<SWrapBox> QueuedParametersBox;
+
+	TSharedPtr<SListView<TSharedRef<FParameterToAdd>>> EntriesList;
+
+	TArray<TSharedRef<FParameterToAdd>> Entries;
+
+	// The fixed library to use. If this is NULL, any library can be used.
+	UAnimNextParameterLibrary* Library = nullptr;
+	
+	bool bCancelPressed = false;
+};
+
+}
diff --git a/Engine/Plugins/Experimental/AnimNext/Source/AnimNextEditor/Private/Param/SLinkParametersDialog.cpp b/Engine/Plugins/Experimental/AnimNext/Source/AnimNextEditor/Private/Param/SLinkParametersDialog.cpp
new file mode 100644
--- /dev/null
+++ b/Engine/Plugins/Experimental/AnimNext/Source/AnimNextEditor/Private/Param/SLinkParametersDialog.cpp
@@ -0,0 +1,154 @@
+// Copyright Epic Games, Inc. All Rights Reserved.
+
+#include "SLinkParametersDialog.h"
+#include "SParameterPicker.h"
+#include "Param/ParameterPickerArgs.h"
+#include "Widgets/Input/SButton.h"
+#include "Framework/Application/SlateApplication.h"
+#include "Framework/Docking/TabManager.h"
+#include "Widgets/Layout/SWrapBox.h"
+#include "Widgets/Layout/SBox.h"
+#include "Widgets/SBoxPanel.h"
+#include "Widgets/Layout/SBorder.h"
+#include "Widgets/SCompoundWidget.h"
+#include "Widgets/Layout/SUniformGridPanel.h"
+#include "Widgets/Text/STextBlock.h"
+
+#define LOCTEXT_NAMESPACE "SLinkParametersDialog"
+
+namespace UE::AnimNext::Editor
+{
+
+class SParameterToLink : public SCompoundWidget
+{
+	SLATE_BEGIN_ARGS(SParameterToLink) {}
+
+	SLATE_END_ARGS()
+
+	void Construct(const FArguments& InArgs, const FParameterBindingReference& InBindingReference)
+	{
+		BindingReference = InBindingReference;
+		
+		ChildSlot
+		[
+			SNew(SBorder)
+			.BorderImage(FAppStyle::GetBrush("ToolPanel.GroupBorder"))
+			.Padding(5.0f, 3.0f)
+			[
+				SNew(STextBlock)
+				.Text(FText::FromName(BindingReference.Parameter))
+				.ToolTipText(FText::FromName((BindingReference.Parameter)))
+			]
+		];
+	}
+
+	FParameterBindingReference BindingReference;
+};
+
+void SLinkParametersDialog::Construct(const FArguments& InArgs)
+{
+	FParameterPickerArgs ParameterPickerArgs;
+	ParameterPickerArgs.bShowUnboundParameters = false;
+	ParameterPickerArgs.OnGetParameterBindings = &OnGetParameterBindings;
+	ParameterPickerArgs.OnSelectionChanged = FSimpleDelegate::CreateSP(this, &SLinkParametersDialog::HandleSelectionChanged);
+	ParameterPickerArgs.OnFilterParameter = InArgs._OnFilterParameter;
+
+	SWindow::Construct(SWindow::FArguments()
+		.Title(LOCTEXT("WindowTitle", "Link Parameters"))
+		.SizingRule(ESizingRule::UserSized)
+		.ClientSize(FVector2D(500.f, 500.f))
+		.SupportsMaximize(false)
+		.SupportsMinimize(false)
+		[
+			SNew(SBox)
+			.Padding(5.0f)
+			[
+				SNew(SVerticalBox)
+				+SVerticalBox::Slot()
+				.FillHeight(1.0f)
+				[
+					SNew(SParameterPicker)
+					.Args(ParameterPickerArgs)
+				]
+				+SVerticalBox::Slot()
+				.AutoHeight()
+				[
+					SNew(SBox)
+					.MinDesiredHeight(40.0f)
+					.Padding(5.0f)
+					[
+						SAssignNew(QueuedParametersBox, SWrapBox)
+						.UseAllottedWidth(true)
+					]
+				]
+				+SVerticalBox::Slot()
+				.AutoHeight()
+				.HAlign(HAlign_Right)
+				[
+					SNew(SUniformGridPanel)
+					.SlotPadding(FAppStyle::Get().GetMargin("StandardDialog.SlotPadding"))
+					.MinDesiredSlotWidth(FAppStyle::GetFloat("StandardDialog.MinDesiredSlotWidth"))
+					.MinDesiredSlotHeight(FAppStyle::GetFloat("StandardDialog.MinDesiredSlotHeight"))
+					+SUniformGridPanel::Slot(0, 0)
+					[
+						SNew(SButton)
+						.HAlign(HAlign_Center)
+						.ButtonStyle(&FAppStyle::Get().GetWidgetStyle<FButtonStyle>("PrimaryButton"))
+						.Text(LOCTEXT("LinkButton", "Link"))
+						.ToolTipText(LOCTEXT("LinkButtonTooltip", "Link the selected parameters to the current parameter block"))
+						.OnClicked_Lambda([this]()
+						{
+							RequestDestroyWindow();
+							return FReply::Handled();
+						})
+					]
+					+SUniformGridPanel::Slot(1, 0)
+					[
+						SNew(SButton)
+						.HAlign(HAlign_Center)
+						.ButtonStyle(&FAppStyle::Get().GetWidgetStyle<FButtonStyle>("Button"))
+						.Text(LOCTEXT("CancelButton", "Cancel"))
+						.ToolTipText(LOCTEXT("CancelButtonTooltip", "Cancel adding linked parameters"))
+						.OnClicked_Lambda([this]()
+						{
+							bCancelPressed = true;
+							RequestDestroyWindow();
+							return FReply::Handled();
+						})
+					]
+				]
+			]
+		]);
+}
+
+bool SLinkParametersDialog::ShowModal(TArray<FParameterBindingReference>& OutParameters)
+{
+	FSlateApplication::Get().AddModalWindow(SharedThis(this), FGlobalTabmanager::Get()->GetRootWindow());
+
+	if(!bCancelPressed)
+	{
+		OnGetParameterBindings.ExecuteIfBound(OutParameters);
+		return true;
+	}
+	return false;
+}
+
+void SLinkParametersDialog::HandleSelectionChanged()
+{
+	QueuedParametersBox->ClearChildren();
+
+	TArray<FParameterBindingReference> ParameterBindings;
+	OnGetParameterBindings.ExecuteIfBound(ParameterBindings);
+	for(const FParameterBindingReference& ParameterBinding : ParameterBindings)
+	{
+		QueuedParametersBox->AddSlot()
+		.Padding(2.0f)
+		[
+			SNew(SParameterToLink, ParameterBinding)
+		];
+	}
+}
+
+}
+
+#undef LOCTEXT_NAMESPACE
\ No newline at end of file
diff --git a/Engine/Plugins/Experimental/AnimNext/Source/AnimNextEditor/Private/Param/SLinkParametersDialog.h b/Engine/Plugins/Experimental/AnimNext/Source/AnimNextEditor/Private/Param/SLinkParametersDialog.h
new file mode 100644
--- /dev/null
+++ b/Engine/Plugins/Experimental/AnimNext/Source/AnimNextEditor/Private/Param/SLinkParametersDialog.h
@@ -0,0 +1,37 @@
+// Copyright Epic Games, Inc. All Rights Reserved.
+
+#pragma once
+
+#include "Param/ParameterPickerArgs.h"
+#include "AssetRegistry/AssetData.h"
+#include "Widgets/SWindow.h"
+
+class SWrapBox;
+
+namespace UE::AnimNext::Editor
+{
+
+class SLinkParametersDialog : public SWindow
+{
+public:
+	SLATE_BEGIN_ARGS(SLinkParametersDialog) {}
+
+	SLATE_EVENT(FOnFilterParameter, OnFilterParameter)
+
+	SLATE_END_ARGS()
+
+	void Construct(const FArguments& InArgs);
+
+	bool ShowModal(TArray<FParameterBindingReference>& OutParameters);
+
+private:
+	void HandleSelectionChanged();
+
+	TSharedPtr<SWrapBox> QueuedParametersBox;
+
+	FOnGetParameterBindings OnGetParameterBindings;
+
+	bool bCancelPressed = false;
+};
+
+}
diff --git a/Engine/Plugins/Experimental/AnimNext/Source/AnimNextEditor/Private/Param/SParameterBlockView.cpp b/Engine/Plugins/Experimental/AnimNext/Source/AnimNextEditor/Private/Param/SParameterBlockView.cpp
new file mode 100644
--- /dev/null
+++ b/Engine/Plugins/Experimental/AnimNext/Source/AnimNextEditor/Private/Param/SParameterBlockView.cpp
@@ -0,0 +1,759 @@
+﻿// Copyright Epic Games, Inc. All Rights Reserved.
+
+#include "SParameterBlockView.h"
+
+#include "Param/AnimNextParameterBlock.h"
+#include "Param/AnimNextParameter.h"
+#include "Param/AnimNextParameterBlock_EditorData.h"
+#include "Param/AnimNextParameterBlockBinding.h"
+#include "Widgets/Layout/SBox.h"
+#include "Widgets/Input/SSearchBox.h"
+#include "Param/AnimNextParameterBlockEntry.h"
+#include "Param/AnimNextParameterLibrary.h"
+#include "DetailLayoutBuilder.h"
+#include "ISourceControlModule.h"
+#include "ISourceControlProvider.h"
+#include "EditorUtils.h"
+#include "UncookedOnlyUtils.h"
+#include "ParameterBlockViewMenuContext.h"
+#include "PropertyBagDetails.h"
+#include "SAddParametersDialog.h"
+#include "SLinkParametersDialog.h"
+#include "SourceControlOperations.h"
+#include "SPinTypeSelector.h"
+#include "SSimpleButton.h"
+#include "Framework/Commands/GenericCommands.h"
+#include "Param/IAnimNextParameterBlockReferenceInterface.h"
+#include "RevisionControlStyle/RevisionControlStyle.h"
+#include "Widgets/Text/SInlineEditableTextBlock.h"
+#include "Param/AnimNextParameterBlockBindingReference.h"
+
+#define LOCTEXT_NAMESPACE "AnimNextParameterBlockView"
+
+namespace UE::AnimNext::Editor
+{
+
+namespace ParameterBlockView
+{
+
+static FName ContextMenuName(TEXT("AnimNext.ParameterBlockView.ContextMenu"));
+static FName Column_RevisionControl(TEXT("RevisionControl"));
+static FName Column_ModifiedStatus(TEXT("ModifiedStatus"));
+static FName Column_Name(TEXT("Name"));
+static FName Column_Type(TEXT("Type"));
+static FName Column_Value(TEXT("Value"));
+
+}
+
+// An entry displayed in the parameters view
+struct FParameterBlockViewEntry
+{
+	FParameterBlockViewEntry() = default;
+
+	FParameterBlockViewEntry(UAnimNextParameterBlockEntry* InEntry)
+		: WeakEntry(InEntry)
+	{}
+
+	bool PassesFilter(const FString& InFilterText) const
+	{
+		if(UAnimNextParameterBlockEntry* Entry = WeakEntry.Get())
+		{
+			return Entry->GetDisplayName().ToString().Contains(InFilterText);
+		}
+
+		return false;
+	}
+
+	// Ptr to the underlying entry
+	TWeakObjectPtr<UAnimNextParameterBlockEntry> WeakEntry;
+
+	// Widget used to rename items
+	TWeakPtr<SInlineEditableTextBlock> NameWidget;
+
+	// Flag to indicate a rename was requested
+	bool bRenameWhenScrolledIntoView = false;
+};
+
+void SParameterBlockView::Construct(const FArguments& InArgs, UAnimNextParameterBlock_EditorData* InEditorData)
+{
+	using namespace ParameterBlockView;
+	
+	check(InEditorData);
+
+	EditorData = InEditorData;
+
+	// Cache asset data for block for comparisons/filtering
+	BlockAssetData = FAssetData(UncookedOnly::FUtils::GetBlock(EditorData));
+
+	OnSelectionChangedDelegate = InArgs._OnSelectionChanged;
+	OnOpenGraphDelegate = InArgs._OnOpenGraph;
+
+	EditorData->ModifiedDelegate.AddSP(this, &SParameterBlockView::HandleBlockModified);
+
+	ChildSlot
+	[
+		SNew(SVerticalBox)
+		+SVerticalBox::Slot()
+		.AutoHeight()
+		[
+			SNew(SHorizontalBox)
+			+SHorizontalBox::Slot()
+			.AutoWidth()
+			.Padding(2.0f)
+			[
+				SNew(SSimpleButton)
+				.Text(LOCTEXT("AddNewParameterButton", "New"))
+				.Icon(FAppStyle::Get().GetBrush("Icons.Plus"))
+				.OnClicked_Lambda([this]()
+				{
+					TSharedRef<SAddParametersDialog> AddParametersDialog = SNew(SAddParametersDialog);
+					TArray<FParameterToAdd> ParametersToAdd;
+					if(AddParametersDialog->ShowModal(ParametersToAdd))
+					{
+						FScopedTransaction Transaction(FText::FormatOrdered(LOCTEXT("AddParameters", "Add {0}|plural(one=parameter,other=parameters)"), ParametersToAdd.Num()));
+
+						PendingSelection.Empty();
+
+						for(const FParameterToAdd& ParameterToAdd : ParametersToAdd)
+						{
+							// Create a new parameter in the supplied library
+							UAnimNextParameterLibrary* Library = Cast<UAnimNextParameterLibrary>(ParameterToAdd.Library.GetAsset());
+							UAnimNextParameter* NewParameter = Library->AddParameter(ParameterToAdd.Name, ParameterToAdd.Type);
+
+							// Create a new entry for the parameter
+							UAnimNextParameterBlockBinding* Binding = EditorData->AddBinding(NewParameter->GetFName(), Library);
+							PendingSelection.Add(Binding);
+						}
+					}
+
+					return FReply::Handled();
+				})
+			]
+			+SHorizontalBox::Slot()
+			.AutoWidth()
+			.Padding(2.0f)
+			[
+				SNew(SSimpleButton)
+				.Text(LOCTEXT("LinkParameterButton", "Link"))
+				.Icon(FAppStyle::Get().GetBrush("Icons.Link"))
+				.OnClicked_Lambda([this]()
+				{
+					TSharedRef<SLinkParametersDialog> AddParametersDialog = SNew(SLinkParametersDialog)
+						.OnFilterParameter(this, &SParameterBlockView::HandleFilterLinkedParameter);
+					TArray<FParameterBindingReference> ParametersToLink;
+					if(AddParametersDialog->ShowModal(ParametersToLink))
+					{
+						FScopedTransaction Transaction(FText::FormatOrdered(LOCTEXT("LinkParameters", "Link {0}|plural(one=parameter,other=parameters)"), ParametersToLink.Num()));
+
+						PendingSelection.Empty();
+
+						for(const FParameterBindingReference& ParameterToLink : ParametersToLink)
+						{
+							check(ParameterToLink.Block.IsValid());
+							check(ParameterToLink.Parameter != NAME_None);
+
+							// Load our assets
+							UAnimNextParameterBlock* ExistingBlock = CastChecked<UAnimNextParameterBlock>(ParameterToLink.Block.GetAsset());
+							UAnimNextParameterLibrary* ExistingLibrary = CastChecked<UAnimNextParameterLibrary>(ParameterToLink.Library.GetAsset());
+
+							// Create a new entry for the supplied parameter
+							UAnimNextParameterBlockBindingReference* BindingReference = EditorData->AddBindingReference(ParameterToLink.Parameter, ExistingLibrary, ExistingBlock);
+							PendingSelection.Add(BindingReference);
+						}
+					}
+
+					return FReply::Handled();
+				})
+			]
+			+SHorizontalBox::Slot()
+			.FillWidth(1.0f)
+			.Padding(2.0f)
+			[
+				SNew(SSearchBox)
+				.OnTextChanged_Lambda([this](FText InText)
+				{
+					FilterText = InText;
+					RefreshFilter();
+				})
+			]
+		]
+		+SVerticalBox::Slot()
+		.FillHeight(1.0f)
+		[
+			SAssignNew(EntriesList, SListView<TSharedRef<FParameterBlockViewEntry>>)
+			.ListItemsSource(&FilteredEntries)
+			.OnGenerateRow(this, &SParameterBlockView::HandleGenerateRow)
+			.OnItemScrolledIntoView(this, &SParameterBlockView::HandleItemScrolledIntoView)
+			.OnSelectionChanged(this, &SParameterBlockView::HandleSelectionChanged)
+			.ItemHeight(20.0f)
+			.HeaderRow(
+				SNew(SHeaderRow)
+				+SHeaderRow::Column(Column_RevisionControl)
+				.DefaultLabel(FText::GetEmpty())
+				.FixedWidth(24.0f)
+				.HeaderContent()
+				[
+					SNew(SBox)
+					.WidthOverride(16.0f)
+					.HeightOverride(16.0f)
+					.VAlign(VAlign_Center)
+					.HAlign(HAlign_Center)
+					[
+						SNew(SImage)
+						.ColorAndOpacity(FSlateColor::UseForeground())
+						.Image(FRevisionControlStyleManager::Get().GetBrush("RevisionControl.Icon"))
+						.ToolTipText(LOCTEXT("RevisionControlStatusTooltip", "Revision control status of this parameter"))
+					]
+				]
+
+				+SHeaderRow::Column(Column_ModifiedStatus)
+				.DefaultLabel(FText::GetEmpty())
+				.FixedWidth(24.0f)
+				.HeaderContent()
+				[
+					SNew(SBox)
+					.WidthOverride(16.0f)
+					.HeightOverride(16.0f)
+					.VAlign(VAlign_Center)
+					.HAlign(HAlign_Center)
+					[
+						SNew(SImage)
+						.ColorAndOpacity(FSlateColor::UseForeground())
+						.Image(FAppStyle::GetBrush("ContentBrowser.ContentDirty"))
+						.ToolTipText(LOCTEXT("ModifiedStatusTooltip", "Modified status of this parameter"))
+					]
+				]
+
+				+SHeaderRow::Column(Column_Name)
+				.DefaultLabel(LOCTEXT("NameColumnHeader", "Name"))
+				.ToolTipText(LOCTEXT("NameColumnHeaderTooltip", "The name of the parameter"))
+				.FillWidth(20.0f)
+
+				+SHeaderRow::Column(Column_Type)
+				.DefaultLabel(LOCTEXT("TypeColumnHeader", "Type"))
+				.ToolTipText(LOCTEXT("TypeColumnHeaderTooltip", "The type of the parameter"))
+				.ManualWidth(145.0f)
+
+				+SHeaderRow::Column(Column_Value)
+				.DefaultLabel(LOCTEXT("ValueColumnHeader", "Value"))
+				.ToolTipText(LOCTEXT("ValueColumnHeaderTooltip", "The value or binding to the parameter"))
+				.FillWidth(10.0f)
+			)
+		]
+	];
+
+	// Update revision control status of all our files
+	if (ISourceControlModule::Get().IsEnabled())
+	{
+		TArray<UPackage*> Packages = EditorData->GetPackage()->GetExternalPackages();
+
+		ISourceControlProvider& SourceControlProvider = ISourceControlModule::Get().GetProvider();
+		SourceControlProvider.Execute(ISourceControlOperation::Create<FUpdateStatus>(), Packages);
+	}
+
+	BindCommands();
+
+	RegisterActiveTimer(1.0f / 60.0f, FWidgetActiveTimerDelegate::CreateLambda(
+		[this](double InCurrentTime, float InDeltaTime)
+		{
+			if(bRefreshRequested)
+			{
+				RefreshEntries();
+				bRefreshRequested = false;
+			}
+
+			return EActiveTimerReturnType::Continue;
+		}));
+
+	RequestRefresh();
+}
+
+FReply SParameterBlockView::OnKeyDown(const FGeometry& MyGeometry, const FKeyEvent& InKeyEvent)
+{
+	if (UICommandList.IsValid() && UICommandList->ProcessCommandBindings(InKeyEvent))
+	{
+		return FReply::Handled();
+	}
+	return FReply::Unhandled();
+}
+
+void SParameterBlockView::RequestRefresh()
+{
+	bRefreshRequested = true;
+}
+
+void SParameterBlockView::RefreshEntries()
+{
+	Entries.Empty();
+
+	TArray<TSharedRef<FParameterBlockViewEntry>> EntriesToSelect;
+
+	for(UAnimNextParameterBlockEntry* Entry : EditorData->Entries)
+	{
+		TSharedRef<FParameterBlockViewEntry> NewEntry = MakeShared<FParameterBlockViewEntry>(Entry);
+		Entries.Add(NewEntry);
+
+		if(PendingSelection.Contains(Entry))
+		{
+			EntriesToSelect.Add(NewEntry);
+		}
+	}
+
+	PendingSelection.Empty();
+
+	RefreshFilter();
+
+	if(EntriesToSelect.Num() > 0)
+	{
+		EntriesList->SetItemSelection(EntriesToSelect, true);
+	}
+}
+
+void SParameterBlockView::RefreshFilter()
+{
+	FilteredEntries.Empty();
+
+	const FString FilterTextAsString = FilterText.ToString();
+	for(const TSharedRef<FParameterBlockViewEntry>& Entry : Entries)
+	{
+		if(Entry->PassesFilter(FilterTextAsString))
+		{
+			FilteredEntries.Add(Entry);
+		}
+	}
+
+	EntriesList->RequestListRefresh();
+}
+
+void SParameterBlockView::BindCommands()
+{
+	UICommandList = MakeShared<FUICommandList>();
+
+	UICommandList->MapAction(FGenericCommands::Get().Delete,
+		FExecuteAction::CreateSP(this, &SParameterBlockView::HandleDelete),
+		FCanExecuteAction::CreateSP(this, &SParameterBlockView::HasValidSelection));
+
+	UICommandList->MapAction(FGenericCommands::Get().Rename,
+		FExecuteAction::CreateSP(this, &SParameterBlockView::HandleRename),
+		FCanExecuteAction::CreateSP(this, &SParameterBlockView::HasValidSingleSelection));
+}
+
+void SParameterBlockView::HandleBlockModified(UAnimNextParameterBlock_EditorData* InEditorData)
+{
+	check(InEditorData == EditorData);
+
+	RequestRefresh();
+}
+
+TSharedRef<SWidget> SParameterBlockView::HandleGetContextContent()
+{
+	using namespace ParameterBlockView;
+	
+	UToolMenus* ToolMenus = UToolMenus::Get();
+
+	if(!ToolMenus->IsMenuRegistered(ContextMenuName))
+	{
+		UToolMenu* Menu = ToolMenus->RegisterMenu(ContextMenuName);
+
+		FToolMenuSection& Section = Menu->AddSection("EntryOperations", LOCTEXT("EntryOperationsMenuSection", "Block Entry"));
+		Section.AddMenuEntry(FGenericCommands::Get().Delete);
+		Section.AddMenuEntry(FGenericCommands::Get().Rename);
+	}
+
+	UParameterBlockViewMenuContext* MenuContext = NewObject<UParameterBlockViewMenuContext>();
+	MenuContext->ParameterBlockView = SharedThis(this);
+	return ToolMenus->GenerateWidget(ContextMenuName, FToolMenuContext(MenuContext));
+}
+
+void SParameterBlockView::HandleDelete()
+{
+	if(EntriesList->GetNumItemsSelected() > 0)
+	{
+		TArray<TSharedRef<FParameterBlockViewEntry>> SelectedItems = EntriesList->GetSelectedItems();
+
+		{
+			FScopedTransaction Transaction(FText::FormatOrdered(LOCTEXT("DeleteParameterBlockEntry", "Delete parameter block {0}|plural(one=entry,other=entries)"), EntriesList->GetNumItemsSelected()));
+
+			TArray<UAnimNextParameterBlockEntry*> EntriesToRemove;
+			Algo::Transform(SelectedItems, EntriesToRemove, [](const TSharedRef<FParameterBlockViewEntry>& InEntry){ return InEntry->WeakEntry.Get(); });
+			EditorData->RemoveEntries(EntriesToRemove);
+		}
+	}
+}
+
+void SParameterBlockView::HandleRename()
+{
+	if(EntriesList->GetNumItemsSelected() == 1)
+	{
+		TArray<TSharedRef<FParameterBlockViewEntry>> SelectedItems = EntriesList->GetSelectedItems();
+		SelectedItems[0]->bRenameWhenScrolledIntoView = true;
+		EntriesList->RequestScrollIntoView(SelectedItems[0]);
+	}
+}
+
+bool SParameterBlockView::HasValidSelection() const
+{
+	return EntriesList->GetNumItemsSelected() > 0;
+}
+
+bool SParameterBlockView::HasValidSingleSelection() const
+{
+	return EntriesList->GetNumItemsSelected() == 1;
+}
+
+void SParameterBlockView::HandleItemScrolledIntoView(TSharedRef<FParameterBlockViewEntry> InEntry, const TSharedPtr<ITableRow>& InWidget)
+{
+	if(InEntry->bRenameWhenScrolledIntoView)
+	{
+		InEntry->bRenameWhenScrolledIntoView = false;
+
+		if(TSharedPtr<SInlineEditableTextBlock> NameWidget = InEntry->NameWidget.Pin())
+		{
+			NameWidget->EnterEditingMode();
+		}
+	}
+}
+
+class SParameterBlockViewRow : public SMultiColumnTableRow<TSharedRef<FParameterBlockViewEntry>>
+{
+	SLATE_BEGIN_ARGS(SParameterBlockViewRow) {}
+	SLATE_END_ARGS()
+
+	void Construct(const FArguments& InArgs, const TSharedRef<STableViewBase>& InOwnerTableView, TSharedRef<SParameterBlockView> InView, TSharedRef<FParameterBlockViewEntry> InEntry)
+	{
+		WeakView = InView;
+		Entry = InEntry;
+
+		SMultiColumnTableRow<TSharedRef<FParameterBlockViewEntry>>::Construct( SMultiColumnTableRow<TSharedRef<FParameterBlockViewEntry>>::FArguments(), InOwnerTableView);
+	}
+
+	virtual TSharedRef<SWidget> GenerateWidgetForColumn(const FName& InColumnName) override
+	{
+		using namespace ParameterBlockView;
+		
+		if(InColumnName == Column_RevisionControl)
+		{
+			return
+				SNew(SBox)
+				.HAlign(HAlign_Center)
+				.VAlign(VAlign_Center)
+				.HeightOverride(20.0f)
+				[
+					SNew(SImage)
+					.ToolTipText_Lambda([this]()
+					{
+						if(UAnimNextParameterBlockEntry* BlockEntry = Entry->WeakEntry.Get())
+						{
+							ISourceControlProvider& SourceControlProvider = ISourceControlModule::Get().GetProvider();
+							UPackage* Package = BlockEntry->GetExternalPackage();
+							if(FSourceControlStatePtr State = SourceControlProvider.GetState(Package, EStateCacheUsage::Use))
+							{
+								return FText::Format(LOCTEXT("RevisionControlStatusFormat", "File: {0}\nStatus: {1}"), FText::FromName(Package->GetFName()),  State->GetDisplayTooltip());
+							}
+						}
+
+						return LOCTEXT("RevisionControlStatus", "Revision control status of this parameter");
+					})
+					.Image_Lambda([this]() -> const FSlateBrush*
+					{
+						if(UAnimNextParameterBlockEntry* BlockEntry = Entry->WeakEntry.Get())
+						{
+							ISourceControlProvider& SourceControlProvider = ISourceControlModule::Get().GetProvider();
+							if(FSourceControlStatePtr State = SourceControlProvider.GetState(BlockEntry->GetExternalPackage(), EStateCacheUsage::Use))
+							{
+								return State->GetIcon().GetSmallIcon();
+							}
+						}
+						return nullptr;
+					})
+				];
+		}
+		else if(InColumnName == Column_ModifiedStatus)
+		{
+			return
+				SNew(SBox)
+				.HAlign(HAlign_Center)
+				.VAlign(VAlign_Center)
+				.HeightOverride(20.0f)
+				[
+					SNew(SImage)
+					.ToolTipText_Lambda([this]()
+					{
+						FTextBuilder TextBuilder;
+						TextBuilder.AppendLine(LOCTEXT("ModifiedTooltip", "Modified Status"));
+
+						if(UAnimNextParameterBlockEntry* BlockEntry = Entry->WeakEntry.Get())
+						{
+							const UPackage* ExternalPackage = BlockEntry->GetExternalPackage();
+							check(ExternalPackage);
+							if(ExternalPackage->IsDirty())
+							{
+								TextBuilder.AppendLine(FText::FromName(ExternalPackage->GetFName()));
+							}
+
+							if(const IAnimNextParameterBlockBindingInterface* Binding = Cast<IAnimNextParameterBlockBindingInterface>(BlockEntry))
+							{
+								if(const UAnimNextParameter* Parameter = Binding->GetParameter())
+								{
+									const UPackage* ParameterPackage = Parameter->GetPackage();
+									check(ParameterPackage);
+									if(ParameterPackage->IsDirty())
+									{
+										TextBuilder.AppendLine(FText::FromName(ParameterPackage->GetFName()));
+									}
+								}
+							}
+						}
+
+						return TextBuilder.ToText();
+					})
+					.Image_Lambda([this]() -> const FSlateBrush*
+					{
+						if(UAnimNextParameterBlockEntry* BlockEntry = Entry->WeakEntry.Get())
+						{
+							bool bIsDirty = false;
+							const UPackage* ExternalPackage = BlockEntry->GetExternalPackage();
+							check(ExternalPackage);
+							if(ExternalPackage->IsDirty())
+							{
+								bIsDirty = true;
+							}
+
+							if(const IAnimNextParameterBlockBindingInterface* Binding = Cast<IAnimNextParameterBlockBindingInterface>(BlockEntry))
+							{
+								if(const UAnimNextParameter* Parameter = Binding->GetParameter())
+								{
+									const UPackage* ParameterPackage = Parameter->GetPackage();
+									check(ParameterPackage);
+									if(ParameterPackage->IsDirty())
+									{
+										bIsDirty = true;
+									}
+								}
+							}
+
+							return bIsDirty ? FAppStyle::GetBrush("ContentBrowser.ContentDirty") : nullptr;
+						}
+						return nullptr;
+					})
+				];
+		}
+		else if(InColumnName == Column_Type)
+		{
+			if(Entry->WeakEntry.Get()->Implements<UAnimNextParameterBlockBindingInterface>())
+			{
+				return
+					SNew(SBox)
+					.HAlign(HAlign_Left)
+					.VAlign(VAlign_Center)
+					.IsEnabled(false)
+					[
+						SNew(SPinTypeSelector, FGetPinTypeTree::CreateStatic(&Editor::FUtils::GetFilteredVariableTypeTree))
+							.TargetPinType_Lambda([this]()
+							{
+								if(IAnimNextParameterBlockBindingInterface* Binding = Cast<IAnimNextParameterBlockBindingInterface>(Entry->WeakEntry.Get()))
+								{
+									return UncookedOnly::FUtils::GetPinTypeFromParamType(Binding->GetParamType());
+								}
+
+								return FEdGraphPinType();
+							})
+							.Schema(GetDefault<UPropertyBagSchema>())
+							.bAllowArrays(true)
+							.TypeTreeFilter(ETypeTreeFilter::None)
+							.Font(IDetailLayoutBuilder::GetDetailFont())
+					];
+			}
+		}
+		else if(InColumnName == Column_Name)
+		{
+			return
+				SNew(SBox)
+				.HAlign(HAlign_Left)
+				.VAlign(VAlign_Center)
+				[
+					// TODO: make this into a picker appropriately
+					SAssignNew(Entry->NameWidget, SInlineEditableTextBlock)
+					.Font(IDetailLayoutBuilder::GetDetailFont())
+					.IsSelected(this, &SParameterBlockViewRow::IsSelectedExclusively)
+					.IsReadOnly_Lambda([this]()
+					{
+						return Cast<UAnimNextParameterBlockBinding>(Entry->WeakEntry.Get()) == nullptr;
+					})
+					.OnTextCommitted_Lambda([this](const FText& InNewText, ETextCommit::Type InCommitType)
+					{
+						if(InCommitType == ETextCommit::OnEnter)
+						{
+							if(IAnimNextParameterBlockBindingInterface* Binding = Cast<IAnimNextParameterBlockBindingInterface>(Entry->WeakEntry.Get()))
+							{
+								FScopedTransaction Transaction(LOCTEXT("SetParameteName", "Set parameter name"));
+
+								Binding->SetParameterName(*InNewText.ToString());
+							}
+						}
+					})
+					.OnVerifyTextChanged_Lambda([this](const FText& InNewText, FText& OutErrorText)
+					{
+						const FString NewString = InNewText.ToString();
+
+						if(IAnimNextParameterBlockBindingInterface* Binding = Cast<IAnimNextParameterBlockBindingInterface>(Entry->WeakEntry.Get()))
+						{
+							// Make sure the new name only contains valid characters
+							if (!FName::IsValidXName(NewString, INVALID_OBJECTNAME_CHARACTERS INVALID_LONGPACKAGE_CHARACTERS, &OutErrorText))
+							{
+								return false;
+							}
+
+							const FName Name(*NewString);
+							if(!FUtils::DoesParameterExistInLibrary(Binding->GetLibrary(), Name))
+							{
+								OutErrorText = LOCTEXT("Error_NameDoesNotExistInLibrary", "This name does not exist in the specified library");
+								return false;
+							}
+						}
+
+						if(IAnimNextParameterBlockReferenceInterface* Reference = Cast<IAnimNextParameterBlockReferenceInterface>(Entry->WeakEntry.Get()))
+						{
+							const FName Name(*NewString);
+							if(!FUtils::DoesParameterExistInLibrary(Reference->GetBlock(), Name))
+							{
+								OutErrorText = LOCTEXT("Error_NameDoesNotExist", "This name does not exist in the specified block");
+								return false;
+							}
+						}
+
+						return true;
+					})
+					.Text_Lambda([this]()
+					{
+						if(UAnimNextParameterBlockEntry* BlockEntry = Entry->WeakEntry.Get())
+						{
+							return BlockEntry->GetDisplayName();
+						}
+						return FText::GetEmpty();
+					})
+					.ToolTipText_Lambda([this]()
+					{
+						if(UAnimNextParameterBlockEntry* BlockEntry = Entry->WeakEntry.Get())
+						{
+							return BlockEntry->GetDisplayNameTooltip();
+						}
+						return FText::GetEmpty();
+					})
+				];
+		}
+		else if(InColumnName == Column_Value)
+		{
+			TSharedPtr<SHorizontalBox> HorizontalBox;
+
+			TSharedRef<SWidget> Widget =
+				SNew(SBox)
+				.HAlign(HAlign_Left)
+				.VAlign(VAlign_Center)
+				[
+					SAssignNew(HorizontalBox, SHorizontalBox)
+				];
+
+			if(UAnimNextParameterBlockEntry* BlockEntry = Entry->WeakEntry.Get())
+			{
+				if(IAnimNextParameterBlockReferenceInterface* ReferenceInterface = Cast<IAnimNextParameterBlockReferenceInterface>(BlockEntry))
+				{
+					const UAnimNextParameterBlock* ReferencedBlock = ReferenceInterface->GetBlock();
+					HorizontalBox->AddSlot()
+					[
+						SNew(STextBlock)
+						.Font(IDetailLayoutBuilder::GetDetailFont())
+						.IsEnabled(false)
+						.Text_Lambda([WeakBlock = TWeakObjectPtr<const UAnimNextParameterBlock>(ReferencedBlock)]()
+						{
+							if(const UAnimNextParameterBlock* Block = WeakBlock.Get())
+							{
+								return FText::FromString(Block->GetName());
+							}
+
+							return LOCTEXT("MissingParameterBlock", "Missing Parameter Block");
+						})
+						.ToolTipText_Lambda([WeakBlock = TWeakObjectPtr<const UAnimNextParameterBlock>(ReferencedBlock)]()
+						{
+							if(const UAnimNextParameterBlock* Block = WeakBlock.Get())
+							{
+								return FText::FromString(Block->GetPathName());
+							}
+
+							return LOCTEXT("MissingParameterBlock", "Missing Parameter Block");
+						})
+					];
+				}
+				
+				if(IAnimNextParameterBlockGraphInterface* GraphInterface = Cast<IAnimNextParameterBlockGraphInterface>(BlockEntry))
+				{
+					URigVMGraph* ReferencedGraph = GraphInterface->GetGraph();
+					HorizontalBox->AddSlot()
+					[
+						SNew(SSimpleButton)
+						.Icon(FAppStyle::GetBrush("Icons.Blueprints"))
+						.OnClicked_Lambda([this, WeakGraph = TWeakObjectPtr<URigVMGraph>(ReferencedGraph)]()
+						{
+							if(URigVMGraph* Graph = WeakGraph.Get())
+							{
+								if(TSharedPtr<SParameterBlockView> View = WeakView.Pin())
+								{
+									View->OnOpenGraphDelegate.ExecuteIfBound(Graph);
+								}
+							}
+							return FReply::Handled();
+						})
+					];
+				}
+			}
+			
+			return Widget;
+		}
+
+		return SNullWidget::NullWidget;
+	}
+
+	TWeakPtr<SParameterBlockView> WeakView;
+	TSharedPtr<FParameterBlockViewEntry> Entry;
+};
+
+TSharedRef<ITableRow> SParameterBlockView::HandleGenerateRow(TSharedRef<FParameterBlockViewEntry> InEntry, const TSharedRef<STableViewBase>& InOwnerTable)
+{
+	return SNew(SParameterBlockViewRow, InOwnerTable, SharedThis(this), InEntry);
+}
+
+void SParameterBlockView::HandleSelectionChanged(TSharedPtr<FParameterBlockViewEntry> InEntry, ESelectInfo::Type InSelectionType)
+{
+	if(OnSelectionChangedDelegate.IsBound())
+	{
+		TArray<UObject*> SelectedItems;
+		SelectedItems.Reserve(EntriesList->GetNumItemsSelected());
+		for(const TSharedRef<FParameterBlockViewEntry>& SelectedItem : EntriesList->GetSelectedItems())
+		{
+			if(UAnimNextParameterBlockEntry* Entry = SelectedItem->WeakEntry.Get())
+			{
+				SelectedItems.Add(Entry);
+			}
+		}
+
+		OnSelectionChangedDelegate.Execute(SelectedItems);
+	}
+}
+
+EFilterParameterResult SParameterBlockView::HandleFilterLinkedParameter(const FParameterBindingReference& InParameterBinding)
+{
+	// Don't display parameters that are already bound in this block
+	if(InParameterBinding.Block.IsValid())
+	{
+		if(BlockAssetData == InParameterBinding.Block)
+		{
+			return EFilterParameterResult::Exclude;
+		}
+	}
+
+	return EFilterParameterResult::Include;
+}
+
+}
+
+#undef LOCTEXT_NAMESPACE
\ No newline at end of file
diff --git a/Engine/Plugins/Experimental/AnimNext/Source/AnimNextEditor/Private/Param/SParameterBlockView.h b/Engine/Plugins/Experimental/AnimNext/Source/AnimNextEditor/Private/Param/SParameterBlockView.h
new file mode 100644
--- /dev/null
+++ b/Engine/Plugins/Experimental/AnimNext/Source/AnimNextEditor/Private/Param/SParameterBlockView.h
@@ -0,0 +1,103 @@
+﻿// Copyright Epic Games, Inc. All Rights Reserved.
+
+#pragma once
+
+#include "Widgets/SCompoundWidget.h"
+#include "Widgets/Views/SListView.h"
+
+class UAnimNextParameterBlockBinding;
+class UAnimNextParameter;
+class UAnimNextParameterBlock_EditorData;
+class UAnimNextParameterBlockEntry;
+class URigVMGraph;
+struct FAssetData;
+
+namespace UE::AnimNext::Editor
+{
+
+enum class EFilterParameterResult : int32;
+struct FParameterBindingReference;
+
+struct FParameterBlockViewEntry; 
+
+class SParameterBlockView : public SCompoundWidget
+{
+public:
+	DECLARE_DELEGATE_OneParam(FOnSelectionChanged, const TArray<UObject*>& /*InEntries*/);
+
+	DECLARE_DELEGATE_OneParam(FOnOpenGraph, URigVMGraph* /*InGraph*/);
+	
+	SLATE_BEGIN_ARGS(SParameterBlockView) {}
+
+	SLATE_EVENT(SParameterBlockView::FOnSelectionChanged, OnSelectionChanged)
+	
+	SLATE_EVENT(SParameterBlockView::FOnOpenGraph, OnOpenGraph)
+
+	SLATE_END_ARGS()
+
+	void Construct(const FArguments& InArgs, UAnimNextParameterBlock_EditorData* InEditorData);
+
+private:
+	virtual FReply OnKeyDown(const FGeometry& MyGeometry, const FKeyEvent& InKeyEvent) override;
+
+	void RequestRefresh();
+
+	void RefreshEntries();
+
+	void RefreshFilter();
+
+	// Bind input commands
+	void BindCommands();
+
+	// Handle modifications to the block
+	void HandleBlockModified(UAnimNextParameterBlock_EditorData* InEditorData);
+
+	// Get the content for the context menu
+	TSharedRef<SWidget> HandleGetContextContent();
+
+	void HandleDelete();
+
+	void HandleRename();
+
+	bool HasValidSelection() const;
+
+	bool HasValidSingleSelection() const;
+
+	// Generate a row for the list view
+	TSharedRef<ITableRow> HandleGenerateRow(TSharedRef<FParameterBlockViewEntry> InEntry, const TSharedRef<STableViewBase>& InOwnerTable);
+
+	// Handle rename after scrolling into view
+	void HandleItemScrolledIntoView(TSharedRef<FParameterBlockViewEntry> Entry, const TSharedPtr<ITableRow>& Widget);
+
+	// Handle selection
+	void HandleSelectionChanged(TSharedPtr<FParameterBlockViewEntry> InEntry, ESelectInfo::Type InSelectionType);
+
+	EFilterParameterResult HandleFilterLinkedParameter(const FParameterBindingReference& InParameterBinding);
+	
+private:
+	friend class SParameterBlockViewRow;
+	
+	TSharedPtr<SListView<TSharedRef<FParameterBlockViewEntry>>> EntriesList;
+
+	TArray<TSharedRef<FParameterBlockViewEntry>> Entries;
+
+	FText FilterText;
+
+	TArray<TSharedRef<FParameterBlockViewEntry>> FilteredEntries;
+
+	UAnimNextParameterBlock_EditorData* EditorData = nullptr;
+
+	TSharedPtr<FUICommandList> UICommandList = nullptr;
+
+	FOnSelectionChanged OnSelectionChangedDelegate;
+
+	FOnOpenGraph OnOpenGraphDelegate;
+
+	FAssetData BlockAssetData;
+
+	TArray<UAnimNextParameterBlockEntry*> PendingSelection;
+
+	bool bRefreshRequested = false;
+};
+
+}
\ No newline at end of file
diff --git a/Engine/Plugins/Experimental/AnimNext/Source/AnimNextEditor/Private/Param/SParameterLibraryView.cpp b/Engine/Plugins/Experimental/AnimNext/Source/AnimNextEditor/Private/Param/SParameterLibraryView.cpp
new file mode 100644
--- /dev/null
+++ b/Engine/Plugins/Experimental/AnimNext/Source/AnimNextEditor/Private/Param/SParameterLibraryView.cpp
@@ -0,0 +1,587 @@
+﻿// Copyright Epic Games, Inc. All Rights Reserved.
+
+#include "SParameterLibraryView.h"
+
+#include "Param/AnimNextParameterLibrary.h"
+#include "Param/AnimNextParameter.h"
+#include "Widgets/Layout/SBox.h"
+#include "Widgets/Input/SSearchBox.h"
+#include "DetailLayoutBuilder.h"
+#include "ISourceControlModule.h"
+#include "ISourceControlProvider.h"
+#include "EditorUtils.h"
+#include "UncookedOnlyUtils.h"
+#include "ParameterLibraryViewMenuContext.h"
+#include "PropertyBagDetails.h"
+#include "SAddParametersDialog.h"
+#include "SLinkParametersDialog.h"
+#include "SourceControlOperations.h"
+#include "SPinTypeSelector.h"
+#include "SSimpleButton.h"
+#include "Framework/Commands/GenericCommands.h"
+#include "RevisionControlStyle/RevisionControlStyle.h"
+#include "Widgets/Text/SInlineEditableTextBlock.h"
+
+#define LOCTEXT_NAMESPACE "AnimNextParameterLibraryView"
+
+namespace UE::AnimNext::Editor
+{
+
+namespace ParameterLibraryView
+{
+
+static FName ContextMenuName(TEXT("AnimNext.ParameterLibraryView.ContextMenu"));
+static FName Column_RevisionControl(TEXT("RevisionControl"));
+static FName Column_ModifiedStatus(TEXT("ModifiedStatus"));
+static FName Column_Name(TEXT("Name"));
+static FName Column_Type(TEXT("Type"));
+
+}
+
+// An entry displayed in the parameters view
+struct FParameterLibraryViewEntry
+{
+	FParameterLibraryViewEntry() = default;
+
+	explicit FParameterLibraryViewEntry(UAnimNextParameter* InEntry)
+		: WeakEntry(InEntry)
+	{}
+
+	bool PassesFilter(const FString& InFilterText) const
+	{
+		if(UAnimNextParameter* Entry = WeakEntry.Get())
+		{
+			return Entry->GetName().Contains(InFilterText);
+		}
+
+		return false;
+	}
+
+	// Ptr to the underlying entry
+	TWeakObjectPtr<UAnimNextParameter> WeakEntry;
+
+	// Widget used to rename items
+	TWeakPtr<SInlineEditableTextBlock> NameWidget;
+
+	// Flag to indicate a rename was requested
+	bool bRenameWhenScrolledIntoView = false;
+};
+
+void SParameterLibraryView::Construct(const FArguments& InArgs, UAnimNextParameterLibrary* InLibrary)
+{
+	using namespace ParameterLibraryView;
+	
+	check(InLibrary);
+
+	Library = InLibrary;
+
+	OnSelectionChangedDelegate = InArgs._OnSelectionChanged;
+
+	Library->ModifiedDelegate.AddSP(this, &SParameterLibraryView::HandleLibraryModified);
+
+	ChildSlot
+	[
+		SNew(SVerticalBox)
+		+SVerticalBox::Slot()
+		.AutoHeight()
+		[
+			SNew(SHorizontalBox)
+			+SHorizontalBox::Slot()
+			.AutoWidth()
+			.Padding(2.0f)
+			[
+				SNew(SSimpleButton)
+				.Text(LOCTEXT("AddNewParameterButton", "New"))
+				.Icon(FAppStyle::Get().GetBrush("Icons.Plus"))
+				.OnClicked_Lambda([this]()
+				{
+					TSharedRef<SAddParametersDialog> AddParametersDialog =
+						SNew(SAddParametersDialog)
+						.Library(Library);
+					TArray<FParameterToAdd> ParametersToAdd;
+					if(AddParametersDialog->ShowModal(ParametersToAdd))
+					{
+						FScopedTransaction Transaction(FText::FormatOrdered(LOCTEXT("AddParameters", "Add {0}|plural(one=parameter,other=parameters)"), ParametersToAdd.Num()));
+
+						PendingSelection.Empty();
+
+						for(const FParameterToAdd& ParameterToAdd : ParametersToAdd)
+						{
+							// Create a new parameter in the library
+							Library->AddParameter(ParameterToAdd.Name, ParameterToAdd.Type);
+						}
+					}
+
+					return FReply::Handled();
+				})
+			]
+			+SHorizontalBox::Slot()
+			.FillWidth(1.0f)
+			.Padding(2.0f)
+			[
+				SNew(SSearchBox)
+				.OnTextChanged_Lambda([this](FText InText)
+				{
+					FilterText = InText;
+					RefreshFilter();
+				})
+			]
+		]
+		+SVerticalBox::Slot()
+		.FillHeight(1.0f)
+		[
+			SAssignNew(EntriesList, SListView<TSharedRef<FParameterLibraryViewEntry>>)
+			.ListItemsSource(&FilteredEntries)
+			.OnGenerateRow(this, &SParameterLibraryView::HandleGenerateRow)
+			.OnItemScrolledIntoView(this, &SParameterLibraryView::HandleItemScrolledIntoView)
+			.OnSelectionChanged(this, &SParameterLibraryView::HandleSelectionChanged)
+			.ItemHeight(20.0f)
+			.HeaderRow(
+				SNew(SHeaderRow)
+				+SHeaderRow::Column(Column_RevisionControl)
+				.DefaultLabel(FText::GetEmpty())
+				.FixedWidth(24.0f)
+				.HeaderContent()
+				[
+					SNew(SBox)
+					.WidthOverride(16.0f)
+					.HeightOverride(16.0f)
+					.VAlign(VAlign_Center)
+					.HAlign(HAlign_Center)
+					[
+						SNew(SImage)
+						.ColorAndOpacity(FSlateColor::UseForeground())
+						.Image(FRevisionControlStyleManager::Get().GetBrush("RevisionControl.Icon"))
+						.ToolTipText(LOCTEXT("RevisionControlStatusTooltip", "Revision control status of this parameter"))
+					]
+				]
+
+				+SHeaderRow::Column(Column_ModifiedStatus)
+				.DefaultLabel(FText::GetEmpty())
+				.FixedWidth(24.0f)
+				.HeaderContent()
+				[
+					SNew(SBox)
+					.WidthOverride(16.0f)
+					.HeightOverride(16.0f)
+					.VAlign(VAlign_Center)
+					.HAlign(HAlign_Center)
+					[
+						SNew(SImage)
+						.ColorAndOpacity(FSlateColor::UseForeground())
+						.Image(FAppStyle::GetBrush("ContentBrowser.ContentDirty"))
+						.ToolTipText(LOCTEXT("ModifiedStatusTooltip", "Modified status of this parameter"))
+					]
+				]
+
+				+SHeaderRow::Column(Column_Name)
+				.DefaultLabel(LOCTEXT("NameColumnHeader", "Name"))
+				.ToolTipText(LOCTEXT("NameColumnHeaderTooltip", "The name of the parameter"))
+				.FillWidth(20.0f)
+
+				+SHeaderRow::Column(Column_Type)
+				.DefaultLabel(LOCTEXT("TypeColumnHeader", "Type"))
+				.ToolTipText(LOCTEXT("TypeColumnHeaderTooltip", "The type of the parameter"))
+				.ManualWidth(145.0f)
+			)
+		]
+	];
+
+	// Update revision control status of all our files
+	if (ISourceControlModule::Get().IsEnabled())
+	{
+		TArray<UPackage*> Packages = Library->GetPackage()->GetExternalPackages();
+
+		ISourceControlProvider& SourceControlProvider = ISourceControlModule::Get().GetProvider();
+		SourceControlProvider.Execute(ISourceControlOperation::Create<FUpdateStatus>(), Packages);
+	}
+
+	BindCommands();
+
+	RegisterActiveTimer(1.0f / 60.0f, FWidgetActiveTimerDelegate::CreateLambda(
+		[this](double InCurrentTime, float InDeltaTime)
+		{
+			if(bRefreshRequested)
+			{
+				RefreshEntries();
+				bRefreshRequested = false;
+			}
+
+			return EActiveTimerReturnType::Continue;
+		}));
+
+	RequestRefresh();
+}
+
+FReply SParameterLibraryView::OnKeyDown(const FGeometry& MyGeometry, const FKeyEvent& InKeyEvent)
+{
+	if (UICommandList.IsValid() && UICommandList->ProcessCommandBindings(InKeyEvent))
+	{
+		return FReply::Handled();
+	}
+	return FReply::Unhandled();
+}
+
+void SParameterLibraryView::RequestRefresh()
+{
+	bRefreshRequested = true;
+}
+
+void SParameterLibraryView::RefreshEntries()
+{
+	Entries.Empty();
+
+	TArray<TSharedRef<FParameterLibraryViewEntry>> EntriesToSelect;
+
+	for(UAnimNextParameter* Entry : Library->Parameters)
+	{
+		TSharedRef<FParameterLibraryViewEntry> NewEntry = MakeShared<FParameterLibraryViewEntry>(Entry);
+		Entries.Add(NewEntry);
+
+		if(PendingSelection.Contains(Entry))
+		{
+			EntriesToSelect.Add(NewEntry);
+		}
+	}
+
+	PendingSelection.Empty();
+
+	RefreshFilter();
+
+	if(EntriesToSelect.Num() > 0)
+	{
+		EntriesList->SetItemSelection(EntriesToSelect, true);
+	}
+}
+
+void SParameterLibraryView::RefreshFilter()
+{
+	FilteredEntries.Empty();
+
+	const FString FilterTextAsString = FilterText.ToString();
+	for(const TSharedRef<FParameterLibraryViewEntry>& Entry : Entries)
+	{
+		if(Entry->PassesFilter(FilterTextAsString))
+		{
+			FilteredEntries.Add(Entry);
+		}
+	}
+
+	EntriesList->RequestListRefresh();
+}
+
+void SParameterLibraryView::BindCommands()
+{
+	UICommandList = MakeShared<FUICommandList>();
+
+	UICommandList->MapAction(FGenericCommands::Get().Delete,
+		FExecuteAction::CreateSP(this, &SParameterLibraryView::HandleDelete),
+		FCanExecuteAction::CreateSP(this, &SParameterLibraryView::HasValidSelection));
+
+	UICommandList->MapAction(FGenericCommands::Get().Rename,
+		FExecuteAction::CreateSP(this, &SParameterLibraryView::HandleRename),
+		FCanExecuteAction::CreateSP(this, &SParameterLibraryView::HasValidSingleSelection));
+}
+
+void SParameterLibraryView::HandleLibraryModified(UAnimNextParameterLibrary* InLibrary)
+{
+	check(InLibrary == Library);
+
+	RequestRefresh();
+}
+
+TSharedRef<SWidget> SParameterLibraryView::HandleGetContextContent()
+{
+	using namespace ParameterLibraryView;
+	
+	UToolMenus* ToolMenus = UToolMenus::Get();
+
+	if(!ToolMenus->IsMenuRegistered(ContextMenuName))
+	{
+		UToolMenu* Menu = ToolMenus->RegisterMenu(ContextMenuName);
+
+		FToolMenuSection& Section = Menu->AddSection("EntryOperations", LOCTEXT("EntryOperationsMenuSection", "Parameter"));
+		Section.AddMenuEntry(FGenericCommands::Get().Delete);
+		Section.AddMenuEntry(FGenericCommands::Get().Rename);
+	}
+
+	UParameterLibraryViewMenuContext* MenuContext = NewObject<UParameterLibraryViewMenuContext>();
+	MenuContext->ParameterLibraryView = SharedThis(this);
+	return ToolMenus->GenerateWidget(ContextMenuName, FToolMenuContext(MenuContext));
+}
+
+void SParameterLibraryView::HandleDelete()
+{
+	if(EntriesList->GetNumItemsSelected() > 0)
+	{
+		TArray<TSharedRef<FParameterLibraryViewEntry>> SelectedItems = EntriesList->GetSelectedItems();
+
+		{
+			FScopedTransaction Transaction(FText::FormatOrdered(LOCTEXT("DeleteParameter", "Delete {0}|plural(one=parameter,other=parameters)"), EntriesList->GetNumItemsSelected()));
+
+			TArray<FName> EntriesToRemove;
+			Algo::Transform(SelectedItems, EntriesToRemove, [](const TSharedRef<FParameterLibraryViewEntry>& InEntry){ return InEntry->WeakEntry.Get()->GetFName(); });
+			Library->RemoveParameters(EntriesToRemove);
+		}
+	}
+}
+
+void SParameterLibraryView::HandleRename()
+{
+	if(EntriesList->GetNumItemsSelected() == 1)
+	{
+		TArray<TSharedRef<FParameterLibraryViewEntry>> SelectedItems = EntriesList->GetSelectedItems();
+		SelectedItems[0]->bRenameWhenScrolledIntoView = true;
+		EntriesList->RequestScrollIntoView(SelectedItems[0]);
+	}
+}
+
+bool SParameterLibraryView::HasValidSelection() const
+{
+	return EntriesList->GetNumItemsSelected() > 0;
+}
+
+bool SParameterLibraryView::HasValidSingleSelection() const
+{
+	return EntriesList->GetNumItemsSelected() == 1;
+}
+
+void SParameterLibraryView::HandleItemScrolledIntoView(TSharedRef<FParameterLibraryViewEntry> InEntry, const TSharedPtr<ITableRow>& InWidget)
+{
+	if(InEntry->bRenameWhenScrolledIntoView)
+	{
+		InEntry->bRenameWhenScrolledIntoView = false;
+
+		if(TSharedPtr<SInlineEditableTextBlock> NameWidget = InEntry->NameWidget.Pin())
+		{
+			NameWidget->EnterEditingMode();
+		}
+	}
+}
+
+class SParameterLibraryViewRow : public SMultiColumnTableRow<TSharedRef<FParameterLibraryViewEntry>>
+{
+	SLATE_BEGIN_ARGS(SParameterLibraryViewRow) {}
+	SLATE_END_ARGS()
+
+	void Construct(const FArguments& InArgs, const TSharedRef<STableViewBase>& InOwnerTableView, TSharedRef<SParameterLibraryView> InView, TSharedRef<FParameterLibraryViewEntry> InEntry)
+	{
+		WeakView = InView;
+		Entry = InEntry;
+
+		SMultiColumnTableRow<TSharedRef<FParameterLibraryViewEntry>>::Construct( SMultiColumnTableRow<TSharedRef<FParameterLibraryViewEntry>>::FArguments(), InOwnerTableView);
+	}
+
+	virtual TSharedRef<SWidget> GenerateWidgetForColumn(const FName& InColumnName) override
+	{
+		using namespace ParameterLibraryView;
+		
+		if(InColumnName == Column_RevisionControl)
+		{
+			return
+				SNew(SBox)
+				.HAlign(HAlign_Center)
+				.VAlign(VAlign_Center)
+				.HeightOverride(20.0f)
+				[
+					SNew(SImage)
+					.ToolTipText_Lambda([this]()
+					{
+						if(UAnimNextParameter* Parameter = Entry->WeakEntry.Get())
+						{
+							ISourceControlProvider& SourceControlProvider = ISourceControlModule::Get().GetProvider();
+							UPackage* Package = Parameter->GetExternalPackage();
+							if(FSourceControlStatePtr State = SourceControlProvider.GetState(Package, EStateCacheUsage::Use))
+							{
+								return FText::Format(LOCTEXT("RevisionControlStatusFormat", "File: {0}\nStatus: {1}"), FText::FromName(Package->GetFName()),  State->GetDisplayTooltip());
+							}
+						}
+
+						return LOCTEXT("RevisionControlStatus", "Revision control status of this parameter");
+					})
+					.Image_Lambda([this]() -> const FSlateBrush*
+					{
+						if(UAnimNextParameter* Parameter = Entry->WeakEntry.Get())
+						{
+							ISourceControlProvider& SourceControlProvider = ISourceControlModule::Get().GetProvider();
+							if(FSourceControlStatePtr State = SourceControlProvider.GetState(Parameter->GetExternalPackage(), EStateCacheUsage::Use))
+							{
+								return State->GetIcon().GetSmallIcon();
+							}
+						}
+						return nullptr;
+					})
+				];
+		}
+		else if(InColumnName == Column_ModifiedStatus)
+		{
+			return
+				SNew(SBox)
+				.HAlign(HAlign_Center)
+				.VAlign(VAlign_Center)
+				.HeightOverride(20.0f)
+				[
+					SNew(SImage)
+					.ToolTipText_Lambda([this]()
+					{
+						FTextBuilder TextBuilder;
+						TextBuilder.AppendLine(LOCTEXT("ModifiedTooltip", "Modified Status"));
+
+						if(UAnimNextParameter* Parameter = Entry->WeakEntry.Get())
+						{
+							const UPackage* ExternalPackage = Parameter->GetExternalPackage();
+							check(ExternalPackage);
+							if(ExternalPackage->IsDirty())
+							{
+								TextBuilder.AppendLine(FText::FromName(ExternalPackage->GetFName()));
+							}
+						}
+
+						return TextBuilder.ToText();
+					})
+					.Image_Lambda([this]() -> const FSlateBrush*
+					{
+						if(UAnimNextParameter* Parameter = Entry->WeakEntry.Get())
+						{
+							bool bIsDirty = false;
+							const UPackage* ExternalPackage = Parameter->GetExternalPackage();
+							check(ExternalPackage);
+							if(ExternalPackage->IsDirty())
+							{
+								bIsDirty = true;
+							}
+
+							return bIsDirty ? FAppStyle::GetBrush("ContentBrowser.ContentDirty") : nullptr;
+						}
+						return nullptr;
+					})
+				];
+		}
+		else if(InColumnName == Column_Type)
+		{
+			return
+				SNew(SBox)
+				.HAlign(HAlign_Left)
+				.VAlign(VAlign_Center)
+				[
+					SNew(SPinTypeSelector, FGetPinTypeTree::CreateStatic(&Editor::FUtils::GetFilteredVariableTypeTree))
+						.TargetPinType_Lambda([this]()
+						{
+							if(UAnimNextParameter* Parameter = Entry->WeakEntry.Get())
+							{
+								return UncookedOnly::FUtils::GetPinTypeFromParamType(Parameter->GetType());
+							}
+
+							return FEdGraphPinType();
+						})
+						.OnPinTypeChanged_Lambda([this](const FEdGraphPinType& InPinType)
+						{
+							if(UAnimNextParameter* Parameter = Entry->WeakEntry.Get())
+							{
+								FScopedTransaction Transaction(LOCTEXT("SetParameterType", "Set parameter type"));
+
+								Parameter->SetType(UncookedOnly::FUtils::GetParamTypeFromPinType(InPinType));
+							}
+						})
+						.Schema(GetDefault<UPropertyBagSchema>())
+						.bAllowArrays(true)
+						.TypeTreeFilter(ETypeTreeFilter::None)
+						.Font(IDetailLayoutBuilder::GetDetailFont())
+				];
+		}
+		else if(InColumnName == Column_Name)
+		{
+			return
+				SNew(SBox)
+				.HAlign(HAlign_Left)
+				.VAlign(VAlign_Center)
+				[
+					SAssignNew(Entry->NameWidget, SInlineEditableTextBlock)
+					.Font(IDetailLayoutBuilder::GetDetailFont())
+					.IsSelected(this, &SParameterLibraryViewRow::IsSelectedExclusively)
+					.OnTextCommitted_Lambda([this](const FText& InNewText, ETextCommit::Type InCommitType)
+					{
+						if(InCommitType == ETextCommit::OnEnter)
+						{
+							if(UAnimNextParameter* Parameter = Entry->WeakEntry.Get())
+							{
+								FScopedTransaction Transaction(LOCTEXT("SetParameterName", "Set parameter name"));
+
+								Parameter->Rename(*InNewText.ToString());
+							}
+						}
+					})
+					.OnVerifyTextChanged_Lambda([this](const FText& InNewText, FText& OutErrorText)
+					{
+						const FString NewString = InNewText.ToString();
+
+						if(UAnimNextParameter* Parameter = Entry->WeakEntry.Get())
+						{
+							// Make sure the new name only contains valid characters
+							if (!FName::IsValidXName(NewString, INVALID_OBJECTNAME_CHARACTERS INVALID_LONGPACKAGE_CHARACTERS, &OutErrorText))
+							{
+								return false;
+							}
+
+							const FName Name(*NewString);
+							if(FUtils::DoesParameterExistInLibrary(WeakView.Pin()->Library, Name))
+							{
+								OutErrorText = LOCTEXT("Error_NameExistsInLibrary", "This name already exists in this library");
+								return false;
+							}
+						}
+
+						return true;
+					})
+					.Text_Lambda([this]()
+					{
+						if(UAnimNextParameter* Parameter = Entry->WeakEntry.Get())
+						{
+							return FText::FromName(Parameter->GetFName());
+						}
+						return FText::GetEmpty();
+					})
+					.ToolTipText_Lambda([this]()
+					{
+						if(UAnimNextParameter* Parameter = Entry->WeakEntry.Get())
+						{
+							return FText::FromStringView(Parameter->GetComment());
+						}
+						return FText::GetEmpty();
+					})
+				];
+		}
+
+		return SNullWidget::NullWidget;
+	}
+
+	TWeakPtr<SParameterLibraryView> WeakView;
+	TSharedPtr<FParameterLibraryViewEntry> Entry;
+};
+
+TSharedRef<ITableRow> SParameterLibraryView::HandleGenerateRow(TSharedRef<FParameterLibraryViewEntry> InEntry, const TSharedRef<STableViewBase>& InOwnerTable)
+{
+	return SNew(SParameterLibraryViewRow, InOwnerTable, SharedThis(this), InEntry);
+}
+
+void SParameterLibraryView::HandleSelectionChanged(TSharedPtr<FParameterLibraryViewEntry> InEntry, ESelectInfo::Type InSelectionType)
+{
+	if(OnSelectionChangedDelegate.IsBound())
+	{
+		TArray<UObject*> SelectedItems;
+		SelectedItems.Reserve(EntriesList->GetNumItemsSelected());
+		for(const TSharedRef<FParameterLibraryViewEntry>& SelectedItem : EntriesList->GetSelectedItems())
+		{
+			if(UAnimNextParameter* Parameter = SelectedItem->WeakEntry.Get())
+			{
+				SelectedItems.Add(Parameter);
+			}
+		}
+
+		OnSelectionChangedDelegate.Execute(SelectedItems);
+	}
+}
+
+}
+
+#undef LOCTEXT_NAMESPACE
\ No newline at end of file
diff --git a/Engine/Plugins/Experimental/AnimNext/Source/AnimNextEditor/Private/Param/SParameterLibraryView.h b/Engine/Plugins/Experimental/AnimNext/Source/AnimNextEditor/Private/Param/SParameterLibraryView.h
new file mode 100644
--- /dev/null
+++ b/Engine/Plugins/Experimental/AnimNext/Source/AnimNextEditor/Private/Param/SParameterLibraryView.h
@@ -0,0 +1,92 @@
+﻿// Copyright Epic Games, Inc. All Rights Reserved.
+
+#pragma once
+
+#include "Widgets/SCompoundWidget.h"
+#include "Widgets/Views/SListView.h"
+
+class UAnimNextParameterLibrary;
+class UAnimNextParameter;
+struct FAssetData;
+
+namespace UE::AnimNext::Editor
+{
+
+enum class EFilterParameterResult : int32;
+struct FParameterBindingReference;
+
+
+
+struct FParameterLibraryViewEntry; 
+
+class SParameterLibraryView : public SCompoundWidget
+{
+public:
+	DECLARE_DELEGATE_OneParam(FOnSelectionChanged, const TArray<UObject*>& /*InEntries*/);
+
+	SLATE_BEGIN_ARGS(SParameterLibraryView) {}
+
+	SLATE_EVENT(SParameterLibraryView::FOnSelectionChanged, OnSelectionChanged)
+
+	SLATE_END_ARGS()
+
+	void Construct(const FArguments& InArgs, UAnimNextParameterLibrary* InLibrary);
+
+private:
+	virtual FReply OnKeyDown(const FGeometry& MyGeometry, const FKeyEvent& InKeyEvent) override;
+
+	void RequestRefresh();
+
+	void RefreshEntries();
+
+	void RefreshFilter();
+
+	// Bind input commands
+	void BindCommands();
+
+	// Handle modifications to the library
+	void HandleLibraryModified(UAnimNextParameterLibrary* InEditorData);
+
+	// Get the content for the context menu
+	TSharedRef<SWidget> HandleGetContextContent();
+
+	void HandleDelete();
+
+	void HandleRename();
+
+	bool HasValidSelection() const;
+
+	bool HasValidSingleSelection() const;
+
+	// Generate a row for the list view
+	TSharedRef<ITableRow> HandleGenerateRow(TSharedRef<FParameterLibraryViewEntry> InEntry, const TSharedRef<STableViewBase>& InOwnerTable);
+
+	// Handle rename after scrolling into view
+	void HandleItemScrolledIntoView(TSharedRef<FParameterLibraryViewEntry> Entry, const TSharedPtr<ITableRow>& Widget);
+
+	// Handle selection
+	void HandleSelectionChanged(TSharedPtr<FParameterLibraryViewEntry> InEntry, ESelectInfo::Type InSelectionType);
+
+private:
+	friend class SParameterLibraryViewRow;
+	
+	TSharedPtr<SListView<TSharedRef<FParameterLibraryViewEntry>>> EntriesList;
+
+	TArray<TSharedRef<FParameterLibraryViewEntry>> Entries;
+
+	FText FilterText;
+
+	TArray<TSharedRef<FParameterLibraryViewEntry>> FilteredEntries;
+
+	UAnimNextParameterLibrary* Library = nullptr;
+
+	TSharedPtr<FUICommandList> UICommandList = nullptr;
+
+	FOnSelectionChanged OnSelectionChangedDelegate;
+
+	TArray<UAnimNextParameter*> PendingSelection;
+
+	bool bRefreshRequested = false;
+};
+
+}
\ No newline at end of file
diff --git a/Engine/Plugins/Experimental/AnimNext/Source/AnimNextEditor/Private/Param/SParameterPicker.cpp b/Engine/Plugins/Experimental/AnimNext/Source/AnimNextEditor/Private/Param/SParameterPicker.cpp
new file mode 100644
--- /dev/null
+++ b/Engine/Plugins/Experimental/AnimNext/Source/AnimNextEditor/Private/Param/SParameterPicker.cpp
@@ -0,0 +1,323 @@
+// Copyright Epic Games, Inc. All Rights Reserved.
+
+#include "SParameterPicker.h"
+
+#include "Param/AnimNextParameterBlock.h"
+#include "UncookedOnlyUtils.h"
+#include "AssetRegistry/AssetRegistryModule.h"
+#include "Kismet2/BlueprintEditorUtils.h"
+#include "Param/ParamType.h"
+#include "Param/AnimNextParameter.h"
+#include "Param/AnimNextParameterBlock_EditorData.h"
+#include "Param/AnimNextParameterLibrary.h"
+#include "DetailLayoutBuilder.h"
+#include "EditorUtils.h"
+#include "Param/ParameterPickerArgs.h"
+#include "Widgets/Input/SSearchBox.h"
+
+#define LOCTEXT_NAMESPACE "SParameterPicker"
+
+namespace UE::AnimNext::Editor
+{
+
+namespace ParameterPicker
+{
+static FName Column_Parameter(TEXT("Parameter"));
+static FName Column_Library(TEXT("Library"));
+static FName Column_Block(TEXT("Block"));
+}
+
+struct FParameterPickerEntry
+{
+	FParameterPickerEntry() = default;
+
+	FParameterPickerEntry(const FParameterBindingReference& InBinding, const FAnimNextParamType& InParamType)
+		: Binding(InBinding)
+		, ParamType(InParamType)
+	{
+		PinType = UE::AnimNext::UncookedOnly::FUtils::GetPinTypeFromParamType(ParamType);
+		PinIcon = FBlueprintEditorUtils::GetIconFromPin(PinType, /* bIsLarge = */true);
+		PinColor = GetDefault<UEdGraphSchema_K2>()->GetPinTypeColor(PinType);
+	}
+
+	bool PassesFilter(const FString& InFilterText) const
+	{
+		return Binding.Parameter.ToString().Contains(InFilterText);
+	}
+
+	FParameterBindingReference Binding;
+
+	FAnimNextParamType ParamType;
+
+	FEdGraphPinType PinType;
+
+	const FSlateBrush* PinIcon;
+
+	FLinearColor PinColor;
+};
+
+void SParameterPicker::Construct(const FArguments& InArgs)
+{
+	using namespace ParameterPicker;
+
+	Args = InArgs._Args;
+
+	if(Args.OnGetParameterBindings != nullptr)
+	{
+		Args.OnGetParameterBindings->BindSP(this, &SParameterPicker::HandleGetParameterBindings);
+	}
+
+	ChildSlot
+	[
+		SNew(SVerticalBox)
+		+SVerticalBox::Slot()
+		.AutoHeight()
+		.Padding(2.0f)
+		[
+			SNew(SSearchBox)
+			.OnTextChanged_Lambda([this](FText InText)
+			{
+				FilterText = InText;
+				RefreshFilter();
+			})
+		]
+		+SVerticalBox::Slot()
+		.FillHeight(1.0f)
+		.Padding(2.0f)
+		[
+			SAssignNew(EntriesList, SListView<TSharedRef<FParameterPickerEntry>>)
+			.ListItemsSource(&FilteredEntries)
+			.OnGenerateRow(this, &SParameterPicker::HandleGenerateRow)
+			.OnSelectionChanged(this, &SParameterPicker::HandleSelectionChanged)
+			.ItemHeight(20.0f)
+			.HeaderRow(
+				SNew(SHeaderRow)
+				+SHeaderRow::Column(Column_Parameter)
+				.DefaultLabel(LOCTEXT("ParameterColumnHeader", "Parameter"))
+				.ToolTipText(LOCTEXT("ParameterColumnHeaderTooltip", "The parameter"))
+				.FillWidth(0.33f)
+				
+				+SHeaderRow::Column(Column_Library)
+				.DefaultLabel(LOCTEXT("LibraryColumnHeader", "Library"))
+				.ToolTipText(LOCTEXT("LibraryColumnHeaderTooltip", "The library that the parameter is declared in"))
+				.FillWidth(0.33f)
+				
+				+SHeaderRow::Column(Column_Block)
+				.DefaultLabel(LOCTEXT("BlockColumnHeader", "Block"))
+				.ToolTipText(LOCTEXT("BlockColumnHeaderTooltip", "The parameter block that has a binding to the parameter"))
+				.FillWidth(0.33f)
+			)
+		]
+	];
+
+	RefreshEntries();
+}
+
+void SParameterPicker::RefreshEntries()
+{
+	Entries.Empty();
+
+	IAssetRegistry& AssetRegistry = FModuleManager::Get().LoadModuleChecked<FAssetRegistryModule>(TEXT("AssetRegistry")).Get();
+
+	FARFilter ARFilter;
+
+	TSet<TTuple<FName, FAssetData>> BoundParameters;
+
+	TMap<FAssetData, FAnimNextParameterLibraryAssetRegistryExports> LibraryExportMap;
+
+	auto GetExportsForLibrary = [&LibraryExportMap](const FAssetData& InLibrary) -> const FAnimNextParameterLibraryAssetRegistryExports&
+	{
+		if(FAnimNextParameterLibraryAssetRegistryExports* ExistingExports = LibraryExportMap.Find(InLibrary))
+		{
+			return *ExistingExports;
+		}
+
+		FAnimNextParameterLibraryAssetRegistryExports& NewExports = LibraryExportMap.Add(InLibrary);
+		FUtils::GetExportedParametersForLibrary(InLibrary, NewExports);
+		return NewExports;
+	};
+	
+	// Find all blocks and their bound parameters
+	if(Args.bShowBoundParameters)
+	{
+		ARFilter.ClassPaths = { UAnimNextParameterBlock::StaticClass()->GetClassPathName() };
+		
+		TArray<FAssetData> BlockAssets;
+		AssetRegistry.GetAssets(ARFilter, BlockAssets);
+
+		for(const FAssetData& BlockAsset : BlockAssets)
+		{
+			FAnimNextParameterBlockAssetRegistryExports Exports;
+			if(FUtils::GetExportedBindingsForBlock(BlockAsset, Exports))
+			{
+				for(const FAnimNextParameterBlockAssetRegistryExportEntry& Export : Exports.Bindings)
+				{
+					const FSoftObjectPath LibraryAssetPath(Export.Library);
+					FAssetData LibraryAsset = AssetRegistry.GetAssetByObjectPath(LibraryAssetPath);
+					if(LibraryAsset.IsValid())
+					{
+						BoundParameters.Add({ Export.Name, LibraryAsset });
+
+						FParameterBindingReference NewReference(Export.Name, LibraryAsset, BlockAsset);
+						if(!Args.OnFilterParameter.IsBound() || Args.OnFilterParameter.Execute(NewReference) == EFilterParameterResult::Include)
+						{
+							const FAnimNextParameterLibraryAssetRegistryExports& LibraryExports = GetExportsForLibrary(LibraryAsset);
+							FAnimNextParamType ParamType = FUtils::GetParameterTypeFromLibraryExports(Export.Name, LibraryExports);
+
+							TSharedRef<FParameterPickerEntry> NewEntry = MakeShared<FParameterPickerEntry>(NewReference, ParamType);
+							Entries.Add(NewEntry);
+						}
+					}
+				}
+			};
+		}
+	}
+
+	// Find all library parameters (that have not already been added as bound above)
+	if(Args.bShowUnboundParameters)
+	{
+		ARFilter.ClassPaths = { UAnimNextParameterLibrary::StaticClass()->GetClassPathName() };
+
+		TArray<FAssetData> LibraryAssets;
+		AssetRegistry.GetAssets(ARFilter, LibraryAssets);
+
+		for(const FAssetData& LibraryAsset : LibraryAssets)
+		{
+			const FAnimNextParameterLibraryAssetRegistryExports& Exports = GetExportsForLibrary(LibraryAsset);
+			for(const FAnimNextParameterLibraryAssetRegistryExportEntry& Export : Exports.Parameters)
+			{
+				if(!BoundParameters.Contains( { Export.Name, LibraryAsset } ))
+				{
+					FParameterBindingReference NewReference(Export.Name, LibraryAsset);
+					if(!Args.OnFilterParameter.IsBound() || Args.OnFilterParameter.Execute(NewReference) == EFilterParameterResult::Include)
+					{
+						TSharedRef<FParameterPickerEntry> NewEntry = MakeShared<FParameterPickerEntry>(NewReference, Export.Type);
+						Entries.Add(NewEntry);
+					}
+				}
+			}
+		}
+	}
+
+	RefreshFilter();
+}
+
+void SParameterPicker::RefreshFilter()
+{
+	FilteredEntries.Empty();
+
+	const FString FilterTextAsString = FilterText.ToString();
+	for(const TSharedRef<FParameterPickerEntry>& Entry : Entries)
+	{
+		if(Entry->PassesFilter(FilterTextAsString))
+		{
+			FilteredEntries.Add(Entry);
+		}
+	}
+
+	EntriesList->RequestListRefresh();
+}
+
+class SParameterPickerRow : public SMultiColumnTableRow<TSharedRef<FParameterPickerEntry>>
+{
+	SLATE_BEGIN_ARGS(SParameterPickerRow) {}
+	SLATE_END_ARGS()
+
+	void Construct(const FArguments& InArgs, const TSharedRef<STableViewBase>& InOwnerTableView, TSharedRef<FParameterPickerEntry> InEntry)
+	{
+		Entry = InEntry;
+
+		SMultiColumnTableRow<TSharedRef<FParameterPickerEntry>>::Construct( SMultiColumnTableRow<TSharedRef<FParameterPickerEntry>>::FArguments(), InOwnerTableView);
+	}
+
+	virtual TSharedRef<SWidget> GenerateWidgetForColumn(const FName& InColumnName) override
+	{
+		using namespace ParameterPicker;
+
+		if(InColumnName == Column_Parameter)
+		{
+			return
+				SNew(SHorizontalBox)
+				+SHorizontalBox::Slot()
+				.HAlign(HAlign_Center)
+				.VAlign(VAlign_Center)
+				.Padding(2.0f)
+				.AutoWidth()
+				[
+					SNew(SImage)
+					.Image(Entry->PinIcon)
+					.ColorAndOpacity(Entry->PinColor)
+				]
+				+SHorizontalBox::Slot()
+				.VAlign(VAlign_Center)
+				.Padding(2.0f)
+				.FillWidth(1.0f)
+				[
+					SNew(STextBlock)
+					.Font(IDetailLayoutBuilder::GetDetailFont())
+					.Text(FText::FromName(Entry->Binding.Parameter))
+					.ToolTipText(FText::FromName(Entry->Binding.Parameter))
+				];
+		}
+		else if(InColumnName == Column_Library)
+		{
+			if(Entry->Binding.Library.IsValid())
+			{
+				
+				return
+					SNew(SBox)
+					.VAlign(VAlign_Center)
+					[
+						SNew(STextBlock)
+						.Font(IDetailLayoutBuilder::GetDetailFont())
+						.Text(FText::FromName(Entry->Binding.Library.AssetName))
+						.ToolTipText(FText::FromName(Entry->Binding.Library.PackageName))
+					];
+			}
+		}
+		else if(InColumnName == Column_Block)
+		{
+			if(Entry->Binding.Block.IsValid())
+			{
+				return
+					SNew(SBox)
+					.VAlign(VAlign_Center)
+					[
+						SNew(STextBlock)
+						.Font(IDetailLayoutBuilder::GetDetailFont())
+						.Text(FText::FromName(Entry->Binding.Block.AssetName))
+						.ToolTipText(FText::FromName(Entry->Binding.Block.PackageName))
+					];
+			}
+		}
+
+		return SNullWidget::NullWidget;
+	}
+	
+	TSharedPtr<FParameterPickerEntry> Entry;
+};
+
+TSharedRef<ITableRow> SParameterPicker::HandleGenerateRow(TSharedRef<FParameterPickerEntry> InEntry, const TSharedRef<STableViewBase>& InOwnerTable) const
+{
+	return SNew(SParameterPickerRow, InOwnerTable, InEntry);
+}
+
+void SParameterPicker::HandleSelectionChanged(TSharedPtr<FParameterPickerEntry> InEntry, ESelectInfo::Type InSelectInfo)
+{
+	Args.OnSelectionChanged.ExecuteIfBound();
+}
+
+void SParameterPicker::HandleGetParameterBindings(TArray<FParameterBindingReference>& OutParameterBindings) const
+{
+	TArray<TSharedRef<FParameterPickerEntry>> SelectedEntries;
+	EntriesList->GetSelectedItems(SelectedEntries);
+
+	for(TSharedRef<FParameterPickerEntry>& Entry : SelectedEntries)
+	{
+		OutParameterBindings.Emplace(Entry->Binding);
+	}
+}
+
+}
+
+#undef LOCTEXT_NAMESPACE
\ No newline at end of file
diff --git a/Engine/Plugins/Experimental/AnimNext/Source/AnimNextEditor/Private/Param/SParameterPicker.h b/Engine/Plugins/Experimental/AnimNext/Source/AnimNextEditor/Private/Param/SParameterPicker.h
new file mode 100644
--- /dev/null
+++ b/Engine/Plugins/Experimental/AnimNext/Source/AnimNextEditor/Private/Param/SParameterPicker.h
@@ -0,0 +1,48 @@
+// Copyright Epic Games, Inc. All Rights Reserved.
+
+#pragma once
+
+#include "Widgets/SCompoundWidget.h"
+#include "Param/ParameterPickerArgs.h"
+#include "Widgets/Views/SListView.h"
+
+namespace UE::AnimNext::Editor
+{
+
+struct FParameterPickerEntry;
+
+class SParameterPicker : public SCompoundWidget
+{
+public:
+	SLATE_BEGIN_ARGS(SParameterPicker) {}
+
+	SLATE_ARGUMENT(FParameterPickerArgs, Args)
+
+	SLATE_END_ARGS()
+
+	void Construct(const FArguments& InArgs);
+
+private:
+	void RefreshEntries();
+
+	void RefreshFilter();
+
+	TSharedRef<ITableRow> HandleGenerateRow(TSharedRef<FParameterPickerEntry> InEntry, const TSharedRef<STableViewBase>& InOwnerTable) const;
+
+	void HandleSelectionChanged(TSharedPtr<FParameterPickerEntry> InEntry, ESelectInfo::Type InSelectInfo);
+	
+	void HandleGetParameterBindings(TArray<FParameterBindingReference>& OutParameterBindings) const;
+
+private:
+	FParameterPickerArgs Args;
+
+	TSharedPtr<SListView<TSharedRef<FParameterPickerEntry>>> EntriesList;
+
+	TArray<TSharedRef<FParameterPickerEntry>> Entries;
+
+	TArray<TSharedRef<FParameterPickerEntry>> FilteredEntries;
+
+	FText FilterText;
+};
+
+}
\ No newline at end of file
diff --git a/Engine/Plugins/Experimental/AnimNext/Source/AnimNextEditor/Private/Param/SParametersActionMenu.cpp b/Engine/Plugins/Experimental/AnimNext/Source/AnimNextEditor/Private/Param/SParametersActionMenu.cpp
new file mode 100644
--- /dev/null
+++ b/Engine/Plugins/Experimental/AnimNext/Source/AnimNextEditor/Private/Param/SParametersActionMenu.cpp
@@ -0,0 +1,226 @@
+﻿// Copyright Epic Games, Inc. All Rights Reserved.
+
+#include "SParametersActionMenu.h"
+
+#include "ParametersEditorSchemaActions.h"
+#include "Framework/Application/SlateApplication.h"
+#include "IDocumentation.h"
+#include "SGraphActionMenu.h"
+#include "Widgets/Images/SImage.h"
+#include "Widgets/Layout/SBox.h"
+#include "Widgets/SBoxPanel.h"
+#include "Widgets/Text/STextBlock.h"
+#include "SGraphPalette.h"
+#include "RigVMCore/RigVMRegistry.h"
+#include "Units/RigUnit.h"
+#include "Widgets/SToolTip.h"
+#include "Param/ParametersExecuteContext.h"
+
+#define LOCTEXT_NAMESPACE "AnimNextParametersEditor"
+
+namespace UE::AnimNext::Editor
+{
+
+static void CollectAllAnimNextParameterActions(FGraphContextMenuBuilder& MenuBuilder)
+{
+	static const TArray<UScriptStruct*> AllowedExecuteContexts =
+	{
+		FRigVMExecuteContext::StaticStruct(),
+		FAnimNextParametersExecuteContext::StaticStruct()
+	};
+
+	const TChunkedArray<FRigVMFunction>& Functions = FRigVMRegistry::Get().GetFunctions();
+	const int32 NumFunctions = Functions.Num();
+	for(int i=0; i < NumFunctions; i++)
+	{
+		const FRigVMFunction& Function = Functions[i];
+
+		const UScriptStruct* FunctionContext = Function.GetExecuteContextStruct();
+		if (FunctionContext == nullptr || !AllowedExecuteContexts.Contains(FunctionContext))
+		{
+			continue;
+		}
+
+		UScriptStruct* Struct = Function.Struct;
+		if (Struct == nullptr)
+		{
+			continue;
+		}
+
+		FString CategoryMetadata, DisplayNameMetadata, MenuDescSuffixMetadata;
+		Struct->GetStringMetaDataHierarchical(FRigVMStruct::CategoryMetaName, &CategoryMetadata);
+		Struct->GetStringMetaDataHierarchical(FRigVMStruct::DisplayNameMetaName, &DisplayNameMetadata);
+		Struct->GetStringMetaDataHierarchical(FRigVMStruct::MenuDescSuffixMetaName, &MenuDescSuffixMetadata);
+		if (!MenuDescSuffixMetadata.IsEmpty())
+		{
+			MenuDescSuffixMetadata = TEXT(" ") + MenuDescSuffixMetadata;
+		}
+		const FText NodeCategory = FText::FromString(CategoryMetadata);
+		const FText MenuDesc = FText::FromString(DisplayNameMetadata + MenuDescSuffixMetadata);
+		const FText ToolTip = Struct->GetToolTipText();
+
+		if (MenuDesc.IsEmpty())
+		{
+			continue;
+		}
+
+		MenuBuilder.AddAction(MakeShared<FAnimNextParameterSchemaAction_RigUnit>(Struct, NodeCategory, MenuDesc, ToolTip));
+	};
+}
+
+SParametersActionMenu::~SParametersActionMenu()
+{
+	OnClosedCallback.ExecuteIfBound();
+	OnCloseReasonCallback.ExecuteIfBound(bActionExecuted, false, !DraggedFromPins.IsEmpty());
+}
+
+void SParametersActionMenu::Construct(const FArguments& InArgs)
+{
+	Graph = InArgs._Graph;
+	DraggedFromPins = InArgs._DraggedFromPins;
+	NewNodePosition = InArgs._NewNodePosition;
+	OnClosedCallback = InArgs._OnClosedCallback;
+	bAutoExpandActionMenu = InArgs._AutoExpandActionMenu;
+	OnCloseReasonCallback = InArgs._OnCloseReason;
+
+	SBorder::Construct(SBorder::FArguments()
+		.BorderImage(FAppStyle::Get().GetBrush("Menu.Background"))
+		.Padding(5)
+		[
+			SNew(SBox)
+			.WidthOverride(400)
+			.HeightOverride(400)
+			[
+				SNew(SVerticalBox)
+				+SVerticalBox::Slot()
+				.AutoHeight()
+				.Padding(2, 2, 2, 5)
+				[
+					SNew(SHorizontalBox)
+					+SHorizontalBox::Slot()
+					.AutoWidth()
+					.VAlign(VAlign_Center)
+					[
+						SNew(STextBlock)
+						.Text(LOCTEXT("ContextText", "All AnimNext Node Classes"))
+						.Font(FAppStyle::Get().GetFontStyle("BlueprintEditor.ActionMenu.ContextDescriptionFont"))
+						.ToolTip(IDocumentation::Get()->CreateToolTip(
+							LOCTEXT("ActionMenuContextTextTooltip", "Describes the current context of the action list"),
+							nullptr,
+							TEXT("Shared/Editors/AnimNextParametersEditor"),
+							TEXT("AnimNextParameterSParametersActionMenuContextText")))
+						.WrapTextAt(280)
+					]
+				]
+				+SVerticalBox::Slot()
+				[
+					SAssignNew(GraphActionMenu, SGraphActionMenu)
+						.OnActionSelected(this, &SParametersActionMenu::OnActionSelected)
+						.OnCreateWidgetForAction(SGraphActionMenu::FOnCreateWidgetForAction::CreateSP(this, &SParametersActionMenu::OnCreateWidgetForAction))
+						.OnCollectAllActions(this, &SParametersActionMenu::CollectAllActions)
+						.DraggedFromPins(DraggedFromPins)
+						.GraphObj(Graph)
+						.AlphaSortItems(true)
+						.bAllowPreselectedItemActivation(true)
+				]
+			]
+		]
+	);
+}
+
+void SParametersActionMenu::CollectAllActions(FGraphActionListBuilderBase& OutAllActions)
+{
+	if (!Graph)
+	{
+		return;
+	}
+
+	FGraphContextMenuBuilder MenuBuilder(Graph);
+	if (!DraggedFromPins.IsEmpty())
+	{
+		MenuBuilder.FromPin = DraggedFromPins[0];
+	}
+
+	// Cannot call GetGraphContextActions() during serialization and GC due to its use of FindObject()
+	if(!GIsSavingPackage && !IsGarbageCollecting())
+	{
+		CollectAllAnimNextParameterActions(MenuBuilder);
+	}
+
+	OutAllActions.Append(MenuBuilder);
+}
+
+TSharedRef<SEditableTextBox> SParametersActionMenu::GetFilterTextBox()
+{
+	return GraphActionMenu->GetFilterTextBox();
+}
+
+TSharedRef<SWidget> SParametersActionMenu::OnCreateWidgetForAction(FCreateWidgetForActionData* const InCreateData)
+{
+	check(InCreateData);
+	InCreateData->bHandleMouseButtonDown = false;
+
+	const FSlateBrush* IconBrush = nullptr;
+	FLinearColor IconColor;
+	TSharedPtr<FAnimNextParameterSchemaAction> Action = StaticCastSharedPtr<FAnimNextParameterSchemaAction>(InCreateData->Action);
+	if (Action.IsValid())
+	{
+		IconBrush = Action->GetIconBrush();
+		IconColor = Action->GetIconColor();
+	}
+
+	TSharedPtr<SHorizontalBox> WidgetBox = SNew(SHorizontalBox);
+	if (IconBrush)
+	{
+		WidgetBox->AddSlot()
+			.AutoWidth()
+			.VAlign(VAlign_Center)
+			.Padding(5, 0, 0, 0)
+			[
+				SNew(SImage)
+				.ColorAndOpacity(IconColor)
+				.Image(IconBrush)
+			];
+	}
+
+	WidgetBox->AddSlot()
+		.AutoWidth()
+		.VAlign(VAlign_Center)
+		.Padding(0, 0, 0, 0)
+		[
+			SNew(SGraphPaletteItem, InCreateData)
+		];
+
+	return WidgetBox->AsShared();
+}
+
+void SParametersActionMenu::OnActionSelected(const TArray<TSharedPtr<FEdGraphSchemaAction>>& SelectedAction, ESelectInfo::Type InSelectionType)
+{
+	if (!Graph)
+	{
+		return;
+	}
+
+	if (InSelectionType != ESelectInfo::OnMouseClick  && InSelectionType != ESelectInfo::OnKeyPress && !SelectedAction.IsEmpty())
+	{
+		return;
+	}
+
+	for (const TSharedPtr<FEdGraphSchemaAction>& Action : SelectedAction)
+	{
+		if (Action.IsValid() && Graph)
+		{
+			if (!bActionExecuted && (Action->GetTypeId() != FEdGraphSchemaAction_Dummy::StaticGetTypeId()))
+			{
+				FSlateApplication::Get().DismissAllMenus();
+				bActionExecuted = true;
+			}
+
+			Action->PerformAction(Graph, DraggedFromPins, NewNodePosition);
+		}
+	}
+}
+
+} 
+
+#undef LOCTEXT_NAMESPACE
diff --git a/Engine/Plugins/Experimental/AnimNext/Source/AnimNextEditor/Private/Param/SParametersActionMenu.h b/Engine/Plugins/Experimental/AnimNext/Source/AnimNextEditor/Private/Param/SParametersActionMenu.h
new file mode 100644
--- /dev/null
+++ b/Engine/Plugins/Experimental/AnimNext/Source/AnimNextEditor/Private/Param/SParametersActionMenu.h
@@ -0,0 +1,72 @@
+﻿// Copyright Epic Games, Inc. All Rights Reserved.
+
+#pragma once
+
+#include "EdGraph/EdGraph.h"
+#include "EdGraph/EdGraphPin.h"
+#include "EdGraph/EdGraphSchema.h"
+#include "Layout/Margin.h"
+#include "Layout/Visibility.h"
+#include "SGraphActionMenu.h"
+#include "Widgets/DeclarativeSyntaxSupport.h"
+#include "Widgets/Input/SEditableTextBox.h"
+#include "Widgets/Layout/SBorder.h"
+#include "Widgets/SWidget.h"
+#include "Widgets/Views/SExpanderArrow.h"
+
+class SEditableTextBox;
+class SGraphActionMenu;
+class UEdGraph;
+
+namespace UE::AnimNext::Editor
+{
+
+class SParametersActionMenu : public SBorder
+{
+public:
+	/** Delegate for the OnCloseReason event which is always raised when the SActionMenu closes */
+	DECLARE_DELEGATE_ThreeParams(FClosedReason, bool /*bActionExecuted*/, bool /*bContextSensitiveChecked*/, bool /*bGraphPinContext*/);
+
+	SLATE_BEGIN_ARGS(SParametersActionMenu)
+		: _Graph(static_cast<UEdGraph*>(nullptr))
+		, _NewNodePosition(FVector2D::ZeroVector)
+		, _AutoExpandActionMenu(false)
+	{
+	}
+		SLATE_ARGUMENT(UEdGraph*, Graph)
+		SLATE_ARGUMENT(FVector2D, NewNodePosition)
+		SLATE_ARGUMENT(TArray<UEdGraphPin*>, DraggedFromPins)
+		SLATE_ARGUMENT(SGraphEditor::FActionMenuClosed, OnClosedCallback)
+		SLATE_ARGUMENT(bool, AutoExpandActionMenu)
+		SLATE_EVENT(FClosedReason, OnCloseReason)
+	SLATE_END_ARGS()
+
+	void Construct(const FArguments& InArgs);
+
+	~SParametersActionMenu();
+
+	TSharedRef<SEditableTextBox> GetFilterTextBox();
+
+protected:
+	void OnActionSelected(const TArray<TSharedPtr<FEdGraphSchemaAction>>& SelectedAction, ESelectInfo::Type InSelectionType);
+
+	TSharedRef<SWidget> OnCreateWidgetForAction(struct FCreateWidgetForActionData* const InCreateData);
+
+	/** Callback used to populate all actions list in SGraphActionMenu */
+	void CollectAllActions(FGraphActionListBuilderBase& OutAllActions);
+
+private:
+	UEdGraph* Graph = nullptr;
+	bool bAutoExpandActionMenu = false;
+	bool bActionExecuted = false;
+
+	TArray<UEdGraphPin*> DraggedFromPins;
+	FVector2D NewNodePosition = FVector2D::ZeroVector;
+
+	SGraphEditor::FActionMenuClosed OnClosedCallback;
+	FClosedReason OnCloseReasonCallback;
+
+	TSharedPtr<SGraphActionMenu> GraphActionMenu;
+};
+
+}
diff --git a/Engine/Plugins/Experimental/AnimNext/Source/AnimNextUncookedOnly/AnimNextUncookedOnly.Build.cs b/Engine/Plugins/Experimental/AnimNext/Source/AnimNextUncookedOnly/AnimNextUncookedOnly.Build.cs
new file mode 100644
--- /dev/null
+++ b/Engine/Plugins/Experimental/AnimNext/Source/AnimNextUncookedOnly/AnimNextUncookedOnly.Build.cs
@@ -0,0 +1,47 @@
+// Copyright Epic Games, Inc. All Rights Reserved.
+
+namespace UnrealBuildTool.Rules
+{
+	public class AnimNextUncookedOnly : ModuleRules
+	{
+		public AnimNextUncookedOnly(ReadOnlyTargetRules Target) : base(Target)
+		{
+			PublicDependencyModuleNames.AddRange(
+				new string[]
+				{
+					"RigVM",
+					"RigVMDeveloper",
+					"ControlRig",
+					"ControlRigDeveloper",
+				}
+			);
+
+			PrivateDependencyModuleNames.AddRange(
+				new string[]
+				{
+					"Core",
+					"CoreUObject",
+					"Engine",
+					"AnimNext",
+					"BlueprintGraph",	// For K2 schema
+					"AnimationCore",
+					"AnimGraph",
+					"Kismet",
+					"Slate",
+					"SlateCore",
+					"StructUtils"
+				}
+			);
+
+			if (Target.bBuildEditor)
+			{
+				PrivateDependencyModuleNames.AddRange(
+					new string[]
+					{
+						"UnrealEd",
+					}
+				);
+			}
+		}
+	}
+}
\ No newline at end of file
diff --git a/Engine/Plugins/Experimental/AnimNext/Source/AnimNextUncookedOnly/Internal/Graph/AnimNextGraph_EdGraph.h b/Engine/Plugins/Experimental/AnimNext/Source/AnimNextUncookedOnly/Internal/Graph/AnimNextGraph_EdGraph.h
new file mode 100644
--- /dev/null
+++ b/Engine/Plugins/Experimental/AnimNext/Source/AnimNextUncookedOnly/Internal/Graph/AnimNextGraph_EdGraph.h
@@ -0,0 +1,27 @@
+﻿// Copyright Epic Games, Inc. All Rights Reserved.
+
+#pragma once
+
+#include "CoreMinimal.h"
+#include "Graph/ControlRigGraph.h"
+#include "AnimNextGraph_EdGraph.generated.h"
+
+class UAnimNextGraph_EditorData;
+
+namespace UE::AnimNext::UncookedOnly
+{
+	struct FUtils;
+}
+
+UCLASS(MinimalAPI)
+class UAnimNextGraph_EdGraph : public UControlRigGraph
+{
+	GENERATED_BODY()
+
+	friend class UAnimNextGraph_EditorData;
+
+	// UControlRigGraph interface
+	virtual FRigVMClient* GetRigVMClient() const override;
+	
+	void Initialize(UAnimNextGraph_EditorData* InEditorData);
+};
\ No newline at end of file
diff --git a/Engine/Plugins/Experimental/AnimNext/Source/AnimNextUncookedOnly/Internal/Graph/AnimNextGraph_EdGraphNode.h b/Engine/Plugins/Experimental/AnimNext/Source/AnimNextUncookedOnly/Internal/Graph/AnimNextGraph_EdGraphNode.h
new file mode 100644
--- /dev/null
+++ b/Engine/Plugins/Experimental/AnimNext/Source/AnimNextUncookedOnly/Internal/Graph/AnimNextGraph_EdGraphNode.h
@@ -0,0 +1,13 @@
+﻿// Copyright Epic Games, Inc. All Rights Reserved.
+
+#pragma once
+
+#include "CoreMinimal.h"
+#include "Graph/ControlRigGraphNode.h"
+#include "AnimNextGraph_EdGraphNode.generated.h"
+
+UCLASS(MinimalAPI)
+class UAnimNextGraph_EdGraphNode : public UControlRigGraphNode
+{
+	GENERATED_BODY()
+};
\ No newline at end of file
diff --git a/Engine/Plugins/Experimental/AnimNext/Source/AnimNextUncookedOnly/Internal/Graph/AnimNextGraph_EdGraphSchema.h b/Engine/Plugins/Experimental/AnimNext/Source/AnimNextUncookedOnly/Internal/Graph/AnimNextGraph_EdGraphSchema.h
new file mode 100644
--- /dev/null
+++ b/Engine/Plugins/Experimental/AnimNext/Source/AnimNextUncookedOnly/Internal/Graph/AnimNextGraph_EdGraphSchema.h
@@ -0,0 +1,20 @@
+﻿// Copyright Epic Games, Inc. All Rights Reserved.
+
+#pragma once
+
+#include "CoreMinimal.h"
+#include "Graph/ControlRigGraphSchema.h"
+#include "AnimNextGraph_EdGraphNode.h"
+#include "AnimNextGraph_EdGraphSchema.generated.h"
+
+UCLASS(MinimalAPI)
+class UAnimNextGraph_EdGraphSchema : public UControlRigGraphSchema
+{
+	GENERATED_BODY()
+
+	// UControlRigGraphSchema interface
+	virtual TSubclassOf<UControlRigGraphNode> GetGraphNodeClass() const override { return UAnimNextGraph_EdGraphNode::StaticClass(); }
+
+	// UEdGraphSchema interface
+	virtual void TrySetDefaultValue(UEdGraphPin& InPin, const FString& InNewDefaultValue, bool bMarkAsModified) const override;
+};
\ No newline at end of file
