AnimNext: Initial editor-side parameter blocks & libraries

diff --git a/Engine/Plugins/Experimental/AnimNext/Source/AnimNextUncookedOnly/Internal/Graph/AnimNextGraph_EditorData.h b/Engine/Plugins/Experimental/AnimNext/Source/AnimNextUncookedOnly/Internal/Graph/AnimNextGraph_EditorData.h
new file mode 100644
--- /dev/null
+++ b/Engine/Plugins/Experimental/AnimNext/Source/AnimNextUncookedOnly/Internal/Graph/AnimNextGraph_EditorData.h
@@ -0,0 +1,141 @@
+﻿// Copyright Epic Games, Inc. All Rights Reserved.
+
+#pragma once
+
+#include "CoreMinimal.h"
+#include "RigVMModel/RigVMGraph.h"
+#include "AnimNextGraph_EdGraph.h"
+#include "RigVMCore/RigVMGraphFunctionHost.h"
+#include "AnimNextGraph_EditorData.generated.h"
+
+class UAnimNextGraph;
+enum class ERigVMGraphNotifType : uint8;
+namespace UE::AnimNext::UncookedOnly
+{
+	struct FUtils;
+}
+
+namespace UE::AnimNext::Editor
+{
+	class FGraphEditor;
+}
+
+enum class EAnimNextGraphLoadType : uint8
+{
+	PostLoad,
+	CheckUserDefinedStructs
+};
+
+UCLASS(MinimalAPI)
+class UAnimNextGraph_EditorData : public UObject, public IRigVMClientHost, public IRigVMGraphFunctionHost
+{
+	GENERATED_BODY()
+
+	UAnimNextGraph_EditorData(const FObjectInitializer& ObjectInitializer);
+	
+	friend class UAnimNextGraphFactory;
+	friend class UAnimNextGraph_EdGraph;
+	friend struct UE::AnimNext::UncookedOnly::FUtils;
+	friend class UE::AnimNext::Editor::FGraphEditor;
+	friend struct FAnimNextGraphSchemaAction_RigUnit;
+
+	// UObject interface
+	virtual void Serialize(FArchive& Ar) override;
+	virtual void PostLoad() override;
+	virtual bool IsEditorOnly() const override { return true; }
+#if WITH_EDITOR
+	void HandlePackageDone(const FEndLoadPackageContext& Context);
+	void HandlePackageDone();
+#endif // WITH_EDITOR
+
+	// IRigVMClientHost interface
+	virtual FRigVMClient* GetRigVMClient() override;
+	virtual const FRigVMClient* GetRigVMClient() const override;
+	virtual IRigVMGraphFunctionHost* GetRigVMGraphFunctionHost() override;
+	virtual const IRigVMGraphFunctionHost* GetRigVMGraphFunctionHost() const override;
+	virtual void HandleRigVMGraphAdded(const FRigVMClient* InClient, const FString& InNodePath) override;
+	virtual void HandleRigVMGraphRemoved(const FRigVMClient* InClient, const FString& InNodePath) override {}
+	virtual void HandleRigVMGraphRenamed(const FRigVMClient* InClient, const FString& InOldNodePath, const FString& InNewNodePath) override {}
+	virtual void HandleConfigureRigVMController(const FRigVMClient* InClient, URigVMController* InControllerToConfigure) override;
+	virtual UObject* GetEditorObjectForRigVMGraph(URigVMGraph* InVMGraph) const override;
+
+	// IRigVMGraphFunctionHost interface
+	virtual FRigVMGraphFunctionStore* GetRigVMGraphFunctionStore() override;
+	virtual const FRigVMGraphFunctionStore* GetRigVMGraphFunctionStore() const override;
+
+	ANIMNEXTUNCOOKEDONLY_API void Initialize(bool bRecompileVM);
+
+	void RefreshAllModels(EAnimNextGraphLoadType InLoadType);
+
+	void RecompileVM();
+	
+	void RecompileVMIfRequired();
+
+	void RequestAutoVMRecompilation();
+	
+	void HandleModifiedEvent(ERigVMGraphNotifType InNotifType, URigVMGraph* InGraph, UObject* InSubject);
+
+	ANIMNEXTUNCOOKEDONLY_API URigVMGraph* GetVMGraphForEdGraph(const UEdGraph* InGraph) const;
+
+	void CreateEdGraphForCollapseNode(URigVMCollapseNode* InNode);
+
+	void UpdateGraphReturnType();
+
+	FEdGraphPinType FindGraphReturnPinType() const;
+
+	bool IsNodeExecConnected(const UEdGraphNode* Node) const;
+
+	UPROPERTY()
+	TObjectPtr<UAnimNextGraph_EdGraph> RootGraph;
+
+	UPROPERTY()
+	TObjectPtr<UAnimNextGraph_EdGraph> EntryPointGraph;
+
+	UPROPERTY()
+	TObjectPtr<UAnimNextGraph_EdGraph> FunctionLibraryEdGraph;
+
+	UPROPERTY()
+	FRigVMClient RigVMClient;
+
+	UPROPERTY()
+	FRigVMGraphFunctionStore GraphFunctionStore;
+
+	UPROPERTY()
+	TObjectPtr<URigVMLibraryNode> EntryPoint;
+
+	UPROPERTY(transient)
+	TMap<TObjectPtr<URigVMGraph>, TObjectPtr<URigVMController>> Controllers;
+
+	UPROPERTY(EditAnywhere, Category = "User Interface")
+	FRigGraphDisplaySettings RigGraphDisplaySettings;
+
+	UPROPERTY(EditAnywhere, Category = "VM")
+	FRigVMRuntimeSettings VMRuntimeSettings;
+	
+	UPROPERTY(EditAnywhere, BlueprintReadWrite, Category = "VM", meta = (AllowPrivateAccess = "true"))
+	FRigVMCompileSettings VMCompileSettings;
+
+	UPROPERTY(EditAnywhere, Category = "Python Log Settings")
+	FControlRigPythonSettings PythonLogSettings;
+
+	UPROPERTY(transient, DuplicateTransient)
+	TMap<FString, FRigVMOperand> PinToOperandMap;
+
+	UPROPERTY(transient, DuplicateTransient)
+	bool bVMRecompilationRequired = false;
+
+	UPROPERTY(transient, DuplicateTransient)
+	bool bIsCompiling = false;
+	
+	FCompilerResultsLog CompileLog;
+
+	FOnVMCompiledEvent VMCompiledEvent;
+	FRigVMGraphModifiedEvent ModifiedEvent;
+
+	bool bAutoRecompileVM = true;
+	bool bErrorsDuringCompilation = false;
+	bool bSuspendModelNotificationsForSelf = false;
+	bool bSuspendModelNotificationsForOthers = false;
+	bool bSuspendAllNotifications = false;
+	bool bCompileInDebugMode = false;
+};
diff --git a/Engine/Plugins/Experimental/AnimNext/Source/AnimNextUncookedOnly/Internal/Param/AnimNextParameter.h b/Engine/Plugins/Experimental/AnimNext/Source/AnimNextUncookedOnly/Internal/Param/AnimNextParameter.h
new file mode 100644
--- /dev/null
+++ b/Engine/Plugins/Experimental/AnimNext/Source/AnimNextUncookedOnly/Internal/Param/AnimNextParameter.h
@@ -0,0 +1,69 @@
+﻿// Copyright Epic Games, Inc. All Rights Reserved.
+
+#pragma once
+
+#include "CoreMinimal.h"
+#include "Param/ParamType.h"
+#include "AnimNextParameter.generated.h"
+
+class UAnimNextParameter;
+
+namespace UE::AnimNext::Editor
+{
+struct FUtils;
+struct FParameterPickerEntry;
+class SParameterBlockView;
+class SParameterLibraryViewRow;
+}
+
+namespace UE::AnimNext::UncookedOnly
+{
+// A delegate for subscribing / reacting to parameter modifications.
+DECLARE_MULTICAST_DELEGATE_OneParam(FOnParameterModified, UAnimNextParameter* /* InParameter */);
+}
+
+/** Defines a single parameter */
+UCLASS(MinimalAPI, BlueprintType)
+class UAnimNextParameter : public UObject
+{
+	GENERATED_BODY()
+
+#if WITH_EDITORONLY_DATA
+	friend class UAnimNextParameterBlock_EditorData;
+	friend class UAnimNextParameterLibrary;
+	friend class UAnimNextParameterBlockBinding;
+	friend class UAnimNextParameterBlockBindingReference;
+	friend struct UE::AnimNext::Editor::FUtils;
+	friend class UE::AnimNext::Editor::SParameterBlockView;
+	friend struct UE::AnimNext::Editor::FParameterPickerEntry;
+	friend class UE::AnimNext::Editor::SParameterLibraryViewRow;
+	friend struct FRigVMDispatch_SetParameter;
+
+	// Get the type of this parameter
+	const FAnimNextParamType& GetType() const { return Type; }
+
+	// Set the type of this parameter
+	ANIMNEXTUNCOOKEDONLY_API void SetType(const FAnimNextParamType& InType, bool bSetupUndoRedo = true);
+
+	// Get the comment for this parameter
+	FStringView GetComment() const { return Comment; }
+
+	// Returns the modified event which can be used to subscribe to changes to this parameter
+	ANIMNEXTUNCOOKEDONLY_API UE::AnimNext::UncookedOnly::FOnParameterModified& OnModified();
+
+	// UObject interface 
+	virtual void PostTransacted(const FTransactionObjectEvent& TransactionEvent) override;
+	virtual bool IsAsset() const override;
+
+	/** The parameter's type */
+	UPROPERTY(EditAnywhere, Category = "Parameter", AssetRegistrySearchable)
+	FAnimNextParamType Type = FAnimNextParamType::GetType<bool>();
+	
+	/** Comment to display in editor */
+	UPROPERTY(EditAnywhere, Category = "Parameter", meta=(MultiLine))
+	FString Comment;
+
+	// Delegate to subscribe to modifications
+	UE::AnimNext::UncookedOnly::FOnParameterModified ModifiedDelegate;
+#endif // #if WITH_EDITORONLY_DATA
+};
\ No newline at end of file
diff --git a/Engine/Plugins/Experimental/AnimNext/Source/AnimNextUncookedOnly/Internal/Param/AnimNextParameterBlockBinding.h b/Engine/Plugins/Experimental/AnimNext/Source/AnimNextUncookedOnly/Internal/Param/AnimNextParameterBlockBinding.h
new file mode 100644
--- /dev/null
+++ b/Engine/Plugins/Experimental/AnimNext/Source/AnimNextUncookedOnly/Internal/Param/AnimNextParameterBlockBinding.h
@@ -0,0 +1,55 @@
+﻿// Copyright Epic Games, Inc. All Rights Reserved.
+
+#pragma once
+
+#include "CoreMinimal.h"
+#include "AnimNextParameterBlockEntry.h"
+#include "IAnimNextParameterBlockGraphInterface.h"
+#include "IAnimNextParameterBlockBindingInterface.h"
+#include "AnimNextParameterBlockBinding.generated.h"
+
+class UAssetDefinition_AnimNextParameterBlockBinding;
+class UAnimNextParameter;
+
+namespace UE::AnimNext::Editor
+{
+	class SParameterBlockViewRow;
+}
+
+/** Parameter binding block entry */
+UCLASS(MinimalAPI, BlueprintType)
+class UAnimNextParameterBlockBinding : public UAnimNextParameterBlockEntry, public IAnimNextParameterBlockBindingInterface, public IAnimNextParameterBlockGraphInterface
+{
+	GENERATED_BODY()
+
+	friend class UAnimNextParameterBlock_EditorData;
+	friend class UAssetDefinition_AnimNextParameterBlockBinding;
+	friend class UE::AnimNext::Editor::SParameterBlockViewRow;
+
+	// IAnimNextParameterBlockBindingInterface interface
+	virtual FAnimNextParamType GetParamType() const override;
+	virtual void SetParameterName(FName InName, bool bSetupUndoRedo = true) override;
+	virtual FName GetParameterName() const override;
+	virtual const UAnimNextParameter* GetParameter() const override;
+	virtual const UAnimNextParameterLibrary* GetLibrary() const override;
+
+	// IAnimNextParameterBlockGraphInterface interface
+	virtual URigVMGraph* GetGraph() const override { return BindingGraph; }
+
+	// UAnimNextParameterBlockEntry interface
+	virtual FText GetDisplayName() const override;
+	virtual FText GetDisplayNameTooltip() const override;
+	virtual void GetEditedObjects(TArray<UObject*>& OutObjects) const override;
+
+	/** Parameter name we reference */
+	UPROPERTY(VisibleAnywhere, Category = Parameter)
+	FName ParameterName;
+
+	/** Parameter library we reference */
+	UPROPERTY(VisibleAnywhere, Category = Parameter)
+	TObjectPtr<UAnimNextParameterLibrary> Library;
+
+	/** Binding graph */
+	UPROPERTY()
+	TObjectPtr<URigVMGraph> BindingGraph;
+};
\ No newline at end of file
diff --git a/Engine/Plugins/Experimental/AnimNext/Source/AnimNextUncookedOnly/Internal/Param/AnimNextParameterBlockBindingReference.h b/Engine/Plugins/Experimental/AnimNext/Source/AnimNextUncookedOnly/Internal/Param/AnimNextParameterBlockBindingReference.h
new file mode 100644
--- /dev/null
+++ b/Engine/Plugins/Experimental/AnimNext/Source/AnimNextUncookedOnly/Internal/Param/AnimNextParameterBlockBindingReference.h
@@ -0,0 +1,57 @@
+﻿// Copyright Epic Games, Inc. All Rights Reserved.
+
+#pragma once
+
+#include "CoreMinimal.h"
+#include "AnimNextParameterBlockEntry.h"
+#include "IAnimNextParameterBlockReferenceInterface.h"
+#include "IAnimNextParameterBlockBindingInterface.h"
+#include "AnimNextParameterBlockBindingReference.generated.h"
+
+class UAnimNextParameterBlock;
+class UAssetDefinition_AnimNextParameterBlockBindingReference;
+
+namespace UE::AnimNext::Editor
+{
+	struct FUtils;
+	class SParameterBlockViewRow;
+}
+
+class UAnimNextParameter;
+
+/** Parameter binding block entry */
+UCLASS(MinimalAPI, BlueprintType)
+class UAnimNextParameterBlockBindingReference : public UAnimNextParameterBlockEntry, public IAnimNextParameterBlockBindingInterface, public IAnimNextParameterBlockReferenceInterface
+{
+	GENERATED_BODY()
+
+	friend class UAnimNextParameterBlock_EditorData;
+	friend class UAssetDefinition_AnimNextParameterBlockBindingReference;
+
+	// IAnimNextParameterBlockBindingInterface interface
+	virtual FAnimNextParamType GetParamType() const override;
+	virtual FName GetParameterName() const override;
+	virtual void SetParameterName(FName InName, bool bSetupUndoRedo = true) override;
+	virtual const UAnimNextParameter* GetParameter() const override;
+	virtual const UAnimNextParameterLibrary* GetLibrary() const override;
+
+	// IAnimNextParameterBlockReferenceInterface interface
+	virtual const UAnimNextParameterBlock* GetBlock() const override;
+
+	// UAnimNextParameterBlockEntry interface
+	virtual FText GetDisplayName() const override;
+	virtual FText GetDisplayNameTooltip() const override;
+	virtual void GetEditedObjects(TArray<UObject*>& OutObjects) const override;
+
+	/** Parameter name we reference */
+	UPROPERTY(VisibleAnywhere, Category = Parameter)
+	FName ParameterName;
+
+	/** Parameter library we reference */
+	UPROPERTY(VisibleAnywhere, Category = Parameter)
+	TObjectPtr<UAnimNextParameterLibrary> Library;
+
+	/** Parameter block we reference */
+	UPROPERTY()
+	TObjectPtr<UAnimNextParameterBlock> Block;
+};
\ No newline at end of file
diff --git a/Engine/Plugins/Experimental/AnimNext/Source/AnimNextUncookedOnly/Internal/Param/AnimNextParameterBlockEntry.h b/Engine/Plugins/Experimental/AnimNext/Source/AnimNextUncookedOnly/Internal/Param/AnimNextParameterBlockEntry.h
new file mode 100644
--- /dev/null
+++ b/Engine/Plugins/Experimental/AnimNext/Source/AnimNextUncookedOnly/Internal/Param/AnimNextParameterBlockEntry.h
@@ -0,0 +1,47 @@
+﻿// Copyright Epic Games, Inc. All Rights Reserved.
+
+#pragma once
+
+#include "CoreMinimal.h"
+#include "AnimNextParameterBlockEntry.generated.h"
+
+class UAnimNextParameterBlock_EditorData;
+enum class ERigVMGraphNotifType : uint8;
+class URigVMGraph;
+
+namespace UE::AnimNext::Editor
+{
+class SParameterBlockViewRow;
+class FParameterBlockEditor;
+struct FParameterBlockViewEntry;
+}
+
+/** Base class that defines an entry in a parameter block, e.g. a parameter binding */
+UCLASS(MinimalAPI, BlueprintType)
+class UAnimNextParameterBlockEntry : public UObject
+{
+	GENERATED_BODY()
+
+	friend class UE::AnimNext::Editor::SParameterBlockViewRow;
+	friend class UE::AnimNext::Editor::FParameterBlockEditor;
+	friend struct UE::AnimNext::Editor::FParameterBlockViewEntry;
+
+	void Initialize(UAnimNextParameterBlock_EditorData* InEditorData);
+
+	void HandleRigVMGraphModifiedEvent(ERigVMGraphNotifType InNotifType, URigVMGraph* InGraph, UObject* InSubject);
+	
+	// Get the name to be displayed in the UI for this entry
+	virtual FText GetDisplayName() const PURE_VIRTUAL(GetDisplayName, return FText::GetEmpty();)
+
+	// Get the tooltip to be displayed for the name in the UI for this entry
+	virtual FText GetDisplayNameTooltip() const PURE_VIRTUAL(GetDisplayNameTooltip, return FText::GetEmpty();)
+
+	// Get any external objects/assets that this entry can edit
+	virtual void GetEditedObjects(TArray<UObject*>& OutObjects) const {}
+
+	// UObject interface
+	virtual bool IsAsset() const override;
+
+protected:
+	void BroadcastModified();
+};
\ No newline at end of file
diff --git a/Engine/Plugins/Experimental/AnimNext/Source/AnimNextUncookedOnly/Internal/Param/AnimNextParameterBlockProcessor.h b/Engine/Plugins/Experimental/AnimNext/Source/AnimNextUncookedOnly/Internal/Param/AnimNextParameterBlockProcessor.h
new file mode 100644
--- /dev/null
+++ b/Engine/Plugins/Experimental/AnimNext/Source/AnimNextUncookedOnly/Internal/Param/AnimNextParameterBlockProcessor.h
@@ -0,0 +1,13 @@
+﻿
+#pragma once
+
+#include "CoreMinimal.h"
+#include "AnimNextParameterBlockEntry.h"
+#include "AnimNextParameterBlockProcessor.generated.h"
+
+/** Block entry that performs arbitrary logic at the point the block is applied */
+UCLASS(MinimalAPI)
+class UAnimNextParameterBlockProcessor : public UAnimNextParameterBlockEntry
+{
+	GENERATED_BODY()
+};
\ No newline at end of file
diff --git a/Engine/Plugins/Experimental/AnimNext/Source/AnimNextUncookedOnly/Internal/Param/AnimNextParameterBlock_EdGraph.h b/Engine/Plugins/Experimental/AnimNext/Source/AnimNextUncookedOnly/Internal/Param/AnimNextParameterBlock_EdGraph.h
new file mode 100644
--- /dev/null
+++ b/Engine/Plugins/Experimental/AnimNext/Source/AnimNextUncookedOnly/Internal/Param/AnimNextParameterBlock_EdGraph.h
@@ -0,0 +1,27 @@
+﻿// Copyright Epic Games, Inc. All Rights Reserved.
+
+#pragma once
+
+#include "CoreMinimal.h"
+#include "Graph/ControlRigGraph.h"
+#include "AnimNextParameterBlock_EdGraph.generated.h"
+
+class UAnimNextGraph_EditorData;
+
+namespace UE::AnimNext::UncookedOnly
+{
+	struct FUtils;
+}
+
+UCLASS(MinimalAPI)
+class UAnimNextParameterBlock_EdGraph : public UControlRigGraph
+{
+	GENERATED_BODY()
+
+	friend class UAnimNextParameterBlock_EditorData;
+
+	// UControlRigGraph interface
+	virtual FRigVMClient* GetRigVMClient() const override;
+	
+	void Initialize(UAnimNextParameterBlock_EditorData* InEditorData);
+};
\ No newline at end of file
diff --git a/Engine/Plugins/Experimental/AnimNext/Source/AnimNextUncookedOnly/Internal/Param/AnimNextParameterBlock_EdGraphNode.h b/Engine/Plugins/Experimental/AnimNext/Source/AnimNextUncookedOnly/Internal/Param/AnimNextParameterBlock_EdGraphNode.h
new file mode 100644
--- /dev/null
+++ b/Engine/Plugins/Experimental/AnimNext/Source/AnimNextUncookedOnly/Internal/Param/AnimNextParameterBlock_EdGraphNode.h
@@ -0,0 +1,13 @@
+﻿// Copyright Epic Games, Inc. All Rights Reserved.
+
+#pragma once
+
+#include "CoreMinimal.h"
+#include "Graph/ControlRigGraphNode.h"
+#include "AnimNextParameterBlock_EdGraphNode.generated.h"
+
+UCLASS(MinimalAPI)
+class UAnimNextParameterBlock_EdGraphNode : public UControlRigGraphNode
+{
+	GENERATED_BODY()
+};
\ No newline at end of file
diff --git a/Engine/Plugins/Experimental/AnimNext/Source/AnimNextUncookedOnly/Internal/Param/AnimNextParameterBlock_EdGraphSchema.h b/Engine/Plugins/Experimental/AnimNext/Source/AnimNextUncookedOnly/Internal/Param/AnimNextParameterBlock_EdGraphSchema.h
new file mode 100644
--- /dev/null
+++ b/Engine/Plugins/Experimental/AnimNext/Source/AnimNextUncookedOnly/Internal/Param/AnimNextParameterBlock_EdGraphSchema.h
@@ -0,0 +1,20 @@
+﻿// Copyright Epic Games, Inc. All Rights Reserved.
+
+#pragma once
+
+#include "CoreMinimal.h"
+#include "Graph/ControlRigGraphSchema.h"
+#include "AnimNextParameterBlock_EdGraphNode.h"
+#include "AnimNextParameterBlock_EdGraphSchema.generated.h"
+
+UCLASS(MinimalAPI)
+class UAnimNextParameterBlock_EdGraphSchema : public UControlRigGraphSchema
+{
+	GENERATED_BODY()
+
+	// UControlRigGraphSchema interface
+	virtual TSubclassOf<UControlRigGraphNode> GetGraphNodeClass() const override { return UAnimNextParameterBlock_EdGraphNode::StaticClass(); }
+
+	// UEdGraphSchema interface
+	virtual void TrySetDefaultValue(UEdGraphPin& InPin, const FString& InNewDefaultValue, bool bMarkAsModified) const override;
+};
\ No newline at end of file
diff --git a/Engine/Plugins/Experimental/AnimNext/Source/AnimNextUncookedOnly/Internal/Param/AnimNextParameterBlock_EditorData.h b/Engine/Plugins/Experimental/AnimNext/Source/AnimNextUncookedOnly/Internal/Param/AnimNextParameterBlock_EditorData.h
new file mode 100644
--- /dev/null
+++ b/Engine/Plugins/Experimental/AnimNext/Source/AnimNextUncookedOnly/Internal/Param/AnimNextParameterBlock_EditorData.h
@@ -0,0 +1,226 @@
+﻿// Copyright Epic Games, Inc. All Rights Reserved.
+
+#pragma once
+
+#include "CoreMinimal.h"
+#include "RigVMModel/RigVMGraph.h"
+#include "AnimNextParameterBlock_EdGraph.h"
+#include "RigVMCore/RigVMGraphFunctionHost.h"
+#include "AnimNextParameterBlock_EditorData.generated.h"
+
+class UAnimNextParameterBlock;
+class UAnimNextParameterBlockEntry;
+class UAnimNextParameter;
+class UAnimNextParameterBlock_Controller;
+enum class ERigVMGraphNotifType : uint8;
+class FAnimationAnimNextParametersEditorTest_Block;;
+
+namespace UE::AnimNext::UncookedOnly
+{
+	struct FUtils;
+	struct FUtilsPrivate;
+}
+
+namespace UE::AnimNext::Editor
+{
+	struct FUtils;
+	class FParameterBlockEditor;
+	class SParameterBlockView;
+	class SParameterPicker;
+	class FParameterBlockTabSummoner;
+}
+
+namespace UE::AnimNext::UncookedOnly
+{
+// A delegate for subscribing / reacting to parameter block modifications.
+DECLARE_MULTICAST_DELEGATE_OneParam(FOnParameterBlockModified, UAnimNextParameterBlock_EditorData* /* InEditorData */);
+}
+
+enum class EAnimNextParameterLoadType : uint8
+{
+	PostLoad,
+	CheckUserDefinedStructs
+};
+
+// Script-callable editor API hoisted onto UAnimNextParameterBlock
+UCLASS()
+class ANIMNEXTUNCOOKEDONLY_API UAnimNextParameterBlockLibrary : public UBlueprintFunctionLibrary
+{
+	GENERATED_BODY()
+	
+	UFUNCTION(BlueprintCallable, Category = "AnimNext|Parameter Block", meta=(ScriptMethod))
+	static UAnimNextParameterBlockBinding* AddBinding(UAnimNextParameterBlock* InBlock, FName InName, UAnimNextParameterLibrary* InLibrary, bool bSetupUndoRedo = true, bool bPrintPythonCommand = true);
+
+	UFUNCTION(BlueprintCallable, Category = "AnimNext|Parameter Block", meta=(ScriptMethod))
+	static UAnimNextParameterBlockBindingReference* AddBindingReference(UAnimNextParameterBlock* InBlock, FName InName, UAnimNextParameterLibrary* InLibrary, UAnimNextParameterBlock* InReferencedBlock, bool bSetupUndoRedo = true, bool bPrintPythonCommand = true);
+
+	UFUNCTION(BlueprintCallable, Category = "AnimNext|Parameter Block", meta=(ScriptMethod))
+	static bool RemoveAllBindings(UAnimNextParameterBlock* InBlock, FName InName, bool bSetupUndoRedo = true, bool bPrintPythonCommand = true);
+
+	UFUNCTION(BlueprintCallable, Category = "AnimNext|Parameter Block", meta=(ScriptMethod))
+	static bool RemoveEntry(UAnimNextParameterBlock* InBlock, UAnimNextParameterBlockEntry* InEntry, bool bSetupUndoRedo = true, bool bPrintPythonCommand = true);
+
+	UFUNCTION(BlueprintCallable, Category = "AnimNext|Parameter Block", meta=(ScriptMethod))
+	static bool RemoveEntries(UAnimNextParameterBlock* InBlock, const TArray<UAnimNextParameterBlockEntry*>& InEntries, bool bSetupUndoRedo = true, bool bPrintPythonCommand = true);
+
+	UFUNCTION(BlueprintCallable, Category = "AnimNext|Parameter Block", meta=(ScriptMethod))
+	static UAnimNextParameterBlockEntry* FindBinding(UAnimNextParameterBlock* InBlock, FName InName);
+};
+
+UCLASS(MinimalAPI)
+class UAnimNextParameterBlock_EditorData : public UObject, public IRigVMClientHost, public IRigVMGraphFunctionHost
+{
+	GENERATED_BODY()
+
+	UAnimNextParameterBlock_EditorData(const FObjectInitializer& ObjectInitializer);
+
+#if WITH_EDITORONLY_DATA
+	ANIMNEXTUNCOOKEDONLY_API static const FName ExportsAssetRegistryTag; 
+
+	friend class UAnimNextParameterBlockLibrary;
+	friend class UAnimNextParameterBlockFactory;
+	friend class UAnimNextParameterBlockEntry;
+	friend class UAnimNextParameterBlock_EdGraph;
+	friend class UAnimNextParameterBlock_Controller;
+	friend struct UE::AnimNext::Editor::FUtils;
+	friend struct UE::AnimNext::UncookedOnly::FUtils;
+	friend struct UE::AnimNext::UncookedOnly::FUtilsPrivate;
+	friend class UE::AnimNext::Editor::FParameterBlockEditor;
+	friend struct FAnimNextParameterSchemaAction_RigUnit;
+	friend class UE::AnimNext::Editor::SParameterBlockView;
+	friend class UE::AnimNext::Editor::SParameterPicker;
+	friend class UE::AnimNext::Editor::FParameterBlockTabSummoner;
+	friend class FAnimationAnimNextParametersEditorTest_Block;
+
+	ANIMNEXTUNCOOKEDONLY_API UAnimNextParameterBlockBinding* AddBinding(FName InName, UAnimNextParameterLibrary* InLibrary, bool bSetupUndoRedo = true, bool bPrintPythonCommand = true);
+
+	ANIMNEXTUNCOOKEDONLY_API UAnimNextParameterBlockBindingReference* AddBindingReference(FName InName, UAnimNextParameterLibrary* InLibrary, UAnimNextParameterBlock* InBlock, bool bSetupUndoRedo = true, bool bPrintPythonCommand = true);
+
+	ANIMNEXTUNCOOKEDONLY_API bool RemoveAllBindings(FName InName, bool bSetupUndoRedo = true, bool bPrintPythonCommand = true);
+
+	ANIMNEXTUNCOOKEDONLY_API bool RemoveEntry(UAnimNextParameterBlockEntry* InEntry, bool bSetupUndoRedo = true, bool bPrintPythonCommand = true);
+
+	ANIMNEXTUNCOOKEDONLY_API bool RemoveEntries(const TArray<UAnimNextParameterBlockEntry*>& InEntries, bool bSetupUndoRedo = true, bool bPrintPythonCommand = true);
+
+	ANIMNEXTUNCOOKEDONLY_API UAnimNextParameterBlockEntry* FindBinding(FName InName) const;
+
+	// UObject interface
+	virtual void Serialize(FArchive& Ar) override;
+	virtual void PostLoad() override;
+	virtual void PostTransacted(const FTransactionObjectEvent& TransactionEvent) override;
+	virtual bool IsEditorOnly() const override { return true; }
+	virtual void GetAssetRegistryTags(TArray<FAssetRegistryTag>& OutTags) const override;
+
+#if WITH_EDITOR
+	void HandlePackageDone(const FEndLoadPackageContext& Context);
+	void HandlePackageDone();
+#endif
+	// IRigVMClientHost interface
+	virtual FRigVMClient* GetRigVMClient() override;
+	virtual const FRigVMClient* GetRigVMClient() const override;
+	virtual IRigVMGraphFunctionHost* GetRigVMGraphFunctionHost() override;
+	virtual const IRigVMGraphFunctionHost* GetRigVMGraphFunctionHost() const override;
+	virtual void HandleRigVMGraphAdded(const FRigVMClient* InClient, const FString& InNodePath) override;
+	virtual void HandleRigVMGraphRemoved(const FRigVMClient* InClient, const FString& InNodePath) override;
+	virtual void HandleRigVMGraphRenamed(const FRigVMClient* InClient, const FString& InOldNodePath, const FString& InNewNodePath) override {}
+	virtual void HandleConfigureRigVMController(const FRigVMClient* InClient, URigVMController* InControllerToConfigure) override;
+	virtual UObject* GetEditorObjectForRigVMGraph(URigVMGraph* InVMGraph) const override;
+
+	// IRigVMGraphFunctionHost interface
+	virtual FRigVMGraphFunctionStore* GetRigVMGraphFunctionStore() override;
+	virtual const FRigVMGraphFunctionStore* GetRigVMGraphFunctionStore() const override;
+
+	UEdGraph* CreateEdGraph(URigVMGraph* InRigVMGraph, bool bForce);
+	bool RemoveEdGraph(URigVMGraph* InModel);
+	
+	ANIMNEXTUNCOOKEDONLY_API void Initialize(bool bRecompileVM);
+
+	void RefreshAllModels(EAnimNextParameterLoadType InLoadType);
+
+	void RecompileVM();
+
+	void RecompileVMIfRequired();
+
+	void Recompile();
+
+	void RecompileIfRequired();
+
+	void RequestAutoRecompilation();
+	
+	void RequestAutoVMRecompilation();
+
+	void IncrementVMRecompileBracket();
+
+	void DecrementVMRecompileBracket();
+	
+	void HandleModifiedEvent(ERigVMGraphNotifType InNotifType, URigVMGraph* InGraph, UObject* InSubject);
+
+	ANIMNEXTUNCOOKEDONLY_API URigVMGraph* GetVMGraphForEdGraph(const UEdGraph* InGraph) const;
+
+	void CreateEdGraphForCollapseNode(URigVMCollapseNode* InNode);
+
+	void BroadcastModified();
+
+	void ReportError(const TCHAR* InMessage) const;
+
+	UPROPERTY()
+	TArray<TObjectPtr<UAnimNextParameterBlock_EdGraph>> Graphs;
+
+	UPROPERTY()
+	FRigVMClient RigVMClient;
+
+	UPROPERTY()
+	FRigVMGraphFunctionStore GraphFunctionStore;
+
+	UPROPERTY(EditAnywhere, Category = "User Interface")
+	FRigGraphDisplaySettings RigGraphDisplaySettings;
+
+	UPROPERTY(EditAnywhere, Category = "VM")
+	FRigVMRuntimeSettings VMRuntimeSettings;
+	
+	UPROPERTY(EditAnywhere, BlueprintReadWrite, Category = "VM", meta = (AllowPrivateAccess = "true"))
+	FRigVMCompileSettings VMCompileSettings;
+
+	UPROPERTY(EditAnywhere, Category = "Python Log Settings")
+	FControlRigPythonSettings PythonLogSettings;
+
+	UPROPERTY(transient, DuplicateTransient)
+	TMap<FString, FRigVMOperand> PinToOperandMap;
+
+	/** All entries in this parameter block - not saved, either serialized or discovered at load time */
+	UPROPERTY(transient)
+	TArray<TObjectPtr<UAnimNextParameterBlockEntry>> Entries;
+
+	UPROPERTY()
+	TArray<FEditedDocumentInfo> LastEditedDocuments;
+
+	UPROPERTY(transient, DuplicateTransient)
+	int32 VMRecompilationBracket = 0;
+
+	UPROPERTY(transient, DuplicateTransient)
+	bool bVMRecompilationRequired = false;
+
+	UPROPERTY(transient, DuplicateTransient)
+	bool bStructRecompilationRequired = false;
+
+	UPROPERTY(transient, DuplicateTransient)
+	bool bIsCompiling = false;
+	
+	FCompilerResultsLog CompileLog;
+
+	FOnVMCompiledEvent RigVMCompiledEvent;
+	FRigVMGraphModifiedEvent RigVMGraphModifiedEvent;
+
+	// Delegate to subscribe to modifications to this block
+	UE::AnimNext::UncookedOnly::FOnParameterBlockModified ModifiedDelegate;
+
+	bool bAutoRecompileVM = true;
+	bool bErrorsDuringCompilation = false;
+	bool bSuspendModelNotificationsForSelf = false;
+	bool bSuspendModelNotificationsForOthers = false;
+	bool bSuspendAllNotifications = false;
+	bool bCompileInDebugMode = false;
+	bool bSuspendPythonMessagesForRigVMClient = true;
+	bool bSuspendBlockNotifications = false;
+#endif // #if WITH_EDITORONLY_DATA
+};
diff --git a/Engine/Plugins/Experimental/AnimNext/Source/AnimNextUncookedOnly/Internal/Param/AnimNextParameterLibrary.h b/Engine/Plugins/Experimental/AnimNext/Source/AnimNextUncookedOnly/Internal/Param/AnimNextParameterLibrary.h
new file mode 100644
--- /dev/null
+++ b/Engine/Plugins/Experimental/AnimNext/Source/AnimNextUncookedOnly/Internal/Param/AnimNextParameterLibrary.h
@@ -0,0 +1,117 @@
+﻿// Copyright Epic Games, Inc. All Rights Reserved.
+
+#pragma once
+
+#include "CoreMinimal.h"
+#include "AnimNextParameter.h"
+#include "AnimNextParameterLibrary.generated.h"
+
+class UAnimNextParameter;
+class UAnimNextParameterLibrary;
+class UAnimNextParameterBlockBinding;
+class UAnimNextParameterBlockBindingReference;
+class UAnimNextParameterBlock_EditorData;
+class FAnimationAnimNextParametersEditorTest_Library;
+class FAnimationAnimNextParametersEditorTest_Block;
+
+namespace UE::AnimNext::Editor
+{
+	struct FUtils;
+	class SParameterBlockView;
+	class SParameterLibraryView;
+}
+
+namespace UE::AnimNext::UncookedOnly
+{
+	// A delegate for subscribing / reacting to parameter library modifications.
+	DECLARE_MULTICAST_DELEGATE_OneParam(FOnParameterLibraryModified, UAnimNextParameterLibrary* /* InLibrary */);
+}
+
+// Library entry used to export to asset registry
+USTRUCT()
+struct FAnimNextParameterLibraryAssetRegistryExportEntry
+{
+	GENERATED_BODY()
+
+	FAnimNextParameterLibraryAssetRegistryExportEntry() = default;
+	
+	FAnimNextParameterLibraryAssetRegistryExportEntry(FName InName, const FAnimNextParamType& InType)
+		: Name(InName)
+		, Type(InType)
+	{}
+	
+	UPROPERTY()
+	FName Name;
+
+	UPROPERTY()
+	FAnimNextParamType Type;
+};
+
+// Library used to export to asset registry
+USTRUCT()
+struct FAnimNextParameterLibraryAssetRegistryExports
+{
+	GENERATED_BODY()
+
+	UPROPERTY()
+	TArray<FAnimNextParameterLibraryAssetRegistryExportEntry> Parameters;
+};
+
+/** Defines a single parameter */
+UCLASS(MinimalAPI, BlueprintType)
+class UAnimNextParameterLibrary : public UObject
+{
+	GENERATED_BODY()
+
+#if WITH_EDITORONLY_DATA
+	ANIMNEXTUNCOOKEDONLY_API static const FName ExportsAssetRegistryTag;
+
+	friend struct UE::AnimNext::Editor::FUtils;
+	friend class UE::AnimNext::Editor::SParameterBlockView;
+	friend class UE::AnimNext::Editor::SParameterLibraryView;
+	friend class UAnimNextParameterBlockBinding;
+	friend class UAnimNextParameterBlockBindingReference;
+	friend class UAnimNextParameterBlock_EditorData;
+	friend class FAnimationAnimNextParametersEditorTest_Library;
+	friend class FAnimationAnimNextParametersEditorTest_Block;
+
+	ANIMNEXTUNCOOKEDONLY_API UAnimNextParameter* AddParameter(FName InName, const FAnimNextParamType& InType, bool bSetupUndoRedo = true, bool bPrintPythonCommand = true);
+
+	UFUNCTION(BlueprintCallable, Category = "AnimNext|Parameter Library")
+	ANIMNEXTUNCOOKEDONLY_API UAnimNextParameter* AddParameter(FName InName, const EPropertyBagPropertyType& InValueType, const EPropertyBagContainerType& InContainerType, UObject* InValueTypeObject = nullptr, bool bSetupUndoRedo = true, bool bPrintPythonCommand = true);
+
+	UFUNCTION(BlueprintCallable, Category = "AnimNext|Parameter Library")
+	ANIMNEXTUNCOOKEDONLY_API bool RemoveParameter(FName InName, bool bSetupUndoRedo = true, bool bPrintPythonCommand = true);
+
+	UFUNCTION(BlueprintCallable, Category = "AnimNext|Parameter Library")
+	ANIMNEXTUNCOOKEDONLY_API bool RemoveParameters(const TArray<FName>& InNames, bool bSetupUndoRedo = true, bool bPrintPythonCommand = true);
+
+	UFUNCTION(BlueprintCallable, Category = "AnimNext|Parameter Library")
+	ANIMNEXTUNCOOKEDONLY_API UAnimNextParameter* FindParameter(FName InName) const;
+
+	// Returns the modified event which can be used to subscribe to changes to this parameter
+	ANIMNEXTUNCOOKEDONLY_API UE::AnimNext::UncookedOnly::FOnParameterLibraryModified& OnModified();
+
+	// UObject interface 
+	virtual void PostTransacted(const FTransactionObjectEvent& TransactionEvent) override;
+	virtual void PostLoad() override;
+	virtual void GetAssetRegistryTags(TArray<FAssetRegistryTag>& OutTags) const override;
+
+	void ReportError(const TCHAR* InMessage) const;
+
+	void BroadcastModified();
+	
+	/** Comment to display in editor */
+	UPROPERTY(EditAnywhere, Category = "Parameter", meta=(MultiLine))
+	FString Comment;
+
+	/** All parameters contained in this library - not saved, discovered at load time */
+	UPROPERTY(Transient)
+	TArray<TObjectPtr<UAnimNextParameter>> Parameters;
+
+	// Delegate to subscribe to modifications
+	UE::AnimNext::UncookedOnly::FOnParameterLibraryModified ModifiedDelegate;
+
+	bool bSuspendNotifications = false;
+#endif // #if WITH_EDITORONLY_DATA
+};
\ No newline at end of file
diff --git a/Engine/Plugins/Experimental/AnimNext/Source/AnimNextUncookedOnly/Internal/Param/IAnimNextParameterBlockBindingInterface.h b/Engine/Plugins/Experimental/AnimNext/Source/AnimNextUncookedOnly/Internal/Param/IAnimNextParameterBlockBindingInterface.h
new file mode 100644
--- /dev/null
+++ b/Engine/Plugins/Experimental/AnimNext/Source/AnimNextUncookedOnly/Internal/Param/IAnimNextParameterBlockBindingInterface.h
@@ -0,0 +1,52 @@
+﻿// Copyright Epic Games, Inc. All Rights Reserved.
+
+#pragma once
+
+#include "Param/ParamType.h"
+#include "UObject/Interface.h"
+#include "IAnimNextParameterBlockBindingInterface.generated.h"
+
+class UAnimNextParameter;
+class UAnimNextParameterBlock_EditorData;
+
+namespace UE::AnimNext::Editor
+{
+	class SParameterBlockView;
+	class SParameterBlockViewRow;
+}
+
+namespace UE::AnimNext::UncookedOnly
+{
+	struct FUtils;
+}
+
+UINTERFACE(meta=(CannotImplementInterfaceInBlueprint))
+class ANIMNEXTUNCOOKEDONLY_API UAnimNextParameterBlockBindingInterface : public UInterface
+{
+	GENERATED_BODY()
+};
+
+class ANIMNEXTUNCOOKEDONLY_API IAnimNextParameterBlockBindingInterface
+{
+	GENERATED_BODY()
+
+	friend class UAnimNextParameterBlock_EditorData;
+	friend class UE::AnimNext::Editor::SParameterBlockView;
+	friend class UE::AnimNext::Editor::SParameterBlockViewRow;
+	friend struct UE::AnimNext::UncookedOnly::FUtils;
+
+	// Get the parameter type
+	virtual FAnimNextParamType GetParamType() const = 0;
+
+	// Get the parameter name
+	virtual FName GetParameterName() const = 0;
+
+	// Set the parameter name
+	virtual void SetParameterName(FName InName, bool bSetupUndoRedo = true) = 0;
+
+	// Resolve the underlying parameter in the library
+	virtual const UAnimNextParameter* GetParameter() const = 0;
+
+	// Get the parameter library hosting the parameter
+	virtual const UAnimNextParameterLibrary* GetLibrary() const = 0;
+};
\ No newline at end of file
diff --git a/Engine/Plugins/Experimental/AnimNext/Source/AnimNextUncookedOnly/Internal/Param/IAnimNextParameterBlockGraphInterface.h b/Engine/Plugins/Experimental/AnimNext/Source/AnimNextUncookedOnly/Internal/Param/IAnimNextParameterBlockGraphInterface.h
new file mode 100644
--- /dev/null
+++ b/Engine/Plugins/Experimental/AnimNext/Source/AnimNextUncookedOnly/Internal/Param/IAnimNextParameterBlockGraphInterface.h
@@ -0,0 +1,40 @@
+﻿// Copyright Epic Games, Inc. All Rights Reserved.
+
+#pragma once
+
+#include "Param/ParamType.h"
+#include "UObject/Interface.h"
+#include "IAnimNextParameterBlockGraphInterface.generated.h"
+
+class URigVMGraph;
+class UAnimNextParameterBlock_EditorData;
+
+namespace UE::AnimNext::Editor
+{
+	class SParameterBlockView;
+	class SParameterBlockViewRow;
+}
+
+namespace UE::AnimNext::UncookedOnly
+{
+	struct FUtilsPrivate;
+}
+
+UINTERFACE(meta=(CannotImplementInterfaceInBlueprint))
+class ANIMNEXTUNCOOKEDONLY_API UAnimNextParameterBlockGraphInterface : public UInterface
+{
+	GENERATED_BODY()
+};
+
+class ANIMNEXTUNCOOKEDONLY_API IAnimNextParameterBlockGraphInterface
+{
+	GENERATED_BODY()
+
+	friend class UAnimNextParameterBlock_EditorData;
+	friend class UE::AnimNext::Editor::SParameterBlockView;
+	friend class UE::AnimNext::Editor::SParameterBlockViewRow;
+	friend struct UE::AnimNext::UncookedOnly::FUtilsPrivate;
+
+	// Get the graph
+	virtual URigVMGraph* GetGraph() const = 0;
+};
\ No newline at end of file
diff --git a/Engine/Plugins/Experimental/AnimNext/Source/AnimNextUncookedOnly/Internal/Param/IAnimNextParameterBlockReferenceInterface.h b/Engine/Plugins/Experimental/AnimNext/Source/AnimNextUncookedOnly/Internal/Param/IAnimNextParameterBlockReferenceInterface.h
new file mode 100644
--- /dev/null
+++ b/Engine/Plugins/Experimental/AnimNext/Source/AnimNextUncookedOnly/Internal/Param/IAnimNextParameterBlockReferenceInterface.h
@@ -0,0 +1,40 @@
+﻿// Copyright Epic Games, Inc. All Rights Reserved.
+
+#pragma once
+
+#include "Param/ParamType.h"
+#include "UObject/Interface.h"
+#include "IAnimNextParameterBlockReferenceInterface.generated.h"
+
+class URigVMGraph;
+class UAnimNextParameter;
+class UAnimNextParameterBlock;
+
+namespace UE::AnimNext::Editor
+{
+class SParameterBlockView;
+class SParameterBlockViewRow;
+}
+
+namespace UE::AnimNext::UncookedOnly
+{
+struct FUtilsPrivate;
+}
+
+UINTERFACE(meta=(CannotImplementInterfaceInBlueprint))
+class ANIMNEXTUNCOOKEDONLY_API UAnimNextParameterBlockReferenceInterface : public UInterface
+{
+	GENERATED_BODY()
+};
+
+class ANIMNEXTUNCOOKEDONLY_API IAnimNextParameterBlockReferenceInterface
+{
+	GENERATED_BODY()
+
+	friend class UE::AnimNext::Editor::SParameterBlockView;
+	friend class UE::AnimNext::Editor::SParameterBlockViewRow;
+	friend struct UE::AnimNext::UncookedOnly::FUtilsPrivate;
+
+	// Get the block we reference
+	virtual const UAnimNextParameterBlock* GetBlock() const = 0;
+};
\ No newline at end of file
diff --git a/Engine/Plugins/Experimental/AnimNext/Source/AnimNextUncookedOnly/Internal/UncookedOnlyUtils.h b/Engine/Plugins/Experimental/AnimNext/Source/AnimNextUncookedOnly/Internal/UncookedOnlyUtils.h
new file mode 100644
--- /dev/null
+++ b/Engine/Plugins/Experimental/AnimNext/Source/AnimNextUncookedOnly/Internal/UncookedOnlyUtils.h
@@ -0,0 +1,74 @@
+﻿// Copyright Epic Games, Inc. All Rights Reserved.
+
+#pragma once
+
+#include "CoreMinimal.h"
+#include "Param/ParamTypeHandle.h"
+#include "RigVMCore/RigVMTemplate.h"
+
+class UAnimNextGraph;
+class UAnimNextGraph_EditorData;
+class UAnimNextGraph_EdGraph;
+class URigVMController;
+class URigVMGraph;
+class UAnimNextParameter;
+class UAnimNextParameterBlock;
+class UAnimNextParameterBlock_EditorData;
+class UAnimNextGraph_EdGraph;
+struct FEdGraphPinType;
+
+namespace UE::AnimNext::UncookedOnly
+{
+
+struct ANIMNEXTUNCOOKEDONLY_API FUtils
+{
+	static void Compile(UAnimNextGraph* InGraph);
+	
+	static UAnimNextGraph_EditorData* GetEditorData(const UAnimNextGraph* InAnimNextInterfaceGraph);
+	
+	static void RecreateVM(UAnimNextGraph* InGraph);
+
+	/**
+	 * Get an AnimNext parameter type handle from an FEdGraphPinType.
+	 * Note that the returned handle may not be valid, so should be checked using IsValid() before use.
+	 **/
+	static FParamTypeHandle GetParameterHandleFromPin(const FEdGraphPinType& InPinType);
+
+	static void Compile(UAnimNextParameterBlock* InParameterBlock);
+
+	static void CompileVM(UAnimNextParameterBlock* InParameterBlock);
+
+	static void CompileStruct(UAnimNextParameterBlock* InParameterBlock);
+	
+	static UAnimNextParameterBlock_EditorData* GetEditorData(const UAnimNextParameterBlock* InParameterBlock);
+
+	static UAnimNextParameterBlock* GetBlock(const UAnimNextParameterBlock_EditorData* InEditorData);
+
+	static void RecreateVM(UAnimNextParameterBlock* InParameterBlock);
+
+	/**
+	 * Get an AnimNext parameter type from an FEdGraphPinType.
+	 * Note that the returned handle may not be valid, so should be checked using IsValid() before use.
+	 **/
+	static FParamTypeHandle GetParamTypeHandleFromPinType(const FEdGraphPinType& InPinType);
+	static FAnimNextParamType GetParamTypeFromPinType(const FEdGraphPinType& InPinType);
+
+	/**
+	 * Get an FEdGraphPinType from an AnimNext parameter type/handle.
+	 * Note that the returned pin type may not be valid.
+	 **/
+	static FEdGraphPinType GetPinTypeFromParamTypeHandle(const FParamTypeHandle& InParamTypeHandle);
+	static FEdGraphPinType GetPinTypeFromParamType(const FAnimNextParamType& InParamType);
+
+	/**
+	 * Get an FRigVMTemplateArgumentType from an AnimNext parameter type/handle.
+	 * Note that the returned pin type may not be valid.
+	 **/
+	static FRigVMTemplateArgumentType GetRigVMArgTypeFromParamTypeHandle(const FParamTypeHandle& InParamTypeHandle);
+	static FRigVMTemplateArgumentType GetRigVMArgTypeFromParamType(const FAnimNextParamType& InParamType);
+
+	/** Set up a binding graph given the type to set */
+	static void SetupBindingGraphForLiteral(URigVMController* InController, const FAnimNextParamType& InParamType);
+};
+
+}
\ No newline at end of file
diff --git a/Engine/Plugins/Experimental/AnimNext/Source/AnimNextUncookedOnly/Private/AnimNextUncookedOnlyModule.cpp b/Engine/Plugins/Experimental/AnimNext/Source/AnimNextUncookedOnly/Private/AnimNextUncookedOnlyModule.cpp
new file mode 100644
--- /dev/null
+++ b/Engine/Plugins/Experimental/AnimNext/Source/AnimNextUncookedOnly/Private/AnimNextUncookedOnlyModule.cpp
@@ -0,0 +1,5 @@
+// Copyright Epic Games, Inc. All Rights Reserved.
+
+#include "Modules/ModuleManager.h"
+
+IMPLEMENT_MODULE(FDefaultModuleImpl, AnimNextUncookedOnly);
\ No newline at end of file
diff --git a/Engine/Plugins/Experimental/AnimNext/Source/AnimNextUncookedOnly/Private/Graph/AnimNextGraph_EdGraph.cpp b/Engine/Plugins/Experimental/AnimNext/Source/AnimNextUncookedOnly/Private/Graph/AnimNextGraph_EdGraph.cpp
new file mode 100644
--- /dev/null
+++ b/Engine/Plugins/Experimental/AnimNext/Source/AnimNextUncookedOnly/Private/Graph/AnimNextGraph_EdGraph.cpp
@@ -0,0 +1,21 @@
+﻿// Copyright Epic Games, Inc. All Rights Reserved.
+
+#include "Graph/AnimNextGraph_EdGraph.h"
+#include "Graph/AnimNextGraph_EditorData.h"
+
+void UAnimNextGraph_EdGraph::Initialize(UAnimNextGraph_EditorData* InEditorData)
+{
+	InEditorData->ModifiedEvent.RemoveAll(this);
+	InEditorData->ModifiedEvent.AddUObject(this, &UAnimNextGraph_EdGraph::HandleModifiedEvent);
+	InEditorData->VMCompiledEvent.RemoveAll(this);
+	InEditorData->VMCompiledEvent.AddUObject(this, &UAnimNextGraph_EdGraph::HandleVMCompiledEvent);
+}
+
+FRigVMClient* UAnimNextGraph_EdGraph::GetRigVMClient() const
+{
+	if (const UAnimNextGraph_EditorData* EditorData = GetTypedOuter<UAnimNextGraph_EditorData>())
+	{
+		return (FRigVMClient*)EditorData->GetRigVMClient();
+	}
+	return nullptr;
+}
diff --git a/Engine/Plugins/Experimental/AnimNext/Source/AnimNextUncookedOnly/Private/Graph/AnimNextGraph_EdGraphSchema.cpp b/Engine/Plugins/Experimental/AnimNext/Source/AnimNextUncookedOnly/Private/Graph/AnimNextGraph_EdGraphSchema.cpp
new file mode 100644
--- /dev/null
+++ b/Engine/Plugins/Experimental/AnimNext/Source/AnimNextUncookedOnly/Private/Graph/AnimNextGraph_EdGraphSchema.cpp
@@ -0,0 +1,42 @@
+﻿// Copyright Epic Games, Inc. All Rights Reserved.
+
+#include "Graph/AnimNextGraph_EdGraphSchema.h"
+#include "Editor.h"
+
+void UAnimNextGraph_EdGraphSchema::TrySetDefaultValue(UEdGraphPin& InPin, const FString& InNewDefaultValue, bool bMarkAsModified) const
+{
+#if WITH_EDITOR
+	if (GEditor)
+	{
+		GEditor->CancelTransaction(0);
+	}
+#endif
+
+	FString UseDefaultValue;
+	TObjectPtr<UObject> UseDefaultObject = nullptr;
+	FText UseDefaultText;
+
+	GetDefault<UEdGraphSchema_K2>()->GetPinDefaultValuesFromString(InPin.PinType, InPin.GetOwningNodeUnchecked(), InNewDefaultValue, UseDefaultValue, UseDefaultObject, UseDefaultText, /*bPreserveTextIdentity*/false);
+
+	// Check the default value and make it an error if it's bogus
+	if (GetDefault<UEdGraphSchema_K2>()->DefaultValueSimpleValidation(InPin.PinType, InPin.PinName, UseDefaultValue, UseDefaultObject, UseDefaultText))
+	{
+		InPin.DefaultObject = UseDefaultObject;
+		InPin.DefaultValue = UseDefaultValue;
+		InPin.DefaultTextValue = UseDefaultText;
+
+		UEdGraphNode* Node = InPin.GetOwningNode();
+		Node->PinDefaultValueChanged(&InPin);
+
+		// If the default value is manually set then treat it as if the value was reset to default and remove the orphaned InPin
+		if (InPin.bOrphanedPin && InPin.DoesDefaultValueMatchAutogenerated())
+		{
+			Node->PinConnectionListChanged(&InPin);
+		}
+
+		if (bMarkAsModified)
+		{
+			Node->MarkPackageDirty();
+		}
+	}
+}
diff --git a/Engine/Plugins/Experimental/AnimNext/Source/AnimNextUncookedOnly/Private/Graph/AnimNextGraph_EditorData.cpp b/Engine/Plugins/Experimental/AnimNext/Source/AnimNextUncookedOnly/Private/Graph/AnimNextGraph_EditorData.cpp
new file mode 100644
--- /dev/null
+++ b/Engine/Plugins/Experimental/AnimNext/Source/AnimNextUncookedOnly/Private/Graph/AnimNextGraph_EditorData.cpp
@@ -0,0 +1,555 @@
+﻿// Copyright Epic Games, Inc. All Rights Reserved.
+
+#include "Graph/AnimNextGraph_EditorData.h"
+
+#include "ControlRigDefines.h"
+#include "Graph/AnimNextGraph.h"
+#include "Graph/AnimNextGraph_EdGraphSchema.h"
+#include "UncookedOnlyUtils.h"
+#include "Graph/GraphExecuteContext.h"
+#include "Rigs/RigHierarchyPose.h"
+#include "RigVMModel/RigVMFunctionLibrary.h"
+#include "RigVMModel/RigVMNotifications.h"
+#include "RigVMModel/Nodes/RigVMCollapseNode.h"
+#include "Curves/CurveFloat.h"
+#include "UObject/ObjectSaveContext.h"
+#include "Kismet2/BlueprintEditorUtils.h"
+
+UAnimNextGraph_EditorData::UAnimNextGraph_EditorData(const FObjectInitializer& ObjectInitializer)
+{
+	RigVMClient.Reset();
+	RigVMClient.SetOuterClientHost(this, GET_MEMBER_NAME_CHECKED(UAnimNextGraph_EditorData, RigVMClient));
+	{
+		TGuardValue<bool> DisableClientNotifs(RigVMClient.bSuspendNotifications, true);
+		RigVMClient.AddModel(TEXT("RigVMGraph"), false, &ObjectInitializer);
+		RigVMClient.GetOrCreateFunctionLibrary(false, &ObjectInitializer);
+	}
+	RigVMClient.SetExecuteContextStruct(FAnimNextGraphExecuteContext::StaticStruct());
+	
+	auto MakeEdGraph = [this, &ObjectInitializer](FName InName) -> UAnimNextGraph_EdGraph*
+	{
+		UAnimNextGraph_EdGraph* EdGraph = ObjectInitializer.CreateDefaultSubobject<UAnimNextGraph_EdGraph>(this, InName);
+		EdGraph->Schema = UAnimNextGraph_EdGraphSchema::StaticClass();
+		EdGraph->bAllowRenaming = false;
+		EdGraph->bEditable = true;
+		EdGraph->bAllowDeletion = false;
+		EdGraph->bIsFunctionDefinition = false;
+		EdGraph->Initialize(this);
+
+		return EdGraph;
+	};
+
+	RootGraph = MakeEdGraph(TEXT("RootEdGraph"));
+	FunctionLibraryEdGraph = MakeEdGraph(TEXT("RigVMFunctionLibraryEdGraph"));
+}
+
+void UAnimNextGraph_EditorData::Serialize(FArchive& Ar)
+{
+	RigVMClient.SetOuterClientHost(this, GET_MEMBER_NAME_CHECKED(UAnimNextGraph_EditorData, RigVMClient));
+
+	UObject::Serialize(Ar);
+}
+
+void UAnimNextGraph_EditorData::Initialize(bool bRecompileVM)
+{
+	UAnimNextGraph* AnimNextGraph = GetTypedOuter<UAnimNextGraph>();
+
+	if (RigVMClient.GetController(0) == nullptr)
+	{
+		check(RigVMClient.Num() == 1);
+		check(RigVMClient.GetFunctionLibrary());
+		
+		RigVMClient.GetOrCreateController(RigVMClient.GetDefaultModel());
+		RigVMClient.GetOrCreateController(RigVMClient.GetFunctionLibrary());
+
+		// Init function library controllers
+		for(URigVMLibraryNode* LibraryNode : RigVMClient.GetFunctionLibrary()->GetFunctions())
+		{
+			RigVMClient.GetOrCreateController(LibraryNode->GetContainedGraph());
+		}
+		
+		if(bRecompileVM)
+		{
+			UE::AnimNext::UncookedOnly::FUtils::Compile(AnimNextGraph);
+		}
+	}
+
+	RootGraph->Initialize(this);
+	FunctionLibraryEdGraph->Initialize(this);
+	if(EntryPointGraph)
+	{
+		EntryPointGraph->Initialize(this);
+	}
+}
+
+void UAnimNextGraph_EditorData::PostLoad()
+{
+	Super::PostLoad();
+	
+	Initialize(/*bRecompileVM*/false);
+	RefreshAllModels(EAnimNextGraphLoadType::PostLoad);
+
+#if WITH_EDITOR
+	if (GIsEditor)
+	{
+		// delay compilation until the package has been loaded
+		FCoreUObjectDelegates::OnEndLoadPackage.AddUObject(this, &UAnimNextGraph_EditorData::HandlePackageDone);
+	}
+#else // !WITH_EDITOR
+	RecompileVMIfRequired();
+#endif // WITH_EDITOR
+}
+
+#if WITH_EDITOR
+void UAnimNextGraph_EditorData::HandlePackageDone(const FEndLoadPackageContext& Context)
+{
+	if (!Context.LoadedPackages.Contains(GetPackage()))
+	{
+		return;
+	}
+	HandlePackageDone();
+}
+
+void UAnimNextGraph_EditorData::HandlePackageDone()
+{
+	FCoreUObjectDelegates::OnEndLoadPackage.RemoveAll(this);
+
+	RecompileVM();
+}
+
+void UAnimNextGraph_EditorData::RefreshAllModels(EAnimNextGraphLoadType InLoadType)
+{
+	const bool bIsPostLoad = InLoadType == EAnimNextGraphLoadType::PostLoad;
+
+	TGuardValue<bool> IsCompilingGuard(bIsCompiling, true);
+	TGuardValue<bool> ClientIgnoreModificationsGuard(RigVMClient.bIgnoreModelNotifications, true);
+
+	TArray<URigVMGraph*> GraphsToDetach = RigVMClient.GetAllModels(true, false);
+	TMap<const URigVMGraph*, TArray<URigVMController::FLinkedPath>> LinkedPaths;
+	
+	if (ensure(IsInGameThread()))
+	{
+		for (const URigVMGraph* GraphToDetach : GraphsToDetach)
+		{
+			URigVMController* Controller = RigVMClient.GetOrCreateController(GraphToDetach);
+			// temporarily disable default value validation during load time, serialized values should always be accepted
+			TGuardValue<bool> PerGraphDisablePinDefaultValueValidation(Controller->bValidatePinDefaults, false);
+			FRigVMControllerNotifGuard NotifGuard(Controller, true);
+			LinkedPaths.Add(GraphToDetach, Controller->GetLinkedPaths());
+			Controller->FastBreakLinkedPaths(LinkedPaths.FindChecked(GraphToDetach));
+			TArray<URigVMNode*> Nodes = GraphToDetach->GetNodes();
+			for (URigVMNode* Node : Nodes)
+			{
+				Controller->RepopulatePinsOnNode(Node, true, true);
+			}
+		}
+		//SetupPinRedirectorsForBackwardsCompatibility();
+	}
+
+	for (const URigVMGraph* GraphToDetach : GraphsToDetach)
+	{
+		URigVMController* Controller = RigVMClient.GetOrCreateController(GraphToDetach);
+		// at this stage, allow all links to be reattached,
+		// RecomputeAllTemplateFilteredPermutations() later should break any invalid links
+		FRigVMControllerNotifGuard NotifGuard(Controller, true);
+		Controller->RestoreLinkedPaths(LinkedPaths.FindChecked(GraphToDetach));
+	}
+
+	if (bIsPostLoad)
+	{
+		//PatchTemplateNodesWithPreferredPermutation();
+	}
+
+	TArray<URigVMGraph*> GraphsToClean = RigVMClient.GetAllModels(true, true);
+
+	// Sort from leaf graphs to root
+	TArray<URigVMGraph*> SortedGraphsToClean;
+	SortedGraphsToClean.Reserve(GraphsToClean.Num());
+	while (SortedGraphsToClean.Num() < GraphsToClean.Num())
+	{
+		bool bGraphAdded = false;
+		for (URigVMGraph* Graph : GraphsToClean)
+		{
+			if (SortedGraphsToClean.Contains(Graph))
+			{
+				continue;
+			}
+
+			TArray<URigVMGraph*> ContainedGraphs;
+			for (URigVMNode* Node : Graph->GetNodes())
+			{
+				if (URigVMLibraryNode* LibraryNode = Cast<URigVMLibraryNode>(Node))
+				{
+					if (URigVMFunctionReferenceNode* FunctionReferenceNode = Cast<URigVMFunctionReferenceNode>(LibraryNode))
+					{
+						if (FunctionReferenceNode->GetReferencedFunctionHeader().LibraryPointer.LibraryNode.GetLongPackageName() != GetPackage()->GetPathName())
+						{
+							continue;
+						}
+						if (URigVMLibraryNode* ReferencedNode = FunctionReferenceNode->LoadReferencedNode())
+						{
+							ContainedGraphs.Add(ReferencedNode->GetContainedGraph());
+							continue;
+						}
+					}
+
+					if (URigVMGraph* ContainedGraph = LibraryNode->GetContainedGraph())
+					{
+						ContainedGraphs.Add(ContainedGraph);
+					}
+				}
+			}
+
+			bool bAllContained = true;
+			for (URigVMGraph* Contained : ContainedGraphs)
+			{
+				if (!SortedGraphsToClean.Contains(Contained))
+				{
+					bAllContained = false;
+					break;
+				}
+			}
+			if (bAllContained)
+			{
+				SortedGraphsToClean.Add(Graph);
+				bGraphAdded = true;
+			}
+		}
+		ensure(bGraphAdded);
+	}
+
+	for (int32 GraphIndex = 0; GraphIndex < SortedGraphsToClean.Num(); GraphIndex++)
+	{
+		URigVMGraph* GraphToClean = SortedGraphsToClean[GraphIndex];
+		URigVMController* Controller = RigVMClient.GetOrCreateController(GraphToClean);
+		//TGuardValue<bool> GuardEditGraph(GraphToClean->bEditable, true);
+		FRigVMControllerNotifGuard NotifGuard(Controller, true);
+
+		for (URigVMNode* ModelNode : GraphToClean->GetNodes())
+		{
+			Controller->RemoveUnusedOrphanedPins(ModelNode);
+		}
+	}
+}
+
+#endif // WITH_EDITOR
+
+FRigVMClient* UAnimNextGraph_EditorData::GetRigVMClient()
+{
+	return &RigVMClient;
+}
+
+const FRigVMClient* UAnimNextGraph_EditorData::GetRigVMClient() const
+{
+	return &RigVMClient;
+}
+
+IRigVMGraphFunctionHost* UAnimNextGraph_EditorData::GetRigVMGraphFunctionHost()
+{
+	return this;
+}
+
+const IRigVMGraphFunctionHost* UAnimNextGraph_EditorData::GetRigVMGraphFunctionHost() const
+{
+	return this;
+}
+
+
+void UAnimNextGraph_EditorData::HandleRigVMGraphAdded(const FRigVMClient* InClient, const FString& InNodePath)
+{
+	if(URigVMGraph* RigVMGraph = InClient->GetModel(InNodePath))
+	{
+		RigVMGraph->SetExecuteContextStruct(FAnimNextGraphExecuteContext::StaticStruct());
+	}
+}
+
+void UAnimNextGraph_EditorData::HandleConfigureRigVMController(const FRigVMClient* InClient,
+                                                                    URigVMController* InControllerToConfigure)
+{
+	InControllerToConfigure->OnModified().AddUObject(this, &UAnimNextGraph_EditorData::HandleModifiedEvent);
+
+	InControllerToConfigure->UnfoldStructDelegate.BindLambda([](const UStruct* InStruct) -> bool
+	{
+		if (InStruct == TBaseStructure<FQuat>::Get())
+		{
+			return false;
+		}
+		if (InStruct == FRuntimeFloatCurve::StaticStruct())
+		{
+			return false;
+		}
+		if (InStruct == FRigPose::StaticStruct())
+		{
+			return false;
+		}
+		return true;
+	});
+
+	TWeakObjectPtr<UAnimNextGraph_EditorData> WeakThis(this);
+
+	// this delegate is used by the controller to determine variable validity
+	// during a bind process. the controller itself doesn't own the variables,
+	// so we need a delegate to request them from the owning blueprint
+	InControllerToConfigure->GetExternalVariablesDelegate.BindLambda([](URigVMGraph* InGraph) -> TArray<FRigVMExternalVariable>
+	{
+		if (InGraph)
+		{
+			if(UAnimNextGraph_EditorData* EditorData = InGraph->GetTypedOuter<UAnimNextGraph_EditorData>())
+			{
+				if (UAnimNextGraph* Graph = EditorData->GetTypedOuter<UAnimNextGraph>())
+				{
+					return Graph->GetRigVMExternalVariables();
+				}
+			}
+		}
+		return TArray<FRigVMExternalVariable>();
+	});
+
+	// this delegate is used by the controller to retrieve the current bytecode of the VM
+	InControllerToConfigure->GetCurrentByteCodeDelegate.BindLambda([WeakThis]() -> const FRigVMByteCode*
+	{
+		if (WeakThis.IsValid())
+		{
+			if(UAnimNextGraph* Graph = WeakThis->GetTypedOuter<UAnimNextGraph>())
+			{
+				if (Graph->RigVM)
+				{
+					return &Graph->RigVM->GetByteCode();
+				}
+			}
+		}
+		return nullptr;
+
+	});
+
+#if WITH_EDITOR
+	InControllerToConfigure->SetupDefaultUnitNodeDelegates(TDelegate<FName(FRigVMExternalVariable, FString)>::CreateLambda(
+		[](FRigVMExternalVariable InVariableToCreate, FString InDefaultValue) -> FName
+		{
+			return NAME_None;
+		}
+	));
+#endif
+
+}
+
+UObject* UAnimNextGraph_EditorData::GetEditorObjectForRigVMGraph(URigVMGraph* InVMGraph) const
+{
+	if(InVMGraph)
+	{
+		TArray<UAnimNextGraph_EdGraph*> AllGraphs = {RootGraph, EntryPointGraph, FunctionLibraryEdGraph};
+		for(UAnimNextGraph_EdGraph* EdGraph : AllGraphs)
+		{
+			if(EdGraph)
+			{
+				if(EdGraph->ModelNodePath == InVMGraph->GetNodePath())
+				{
+					return EdGraph;
+				}
+			}
+		}
+	}
+	return nullptr;
+}
+
+FRigVMGraphFunctionStore* UAnimNextGraph_EditorData::GetRigVMGraphFunctionStore()
+{
+	return &GraphFunctionStore;
+}
+
+const FRigVMGraphFunctionStore* UAnimNextGraph_EditorData::GetRigVMGraphFunctionStore() const
+{
+	return &GraphFunctionStore;
+}
+
+
+void UAnimNextGraph_EditorData::RecompileVM()
+{
+	UE::AnimNext::UncookedOnly::FUtils::Compile(GetTypedOuter<UAnimNextGraph>());
+}
+
+void UAnimNextGraph_EditorData::RecompileVMIfRequired()
+{
+	if (bVMRecompilationRequired)
+	{
+		RecompileVM();
+	}
+}
+
+void UAnimNextGraph_EditorData::RequestAutoVMRecompilation()
+{
+	bVMRecompilationRequired = true;
+	if (bAutoRecompileVM)
+	{
+		RecompileVMIfRequired();
+	}
+}
+
+void UAnimNextGraph_EditorData::HandleModifiedEvent(ERigVMGraphNotifType InNotifType, URigVMGraph* InGraph, UObject* InSubject)
+{
+	bool bNotifForOthersPending = true;
+
+	switch(InNotifType)
+	{
+	case ERigVMGraphNotifType::NodeAdded:
+		{
+			if (URigVMCollapseNode* CollapseNode = Cast<URigVMCollapseNode>(InSubject))
+			{
+				CreateEdGraphForCollapseNode(CollapseNode);
+				break;
+			}
+		}
+		// Fall through to the next case
+	case ERigVMGraphNotifType::LinkAdded:
+	case ERigVMGraphNotifType::LinkRemoved:
+	case ERigVMGraphNotifType::PinArraySizeChanged:
+	case ERigVMGraphNotifType::PinDirectionChanged:
+		{
+			UpdateGraphReturnType();
+			RecompileVM();
+			break;
+		}
+
+	case ERigVMGraphNotifType::PinDefaultValueChanged:
+		{
+			if (InGraph->GetRuntimeAST().IsValid())
+			{
+				URigVMPin* RootPin = CastChecked<URigVMPin>(InSubject)->GetRootPin();
+				FRigVMASTProxy RootPinProxy = FRigVMASTProxy::MakeFromUObject(RootPin);
+				const FRigVMExprAST* Expression = InGraph->GetRuntimeAST()->GetExprForSubject(RootPinProxy);
+				if (Expression == nullptr)
+				{
+					InGraph->ClearAST();
+					break;
+				}
+				else if (Expression->NumParents() > 1)
+				{
+					InGraph->ClearAST();
+					break;
+				}
+			}
+			break;
+		}
+	}
+	
+	// if the notification still has to be sent...
+	if (bNotifForOthersPending && !bSuspendModelNotificationsForOthers)
+	{
+		if (ModifiedEvent.IsBound())
+		{
+			ModifiedEvent.Broadcast(InNotifType, InGraph, InSubject);
+		}
+	}
+}
+
+URigVMGraph* UAnimNextGraph_EditorData::GetVMGraphForEdGraph(const UEdGraph* InGraph) const
+{
+	if (InGraph == RootGraph)
+	{
+		return RigVMClient.GetDefaultModel();
+	}
+	else
+	{
+		const UAnimNextGraph_EdGraph* Graph = Cast<UAnimNextGraph_EdGraph>(InGraph);
+		check(Graph);
+
+		if (Graph->bIsFunctionDefinition)
+		{
+			if (URigVMLibraryNode* LibraryNode = RigVMClient.GetFunctionLibrary()->FindFunction(*Graph->ModelNodePath))
+			{
+				return LibraryNode->GetContainedGraph();
+			}
+		}
+	}
+
+	return nullptr;
+}
+
+void UAnimNextGraph_EditorData::CreateEdGraphForCollapseNode(URigVMCollapseNode* InNode)
+{
+	if (InNode->GetGraph()->IsA<URigVMFunctionLibrary>())
+	{
+		if (URigVMGraph* ContainedGraph = InNode->GetContainedGraph())
+		{
+			if (EntryPointGraph == nullptr)
+			{
+				// create a sub graph
+				UAnimNextGraph_EdGraph* RigFunctionGraph = NewObject<UAnimNextGraph_EdGraph>(this, *InNode->GetName(), RF_Transactional);
+				RigFunctionGraph->Schema = UAnimNextGraph_EdGraphSchema::StaticClass();
+				RigFunctionGraph->bAllowRenaming = true;
+				RigFunctionGraph->bEditable = true;
+				RigFunctionGraph->bAllowDeletion = true;
+				RigFunctionGraph->ModelNodePath = ContainedGraph->GetNodePath();
+				RigFunctionGraph->bIsFunctionDefinition = true;
+
+				EntryPointGraph = RigFunctionGraph;
+				RigFunctionGraph->Initialize(this);
+
+				RigVMClient.GetOrCreateController(ContainedGraph)->ResendAllNotifications();
+			}
+		}
+	}
+}
+
+FEdGraphPinType UAnimNextGraph_EditorData::FindGraphReturnPinType() const
+{
+	static const TCHAR* EndExecStr = TEXT("AnimNextEndExec");
+	static const FName ResultName = TEXT(".Result");
+
+	const UAnimNextGraph_EdGraph* AnimNextInterfaceEdGraph = RootGraph.Get();
+	if (AnimNextInterfaceEdGraph != nullptr)
+	{
+		for (const TObjectPtr<UEdGraphNode>& Node : AnimNextInterfaceEdGraph->Nodes)
+		{
+			if (IsNodeExecConnected(Node))
+			{
+				const FString& NodeName = Node->GetName();
+				if (NodeName.Contains(EndExecStr))
+				{
+					for (const UEdGraphPin* Pin : Node->Pins)
+					{
+						const FString PinName = Pin->PinName.ToString();
+
+						if (PinName.Contains(ResultName.ToString()) && PinName.Contains(EndExecStr))
+						{
+							return Pin->PinType;
+						}
+					}
+				}
+			}
+		}
+	}
+
+	return FEdGraphPinType();
+}
+
+bool UAnimNextGraph_EditorData::IsNodeExecConnected(const UEdGraphNode* Node) const
+{
+	static const TCHAR* ExecuteContextStr = TEXT(".ExecuteContext");
+
+	for (const UEdGraphPin* Pin : Node->Pins)
+	{
+		const FString PinName = Pin->PinName.ToString();
+
+		if (PinName.Contains(ExecuteContextStr) && Pin->LinkedTo.Num() > 0)
+		{
+			return true;
+		}
+	}
+
+	return false;
+}
+
+void UAnimNextGraph_EditorData::UpdateGraphReturnType()
+{
+	if (UAnimNextGraph* Graph = GetTypedOuter<UAnimNextGraph>())
+	{
+		FEdGraphPinType GraphReturnPinType = FindGraphReturnPinType();
+		UE::AnimNext::FParamTypeHandle Handle = UE::AnimNext::UncookedOnly::FUtils::GetParameterHandleFromPin(GraphReturnPinType);
+		if (Handle.IsValid() && Graph->GetReturnTypeHandle() != Handle)
+		{
+			Graph->Modify();
+			Graph->SetReturnTypeHandle(Handle);
+		}
+	}
+}
diff --git a/Engine/Plugins/Experimental/AnimNext/Source/AnimNextUncookedOnly/Private/Param/AnimNextParameter.cpp b/Engine/Plugins/Experimental/AnimNext/Source/AnimNextUncookedOnly/Private/Param/AnimNextParameter.cpp
new file mode 100644
--- /dev/null
+++ b/Engine/Plugins/Experimental/AnimNext/Source/AnimNextUncookedOnly/Private/Param/AnimNextParameter.cpp
@@ -0,0 +1,39 @@
+﻿// Copyright Epic Games, Inc. All Rights Reserved.
+
+#include "Param/AnimNextParameter.h"
+
+#if WITH_EDITORONLY_DATA
+
+UE::AnimNext::UncookedOnly::FOnParameterModified& UAnimNextParameter::OnModified()
+{
+	return ModifiedDelegate;
+}
+
+void UAnimNextParameter::SetType(const FAnimNextParamType& InType, bool bSetupUndoRedo)
+{
+	ensure(InType.IsValid());
+
+	if(bSetupUndoRedo)
+	{
+		Modify();
+	}
+
+	Type = InType;
+
+	ModifiedDelegate.Broadcast(this);
+}
+
+void UAnimNextParameter::PostTransacted(const FTransactionObjectEvent& TransactionEvent)
+{
+	Super::PostTransacted(TransactionEvent);
+	ModifiedDelegate.Broadcast(this);
+}
+
+bool UAnimNextParameter::IsAsset() const
+{
+	// Parameters are considered assets to allow using the asset logic for save dialogs, etc.
+	// Also, they return true even if pending kill, in order to show up as deleted in these dialogs.
+	return IsPackageExternal() && !GetPackage()->HasAnyFlags(RF_Transient) && !HasAnyFlags(RF_Transient | RF_ClassDefaultObject);
+}
+
+#endif // #if WITH_EDITORONLY_DATA
\ No newline at end of file
diff --git a/Engine/Plugins/Experimental/AnimNext/Source/AnimNextUncookedOnly/Private/Param/AnimNextParameterBlockBinding.cpp b/Engine/Plugins/Experimental/AnimNext/Source/AnimNextUncookedOnly/Private/Param/AnimNextParameterBlockBinding.cpp
new file mode 100644
--- /dev/null
+++ b/Engine/Plugins/Experimental/AnimNext/Source/AnimNextUncookedOnly/Private/Param/AnimNextParameterBlockBinding.cpp
@@ -0,0 +1,80 @@
+﻿// Copyright Epic Games, Inc. All Rights Reserved.
+
+#include "Param/AnimNextParameterBlockBinding.h"
+#include "Param/AnimNextParameter.h"
+#include "Param/AnimNextParameterLibrary.h"
+
+#define LOCTEXT_NAMESPACE "AnimNextParameterBlockBinding"
+
+FAnimNextParamType UAnimNextParameterBlockBinding::GetParamType() const
+{
+	if(Library)
+	{
+		if(const UAnimNextParameter* Parameter = Library->FindParameter(ParameterName))
+		{
+			return Parameter->GetType();
+		}
+	}
+
+	return FAnimNextParamType();
+}
+
+FName UAnimNextParameterBlockBinding::GetParameterName() const
+{
+	return ParameterName;
+}
+
+void UAnimNextParameterBlockBinding::SetParameterName(FName InName, bool bSetupUndoRedo)
+{
+	if(bSetupUndoRedo)
+	{
+		Modify();
+	}
+
+	ParameterName = InName;
+
+	BroadcastModified();
+}
+
+const UAnimNextParameter* UAnimNextParameterBlockBinding::GetParameter() const
+{
+	if(Library)
+	{
+		return Library->FindParameter(ParameterName);
+	}
+
+	return nullptr;
+}
+
+const UAnimNextParameterLibrary* UAnimNextParameterBlockBinding::GetLibrary() const
+{
+	return Library;
+}
+
+FText UAnimNextParameterBlockBinding::GetDisplayName() const
+{
+	return FText::FromName(ParameterName);
+}
+
+FText UAnimNextParameterBlockBinding::GetDisplayNameTooltip() const
+{
+	FTextBuilder TextBuilder;
+	TextBuilder.AppendLine(FText::FromName(ParameterName));
+	TextBuilder.AppendLine(Library ? FText::FromString(Library->GetPathName()) : LOCTEXT("MissingLibrary", "Missing parameter library"));
+	return TextBuilder.ToText();
+}
+
+void UAnimNextParameterBlockBinding::GetEditedObjects(TArray<UObject*>& OutObjects) const
+{
+	if(Library)
+	{
+		OutObjects.Add(Library);
+	}
+
+	if(const UAnimNextParameter* Parameter = GetParameter())
+	{
+		OutObjects.Add(const_cast<UAnimNextParameter*>(Parameter));
+	}
+}
+
+#undef LOCTEXT_NAMESPACE
\ No newline at end of file
diff --git a/Engine/Plugins/Experimental/AnimNext/Source/AnimNextUncookedOnly/Private/Param/AnimNextParameterBlockBindingReference.cpp b/Engine/Plugins/Experimental/AnimNext/Source/AnimNextUncookedOnly/Private/Param/AnimNextParameterBlockBindingReference.cpp
new file mode 100644
--- /dev/null
+++ b/Engine/Plugins/Experimental/AnimNext/Source/AnimNextUncookedOnly/Private/Param/AnimNextParameterBlockBindingReference.cpp
@@ -0,0 +1,84 @@
+﻿// Copyright Epic Games, Inc. All Rights Reserved.
+
+#include "Param/AnimNextParameterBlockBindingReference.h"
+#include "Param/AnimNextParameter.h"
+#include "Param/AnimNextParameterLibrary.h"
+#include "Param/AnimNextParameterBlock.h"
+
+#define LOCTEXT_NAMESPACE "AnimNextParameterBlockBindingReference"
+
+FAnimNextParamType UAnimNextParameterBlockBindingReference::GetParamType() const
+{
+	if(Library)
+	{
+		if(const UAnimNextParameter* Parameter = Library->FindParameter(ParameterName))
+		{
+			return Parameter->GetType();
+		}
+	}
+
+	return FAnimNextParamType();
+}
+
+FName UAnimNextParameterBlockBindingReference::GetParameterName() const
+{
+	return ParameterName;
+}
+
+void UAnimNextParameterBlockBindingReference::SetParameterName(FName InName, bool bSetupUndoRedo)
+{
+	ParameterName = InName;
+}
+
+const UAnimNextParameter* UAnimNextParameterBlockBindingReference::GetParameter() const
+{
+	if(Library)
+	{
+		return Library->FindParameter(ParameterName);
+	}
+
+	return nullptr;
+}
+
+const UAnimNextParameterLibrary* UAnimNextParameterBlockBindingReference::GetLibrary() const
+{
+	return Library;
+}
+
+const UAnimNextParameterBlock* UAnimNextParameterBlockBindingReference::GetBlock() const
+{
+	return Block;
+}
+
+FText UAnimNextParameterBlockBindingReference::GetDisplayName() const
+{
+	return FText::FromName(ParameterName);
+}
+
+FText UAnimNextParameterBlockBindingReference::GetDisplayNameTooltip() const
+{
+	FTextBuilder TextBuilder;
+	TextBuilder.AppendLine(FText::FromName(ParameterName));
+	TextBuilder.AppendLine(Library ? FText::FromString(Library->GetPathName()) : LOCTEXT("MissingLibrary", "Missing parameter library"));
+	return TextBuilder.ToText();
+}
+
+void UAnimNextParameterBlockBindingReference::GetEditedObjects(TArray<UObject*>& OutObjects) const
+{
+	if(Library)
+	{
+		OutObjects.Add(Library);
+	}
+
+	if(const UAnimNextParameter* Parameter = GetParameter())
+	{
+		OutObjects.Add(const_cast<UAnimNextParameter*>(Parameter));
+	}
+	
+	if(Block)
+	{
+		OutObjects.Add(Block);
+	}
+}
+
+#undef LOCTEXT_NAMESPACE
\ No newline at end of file
diff --git a/Engine/Plugins/Experimental/AnimNext/Source/AnimNextUncookedOnly/Private/Param/AnimNextParameterBlockEntry.cpp b/Engine/Plugins/Experimental/AnimNext/Source/AnimNextUncookedOnly/Private/Param/AnimNextParameterBlockEntry.cpp
new file mode 100644
--- /dev/null
+++ b/Engine/Plugins/Experimental/AnimNext/Source/AnimNextUncookedOnly/Private/Param/AnimNextParameterBlockEntry.cpp
@@ -0,0 +1,30 @@
+﻿// Copyright Epic Games, Inc. All Rights Reserved.
+
+#include "Param/AnimNextParameterBlockEntry.h"
+#include "Param/AnimNextParameterBlock_EditorData.h"
+
+void UAnimNextParameterBlockEntry::Initialize(UAnimNextParameterBlock_EditorData* InEditorData)
+{
+	InEditorData->RigVMGraphModifiedEvent.RemoveAll(this);
+	InEditorData->RigVMGraphModifiedEvent.AddUObject(this, &UAnimNextParameterBlockEntry::HandleRigVMGraphModifiedEvent);
+}
+
+void UAnimNextParameterBlockEntry::HandleRigVMGraphModifiedEvent(ERigVMGraphNotifType InNotifType, URigVMGraph* InGraph, UObject* InSubject)
+{
+	
+}
+
+bool UAnimNextParameterBlockEntry::IsAsset() const
+{
+	// Entries are considered assets to allow using the asset logic for save dialogs, etc.
+	// Also, they return true even if pending kill, in order to show up as deleted in these dialogs.
+	return IsPackageExternal() && !GetPackage()->HasAnyFlags(RF_Transient) && !HasAnyFlags(RF_Transient | RF_ClassDefaultObject);
+}
+
+void UAnimNextParameterBlockEntry::BroadcastModified()
+{
+	if(UAnimNextParameterBlock_EditorData* EditorData = Cast<UAnimNextParameterBlock_EditorData>(GetOuter()))
+	{
+		EditorData->BroadcastModified();
+	}
+}
\ No newline at end of file
diff --git a/Engine/Plugins/Experimental/AnimNext/Source/AnimNextUncookedOnly/Private/Param/AnimNextParameterBlock_EdGraph.cpp b/Engine/Plugins/Experimental/AnimNext/Source/AnimNextUncookedOnly/Private/Param/AnimNextParameterBlock_EdGraph.cpp
new file mode 100644
--- /dev/null
+++ b/Engine/Plugins/Experimental/AnimNext/Source/AnimNextUncookedOnly/Private/Param/AnimNextParameterBlock_EdGraph.cpp
@@ -0,0 +1,21 @@
+﻿// Copyright Epic Games, Inc. All Rights Reserved.
+
+#include "Param/AnimNextParameterBlock_EdGraph.h"
+#include "Param/AnimNextParameterBlock_EditorData.h"
+
+void UAnimNextParameterBlock_EdGraph::Initialize(UAnimNextParameterBlock_EditorData* InEditorData)
+{
+	InEditorData->RigVMGraphModifiedEvent.RemoveAll(this);
+	InEditorData->RigVMGraphModifiedEvent.AddUObject(this, &UAnimNextParameterBlock_EdGraph::HandleModifiedEvent);
+	InEditorData->RigVMCompiledEvent.RemoveAll(this);
+	InEditorData->RigVMCompiledEvent.AddUObject(this, &UAnimNextParameterBlock_EdGraph::HandleVMCompiledEvent);
+}
+
+FRigVMClient* UAnimNextParameterBlock_EdGraph::GetRigVMClient() const
+{
+	if (const UAnimNextParameterBlock_EditorData* EditorData = GetTypedOuter<UAnimNextParameterBlock_EditorData>())
+	{
+		return (FRigVMClient*)EditorData->GetRigVMClient();
+	}
+	return nullptr;
+}
diff --git a/Engine/Plugins/Experimental/AnimNext/Source/AnimNextUncookedOnly/Private/Param/AnimNextParameterBlock_EdGraphSchema.cpp b/Engine/Plugins/Experimental/AnimNext/Source/AnimNextUncookedOnly/Private/Param/AnimNextParameterBlock_EdGraphSchema.cpp
new file mode 100644
--- /dev/null
+++ b/Engine/Plugins/Experimental/AnimNext/Source/AnimNextUncookedOnly/Private/Param/AnimNextParameterBlock_EdGraphSchema.cpp
@@ -0,0 +1,42 @@
+﻿// Copyright Epic Games, Inc. All Rights Reserved.
+
+#include "Param/AnimNextParameterBlock_EdGraphSchema.h"
+#include "Editor.h"
+
+void UAnimNextParameterBlock_EdGraphSchema::TrySetDefaultValue(UEdGraphPin& InPin, const FString& InNewDefaultValue, bool bMarkAsModified) const
+{
+#if WITH_EDITOR
+	if (GEditor)
+	{
+		GEditor->CancelTransaction(0);
+	}
+#endif
+
+	FString UseDefaultValue;
+	TObjectPtr<UObject> UseDefaultObject = nullptr;
+	FText UseDefaultText;
+
+	GetDefault<UEdGraphSchema_K2>()->GetPinDefaultValuesFromString(InPin.PinType, InPin.GetOwningNodeUnchecked(), InNewDefaultValue, UseDefaultValue, UseDefaultObject, UseDefaultText, /*bPreserveTextIdentity*/false);
+
+	// Check the default value and make it an error if it's bogus
+	if (GetDefault<UEdGraphSchema_K2>()->DefaultValueSimpleValidation(InPin.PinType, InPin.PinName, UseDefaultValue, UseDefaultObject, UseDefaultText))
+	{
+		InPin.DefaultObject = UseDefaultObject;
+		InPin.DefaultValue = UseDefaultValue;
+		InPin.DefaultTextValue = UseDefaultText;
+
+		UEdGraphNode* Node = InPin.GetOwningNode();
+		Node->PinDefaultValueChanged(&InPin);
+
+		// If the default value is manually set then treat it as if the value was reset to default and remove the orphaned InPin
+		if (InPin.bOrphanedPin && InPin.DoesDefaultValueMatchAutogenerated())
+		{
+			Node->PinConnectionListChanged(&InPin);
+		}
+
+		if (bMarkAsModified)
+		{
+			Node->MarkPackageDirty();
+		}
+	}
+}
