AnimNext: Initial editor-side parameter blocks & libraries

diff --git a/Engine/Plugins/Experimental/AnimNext/Source/AnimNextUncookedOnly/Private/Param/AnimNextParameterBlock_EditorData.cpp b/Engine/Plugins/Experimental/AnimNext/Source/AnimNextUncookedOnly/Private/Param/AnimNextParameterBlock_EditorData.cpp
new file mode 100644
--- /dev/null
+++ b/Engine/Plugins/Experimental/AnimNext/Source/AnimNextUncookedOnly/Private/Param/AnimNextParameterBlock_EditorData.cpp
@@ -0,0 +1,980 @@
+﻿// Copyright Epic Games, Inc. All Rights Reserved.
+
+#include "Param/AnimNextParameterBlock_EditorData.h"
+
+#include "ControlRigDefines.h"
+#include "Param/AnimNextParameterBlock_EdGraphSchema.h"
+#include "UncookedOnlyUtils.h"
+#include "Param/ParametersExecuteContext.h"
+#include "Rigs/RigHierarchyPose.h"
+#include "RigVMModel/RigVMFunctionLibrary.h"
+#include "RigVMModel/RigVMNotifications.h"
+#include "RigVMModel/Nodes/RigVMCollapseNode.h"
+#include "Curves/CurveFloat.h"
+#include "Param/AnimNextParameterBlock.h"
+#include "Param/AnimNextParameterBlockEntry.h"
+#include "ExternalPackageHelper.h"
+#include "Param/AnimNextParameter.h"
+#include "Param/AnimNextParameterBlockBinding.h"
+#include "Param/AnimNextParameterBlockBindingReference.h"
+#include "Param/AnimNextParameterLibrary.h"
+#include "Misc/TransactionObjectEvent.h"
+
+#if WITH_EDITORONLY_DATA
+
+const FName UAnimNextParameterBlock_EditorData::ExportsAssetRegistryTag = TEXT("Exports");
+
+namespace UE::AnimNext::Parameters::Private
+{
+
+static constexpr TCHAR DefaultBindingName[] = TEXT("BindingGraph");
+
+static UAnimNextParameterBlockEntry* CreateNewParameterBlockEntry(UAnimNextParameterBlock_EditorData* InEditorData, TSubclassOf<UAnimNextParameterBlockEntry> InClass)
+{
+	UAnimNextParameterBlockEntry* NewEntry = NewObject<UAnimNextParameterBlockEntry>(InEditorData, InClass.Get(), NAME_None, RF_Transactional);
+	// If we are a transient asset, dont use external packages
+	UAnimNextParameterBlock* Block = UncookedOnly::FUtils::GetBlock(InEditorData);
+	check(Block);
+	if(!Block->HasAnyFlags(RF_Transient))
+	{
+		FExternalPackageHelper::SetPackagingMode(NewEntry, InEditorData, true, false, PKG_None);
+	}
+	return NewEntry;
+}
+
+template<typename EntryClassType>
+static EntryClassType* CreateNewParameterBlockEntry(UAnimNextParameterBlock_EditorData* InEditorData)
+{
+	return CastChecked<EntryClassType>(CreateNewParameterBlockEntry(InEditorData, EntryClassType::StaticClass()));
+}
+
+}
+
+#endif // #if WITH_EDITORONLY_DATA
+
+UAnimNextParameterBlock_EditorData::UAnimNextParameterBlock_EditorData(const FObjectInitializer& ObjectInitializer)
+{
+	RigVMClient.Reset();
+	RigVMClient.SetOuterClientHost(this, GET_MEMBER_NAME_CHECKED(UAnimNextParameterBlock_EditorData, RigVMClient));
+	{
+		TGuardValue<bool> DisableClientNotifs(RigVMClient.bSuspendNotifications, true);
+		RigVMClient.AddModel(TEXT("RigVMGraph"), false, &ObjectInitializer);
+		RigVMClient.GetOrCreateFunctionLibrary(false, &ObjectInitializer);
+	}
+	RigVMClient.SetExecuteContextStruct(FAnimNextParametersExecuteContext::StaticStruct());
+}
+
+#if WITH_EDITORONLY_DATA
+
+UAnimNextParameterBlockBinding* UAnimNextParameterBlockLibrary::AddBinding(UAnimNextParameterBlock* InBlock, FName InName, UAnimNextParameterLibrary* InLibrary, bool bSetupUndoRedo, bool bPrintPythonCommand)
+{
+	return UE::AnimNext::UncookedOnly::FUtils::GetEditorData(InBlock)->AddBinding(InName, InLibrary, bSetupUndoRedo, bPrintPythonCommand);
+}
+
+UAnimNextParameterBlockBinding* UAnimNextParameterBlock_EditorData::AddBinding(FName InName, UAnimNextParameterLibrary* InLibrary, bool bSetupUndoRedo, bool bPrintPythonCommand)
+{
+	if(InName == NAME_None)
+	{
+		ReportError(TEXT("UAnimNextParameterBlock_EditorData::AddBinding: Invalid parameter name supplied."));
+		return nullptr;
+	}
+
+	if(InLibrary == nullptr)
+	{
+		ReportError(TEXT("UAnimNextParameterBlock_EditorData::AddBinding: Invalid parameter library supplied."));
+		return nullptr;
+	}
+
+	// Check parameter exists in library
+	UAnimNextParameter* Parameter = InLibrary->FindParameter(InName);
+	if(Parameter == nullptr)
+	{
+		ReportError(TEXT("UAnimNextParameterBlock_EditorData::AddBinding: Parameter does not exist in library."));
+		return nullptr;
+	}
+
+	// Check for duplicate bindings
+	const bool bAlreadyExists = Entries.ContainsByPredicate([InName, InLibrary](const UAnimNextParameterBlockEntry* InEntry)
+	{
+		if(const IAnimNextParameterBlockBindingInterface* Binding = Cast<IAnimNextParameterBlockBindingInterface>(InEntry))
+		{
+			return Binding->GetParameterName() == InName && Binding->GetLibrary() == InLibrary;
+		}
+		return false;
+	});
+
+	if(bAlreadyExists)
+	{
+		ReportError(TEXT("UAnimNextParameterBlock_EditorData::AddBinding: A binding already exists for the supplied parameter name and library."));
+		return nullptr;
+	}
+
+	UAnimNextParameterBlockBinding* NewEntry = UE::AnimNext::Parameters::Private::CreateNewParameterBlockEntry<UAnimNextParameterBlockBinding>(this);
+	NewEntry->ParameterName = InName;
+	NewEntry->Library = InLibrary;
+
+	if(bSetupUndoRedo)
+	{
+		NewEntry->Modify();
+		Modify();
+	}
+	
+	Entries.Add(NewEntry);
+
+	// Add new binding graph
+	{
+		TGuardValue<bool> EnablePythonPrint(bSuspendPythonMessagesForRigVMClient, !bPrintPythonCommand);
+		TGuardValue<bool> DisableAutoCompile(bAutoRecompileVM, false);
+		URigVMGraph* NewGraph = RigVMClient.AddModel(UE::AnimNext::Parameters::Private::DefaultBindingName, bSetupUndoRedo);
+		ensure(NewGraph);
+		NewEntry->BindingGraph = NewGraph;
+
+		URigVMController* Controller = RigVMClient.GetController(NewGraph);
+		UE::AnimNext::UncookedOnly::FUtils::SetupBindingGraphForLiteral(Controller, Parameter->GetType());
+	}
+
+	BroadcastModified();
+
+	return NewEntry;
+}
+
+UAnimNextParameterBlockBindingReference* UAnimNextParameterBlockLibrary::AddBindingReference(UAnimNextParameterBlock* InBlock, FName InName, UAnimNextParameterLibrary* InLibrary, UAnimNextParameterBlock* InReferencedBlock, bool bSetupUndoRedo, bool bPrintPythonCommand)
+{
+	return UE::AnimNext::UncookedOnly::FUtils::GetEditorData(InBlock)->AddBindingReference(InName, InLibrary, InReferencedBlock, bSetupUndoRedo, bPrintPythonCommand);
+}
+
+UAnimNextParameterBlockBindingReference* UAnimNextParameterBlock_EditorData::AddBindingReference(FName InName, UAnimNextParameterLibrary* InLibrary, UAnimNextParameterBlock* InReferencedBlock, bool bSetupUndoRedo, bool bPrintPythonCommand)
+{
+	if(InName == NAME_None)
+	{
+		ReportError(TEXT("UAnimNextParameterBlock_EditorData::AddBindingReference: Invalid parameter name supplied."));
+		return nullptr;
+	}
+
+	if(InLibrary == nullptr)
+	{
+		ReportError(TEXT("UAnimNextParameterBlock_EditorData::AddBindingReference: Invalid parameter library supplied."));
+		return nullptr;
+	}
+	
+	if(InReferencedBlock == nullptr)
+	{
+		ReportError(TEXT("UAnimNextParameterBlock_EditorData::AddBindingReference: Invalid block supplied."));
+		return nullptr;
+	}
+
+	// Check parameter exists in library
+	UAnimNextParameter* Parameter = InLibrary->FindParameter(InName);
+	if(Parameter == nullptr)
+	{
+		ReportError(TEXT("UAnimNextParameterBlock_EditorData::AddBindingReference: Parameter does not exist in library."));
+		return nullptr;
+	}
+
+	// Check for duplicates
+	const bool bAlreadyExists = Entries.ContainsByPredicate([InName, InLibrary](const UAnimNextParameterBlockEntry* InEntry)
+	{
+		if(const IAnimNextParameterBlockBindingInterface* Binding = Cast<IAnimNextParameterBlockBindingInterface>(InEntry))
+		{
+			return Binding->GetParameterName() == InName && Binding->GetLibrary() == InLibrary;
+		}
+
+		return false;
+	});
+
+	if(bAlreadyExists)
+	{
+		ReportError(TEXT("UAnimNextParameterBlock_EditorData::AddBindingReference: A binding already exists for the supplied parameter and library."));
+		return nullptr;
+	}
+
+	UAnimNextParameterBlockBindingReference* NewEntry = UE::AnimNext::Parameters::Private::CreateNewParameterBlockEntry<UAnimNextParameterBlockBindingReference>(this);
+	NewEntry->ParameterName = InName;
+	NewEntry->Library = InLibrary;
+	NewEntry->Block = InReferencedBlock;
+
+	if(bSetupUndoRedo)
+	{
+		NewEntry->Modify();
+		Modify();
+	}
+
+	Entries.Add(NewEntry);
+
+	BroadcastModified();
+
+	return NewEntry;
+}
+
+bool UAnimNextParameterBlockLibrary::RemoveAllBindings(UAnimNextParameterBlock* InBlock, FName InName,  bool bSetupUndoRedo, bool bPrintPythonCommand)
+{
+	return UE::AnimNext::UncookedOnly::FUtils::GetEditorData(InBlock)->RemoveAllBindings(InName, bSetupUndoRedo, bPrintPythonCommand);
+}
+
+bool UAnimNextParameterBlock_EditorData::RemoveAllBindings(FName InName, bool bSetupUndoRedo, bool bPrintPythonCommand)
+{
+	if(InName == NAME_None)
+	{
+		ReportError(TEXT("UAnimNextParameterBlock_EditorData::RemoveBinding: Invalid parameter name supplied."));
+		return false;
+	}
+
+	// Check parameter binding exists
+	TObjectPtr<UAnimNextParameterBlockEntry>* EntryToRemove = Entries.FindByPredicate([InName](const UAnimNextParameterBlockEntry* InEntry)
+	{
+		if(const IAnimNextParameterBlockBindingInterface* Binding = Cast<IAnimNextParameterBlockBindingInterface>(InEntry))
+		{
+			return Binding->GetParameterName() == InName;
+		}
+		return false;
+	});
+
+	if(EntryToRemove == nullptr)
+	{
+		ReportError(TEXT("UAnimNextParameterBlock_EditorData::RemoveBinding: Block does not contain a binding for parameter."));
+		return false;
+	}
+
+	UAnimNextParameterBlockBinding* BindingToRemove = CastChecked<UAnimNextParameterBlockBinding>(*EntryToRemove);
+
+	if(bSetupUndoRedo)
+	{
+		Modify();
+	}
+
+	Entries.RemoveAll([InName](const UAnimNextParameterBlockEntry* InEntry)
+	{
+		if(const IAnimNextParameterBlockBindingInterface* Binding = Cast<IAnimNextParameterBlockBindingInterface>(InEntry))
+		{
+			return Binding->GetParameterName() == InName;
+		}
+		return false;
+	});
+
+	// Remove any binding graphs
+	TGuardValue<bool> EnablePythonPrint(bSuspendPythonMessagesForRigVMClient, !bPrintPythonCommand);
+	TGuardValue<bool> DisableAutoCompile(bAutoRecompileVM, false);
+	const bool bResult = RigVMClient.RemoveModel(BindingToRemove->BindingGraph->GetNodePath(), bSetupUndoRedo);
+
+	BroadcastModified();
+
+	return bResult;
+}
+
+bool UAnimNextParameterBlockLibrary::RemoveEntry(UAnimNextParameterBlock* InBlock, UAnimNextParameterBlockEntry* InEntry,  bool bSetupUndoRedo, bool bPrintPythonCommand)
+{
+	return UE::AnimNext::UncookedOnly::FUtils::GetEditorData(InBlock)->RemoveEntry(InEntry, bSetupUndoRedo, bPrintPythonCommand);
+}
+
+bool UAnimNextParameterBlock_EditorData::RemoveEntry(UAnimNextParameterBlockEntry* InEntry, bool bSetupUndoRedo, bool bPrintPythonCommand)
+{
+	if(InEntry == nullptr)
+	{
+		ReportError(TEXT("UAnimNextParameterBlock_EditorData::RemoveEntry: Invalid entry supplied."));
+		return false;
+	}
+
+	// Remove from internal array
+	TObjectPtr<UAnimNextParameterBlockEntry>* EntryToRemovePtr = Entries.FindByKey(InEntry);
+	if(EntryToRemovePtr == nullptr)
+	{
+		ReportError(TEXT("UAnimNextParameterBlock_EditorData::RemoveEntry: Block does not contain the supplied entry."));
+		return false;
+	}
+
+	if(bSetupUndoRedo)
+	{
+		Modify();
+	}
+
+	UAnimNextParameterBlockEntry* EntryToRemove = *EntryToRemovePtr;
+	Entries.Remove(EntryToRemove);
+
+	bool bResult = true;
+	if(const IAnimNextParameterBlockGraphInterface* GraphInterface = Cast<IAnimNextParameterBlockGraphInterface>(EntryToRemove))
+	{
+		// Remove any binding graphs
+		TGuardValue<bool> EnablePythonPrint(bSuspendPythonMessagesForRigVMClient, !bPrintPythonCommand);
+		TGuardValue<bool> DisableAutoCompile(bAutoRecompileVM, false);
+		bResult = RigVMClient.RemoveModel(GraphInterface->GetGraph()->GetNodePath(), bSetupUndoRedo);
+	}
+
+	BroadcastModified();
+
+	return bResult;
+}
+
+bool UAnimNextParameterBlockLibrary::RemoveEntries(UAnimNextParameterBlock* InBlock, const TArray<UAnimNextParameterBlockEntry*>& InEntries,  bool bSetupUndoRedo, bool bPrintPythonCommand)
+{
+	return UE::AnimNext::UncookedOnly::FUtils::GetEditorData(InBlock)->RemoveEntries(InEntries, bSetupUndoRedo, bPrintPythonCommand);
+}
+
+bool UAnimNextParameterBlock_EditorData::RemoveEntries(const TArray<UAnimNextParameterBlockEntry*>& InEntries, bool bSetupUndoRedo, bool bPrintPythonCommand)
+{
+	bool bResult = true;
+	{
+		TGuardValue<bool> DisableBlockNotifications(bSuspendBlockNotifications, true);
+		TGuardValue<bool> DisableAutoCompile(bAutoRecompileVM, false);
+		for(UAnimNextParameterBlockEntry* Entry : InEntries)
+		{
+			bResult &= RemoveEntry(Entry, bSetupUndoRedo, bPrintPythonCommand);
+		}
+	}
+
+	BroadcastModified();
+
+	return bResult;
+}
+
+UAnimNextParameterBlockEntry* UAnimNextParameterBlockLibrary::FindBinding(UAnimNextParameterBlock* InBlock, FName InName)
+{
+	return UE::AnimNext::UncookedOnly::FUtils::GetEditorData(InBlock)->FindBinding(InName);
+}
+
+UAnimNextParameterBlockEntry* UAnimNextParameterBlock_EditorData::FindBinding(FName InName) const
+{
+	if(InName == NAME_None)
+	{
+		ReportError(TEXT("UAnimNextParameterBlock_EditorData::FindBinding: Invalid parameter name supplied."));
+		return nullptr;
+	}
+
+	const TObjectPtr<UAnimNextParameterBlockEntry>* FoundEntry = Entries.FindByPredicate([InName](const UAnimNextParameterBlockEntry* InEntry)
+	{
+		if(const IAnimNextParameterBlockBindingInterface* Binding = Cast<IAnimNextParameterBlockBindingInterface>(InEntry))
+		{
+			return Binding->GetParameterName() == InName;
+		}
+		return false;
+	});
+
+	return FoundEntry != nullptr ? *FoundEntry : nullptr;
+}
+
+void UAnimNextParameterBlock_EditorData::BroadcastModified()
+{
+	Recompile();
+
+	if(!bSuspendBlockNotifications)
+	{
+		ModifiedDelegate.Broadcast(this);
+	}
+}
+
+void UAnimNextParameterBlock_EditorData::ReportError(const TCHAR* InMessage) const
+{
+#if WITH_EDITOR
+	FScriptExceptionHandler::Get().HandleException(ELogVerbosity::Error, InMessage, TEXT(""));
+#endif
+}
+
+void UAnimNextParameterBlock_EditorData::Serialize(FArchive& Ar)
+{
+	RigVMClient.SetOuterClientHost(this, GET_MEMBER_NAME_CHECKED(UAnimNextParameterBlock_EditorData, RigVMClient));
+
+	UObject::Serialize(Ar);
+}
+
+void UAnimNextParameterBlock_EditorData::Initialize(bool bRecompileVM)
+{
+	UAnimNextParameterBlock* ParameterBlock = GetTypedOuter<UAnimNextParameterBlock>();
+
+	if (RigVMClient.GetController(0) == nullptr)
+	{
+		check(RigVMClient.Num() == 1);
+		check(RigVMClient.GetFunctionLibrary());
+		
+		RigVMClient.GetOrCreateController(RigVMClient.GetDefaultModel());
+		RigVMClient.GetOrCreateController(RigVMClient.GetFunctionLibrary());
+
+		// Init function library controllers
+		for(URigVMLibraryNode* LibraryNode : RigVMClient.GetFunctionLibrary()->GetFunctions())
+		{
+			RigVMClient.GetOrCreateController(LibraryNode->GetContainedGraph());
+		}
+		
+		if(bRecompileVM)
+		{
+			UE::AnimNext::UncookedOnly::FUtils::Compile(ParameterBlock);
+		}
+	}
+
+	for(UAnimNextParameterBlock_EdGraph* Graph : Graphs)
+	{
+		Graph->Initialize(this);
+	}
+}
+
+void UAnimNextParameterBlock_EditorData::PostLoad()
+{
+	Super::PostLoad();
+	
+	Initialize(/*bRecompileVM*/false);
+	RefreshAllModels(EAnimNextParameterLoadType::PostLoad);
+
+#if WITH_EDITOR
+	FExternalPackageHelper::LoadObjectsFromExternalPackages<UAnimNextParameterBlockEntry>(this, [this](UAnimNextParameterBlockEntry* InLoadedEntry)
+	{
+		check(IsValid(InLoadedEntry));
+		Entries.Add(InLoadedEntry);
+	});
+
+	if (GIsEditor)
+	{
+		// delay compilation until the package has been loaded
+		FCoreUObjectDelegates::OnEndLoadPackage.AddUObject(this, &UAnimNextParameterBlock_EditorData::HandlePackageDone);
+	}
+#else // !WITH_EDITOR
+	RecompileVMIfRequired();
+#endif // WITH_EDITOR
+}
+
+void UAnimNextParameterBlock_EditorData::PostTransacted(const FTransactionObjectEvent& TransactionEvent)
+{
+	Super::PostTransacted(TransactionEvent);
+
+	if (TransactionEvent.GetEventType() == ETransactionObjectEventType::UndoRedo)
+	{
+		BroadcastModified();
+	}
+}
+
+void UAnimNextParameterBlock_EditorData::GetAssetRegistryTags(TArray<FAssetRegistryTag>& OutTags) const
+{
+	Super::GetAssetRegistryTags(OutTags);
+
+	FAnimNextParameterBlockAssetRegistryExports Exports;
+	Exports.Bindings.Reserve(Entries.Num());
+
+	for(const UAnimNextParameterBlockEntry* Entry : Entries)
+	{
+		if(const IAnimNextParameterBlockBindingInterface* Binding = Cast<IAnimNextParameterBlockBindingInterface>(Entry))
+		{
+			Exports.Bindings.Emplace(Binding->GetParameterName(), FSoftObjectPath(Binding->GetLibrary()));
+		}
+	}
+
+	FString TagValue;
+	FAnimNextParameterBlockAssetRegistryExports::StaticStruct()->ExportText(TagValue, &Exports, nullptr, nullptr, PPF_None, nullptr);
+
+	OutTags.Add(FAssetRegistryTag(ExportsAssetRegistryTag, TagValue, FAssetRegistryTag::TT_Hidden));
+}
+
+#if WITH_EDITOR
+void UAnimNextParameterBlock_EditorData::HandlePackageDone(const FEndLoadPackageContext& Context)
+{
+	if (!Context.LoadedPackages.Contains(GetPackage()))
+	{
+		return;
+	}
+	HandlePackageDone();
+}
+
+void UAnimNextParameterBlock_EditorData::HandlePackageDone()
+{
+	FCoreUObjectDelegates::OnEndLoadPackage.RemoveAll(this);
+
+	RecompileVM();
+}
+
+void UAnimNextParameterBlock_EditorData::RefreshAllModels(EAnimNextParameterLoadType InLoadType)
+{
+	const bool bIsPostLoad = InLoadType == EAnimNextParameterLoadType::PostLoad;
+
+	TGuardValue<bool> IsCompilingGuard(bIsCompiling, true);
+	TGuardValue<bool> ClientIgnoreModificationsGuard(RigVMClient.bIgnoreModelNotifications, true);
+
+	TArray<URigVMGraph*> GraphsToDetach = RigVMClient.GetAllModels(true, false);
+	TMap<const URigVMGraph*, TArray<URigVMController::FLinkedPath>> LinkedPaths;
+	
+	if (ensure(IsInGameThread()))
+	{
+		for (URigVMGraph* GraphToDetach : GraphsToDetach)
+		{
+			URigVMController* Controller = RigVMClient.GetOrCreateController(GraphToDetach);
+			// temporarily disable default value validation during load time, serialized values should always be accepted
+			TGuardValue<bool> PerGraphDisablePinDefaultValueValidation(Controller->bValidatePinDefaults, false);
+			FRigVMControllerNotifGuard NotifGuard(Controller, true);
+			LinkedPaths.Add(GraphToDetach, Controller->GetLinkedPaths());
+			Controller->FastBreakLinkedPaths(LinkedPaths.FindChecked(GraphToDetach));
+			TArray<URigVMNode*> Nodes = GraphToDetach->GetNodes();
+			for (URigVMNode* Node : Nodes)
+			{
+				Controller->RepopulatePinsOnNode(Node, true, true);
+			}
+		}
+		//SetupPinRedirectorsForBackwardsCompatibility();
+	}
+
+	for (URigVMGraph* GraphToDetach : GraphsToDetach)
+	{
+		URigVMController* Controller = RigVMClient.GetOrCreateController(GraphToDetach);
+		// at this stage, allow all links to be reattached,
+		// RecomputeAllTemplateFilteredPermutations() later should break any invalid links
+		FRigVMControllerNotifGuard NotifGuard(Controller, true);
+		Controller->RestoreLinkedPaths(LinkedPaths.FindChecked(GraphToDetach));
+	}
+
+	if (bIsPostLoad)
+	{
+		//PatchTemplateNodesWithPreferredPermutation();
+	}
+
+	TArray<URigVMGraph*> GraphsToClean = RigVMClient.GetAllModels(true, true);
+
+	// Sort from leaf graphs to root
+	TArray<URigVMGraph*> SortedGraphsToClean;
+	SortedGraphsToClean.Reserve(GraphsToClean.Num());
+	while (SortedGraphsToClean.Num() < GraphsToClean.Num())
+	{
+		bool bGraphAdded = false;
+		for (URigVMGraph* Graph : GraphsToClean)
+		{
+			if (SortedGraphsToClean.Contains(Graph))
+			{
+				continue;
+			}
+
+			TArray<URigVMGraph*> ContainedGraphs;
+			for (URigVMNode* Node : Graph->GetNodes())
+			{
+				if (URigVMLibraryNode* LibraryNode = Cast<URigVMLibraryNode>(Node))
+				{
+					if (URigVMFunctionReferenceNode* FunctionReferenceNode = Cast<URigVMFunctionReferenceNode>(LibraryNode))
+					{
+						if (FunctionReferenceNode->GetReferencedFunctionHeader().LibraryPointer.LibraryNode.GetLongPackageName() != GetPackage()->GetPathName())
+						{
+							continue;
+						}
+						if (URigVMLibraryNode* ReferencedNode = FunctionReferenceNode->LoadReferencedNode())
+						{
+							ContainedGraphs.Add(ReferencedNode->GetContainedGraph());
+							continue;
+						}
+					}
+
+					if (URigVMGraph* ContainedGraph = LibraryNode->GetContainedGraph())
+					{
+						ContainedGraphs.Add(ContainedGraph);
+					}
+				}
+			}
+
+			bool bAllContained = true;
+			for (URigVMGraph* Contained : ContainedGraphs)
+			{
+				if (!SortedGraphsToClean.Contains(Contained))
+				{
+					bAllContained = false;
+					break;
+				}
+			}
+			if (bAllContained)
+			{
+				SortedGraphsToClean.Add(Graph);
+				bGraphAdded = true;
+			}
+		}
+		ensure(bGraphAdded);
+	}
+
+	for (int32 GraphIndex = 0; GraphIndex < SortedGraphsToClean.Num(); GraphIndex++)
+	{
+		URigVMGraph* GraphToClean = SortedGraphsToClean[GraphIndex];
+		URigVMController* Controller = RigVMClient.GetOrCreateController(GraphToClean);
+		//TGuardValue<bool> GuardEditGraph(GraphToClean->bEditable, true);
+		FRigVMControllerNotifGuard NotifGuard(Controller, true);
+
+		for (URigVMNode* ModelNode : GraphToClean->GetNodes())
+		{
+			Controller->RemoveUnusedOrphanedPins(ModelNode);
+		}
+	}
+}
+
+#endif // WITH_EDITOR
+
+FRigVMClient* UAnimNextParameterBlock_EditorData::GetRigVMClient()
+{
+	return &RigVMClient;
+}
+
+const FRigVMClient* UAnimNextParameterBlock_EditorData::GetRigVMClient() const
+{
+	return &RigVMClient;
+}
+
+IRigVMGraphFunctionHost* UAnimNextParameterBlock_EditorData::GetRigVMGraphFunctionHost()
+{
+	return this;
+}
+
+const IRigVMGraphFunctionHost* UAnimNextParameterBlock_EditorData::GetRigVMGraphFunctionHost() const
+{
+	return this;
+}
+
+void UAnimNextParameterBlock_EditorData::HandleRigVMGraphAdded(const FRigVMClient* InClient, const FString& InNodePath)
+{
+	if(URigVMGraph* RigVMGraph = InClient->GetModel(InNodePath))
+	{
+		RigVMGraph->SetExecuteContextStruct(FAnimNextParametersExecuteContext::StaticStruct());
+
+		if(!HasAnyFlags(RF_ClassDefaultObject | RF_NeedInitialization | RF_NeedLoad | RF_NeedPostLoad) &&
+			GetOuter() != GetTransientPackage())
+		{
+			CreateEdGraph(RigVMGraph, true);
+			RecompileVM();
+		}
+		
+#if WITH_EDITOR
+		if(!bSuspendPythonMessagesForRigVMClient)
+		{
+			const FString BlueprintName = URigVMController::GetSanitizedName(GetName(), true, false);
+			RigVMPythonUtils::Print(BlueprintName, FString::Printf(TEXT("block.add_binding('%s')"), *RigVMGraph->GetName()));
+		}
+#endif
+	}
+}
+
+void UAnimNextParameterBlock_EditorData::HandleRigVMGraphRemoved(const FRigVMClient* InClient, const FString& InNodePath)
+{
+	if(URigVMGraph* RigVMGraph = InClient->GetModel(InNodePath))
+	{
+		RemoveEdGraph(RigVMGraph);
+		RecompileVM();
+
+#if WITH_EDITOR
+		if(!bSuspendPythonMessagesForRigVMClient)
+		{
+			const FString BlueprintName = URigVMController::GetSanitizedName(GetName(), true, false);
+			RigVMPythonUtils::Print(BlueprintName, FString::Printf(TEXT("block.remove_binding('%s')"), *RigVMGraph->GetName()));
+		}
+#endif
+	}
+}
+
+void UAnimNextParameterBlock_EditorData::HandleConfigureRigVMController(const FRigVMClient* InClient, URigVMController* InControllerToConfigure)
+{
+	InControllerToConfigure->OnModified().AddUObject(this, &UAnimNextParameterBlock_EditorData::HandleModifiedEvent);
+
+	InControllerToConfigure->UnfoldStructDelegate.BindLambda([](const UStruct* InStruct) -> bool
+	{
+		if (InStruct == TBaseStructure<FQuat>::Get())
+		{
+			return false;
+		}
+		if (InStruct == FRuntimeFloatCurve::StaticStruct())
+		{
+			return false;
+		}
+		if (InStruct == FRigPose::StaticStruct())
+		{
+			return false;
+		}
+		return true;
+	});
+
+	TWeakObjectPtr<UAnimNextParameterBlock_EditorData> WeakThis(this);
+
+	// this delegate is used by the controller to determine variable validity
+	// during a bind process. the controller itself doesn't own the variables,
+	// so we need a delegate to request them from the owning blueprint
+	InControllerToConfigure->GetExternalVariablesDelegate.BindLambda([](URigVMGraph* InGraph) -> TArray<FRigVMExternalVariable>
+	{
+		if (InGraph)
+		{
+			if(UAnimNextParameterBlock_EditorData* EditorData = InGraph->GetTypedOuter<UAnimNextParameterBlock_EditorData>())
+			{
+				if (UAnimNextParameterBlock* ParameterBlock = EditorData->GetTypedOuter<UAnimNextParameterBlock>())
+				{
+					return ParameterBlock->GetRigVMExternalVariables();
+				}
+			}
+		}
+		return TArray<FRigVMExternalVariable>();
+	});
+
+	// this delegate is used by the controller to retrieve the current bytecode of the VM
+	InControllerToConfigure->GetCurrentByteCodeDelegate.BindLambda([WeakThis]() -> const FRigVMByteCode*
+	{
+		if (WeakThis.IsValid())
+		{
+			if(UAnimNextParameterBlock* ParameterBlock = WeakThis->GetTypedOuter<UAnimNextParameterBlock>())
+			{
+				if (ParameterBlock->RigVM)
+				{
+					return &ParameterBlock->RigVM->GetByteCode();
+				}
+			}
+		}
+		return nullptr;
+
+	});
+
+#if WITH_EDITOR
+	InControllerToConfigure->SetupDefaultUnitNodeDelegates(TDelegate<FName(FRigVMExternalVariable, FString)>::CreateLambda(
+		[](FRigVMExternalVariable InVariableToCreate, FString InDefaultValue) -> FName
+		{
+			return NAME_None;
+		}
+	));
+#endif
+
+}
+
+UObject* UAnimNextParameterBlock_EditorData::GetEditorObjectForRigVMGraph(URigVMGraph* InVMGraph) const
+{
+	if(InVMGraph)
+	{
+		for(UAnimNextParameterBlock_EdGraph* EdGraph : Graphs)
+		{
+			if(EdGraph)
+			{
+				if(EdGraph->ModelNodePath == InVMGraph->GetNodePath())
+				{
+					return EdGraph;
+				}
+			}
+		}
+	}
+	return nullptr;
+}
+
+FRigVMGraphFunctionStore* UAnimNextParameterBlock_EditorData::GetRigVMGraphFunctionStore()
+{
+	return &GraphFunctionStore;
+}
+
+const FRigVMGraphFunctionStore* UAnimNextParameterBlock_EditorData::GetRigVMGraphFunctionStore() const
+{
+	return &GraphFunctionStore;
+}
+
+void UAnimNextParameterBlock_EditorData::RecompileVM()
+{
+	UE::AnimNext::UncookedOnly::FUtils::CompileVM(GetTypedOuter<UAnimNextParameterBlock>());
+}
+
+void UAnimNextParameterBlock_EditorData::RecompileVMIfRequired()
+{
+	if (bVMRecompilationRequired)
+	{
+		RecompileVM();
+	}
+}
+
+void UAnimNextParameterBlock_EditorData::Recompile()
+{
+	// NOTE: Recompiling the struct entails recompiling the VM to ensure that offsets are up to date in bytecode
+	UE::AnimNext::UncookedOnly::FUtils::Compile(GetTypedOuter<UAnimNextParameterBlock>());
+}
+
+void UAnimNextParameterBlock_EditorData::RecompileIfRequired()
+{
+	if (bStructRecompilationRequired || bVMRecompilationRequired)
+	{
+		Recompile();
+	}
+}
+
+void UAnimNextParameterBlock_EditorData::RequestAutoRecompilation()
+{
+	bVMRecompilationRequired = true;
+	bStructRecompilationRequired = true;
+	if (bAutoRecompileVM && VMRecompilationBracket == 0)
+	{
+		RecompileIfRequired();
+	}
+}
+
+void UAnimNextParameterBlock_EditorData::RequestAutoVMRecompilation()
+{
+	bVMRecompilationRequired = true;
+	if (bAutoRecompileVM && VMRecompilationBracket == 0)
+	{
+		RecompileVMIfRequired();
+	}
+}
+
+void UAnimNextParameterBlock_EditorData::IncrementVMRecompileBracket()
+{
+	VMRecompilationBracket++;
+}
+
+void UAnimNextParameterBlock_EditorData::DecrementVMRecompileBracket()
+{
+	if (VMRecompilationBracket == 1)
+	{
+		if (bAutoRecompileVM)
+		{
+			RecompileVMIfRequired();
+		}
+		VMRecompilationBracket = 0;
+	}
+	else if (VMRecompilationBracket > 0)
+	{
+		VMRecompilationBracket--;
+	}
+}
+
+void UAnimNextParameterBlock_EditorData::HandleModifiedEvent(ERigVMGraphNotifType InNotifType, URigVMGraph* InGraph, UObject* InSubject)
+{
+	bool bNotifForOthersPending = true;
+
+	switch(InNotifType)
+	{
+	case ERigVMGraphNotifType::InteractionBracketOpened:
+		{
+			IncrementVMRecompileBracket();
+			break;
+		}
+	case ERigVMGraphNotifType::InteractionBracketClosed:
+	case ERigVMGraphNotifType::InteractionBracketCanceled:
+		{
+			DecrementVMRecompileBracket();
+			break;
+		}
+	case ERigVMGraphNotifType::NodeAdded:
+		{
+			if (URigVMCollapseNode* CollapseNode = Cast<URigVMCollapseNode>(InSubject))
+			{
+				CreateEdGraphForCollapseNode(CollapseNode);
+				break;
+			}
+		}
+		// Fall through to the next case
+	case ERigVMGraphNotifType::LinkAdded:
+	case ERigVMGraphNotifType::LinkRemoved:
+	case ERigVMGraphNotifType::PinArraySizeChanged:
+	case ERigVMGraphNotifType::PinDirectionChanged:
+		{
+			RequestAutoVMRecompilation();
+			break;
+		}
+
+	case ERigVMGraphNotifType::PinDefaultValueChanged:
+		{
+			if (InGraph->GetRuntimeAST().IsValid())
+			{
+				URigVMPin* RootPin = CastChecked<URigVMPin>(InSubject)->GetRootPin();
+				FRigVMASTProxy RootPinProxy = FRigVMASTProxy::MakeFromUObject(RootPin);
+				const FRigVMExprAST* Expression = InGraph->GetRuntimeAST()->GetExprForSubject(RootPinProxy);
+				if (Expression == nullptr)
+				{
+					InGraph->ClearAST();
+					break;
+				}
+				else if (Expression->NumParents() > 1)
+				{
+					InGraph->ClearAST();
+					break;
+				}
+			}
+			break;
+		}
+	}
+	
+	// if the notification still has to be sent...
+	if (bNotifForOthersPending && !bSuspendModelNotificationsForOthers)
+	{
+		if (RigVMGraphModifiedEvent.IsBound())
+		{
+			RigVMGraphModifiedEvent.Broadcast(InNotifType, InGraph, InSubject);
+		}
+	}
+}
+
+URigVMGraph* UAnimNextParameterBlock_EditorData::GetVMGraphForEdGraph(const UEdGraph* InGraph) const
+{
+	const UAnimNextParameterBlock_EdGraph* Graph = Cast<UAnimNextParameterBlock_EdGraph>(InGraph);
+	check(Graph);
+
+	if (Graph->bIsFunctionDefinition)
+	{
+		if (URigVMLibraryNode* LibraryNode = RigVMClient.GetFunctionLibrary()->FindFunction(*Graph->ModelNodePath))
+		{
+			return LibraryNode->GetContainedGraph();
+		}
+	}
+
+	return nullptr;
+}
+
+void UAnimNextParameterBlock_EditorData::CreateEdGraphForCollapseNode(URigVMCollapseNode* InNode)
+{
+	if (InNode->GetGraph()->IsA<URigVMFunctionLibrary>())
+	{
+		if (URigVMGraph* ContainedGraph = InNode->GetContainedGraph())
+		{
+			// create a sub graph
+			UAnimNextParameterBlock_EdGraph* RigFunctionGraph = NewObject<UAnimNextParameterBlock_EdGraph>(this, *InNode->GetName(), RF_Transactional);
+			RigFunctionGraph->Schema = UAnimNextParameterBlock_EdGraphSchema::StaticClass();
+			RigFunctionGraph->bAllowRenaming = true;
+			RigFunctionGraph->bEditable = true;
+			RigFunctionGraph->bAllowDeletion = true;
+			RigFunctionGraph->ModelNodePath = ContainedGraph->GetNodePath();
+			RigFunctionGraph->bIsFunctionDefinition = true;
+
+			RigFunctionGraph->Initialize(this);
+
+			RigVMClient.GetOrCreateController(ContainedGraph)->ResendAllNotifications();
+		}
+	}
+}
+
+UEdGraph* UAnimNextParameterBlock_EditorData::CreateEdGraph(URigVMGraph* InRigVMGraph, bool bForce)
+{
+	check(InRigVMGraph);
+
+// #if WITH_EDITORONLY_DATA
+// 	if(InRigVMGraph->IsA<URigVMFunctionLibrary>())
+// 	{
+// 		return FunctionLibraryEdGraph;
+// 	}
+// #endif
+	
+	if(bForce)
+	{
+		RemoveEdGraph(InRigVMGraph);
+	}
+
+	FString GraphName = InRigVMGraph->GetName();
+//	GraphName.RemoveFromStart(RigVMModelPrefix);
+//	GraphName.TrimStartAndEndInline();
+
+	if(GraphName.IsEmpty())
+	{
+		GraphName = UControlRigGraphSchema::GraphName_ControlRig.ToString();
+	}
+
+	GraphName = RigVMClient.GetUniqueName(*GraphName).ToString();
+
+	UAnimNextParameterBlock_EdGraph* RigFunctionGraph = NewObject<UAnimNextParameterBlock_EdGraph>(this, *GraphName, RF_Transactional);
+	RigFunctionGraph->Schema = UAnimNextParameterBlock_EdGraphSchema::StaticClass();
+	RigFunctionGraph->bAllowDeletion = true;
+	RigFunctionGraph->bIsFunctionDefinition = true;
+	RigFunctionGraph->ModelNodePath = InRigVMGraph->GetNodePath();
+	RigFunctionGraph->Initialize(this);
+
+	Graphs.Add(RigFunctionGraph);
+
+	return RigFunctionGraph;
+}
+
+bool UAnimNextParameterBlock_EditorData::RemoveEdGraph(URigVMGraph* InModel)
+{
+	if(UAnimNextParameterBlock_EdGraph* RigFunctionGraph = Cast<UAnimNextParameterBlock_EdGraph>(GetEditorObjectForRigVMGraph(InModel)))
+	{
+		if(Graphs.Contains(RigFunctionGraph))
+		{
+			Modify();
+			Graphs.Remove(RigFunctionGraph);
+		}
+		RigVMClient.DestroyObject(RigFunctionGraph);
+		return true;
+	}
+	return false;
+}
+
+#endif // #if WITH_EDITORONLY_DATA
\ No newline at end of file
diff --git a/Engine/Plugins/Experimental/AnimNext/Source/AnimNextUncookedOnly/Private/Param/AnimNextParameterLibrary.cpp b/Engine/Plugins/Experimental/AnimNext/Source/AnimNextUncookedOnly/Private/Param/AnimNextParameterLibrary.cpp
new file mode 100644
--- /dev/null
+++ b/Engine/Plugins/Experimental/AnimNext/Source/AnimNextUncookedOnly/Private/Param/AnimNextParameterLibrary.cpp
@@ -0,0 +1,198 @@
+// Copyright Epic Games, Inc. All Rights Reserved.
+
+#include "Param/AnimNextParameterLibrary.h"
+#include "ExternalPackageHelper.h"
+
+#if WITH_EDITORONLY_DATA
+
+const FName UAnimNextParameterLibrary::ExportsAssetRegistryTag = TEXT("Exports");
+
+namespace UE::AnimNext::UncookedOnly::Private
+{
+
+static UAnimNextParameter* CreateNewParameter(UAnimNextParameterLibrary* InLibrary, FName InName)
+{
+	UAnimNextParameter* NewParameter = NewObject<UAnimNextParameter>(InLibrary, InName, RF_Transactional);
+	// If we are a transient asset, dont use external packages
+	if(!InLibrary->HasAnyFlags(RF_Transient))
+	{
+		FExternalPackageHelper::SetPackagingMode(NewParameter, InLibrary, true, false, PKG_None);
+	}
+	return NewParameter;
+}
+
+}
+
+UE::AnimNext::UncookedOnly::FOnParameterLibraryModified& UAnimNextParameterLibrary::OnModified()
+{
+	return ModifiedDelegate;
+}
+
+UAnimNextParameter* UAnimNextParameterLibrary::AddParameter(FName InName, const FAnimNextParamType& InType, bool bSetupUndoRedo, bool bPrintPythonCommand)
+{
+	return AddParameter(InName, InType.GetValueType(), InType.GetContainerType(), InType.GetValueTypeObject(), bSetupUndoRedo, bPrintPythonCommand);
+}
+
+UAnimNextParameter* UAnimNextParameterLibrary::AddParameter(FName InName, const EPropertyBagPropertyType& InValueType, const EPropertyBagContainerType& InContainerType, UObject* InValueTypeObject, bool bSetupUndoRedo, bool bPrintPythonCommand)
+{
+	if(InName == NAME_None)
+	{
+		ReportError(TEXT("UAnimNextParameterLibrary::AddParameter: Invalid name supplied."));
+		return nullptr;
+	}
+
+	// Check for duplicates
+	const bool bAlreadyExists = Parameters.ContainsByPredicate([InName](const UAnimNextParameter* InParameter)
+	{
+		return InParameter->GetName() == InName; 
+	});
+
+	if(bAlreadyExists)
+	{
+		ReportError(TEXT("UAnimNextParameterLibrary::AddParameter: The requested parameter already exists."));
+		return nullptr;
+	}
+
+	const FAnimNextParamType ParamType(InValueType, InContainerType, InValueTypeObject);
+	if(!ParamType.IsValid())
+	{
+		ReportError(TEXT("UAnimNextParameterLibrary::AddParameter: Invalid parameter type supplied."));
+		return nullptr;
+	}
+
+	if(bSetupUndoRedo)
+	{
+		Modify();
+	}
+
+	UAnimNextParameter* NewParameter = UE::AnimNext::UncookedOnly::Private::CreateNewParameter(this, InName);
+	Parameters.Add(NewParameter);
+	NewParameter->SetType(ParamType, bSetupUndoRedo);
+
+	BroadcastModified();
+
+	return NewParameter;
+}
+
+bool UAnimNextParameterLibrary::RemoveParameter(FName InName, bool bSetupUndoRedo, bool bPrintPythonCommand)
+{
+	if(InName == NAME_None)
+	{
+		ReportError(TEXT("UAnimNextParameterLibrary::RemoveParameter: Invalid name supplied."));
+		return false;
+	}
+
+	// Remove from internal array
+	TObjectPtr<UAnimNextParameter>* ParameterToRemove = Parameters.FindByPredicate([InName](const UAnimNextParameter* InParameter)
+	{
+		return InParameter->GetFName() == InName;
+	});
+
+	if(ParameterToRemove == nullptr)
+	{
+		ReportError(TEXT("UAnimNextParameterLibrary::RemoveParameter: Library does not contain the supplied parameter."));
+		return false;
+	}
+
+	if(bSetupUndoRedo)
+	{
+		Modify();
+	}
+
+	Parameters.RemoveSingle(*ParameterToRemove);
+
+	BroadcastModified();
+
+	return true;
+}
+
+bool UAnimNextParameterLibrary::RemoveParameters(const TArray<FName>& InNames, bool bSetupUndoRedo, bool bPrintPythonCommand)
+{
+	if(InNames.Num() == 0)
+	{
+		return false;
+	}
+
+	bool bResult = true;
+	{
+		TGuardValue<bool> DisableNotifications(bSuspendNotifications, true);
+		for(FName Name : InNames)
+		{
+			bResult &= RemoveParameter(Name, bSetupUndoRedo, bPrintPythonCommand);
+		}
+	}
+
+	BroadcastModified();
+
+	return bResult;
+}
+
+UAnimNextParameter* UAnimNextParameterLibrary::FindParameter(FName InName) const
+{
+	if(InName == NAME_None)
+	{
+		ReportError(TEXT("UAnimNextParameterLibrary::FindParameter: Invalid name supplied."));
+		return nullptr;
+	}
+
+	const TObjectPtr<UAnimNextParameter>* FoundParameter = Parameters.FindByPredicate([InName](const UAnimNextParameter* InParameter)
+	{
+		return InParameter->GetFName() == InName;
+	});
+
+	return FoundParameter != nullptr ? *FoundParameter : nullptr;
+}
+
+void UAnimNextParameterLibrary::ReportError(const TCHAR* InMessage) const
+{
+#if WITH_EDITOR
+	FScriptExceptionHandler::Get().HandleException(ELogVerbosity::Error, InMessage, TEXT(""));
+#endif
+}
+
+void UAnimNextParameterLibrary::BroadcastModified()
+{
+	if(!bSuspendNotifications)
+	{
+		ModifiedDelegate.Broadcast(this);
+	}
+}
+
+void UAnimNextParameterLibrary::PostLoad()
+{
+	Super::PostLoad();
+
+#if WITH_EDITOR
+	FExternalPackageHelper::LoadObjectsFromExternalPackages<UAnimNextParameter>(this, [this](UAnimNextParameter* InLoadedParameter)
+	{
+		check(IsValid(InLoadedParameter));
+		Parameters.Add(InLoadedParameter);
+	});
+#endif // WITH_EDITOR
+}
+
+void UAnimNextParameterLibrary::GetAssetRegistryTags(TArray<FAssetRegistryTag>& OutTags) const
+{
+	Super::GetAssetRegistryTags(OutTags);
+
+	FAnimNextParameterLibraryAssetRegistryExports Exports;
+	Exports.Parameters.Reserve(Parameters.Num());
+
+	for(const UAnimNextParameter* Parameter : Parameters)
+	{
+		Exports.Parameters.Emplace(Parameter->GetFName(), Parameter->GetType());
+	}
+
+	FString TagValue;
+	FAnimNextParameterLibraryAssetRegistryExports::StaticStruct()->ExportText(TagValue, &Exports, nullptr, nullptr, PPF_None, nullptr);
+
+	OutTags.Add(FAssetRegistryTag(ExportsAssetRegistryTag, TagValue, FAssetRegistryTag::TT_Hidden));
+}
+
+void UAnimNextParameterLibrary::PostTransacted(const FTransactionObjectEvent& TransactionEvent)
+{
+	Super::PostTransacted(TransactionEvent);
+	BroadcastModified();
+}
+
+#endif // #if WITH_EDITORONLY_DATA
\ No newline at end of file
diff --git a/Engine/Plugins/Experimental/AnimNext/Source/AnimNextUncookedOnly/Private/UncookedOnlyUtils.cpp b/Engine/Plugins/Experimental/AnimNext/Source/AnimNextUncookedOnly/Private/UncookedOnlyUtils.cpp
new file mode 100644
--- /dev/null
+++ b/Engine/Plugins/Experimental/AnimNext/Source/AnimNextUncookedOnly/Private/UncookedOnlyUtils.cpp
@@ -0,0 +1,635 @@
+﻿// Copyright Epic Games, Inc. All Rights Reserved.
+
+#include "UncookedOnlyUtils.h"
+#include "Graph/AnimNextGraph.h"
+#include "Graph/AnimNextGraph_EditorData.h"
+#include "RigVMCompiler/RigVMCompiler.h"
+#include "RigVMCore/RigVM.h"
+#include "Graph/GraphExecuteContext.h"
+#include "Param/AnimNextParameter.h"
+#include "Param/AnimNextParameterBlock.h"
+#include "Param/AnimNextParameterBlock_EditorData.h"
+#include "Param/AnimNextParameterBlock_EdGraph.h"
+#include "Param/ParametersExecuteContext.h"
+#include "Param/RigUnit_AnimNextParametersBeginExecution.h"
+#include "Param/RigVMDispatch_SetParameter.h"
+#include "Param/AnimNextParameterBlockEntry.h"
+#include "Param/IAnimNextParameterBlockBindingInterface.h"
+
+namespace UE::AnimNext::UncookedOnly
+{
+
+void FUtils::Compile(UAnimNextGraph* InGraph)
+{
+	check(InGraph);
+	
+	UAnimNextGraph_EditorData* EditorData = GetEditorData(InGraph);
+	
+	if(EditorData->bIsCompiling)
+	{
+		return;
+	}
+	
+	TGuardValue<bool> CompilingGuard(EditorData->bIsCompiling, true);
+	
+	EditorData->bErrorsDuringCompilation = false;
+
+	EditorData->RigGraphDisplaySettings.MinMicroSeconds = EditorData->RigGraphDisplaySettings.LastMinMicroSeconds = DBL_MAX;
+	EditorData->RigGraphDisplaySettings.MaxMicroSeconds = EditorData->RigGraphDisplaySettings.LastMaxMicroSeconds = (double)INDEX_NONE;
+	
+	TGuardValue<bool> ReentrantGuardSelf(EditorData->bSuspendModelNotificationsForSelf, true);
+	TGuardValue<bool> ReentrantGuardOthers(EditorData->bSuspendModelNotificationsForOthers, true);
+
+	RecreateVM(InGraph);
+
+	InGraph->VMRuntimeSettings = EditorData->VMRuntimeSettings;
+
+	EditorData->CompileLog.Messages.Reset();
+	EditorData->CompileLog.NumErrors = EditorData->CompileLog.NumWarnings = 0;
+
+	URigVMCompiler* Compiler = URigVMCompiler::StaticClass()->GetDefaultObject<URigVMCompiler>();
+	EditorData->VMCompileSettings.SetExecuteContextStruct(EditorData->RigVMClient.GetExecuteContextStruct());
+	Compiler->Settings = (EditorData->bCompileInDebugMode) ? FRigVMCompileSettings::Fast(EditorData->VMCompileSettings.GetExecuteContextStruct()) : EditorData->VMCompileSettings;
+	URigVMController* RootController = EditorData->GetRigVMClient()->GetOrCreateController(EditorData->GetRigVMClient()->GetDefaultModel());
+	Compiler->Compile(EditorData->GetRigVMClient()->GetAllModels(false, false), RootController, InGraph->RigVM, InGraph->GetRigVMExternalVariables(), &EditorData->PinToOperandMap);
+
+	if (EditorData->bErrorsDuringCompilation)
+	{
+		if(Compiler->Settings.SurpressErrors)
+		{
+			Compiler->Settings.Reportf(EMessageSeverity::Info, InGraph,TEXT("Compilation Errors may be suppressed for AnimNext Interface Graph: %s. See VM Compile Settings for more Details"), *InGraph->GetName());
+		}
+	}
+
+	EditorData->bVMRecompilationRequired = false;
+	if(InGraph->RigVM)
+	{
+		EditorData->VMCompiledEvent.Broadcast(InGraph, InGraph->RigVM);
+	}
+
+#if WITH_EDITOR
+//	RefreshBreakpoints(EditorData);
+#endif
+}
+
+void FUtils::RecreateVM(UAnimNextGraph* InGraph)
+{
+	InGraph->RigVM = NewObject<URigVM>(InGraph, TEXT("VM"), RF_NoFlags);
+
+	// Cooked platforms will load these pointers from disk
+	if (!FPlatformProperties::RequiresCookedData())
+	{
+		// We dont support ERigVMMemoryType::Work memory as we dont operate on an instance
+	//	InGraph->RigVM->GetMemoryByType(ERigVMMemoryType::Work, true);
+		InGraph->RigVM->GetMemoryByType(ERigVMMemoryType::Literal, true);
+		InGraph->RigVM->GetMemoryByType(ERigVMMemoryType::Debug, true);
+	}
+
+	InGraph->RigVM->Reset();
+}
+
+UAnimNextGraph_EditorData* FUtils::GetEditorData(const UAnimNextGraph* InAnimNextGraph)
+{
+	check(InAnimNextGraph);
+	
+	return CastChecked<UAnimNextGraph_EditorData>(InAnimNextGraph->EditorData);
+}
+
+FParamTypeHandle FUtils::GetParameterHandleFromPin(const FEdGraphPinType& InPinType)
+{
+	FAnimNextParamType::EValueType ValueType = FAnimNextParamType::EValueType::None;
+	FAnimNextParamType::EContainerType ContainerType = FAnimNextParamType::EContainerType::None;
+	UObject* ValueTypeObject = nullptr;
+
+	if (InPinType.PinCategory == UEdGraphSchema_K2::PC_Boolean)
+	{
+		ValueType = FAnimNextParamType::EValueType::Bool;
+	}
+	else if (InPinType.PinCategory == UEdGraphSchema_K2::PC_Byte)
+	{
+		ValueType = FAnimNextParamType::EValueType::Byte;
+	}
+	else if (InPinType.PinCategory == UEdGraphSchema_K2::PC_Int)
+	{
+		ValueType = FAnimNextParamType::EValueType::Int32;
+	}
+	else if (InPinType.PinCategory == UEdGraphSchema_K2::PC_Int64)
+	{
+		ValueType = FAnimNextParamType::EValueType::Int64;
+	}
+	else if (InPinType.PinCategory == UEdGraphSchema_K2::PC_Real)
+	{
+		if (InPinType.PinSubCategory == UEdGraphSchema_K2::PC_Float)
+		{
+			ValueType = FAnimNextParamType::EValueType::Float;
+		}
+		else if (InPinType.PinSubCategory == UEdGraphSchema_K2::PC_Double)
+		{
+			ValueType = FAnimNextParamType::EValueType::Double;
+		}
+		else
+		{
+			ensure(false);	// Reals should be either floats or doubles
+		}
+	}
+	else if (InPinType.PinCategory == UEdGraphSchema_K2::PC_Float)
+	{
+		ValueType = FAnimNextParamType::EValueType::Float;
+	}
+	else if (InPinType.PinCategory == UEdGraphSchema_K2::PC_Double)
+	{
+		ValueType = FAnimNextParamType::EValueType::Double;
+	}
+	else if (InPinType.PinCategory == UEdGraphSchema_K2::PC_Name)
+	{
+		ValueType = FAnimNextParamType::EValueType::Name;
+	}
+	else if (InPinType.PinCategory == UEdGraphSchema_K2::PC_String)
+	{
+		ValueType = FAnimNextParamType::EValueType::String;
+	}
+	else if (InPinType.PinCategory == UEdGraphSchema_K2::PC_Text)
+	{
+		ValueType = FAnimNextParamType::EValueType::Text;
+	}
+	else if (InPinType.PinCategory == UEdGraphSchema_K2::PC_Enum)
+	{
+		ValueType = FAnimNextParamType::EValueType::Enum;
+		ValueTypeObject = InPinType.PinSubCategoryObject.Get();
+		ensure(ValueTypeObject);
+	}
+	else if (InPinType.PinCategory == UEdGraphSchema_K2::PC_Struct)
+	{
+		ValueType = FAnimNextParamType::EValueType::Struct;
+		ValueTypeObject = Cast<UScriptStruct>(InPinType.PinSubCategoryObject.Get());
+		ensure(ValueTypeObject);
+	}
+	else if (InPinType.PinCategory == UEdGraphSchema_K2::PC_Object)
+	{
+		ValueType = FAnimNextParamType::EValueType::Object;
+		ValueTypeObject = Cast<UClass>(InPinType.PinSubCategoryObject.Get());
+		ensure(ValueTypeObject);
+	}
+	else if (InPinType.PinCategory == UEdGraphSchema_K2::PC_SoftObject)
+	{
+		ValueType = FAnimNextParamType::EValueType::SoftObject;
+		ValueTypeObject = Cast<UClass>(InPinType.PinSubCategoryObject.Get());
+		ensure(ValueTypeObject);
+	}
+	else if (InPinType.PinCategory == UEdGraphSchema_K2::PC_SoftClass)
+	{
+		ValueType = FAnimNextParamType::EValueType::SoftClass;
+		ValueTypeObject = Cast<UClass>(InPinType.PinSubCategoryObject.Get());
+		ensure(ValueTypeObject);
+	}
+
+	if(InPinType.ContainerType == EPinContainerType::Array)
+	{
+		ContainerType = FAnimNextParamType::EContainerType::Array;
+	}
+	else if(InPinType.ContainerType == EPinContainerType::Set)
+	{
+		ensureMsgf(false, TEXT("Set pins are not yet supported"));
+	}
+	if(InPinType.ContainerType == EPinContainerType::Map)
+	{
+		ensureMsgf(false, TEXT("Map pins are not yet supported"));
+	}
+	
+	return FAnimNextParamType(ValueType, ContainerType, ValueTypeObject).GetHandle();
+}
+
+void FUtils::CompileVM(UAnimNextParameterBlock* InParameterBlock)
+{
+	check(InParameterBlock);
+
+	UAnimNextParameterBlock_EditorData* EditorData = FUtils::GetEditorData(InParameterBlock);
+	if(EditorData->bIsCompiling)
+	{
+		return;
+	}
+
+	TGuardValue<bool> CompilingGuard(EditorData->bIsCompiling, true);
+
+	EditorData->bErrorsDuringCompilation = false;
+
+	EditorData->RigGraphDisplaySettings.MinMicroSeconds = EditorData->RigGraphDisplaySettings.LastMinMicroSeconds = DBL_MAX;
+	EditorData->RigGraphDisplaySettings.MaxMicroSeconds = EditorData->RigGraphDisplaySettings.LastMaxMicroSeconds = (double)INDEX_NONE;
+	
+	TGuardValue<bool> ReentrantGuardSelf(EditorData->bSuspendModelNotificationsForSelf, true);
+	TGuardValue<bool> ReentrantGuardOthers(EditorData->bSuspendModelNotificationsForOthers, true);
+
+	RecreateVM(InParameterBlock);
+
+	InParameterBlock->VMRuntimeSettings = EditorData->VMRuntimeSettings;
+
+	EditorData->CompileLog.Messages.Reset();
+	EditorData->CompileLog.NumErrors = EditorData->CompileLog.NumWarnings = 0;
+
+	URigVMCompiler* Compiler = URigVMCompiler::StaticClass()->GetDefaultObject<URigVMCompiler>();
+	EditorData->VMCompileSettings.SetExecuteContextStruct(EditorData->RigVMClient.GetExecuteContextStruct());
+	Compiler->Settings = (EditorData->bCompileInDebugMode) ? FRigVMCompileSettings::Fast(EditorData->VMCompileSettings.GetExecuteContextStruct()) : EditorData->VMCompileSettings;
+	URigVMController* RootController = EditorData->GetRigVMClient()->GetOrCreateController(EditorData->GetRigVMClient()->GetDefaultModel());
+	Compiler->Compile(EditorData->GetRigVMClient()->GetAllModels(false, false), RootController, InParameterBlock->RigVM, InParameterBlock->GetRigVMExternalVariables(), &EditorData->PinToOperandMap);
+
+	if (EditorData->bErrorsDuringCompilation)
+	{
+		if(Compiler->Settings.SurpressErrors)
+		{
+			Compiler->Settings.Reportf(EMessageSeverity::Info, InParameterBlock, TEXT("Compilation Errors may be suppressed for AnimNext Interface Graph: %s. See VM Compile Settings for more Details"), *InParameterBlock->GetName());
+		}
+	}
+
+	EditorData->bVMRecompilationRequired = false;
+	if(InParameterBlock->RigVM)
+	{
+		EditorData->RigVMCompiledEvent.Broadcast(InParameterBlock, InParameterBlock->RigVM);
+	}
+
+#if WITH_EDITOR
+//	RefreshBreakpoints(EditorData);
+#endif
+}
+
+void FUtils::CompileStruct(UAnimNextParameterBlock* InParameterBlock)
+{
+	check(InParameterBlock);
+
+	UAnimNextParameterBlock_EditorData* EditorData = GetEditorData(InParameterBlock);
+	if(EditorData->bIsCompiling)
+	{
+		return;
+	}
+
+	TGuardValue<bool> CompilingGuard(EditorData->bIsCompiling, true);
+
+	FInstancedPropertyBag& PropertyBag = InParameterBlock->PropertyBag;
+	PropertyBag.Reset();
+
+	TArray<FPropertyBagPropertyDesc> PropertyDescs;
+	PropertyDescs.Reserve(EditorData->Entries.Num());
+	
+	// Gather all properties in this block
+	for(const UAnimNextParameterBlockEntry* Entry : EditorData->Entries)
+	{
+		if(const IAnimNextParameterBlockBindingInterface* Binding = Cast<IAnimNextParameterBlockBindingInterface>(Entry))
+		{
+			const FAnimNextParamType& Type = Binding->GetParamType();
+			const UAnimNextParameter* Parameter = Binding->GetParameter();
+			PropertyDescs.Emplace(Parameter->GetFName(), Type.GetContainerType(), Type.GetValueType(), Type.GetValueTypeObject());
+		}
+	}
+
+	// Bulk add to the bag
+	PropertyBag.AddProperties(PropertyDescs);
+
+	// TODO: Now copy over defaults for those properties that need it (literals)
+
+	EditorData->bStructRecompilationRequired = false;
+}
+
+void FUtils::Compile(UAnimNextParameterBlock* InParameterBlock)
+{
+	check(InParameterBlock);
+
+	UAnimNextParameterBlock_EditorData* EditorData = GetEditorData(InParameterBlock);
+	if(EditorData->bIsCompiling)
+	{
+		return;
+	}
+
+	TGuardValue<bool> CompilingGuard(EditorData->bIsCompiling, true);
+
+	CompileStruct(InParameterBlock);
+	CompileVM(InParameterBlock);
+}
+
+void FUtils::RecreateVM(UAnimNextParameterBlock* InParameterBlock)
+{
+	InParameterBlock->RigVM = NewObject<URigVM>(InParameterBlock, TEXT("VM"), RF_NoFlags);
+
+	// Cooked platforms will load these pointers from disk
+	if (!FPlatformProperties::RequiresCookedData())
+	{
+		// We dont support ERigVMMemoryType::Work memory as we dont operate on an instance
+	//	InParameterBlock->RigVM->GetMemoryByType(ERigVMMemoryType::Work, true);
+		InParameterBlock->RigVM->GetMemoryByType(ERigVMMemoryType::Literal, true);
+		InParameterBlock->RigVM->GetMemoryByType(ERigVMMemoryType::Debug, true);
+	}
+
+	InParameterBlock->RigVM->Reset();
+}
+
+UAnimNextParameterBlock_EditorData* FUtils::GetEditorData(const UAnimNextParameterBlock* InParameterBlock)
+{
+	check(InParameterBlock);
+
+	return CastChecked<UAnimNextParameterBlock_EditorData>(InParameterBlock->EditorData);
+}
+
+UAnimNextParameterBlock* FUtils::GetBlock(const UAnimNextParameterBlock_EditorData* InEditorData)
+{
+	check(InEditorData);
+
+	return CastChecked<UAnimNextParameterBlock>(InEditorData->GetOuter());
+}
+
+FParamTypeHandle FUtils::GetParamTypeHandleFromPinType(const FEdGraphPinType& InPinType)
+{
+	return GetParamTypeFromPinType(InPinType).GetHandle();
+}
+
+FAnimNextParamType FUtils::GetParamTypeFromPinType(const FEdGraphPinType& InPinType)
+{
+	FAnimNextParamType::EValueType ValueType = FAnimNextParamType::EValueType::None;
+	FAnimNextParamType::EContainerType ContainerType = FAnimNextParamType::EContainerType::None;
+	UObject* ValueTypeObject = nullptr;
+
+	if (InPinType.PinCategory == UEdGraphSchema_K2::PC_Boolean)
+	{
+		ValueType = FAnimNextParamType::EValueType::Bool;
+	}
+	else if (InPinType.PinCategory == UEdGraphSchema_K2::PC_Byte)
+	{
+		ValueType = FAnimNextParamType::EValueType::Byte;
+	}
+	else if (InPinType.PinCategory == UEdGraphSchema_K2::PC_Int)
+	{
+		ValueType = FAnimNextParamType::EValueType::Int32;
+	}
+	else if (InPinType.PinCategory == UEdGraphSchema_K2::PC_Int64)
+	{
+		ValueType = FAnimNextParamType::EValueType::Int64;
+	}
+	else if (InPinType.PinCategory == UEdGraphSchema_K2::PC_Real)
+	{
+		if (InPinType.PinSubCategory == UEdGraphSchema_K2::PC_Float)
+		{
+			ValueType = FAnimNextParamType::EValueType::Float;
+		}
+		else if (InPinType.PinSubCategory == UEdGraphSchema_K2::PC_Double)
+		{
+			ValueType = FAnimNextParamType::EValueType::Double;
+		}
+		else
+		{
+			ensure(false);	// Reals should be either floats or doubles
+		}
+	}
+	else if (InPinType.PinCategory == UEdGraphSchema_K2::PC_Float)
+	{
+		ValueType = FAnimNextParamType::EValueType::Float;
+	}
+	else if (InPinType.PinCategory == UEdGraphSchema_K2::PC_Double)
+	{
+		ValueType = FAnimNextParamType::EValueType::Double;
+	}
+	else if (InPinType.PinCategory == UEdGraphSchema_K2::PC_Name)
+	{
+		ValueType = FAnimNextParamType::EValueType::Name;
+	}
+	else if (InPinType.PinCategory == UEdGraphSchema_K2::PC_String)
+	{
+		ValueType = FAnimNextParamType::EValueType::String;
+	}
+	else if (InPinType.PinCategory == UEdGraphSchema_K2::PC_Text)
+	{
+		ValueType = FAnimNextParamType::EValueType::Text;
+	}
+	else if (InPinType.PinCategory == UEdGraphSchema_K2::PC_Enum)
+	{
+		ValueType = FAnimNextParamType::EValueType::Enum;
+		ValueTypeObject = InPinType.PinSubCategoryObject.Get();
+		ensure(ValueTypeObject);
+	}
+	else if (InPinType.PinCategory == UEdGraphSchema_K2::PC_Struct)
+	{
+		ValueType = FAnimNextParamType::EValueType::Struct;
+		ValueTypeObject = Cast<UScriptStruct>(InPinType.PinSubCategoryObject.Get());
+		ensure(ValueTypeObject);
+	}
+	else if (InPinType.PinCategory == UEdGraphSchema_K2::PC_Object)
+	{
+		ValueType = FAnimNextParamType::EValueType::Object;
+		ValueTypeObject = Cast<UClass>(InPinType.PinSubCategoryObject.Get());
+		ensure(ValueTypeObject);
+	}
+	else if (InPinType.PinCategory == UEdGraphSchema_K2::PC_SoftObject)
+	{
+		ValueType = FAnimNextParamType::EValueType::SoftObject;
+		ValueTypeObject = Cast<UClass>(InPinType.PinSubCategoryObject.Get());
+		ensure(ValueTypeObject);
+	}
+	else if (InPinType.PinCategory == UEdGraphSchema_K2::PC_SoftClass)
+	{
+		ValueType = FAnimNextParamType::EValueType::SoftClass;
+		ValueTypeObject = Cast<UClass>(InPinType.PinSubCategoryObject.Get());
+		ensure(ValueTypeObject);
+	}
+
+	if(InPinType.ContainerType == EPinContainerType::Array)
+	{
+		ContainerType = FAnimNextParamType::EContainerType::Array;
+	}
+	else if(InPinType.ContainerType == EPinContainerType::Set)
+	{
+		ensureMsgf(false, TEXT("Set pins are not yet supported"));
+	}
+	if(InPinType.ContainerType == EPinContainerType::Map)
+	{
+		ensureMsgf(false, TEXT("Map pins are not yet supported"));
+	}
+	
+	return FAnimNextParamType(ValueType, ContainerType, ValueTypeObject);
+}
+
+FEdGraphPinType FUtils::GetPinTypeFromParamTypeHandle(const FParamTypeHandle& InParamTypeHandle)
+{
+	return GetPinTypeFromParamType(InParamTypeHandle.GetType());
+}
+
+FEdGraphPinType FUtils::GetPinTypeFromParamType(const FAnimNextParamType& InParamType)
+{
+	FEdGraphPinType PinType;
+	PinType.PinSubCategory = NAME_None;
+
+	// Container type
+	switch (InParamType.ContainerType)
+	{
+	case FAnimNextParamType::EContainerType::Array:
+		PinType.ContainerType = EPinContainerType::Array;
+		break;
+	default:
+		PinType.ContainerType = EPinContainerType::None;
+	}
+
+	// Value type
+	switch (InParamType.ValueType)
+	{
+	case EPropertyBagPropertyType::Bool:
+		PinType.PinCategory = UEdGraphSchema_K2::PC_Boolean;
+		break;
+	case EPropertyBagPropertyType::Byte:
+		PinType.PinCategory = UEdGraphSchema_K2::PC_Byte;
+		break;
+	case EPropertyBagPropertyType::Int32:
+		PinType.PinCategory = UEdGraphSchema_K2::PC_Int;
+		break;
+	case EPropertyBagPropertyType::Int64:
+		PinType.PinCategory = UEdGraphSchema_K2::PC_Int64;
+		break;
+	case EPropertyBagPropertyType::Float:
+		PinType.PinCategory = UEdGraphSchema_K2::PC_Real;
+		PinType.PinSubCategory = UEdGraphSchema_K2::PC_Float;
+		break;
+	case EPropertyBagPropertyType::Double:
+		PinType.PinCategory = UEdGraphSchema_K2::PC_Real;
+		PinType.PinSubCategory = UEdGraphSchema_K2::PC_Double;
+		break;
+	case EPropertyBagPropertyType::Name:
+		PinType.PinCategory = UEdGraphSchema_K2::PC_Name;
+		break;
+	case EPropertyBagPropertyType::String:
+		PinType.PinCategory = UEdGraphSchema_K2::PC_String;
+		break;
+	case EPropertyBagPropertyType::Text:
+		PinType.PinCategory = UEdGraphSchema_K2::PC_Text;
+		break;
+	case EPropertyBagPropertyType::Enum:
+		// @todo: some pin coloring is not correct due to this (byte-as-enum vs enum). 
+		PinType.PinCategory = UEdGraphSchema_K2::PC_Enum;
+		PinType.PinSubCategoryObject = const_cast<UObject*>(InParamType.ValueTypeObject.Get());
+		break;
+	case EPropertyBagPropertyType::Struct:
+		PinType.PinCategory = UEdGraphSchema_K2::PC_Struct;
+		PinType.PinSubCategoryObject = const_cast<UObject*>(InParamType.ValueTypeObject.Get());
+		break;
+	case EPropertyBagPropertyType::Object:
+		PinType.PinCategory = UEdGraphSchema_K2::PC_Object;
+		PinType.PinSubCategoryObject = const_cast<UObject*>(InParamType.ValueTypeObject.Get());
+		break;
+	case EPropertyBagPropertyType::SoftObject:
+		PinType.PinCategory = UEdGraphSchema_K2::PC_SoftObject;
+		PinType.PinSubCategoryObject = const_cast<UObject*>(InParamType.ValueTypeObject.Get());
+		break;
+	case EPropertyBagPropertyType::Class:
+		PinType.PinCategory = UEdGraphSchema_K2::PC_Class;
+		PinType.PinSubCategoryObject = const_cast<UObject*>(InParamType.ValueTypeObject.Get());
+		break;
+	case EPropertyBagPropertyType::SoftClass:
+		PinType.PinCategory = UEdGraphSchema_K2::PC_SoftClass;
+		PinType.PinSubCategoryObject = const_cast<UObject*>(InParamType.ValueTypeObject.Get());
+		break;
+	default:
+		ensureMsgf(false, TEXT("Unhandled value type %d"), InParamType.ValueType);
+		break;
+	}
+
+	return PinType;
+}
+
+FRigVMTemplateArgumentType FUtils::GetRigVMArgTypeFromParamTypeHandle(const FParamTypeHandle& InParamTypeHandle)
+{
+	return GetRigVMArgTypeFromParamType(InParamTypeHandle.GetType());
+}
+
+FRigVMTemplateArgumentType FUtils::GetRigVMArgTypeFromParamType(const FAnimNextParamType& InParamType)
+{
+	FRigVMTemplateArgumentType ArgType;
+
+	FString CPPTypeString;
+
+	// Value type
+	switch (InParamType.ValueType)
+	{
+	case EPropertyBagPropertyType::Bool:
+		CPPTypeString = RigVMTypeUtils::BoolType;
+		break;
+	case EPropertyBagPropertyType::Byte:
+		CPPTypeString = RigVMTypeUtils::UInt8Type;
+		break;
+	case EPropertyBagPropertyType::Int32:
+		CPPTypeString = RigVMTypeUtils::UInt32Type;
+		break;
+	case EPropertyBagPropertyType::Int64:
+		ensureMsgf(false, TEXT("Unhandled value type %d"), InParamType.ValueType);
+		break;
+	case EPropertyBagPropertyType::Float:
+		CPPTypeString = RigVMTypeUtils::FloatType;
+		break;
+	case EPropertyBagPropertyType::Double:
+		CPPTypeString = RigVMTypeUtils::DoubleType;
+		break;
+	case EPropertyBagPropertyType::Name:
+		CPPTypeString = RigVMTypeUtils::FNameType;
+		break;
+	case EPropertyBagPropertyType::String:
+		CPPTypeString = RigVMTypeUtils::FStringType;
+		break;
+	case EPropertyBagPropertyType::Text:
+		ensureMsgf(false, TEXT("Unhandled value type %d"), InParamType.ValueType);
+		break;
+	case EPropertyBagPropertyType::Enum:
+		CPPTypeString = RigVMTypeUtils::CPPTypeFromEnum(Cast<UEnum>(InParamType.ValueTypeObject.Get()));
+		ArgType.CPPTypeObject = InParamType.ValueTypeObject.Get();
+		break;
+	case EPropertyBagPropertyType::Struct:
+		CPPTypeString = RigVMTypeUtils::GetUniqueStructTypeName(Cast<UScriptStruct>(InParamType.ValueTypeObject.Get()));
+		ArgType.CPPTypeObject = InParamType.ValueTypeObject.Get();
+		break;
+	case EPropertyBagPropertyType::Object:
+		CPPTypeString = RigVMTypeUtils::CPPTypeFromObject(Cast<UClass>(InParamType.ValueTypeObject.Get()));
+		ArgType.CPPTypeObject = InParamType.ValueTypeObject.Get();
+		break;
+	case EPropertyBagPropertyType::SoftObject:
+		ensureMsgf(false, TEXT("Unhandled value type %d"), InParamType.ValueType);
+		break;
+	case EPropertyBagPropertyType::Class:
+		ensureMsgf(false, TEXT("Unhandled value type %d"), InParamType.ValueType);
+		break;
+	case EPropertyBagPropertyType::SoftClass:
+		ensureMsgf(false, TEXT("Unhandled value type %d"), InParamType.ValueType);
+		break;
+	default:
+		ensureMsgf(false, TEXT("Unhandled value type %d"), InParamType.ValueType);
+		break;
+	}
+
+	// Container type
+	switch (InParamType.ContainerType)
+	{
+	case FAnimNextParamType::EContainerType::None:
+		break;
+	case FAnimNextParamType::EContainerType::Array:
+		CPPTypeString = FString::Printf(RigVMTypeUtils::TArrayTemplate, *CPPTypeString);
+		break;
+	default:
+		ensureMsgf(false, TEXT("Unhandled container type %d"), InParamType.ContainerType);
+		break;
+	}
+
+	ArgType.CPPType = *CPPTypeString;
+
+	return ArgType;
+}
+
+void FUtils::SetupBindingGraphForLiteral(URigVMController* InController, const FAnimNextParamType& InParamType)
+{
+	FRigVMTemplateArgumentType ArgType = GetRigVMArgTypeFromParamType(InParamType);
+	TRigVMTypeIndex TypeIndex = FRigVMRegistry::Get().GetTypeIndex(ArgType);
+
+	// Clear the graph
+	InController->RemoveNodes(InController->GetGraph()->GetNodes());
+
+	// Add new nodes for a simple literal binding
+	URigVMUnitNode* EntryPointNode = InController->AddUnitNode(FRigUnit_AnimNextParametersBeginExecution::StaticStruct(), FRigUnit::GetMethodName(), FVector2D(-200.0f, 0.0f), FString(), false);
+
+	const FName FactoryName = FRigVMDispatch_SetParameter().GetFactoryName();
+	FRigVMDispatch_SetParameter* Factory = static_cast<FRigVMDispatch_SetParameter*>(FRigVMRegistry::Get().FindDispatchFactory(FactoryName));
+	URigVMTemplateNode* SetParameterNode = InController->AddTemplateNode(Factory->GetTemplate()->GetNotation(), FVector2D(200.0f, 0.0f));
+	InController->ResolveWildCardPin(SetParameterNode->FindPin(FRigVMDispatch_SetParameter::ValueName.ToString()), TypeIndex);
+
+	InController->AddLink(EntryPointNode->FindPin(GET_MEMBER_NAME_STRING_CHECKED(FRigUnit_AnimNextParametersBeginExecution, ExecuteContext)), SetParameterNode->FindPin(FRigVMDispatch_SetParameter::ExecuteContextName.ToString()));
+}
+
+}
diff --git a/Engine/Plugins/Experimental/StructUtils/Source/StructUtils/Public/PropertyBag.h b/Engine/Plugins/Experimental/StructUtils/Source/StructUtils/Public/PropertyBag.h
--- a/Engine/Plugins/Experimental/StructUtils/Source/StructUtils/Public/PropertyBag.h
+++ b/Engine/Plugins/Experimental/StructUtils/Source/StructUtils/Public/PropertyBag.h
@@ -8,7 +8,7 @@
 #include "PropertyBag.generated.h"
 
 /** Property bag property type, loosely based on BluePrint pin types. */
-UENUM()
+UENUM(BlueprintType)
 enum class EPropertyBagPropertyType : uint8
 {
 	None UMETA(Hidden),
@@ -32,10 +32,10 @@ enum class EPropertyBagPropertyType : uint8
 };
 
 /** Property bag property container type. */
-UENUM()
+UENUM(BlueprintType)
 enum class EPropertyBagContainerType : uint8
 {
-	None UMETA(Hidden),
+	None,
 	Array,
 
 	Count UMETA(Hidden)
diff --git a/Engine/Plugins/Runtime/RigVM/Source/RigVM/Public/RigVMCore/RigVMTemplate.h b/Engine/Plugins/Runtime/RigVM/Source/RigVM/Public/RigVMCore/RigVMTemplate.h
--- a/Engine/Plugins/Runtime/RigVM/Source/RigVM/Public/RigVMCore/RigVMTemplate.h
+++ b/Engine/Plugins/Runtime/RigVM/Source/RigVM/Public/RigVMCore/RigVMTemplate.h
@@ -75,6 +75,12 @@ struct RIGVM_API FRigVMTemplateArgumentType
 		checkf(!CPPType.IsNone(), TEXT("FRigVMTemplateArgumentType(): Input CPPType '%s' could not be resolved."), *InCPPTypeString);
 	}
 
+	FRigVMTemplateArgumentType(UClass* InClass)
+	: CPPType(*RigVMTypeUtils::CPPTypeFromObject(InClass))
+	, CPPTypeObject(InClass)
+	{
+	}
+
 	FRigVMTemplateArgumentType(UScriptStruct* InScriptStruct)
 	: CPPType(*RigVMTypeUtils::GetUniqueStructTypeName(InScriptStruct))
 	, CPPTypeObject(InScriptStruct)
