CIS: Fix incremental Linux build

diff --git a/Engine/Plugins/Experimental/AnimNext/Source/AnimNextUncookedOnly/Internal/Param/AnimNextParameterBlockBinding.h b/Engine/Plugins/Experimental/AnimNext/Source/AnimNextUncookedOnly/Internal/Param/AnimNextParameterBlockBinding.h
--- a/Engine/Plugins/Experimental/AnimNext/Source/AnimNextUncookedOnly/Internal/Param/AnimNextParameterBlockBinding.h
+++ b/Engine/Plugins/Experimental/AnimNext/Source/AnimNextUncookedOnly/Internal/Param/AnimNextParameterBlockBinding.h
@@ -10,6 +10,8 @@
 
 class UAssetDefinition_AnimNextParameterBlockBinding;
 class UAnimNextParameter;
+class UAnimNextParameterLibrary;
+class URigVMGraph;
 
 namespace UE::AnimNext::Editor
 {
diff --git a/Engine/Plugins/Experimental/AnimNext/Source/AnimNextUncookedOnly/Internal/Param/AnimNextParameterBlockBindingReference.h b/Engine/Plugins/Experimental/AnimNext/Source/AnimNextUncookedOnly/Internal/Param/AnimNextParameterBlockBindingReference.h
--- a/Engine/Plugins/Experimental/AnimNext/Source/AnimNextUncookedOnly/Internal/Param/AnimNextParameterBlockBindingReference.h
+++ b/Engine/Plugins/Experimental/AnimNext/Source/AnimNextUncookedOnly/Internal/Param/AnimNextParameterBlockBindingReference.h
@@ -8,7 +8,9 @@
 #include "IAnimNextParameterBlockBindingInterface.h"
 #include "AnimNextParameterBlockBindingReference.generated.h"
 
+class UAnimNextParameter;
 class UAnimNextParameterBlock;
+class UAnimNextParameterLibrary;
 class UAssetDefinition_AnimNextParameterBlockBindingReference;
 
 namespace UE::AnimNext::Editor
@@ -17,8 +19,6 @@ namespace UE::AnimNext::Editor
 	class SParameterBlockViewRow;
 }
 
-class UAnimNextParameter;
-
 /** Parameter binding block entry */
 UCLASS(MinimalAPI, BlueprintType)
 class UAnimNextParameterBlockBindingReference : public UAnimNextParameterBlockEntry, public IAnimNextParameterBlockBindingInterface, public IAnimNextParameterBlockReferenceInterface
diff --git a/Engine/Plugins/Experimental/AnimNext/Source/AnimNextUncookedOnly/Internal/Param/AnimNextParameterBlockEntry.h b/Engine/Plugins/Experimental/AnimNext/Source/AnimNextUncookedOnly/Internal/Param/AnimNextParameterBlockEntry.h
--- a/Engine/Plugins/Experimental/AnimNext/Source/AnimNextUncookedOnly/Internal/Param/AnimNextParameterBlockEntry.h
+++ b/Engine/Plugins/Experimental/AnimNext/Source/AnimNextUncookedOnly/Internal/Param/AnimNextParameterBlockEntry.h
@@ -17,7 +17,7 @@ struct FParameterBlockViewEntry;
 }
 
 /** Base class that defines an entry in a parameter block, e.g. a parameter binding */
-UCLASS(MinimalAPI, BlueprintType)
+UCLASS(MinimalAPI, BlueprintType, Abstract)
 class UAnimNextParameterBlockEntry : public UObject
 {
 	GENERATED_BODY()
diff --git a/Engine/Plugins/Experimental/AnimNext/Source/AnimNextUncookedOnly/Internal/Param/AnimNextParameterBlockProcessor.h b/Engine/Plugins/Experimental/AnimNext/Source/AnimNextUncookedOnly/Internal/Param/AnimNextParameterBlockProcessor.h
--- a/Engine/Plugins/Experimental/AnimNext/Source/AnimNextUncookedOnly/Internal/Param/AnimNextParameterBlockProcessor.h
+++ b/Engine/Plugins/Experimental/AnimNext/Source/AnimNextUncookedOnly/Internal/Param/AnimNextParameterBlockProcessor.h
@@ -10,4 +10,8 @@ UCLASS(MinimalAPI)
 class UAnimNextParameterBlockProcessor : public UAnimNextParameterBlockEntry
 {
 	GENERATED_BODY()
+
+	// UAnimNextParameterBlockEntry interface
+	virtual FText GetDisplayName() const { return FText::GetEmpty(); }
+	virtual FText GetDisplayNameTooltip() const { return FText::GetEmpty(); }
 };
\ No newline at end of file
diff --git a/Engine/Plugins/Experimental/AnimNext/Source/AnimNextUncookedOnly/Internal/Param/IAnimNextParameterBlockBindingInterface.h b/Engine/Plugins/Experimental/AnimNext/Source/AnimNextUncookedOnly/Internal/Param/IAnimNextParameterBlockBindingInterface.h
--- a/Engine/Plugins/Experimental/AnimNext/Source/AnimNextUncookedOnly/Internal/Param/IAnimNextParameterBlockBindingInterface.h
+++ b/Engine/Plugins/Experimental/AnimNext/Source/AnimNextUncookedOnly/Internal/Param/IAnimNextParameterBlockBindingInterface.h
@@ -7,6 +7,7 @@
 #include "IAnimNextParameterBlockBindingInterface.generated.h"
 
 class UAnimNextParameter;
+class UAnimNextParameterLibrary;
 class UAnimNextParameterBlock_EditorData;
 
 namespace UE::AnimNext::Editor
