AnimNext: Added parameter block/library tests

Also moved tests from EngineTest to plugin
Converted moves to deletes/adds because of path restriction changes across the move

diff --git a/Engine/Plugins/Experimental/AnimNext/Source/AnimNextTestSuite/AnimNextTestSuite.Build.cs b/Engine/Plugins/Experimental/AnimNext/Source/AnimNextTestSuite/AnimNextTestSuite.Build.cs
--- a/Engine/Plugins/Experimental/AnimNext/Source/AnimNextTestSuite/AnimNextTestSuite.Build.cs
+++ b/Engine/Plugins/Experimental/AnimNext/Source/AnimNextTestSuite/AnimNextTestSuite.Build.cs
@@ -15,12 +15,21 @@ namespace UnrealBuildTool.Rules
 					"CoreUObject",
 					"Engine",
 					"AnimNext",
+					"StructUtils"
 				}
 			);
 
 			if (Target.bBuildEditor == true)
 			{
-				PrivateDependencyModuleNames.Add("UnrealEd");
+				PrivateDependencyModuleNames.AddRange(
+					new string[]
+					{
+						"UnrealEd",
+						"AnimNextUncookedOnly",
+						"AnimNextEditor",
+						"PythonScriptPlugin"
+					}
+				);
 			}
 		}
 	}
diff --git a/Engine/Plugins/Experimental/AnimNext/Source/AnimNextTestSuite/Private/AnimNextEditorTest.cpp b/Engine/Plugins/Experimental/AnimNext/Source/AnimNextTestSuite/Private/AnimNextEditorTest.cpp
new file mode 100644
--- /dev/null
+++ b/Engine/Plugins/Experimental/AnimNext/Source/AnimNextTestSuite/Private/AnimNextEditorTest.cpp
@@ -0,0 +1,201 @@
+﻿// Copyright Epic Games, Inc. All Rights Reserved.
+
+#include "CoreMinimal.h"
+#include "UncookedOnlyUtils.h"
+#include "Param/AnimNextParameterBlock_EditorData.h"
+#include "Misc/AutomationTest.h"
+#include "Param/ParameterBlockFactory.h"
+#include "Param/ParameterLibraryFactory.h"
+#include "Param/AnimNextParameterBlock.h"
+#include "Param/AnimNextParameterBlockBinding.h"
+#include "Param/AnimNextParameterLibrary.h"
+#include "Animation/AnimSequence.h"
+#if WITH_EDITOR
+#include "ScopedTransaction.h"
+#include "Editor.h"
+#include "IPythonScriptPlugin.h"
+#endif
+
+// AnimNext Editor Tests
+
+#if WITH_DEV_AUTOMATION_TESTS && WITH_EDITOR
+
+IMPLEMENT_SIMPLE_AUTOMATION_TEST(FAnimationAnimNextParametersEditorTest_Library, "Animation.AnimNext.Parameters.Editor.Library", EAutomationTestFlags::EditorContext | EAutomationTestFlags::EngineFilter)
+
+bool FAnimationAnimNextParametersEditorTest_Library::RunTest(const FString& InParameters)
+{
+	using namespace UE::AnimNext;
+
+	UFactory* LibraryFactory = NewObject<UAnimNextParameterLibraryFactory>();
+	UAnimNextParameterLibrary* Library = Cast<UAnimNextParameterLibrary>(LibraryFactory->FactoryCreateNew(UAnimNextParameterLibrary::StaticClass(), GetTransientPackage(), TEXT("TestAnimNextParameterLibrary"), RF_Transient, nullptr, nullptr, NAME_None));
+	if(Library == nullptr)
+	{
+		AddError(TEXT("Could not create parameter library."));
+		return false;
+	}
+
+	AddErrorIfFalse(Library->AddParameter(TEXT("TestParam0"), FAnimNextParamType::GetType<bool>()) != nullptr, TEXT("Could not create new bool parameter in library."));
+	AddErrorIfFalse(Library->AddParameter(TEXT("TestParam1"), FAnimNextParamType::GetType<uint8>()) != nullptr, TEXT("Could not create new uint8 parameter in library."));
+	AddErrorIfFalse(Library->AddParameter(TEXT("TestParam2"), FAnimNextParamType::GetType<int32>()) != nullptr, TEXT("Could not create new int32 parameter in library."));
+	AddErrorIfFalse(Library->AddParameter(TEXT("TestParam3"), FAnimNextParamType::GetType<int64>()) != nullptr, TEXT("Could not create new int64 parameter in library."));
+	AddErrorIfFalse(Library->AddParameter(TEXT("TestParam4"), FAnimNextParamType::GetType<float>()) != nullptr, TEXT("Could not create new float parameter in library."));
+	AddErrorIfFalse(Library->AddParameter(TEXT("TestParam5"), FAnimNextParamType::GetType<double>()) != nullptr, TEXT("Could not create new double parameter in library."));
+	AddErrorIfFalse(Library->AddParameter(TEXT("TestParam6"), FAnimNextParamType::GetType<FName>()) != nullptr, TEXT("Could not create new FName parameter in library."));
+	AddErrorIfFalse(Library->AddParameter(TEXT("TestParam7"), FAnimNextParamType::GetType<FString>()) != nullptr, TEXT("Could not create new FString parameter in library."));
+	AddErrorIfFalse(Library->AddParameter(TEXT("TestParam8"), FAnimNextParamType::GetType<FText>()) != nullptr, TEXT("Could not create new FText parameter in library."));
+	AddErrorIfFalse(Library->AddParameter(TEXT("TestParam9"), FAnimNextParamType::GetType<EPropertyBagPropertyType>()) != nullptr, TEXT("Could not create new enum parameter in library."));
+	AddErrorIfFalse(Library->AddParameter(TEXT("TestParam10"), FAnimNextParamType::GetType<FVector>()) != nullptr, TEXT("Could not create new FVector parameter in library."));
+	AddErrorIfFalse(Library->AddParameter(TEXT("TestParam11"), FAnimNextParamType::GetType<FQuat>()) != nullptr, TEXT("Could not create new FQuat parameter in library."));
+	AddErrorIfFalse(Library->AddParameter(TEXT("TestParam12"), FAnimNextParamType::GetType<FTransform>()) != nullptr, TEXT("Could not create new FTransform parameter in library."));
+	AddErrorIfFalse(Library->AddParameter(TEXT("TestParam13"), FAnimNextParamType::GetType<TObjectPtr<UObject>>()) != nullptr, TEXT("Could not create new TObjectPtr<UObject> parameter in library."));
+	AddErrorIfFalse(Library->AddParameter(TEXT("TestParam14"), FAnimNextParamType::GetType<TObjectPtr<UAnimSequence>>()) != nullptr, TEXT("Could not create new TObjectPtr<UAnimSequence> parameter in library."));
+	AddErrorIfFalse(Library->AddParameter(TEXT("TestParam15"), FAnimNextParamType::GetType<TArray<float>>()) != nullptr, TEXT("Could not create new TArray<float> parameter in library."));
+	AddErrorIfFalse(Library->AddParameter(TEXT("TestParam16"), FAnimNextParamType::GetType<TArray<TObjectPtr<UAnimSequence>>>()) != nullptr, TEXT("Could not create new TArray<TObjectPtr<UAnimSequence>> parameter in library."));
+
+	AddErrorIfFalse(Library->RemoveParameter(TEXT("TestParam0")), TEXT("Could not remove parameter from library."));
+	AddErrorIfFalse(Library->RemoveParameters({ TEXT("TestParam1"), TEXT("TestParam2") }), TEXT("Could not remove parameters from library."));
+	AddErrorIfFalse(Library->FindParameter(TEXT("TestParam3")) != nullptr, TEXT("Could not find parameter in library."));
+	
+	CollectGarbage(GARBAGE_COLLECTION_KEEPFLAGS);
+
+	return true;
+}
+
+IMPLEMENT_SIMPLE_AUTOMATION_TEST(FAnimationAnimNextParametersEditorTest_Block, "Animation.AnimNext.Parameters.Editor.Block", EAutomationTestFlags::EditorContext | EAutomationTestFlags::EngineFilter)
+
+bool FAnimationAnimNextParametersEditorTest_Block::RunTest(const FString& InParameters)
+{
+	using namespace UE::AnimNext;
+
+	UFactory* LibraryFactory = NewObject<UAnimNextParameterLibraryFactory>();
+	UAnimNextParameterLibrary* Library = Cast<UAnimNextParameterLibrary>(LibraryFactory->FactoryCreateNew(UAnimNextParameterLibrary::StaticClass(), GetTransientPackage(), TEXT("TestAnimNextParameterLibrary"), RF_Transient, nullptr, nullptr, NAME_None));
+	if(Library == nullptr)
+	{
+		AddError(TEXT("Could not create parameter library."));
+		return false;
+	}
+
+	AddErrorIfFalse(Library->AddParameter(TEXT("TestParam"), FAnimNextParamType::GetType<bool>()) != nullptr, TEXT("Could not create new parameter in library."));
+
+	UFactory* BlockFactory = NewObject<UAnimNextParameterBlockFactory>();
+	UAnimNextParameterBlock* Block = Cast<UAnimNextParameterBlock>(BlockFactory->FactoryCreateNew(UAnimNextParameterBlock::StaticClass(), GetTransientPackage(), TEXT("TestAnimNextParameterBlock"), RF_Transient, nullptr, nullptr, NAME_None));
+	if(Block == nullptr)
+	{
+		AddError(TEXT("Could not create parameter block."));
+		return false;
+	}
+
+	UAnimNextParameterBlock_EditorData* EditorData = UncookedOnly::FUtils::GetEditorData(Block);
+	if(EditorData == nullptr)
+	{
+		AddError(TEXT("Parameter block has no editor data."));
+		return false;
+	}
+
+	// AddBinding
+	UAnimNextParameterBlockBinding* Binding = nullptr;
+	{
+		FScopedTransaction Transaction(FText::GetEmpty());
+		Binding = EditorData->AddBinding(TEXT("TestParam"), Library);
+		AddErrorIfFalse(Binding != nullptr, TEXT("Could not create new binding in block."));
+	}
+
+	GEditor->UndoTransaction();
+	AddErrorIfFalse(EditorData->Entries.Num() == 0, FString::Printf(TEXT("Unexpected entry count found in parameter block (Have %d, expected 0)."), EditorData->Entries.Num()));
+
+	GEditor->RedoTransaction();
+	AddErrorIfFalse(EditorData->Entries.Num() == 1, FString::Printf(TEXT("Unexpected entry count found in parameter block (Have %d, expected 1)."), EditorData->Entries.Num()));
+
+	// RemoveAllBindings
+	{
+		FScopedTransaction Transaction(FText::GetEmpty());
+		AddErrorIfFalse(EditorData->RemoveAllBindings(TEXT("TestParam")), TEXT("Failed to remove binding."));
+	}
+
+	GEditor->UndoTransaction();
+	AddErrorIfFalse(EditorData->Entries.Num() == 1, FString::Printf(TEXT("Unexpected entry count found in parameter block (Have %d, expected 1)."), EditorData->Entries.Num()));
+
+	GEditor->RedoTransaction();
+	AddErrorIfFalse(EditorData->Entries.Num() == 0, FString::Printf(TEXT("Unexpected entry count found in parameter block (Have %d, expected 0)."), EditorData->Entries.Num()));
+
+	GEditor->UndoTransaction();
+	AddErrorIfFalse(EditorData->Entries.Num() == 1, FString::Printf(TEXT("Unexpected entry count found in parameter block (Have %d, expected 1)."), EditorData->Entries.Num()));
+
+	// RemoveEntry
+	{
+		FScopedTransaction Transaction(FText::GetEmpty());
+		AddErrorIfFalse(EditorData->RemoveEntry(Binding), TEXT("Failed to remove entry."));
+	}
+
+	GEditor->UndoTransaction();
+	AddErrorIfFalse(EditorData->Entries.Num() == 1, FString::Printf(TEXT("Unexpected entry count found in parameter block (Have %d, expected 1)."), EditorData->Entries.Num()));
+
+	GEditor->RedoTransaction();
+	AddErrorIfFalse(EditorData->Entries.Num() == 0, FString::Printf(TEXT("Unexpected entry count found in parameter block (Have %d, expected 0)."), EditorData->Entries.Num()));
+
+	GEditor->UndoTransaction();
+	AddErrorIfFalse(EditorData->Entries.Num() == 1, FString::Printf(TEXT("Unexpected entry count found in parameter block (Have %d, expected 1)."), EditorData->Entries.Num()));
+	
+	UAnimNextParameterBlock* OtherBlock = Cast<UAnimNextParameterBlock>(BlockFactory->FactoryCreateNew(UAnimNextParameterBlock::StaticClass(), GetTransientPackage(), TEXT("TestAnimNextParameterBlock2"), RF_Transient, nullptr, nullptr, NAME_None));
+	if(OtherBlock == nullptr)
+	{
+		AddError(TEXT("Could not create additional parameter block."));
+		return false;
+	}
+	
+	UAnimNextParameterBlock_EditorData* OtherEditorData = UncookedOnly::FUtils::GetEditorData(OtherBlock);
+	if(OtherEditorData == nullptr)
+	{
+		AddError(TEXT("Additional parameter block has no editor data."));
+		return false;
+	}
+
+	// AddBindingReference
+	{
+		FScopedTransaction Transaction(FText::GetEmpty());
+		UAnimNextParameterBlockBindingReference* Ref = OtherEditorData->AddBindingReference(TEXT("TestParam"), Library, Block);
+		AddErrorIfFalse(Ref != nullptr, TEXT("Could not create new binding reference in block."));
+	}
+
+	GEditor->UndoTransaction();
+	AddErrorIfFalse(OtherEditorData->Entries.Num() == 0, FString::Printf(TEXT("Unexpected entry count found in parameter block (Have %d, expected 0)."), OtherEditorData->Entries.Num()));
+
+	GEditor->RedoTransaction();
+	AddErrorIfFalse(OtherEditorData->Entries.Num() == 1, FString::Printf(TEXT("Unexpected entry count found in parameter block (Have %d, expected 1)."), OtherEditorData->Entries.Num()));
+
+	// FindBinding
+	AddErrorIfFalse(EditorData->FindBinding(TEXT("TestParam")) != nullptr, TEXT("Could not find binding in block."));
+	AddErrorIfFalse(OtherEditorData->FindBinding(TEXT("TestParam")) != nullptr, TEXT("Could not find binding refernce in block."));
+
+	CollectGarbage(GARBAGE_COLLECTION_KEEPFLAGS);
+
+	return true;
+}
+
+IMPLEMENT_SIMPLE_AUTOMATION_TEST(FAnimationAnimNextParametersEditorTest_Python, "Animation.AnimNext.Parameters.Editor.Python", EAutomationTestFlags::EditorContext | EAutomationTestFlags::EngineFilter)
+
+bool FAnimationAnimNextParametersEditorTest_Python::RunTest(const FString& InParameters)
+{
+	using namespace UE::AnimNext;
+
+	const TCHAR* Script = TEXT(
+		"asset_tools = unreal.AssetToolsHelpers.get_asset_tools()\n"
+		"library = unreal.AssetTools.create_asset(asset_tools, asset_name = \"TestLibrary\", package_path = \"/Game/\", asset_class = unreal.AnimNextParameterLibrary, factory = unreal.AnimNextParameterLibraryFactory())\n"
+		"block = unreal.AssetTools.create_asset(asset_tools, asset_name = \"TestBlock\", package_path = \"/Game/\", asset_class = unreal.AnimNextParameterBlock, factory = unreal.AnimNextParameterBlockFactory())\n"
+		"other_block = unreal.AssetTools.create_asset(asset_tools, asset_name = \"TestBlock1\", package_path = \"/Game/\", asset_class = unreal.AnimNextParameterBlock, factory = unreal.AnimNextParameterBlockFactory())\n"
+		"library.add_parameter(name = \"TestParam\", value_type = unreal.PropertyBagPropertyType.BOOL, container_type = unreal.PropertyBagContainerType.NONE)\n"
+		"block.add_binding(name = \"TestParam\", library = library)\n"
+		"other_block.add_binding_reference(name = \"TestParam\", library = library, referenced_block = block)\n"
+		"unreal.EditorAssetLibrary.delete_loaded_asset(library)\n"
+		"unreal.EditorAssetLibrary.delete_loaded_asset(block)\n"
+		"unreal.EditorAssetLibrary.delete_loaded_asset(other_block)\n"
+	);
+
+	IPythonScriptPlugin::Get()->ExecPythonCommand(Script);
+
+	CollectGarbage(GARBAGE_COLLECTION_KEEPFLAGS);
+
+	return true;
+}
+
+#endif	// WITH_DEV_AUTOMATION_TESTS && WITH_EDITOR
\ No newline at end of file
diff --git a/Engine/Plugins/Experimental/AnimNext/Source/AnimNextTestSuite/Private/AnimNextParametersTest.cpp b/Engine/Plugins/Experimental/AnimNext/Source/AnimNextTestSuite/Private/AnimNextParametersTest.cpp
new file mode 100644
--- /dev/null
+++ b/Engine/Plugins/Experimental/AnimNext/Source/AnimNextTestSuite/Private/AnimNextParametersTest.cpp
@@ -0,0 +1,288 @@
+﻿// Copyright Epic Games, Inc. All Rights Reserved.
+
+#include "CoreMinimal.h"
+#include "Misc/AutomationTest.h"
+#include "Param/ParamType.h"
+#include "Param/ParamTypeHandle.h"
+#include "Animation/AnimSequence.h"
+#include "Async/ParallelFor.h"
+
+// AnimNext Parameters Tests
+
+#if WITH_DEV_AUTOMATION_TESTS
+
+IMPLEMENT_SIMPLE_AUTOMATION_TEST(FAnimationAnimNextParametersTest_Types, "Animation.AnimNext.Parameters.Types", EAutomationTestFlags::EditorContext | EAutomationTestFlags::EngineFilter)
+
+bool FAnimationAnimNextParametersTest_Types::RunTest(const FString& InParameters)
+{
+	using namespace UE::AnimNext;
+
+	// None is invalid
+	FAnimNextParamType ParameterTypeValueNone(FAnimNextParamType::EValueType::None); 
+	AddErrorIfFalse(!ParameterTypeValueNone.IsValid(), TEXT("Parameter type None is valid."));
+	AddErrorIfFalse(!ParameterTypeValueNone.GetHandle().IsValid(), TEXT("Parameter type handle None is valid."));
+
+	// None is invalid for all containers
+	for(uint8 ContainerType = (uint8)FAnimNextParamType::EContainerType::None; ContainerType <= (uint8)FAnimNextParamType::EContainerType::Array; ++ContainerType)
+	{
+		FAnimNextParamType ParameterTypeValueContainerNone(FAnimNextParamType::EValueType::None, (FAnimNextParamType::EContainerType)ContainerType); 
+		AddErrorIfFalse(!ParameterTypeValueContainerNone.IsValid(), FString::Printf(TEXT("Parameter type None, container type %d is valid."), ContainerType));
+		AddErrorIfFalse(!ParameterTypeValueContainerNone.GetHandle().IsValid(), FString::Printf(TEXT("Parameter type handle None, container type %d is valid."), ContainerType));
+	}
+
+	// Null object types
+	for(uint8 ObjectValueType = (uint8)FAnimNextParamType::EValueType::Enum; ObjectValueType <= (uint8)FAnimNextParamType::EValueType::SoftClass; ++ObjectValueType)
+	{
+		for(uint8 ContainerType = (uint8)FAnimNextParamType::EContainerType::None; ContainerType <= (uint8)FAnimNextParamType::EContainerType::Array; ++ContainerType)
+		{
+			FAnimNextParamType ParameterTypeNullObject((FAnimNextParamType::EValueType)ObjectValueType, (FAnimNextParamType::EContainerType)ContainerType, nullptr); 
+			AddErrorIfFalse(!ParameterTypeNullObject.IsValid(), FString::Printf(TEXT("Parameter type %d, container type %d with null object is valid."), ObjectValueType, ContainerType));
+			AddErrorIfFalse(!ParameterTypeNullObject.GetHandle().IsValid(), FString::Printf(TEXT("Parameter type handle %d, container type %d with null object is valid."), ObjectValueType, ContainerType));
+		}
+	}
+
+	// Non object types
+	for(uint8 ValueType = (uint8)FAnimNextParamType::EValueType::Bool; ValueType < (uint8)FAnimNextParamType::EValueType::Enum; ++ValueType)
+	{
+		for(uint8 ContainerType = (uint8)FAnimNextParamType::EContainerType::None; ContainerType <= (uint8)FAnimNextParamType::EContainerType::Array; ++ContainerType)
+		{
+			FAnimNextParamType TestValueContainerParameterType((FAnimNextParamType::EValueType)ValueType, (FAnimNextParamType::EContainerType)ContainerType);
+			AddErrorIfFalse(TestValueContainerParameterType.IsValid(), FString::Printf(TEXT("Parameter type %d, container type %d is invalid."), ValueType, ContainerType));
+			AddErrorIfFalse(TestValueContainerParameterType.GetHandle().IsValid(), FString::Printf(TEXT("Parameter type handle %d, container type %d is invalid."), ValueType, ContainerType));
+		}
+	}
+
+	UObject* ExampleValidObjects[(uint8)FAnimNextParamType::EValueType::SoftClass + 1] =
+	{
+		nullptr,
+		nullptr,
+		nullptr,
+		nullptr,
+		nullptr,
+		nullptr,
+		nullptr,
+		nullptr,
+		nullptr,
+		nullptr,
+		FindObjectChecked<UEnum>(nullptr, TEXT("/Script/StructUtils.EPropertyBagPropertyType")),
+		FAnimNextParamType::StaticStruct(),
+		UObject::StaticClass(),
+		UObject::StaticClass(),
+		UObject::StaticClass(),
+		UObject::StaticClass()
+	};
+
+	// Non-null valid object types
+	for(uint8 ObjectValueType = (uint8)FAnimNextParamType::EValueType::Enum; ObjectValueType <= (uint8)FAnimNextParamType::EValueType::SoftClass; ++ObjectValueType)
+	{
+		for(uint8 ContainerType = (uint8)FAnimNextParamType::EContainerType::None; ContainerType <= (uint8)FAnimNextParamType::EContainerType::Array; ++ContainerType)
+		{
+			FAnimNextParamType TestValueContainerParameterType((FAnimNextParamType::EValueType)ObjectValueType, (FAnimNextParamType::EContainerType)ContainerType, ExampleValidObjects[ObjectValueType]);
+			AddErrorIfFalse(TestValueContainerParameterType.IsValid(), FString::Printf(TEXT("Object parameter type %d, container type %d is invalid."), ObjectValueType, ContainerType));
+			AddErrorIfFalse(TestValueContainerParameterType.GetHandle().IsValid(), FString::Printf(TEXT("Object parameter type handle %d, container type %d is invalid."), ObjectValueType, ContainerType));
+		}
+	}
+
+	UObject* ExampleInvalidObjects[(uint8)FAnimNextParamType::EValueType::SoftClass + 1] =
+	{
+		nullptr,
+		nullptr,
+		nullptr,
+		nullptr,
+		nullptr,
+		nullptr,
+		nullptr,
+		nullptr,
+		nullptr,
+		nullptr,
+		FAnimNextParamType::StaticStruct(),
+		FindObjectChecked<UEnum>(nullptr, TEXT("/Script/StructUtils.EPropertyBagPropertyType")),
+		FAnimNextParamType::StaticStruct(),
+		FAnimNextParamType::StaticStruct(),
+		FAnimNextParamType::StaticStruct(),
+		FAnimNextParamType::StaticStruct()
+	};
+
+	// Non-null invalid object types
+	for(uint8 ObjectValueType = (uint8)FAnimNextParamType::EValueType::Enum; ObjectValueType <= (uint8)FAnimNextParamType::EValueType::SoftClass; ++ObjectValueType)
+	{
+		for(uint8 ContainerType = (uint8)FAnimNextParamType::EContainerType::None; ContainerType <= (uint8)FAnimNextParamType::EContainerType::Array; ++ContainerType)
+		{
+			FAnimNextParamType TestValueContainerParameterType((FAnimNextParamType::EValueType)ObjectValueType, (FAnimNextParamType::EContainerType)ContainerType, ExampleInvalidObjects[ObjectValueType]);
+			AddErrorIfFalse(!TestValueContainerParameterType.IsValid(), FString::Printf(TEXT("Object parameter type %d, container type %d is valid."), ObjectValueType, ContainerType));
+			AddErrorIfFalse(!TestValueContainerParameterType.GetHandle().IsValid(), FString::Printf(TEXT("Object parameter type handle %d, container type %d is valid."), ObjectValueType, ContainerType));
+		}
+	}
+
+	// Check type inference
+	AddErrorIfFalse(FAnimNextParamType::GetType<bool>().IsValid(), TEXT("bool parameter is invalid."));
+	AddErrorIfFalse(FAnimNextParamType::GetType<uint8>().IsValid(), TEXT("uint8 parameter is invalid."));
+	AddErrorIfFalse(FAnimNextParamType::GetType<int32>().IsValid(), TEXT("int32 parameter is invalid."));
+	AddErrorIfFalse(FAnimNextParamType::GetType<int64>().IsValid(), TEXT("int64 parameter is invalid."));
+	AddErrorIfFalse(FAnimNextParamType::GetType<float>().IsValid(), TEXT("float parameter is invalid."));
+	AddErrorIfFalse(FAnimNextParamType::GetType<double>().IsValid(), TEXT("double parameter is invalid."));
+	AddErrorIfFalse(FAnimNextParamType::GetType<FName>().IsValid(), TEXT("FName parameter is invalid."));
+	AddErrorIfFalse(FAnimNextParamType::GetType<FString>().IsValid(), TEXT("FString parameter is invalid."));
+	AddErrorIfFalse(FAnimNextParamType::GetType<FText>().IsValid(), TEXT("FText parameter is invalid."));
+	AddErrorIfFalse(FAnimNextParamType::GetType<EPropertyBagContainerType>().IsValid(), TEXT("Enum parameter is invalid."));
+	AddErrorIfFalse(FAnimNextParamType::GetType<FAnimNextParamType>().IsValid(), TEXT("Struct parameter is invalid."));
+	AddErrorIfFalse(FAnimNextParamType::GetType<FVector>().IsValid(), TEXT("Struct parameter is invalid."));
+	AddErrorIfFalse(FAnimNextParamType::GetType<FTransform>().IsValid(), TEXT("Struct parameter is invalid."));
+	AddErrorIfFalse(FAnimNextParamType::GetType<FQuat>().IsValid(), TEXT("Struct parameter is invalid."));
+	AddErrorIfFalse(FAnimNextParamType::GetType<UObject>().IsValid(), TEXT("UObject parameter is invalid."));
+	AddErrorIfFalse(FAnimNextParamType::GetType<TObjectPtr<UObject>>().IsValid(), TEXT("TObjectPtr<UObject> parameter is invalid."));
+	AddErrorIfFalse(FAnimNextParamType::GetType<UClass>().IsValid(), TEXT("UClass parameter is invalid."));
+	AddErrorIfFalse(FAnimNextParamType::GetType<TSubclassOf<UObject>>().IsValid(), TEXT("TSubclassOff<UObject> parameter is invalid."));
+	AddErrorIfFalse(FAnimNextParamType::GetType<TSoftObjectPtr<UObject>>().IsValid(), TEXT("TSoftObjectPtr<UObject> parameter is invalid."));
+	AddErrorIfFalse(FAnimNextParamType::GetType<TSoftClassPtr<UObject>>().IsValid(), TEXT("TSoftClassPtr<UObject> parameter is invalid."));
+
+	AddErrorIfFalse(FAnimNextParamType::GetType<TArray<bool>>().IsValid(), TEXT("bool array parameter is invalid."));
+	AddErrorIfFalse(FAnimNextParamType::GetType<TArray<uint8>>().IsValid(), TEXT("uint8 array parameter is invalid."));
+	AddErrorIfFalse(FAnimNextParamType::GetType<TArray<int32>>().IsValid(), TEXT("int32 array parameter is invalid."));
+	AddErrorIfFalse(FAnimNextParamType::GetType<TArray<int64>>().IsValid(), TEXT("int64 array parameter is invalid."));
+	AddErrorIfFalse(FAnimNextParamType::GetType<TArray<float>>().IsValid(), TEXT("float array parameter is invalid."));
+	AddErrorIfFalse(FAnimNextParamType::GetType<TArray<double>>().IsValid(), TEXT("double array parameter is invalid."));
+	AddErrorIfFalse(FAnimNextParamType::GetType<TArray<FName>>().IsValid(), TEXT("FName array parameter is invalid."));
+	AddErrorIfFalse(FAnimNextParamType::GetType<TArray<FString>>().IsValid(), TEXT("FString array parameter is invalid."));
+	AddErrorIfFalse(FAnimNextParamType::GetType<TArray<FText>>().IsValid(), TEXT("FText array parameter is invalid."));
+	AddErrorIfFalse(FAnimNextParamType::GetType<TArray<EPropertyBagContainerType>>().IsValid(), TEXT("Enum array parameter is invalid."));
+	AddErrorIfFalse(FAnimNextParamType::GetType<TArray<FAnimNextParamType>>().IsValid(), TEXT("Struct array parameter is invalid."));
+	AddErrorIfFalse(FAnimNextParamType::GetType<TArray<FVector>>().IsValid(), TEXT("Struct array parameter is invalid."));
+	AddErrorIfFalse(FAnimNextParamType::GetType<TArray<FTransform>>().IsValid(), TEXT("Struct array parameter is invalid."));
+	AddErrorIfFalse(FAnimNextParamType::GetType<TArray<FQuat>>().IsValid(), TEXT("Struct array parameter is invalid."));
+	AddErrorIfFalse(FAnimNextParamType::GetType<TArray<UObject>>().IsValid(), TEXT("UObject array parameter is invalid."));
+	AddErrorIfFalse(FAnimNextParamType::GetType<TArray<TObjectPtr<UObject>>>().IsValid(), TEXT("TObjectPtr<UObject> array parameter is invalid."));
+	AddErrorIfFalse(FAnimNextParamType::GetType<TArray<UClass>>().IsValid(), TEXT("UClass array parameter is invalid."));
+	AddErrorIfFalse(FAnimNextParamType::GetType<TArray<TSubclassOf<UObject>>>().IsValid(), TEXT("TSubclassOff<UObject> array parameter is invalid."));
+	AddErrorIfFalse(FAnimNextParamType::GetType<TArray<TSoftObjectPtr<UObject>>>().IsValid(), TEXT("TSoftObjectPtr<UObject> array parameter is invalid."));
+	AddErrorIfFalse(FAnimNextParamType::GetType<TArray<TSoftClassPtr<UObject>>>().IsValid(), TEXT("TSoftClassPtr<UObject> array parameter is invalid."));
+
+	// Type inference for handles
+	AddErrorIfFalse(FAnimNextParamType::GetType<bool>().GetHandle().IsValid(), TEXT("bool parameter handle is invalid."));
+	AddErrorIfFalse(FAnimNextParamType::GetType<uint8>().GetHandle().IsValid(), TEXT("uint8 parameter handle is invalid."));
+	AddErrorIfFalse(FAnimNextParamType::GetType<int32>().GetHandle().IsValid(), TEXT("int32 parameter handle is invalid."));
+	AddErrorIfFalse(FAnimNextParamType::GetType<int64>().GetHandle().IsValid(), TEXT("int64 parameter handle is invalid."));
+	AddErrorIfFalse(FAnimNextParamType::GetType<float>().GetHandle().IsValid(), TEXT("float parameter handle is invalid."));
+	AddErrorIfFalse(FAnimNextParamType::GetType<double>().GetHandle().IsValid(), TEXT("double parameter handle is invalid."));
+	AddErrorIfFalse(FAnimNextParamType::GetType<FName>().GetHandle().IsValid(), TEXT("FName parameter handle is invalid."));
+	AddErrorIfFalse(FAnimNextParamType::GetType<FString>().GetHandle().IsValid(), TEXT("FString parameter handle is invalid."));
+	AddErrorIfFalse(FAnimNextParamType::GetType<FText>().GetHandle().IsValid(), TEXT("FText parameter handle is invalid."));
+	AddErrorIfFalse(FAnimNextParamType::GetType<EPropertyBagContainerType>().GetHandle().IsValid(), TEXT("Enum parameter handle is invalid."));
+	AddErrorIfFalse(FAnimNextParamType::GetType<FAnimNextParamType>().GetHandle().IsValid(), TEXT("Struct parameter handle is invalid."));
+	AddErrorIfFalse(FAnimNextParamType::GetType<FVector>().GetHandle().IsValid(), TEXT("Struct parameter handle is invalid."));
+	AddErrorIfFalse(FAnimNextParamType::GetType<FTransform>().GetHandle().IsValid(), TEXT("Struct parameter handle is invalid."));
+	AddErrorIfFalse(FAnimNextParamType::GetType<FQuat>().GetHandle().IsValid(), TEXT("Struct parameter handle is invalid."));
+	AddErrorIfFalse(FAnimNextParamType::GetType<UObject>().GetHandle().IsValid(), TEXT("UObject parameter handle is invalid."));
+	AddErrorIfFalse(FAnimNextParamType::GetType<TObjectPtr<UObject>>().GetHandle().IsValid(), TEXT("TObjectPtr<UObject> parameter handle is invalid."));
+	AddErrorIfFalse(FAnimNextParamType::GetType<UClass>().GetHandle().IsValid(), TEXT("UClass parameter handle is invalid."));
+	AddErrorIfFalse(FAnimNextParamType::GetType<TSubclassOf<UObject>>().GetHandle().IsValid(), TEXT("TSubclassOff<UObject> parameter handle is invalid."));
+	AddErrorIfFalse(FAnimNextParamType::GetType<TSoftObjectPtr<UObject>>().GetHandle().IsValid(), TEXT("TSoftObjectPtr<UObject> parameter handle is invalid."));
+	AddErrorIfFalse(FAnimNextParamType::GetType<TSoftClassPtr<UObject>>().GetHandle().IsValid(), TEXT("TSoftClassPtr<UObject> parameter handle is invalid."));
+
+	AddErrorIfFalse(FAnimNextParamType::GetType<TArray<bool>>().GetHandle().IsValid(), TEXT("bool array parameter handle is invalid."));
+	AddErrorIfFalse(FAnimNextParamType::GetType<TArray<uint8>>().GetHandle().IsValid(), TEXT("uint8 array parameter handle is invalid."));
+	AddErrorIfFalse(FAnimNextParamType::GetType<TArray<int32>>().GetHandle().IsValid(), TEXT("int32 array parameter handle is invalid."));
+	AddErrorIfFalse(FAnimNextParamType::GetType<TArray<int64>>().GetHandle().IsValid(), TEXT("int64 array parameter handle is invalid."));
+	AddErrorIfFalse(FAnimNextParamType::GetType<TArray<float>>().GetHandle().IsValid(), TEXT("float array parameter handle is invalid."));
+	AddErrorIfFalse(FAnimNextParamType::GetType<TArray<double>>().GetHandle().IsValid(), TEXT("double array parameter handle is invalid."));
+	AddErrorIfFalse(FAnimNextParamType::GetType<TArray<FName>>().GetHandle().IsValid(), TEXT("FName array parameter handle is invalid."));
+	AddErrorIfFalse(FAnimNextParamType::GetType<TArray<FString>>().GetHandle().IsValid(), TEXT("FString array parameter handle is invalid."));
+	AddErrorIfFalse(FAnimNextParamType::GetType<TArray<FText>>().GetHandle().IsValid(), TEXT("FText array parameter handle is invalid."));
+	AddErrorIfFalse(FAnimNextParamType::GetType<TArray<EPropertyBagContainerType>>().GetHandle().IsValid(), TEXT("Enum array parameter handle is invalid."));
+	AddErrorIfFalse(FAnimNextParamType::GetType<TArray<FAnimNextParamType>>().GetHandle().IsValid(), TEXT("Struct array parameter handle is invalid."));
+	AddErrorIfFalse(FAnimNextParamType::GetType<TArray<FVector>>().GetHandle().IsValid(), TEXT("Struct array parameter handle is invalid."));
+	AddErrorIfFalse(FAnimNextParamType::GetType<TArray<FTransform>>().GetHandle().IsValid(), TEXT("Struct array parameter handle is invalid."));
+	AddErrorIfFalse(FAnimNextParamType::GetType<TArray<FQuat>>().GetHandle().IsValid(), TEXT("Struct array parameter handle is invalid."));
+	AddErrorIfFalse(FAnimNextParamType::GetType<TArray<UObject>>().GetHandle().IsValid(), TEXT("UObject array parameter handle is invalid."));
+	AddErrorIfFalse(FAnimNextParamType::GetType<TArray<TObjectPtr<UObject>>>().GetHandle().IsValid(), TEXT("TObjectPtr<UObject> array parameter handle is invalid."));
+	AddErrorIfFalse(FAnimNextParamType::GetType<TArray<UClass>>().GetHandle().IsValid(), TEXT("UClass array parameter handle is invalid."));
+	AddErrorIfFalse(FAnimNextParamType::GetType<TArray<TSubclassOf<UObject>>>().GetHandle().IsValid(), TEXT("TSubclassOff<UObject> array parameter handle is invalid."));
+	AddErrorIfFalse(FAnimNextParamType::GetType<TArray<TSoftObjectPtr<UObject>>>().GetHandle().IsValid(), TEXT("TSoftObjectPtr<UObject> array parameter handle is invalid."));
+	AddErrorIfFalse(FAnimNextParamType::GetType<TArray<TSoftClassPtr<UObject>>>().GetHandle().IsValid(), TEXT("TSoftClassPtr<UObject> array parameter handle is invalid."));
+
+	// Raw handle type inference
+	AddErrorIfFalse(FParamTypeHandle::GetHandle<bool>().IsValid(), TEXT("bool parameter handle is invalid."));
+	AddErrorIfFalse(FParamTypeHandle::GetHandle<uint8>().IsValid(), TEXT("uint8 parameter handle is invalid."));
+	AddErrorIfFalse(FParamTypeHandle::GetHandle<int32>().IsValid(), TEXT("int32 parameter handle is invalid."));
+	AddErrorIfFalse(FParamTypeHandle::GetHandle<int64>().IsValid(), TEXT("int64 parameter handle is invalid."));
+	AddErrorIfFalse(FParamTypeHandle::GetHandle<float>().IsValid(), TEXT("float parameter handle is invalid."));
+	AddErrorIfFalse(FParamTypeHandle::GetHandle<double>().IsValid(), TEXT("double parameter handle is invalid."));
+	AddErrorIfFalse(FParamTypeHandle::GetHandle<FName>().IsValid(), TEXT("FName parameter handle is invalid."));
+	AddErrorIfFalse(FParamTypeHandle::GetHandle<FString>().IsValid(), TEXT("FString parameter handle is invalid."));
+	AddErrorIfFalse(FParamTypeHandle::GetHandle<FText>().IsValid(), TEXT("FText parameter handle is invalid."));
+	AddErrorIfFalse(FParamTypeHandle::GetHandle<EPropertyBagContainerType>().IsValid(), TEXT("Enum parameter handle is invalid."));
+	AddErrorIfFalse(FParamTypeHandle::GetHandle<FAnimNextParamType>().IsValid(), TEXT("Struct parameter handle is invalid."));
+	AddErrorIfFalse(FParamTypeHandle::GetHandle<FVector>().IsValid(), TEXT("Struct parameter handle is invalid."));
+	AddErrorIfFalse(FParamTypeHandle::GetHandle<FTransform>().IsValid(), TEXT("Struct parameter handle is invalid."));
+	AddErrorIfFalse(FParamTypeHandle::GetHandle<FQuat>().IsValid(), TEXT("Struct parameter handle is invalid."));
+	AddErrorIfFalse(FParamTypeHandle::GetHandle<UObject>().IsValid(), TEXT("UObject parameter handle is invalid."));
+	AddErrorIfFalse(FParamTypeHandle::GetHandle<TObjectPtr<UObject>>().IsValid(), TEXT("TObjectPtr<UObject> parameter handle is invalid."));
+	AddErrorIfFalse(FParamTypeHandle::GetHandle<UClass>().IsValid(), TEXT("UClass parameter handle is invalid."));
+	AddErrorIfFalse(FParamTypeHandle::GetHandle<TSubclassOf<UObject>>().IsValid(), TEXT("TSubclassOf<UObject> parameter handle is invalid."));
+	AddErrorIfFalse(FParamTypeHandle::GetHandle<TSoftObjectPtr<UObject>>().IsValid(), TEXT("TSoftObjectPtr<UObject> parameter handle is invalid."));
+	AddErrorIfFalse(FParamTypeHandle::GetHandle<TSoftClassPtr<UObject>>().IsValid(), TEXT("TSoftClassPtr<UObject> parameter handle is invalid."));
+
+	AddErrorIfFalse(FParamTypeHandle::GetHandle<TArray<bool>>().IsValid(), TEXT("bool array parameter handle is invalid."));
+	AddErrorIfFalse(FParamTypeHandle::GetHandle<TArray<uint8>>().IsValid(), TEXT("uint8 array parameter handle is invalid."));
+	AddErrorIfFalse(FParamTypeHandle::GetHandle<TArray<int32>>().IsValid(), TEXT("int32 array parameter handle is invalid."));
+	AddErrorIfFalse(FParamTypeHandle::GetHandle<TArray<int64>>().IsValid(), TEXT("int64 array parameter handle is invalid."));
+	AddErrorIfFalse(FParamTypeHandle::GetHandle<TArray<float>>().IsValid(), TEXT("float array parameter handle is invalid."));
+	AddErrorIfFalse(FParamTypeHandle::GetHandle<TArray<double>>().IsValid(), TEXT("double array parameter handle is invalid."));
+	AddErrorIfFalse(FParamTypeHandle::GetHandle<TArray<FName>>().IsValid(), TEXT("FName array parameter handle is invalid."));
+	AddErrorIfFalse(FParamTypeHandle::GetHandle<TArray<FString>>().IsValid(), TEXT("FString array parameter handle is invalid."));
+	AddErrorIfFalse(FParamTypeHandle::GetHandle<TArray<FText>>().IsValid(), TEXT("FText array parameter handle is invalid."));
+	AddErrorIfFalse(FParamTypeHandle::GetHandle<TArray<EPropertyBagContainerType>>().IsValid(), TEXT("Enum array parameter handle is invalid."));
+	AddErrorIfFalse(FParamTypeHandle::GetHandle<TArray<FAnimNextParamType>>().IsValid(), TEXT("Struct array parameter handle is invalid."));
+	AddErrorIfFalse(FParamTypeHandle::GetHandle<TArray<FVector>>().IsValid(), TEXT("Struct array parameter handle is invalid."));
+	AddErrorIfFalse(FParamTypeHandle::GetHandle<TArray<FTransform>>().IsValid(), TEXT("Struct array parameter handle is invalid."));
+	AddErrorIfFalse(FParamTypeHandle::GetHandle<TArray<FQuat>>().IsValid(), TEXT("Struct array parameter handle is invalid."));
+	AddErrorIfFalse(FParamTypeHandle::GetHandle<TArray<UObject>>().IsValid(), TEXT("UObject array parameter handle is invalid."));
+	AddErrorIfFalse(FParamTypeHandle::GetHandle<TArray<TObjectPtr<UObject>>>().IsValid(), TEXT("TObjectPtr<UObject> array parameter handle is invalid."));
+	AddErrorIfFalse(FParamTypeHandle::GetHandle<TArray<UClass>>().IsValid(), TEXT("UClass array parameter handle is invalid."));
+	AddErrorIfFalse(FParamTypeHandle::GetHandle<TArray<TSubclassOf<UObject>>>().IsValid(), TEXT("TSubclassOf<UObject> array parameter handle is invalid."));
+	AddErrorIfFalse(FParamTypeHandle::GetHandle<TArray<TSoftObjectPtr<UObject>>>().IsValid(), TEXT("TSoftObjectPtr<UObject> array parameter handle is invalid."));
+	AddErrorIfFalse(FParamTypeHandle::GetHandle<TArray<TSoftClassPtr<UObject>>>().IsValid(), TEXT("TSoftClassPtr<UObject> array parameter handle is invalid."));
+
+	// Check built in type handles
+	// Should be built in:
+	AddErrorIfFalse(FParamTypeHandle::GetHandle<bool>().IsBuiltInType(), TEXT("bool parameter handle is not built in."));
+	AddErrorIfFalse(FParamTypeHandle::GetHandle<uint8>().IsBuiltInType(), TEXT("uint8 parameter handle is not built in."));
+	AddErrorIfFalse(FParamTypeHandle::GetHandle<int32>().IsBuiltInType(), TEXT("int32 parameter handle is not built in."));
+	AddErrorIfFalse(FParamTypeHandle::GetHandle<int64>().IsBuiltInType(), TEXT("int64 parameter handle is not built in."));
+	AddErrorIfFalse(FParamTypeHandle::GetHandle<float>().IsBuiltInType(), TEXT("float parameter handle is not built in."));
+	AddErrorIfFalse(FParamTypeHandle::GetHandle<double>().IsBuiltInType(), TEXT("double parameter handle is not built in."));
+	AddErrorIfFalse(FParamTypeHandle::GetHandle<FName>().IsBuiltInType(), TEXT("FName parameter handle is not built in."));
+	AddErrorIfFalse(FParamTypeHandle::GetHandle<FString>().IsBuiltInType(), TEXT("FString parameter handle is not built in."));
+	AddErrorIfFalse(FParamTypeHandle::GetHandle<FText>().IsBuiltInType(), TEXT("FText parameter handle is not built in."));
+	AddErrorIfFalse(FParamTypeHandle::GetHandle<FVector>().IsBuiltInType(), TEXT("FVector parameter handle is not built in."));
+	AddErrorIfFalse(FParamTypeHandle::GetHandle<FVector4>().IsBuiltInType(), TEXT("FVector4 parameter handle is not built in."));
+	AddErrorIfFalse(FParamTypeHandle::GetHandle<FQuat>().IsBuiltInType(), TEXT("FQuat parameter handle is not built in."));
+	AddErrorIfFalse(FParamTypeHandle::GetHandle<FTransform>().IsBuiltInType(), TEXT("FTransform parameter handle is not built in."));
+
+	// Should not be built in:
+	AddErrorIfFalse(!FParamTypeHandle::GetHandle<EPropertyBagContainerType>().IsBuiltInType(), TEXT("Enum parameter handle is built in."));
+	AddErrorIfFalse(!FParamTypeHandle::GetHandle<FAnimNextParamType>().IsBuiltInType(), TEXT("Struct parameter handle is built in."));
+	AddErrorIfFalse(!FParamTypeHandle::GetHandle<UObject>().IsBuiltInType(), TEXT("UObject parameter handle is built in."));
+	AddErrorIfFalse(!FParamTypeHandle::GetHandle<TObjectPtr<UObject>>().IsBuiltInType(), TEXT("TObjectPtr<UObject> parameter handle is built in."));
+	AddErrorIfFalse(!FParamTypeHandle::GetHandle<UClass>().IsBuiltInType(), TEXT("UClass parameter handle is built in."));
+	AddErrorIfFalse(!FParamTypeHandle::GetHandle<TSubclassOf<UObject>>().IsBuiltInType(), TEXT("TSubclassOff<UObject> parameter handle is built in."));
+	AddErrorIfFalse(!FParamTypeHandle::GetHandle<TSoftObjectPtr<UObject>>().IsBuiltInType(), TEXT("TSoftObjectPtr<UObject> parameter handle is built in."));
+	AddErrorIfFalse(!FParamTypeHandle::GetHandle<TSoftClassPtr<UObject>>().IsBuiltInType(), TEXT("TSoftClassPtr<UObject> parameter handle is built in."));
+
+	// Check adding types across multiple threads is invariant
+	FParamTypeHandle::ResetCustomTypes();
+	FParamTypeHandle Handles[100];
+	ParallelFor(100, [&Handles](int32 InIndex)
+	{
+		Handles[InIndex] = FParamTypeHandle::GetHandle<UAnimSequence>();
+	});
+
+	for(int32 Index = 1; Index < 100; ++Index)
+	{
+		AddErrorIfFalse(Handles[Index] == Handles[0], FString::Printf(TEXT("Different type (index %d) has been defined when generating a handle across threads (0x%08x vs 0x%08x)."), Index, Handles[Index].Value, Handles[0].Value));
+	}
+
+	return true;
+}
+
+#endif	// WITH_DEV_AUTOMATION_TESTS
\ No newline at end of file
diff --git a/Engine/Plugins/Experimental/AnimNext/Source/AnimNextTestSuite/Private/AnimNextTest.cpp b/Engine/Plugins/Experimental/AnimNext/Source/AnimNextTestSuite/Private/AnimNextTest.cpp
new file mode 100644
--- /dev/null
+++ b/Engine/Plugins/Experimental/AnimNext/Source/AnimNextTestSuite/Private/AnimNextTest.cpp
@@ -0,0 +1,524 @@
+// Copyright Epic Games, Inc. All Rights Reserved.
+
+
+#include "AnimNextTest.h"
+#include "CoreMinimal.h"
+#include "Interface/AnimNextInterface.h"
+#include "Param/Param.h"
+#include "Interface/InterfaceContext.h"
+#include "Param/ParamStorage.h"
+#include "Misc/AutomationTest.h"
+#include "DataRegistry.h"
+#include "ReferencePose.h"
+
+#include <chrono>
+
+namespace 
+{
+
+class FHighResTimer
+{
+public:
+	using FTimePoint = std::chrono::time_point<std::chrono::high_resolution_clock>;
+
+	static inline FTimePoint GetTimeMark()
+	{
+		return std::chrono::high_resolution_clock::now();
+	}
+
+	static inline int64 GetTimeDiffMicroSec(const FTimePoint& RefTime)
+	{
+		return std::chrono::duration_cast<std::chrono::microseconds>(GetTimeMark() - RefTime).count();
+	}
+	static inline int64 GetTimeDiffNanoSec(const FTimePoint& RefTime)
+	{
+		return std::chrono::duration_cast<std::chrono::nanoseconds>(GetTimeMark() - RefTime).count();
+	}
+};
+
+} // end namespace
+
+//----------------------------------------------------------------------------
+
+/*static*/ FName UAnimNextInterface_TestData_Multiply::NameParamA(TEXT("TestDataA"));
+/*static*/ FName UAnimNextInterface_TestData_Multiply::NameParamB(TEXT("TestDataB"));
+/*static*/ FName UAnimNextInterface_TestData_Multiply::NameParamResult(TEXT("NameParamAxB"));
+
+
+bool UAnimNextInterfaceTestDataLiteral::GetDataImpl(const UE::AnimNext::FContext& Context) const
+{
+	Context.SetResult(Value);
+
+	return true;
+}
+
+bool UAnimNextInterface_TestData_Multiply::GetDataImpl(const UE::AnimNext::FContext& Context) const
+{
+	using namespace UE::AnimNext;
+
+	bool bResult = false;
+
+	const FTestData& ParamA = Context.GetParameterChecked<const FTestData>(NameParamA);
+	const FTestData& ParamB = Context.GetParameterChecked<const FTestData>(NameParamB);
+	FTestData& Result = Context.GetResult<FTestData>();
+
+	bResult = true;
+
+	Result.A = ParamA.A * ParamB.A;
+	Result.B = ParamA.B * ParamB.B;
+
+	return bResult;
+}
+
+//----------------------------------------------------------------------------
+
+/*static*/ FName UAnimNextInterface_TestData_Split::InputName_AB(TEXT("Input_AB"));
+/*static*/ FName UAnimNextInterface_TestData_Split::OutputName_A(TEXT("Output_A"));
+/*static*/ FName UAnimNextInterface_TestData_Split::OutputName_B(TEXT("Output_B"));
+
+bool UAnimNextInterface_TestData_Split::GetDataImpl(const UE::AnimNext::FContext& Context) const
+{
+	using namespace UE::AnimNext;
+
+	bool bResult = false;
+
+	const FTestData& Input_AB = Context.GetParameterChecked<const FTestData>(InputName_AB);
+	float& Output_A = Context.GetParameterChecked<float>(OutputName_A);
+	float& Output_B = Context.GetParameterChecked<float>(OutputName_B);
+
+	bResult = true;
+
+	Output_A = Input_AB.A;
+	Output_B = Input_AB.B;
+
+	return bResult;
+}
+
+
+//****************************************************************************
+// AnimNext Interface Test
+//****************************************************************************
+#if WITH_DEV_AUTOMATION_TESTS
+
+IMPLEMENT_SIMPLE_AUTOMATION_TEST(FAnimationAnimNextInterfaceTest_SingleElement, "Animation.AnimNext.Interface.Single", EAutomationTestFlags::EditorContext | EAutomationTestFlags::EngineFilter)
+
+bool FAnimationAnimNextInterfaceTest_SingleElement::RunTest(const FString& InParameters)
+{
+	using namespace UE::AnimNext;
+
+	UAnimNextInterface_TestData_Multiply* AnimNextInterface_TestData_Multiply = Cast<UAnimNextInterface_TestData_Multiply>(UAnimNextInterface_TestData_Multiply::StaticClass()->GetDefaultObject());
+	TScriptInterface<IAnimNextInterface> TestDataMultiply = AnimNextInterface_TestData_Multiply;
+
+	AddErrorIfFalse(TestDataMultiply.GetInterface() != nullptr, "UAnimNextInterface_TestData_Multiply -> Interface is Null.");
+
+	constexpr FTestData TestDataA = { 1.f, 1.f };
+	constexpr FTestData TestDataB = { 1.f, 2.f };
+	FTestData TestDataResult = { 0.f, 0.f };
+
+	FState State;
+	FParamStorage ParamStorage(16, 256, 2);
+	FContext Context(0.f, State, ParamStorage);
+
+	const FContext TestContext = Context.WithParameters({
+		TPairInitializer(UAnimNextInterface_TestData_Multiply::NameParamA, TWrapParam<const FTestData>(TestDataA)),
+		TPairInitializer(UAnimNextInterface_TestData_Multiply::NameParamB, TWrapParam<const FTestData>(TestDataB))
+		});
+
+	Interface::GetDataSafe(TestDataMultiply, TestContext, TestDataResult);
+
+	// --- Check the results from the batched operation ---
+	AddErrorIfFalse(TestDataResult.A == TestDataA.A * TestDataB.A, FString::Printf(TEXT("TestDataResult.A is : [%.04f] Expected : [%.04f]"), TestDataResult.A, TestDataA.A * TestDataB.A));
+	AddErrorIfFalse(TestDataResult.B == TestDataA.B * TestDataB.B, FString::Printf(TEXT("TestDataResult.B is : [%.04f] Expected : [%.04f]"), TestDataResult.B, TestDataA.B * TestDataB.B));
+
+	// Final GC to make sure everything is cleaned up
+	CollectGarbage(GARBAGE_COLLECTION_KEEPFLAGS);
+
+	return true;
+}
+
+IMPLEMENT_SIMPLE_AUTOMATION_TEST(FAnimationAnimNextInterfaceTest_Chained, "Animation.AnimNext.Interface.Chained", EAutomationTestFlags::EditorContext | EAutomationTestFlags::EngineFilter)
+
+bool FAnimationAnimNextInterfaceTest_Chained::RunTest(const FString& InParameters)
+{
+	using namespace UE::AnimNext;
+
+	FState RootState;
+	FParamStorage ParamStorage(16, 2048, 8);
+	const FContext RootContext(0.f, RootState, ParamStorage);
+	
+	// Simulate we pass a const context and have to return data
+	TParamValue<FTestData> RootContextResultParam;
+	const FContext InContext = RootContext.WithResult(RootContextResultParam);
+
+	FTestData LiteralResult;
+	TParamValue<FTestData> ParamLiteralResult;
+	//TContextStorageParam<FTestData> ParamLiteralResult;
+
+	// Retrieve the literal value
+	{
+		TScriptInterface<IAnimNextInterface> TestDataLiteral = Cast<UAnimNextInterfaceTestDataLiteral>(UAnimNextInterfaceTestDataLiteral::StaticClass()->GetDefaultObject());
+		AddErrorIfFalse(TestDataLiteral.GetInterface() != nullptr, "UAnimNextInterfaceTestDataLiteral -> Interface is Null.");
+
+		const FContext LiteralContext = InContext.WithResult(ParamLiteralResult);
+
+		Interface::GetDataSafe(TestDataLiteral, LiteralContext);
+
+		AddErrorIfFalse(ParamLiteralResult->A == 1, FString::Printf(TEXT("LiteralResult.A is : [%.04f] Expected : [%.04f]"), ParamLiteralResult->A, 1));
+		AddErrorIfFalse(ParamLiteralResult->B == 1, FString::Printf(TEXT("LiteralResult.B is : [%.04f] Expected : [%.04f]"), ParamLiteralResult->B, 1));
+	}
+
+	const FTestData TestDataB = { 2.f, 2.f };
+
+	// now chain an operation using the literal result as a parameter
+	{
+		UAnimNextInterface_TestData_Multiply* AnimNextInterface_TestData_Multiply = Cast<UAnimNextInterface_TestData_Multiply>(UAnimNextInterface_TestData_Multiply::StaticClass()->GetDefaultObject());
+		TScriptInterface<IAnimNextInterface> TestDataMultiply = AnimNextInterface_TestData_Multiply;
+		AddErrorIfFalse(TestDataMultiply.GetInterface() != nullptr, "UAnimNextInterface_TestData_Multiply -> Interface is Null.");
+
+		const FContext ChainContext = InContext.WithResultAndParameters(InContext.GetResultParam(), {
+			TPairInitializer(UAnimNextInterface_TestData_Multiply::NameParamA, ParamLiteralResult),						// chain the result of the literal call as the input
+			TPairInitializer(UAnimNextInterface_TestData_Multiply::NameParamB, TWrapParam<const FTestData>(TestDataB))	// set the second operand
+			});
+
+		Interface::GetDataSafe(TestDataMultiply, ChainContext);
+	}
+
+	// --- Check the results from the chained operation on the root component ---
+	AddErrorIfFalse(RootContextResultParam->A == ParamLiteralResult->A * TestDataB.A, FString::Printf(TEXT("TestDataResult.A is : [%.04f] Expected : [%.04f]"), RootContextResultParam->A, ParamLiteralResult->A * TestDataB.A));
+	AddErrorIfFalse(RootContextResultParam->B == ParamLiteralResult->B * TestDataB.B, FString::Printf(TEXT("TestDataResult.B is : [%.04f] Expected : [%.04f]"), RootContextResultParam->B, ParamLiteralResult->B * TestDataB.B));
+
+	// now chain an operation using the literal result as a parameter but copying the input parameters on the Context Parameter storage
+	{
+		UAnimNextInterface_TestData_Multiply* AnimNextInterface_TestData_Multiply = Cast<UAnimNextInterface_TestData_Multiply>(UAnimNextInterface_TestData_Multiply::StaticClass()->GetDefaultObject());
+		TScriptInterface<IAnimNextInterface> TestDataMultiply = AnimNextInterface_TestData_Multiply;
+		AddErrorIfFalse(TestDataMultiply.GetInterface() != nullptr, "UAnimNextInterface_TestData_Multiply -> Interface is Null.");
+
+		const FContext ChainContext = InContext.WithResultAndParameters(InContext.GetResultParam(), {
+			TPairInitializer(UAnimNextInterface_TestData_Multiply::NameParamA, TContextStorageParam<const FTestData>(ParamLiteralResult)),	// chain the result of the literal call as the input
+			TPairInitializer(UAnimNextInterface_TestData_Multiply::NameParamB, TContextStorageParam<const FTestData>(TestDataB))			// set the second operand
+			});
+
+		Interface::GetDataSafe(TestDataMultiply, ChainContext);
+	}
+
+	// --- Check the results from the chained operation on the root component ---
+	AddErrorIfFalse(RootContextResultParam->A == ParamLiteralResult->A * TestDataB.A, FString::Printf(TEXT("TestDataResult.A is : [%.04f] Expected : [%.04f]"), RootContextResultParam->A, ParamLiteralResult->A * TestDataB.A));
+	AddErrorIfFalse(RootContextResultParam->B == ParamLiteralResult->B * TestDataB.B, FString::Printf(TEXT("TestDataResult.B is : [%.04f] Expected : [%.04f]"), RootContextResultParam->B, ParamLiteralResult->B * TestDataB.B));
+
+	// Final GC to make sure everything is cleaned up
+	CollectGarbage(GARBAGE_COLLECTION_KEEPFLAGS);
+
+	return true;
+}
+
+IMPLEMENT_SIMPLE_AUTOMATION_TEST(FAnimationAnimNextInterfaceTest_AltParamStorage, "Animation.AnimNext.Interface.AltParamStorage", EAutomationTestFlags::EditorContext | EAutomationTestFlags::EngineFilter)
+
+bool FAnimationAnimNextInterfaceTest_AltParamStorage::RunTest(const FString& InParameters)
+{
+	using namespace UE::AnimNext;
+
+	FState RootState;
+	FParamStorage ParamStorage(16, 2048, 8);
+	const FContext RootContext(0.f, RootState, ParamStorage);
+
+	UAnimNextInterface_TestData_Multiply* AnimNextInterface_TestData_Multiply = Cast<UAnimNextInterface_TestData_Multiply>(UAnimNextInterface_TestData_Multiply::StaticClass()->GetDefaultObject());
+	TScriptInterface<IAnimNextInterface> TestDataMultiply = AnimNextInterface_TestData_Multiply;
+
+	AddErrorIfFalse(TestDataMultiply.GetInterface() != nullptr, "UAnimNextInterface_TestData_Multiply -> Interface is Null.");
+
+	constexpr FTestData TestDataA = { 1.f, 1.f };
+	constexpr FTestData TestDataB = { 1.f, 2.f };
+	FTestData TestDataResult = { 0.f, 0.f };
+
+	FContext TestContext = RootContext.CreateSubContext();
+
+	TestContext.AddValue(FContext::EParamType::Input, UAnimNextInterface_TestData_Multiply::NameParamA, TestDataA);
+	TestContext.AddValue(FContext::EParamType::Input, UAnimNextInterface_TestData_Multiply::NameParamB, TestDataB);
+	FParamHandle HOutput = TestContext.AddReference(FContext::EParamType::Output, UAnimNextInterface_TestData_Multiply::NameParamResult, TestDataResult);
+
+	// for compatibility reasons, as the current code expects to have a return value set
+	TestContext.SetHParamAsResult(HOutput);  
+
+	const bool bOk = TestDataMultiply->GetData(TestContext);
+
+	AddErrorIfFalse(bOk, "TestDataMultiply -> GetData returned FAILURE.");
+
+	// --- Check the results from the operation ---
+	AddErrorIfFalse(TestDataResult.A == TestDataA.A * TestDataB.A, FString::Printf(TEXT("TestDataResult.A is : [%.04f] Expected : [%.04f]"), TestDataResult.A, TestDataA.A * TestDataB.A));
+	AddErrorIfFalse(TestDataResult.B == TestDataA.B * TestDataB.B, FString::Printf(TEXT("TestDataResult.B is : [%.04f] Expected : [%.04f]"), TestDataResult.B, TestDataA.B * TestDataB.B));
+
+	// Final GC to make sure everything is cleaned up
+	CollectGarbage(GARBAGE_COLLECTION_KEEPFLAGS);
+
+	return true;
+}
+
+
+IMPLEMENT_SIMPLE_AUTOMATION_TEST(FAnimationAnimNextInterfaceTest_AltParamStorageMultipleOutputs, "Animation.AnimNext.Interface.AltParamStorageMultipleOutputs", EAutomationTestFlags::EditorContext | EAutomationTestFlags::EngineFilter)
+
+bool FAnimationAnimNextInterfaceTest_AltParamStorageMultipleOutputs::RunTest(const FString& InParameters)
+{
+	using namespace UE::AnimNext;
+
+	FState RootState;
+	FParamStorage ParamStorage(16, 2048, 8);
+	const FContext RootContext(0.f, RootState, ParamStorage);
+
+	TScriptInterface<IAnimNextInterface> TestDataSplit(Cast<UAnimNextInterface_TestData_Split>(UAnimNextInterface_TestData_Split::StaticClass()->GetDefaultObject()));
+
+	AddErrorIfFalse(TestDataSplit.GetInterface() != nullptr, "UAnimNextInterface_TestData_Split -> Interface is Null.");
+
+	constexpr FTestData Input_AB = { 1.f, 2.f };
+	float Output_A = 0.f;
+	float Output_B = 0.f;
+
+	FContext TestContext = RootContext.CreateSubContext();
+
+	// Adding an input as a value (will be copied and the copy will become mutable)
+	FParamHandle RetValCompatiblility = TestContext.AddValue(FContext::EParamType::Input, UAnimNextInterface_TestData_Split::InputName_AB, Input_AB);
+	// Adding first output as a reference
+	TestContext.AddReference(FContext::EParamType::Output, UAnimNextInterface_TestData_Split::OutputName_A, Output_A);
+	// Adding second output as a reference
+	TestContext.AddReference(FContext::EParamType::Output, UAnimNextInterface_TestData_Split::OutputName_B, Output_B);
+
+	// Not used, but added for compatibility reasons, as the current code expects to have a return value set of the interface type
+	TestContext.SetHParamAsResult(RetValCompatiblility);
+
+	const bool bOk = TestDataSplit->GetData(TestContext);
+
+	AddErrorIfFalse(bOk, "TestDataSplit -> GetData returned FAILURE.");
+
+	// --- Check the results from the operation ---
+	AddErrorIfFalse(Output_A == Input_AB.A, FString::Printf(TEXT("Output_A is : [%.04f] Expected : [%.04f]"), Output_A, Input_AB.A));
+	AddErrorIfFalse(Output_B == Input_AB.B, FString::Printf(TEXT("Output_B is : [%.04f] Expected : [%.04f]"), Output_B, Input_AB.B));
+
+	// Final GC to make sure everything is cleaned up
+	CollectGarbage(GARBAGE_COLLECTION_KEEPFLAGS);
+
+	return true;
+}
+
+IMPLEMENT_SIMPLE_AUTOMATION_TEST(FAnimationAnimNextInterfaceTest_AltParamStorageChained, "Animation.AnimNext.Interface.AltParamStorageChained", EAutomationTestFlags::EditorContext | EAutomationTestFlags::EngineFilter)
+
+bool FAnimationAnimNextInterfaceTest_AltParamStorageChained::RunTest(const FString& InParameters)
+{
+	using namespace UE::AnimNext;
+
+	FState RootState;
+	FParamStorage ParamStorage(16, 2048, 8);
+	const FContext RootContext(0.f, RootState, ParamStorage);
+
+	UAnimNextInterface_TestData_Multiply* AnimNextInterface_TestData_Multiply = Cast<UAnimNextInterface_TestData_Multiply>(UAnimNextInterface_TestData_Multiply::StaticClass()->GetDefaultObject());
+	TScriptInterface<IAnimNextInterface> TestDataMultiply = AnimNextInterface_TestData_Multiply;
+
+	AddErrorIfFalse(TestDataMultiply.GetInterface() != nullptr, "UAnimNextInterface_TestData_Multiply -> Interface is Null.");
+
+	// --- Initial Test : Multiply two TestData and store the value at Result ---
+
+	FParamHandle TestOutput_AxB_Param;
+
+	{
+		constexpr FTestData TestDataA = { 1.f, 1.f };
+		constexpr FTestData TestDataB = { 1.f, 2.f };
+
+		FContext TestContext = RootContext.CreateSubContext();
+
+		// Adding a const value. Note that adding a value will convert const into mutable, as the copy will not be const anymore
+		TestContext.AddValue(FContext::EParamType::Input, UAnimNextInterface_TestData_Multiply::NameParamA, TestDataA);
+		// Adding a reference to const. This param will keep the const
+		TestContext.AddReference(FContext::EParamType::Input, UAnimNextInterface_TestData_Multiply::NameParamB, TestDataB);
+		// Adding a value using a temp taht will be copied into the value
+		TestOutput_AxB_Param = TestContext.AddValue(FContext::EParamType::Output, UAnimNextInterface_TestData_Multiply::NameParamResult, FTestData{ 0.f, 0.f });
+
+		// for compatibility reasons, as the current code expects to have a return value set
+		TestContext.SetHParamAsResult(TestOutput_AxB_Param);  
+
+		const bool bOk = TestDataMultiply->GetData(TestContext);
+
+		AddErrorIfFalse(bOk, "TestDataMultiply -> GetData returned FAILURE.");
+
+		// --- Check the results from the operation ---
+		const FTestData& TestDataResult = TestContext.GetParameterAs<const FTestData>(TestOutput_AxB_Param);
+		AddErrorIfFalse(TestDataResult.A == TestDataA.A * TestDataB.A, FString::Printf(TEXT("TestDataResult.A is : [%.04f] Expected : [%.04f]"), TestDataResult.A, TestDataA.A * TestDataB.A));
+		AddErrorIfFalse(TestDataResult.B == TestDataA.B * TestDataB.B, FString::Printf(TEXT("TestDataResult.B is : [%.04f] Expected : [%.04f]"), TestDataResult.B, TestDataA.B * TestDataB.B));
+	}
+
+	// --- Chained Test : Use the previous call Result as an input ---
+
+	{
+		FContext TestContext = RootContext.CreateSubContext();
+
+		constexpr FTestData TestDataC = { 2.f, 2.f };
+
+		// Add the result of the previous call as an input for this operation
+		TestContext.AddValue(FContext::EParamType::Input, UAnimNextInterface_TestData_Multiply::NameParamA, TestOutput_AxB_Param);
+		// Add a reference to const data
+		TestContext.AddReference(FContext::EParamType::Input, UAnimNextInterface_TestData_Multiply::NameParamB, TestDataC);
+		// Add a value
+		const FParamHandle ChainedResultParam = TestContext.AddValue(FContext::EParamType::Output, UAnimNextInterface_TestData_Multiply::NameParamResult, FTestData{ 0.f, 0.f });
+
+		// for compatibility reasons, as the current code expects to have a return value set
+		TestContext.SetHParamAsResult(ChainedResultParam);
+
+		const bool bOk = TestDataMultiply->GetData(TestContext);
+
+		AddErrorIfFalse(bOk, "TestDataMultiply -> GetData returned FAILURE.");
+
+		// --- Check the results from the operation ---
+		const FTestData& TestOutput_AxB = TestContext.GetParameterAs<const FTestData>(TestOutput_AxB_Param);
+		const FTestData& ChainedResult = TestContext.GetParameterAs<const FTestData>(ChainedResultParam);
+		
+		AddErrorIfFalse(ChainedResult.A == TestOutput_AxB.A * TestDataC.A, FString::Printf(TEXT("Chained TestDataResult.A is : [%.04f] Expected : [%.04f]"), ChainedResult.A, TestOutput_AxB.A * TestDataC.A));
+		AddErrorIfFalse(ChainedResult.B == TestOutput_AxB.B * TestDataC.B, FString::Printf(TEXT("Chained TestDataResult.B is : [%.04f] Expected : [%.04f]"), ChainedResult.B, TestOutput_AxB.B * TestDataC.B));
+	}
+
+
+	// Final GC to make sure everything is cleaned up
+	CollectGarbage(GARBAGE_COLLECTION_KEEPFLAGS);
+
+	return true;
+}
+
+IMPLEMENT_SIMPLE_AUTOMATION_TEST(FAnimationAnimNextInterfaceTest_AnimationDataRegistry, "Animation.AnimNext.AnimationDataRegistry.TransformsTest", EAutomationTestFlags::EditorContext | EAutomationTestFlags::EngineFilter)
+
+bool FAnimationAnimNextInterfaceTest_AnimationDataRegistry::RunTest(const FString& InParameters)
+{
+	using namespace UE::AnimNext;
+
+	FDataRegistry* AnimDataReg = FDataRegistry::Get();
+
+
+	//constexpr int32 BASIC_TYPE_ALLOC_BLOCK = 1000;
+	//AnimDataReg.RegisterDataType<uint8>(EAnimationDataType::RawMemory, BASIC_TYPE_ALLOC_BLOCK);
+	//AnimDataReg.RegisterDataType<FVector>(EAnimationDataType::Translation, BASIC_TYPE_ALLOC_BLOCK);
+	//AnimDataReg.RegisterDataType<FVector>(EAnimationDataType::Scale, BASIC_TYPE_ALLOC_BLOCK);
+	//AnimDataReg.RegisterDataType<FQuat>(EAnimationDataType::Rotation, BASIC_TYPE_ALLOC_BLOCK);
+	//AnimDataReg.RegisterDataType<FTransform>(EAnimationDataType::Transform, BASIC_TYPE_ALLOC_BLOCK);
+	//AnimDataReg.RegisterDataType<FTransform>(EAnimationDataType::Displacement, BASIC_TYPE_ALLOC_BLOCK);
+
+	//AnimDataReg.RegisterDataType<FAnimTransformArray>(EAnimationDataType::TransformArray, BASIC_TYPE_ALLOC_BLOCK);
+	//AnimDataReg.RegisterDataType<FAnimTransformArraySoA>(EAnimationDataType::TransformArraySoA, BASIC_TYPE_ALLOC_BLOCK);
+
+	constexpr int32 NumElements = 1;
+	constexpr int32 NumPasses = 100;
+	constexpr int32 NumTransforms = 100;
+
+	TArray<int64> TimeDiffAoS;
+	TArray<int64> TimeDiffSoA;
+	TimeDiffAoS.SetNumZeroed(NumPasses);
+	TimeDiffSoA.SetNumZeroed(NumPasses);
+
+	const double BlendFactor = 0.5f;
+	const FVector Half = FVector(0.5, 0.5, 0.5);
+	const FQuat Rot90 = FQuat::MakeFromEuler(FVector(90.0, 0.0, 0.0));
+
+	const FTransform TargetBlendValue(Rot90, FVector::OneVector, Half);
+	FTransform ReferenceResult;
+	ReferenceResult.Blend(FTransform::Identity, TargetBlendValue, BlendFactor);
+
+	// --- AoS Test ---
+
+	FTransformArrayAoSHeap ResultAoS(NumTransforms);
+	for (int i = 0; i < NumPasses; ++i)
+	{
+		// *** TransformArrayAoS Test *** 
+		{
+			FTransformArrayAoSHeap TransformsA(NumTransforms);
+			FTransformArrayAoSHeap TransformsB(NumTransforms);
+
+			for (auto& Transform : TransformsB.GetTransforms())
+			{
+				Transform.SetTranslation(FVector::OneVector);
+				Transform.SetRotation(Rot90);
+				Transform.SetScale3D(Half);
+			}
+
+			FHighResTimer::FTimePoint StartTime = FHighResTimer::GetTimeMark();
+
+			ResultAoS.Blend(TransformsA, TransformsB, BlendFactor);
+
+			TimeDiffAoS[i] = FHighResTimer::GetTimeDiffNanoSec(StartTime);
+		}
+	}
+	for (auto& Transform : ResultAoS.Transforms)
+	{
+		AddErrorIfFalse(Transform.GetTranslation() == ReferenceResult.GetTranslation(), "AoS Translation Blend error.");
+		AddErrorIfFalse(Transform.GetRotation() == ReferenceResult.GetRotation(), "AoS Rotation Blend error.");
+		AddErrorIfFalse(Transform.GetScale3D() == ReferenceResult.GetScale3D(), "AoS Scale Blend error.");
+	}
+
+	// --- SoA Test --- 
+	FTransformArraySoAHeap ResultSoA(NumTransforms);
+	for (int i = 0; i < NumPasses; ++i)
+	{
+		//FAnimTransformArray* TransoformArray = AnimDataReg->AllocateData<FAnimTransformArray>(NumElements, NumTransforms);
+
+		// *** TransformArraySoA Test *** 
+		{
+			FTransformArraySoAHeap TransformsA(NumTransforms);
+			FTransformArraySoAHeap TransformsB(NumTransforms);
+
+			for (auto& Translation : TransformsB.Translations)
+			{
+				Translation = FVector::OneVector;
+			}
+			for (auto& Rotation : TransformsB.Rotations)
+			{
+				Rotation = Rot90;
+			}
+			for (auto& Scale3D : TransformsB.Scales3D)
+			{
+				Scale3D = Half;
+			}
+
+			FHighResTimer::FTimePoint StartTime = FHighResTimer::GetTimeMark();
+
+			ResultSoA.Blend(TransformsA, TransformsB, BlendFactor);
+
+			TimeDiffSoA[i] = FHighResTimer::GetTimeDiffNanoSec(StartTime);
+		}
+	}
+	for (const auto& Translation : ResultSoA.Translations)
+	{
+		AddErrorIfFalse(Translation == ReferenceResult.GetTranslation(), "SoA Translation Blend error.");
+	}
+	for (const auto& Rotation : ResultSoA.Rotations)
+	{
+		AddErrorIfFalse(Rotation == ReferenceResult.GetRotation(), "SoA Rotation Blend error.");
+	}
+	for (const auto& Scale3D : ResultSoA.Scales3D)
+	{
+		AddErrorIfFalse(Scale3D == ReferenceResult.GetScale3D(), "SoA Scale Blend error.");
+	}
+
+	// --- Time Averaging ---
+
+	TimeDiffAoS.Sort();
+	TimeDiffSoA.Sort();
+
+	const int32 NumToAverage = FMath::Min(3, TimeDiffAoS.Num());
+	int64 AverageAoS = 0LL;
+	for (int i = 0; i < NumToAverage; ++i)
+	{
+		AverageAoS += TimeDiffAoS[i];
+	}
+	AverageAoS /= NumToAverage;
+
+	int64 AverageSoA = 0LL;
+	for (int i = 0; i < NumToAverage; ++i)
+	{
+		AverageSoA += TimeDiffSoA[i];
+	}
+	AverageSoA /= NumToAverage;
+
+
+	AddInfo(FString::Printf(TEXT("Average AoS Duration: %d nanosecs."), AverageAoS));
+	AddInfo(FString::Printf(TEXT("Average SoA Duration: %d nanosecs."), AverageSoA));
+
+	// Final GC to make sure everything is cleaned up
+	CollectGarbage(GARBAGE_COLLECTION_KEEPFLAGS);
+
+	return true;
+}
+
+#endif // WITH_DEV_AUTOMATION_TESTS
\ No newline at end of file
diff --git a/Engine/Plugins/Experimental/AnimNext/Source/AnimNextTestSuite/Private/AnimNextTest.h b/Engine/Plugins/Experimental/AnimNext/Source/AnimNextTestSuite/Private/AnimNextTest.h
new file mode 100644
--- /dev/null
+++ b/Engine/Plugins/Experimental/AnimNext/Source/AnimNextTestSuite/Private/AnimNextTest.h
@@ -0,0 +1,75 @@
+// Copyright Epic Games, Inc. All Rights Reserved.
+
+#pragma once
+
+#include "Interface/InterfaceTypes.h"
+#include "Interface/IAnimNextInterface.h"
+
+#include "AnimNextTest.generated.h"
+
+//// --- Raw type ---
+USTRUCT()
+struct FTestData
+{
+	GENERATED_BODY()
+
+	float A = 0.f;
+	float B = 0.f;
+};
+
+// -- Defines the TestData as an interface that will be used by any variant operating on TestData --
+UCLASS()
+class UAnimNextInterface_TestData : public UObject, public IAnimNextInterface
+{
+	GENERATED_BODY()
+
+	ANIM_NEXT_INTERFACE_RETURN_TYPE(FTestData)
+
+	virtual bool GetDataImpl(const UE::AnimNext::FContext& Context) const override
+	{
+		checkf(false, TEXT("UAnimNextInterfaceFloat::GetDataImpl must be overridden"));
+		return false;
+	}
+};
+
+// --- Constant user defined value ---
+UCLASS()
+class UAnimNextInterfaceTestDataLiteral : public UAnimNextInterface_TestData
+{
+	GENERATED_BODY()
+
+	// IAnimNextInterface interface
+	virtual bool GetDataImpl(const UE::AnimNext::FContext& Context) const final override;
+
+	FTestData Value = {1.f, 1.f};
+};
+
+// --- TestData multiply operator  ---
+UCLASS()
+class UAnimNextInterface_TestData_Multiply : public UAnimNextInterface_TestData
+{
+	GENERATED_BODY()
+
+	// IAnimNextInterface interface
+	virtual bool GetDataImpl(const UE::AnimNext::FContext& Context) const final override;
+
+public:
+	static FName NameParamA;
+	static FName NameParamB;
+	static FName NameParamResult;
+};
+
+// --- TestData split operator : creates two outputs from a single input (multi output test) ---
+UCLASS()
+class UAnimNextInterface_TestData_Split : public UAnimNextInterface_TestData
+{
+	GENERATED_BODY()
+
+	// IAnimNextInterface interface
+	virtual bool GetDataImpl(const UE::AnimNext::FContext& Context) const final override;
+
+public:
+	static FName InputName_AB;
+	static FName OutputName_A;
+	static FName OutputName_B;
+};
