CIS fixes: Static analysis (IncludeTool) and Copyright notices

diff --git a/Engine/Plugins/Experimental/AnimNext/Source/AnimNextTestSuite/Private/AnimNextTest.cpp b/Engine/Plugins/Experimental/AnimNext/Source/AnimNextTestSuite/Private/AnimNextTest.cpp
--- a/Engine/Plugins/Experimental/AnimNext/Source/AnimNextTestSuite/Private/AnimNextTest.cpp
+++ b/Engine/Plugins/Experimental/AnimNext/Source/AnimNextTestSuite/Private/AnimNextTest.cpp
@@ -58,9 +58,9 @@ bool UAnimNextInterface_TestData_Multiply::GetDataImpl(const UE::AnimNext::FCont
 
 	bool bResult = false;
 
-	const FTestData& ParamA = Context.GetParameterChecked<const FTestData>(NameParamA);
-	const FTestData& ParamB = Context.GetParameterChecked<const FTestData>(NameParamB);
-	FTestData& Result = Context.GetResult<FTestData>();
+	const FAnimNextTestData& ParamA = Context.GetParameterChecked<const FAnimNextTestData>(NameParamA);
+	const FAnimNextTestData& ParamB = Context.GetParameterChecked<const FAnimNextTestData>(NameParamB);
+	FAnimNextTestData& Result = Context.GetResult<FAnimNextTestData>();
 
 	bResult = true;
 
@@ -82,7 +82,7 @@ bool UAnimNextInterface_TestData_Split::GetDataImpl(const UE::AnimNext::FContext
 
 	bool bResult = false;
 
-	const FTestData& Input_AB = Context.GetParameterChecked<const FTestData>(InputName_AB);
+	const FAnimNextTestData& Input_AB = Context.GetParameterChecked<const FAnimNextTestData>(InputName_AB);
 	float& Output_A = Context.GetParameterChecked<float>(OutputName_A);
 	float& Output_B = Context.GetParameterChecked<float>(OutputName_B);
 
@@ -111,17 +111,17 @@ bool FAnimationAnimNextInterfaceTest_SingleElement::RunTest(const FString& InPar
 
 	AddErrorIfFalse(TestDataMultiply.GetInterface() != nullptr, "UAnimNextInterface_TestData_Multiply -> Interface is Null.");
 
-	constexpr FTestData TestDataA = { 1.f, 1.f };
-	constexpr FTestData TestDataB = { 1.f, 2.f };
-	FTestData TestDataResult = { 0.f, 0.f };
+	constexpr FAnimNextTestData TestDataA = { 1.f, 1.f };
+	constexpr FAnimNextTestData TestDataB = { 1.f, 2.f };
+	FAnimNextTestData TestDataResult = { 0.f, 0.f };
 
 	FState State;
 	FParamStorage ParamStorage(16, 256, 2);
 	FContext Context(0.f, State, ParamStorage);
 
 	const FContext TestContext = Context.WithParameters({
-		TPairInitializer(UAnimNextInterface_TestData_Multiply::NameParamA, TWrapParam<const FTestData>(TestDataA)),
-		TPairInitializer(UAnimNextInterface_TestData_Multiply::NameParamB, TWrapParam<const FTestData>(TestDataB))
+		TPairInitializer(UAnimNextInterface_TestData_Multiply::NameParamA, TWrapParam<const FAnimNextTestData>(TestDataA)),
+		TPairInitializer(UAnimNextInterface_TestData_Multiply::NameParamB, TWrapParam<const FAnimNextTestData>(TestDataB))
 		});
 
 	Interface::GetDataSafe(TestDataMultiply, TestContext, TestDataResult);
@@ -147,11 +147,11 @@ bool FAnimationAnimNextInterfaceTest_Chained::RunTest(const FString& InParameter
 	const FContext RootContext(0.f, RootState, ParamStorage);
 	
 	// Simulate we pass a const context and have to return data
-	TParamValue<FTestData> RootContextResultParam;
+	TParamValue<FAnimNextTestData> RootContextResultParam;
 	const FContext InContext = RootContext.WithResult(RootContextResultParam);
 
-	FTestData LiteralResult;
-	TParamValue<FTestData> ParamLiteralResult;
+	FAnimNextTestData LiteralResult;
+	TParamValue<FAnimNextTestData> ParamLiteralResult;
 	//TContextStorageParam<FTestData> ParamLiteralResult;
 
 	// Retrieve the literal value
@@ -167,7 +167,7 @@ bool FAnimationAnimNextInterfaceTest_Chained::RunTest(const FString& InParameter
 		AddErrorIfFalse(ParamLiteralResult->B == 1, FString::Printf(TEXT("LiteralResult.B is : [%.04f] Expected : [%.04f]"), ParamLiteralResult->B, 1));
 	}
 
-	const FTestData TestDataB = { 2.f, 2.f };
+	const FAnimNextTestData TestDataB = { 2.f, 2.f };
 
 	// now chain an operation using the literal result as a parameter
 	{
@@ -177,7 +177,7 @@ bool FAnimationAnimNextInterfaceTest_Chained::RunTest(const FString& InParameter
 
 		const FContext ChainContext = InContext.WithResultAndParameters(InContext.GetResultParam(), {
 			TPairInitializer(UAnimNextInterface_TestData_Multiply::NameParamA, ParamLiteralResult),						// chain the result of the literal call as the input
-			TPairInitializer(UAnimNextInterface_TestData_Multiply::NameParamB, TWrapParam<const FTestData>(TestDataB))	// set the second operand
+			TPairInitializer(UAnimNextInterface_TestData_Multiply::NameParamB, TWrapParam<const FAnimNextTestData>(TestDataB))	// set the second operand
 			});
 
 		Interface::GetDataSafe(TestDataMultiply, ChainContext);
@@ -194,8 +194,8 @@ bool FAnimationAnimNextInterfaceTest_Chained::RunTest(const FString& InParameter
 		AddErrorIfFalse(TestDataMultiply.GetInterface() != nullptr, "UAnimNextInterface_TestData_Multiply -> Interface is Null.");
 
 		const FContext ChainContext = InContext.WithResultAndParameters(InContext.GetResultParam(), {
-			TPairInitializer(UAnimNextInterface_TestData_Multiply::NameParamA, TContextStorageParam<const FTestData>(ParamLiteralResult)),	// chain the result of the literal call as the input
-			TPairInitializer(UAnimNextInterface_TestData_Multiply::NameParamB, TContextStorageParam<const FTestData>(TestDataB))			// set the second operand
+			TPairInitializer(UAnimNextInterface_TestData_Multiply::NameParamA, TContextStorageParam<const FAnimNextTestData>(ParamLiteralResult)),	// chain the result of the literal call as the input
+			TPairInitializer(UAnimNextInterface_TestData_Multiply::NameParamB, TContextStorageParam<const FAnimNextTestData>(TestDataB))			// set the second operand
 			});
 
 		Interface::GetDataSafe(TestDataMultiply, ChainContext);
@@ -226,9 +226,9 @@ bool FAnimationAnimNextInterfaceTest_AltParamStorage::RunTest(const FString& InP
 
 	AddErrorIfFalse(TestDataMultiply.GetInterface() != nullptr, "UAnimNextInterface_TestData_Multiply -> Interface is Null.");
 
-	constexpr FTestData TestDataA = { 1.f, 1.f };
-	constexpr FTestData TestDataB = { 1.f, 2.f };
-	FTestData TestDataResult = { 0.f, 0.f };
+	constexpr FAnimNextTestData TestDataA = { 1.f, 1.f };
+	constexpr FAnimNextTestData TestDataB = { 1.f, 2.f };
+	FAnimNextTestData TestDataResult = { 0.f, 0.f };
 
 	FContext TestContext = RootContext.CreateSubContext();
 
@@ -268,7 +268,7 @@ bool FAnimationAnimNextInterfaceTest_AltParamStorageMultipleOutputs::RunTest(con
 
 	AddErrorIfFalse(TestDataSplit.GetInterface() != nullptr, "UAnimNextInterface_TestData_Split -> Interface is Null.");
 
-	constexpr FTestData Input_AB = { 1.f, 2.f };
+	constexpr FAnimNextTestData Input_AB = { 1.f, 2.f };
 	float Output_A = 0.f;
 	float Output_B = 0.f;
 
@@ -318,8 +318,8 @@ bool FAnimationAnimNextInterfaceTest_AltParamStorageChained::RunTest(const FStri
 	FParamHandle TestOutput_AxB_Param;
 
 	{
-		constexpr FTestData TestDataA = { 1.f, 1.f };
-		constexpr FTestData TestDataB = { 1.f, 2.f };
+		constexpr FAnimNextTestData TestDataA = { 1.f, 1.f };
+		constexpr FAnimNextTestData TestDataB = { 1.f, 2.f };
 
 		FContext TestContext = RootContext.CreateSubContext();
 
@@ -328,7 +328,7 @@ bool FAnimationAnimNextInterfaceTest_AltParamStorageChained::RunTest(const FStri
 		// Adding a reference to const. This param will keep the const
 		TestContext.AddReference(FContext::EParamType::Input, UAnimNextInterface_TestData_Multiply::NameParamB, TestDataB);
 		// Adding a value using a temp taht will be copied into the value
-		TestOutput_AxB_Param = TestContext.AddValue(FContext::EParamType::Output, UAnimNextInterface_TestData_Multiply::NameParamResult, FTestData{ 0.f, 0.f });
+		TestOutput_AxB_Param = TestContext.AddValue(FContext::EParamType::Output, UAnimNextInterface_TestData_Multiply::NameParamResult, FAnimNextTestData{ 0.f, 0.f });
 
 		// for compatibility reasons, as the current code expects to have a return value set
 		TestContext.SetHParamAsResult(TestOutput_AxB_Param);  
@@ -338,7 +338,7 @@ bool FAnimationAnimNextInterfaceTest_AltParamStorageChained::RunTest(const FStri
 		AddErrorIfFalse(bOk, "TestDataMultiply -> GetData returned FAILURE.");
 
 		// --- Check the results from the operation ---
-		const FTestData& TestDataResult = TestContext.GetParameterAs<const FTestData>(TestOutput_AxB_Param);
+		const FAnimNextTestData& TestDataResult = TestContext.GetParameterAs<const FAnimNextTestData>(TestOutput_AxB_Param);
 		AddErrorIfFalse(TestDataResult.A == TestDataA.A * TestDataB.A, FString::Printf(TEXT("TestDataResult.A is : [%.04f] Expected : [%.04f]"), TestDataResult.A, TestDataA.A * TestDataB.A));
 		AddErrorIfFalse(TestDataResult.B == TestDataA.B * TestDataB.B, FString::Printf(TEXT("TestDataResult.B is : [%.04f] Expected : [%.04f]"), TestDataResult.B, TestDataA.B * TestDataB.B));
 	}
@@ -348,14 +348,14 @@ bool FAnimationAnimNextInterfaceTest_AltParamStorageChained::RunTest(const FStri
 	{
 		FContext TestContext = RootContext.CreateSubContext();
 
-		constexpr FTestData TestDataC = { 2.f, 2.f };
+		constexpr FAnimNextTestData TestDataC = { 2.f, 2.f };
 
 		// Add the result of the previous call as an input for this operation
 		TestContext.AddValue(FContext::EParamType::Input, UAnimNextInterface_TestData_Multiply::NameParamA, TestOutput_AxB_Param);
 		// Add a reference to const data
 		TestContext.AddReference(FContext::EParamType::Input, UAnimNextInterface_TestData_Multiply::NameParamB, TestDataC);
 		// Add a value
-		const FParamHandle ChainedResultParam = TestContext.AddValue(FContext::EParamType::Output, UAnimNextInterface_TestData_Multiply::NameParamResult, FTestData{ 0.f, 0.f });
+		const FParamHandle ChainedResultParam = TestContext.AddValue(FContext::EParamType::Output, UAnimNextInterface_TestData_Multiply::NameParamResult, FAnimNextTestData{ 0.f, 0.f });
 
 		// for compatibility reasons, as the current code expects to have a return value set
 		TestContext.SetHParamAsResult(ChainedResultParam);
@@ -365,8 +365,8 @@ bool FAnimationAnimNextInterfaceTest_AltParamStorageChained::RunTest(const FStri
 		AddErrorIfFalse(bOk, "TestDataMultiply -> GetData returned FAILURE.");
 
 		// --- Check the results from the operation ---
-		const FTestData& TestOutput_AxB = TestContext.GetParameterAs<const FTestData>(TestOutput_AxB_Param);
-		const FTestData& ChainedResult = TestContext.GetParameterAs<const FTestData>(ChainedResultParam);
+		const FAnimNextTestData& TestOutput_AxB = TestContext.GetParameterAs<const FAnimNextTestData>(TestOutput_AxB_Param);
+		const FAnimNextTestData& ChainedResult = TestContext.GetParameterAs<const FAnimNextTestData>(ChainedResultParam);
 		
 		AddErrorIfFalse(ChainedResult.A == TestOutput_AxB.A * TestDataC.A, FString::Printf(TEXT("Chained TestDataResult.A is : [%.04f] Expected : [%.04f]"), ChainedResult.A, TestOutput_AxB.A * TestDataC.A));
 		AddErrorIfFalse(ChainedResult.B == TestOutput_AxB.B * TestDataC.B, FString::Printf(TEXT("Chained TestDataResult.B is : [%.04f] Expected : [%.04f]"), ChainedResult.B, TestOutput_AxB.B * TestDataC.B));
diff --git a/Engine/Plugins/Experimental/AnimNext/Source/AnimNextTestSuite/Private/AnimNextTest.h b/Engine/Plugins/Experimental/AnimNext/Source/AnimNextTestSuite/Private/AnimNextTest.h
--- a/Engine/Plugins/Experimental/AnimNext/Source/AnimNextTestSuite/Private/AnimNextTest.h
+++ b/Engine/Plugins/Experimental/AnimNext/Source/AnimNextTestSuite/Private/AnimNextTest.h
@@ -9,7 +9,7 @@
 
 //// --- Raw type ---
 USTRUCT()
-struct FTestData
+struct FAnimNextTestData
 {
 	GENERATED_BODY()
 
@@ -23,7 +23,7 @@ class UAnimNextInterface_TestData : public UObject, public IAnimNextInterface
 {
 	GENERATED_BODY()
 
-	ANIM_NEXT_INTERFACE_RETURN_TYPE(FTestData)
+	ANIM_NEXT_INTERFACE_RETURN_TYPE(FAnimNextTestData)
 
 	virtual bool GetDataImpl(const UE::AnimNext::FContext& Context) const override
 	{
@@ -41,7 +41,7 @@ class UAnimNextInterfaceTestDataLiteral : public UAnimNextInterface_TestData
 	// IAnimNextInterface interface
 	virtual bool GetDataImpl(const UE::AnimNext::FContext& Context) const final override;
 
-	FTestData Value = {1.f, 1.f};
+	FAnimNextTestData Value = {1.f, 1.f};
 };
 
 // --- TestData multiply operator  ---
diff --git a/Engine/Plugins/Experimental/AnimNext/Source/AnimNextUncookedOnly/Internal/Param/AnimNextParameterBlockProcessor.h b/Engine/Plugins/Experimental/AnimNext/Source/AnimNextUncookedOnly/Internal/Param/AnimNextParameterBlockProcessor.h
--- a/Engine/Plugins/Experimental/AnimNext/Source/AnimNextUncookedOnly/Internal/Param/AnimNextParameterBlockProcessor.h
+++ b/Engine/Plugins/Experimental/AnimNext/Source/AnimNextUncookedOnly/Internal/Param/AnimNextParameterBlockProcessor.h
@@ -1,4 +1,5 @@
-﻿
+﻿// Copyright Epic Games, Inc. All Rights Reserved.
+
 #pragma once
 
 #include "CoreMinimal.h"
@@ -12,6 +13,6 @@ class UAnimNextParameterBlockProcessor : public UAnimNextParameterBlockEntry
 	GENERATED_BODY()
 
 	// UAnimNextParameterBlockEntry interface
-	virtual FText GetDisplayName() const { return FText::GetEmpty(); }
-	virtual FText GetDisplayNameTooltip() const { return FText::GetEmpty(); }
+	virtual FText GetDisplayName() const override { return FText::GetEmpty(); }
+	virtual FText GetDisplayNameTooltip() const override { return FText::GetEmpty(); }
 };
\ No newline at end of file
