Fix non-unity compile issues

diff --git a/Engine/Plugins/Experimental/AnimNext/Source/AnimNextEditor/Internal/AnimNextEditorModule.h b/Engine/Plugins/Experimental/AnimNext/Source/AnimNextEditor/Internal/AnimNextEditorModule.h
--- a/Engine/Plugins/Experimental/AnimNext/Source/AnimNextEditor/Internal/AnimNextEditorModule.h
+++ b/Engine/Plugins/Experimental/AnimNext/Source/AnimNextEditor/Internal/AnimNextEditorModule.h
@@ -5,6 +5,8 @@
 #include "CoreMinimal.h"
 #include "Modules/ModuleManager.h"
 
+class SWidget;
+
 namespace UE::AnimNext::Editor
 {
 
diff --git a/Engine/Plugins/Experimental/AnimNext/Source/AnimNextEditor/Internal/Param/ParameterPickerArgs.h b/Engine/Plugins/Experimental/AnimNext/Source/AnimNextEditor/Internal/Param/ParameterPickerArgs.h
--- a/Engine/Plugins/Experimental/AnimNext/Source/AnimNextEditor/Internal/Param/ParameterPickerArgs.h
+++ b/Engine/Plugins/Experimental/AnimNext/Source/AnimNextEditor/Internal/Param/ParameterPickerArgs.h
@@ -2,9 +2,9 @@
 
 #pragma once
 
+#include "UObject/NameTypes.h"
 #include "Delegates/DelegateCombinations.h"
-
-struct FAssetData;
+#include "AssetRegistry/AssetData.h"
 
 namespace UE::AnimNext::Editor
 {
diff --git a/Engine/Plugins/Experimental/AnimNext/Source/AnimNextEditor/Private/EditorUtils.h b/Engine/Plugins/Experimental/AnimNext/Source/AnimNextEditor/Private/EditorUtils.h
--- a/Engine/Plugins/Experimental/AnimNext/Source/AnimNextEditor/Private/EditorUtils.h
+++ b/Engine/Plugins/Experimental/AnimNext/Source/AnimNextEditor/Private/EditorUtils.h
@@ -5,6 +5,7 @@
 #include "CoreMinimal.h"
 #include "Param/AnimNextParameterBlockEntry.h"
 #include "Param/ParamType.h"
+#include "EdGraphSchema_K2.h"
 
 class UAnimNextGraph;
 class UAnimNextGraph_EditorData;
diff --git a/Engine/Plugins/Experimental/AnimNext/Source/AnimNextEditor/Private/Param/AnimNextParameterSettings.cpp b/Engine/Plugins/Experimental/AnimNext/Source/AnimNextEditor/Private/Param/AnimNextParameterSettings.cpp
--- a/Engine/Plugins/Experimental/AnimNext/Source/AnimNextEditor/Private/Param/AnimNextParameterSettings.cpp
+++ b/Engine/Plugins/Experimental/AnimNext/Source/AnimNextEditor/Private/Param/AnimNextParameterSettings.cpp
@@ -4,6 +4,7 @@
 
 #include "AssetRegistry/AssetRegistryModule.h"
 #include "AssetRegistry/IAssetRegistry.h"
+#include "AssetRegistry/AssetData.h"
 #include "Misc/PackageName.h"
 #include "Misc/Paths.h"
 
diff --git a/Engine/Plugins/Experimental/AnimNext/Source/AnimNextEditor/Private/Param/ParameterLibraryFactory.cpp b/Engine/Plugins/Experimental/AnimNext/Source/AnimNextEditor/Private/Param/ParameterLibraryFactory.cpp
--- a/Engine/Plugins/Experimental/AnimNext/Source/AnimNextEditor/Private/Param/ParameterLibraryFactory.cpp
+++ b/Engine/Plugins/Experimental/AnimNext/Source/AnimNextEditor/Private/Param/ParameterLibraryFactory.cpp
@@ -2,6 +2,7 @@
 
 #include "Param/ParameterLibraryFactory.h"
 #include "Param/AnimNextParameterLibrary.h"
+#include "UObject/Package.h"
 
 UAnimNextParameterLibraryFactory::UAnimNextParameterLibraryFactory()
 {
diff --git a/Engine/Plugins/Experimental/AnimNext/Source/AnimNextEditor/Private/Param/SAddParametersDialog.cpp b/Engine/Plugins/Experimental/AnimNext/Source/AnimNextEditor/Private/Param/SAddParametersDialog.cpp
--- a/Engine/Plugins/Experimental/AnimNext/Source/AnimNextEditor/Private/Param/SAddParametersDialog.cpp
+++ b/Engine/Plugins/Experimental/AnimNext/Source/AnimNextEditor/Private/Param/SAddParametersDialog.cpp
@@ -23,6 +23,7 @@
 #include "SSimpleComboButton.h"
 #include "Param/AnimNextParameterLibrary.h"
 #include "Param/ParameterLibraryFactory.h"
+#include "ToolMenus.h"
 
 #define LOCTEXT_NAMESPACE "SAddParametersDialog"
 
diff --git a/Engine/Plugins/Experimental/AnimNext/Source/AnimNextEditor/Private/Param/SAddParametersDialog.h b/Engine/Plugins/Experimental/AnimNext/Source/AnimNextEditor/Private/Param/SAddParametersDialog.h
--- a/Engine/Plugins/Experimental/AnimNext/Source/AnimNextEditor/Private/Param/SAddParametersDialog.h
+++ b/Engine/Plugins/Experimental/AnimNext/Source/AnimNextEditor/Private/Param/SAddParametersDialog.h
@@ -5,6 +5,7 @@
 #include "Param/ParamType.h"
 #include "Widgets/SWindow.h"
 #include "AssetRegistry/AssetData.h"
+#include "Widgets/Views/SListView.h"
 
 class SWrapBox;
 class UAnimNextParameterLibrary;
diff --git a/Engine/Plugins/Experimental/AnimNext/Source/AnimNextEditor/Private/Param/SLinkParametersDialog.h b/Engine/Plugins/Experimental/AnimNext/Source/AnimNextEditor/Private/Param/SLinkParametersDialog.h
--- a/Engine/Plugins/Experimental/AnimNext/Source/AnimNextEditor/Private/Param/SLinkParametersDialog.h
+++ b/Engine/Plugins/Experimental/AnimNext/Source/AnimNextEditor/Private/Param/SLinkParametersDialog.h
@@ -16,7 +16,7 @@ class SLinkParametersDialog : public SWindow
 public:
 	SLATE_BEGIN_ARGS(SLinkParametersDialog) {}
 
-	SLATE_EVENT(FOnFilterParameter, OnFilterParameter)
+	SLATE_EVENT(UE::AnimNext::Editor::FOnFilterParameter, OnFilterParameter)
 
 	SLATE_END_ARGS()
 
diff --git a/Engine/Plugins/Experimental/AnimNext/Source/AnimNextEditor/Private/Param/SParameterBlockView.cpp b/Engine/Plugins/Experimental/AnimNext/Source/AnimNextEditor/Private/Param/SParameterBlockView.cpp
--- a/Engine/Plugins/Experimental/AnimNext/Source/AnimNextEditor/Private/Param/SParameterBlockView.cpp
+++ b/Engine/Plugins/Experimental/AnimNext/Source/AnimNextEditor/Private/Param/SParameterBlockView.cpp
@@ -27,6 +27,8 @@
 #include "RevisionControlStyle/RevisionControlStyle.h"
 #include "Widgets/Text/SInlineEditableTextBlock.h"
 #include "Param/AnimNextParameterBlockBindingReference.h"
+#include "ToolMenus.h"
+#include "ScopedTransaction.h"
 
 #define LOCTEXT_NAMESPACE "AnimNextParameterBlockView"
 
diff --git a/Engine/Plugins/Experimental/AnimNext/Source/AnimNextEditor/Private/Param/SParameterBlockView.h b/Engine/Plugins/Experimental/AnimNext/Source/AnimNextEditor/Private/Param/SParameterBlockView.h
--- a/Engine/Plugins/Experimental/AnimNext/Source/AnimNextEditor/Private/Param/SParameterBlockView.h
+++ b/Engine/Plugins/Experimental/AnimNext/Source/AnimNextEditor/Private/Param/SParameterBlockView.h
@@ -4,13 +4,14 @@
 
 #include "Widgets/SCompoundWidget.h"
 #include "Widgets/Views/SListView.h"
+#include "AssetRegistry/AssetData.h"
 
 class UAnimNextParameterBlockBinding;
 class UAnimNextParameter;
 class UAnimNextParameterBlock_EditorData;
 class UAnimNextParameterBlockEntry;
 class URigVMGraph;
-struct FAssetData;
+class FUICommandList;
 
 namespace UE::AnimNext::Editor
 {
diff --git a/Engine/Plugins/Experimental/AnimNext/Source/AnimNextEditor/Private/Param/SParameterLibraryView.cpp b/Engine/Plugins/Experimental/AnimNext/Source/AnimNextEditor/Private/Param/SParameterLibraryView.cpp
--- a/Engine/Plugins/Experimental/AnimNext/Source/AnimNextEditor/Private/Param/SParameterLibraryView.cpp
+++ b/Engine/Plugins/Experimental/AnimNext/Source/AnimNextEditor/Private/Param/SParameterLibraryView.cpp
@@ -21,6 +21,9 @@
 #include "Framework/Commands/GenericCommands.h"
 #include "RevisionControlStyle/RevisionControlStyle.h"
 #include "Widgets/Text/SInlineEditableTextBlock.h"
+#include "ScopedTransaction.h"
+#include "Framework/Commands/UICommandList.h"
+#include "ToolMenus.h"
 
 #define LOCTEXT_NAMESPACE "AnimNextParameterLibraryView"
 
diff --git a/Engine/Plugins/Experimental/AnimNext/Source/AnimNextEditor/Private/Param/SParameterLibraryView.h b/Engine/Plugins/Experimental/AnimNext/Source/AnimNextEditor/Private/Param/SParameterLibraryView.h
--- a/Engine/Plugins/Experimental/AnimNext/Source/AnimNextEditor/Private/Param/SParameterLibraryView.h
+++ b/Engine/Plugins/Experimental/AnimNext/Source/AnimNextEditor/Private/Param/SParameterLibraryView.h
@@ -8,15 +8,13 @@
 class UAnimNextParameterLibrary;
 class UAnimNextParameter;
 struct FAssetData;
+class FUICommandList;
 
 namespace UE::AnimNext::Editor
 {
 
 enum class EFilterParameterResult : int32;
 struct FParameterBindingReference;
-
-
-
 struct FParameterLibraryViewEntry; 
 
 class SParameterLibraryView : public SCompoundWidget
diff --git a/Engine/Plugins/Experimental/AnimNext/Source/AnimNextEditor/Private/Param/SParametersActionMenu.h b/Engine/Plugins/Experimental/AnimNext/Source/AnimNextEditor/Private/Param/SParametersActionMenu.h
--- a/Engine/Plugins/Experimental/AnimNext/Source/AnimNextEditor/Private/Param/SParametersActionMenu.h
+++ b/Engine/Plugins/Experimental/AnimNext/Source/AnimNextEditor/Private/Param/SParametersActionMenu.h
@@ -13,6 +13,7 @@
 #include "Widgets/Layout/SBorder.h"
 #include "Widgets/SWidget.h"
 #include "Widgets/Views/SExpanderArrow.h"
+#include "GraphEditor.h"
 
 class SEditableTextBox;
 class SGraphActionMenu;
diff --git a/Engine/Plugins/Experimental/AnimNext/Source/AnimNextUncookedOnly/Private/Param/AnimNextParameter.cpp b/Engine/Plugins/Experimental/AnimNext/Source/AnimNextUncookedOnly/Private/Param/AnimNextParameter.cpp
--- a/Engine/Plugins/Experimental/AnimNext/Source/AnimNextUncookedOnly/Private/Param/AnimNextParameter.cpp
+++ b/Engine/Plugins/Experimental/AnimNext/Source/AnimNextUncookedOnly/Private/Param/AnimNextParameter.cpp
@@ -1,6 +1,7 @@
 ﻿// Copyright Epic Games, Inc. All Rights Reserved.
 
 #include "Param/AnimNextParameter.h"
+#include "UObject/Package.h"
 
 #if WITH_EDITORONLY_DATA
 
