CIS: More non-unity fixes

diff --git a/Engine/Plugins/Experimental/AnimNext/Source/AnimNextUncookedOnly/Internal/Param/AnimNextParameterBlock_EdGraph.h b/Engine/Plugins/Experimental/AnimNext/Source/AnimNextUncookedOnly/Internal/Param/AnimNextParameterBlock_EdGraph.h
--- a/Engine/Plugins/Experimental/AnimNext/Source/AnimNextUncookedOnly/Internal/Param/AnimNextParameterBlock_EdGraph.h
+++ b/Engine/Plugins/Experimental/AnimNext/Source/AnimNextUncookedOnly/Internal/Param/AnimNextParameterBlock_EdGraph.h
@@ -6,7 +6,7 @@
 #include "Graph/ControlRigGraph.h"
 #include "AnimNextParameterBlock_EdGraph.generated.h"
 
-class UAnimNextGraph_EditorData;
+class UAnimNextParameterBlock_EditorData;
 
 namespace UE::AnimNext::UncookedOnly
 {
diff --git a/Engine/Plugins/Experimental/AnimNext/Source/AnimNextUncookedOnly/Private/Param/AnimNextParameterBlock_EditorData.cpp b/Engine/Plugins/Experimental/AnimNext/Source/AnimNextUncookedOnly/Private/Param/AnimNextParameterBlock_EditorData.cpp
--- a/Engine/Plugins/Experimental/AnimNext/Source/AnimNextUncookedOnly/Private/Param/AnimNextParameterBlock_EditorData.cpp
+++ b/Engine/Plugins/Experimental/AnimNext/Source/AnimNextUncookedOnly/Private/Param/AnimNextParameterBlock_EditorData.cpp
@@ -3,6 +3,7 @@
 #include "Param/AnimNextParameterBlock_EditorData.h"
 
 #include "ControlRigDefines.h"
+#include "Param/AnimNextParameterBlock_EdGraph.h"
 #include "Param/AnimNextParameterBlock_EdGraphSchema.h"
 #include "UncookedOnlyUtils.h"
 #include "Param/ParametersExecuteContext.h"
