CIS: Fixup LOCTEXT duplicate key

diff --git a/Engine/Plugins/Experimental/AnimNext/Source/AnimNextEditor/Private/Param/SAddParametersDialog.cpp b/Engine/Plugins/Experimental/AnimNext/Source/AnimNextEditor/Private/Param/SAddParametersDialog.cpp
--- a/Engine/Plugins/Experimental/AnimNext/Source/AnimNextEditor/Private/Param/SAddParametersDialog.cpp
+++ b/Engine/Plugins/Experimental/AnimNext/Source/AnimNextEditor/Private/Param/SAddParametersDialog.cpp
@@ -110,9 +110,9 @@ void SAddParametersDialog::Construct(const FArguments& InArgs)
 						.ButtonStyle(&FAppStyle::Get().GetWidgetStyle<FButtonStyle>("PrimaryButton"))
 						.Text_Lambda([this]()
 						{
-							return FText::Format(LOCTEXT("AddButtonFormat", "Add {0} {0}|plural(one=Parameter,other=Parameters)"), FText::AsNumber(Entries.Num()));
+							return FText::Format(LOCTEXT("AddParametersButtonFormat", "Add {0} {0}|plural(one=Parameter,other=Parameters)"), FText::AsNumber(Entries.Num()));
 						})
-						.ToolTipText(LOCTEXT("AddButtonTooltip", "Add the selected parameters to the current parameter block"))
+						.ToolTipText(LOCTEXT("AddParametersButtonTooltip", "Add the selected parameters to the current parameter block"))
 						.OnClicked_Lambda([this]()
 						{
 							RequestDestroyWindow();
