Fix ASan issue when unregistering proeprty type customizations on shutdown

diff --git a/Engine/Plugins/Experimental/AnimNext/Source/AnimNextEditor/Private/AnimNextEditorModule.cpp b/Engine/Plugins/Experimental/AnimNext/Source/AnimNextEditor/Private/AnimNextEditorModule.cpp
--- a/Engine/Plugins/Experimental/AnimNext/Source/AnimNextEditor/Private/AnimNextEditorModule.cpp
+++ b/Engine/Plugins/Experimental/AnimNext/Source/AnimNextEditor/Private/AnimNextEditorModule.cpp
@@ -22,12 +22,12 @@ void FModule::StartupModule()
 	AnimNextPropertyTypeIdentifier = MakeShared<FPropertyTypeIdentifier>();
 	FPropertyEditorModule& PropertyModule = FModuleManager::LoadModuleChecked<FPropertyEditorModule>("PropertyEditor");
 	PropertyModule.RegisterCustomPropertyTypeLayout(
-		FInterfaceProperty::StaticClass()->GetFName(),
+		"InterfaceProperty",
 		FOnGetPropertyTypeCustomizationInstance::CreateLambda([] { return MakeShared<FPropertyTypeCustomization>(); }),
 		AnimNextPropertyTypeIdentifier);
 
 	PropertyModule.RegisterCustomPropertyTypeLayout(
-		FAnimNextParamType::StaticStruct()->GetFName(),
+		"AnimNextParamType",
 		FOnGetPropertyTypeCustomizationInstance::CreateLambda([] { return MakeShared<FParamTypePropertyTypeCustomization>(); }));
 }
 
@@ -42,8 +42,8 @@ void FModule::ShutdownModule()
 	if(FModuleManager::Get().IsModuleLoaded("PropertyEditor"))
 	{
 		FPropertyEditorModule& PropertyModule = FModuleManager::LoadModuleChecked<FPropertyEditorModule>("PropertyEditor");
-		PropertyModule.UnregisterCustomPropertyTypeLayout(FInterfaceProperty::StaticClass()->GetFName(), AnimNextPropertyTypeIdentifier);
-		PropertyModule.UnregisterCustomPropertyTypeLayout(FAnimNextParamType::StaticStruct()->GetFName());
+		PropertyModule.UnregisterCustomPropertyTypeLayout("InterfaceProperty", AnimNextPropertyTypeIdentifier);
+		PropertyModule.UnregisterCustomPropertyTypeLayout("AnimNextParamType");
 	}
 }
 
