Implement AnimNext decorator serialization

diff --git a/Engine/Plugins/Experimental/AnimNext/Source/AnimNext/Internal/DecoratorBase/DecoratorReader.h b/Engine/Plugins/Experimental/AnimNext/Source/AnimNext/Internal/DecoratorBase/DecoratorReader.h
new file mode 100644
--- /dev/null
+++ b/Engine/Plugins/Experimental/AnimNext/Source/AnimNext/Internal/DecoratorBase/DecoratorReader.h
@@ -0,0 +1,25 @@
+// Copyright Epic Games, Inc. All Rights Reserved.
+
+#pragma once
+
+#include "CoreMinimal.h"
+#include "Serialization/ArchiveProxy.h"
+
+namespace UE::AnimNext
+{
+	/**
+	  * FDecoratorReader
+	  * 
+	  * The decorator reader is used to read from a serialized binary blob that contains
+	  * the anim graph data. An anim graph contains the following:
+	  *     - A list of FNodeTemplates that the nodes use
+	  *     - The graph shared data (FNodeDescription for every node)
+	  */
+	class ANIMNEXT_API FDecoratorReader : public FArchiveProxy
+	{
+	public:
+		FDecoratorReader(FArchive& Ar);
+
+		void ReadGraphSharedData(TArray<uint8>& GraphSharedData);
+	};
+}
diff --git a/Engine/Plugins/Experimental/AnimNext/Source/AnimNext/Internal/DecoratorBase/DecoratorTemplate.h b/Engine/Plugins/Experimental/AnimNext/Source/AnimNext/Internal/DecoratorBase/DecoratorTemplate.h
--- a/Engine/Plugins/Experimental/AnimNext/Source/AnimNext/Internal/DecoratorBase/DecoratorTemplate.h
+++ b/Engine/Plugins/Experimental/AnimNext/Source/AnimNext/Internal/DecoratorBase/DecoratorTemplate.h
@@ -9,10 +9,12 @@
 
 #include <type_traits>
 
+class FArchive;
+struct FAnimNextDecoratorSharedData;
+
 namespace UE::AnimNext
 {
-	struct FDecoratorDescription;
-	struct FDecoratorInstance;
+	struct FDecoratorInstanceData;
 	struct FNodeDescription;
 	struct FNodeInstance;
 
@@ -63,29 +65,32 @@ namespace UE::AnimNext
 		uint32 GetNodeInstanceOffset() const { return NodeInstanceOffset; }
 
 		// Returns a pointer to the specified decorator description on the current node
-		FDecoratorDescription* GetDecoratorDescription(FNodeDescription& NodeDescription) const
+		FAnimNextDecoratorSharedData* GetDecoratorDescription(FNodeDescription& NodeDescription) const
 		{
-			return reinterpret_cast<FDecoratorDescription*>(reinterpret_cast<uint8*>(&NodeDescription) + GetNodeSharedOffset());
+			return reinterpret_cast<FAnimNextDecoratorSharedData*>(reinterpret_cast<uint8*>(&NodeDescription) + GetNodeSharedOffset());
 		}
 
 		// Returns a pointer to the specified decorator description on the current node
-		const FDecoratorDescription* GetDecoratorDescription(const FNodeDescription& NodeDescription) const
+		const FAnimNextDecoratorSharedData* GetDecoratorDescription(const FNodeDescription& NodeDescription) const
 		{
-			return reinterpret_cast<const FDecoratorDescription*>(reinterpret_cast<const uint8*>(&NodeDescription) + GetNodeSharedOffset());
+			return reinterpret_cast<const FAnimNextDecoratorSharedData*>(reinterpret_cast<const uint8*>(&NodeDescription) + GetNodeSharedOffset());
 		}
 
 		// Returns a pointer to the specified decorator instance on the current node
-		FDecoratorInstance* GetDecoratorInstance(FNodeInstance& NodeInstance) const
+		FDecoratorInstanceData* GetDecoratorInstance(FNodeInstance& NodeInstance) const
 		{
-			return reinterpret_cast<FDecoratorInstance*>(reinterpret_cast<uint8*>(&NodeInstance) + GetNodeInstanceOffset());
+			return reinterpret_cast<FDecoratorInstanceData*>(reinterpret_cast<uint8*>(&NodeInstance) + GetNodeInstanceOffset());
 		}
 
 		// Returns a pointer to the specified decorator instance on the current node
-		const FDecoratorInstance* GetDecoratorInstance(const FNodeInstance& NodeInstance) const
+		const FDecoratorInstanceData* GetDecoratorInstance(const FNodeInstance& NodeInstance) const
 		{
-			return reinterpret_cast<const FDecoratorInstance*>(reinterpret_cast<const uint8*>(&NodeInstance) + GetNodeInstanceOffset());
+			return reinterpret_cast<const FDecoratorInstanceData*>(reinterpret_cast<const uint8*>(&NodeInstance) + GetNodeInstanceOffset());
 		}
 
+		// Serializes this decorator template instance
+		ANIMNEXT_API void Serialize(FArchive& Ar);
+
 	private:
 		uint32	UID;							// decorator globally unique identifier
 
diff --git a/Engine/Plugins/Experimental/AnimNext/Source/AnimNext/Internal/DecoratorBase/DecoratorWriter.h b/Engine/Plugins/Experimental/AnimNext/Source/AnimNext/Internal/DecoratorBase/DecoratorWriter.h
new file mode 100644
--- /dev/null
+++ b/Engine/Plugins/Experimental/AnimNext/Source/AnimNext/Internal/DecoratorBase/DecoratorWriter.h
@@ -0,0 +1,91 @@
+// Copyright Epic Games, Inc. All Rights Reserved.
+
+#pragma once
+
+#include "CoreMinimal.h"
+
+#if WITH_EDITOR
+#include "Serialization/MemoryWriter.h"
+
+#include "DecoratorBase/NodeHandle.h"
+#include "DecoratorBase/NodeTemplateRegistryHandle.h"
+
+namespace UE::AnimNext
+{
+	struct FNodeTemplate;
+
+	/**
+	  * FDecoratorWriter
+	  *
+	  * The decorator writer is used to write a serialized binary blob that contains
+	  * the anim graph data. An anim graph contains the following:
+	  *     - A list of FNodeTemplates that the nodes use
+	  *     - The graph shared data (FNodeDescription for every node)
+	  */
+	class ANIMNEXT_API FDecoratorWriter : public FMemoryWriter
+	{
+	public:
+		// The largest size allowed for the shared data of a single graph
+		// We use 16 bit offsets to index into the shared data
+		static constexpr uint32 MAXIMUM_GRAPH_SHARED_DATA_SIZE = 64 * 1024;
+
+		enum class EErrorState
+		{
+			None,					// All good, no error
+			TooManyNodes,			// Exceeded the maximum number of nodes in a graph, @see FNodeDescription::MAXIMUM_COUNT
+			GraphTooLarge,			// Exceeded the maximum graph size, @see FDecoratorWriter::MAXIMUM_GRAPH_SHARED_DATA_SIZE
+			NodeTemplateTooLarge,	// Exceeded the maximum node template size, @see FNodeTemplate::MAXIMUM_SIZE
+			NodeSharedDataTooLarge,	// Exceeded the maximum node shared data size, @see FNodeDescription::MAXIMUM_SIZE
+		};
+
+		FDecoratorWriter();
+
+		// Registers an instance of the provided node template and assigns a node handle and node UID to it
+		FNodeHandle RegisterNode(const FNodeTemplate& NodeTemplate);
+
+		// Returns the node UID associated with a node handle
+		uint16 GetNodeUID(const FNodeHandle NodeHandle) const;
+
+		// Returns the node handle associated with a node UID
+		FNodeHandle GetNodeHandle(uint16 NodeUID) const;
+
+		// Called before node writing can begin
+		void BeginNodeWriting();
+
+		// Called once node writing has terminated
+		void EndNodeWriting();
+
+		// Writes out the provided node using the decorator properties
+		// Nodes must be written in the same order they were registered in
+		void WriteNode(const FNodeHandle NodeHandle, const TFunction<const TMap<FString, FString>& (uint32 DecoratorIndex)>& GetDecoratorProperties);
+
+		// Returns the error state
+		EErrorState GetErrorState() const;
+
+		// Returns the populated raw graph shared data buffer
+		const TArray<uint8>& GetGraphSharedData() const;
+
+	private:
+		struct FNodeMapping
+		{
+			FNodeHandle NodeHandle;
+			uint32 NodeSize;
+			uint16 NodeUID;
+			FNodeTemplateRegistryHandle NodeTemplateHandle;
+		};
+
+		TArray<uint8> GraphSharedDataArchiveBuffer;
+		TArray<FNodeMapping> NodeMappings;
+
+		// To track the node registration process
+		uint32 GraphSharedDataSize;
+		uint16 NextNodeUID;
+
+		// To track node writing
+		uint32 NumNodesWritten;
+		bool bIsNodeWriting;
+
+		EErrorState ErrorState;
+	};
+}
+#endif
diff --git a/Engine/Plugins/Experimental/AnimNext/Source/AnimNext/Internal/DecoratorBase/NodeTemplate.h b/Engine/Plugins/Experimental/AnimNext/Source/AnimNext/Internal/DecoratorBase/NodeTemplate.h
--- a/Engine/Plugins/Experimental/AnimNext/Source/AnimNext/Internal/DecoratorBase/NodeTemplate.h
+++ b/Engine/Plugins/Experimental/AnimNext/Source/AnimNext/Internal/DecoratorBase/NodeTemplate.h
@@ -7,6 +7,8 @@
 
 #include <type_traits>
 
+class FArchive;
+
 namespace UE::AnimNext
 {
 	/**
@@ -23,6 +25,11 @@ namespace UE::AnimNext
 	 */
 	struct FNodeTemplate
 	{
+		// The largest size allowed for a node template
+		// Node templates are just a lightweight descriptor, in theory they can have any size
+		// We artificially specify a conservative upper bound
+		static constexpr uint32 MAXIMUM_SIZE = 64 * 1024;
+
 		FNodeTemplate(uint32 UID_, uint32 InstanceSize_, uint32 NumDecorators_)
 			: UID(UID_)
 			, InstanceSize(InstanceSize_)
@@ -41,9 +48,15 @@ namespace UE::AnimNext
 		// Returns the number of decorators present in the node template
 		uint32 GetNumDecorators() const { return NumDecorators; }
 
+		// Returns a pointer to the list of decorator template descriptions
+		FDecoratorTemplate* GetDecorators() { return reinterpret_cast<FDecoratorTemplate*>(reinterpret_cast<uint8*>(this) + sizeof(FNodeTemplate)); }
+
 		// Returns a pointer to the list of decorator template descriptions
 		const FDecoratorTemplate* GetDecorators() const { return reinterpret_cast<const FDecoratorTemplate*>(reinterpret_cast<const uint8*>(this) + sizeof(FNodeTemplate)); }
 
+		// Serializes this node template instance and each decorator template that follows
+		ANIMNEXT_API void Serialize(FArchive& Ar);
+
 	private:
 		uint32	UID;				// globally unique template identifier or hash
 
diff --git a/Engine/Plugins/Experimental/AnimNext/Source/AnimNext/Internal/DecoratorBase/NodeTemplateRegistry.h b/Engine/Plugins/Experimental/AnimNext/Source/AnimNext/Internal/DecoratorBase/NodeTemplateRegistry.h
--- a/Engine/Plugins/Experimental/AnimNext/Source/AnimNext/Internal/DecoratorBase/NodeTemplateRegistry.h
+++ b/Engine/Plugins/Experimental/AnimNext/Source/AnimNext/Internal/DecoratorBase/NodeTemplateRegistry.h
@@ -21,11 +21,14 @@ namespace UE::AnimNext
 		// Access the global registry
 		static FNodeTemplateRegistry& Get();
 
+		// Finds the specified node template from its UID and returns its offset
+		FNodeTemplateRegistryHandle Find(uint32 NodeTemplateUID) const;
+
 		// Finds or adds the specified node template and returns its offset
 		FNodeTemplateRegistryHandle FindOrAdd(const FNodeTemplate* NodeTemplate);
 
-		// Removes the specified node template from the registry
-		void Unregister(const FNodeTemplate* NodeTemplate);
+		// Finds and returns a node template based on its handle or nullptr if the handle is invalid
+		FNodeTemplate* Find(FNodeTemplateRegistryHandle TemplateHandle);
 
 		// Finds and returns a node template based on its handle or nullptr if the handle is invalid
 		const FNodeTemplate* Find(FNodeTemplateRegistryHandle TemplateHandle) const;
@@ -33,6 +36,12 @@ namespace UE::AnimNext
 		// Returns the number of registered node templates
 		uint32 GetNum() const;
 
+		// Removes the specified node template from the registry
+		void Unregister(const FNodeTemplate* NodeTemplate);
+
+		// Clears the registry by removing every node template
+		void Clear();
+
 	private:
 		FNodeTemplateRegistry() = default;
 
diff --git a/Engine/Plugins/Experimental/AnimNext/Source/AnimNext/Private/DecoratorBase/Decorator.cpp b/Engine/Plugins/Experimental/AnimNext/Source/AnimNext/Private/DecoratorBase/Decorator.cpp
--- a/Engine/Plugins/Experimental/AnimNext/Source/AnimNext/Private/DecoratorBase/Decorator.cpp
+++ b/Engine/Plugins/Experimental/AnimNext/Source/AnimNext/Private/DecoratorBase/Decorator.cpp
@@ -2,10 +2,69 @@
 
 #include "DecoratorBase/Decorator.h"
 
+#include "DecoratorBase/DecoratorReader.h"
 #include "DecoratorBase/DecoratorRegistry.h"
 
+#if WITH_EDITOR
+#include "DecoratorBase/DecoratorWriter.h"
+#endif
+
 namespace UE::AnimNext
 {
+	void FDecorator::SerializeDecoratorSharedData(FArchive& Ar, FAnimNextDecoratorSharedData& SharedData) const
+	{
+		UScriptStruct* SharedDataStruct = GetDecoratorSharedDataStruct();
+		SharedDataStruct->SerializeItem(Ar, &SharedData, nullptr);
+	}
+
+#if WITH_EDITOR
+	void FDecorator::SaveDecoratorSharedData(FDecoratorWriter& Writer, const TMap<FString, FString>& Properties, FAnimNextDecoratorSharedData& OutSharedData) const
+	{
+		const UScriptStruct* SharedDataStruct = GetDecoratorSharedDataStruct();
+
+		uint8* SharedData = reinterpret_cast<uint8*>(&OutSharedData);
+
+		// Initialize our output struct with its default values
+		SharedDataStruct->InitializeDefaultValue(SharedData);
+
+		// Use UE reflection to iterate over every property
+		// We convert every property from its string representation into its binary form
+		for (const FProperty* Property = SharedDataStruct->PropertyLink; Property != nullptr; Property = Property->PropertyLinkNext)
+		{
+			if (const FString* PropertyValue = Properties.Find(Property->GetName()))
+			{
+				const TCHAR* PropertyValuePtr = **PropertyValue;
+
+				// C-style array properties aren't handled by ExportText, we need to handle it manually
+				const bool bIsCArray = Property->ArrayDim > 1;
+				if (bIsCArray)
+				{
+					ensure(PropertyValuePtr[0] == TEXT('('));
+					PropertyValuePtr++;
+				}
+
+				for (int32 Index = 0; Index < Property->ArrayDim; ++Index)
+				{
+					void* DataPtr = Property->ContainerPtrToValuePtr<void>(SharedData, Index);
+					PropertyValuePtr = Property->ImportText_Direct(PropertyValuePtr, DataPtr, nullptr, PPF_None);
+
+					if (Index + 1 < Property->ArrayDim)
+					{
+						ensure(PropertyValuePtr[0] == TEXT(','));
+						PropertyValuePtr++;
+					}
+				}
+
+				if (bIsCArray)
+				{
+					ensure(PropertyValuePtr[0] == TEXT(')'));
+					PropertyValuePtr++;
+				}
+			}
+		}
+	}
+#endif
+
 	FDecoratorStaticInitHook::FDecoratorStaticInitHook(DecoratorConstructorFunc DecoratorConstructor_)
 		: DecoratorConstructor(DecoratorConstructor_)
 	{
diff --git a/Engine/Plugins/Experimental/AnimNext/Source/AnimNext/Private/DecoratorBase/DecoratorReader.cpp b/Engine/Plugins/Experimental/AnimNext/Source/AnimNext/Private/DecoratorBase/DecoratorReader.cpp
new file mode 100644
--- /dev/null
+++ b/Engine/Plugins/Experimental/AnimNext/Source/AnimNext/Private/DecoratorBase/DecoratorReader.cpp
@@ -0,0 +1,60 @@
+// Copyright Epic Games, Inc. All Rights Reserved.
+
+#include "DecoratorBase/DecoratorReader.h"
+
+#include "DecoratorBase/NodeDescription.h"
+#include "DecoratorBase/NodeTemplate.h"
+#include "DecoratorBase/NodeTemplateRegistry.h"
+
+namespace UE::AnimNext
+{
+	FDecoratorReader::FDecoratorReader(FArchive& Ar)
+		: FArchiveProxy(Ar)
+	{
+	}
+
+	void FDecoratorReader::ReadGraphSharedData(TArray<uint8>& GraphSharedData)
+	{
+		GraphSharedData.Empty(0);
+
+		// Read the node templates and register them as needed
+		uint32 NumNodeTemplates = 0;
+		*this << NumNodeTemplates;
+
+		// We serialize each node template in the same buffer, they get copied into the registry and we don't need to retain them
+		alignas(16) uint8 NodeTemplateBuffer[64 * 1024];
+		FNodeTemplate* NodeTemplate = reinterpret_cast<FNodeTemplate*>(&NodeTemplateBuffer[0]);
+
+		FNodeTemplateRegistry& NodeTemplateRegistry = FNodeTemplateRegistry::Get();
+		for (uint32 NodeTemplateIndex = 0; NodeTemplateIndex < NumNodeTemplates; ++NodeTemplateIndex)
+		{
+			NodeTemplate->Serialize(*this);
+
+			// Register our node template
+			NodeTemplateRegistry.FindOrAdd(NodeTemplate);
+		}
+
+		// Read our graph shared data
+		uint32 NumNodes = 0;
+		*this << NumNodes;
+
+		uint32 SharedDataSize = 0;
+		*this << SharedDataSize;
+
+		GraphSharedData.AddZeroed(SharedDataSize);
+
+		uint32 SharedDataOffset = 0;
+		for (uint32 NodeIndex = 0; NodeIndex < NumNodes; ++NodeIndex)
+		{
+			uint32 NodeSize = 0;
+			*this << NodeSize;
+
+			FNodeDescription* NodeDesc = reinterpret_cast<FNodeDescription*>(&GraphSharedData[SharedDataOffset]);
+			NodeDesc->Serialize(*this);
+
+			SharedDataOffset += NodeSize;
+		}
+
+		ensure(SharedDataOffset == SharedDataSize);
+	}
+}
diff --git a/Engine/Plugins/Experimental/AnimNext/Source/AnimNext/Private/DecoratorBase/DecoratorTemplate.cpp b/Engine/Plugins/Experimental/AnimNext/Source/AnimNext/Private/DecoratorBase/DecoratorTemplate.cpp
new file mode 100644
--- /dev/null
+++ b/Engine/Plugins/Experimental/AnimNext/Source/AnimNext/Private/DecoratorBase/DecoratorTemplate.cpp
@@ -0,0 +1,49 @@
+// Copyright Epic Games, Inc. All Rights Reserved.
+
+#include "DecoratorBase/DecoratorTemplate.h"
+
+#include "Serialization/Archive.h"
+#include "DecoratorBase/Decorator.h"
+#include "DecoratorBase/DecoratorRegistry.h"
+
+namespace UE::AnimNext
+{
+	void FDecoratorTemplate::Serialize(FArchive& Ar)
+	{
+		Ar << UID;
+
+		if (Ar.IsSaving())
+		{
+			const FDecorator* Decorator = FDecoratorRegistry::Get().Find(RegistryHandle);
+
+			uint32 DecoratorUID = Decorator->GetDecoratorUID().GetUID();
+			Ar << DecoratorUID;
+		}
+		else if (Ar.IsLoading())
+		{
+			uint32 DecoratorUID = 0;
+			Ar << DecoratorUID;
+
+			RegistryHandle = FDecoratorRegistry::Get().FindHandle(DecoratorUID);
+		}
+		else
+		{
+			// Counting, etc
+			if (RegistryHandle.IsDynamic())
+			{
+				int32 DynamicIndex = RegistryHandle.GetDynamicIndex();
+				Ar << DynamicIndex;
+			}
+			else if (RegistryHandle.IsStatic())
+			{
+				int32 StaticOffset = RegistryHandle.GetStaticOffset();
+				Ar << StaticOffset;
+			}
+		}
+
+		Ar << Mode;
+		Ar << AdditiveIndexOrNumAdditive;
+		Ar << NodeSharedOffset;
+		Ar << NodeInstanceOffset;
+	}
+}
diff --git a/Engine/Plugins/Experimental/AnimNext/Source/AnimNext/Private/DecoratorBase/DecoratorWriter.cpp b/Engine/Plugins/Experimental/AnimNext/Source/AnimNext/Private/DecoratorBase/DecoratorWriter.cpp
new file mode 100644
--- /dev/null
+++ b/Engine/Plugins/Experimental/AnimNext/Source/AnimNext/Private/DecoratorBase/DecoratorWriter.cpp
@@ -0,0 +1,225 @@
+// Copyright Epic Games, Inc. All Rights Reserved.
+
+#include "DecoratorBase/DecoratorWriter.h"
+
+#if WITH_EDITOR
+#include "DecoratorBase/DecoratorRegistry.h"
+#include "DecoratorBase/DecoratorTemplate.h"
+#include "DecoratorBase/NodeDescription.h"
+#include "DecoratorBase/NodeTemplate.h"
+#include "DecoratorBase/NodeTemplateRegistry.h"
+
+namespace UE::AnimNext
+{
+	namespace Private
+	{
+		// Note: this does not account for inline/pooled/pinned properties
+		static uint32 GetNodeSharedSize(const FNodeTemplate& NodeTemplate)
+		{
+			const FDecoratorRegistry& Registry = FDecoratorRegistry::Get();
+
+			const uint32 NumDecorators = NodeTemplate.GetNumDecorators();
+			const FDecoratorTemplate* DecoratorTemplates = NodeTemplate.GetDecorators();
+
+			uint32 NodeSharedSize = sizeof(FNodeDescription);
+			for (uint32_t DecoratorIndex = 0; DecoratorIndex < NumDecorators; ++DecoratorIndex)
+			{
+				const FDecorator* Decorator = Registry.Find(DecoratorTemplates[DecoratorIndex].GetRegistryHandle());
+
+				const FDecoratorMemoryLayout MemoryLayout = Decorator->GetDecoratorMemoryDescription();
+				NodeSharedSize = Align(NodeSharedSize, MemoryLayout.SharedDataAlignment);
+				NodeSharedSize += MemoryLayout.SharedDataSize;
+			}
+
+			return NodeSharedSize;
+		}
+	}
+
+	FDecoratorWriter::FDecoratorWriter()
+		: FMemoryWriter(GraphSharedDataArchiveBuffer)
+		, GraphSharedDataSize(0)
+		, NextNodeUID(1)	// Node UID 0 is reserved for the invalid value (null)
+		, NumNodesWritten(0)
+		, bIsNodeWriting(false)
+		, ErrorState(EErrorState::None)
+	{
+	}
+
+	FNodeHandle FDecoratorWriter::RegisterNode(const FNodeTemplate& NodeTemplate)
+	{
+		ensure(!bIsNodeWriting);
+
+		if (ErrorState != EErrorState::None)
+		{
+			// We encountered an error, do nothing
+			return FNodeHandle();
+		}
+
+		if (NodeTemplate.GetNodeTemplateSize() > FNodeTemplate::MAXIMUM_SIZE)
+		{
+			// This node template is too large
+			ErrorState = EErrorState::NodeTemplateTooLarge;
+			return FNodeHandle();
+		}
+
+		if (NextNodeUID == FNodeDescription::MAXIMUM_COUNT)
+		{
+			// We have too many nodes in the graph, we need to be able to represent them with 16 bits
+			ErrorState = EErrorState::TooManyNodes;
+			return FNodeHandle();
+		}
+
+		const uint32 NodeSharedDataSize = Private::GetNodeSharedSize(NodeTemplate);
+		if (NodeSharedDataSize > FNodeDescription::MAXIMUM_SIZE)
+		{
+			// This node shared data is too large
+			ErrorState = EErrorState::NodeSharedDataTooLarge;
+			return FNodeHandle();
+		}
+
+		const uint32 NodeSharedDataOffset = GraphSharedDataSize;
+		if (GraphSharedDataSize + NodeSharedDataSize > MAXIMUM_GRAPH_SHARED_DATA_SIZE)
+		{
+			// The graph shared data size is too large
+			ErrorState = EErrorState::GraphTooLarge;
+			return FNodeHandle();
+		}
+
+		GraphSharedDataSize += NodeSharedDataSize;
+
+		const FNodeHandle NodeHandle(NodeSharedDataOffset);
+		const uint16 NodeUID = NextNodeUID++;
+
+		FNodeTemplateRegistry& NodeTemplateRegistry = FNodeTemplateRegistry::Get();
+		const FNodeTemplateRegistryHandle NodeTemplateHandle = NodeTemplateRegistry.FindOrAdd(&NodeTemplate);
+
+		NodeMappings.Add({ NodeHandle, NodeSharedDataSize, NodeUID, NodeTemplateHandle });
+
+		return NodeHandle;
+	}
+
+	uint16 FDecoratorWriter::GetNodeUID(const FNodeHandle NodeHandle) const
+	{
+		const FNodeMapping* NodeMapping = NodeMappings.FindByPredicate([NodeHandle](const FNodeMapping& It) { return It.NodeHandle == NodeHandle; });
+		return NodeMapping != nullptr ? NodeMapping->NodeUID : 0;
+	}
+
+	FNodeHandle FDecoratorWriter::GetNodeHandle(uint16 NodeUID) const
+	{
+		const FNodeMapping* NodeMapping = NodeMappings.FindByPredicate([NodeUID](const FNodeMapping& It) { return It.NodeUID == NodeUID; });
+		return NodeMapping != nullptr ? NodeMapping->NodeHandle : FNodeHandle();
+	}
+
+	void FDecoratorWriter::BeginNodeWriting()
+	{
+		ensure(!bIsNodeWriting);
+		ensure(NumNodesWritten == 0);
+
+		if (ErrorState != EErrorState::None)
+		{
+			// We encountered an error, do nothing
+			return;
+		}
+
+		bIsNodeWriting = true;
+
+		// Serialize the node templates
+		TArray<FNodeTemplateRegistryHandle> NodeTemplateHandles;
+		NodeTemplateHandles.Reserve(NodeMappings.Num());
+
+		for (const FNodeMapping& NodeMapping : NodeMappings)
+		{
+			NodeTemplateHandles.AddUnique(NodeMapping.NodeTemplateHandle);
+		}
+
+		uint32 NumNodeTemplates = NodeTemplateHandles.Num();
+		*this << NumNodeTemplates;
+
+		FNodeTemplateRegistry& NodeTemplateRegistry = FNodeTemplateRegistry::Get();
+		for (FNodeTemplateRegistryHandle NodeTemplateHandle : NodeTemplateHandles)
+		{
+			FNodeTemplate* NodeTemplate = NodeTemplateRegistry.Find(NodeTemplateHandle);
+			NodeTemplate->Serialize(*this);
+		}
+
+		// Begin serializing the graph shared data
+		uint32 NumNodes = NodeMappings.Num();
+		*this << NumNodes;
+
+		uint32 SharedDataSize = GraphSharedDataSize;
+		*this << SharedDataSize;
+	}
+
+	void FDecoratorWriter::EndNodeWriting()
+	{
+		ensure(bIsNodeWriting);
+		bIsNodeWriting = false;
+
+		if (ErrorState != EErrorState::None)
+		{
+			// We encountered an error, do nothing
+			return;
+		}
+
+		ensure(NumNodesWritten == NodeMappings.Num());
+	}
+
+	void FDecoratorWriter::WriteNode(const FNodeHandle NodeHandle, const TFunction<const TMap<FString, FString>& (uint32 DecoratorIndex)>& GetDecoratorProperties)
+	{
+		ensure(bIsNodeWriting);
+
+		if (ErrorState != EErrorState::None)
+		{
+			// We encountered an error, do nothing
+			return;
+		}
+
+		const FDecoratorRegistry& DecoratorRegistry = FDecoratorRegistry::Get();
+		FNodeTemplateRegistry& NodeTemplateRegistry = FNodeTemplateRegistry::Get();
+
+		const FNodeMapping* NodeMapping = NodeMappings.FindByPredicate([NodeHandle](const FNodeMapping& It) { return It.NodeHandle == NodeHandle; });
+		ensure(NodeMapping != nullptr);
+
+		const FNodeTemplate* NodeTemplate = NodeTemplateRegistry.Find(NodeMapping->NodeTemplateHandle);
+		ensure(NodeTemplate);
+
+		// Populate our node description into a temporary buffer
+		alignas(16) uint8 Buffer[64 * 1024];	// Max node size
+
+		FNodeDescription* NodeDesc = new(Buffer) FNodeDescription(NodeMapping->NodeUID, NodeMapping->NodeTemplateHandle);
+
+		// Populate our decorator properties
+		const uint32 NumDecorators = NodeTemplate->GetNumDecorators();
+		const FDecoratorTemplate* DecoratorTemplates = NodeTemplate->GetDecorators();
+		for (uint32 DecoratorIndex = 0; DecoratorIndex < NumDecorators; ++DecoratorIndex)
+		{
+			const FDecoratorRegistryHandle DecoratorHandle = DecoratorTemplates[DecoratorIndex].GetRegistryHandle();
+			const FDecorator* Decorator = DecoratorRegistry.Find(DecoratorHandle);
+
+			FAnimNextDecoratorSharedData* SharedData = DecoratorTemplates[DecoratorIndex].GetDecoratorDescription(*NodeDesc);
+
+			const TMap<FString, FString>& Properties = GetDecoratorProperties(DecoratorIndex);
+			Decorator->SaveDecoratorSharedData(*this, Properties, *SharedData);
+		}
+
+		// Write our node size first so the reader knows how much to skip ahead to the next node
+		uint32 NodeSize = NodeMapping->NodeSize;
+		*this << NodeSize;
+
+		// Append it to our archive
+		NodeDesc->Serialize(*this);
+
+		NumNodesWritten++;
+	}
+
+	FDecoratorWriter::EErrorState FDecoratorWriter::GetErrorState() const
+	{
+		return ErrorState;
+	}
+
+	const TArray<uint8>& FDecoratorWriter::GetGraphSharedData() const
+	{
+		return GraphSharedDataArchiveBuffer;
+	}
+}
+#endif
diff --git a/Engine/Plugins/Experimental/AnimNext/Source/AnimNext/Private/DecoratorBase/ExecutionContext.cpp b/Engine/Plugins/Experimental/AnimNext/Source/AnimNext/Private/DecoratorBase/ExecutionContext.cpp
--- a/Engine/Plugins/Experimental/AnimNext/Source/AnimNext/Private/DecoratorBase/ExecutionContext.cpp
+++ b/Engine/Plugins/Experimental/AnimNext/Source/AnimNext/Private/DecoratorBase/ExecutionContext.cpp
@@ -32,7 +32,7 @@ namespace UE::AnimNext
 		Private::GThreadLocalExecutionContext = nullptr;
 	}
 
-	FDecoratorPtr FExecutionContext::AllocateNodeInstance(FWeakDecoratorPtr ParentBinding, FDecoratorHandle ChildDecoratorHandle)
+	FDecoratorPtr FExecutionContext::AllocateNodeInstance(FWeakDecoratorPtr ParentBinding, FAnimNextDecoratorHandle ChildDecoratorHandle)
 	{
 		ensure(ChildDecoratorHandle.IsValid());
 		if (!ChildDecoratorHandle.IsValid())
@@ -96,10 +96,10 @@ namespace UE::AnimNext
 
 			FWeakDecoratorPtr DecoratorPtr(NodeInstance, DecoratorIndex);
 
-			const FDecoratorDescription* SharedData = DecoratorDesc->GetDecoratorDescription(NodeDesc);
-			FDecoratorInstance* InstanceData = DecoratorDesc->GetDecoratorInstance(*NodeInstance);
+			const FAnimNextDecoratorSharedData* SharedData = DecoratorDesc->GetDecoratorDescription(NodeDesc);
+			FDecoratorInstanceData* InstanceData = DecoratorDesc->GetDecoratorInstance(*NodeInstance);
 
-			Decorator->ConstructInstance(*this, DecoratorPtr, *SharedData, *InstanceData);
+			Decorator->ConstructDecoratorInstance(*this, DecoratorPtr, *SharedData, *InstanceData);
 		}
 
 		if (FailedDecoratorDesc != nullptr)
@@ -115,10 +115,10 @@ namespace UE::AnimNext
 
 				FWeakDecoratorPtr DecoratorPtr(NodeInstance, DecoratorIndex);
 
-				const FDecoratorDescription* SharedData = DecoratorDesc->GetDecoratorDescription(NodeDesc);
-				FDecoratorInstance* InstanceData = DecoratorDesc->GetDecoratorInstance(*NodeInstance);
+				const FAnimNextDecoratorSharedData* SharedData = DecoratorDesc->GetDecoratorDescription(NodeDesc);
+				FDecoratorInstanceData* InstanceData = DecoratorDesc->GetDecoratorInstance(*NodeInstance);
 
-				Decorator->DestructInstance(*this, DecoratorPtr, *SharedData, *InstanceData);
+				Decorator->DestructDecoratorInstance(*this, DecoratorPtr, *SharedData, *InstanceData);
 			}
 
 			FMemory::Free(NodeInstance);
@@ -166,10 +166,10 @@ namespace UE::AnimNext
 
 				FWeakDecoratorPtr DecoratorPtr(NodeInstance, DecoratorIndex);
 
-				const FDecoratorDescription* SharedData = DecoratorDesc->GetDecoratorDescription(NodeDesc);
-				FDecoratorInstance* InstanceData = DecoratorDesc->GetDecoratorInstance(*NodeInstance);
+				const FAnimNextDecoratorSharedData* SharedData = DecoratorDesc->GetDecoratorDescription(NodeDesc);
+				FDecoratorInstanceData* InstanceData = DecoratorDesc->GetDecoratorInstance(*NodeInstance);
 
-				Decorator->DestructInstance(*this, DecoratorPtr, *SharedData, *InstanceData);
+				Decorator->DestructDecoratorInstance(*this, DecoratorPtr, *SharedData, *InstanceData);
 			}
 		}
 
@@ -213,10 +213,10 @@ namespace UE::AnimNext
 				return false;	// Failed to find the matching decorator, did it get unregistered or is the decorator descriptor corrupted?
 			}
 
-			if (const IDecoratorInterface* Interface = Decorator->GetInterface(InterfaceUID))
+			if (const IDecoratorInterface* Interface = Decorator->GetDecoratorInterface(InterfaceUID))
 			{
-				const FDecoratorDescription* SharedData = DecoratorDesc->GetDecoratorDescription(NodeDesc);
-				FDecoratorInstance* InstanceData = DecoratorDesc->GetDecoratorInstance(*NodeInstance);
+				const FAnimNextDecoratorSharedData* SharedData = DecoratorDesc->GetDecoratorDescription(NodeDesc);
+				FDecoratorInstanceData* InstanceData = DecoratorDesc->GetDecoratorInstance(*NodeInstance);
 
 				const uint32 DecoratorIndex = DecoratorDesc - DecoratorDescs;
 				FWeakDecoratorPtr InterfaceDecoratorPtr(NodeInstance, DecoratorIndex);
@@ -271,10 +271,10 @@ namespace UE::AnimNext
 				return false;	// Failed to find the matching decorator, did it get unregistered or is the decorator descriptor corrupted?
 			}
 
-			if (const IDecoratorInterface* Interface = Decorator->GetInterface(InterfaceUID))
+			if (const IDecoratorInterface* Interface = Decorator->GetDecoratorInterface(InterfaceUID))
 			{
-				const FDecoratorDescription* SharedData = DecoratorDesc->GetDecoratorDescription(NodeDesc);
-				FDecoratorInstance* InstanceData = DecoratorDesc->GetDecoratorInstance(*NodeInstance);
+				const FAnimNextDecoratorSharedData* SharedData = DecoratorDesc->GetDecoratorDescription(NodeDesc);
+				FDecoratorInstanceData* InstanceData = DecoratorDesc->GetDecoratorInstance(*NodeInstance);
 
 				const uint32 DecoratorIndex = DecoratorDesc - DecoratorDescs;
 				FWeakDecoratorPtr SuperPtr(NodeInstance, DecoratorIndex);
diff --git a/Engine/Plugins/Experimental/AnimNext/Source/AnimNext/Private/DecoratorBase/NodeDescription.cpp b/Engine/Plugins/Experimental/AnimNext/Source/AnimNext/Private/DecoratorBase/NodeDescription.cpp
new file mode 100644
--- /dev/null
+++ b/Engine/Plugins/Experimental/AnimNext/Source/AnimNext/Private/DecoratorBase/NodeDescription.cpp
@@ -0,0 +1,52 @@
+// Copyright Epic Games, Inc. All Rights Reserved.
+
+#include "DecoratorBase/NodeDescription.h"
+
+#include "Serialization/Archive.h"
+#include "DecoratorBase/DecoratorRegistry.h"
+#include "DecoratorBase/NodeTemplate.h"
+#include "DecoratorBase/NodeTemplateRegistry.h"
+
+namespace UE::AnimNext
+{
+	void FNodeDescription::Serialize(FArchive& Ar)
+	{
+		Ar << UID;
+
+		if (Ar.IsSaving())
+		{
+			const FNodeTemplate* NodeTemplate = FNodeTemplateRegistry::Get().Find(TemplateHandle);
+
+			uint32 TemplateUID = NodeTemplate->GetUID();
+			Ar << TemplateUID;
+		}
+		else if (Ar.IsLoading())
+		{
+			uint32 TemplateUID = 0;
+			Ar << TemplateUID;
+
+			TemplateHandle = FNodeTemplateRegistry::Get().Find(TemplateUID);
+		}
+		else
+		{
+			// Counting, etc
+			int32 TemplateOffset = TemplateHandle.GetTemplateOffset();
+			Ar << TemplateOffset;
+		}
+
+		// Use our template to serialize our decorators
+		const FNodeTemplate* NodeTemplate = FNodeTemplateRegistry::Get().Find(TemplateHandle);
+
+		const uint32 NumDecorators = NodeTemplate->GetNumDecorators();
+		const FDecoratorTemplate* DecoratorTemplates = NodeTemplate->GetDecorators();
+		for (uint32 DecoratorIndex = 0; DecoratorIndex < NumDecorators; ++DecoratorIndex)
+		{
+			const FDecoratorRegistryHandle DecoratorHandle = DecoratorTemplates[DecoratorIndex].GetRegistryHandle();
+			const FDecorator* Decorator = FDecoratorRegistry::Get().Find(DecoratorHandle);
+
+			FAnimNextDecoratorSharedData* SharedData = DecoratorTemplates[DecoratorIndex].GetDecoratorDescription(*this);
+
+			Decorator->SerializeDecoratorSharedData(Ar, *SharedData);
+		}
+	}
+}
diff --git a/Engine/Plugins/Experimental/AnimNext/Source/AnimNext/Private/DecoratorBase/NodeTemplate.cpp b/Engine/Plugins/Experimental/AnimNext/Source/AnimNext/Private/DecoratorBase/NodeTemplate.cpp
new file mode 100644
--- /dev/null
+++ b/Engine/Plugins/Experimental/AnimNext/Source/AnimNext/Private/DecoratorBase/NodeTemplate.cpp
@@ -0,0 +1,21 @@
+// Copyright Epic Games, Inc. All Rights Reserved.
+
+#include "DecoratorBase/NodeTemplate.h"
+
+#include "Serialization/Archive.h"
+
+namespace UE::AnimNext
+{
+	void FNodeTemplate::Serialize(FArchive& Ar)
+	{
+		Ar << UID;
+		Ar << InstanceSize;
+		Ar << NumDecorators;
+
+		FDecoratorTemplate* DecoratorTemplates = GetDecorators();
+		for (uint32 DecoratorIndex = 0; DecoratorIndex < GetNumDecorators(); ++DecoratorIndex)
+		{
+			DecoratorTemplates[DecoratorIndex].Serialize(Ar);
+		}
+	}
+}
diff --git a/Engine/Plugins/Experimental/AnimNext/Source/AnimNext/Private/DecoratorBase/NodeTemplateRegistry.cpp b/Engine/Plugins/Experimental/AnimNext/Source/AnimNext/Private/DecoratorBase/NodeTemplateRegistry.cpp
--- a/Engine/Plugins/Experimental/AnimNext/Source/AnimNext/Private/DecoratorBase/NodeTemplateRegistry.cpp
+++ b/Engine/Plugins/Experimental/AnimNext/Source/AnimNext/Private/DecoratorBase/NodeTemplateRegistry.cpp
@@ -35,6 +35,12 @@ namespace UE::AnimNext
 		}
 	}
 
+	FNodeTemplateRegistryHandle FNodeTemplateRegistry::Find(uint32 NodeTemplateUID) const
+	{
+		const FNodeTemplateRegistryHandle* TemplateHandle = TemplateUIDToHandleMap.Find(NodeTemplateUID);
+		return TemplateHandle != nullptr ? *TemplateHandle : FNodeTemplateRegistryHandle();
+	}
+
 	FNodeTemplateRegistryHandle FNodeTemplateRegistry::FindOrAdd(const FNodeTemplate* NodeTemplate)
 	{
 		if (NodeTemplate == nullptr)
@@ -75,6 +81,16 @@ namespace UE::AnimNext
 		TemplateUIDToHandleMap.Remove(TemplateUID);
 	}
 
+	FNodeTemplate* FNodeTemplateRegistry::Find(FNodeTemplateRegistryHandle TemplateHandle)
+	{
+		if (!TemplateHandle.IsValid())
+		{
+			return nullptr;
+		}
+
+		return reinterpret_cast<FNodeTemplate*>(&TemplateBuffer[TemplateHandle.GetTemplateOffset()]);
+	}
+
 	const FNodeTemplate* FNodeTemplateRegistry::Find(FNodeTemplateRegistryHandle TemplateHandle) const
 	{
 		if (!TemplateHandle.IsValid())
@@ -89,4 +105,10 @@ namespace UE::AnimNext
 	{
 		return TemplateUIDToHandleMap.Num();
 	}
+
+	void FNodeTemplateRegistry::Clear()
+	{
+		TemplateBuffer.Empty(0);
+		TemplateUIDToHandleMap.Empty(0);
+	}
 }
diff --git a/Engine/Plugins/Experimental/AnimNext/Source/AnimNext/Private/Decorators/BlendTwoWay.h b/Engine/Plugins/Experimental/AnimNext/Source/AnimNext/Private/Decorators/BlendTwoWay.h
--- a/Engine/Plugins/Experimental/AnimNext/Source/AnimNext/Private/Decorators/BlendTwoWay.h
+++ b/Engine/Plugins/Experimental/AnimNext/Source/AnimNext/Private/Decorators/BlendTwoWay.h
@@ -3,11 +3,26 @@
 #pragma once
 
 #include "CoreMinimal.h"
+
 #include "DecoratorBase/Decorator.h"
 #include "DecoratorInterfaces/IEvaluate.h"
 #include "DecoratorInterfaces/IHierarchy.h"
 #include "DecoratorInterfaces/IUpdate.h"
 
+#include "BlendTwoWay.generated.h"
+
+USTRUCT()
+struct FAnimNextBlendTwoWayDecoratorSharedData : public FAnimNextDecoratorSharedData
+{
+	GENERATED_BODY()
+
+	UPROPERTY()
+	FAnimNextDecoratorHandle Children[2];
+
+	UPROPERTY()
+	double BlendWeight = 0.0;
+};
+
 namespace UE::AnimNext
 {
 	/**
@@ -15,24 +30,17 @@ namespace UE::AnimNext
 	 * 
 	 * A decorator that can blend two inputs.
 	 */
-	struct FBlendTwoWayDecorator : FDecorator, IEvaluate, IUpdate, IHierarchy
+	struct FBlendTwoWayDecorator : FBaseDecorator, IEvaluate, IUpdate, IHierarchy
 	{
-		DECLARE_ANIM_DECORATOR(FBlendTwoWayDecorator, 0x96a81d1e, FDecorator)
+		DECLARE_ANIM_DECORATOR(FBlendTwoWayDecorator, 0x96a81d1e, FBaseDecorator)
 
-		struct FSharedData : FDecorator::FSharedData
-		{
-			FDecoratorHandle Children[2];
-			double BlendWeight = 0.0;
-		};
+		using FSharedData = FAnimNextBlendTwoWayDecoratorSharedData;
 
 		struct FInstanceData : FDecorator::FInstanceData
 		{
 			FDecoratorPtr Children[2];
 		};
 
-		// FDecorator impl
-		virtual EDecoratorMode GetMode() const override { return EDecoratorMode::Base; }
-
 		// IEvaluate impl
 		virtual void PostEvaluate(FExecutionContext& Context, const TDecoratorBinding<IEvaluate>& Binding) const override;
 
diff --git a/Engine/Plugins/Experimental/AnimNext/Source/AnimNext/Private/Decorators/SequencePlayer.h b/Engine/Plugins/Experimental/AnimNext/Source/AnimNext/Private/Decorators/SequencePlayer.h
--- a/Engine/Plugins/Experimental/AnimNext/Source/AnimNext/Private/Decorators/SequencePlayer.h
+++ b/Engine/Plugins/Experimental/AnimNext/Source/AnimNext/Private/Decorators/SequencePlayer.h
@@ -9,8 +9,22 @@
 #include "DecoratorInterfaces/ITimeline.h"
 #include "DecoratorInterfaces/IUpdate.h"
 
+#include "SequencePlayer.generated.h"
+
 class UAnimSequence;
 
+USTRUCT()
+struct FAnimNextSequencePlayerDecoratorSharedData : public FAnimNextDecoratorSharedData
+{
+	GENERATED_BODY()
+
+	UPROPERTY()
+	TObjectPtr<UAnimSequence> AnimSeq;
+
+	UPROPERTY()
+	double PlayRate = 1.0;
+};
+
 namespace UE::AnimNext
 {
 	/**
@@ -18,25 +32,17 @@ namespace UE::AnimNext
 	 * 
 	 * A decorator that can play an animation sequence.
 	 */
-	struct FSequencePlayerDecorator : FDecorator, IEvaluate, ITimeline, IUpdate
+	struct FSequencePlayerDecorator : FBaseDecorator, IEvaluate, ITimeline, IUpdate
 	{
-		DECLARE_ANIM_DECORATOR(FSequencePlayerDecorator, 0xa628ad12, FDecorator)
+		DECLARE_ANIM_DECORATOR(FSequencePlayerDecorator, 0xa628ad12, FBaseDecorator)
 
-		struct FSharedData : FDecorator::FSharedData
-		{
-			TObjectPtr<UAnimSequence> AnimSeq;
-
-			double PlayRate = 1.0;
-		};
+		using FSharedData = FAnimNextSequencePlayerDecoratorSharedData;
 
 		struct FInstanceData : FDecorator::FInstanceData
 		{
 			double CurrentTime = 0.0;
 		};
 
-		// FDecorator impl
-		virtual EDecoratorMode GetMode() const override { return EDecoratorMode::Base; }
-
 		// IEvaluate impl
 		virtual void PreEvaluate(FExecutionContext& Context, const TDecoratorBinding<IEvaluate>& Binding) const override;
 
diff --git a/Engine/Plugins/Experimental/AnimNext/Source/AnimNext/Public/DecoratorBase/Decorator.h b/Engine/Plugins/Experimental/AnimNext/Source/AnimNext/Public/DecoratorBase/Decorator.h
--- a/Engine/Plugins/Experimental/AnimNext/Source/AnimNext/Public/DecoratorBase/Decorator.h
+++ b/Engine/Plugins/Experimental/AnimNext/Source/AnimNext/Public/DecoratorBase/Decorator.h
@@ -3,26 +3,38 @@
 #pragma once
 
 #include "CoreMinimal.h"
-#include "DecoratorBase/DecoratorDescription.h"
+
 #include "DecoratorBase/DecoratorHandle.h"			// Derived types are likely to refer to other decorators as children
-#include "DecoratorBase/DecoratorInstance.h"
+#include "DecoratorBase/DecoratorInstanceData.h"
 #include "DecoratorBase/DecoratorMode.h"
 #include "DecoratorBase/DecoratorPtr.h"
+#include "DecoratorBase/DecoratorSharedData.h"
 #include "DecoratorBase/DecoratorUID.h"
 #include "DecoratorBase/IDecoratorInterface.h"
 
+class FArchive;
+
 // Helper macros
 // In the decorator class declaration, this macro declares the Super alias and base functions we override
 #define DECLARE_ANIM_DECORATOR(DecoratorName, DecoratorNameHash, SuperDecoratorName) \
-	using Super = SuperDecoratorName; \
+	using DecoratorSuper = SuperDecoratorName; \
 	/* FDecorator impl */ \
 	static constexpr UE::AnimNext::FDecoratorUID DecoratorUID = UE::AnimNext::FDecoratorUID(TEXT(#DecoratorName), DecoratorNameHash); \
 	virtual UE::AnimNext::FDecoratorUID GetDecoratorUID() const override { return DecoratorUID; } \
 	static const UE::AnimNext::FDecoratorMemoryLayout DecoratorMemoryDescription; \
 	virtual UE::AnimNext::FDecoratorMemoryLayout GetDecoratorMemoryDescription() const override { return DecoratorMemoryDescription; } \
-	virtual void ConstructInstance(UE::AnimNext::FExecutionContext& Context, UE::AnimNext::FWeakDecoratorPtr DecoratorPtr, const UE::AnimNext::FDecoratorDescription& DecoratorDesc, UE::AnimNext::FDecoratorInstance& DecoratorInstance) const override; \
-	virtual void DestructInstance(UE::AnimNext::FExecutionContext& Context, UE::AnimNext::FWeakDecoratorPtr DecoratorPtr, const UE::AnimNext::FDecoratorDescription& DecoratorDesc, UE::AnimNext::FDecoratorInstance& DecoratorInstance) const override; \
-	virtual const UE::AnimNext::IDecoratorInterface* GetInterface(UE::AnimNext::FDecoratorInterfaceUID InterfaceUID) const override;
+	virtual UScriptStruct* GetDecoratorSharedDataStruct() const override { return FSharedData::StaticStruct(); } \
+	virtual void ConstructDecoratorInstance(UE::AnimNext::FExecutionContext& Context, UE::AnimNext::FWeakDecoratorPtr DecoratorPtr, const FAnimNextDecoratorSharedData& DecoratorDesc, UE::AnimNext::FDecoratorInstanceData& DecoratorInstance) const override; \
+	virtual void DestructDecoratorInstance(UE::AnimNext::FExecutionContext& Context, UE::AnimNext::FWeakDecoratorPtr DecoratorPtr, const FAnimNextDecoratorSharedData& DecoratorDesc, UE::AnimNext::FDecoratorInstanceData& DecoratorInstance) const override; \
+	virtual const UE::AnimNext::IDecoratorInterface* GetDecoratorInterface(UE::AnimNext::FDecoratorInterfaceUID InterfaceUID) const override; \
+	static_assert(std::is_base_of<FAnimNextDecoratorSharedData, FSharedData>::value, TEXT("Decorator shared data must derive from FAnimNextDecoratorSharedData")); \
+	static_assert(std::is_base_of<FDecoratorInstanceData, FInstanceData>::value, TEXT("Decorator instance data must derive from FDecoratorInstanceData"));
+
+#define DECLARE_ABSTRACT_ANIM_DECORATOR(DecoratorName, DecoratorNameHash, SuperDecoratorName) \
+	using DecoratorSuper = SuperDecoratorName; \
+	/* FDecorator impl */ \
+	static constexpr UE::AnimNext::FDecoratorUID DecoratorUID = UE::AnimNext::FDecoratorUID(TEXT(#DecoratorName), DecoratorNameHash); \
+	virtual UE::AnimNext::FDecoratorUID GetDecoratorUID() const override { return DecoratorUID; }
 
 // In the decorator cpp, these three macros implement the base functionality
 // 
@@ -36,23 +48,23 @@
 #define DEFINE_ANIM_DECORATOR_BEGIN(DecoratorName) \
 	const UE::AnimNext::FDecoratorMemoryLayout DecoratorName::DecoratorMemoryDescription = \
 		UE::AnimNext::FDecoratorMemoryLayout{ sizeof(DecoratorName), alignof(DecoratorName), sizeof(DecoratorName::FSharedData), alignof(DecoratorName::FSharedData), sizeof(DecoratorName::FInstanceData), alignof(DecoratorName::FInstanceData) }; \
-	void DecoratorName::ConstructInstance(UE::AnimNext::FExecutionContext& Context, UE::AnimNext::FWeakDecoratorPtr DecoratorPtr, const UE::AnimNext::FDecoratorDescription& DecoratorDesc, UE::AnimNext::FDecoratorInstance& DecoratorInstance) const \
+	void DecoratorName::ConstructDecoratorInstance(UE::AnimNext::FExecutionContext& Context, UE::AnimNext::FWeakDecoratorPtr DecoratorPtr, const FAnimNextDecoratorSharedData& DecoratorDesc, UE::AnimNext::FDecoratorInstanceData& DecoratorInstance) const \
 	{ \
 		FInstanceData* Data = new(&static_cast<FInstanceData&>(DecoratorInstance)) FInstanceData(); \
 		Data->Construct(Context, DecoratorPtr, static_cast<const FSharedData&>(DecoratorDesc)); \
 	} \
-	void DecoratorName::DestructInstance(UE::AnimNext::FExecutionContext& Context, UE::AnimNext::FWeakDecoratorPtr DecoratorPtr, const UE::AnimNext::FDecoratorDescription& DecoratorDesc, UE::AnimNext::FDecoratorInstance& DecoratorInstance) const \
+	void DecoratorName::DestructDecoratorInstance(UE::AnimNext::FExecutionContext& Context, UE::AnimNext::FWeakDecoratorPtr DecoratorPtr, const FAnimNextDecoratorSharedData& DecoratorDesc, UE::AnimNext::FDecoratorInstanceData& DecoratorInstance) const \
 	{ \
 		FInstanceData& Data = static_cast<FInstanceData&>(DecoratorInstance); \
 		Data.Destruct(Context, DecoratorPtr, DecoratorDesc); \
 		Data.~FInstanceData(); \
 	} \
-	const UE::AnimNext::IDecoratorInterface* DecoratorName::GetInterface(UE::AnimNext::FDecoratorInterfaceUID InterfaceUID_) const \
+	const UE::AnimNext::IDecoratorInterface* DecoratorName::GetDecoratorInterface(UE::AnimNext::FDecoratorInterfaceUID InterfaceUID_) const \
 	{
 
 #define DEFINE_ANIM_DECORATOR_END(DecoratorName) \
 		/* Forward to base implementation */ \
-		return Super::GetInterface(InterfaceUID_); \
+		return DecoratorSuper::GetDecoratorInterface(InterfaceUID_); \
 	}
 
 #define DEFINE_ANIM_DECORATOR_IMPLEMENTS_INTERFACE(InterfaceName) \
@@ -74,6 +86,8 @@ namespace UE::AnimNext
 {
 	struct FDecorator;
 	struct FDecoratorMemoryLayout;
+	class FDecoratorReader;
+	class FDecoratorWriter;
 	struct FExecutionContext;
 
 	// A function pointer to a shim to construct a decorator into the desired memory location
@@ -122,16 +136,18 @@ namespace UE::AnimNext
 	 * on all worker threads.
 	 * 
 	 * Decorators can have shared read-only data that all instances of a graph can use (e.g. hard-coded properties).
-	 * Shared data must derive from FDecoratorDescription.
+	 * Shared data must derive from FAnimNextDecoratorSharedData.
 	 * Decorators can have instance data (e.g. blend weight).
-	 * Instance data must derive from FDecoratorInstance.
+	 * Instance data must derive from FDecoratorInstanceData.
 	 */
 	struct ANIMNEXT_API FDecorator
 	{
+		virtual ~FDecorator() {}
+
 		// Empty shared/instance data types
 		// Derived types must define an alias for these
-		using FSharedData = FDecoratorDescription;
-		using FInstanceData = FDecoratorInstance;
+		using FSharedData = FAnimNextDecoratorSharedData;
+		using FInstanceData = FDecoratorInstanceData;
 
 		// The globally unique UID for this decorator
 		// Derived types will have their own DecoratorUID member that hides/aliases/shadows this one
@@ -144,18 +160,20 @@ namespace UE::AnimNext
 		// Returns the memory requirements of the derived decorator instance
 		virtual FDecoratorMemoryLayout GetDecoratorMemoryDescription() const { return { sizeof(FDecorator), alignof(FDecorator) }; }
 
+		// Returns the UScriptStruct associated with the shared data for the decorator
+		virtual UScriptStruct* GetDecoratorSharedDataStruct() const { return FSharedData::StaticStruct(); }
+
 		// Called when a new instance of the decorator is created or destroyed
 		// Derived types must override this and forward to the instance data constructor/destructor
-		virtual void ConstructInstance(FExecutionContext& Context, FWeakDecoratorPtr DecoratorPtr, const FDecoratorDescription& DecoratorDesc, FDecoratorInstance& DecoratorInstance) const = 0;
-		virtual void DestructInstance(FExecutionContext& Context, FWeakDecoratorPtr DecoratorPtr, const FDecoratorDescription& DecoratorDesc, FDecoratorInstance& DecoratorInstance) const = 0;
+		virtual void ConstructDecoratorInstance(FExecutionContext& Context, FWeakDecoratorPtr DecoratorPtr, const FAnimNextDecoratorSharedData& DecoratorDesc, FDecoratorInstanceData& DecoratorInstance) const = 0;
+		virtual void DestructDecoratorInstance(FExecutionContext& Context, FWeakDecoratorPtr DecoratorPtr, const FAnimNextDecoratorSharedData& DecoratorDesc, FDecoratorInstanceData& DecoratorInstance) const = 0;
 
-		// Returns the decorator mode for this derived type.
-		// TODO: This belongs in the editor side of the decorators since we need this information for our cook, here for now
-		virtual EDecoratorMode GetMode() const = 0;
+		// Returns the decorator mode.
+		virtual EDecoratorMode GetDecoratorMode() const = 0;
 
 		// Returns a pointer to the specified interface if it is supported.
 		// Derived types must override this.
-		virtual const IDecoratorInterface* GetInterface(FDecoratorInterfaceUID InterfaceUID) const
+		virtual const IDecoratorInterface* GetDecoratorInterface(FDecoratorInterfaceUID InterfaceUID) const
 		{
 			// TODO:
 			// if/else sequence with static_casts to get the right v-table
@@ -179,7 +197,32 @@ namespace UE::AnimNext
 			return nullptr;
 		}
 
-		virtual ~FDecorator() {}
+		// Called to serialize decorator shared data
+		virtual void SerializeDecoratorSharedData(FArchive& Ar, FAnimNextDecoratorSharedData& SharedData) const;
+
+#if WITH_EDITOR
+		// Takes the editor properties as authored in the graph and converts them into an instance of the FAnimNextDecoratorSharedData
+		// derived type using UE reflection.
+		// Decorators can override this function to control how editor only properties are coerced into the runtime shared data
+		// instance.
+		virtual void SaveDecoratorSharedData(FDecoratorWriter& Writer, const TMap<FString, FString>& Properties, FAnimNextDecoratorSharedData& OutSharedData) const;
+#endif
+	};
+
+	// Base class for base decorators that are standalone
+	struct ANIMNEXT_API FBaseDecorator : FDecorator
+	{
+		DECLARE_ABSTRACT_ANIM_DECORATOR(FBaseDecorator, 0xd23dcf79, FDecorator)
+
+		virtual EDecoratorMode GetDecoratorMode() const override { return EDecoratorMode::Base; }
+	};
+
+	// Base class for additive decorators that override behavior of other decorators
+	struct ANIMNEXT_API FAdditiveDecorator : FDecorator
+	{
+		DECLARE_ABSTRACT_ANIM_DECORATOR(FAdditiveDecorator, 0x7ab3732a, FDecorator)
+
+		virtual EDecoratorMode GetDecoratorMode() const override { return EDecoratorMode::Additive; }
 	};
 
 	/**
diff --git a/Engine/Plugins/Experimental/AnimNext/Source/AnimNext/Public/DecoratorBase/DecoratorBinding.h b/Engine/Plugins/Experimental/AnimNext/Source/AnimNext/Public/DecoratorBase/DecoratorBinding.h
--- a/Engine/Plugins/Experimental/AnimNext/Source/AnimNext/Public/DecoratorBase/DecoratorBinding.h
+++ b/Engine/Plugins/Experimental/AnimNext/Source/AnimNext/Public/DecoratorBase/DecoratorBinding.h
@@ -6,10 +6,11 @@
 #include "DecoratorBase/DecoratorPtr.h"
 #include "DecoratorBase/DecoratorInterfaceUID.h"
 
+struct FAnimNextDecoratorSharedData;
+
 namespace UE::AnimNext
 {
-	struct FDecoratorDescription;
-	struct FDecoratorInstance;
+	struct FDecoratorInstanceData;
 	struct IDecoratorInterface;
 
 	/**
@@ -48,7 +49,7 @@ namespace UE::AnimNext
 
 	protected:
 		// Creates a valid binding
-		FDecoratorBinding(const IDecoratorInterface* Interface_, const FDecoratorDescription* SharedData_, FDecoratorInstance* InstanceData_, FWeakDecoratorPtr DecoratorPtr_)
+		FDecoratorBinding(const IDecoratorInterface* Interface_, const FAnimNextDecoratorSharedData* SharedData_, FDecoratorInstanceData* InstanceData_, FWeakDecoratorPtr DecoratorPtr_)
 			: Interface(Interface_)
 			, SharedData(SharedData_)
 			, InstanceData(InstanceData_)
@@ -61,8 +62,8 @@ namespace UE::AnimNext
 
 		const IDecoratorInterface*		Interface = nullptr;
 
-		const FDecoratorDescription*	SharedData = nullptr;
-		FDecoratorInstance*				InstanceData = nullptr;
+		const FAnimNextDecoratorSharedData*		SharedData = nullptr;
+		FDecoratorInstanceData*			InstanceData = nullptr;
 
 		FWeakDecoratorPtr				DecoratorPtr;
 
diff --git a/Engine/Plugins/Experimental/AnimNext/Source/AnimNext/Public/DecoratorBase/DecoratorDescription.h b/Engine/Plugins/Experimental/AnimNext/Source/AnimNext/Public/DecoratorBase/DecoratorDescription.h
deleted file mode 100644
diff --git a/Engine/Plugins/Experimental/AnimNext/Source/AnimNext/Public/DecoratorBase/DecoratorHandle.h b/Engine/Plugins/Experimental/AnimNext/Source/AnimNext/Public/DecoratorBase/DecoratorHandle.h
--- a/Engine/Plugins/Experimental/AnimNext/Source/AnimNext/Public/DecoratorBase/DecoratorHandle.h
+++ b/Engine/Plugins/Experimental/AnimNext/Source/AnimNext/Public/DecoratorBase/DecoratorHandle.h
@@ -3,50 +3,52 @@
 #pragma once
 
 #include "CoreMinimal.h"
+
 #include "DecoratorBase/NodeHandle.h"
 
-namespace UE::AnimNext
+#include "DecoratorHandle.generated.h"
+
+/**
+ * Decorator Handle
+ * A decorator handle represents a reference to a specific decorator instance in the shared/read-only portion
+ * of a sub-graph. It points to a FNodeDescription when resolved.
+ * @see FNodeDescription
+ */
+USTRUCT()
+struct FAnimNextDecoratorHandle
 {
-	/**
-	 * Decorator Handle
-	 * A decorator handle represents a reference to a specific decorator instance in the shared/read-only portion
-	 * of a sub-graph. It points to a FNodeDescription when resolved.
-	 * @see FNodeDescription
-	 */
-	struct FDecoratorHandle
-	{
-		// Creates an invalid decorator handle
-		FDecoratorHandle()
-			: DecoratorIndex(0)
-			, SharedOffset(INVALID_SHARED_OFFSET_VALUE)
-		{}
-
-		// Creates a decorator handle pointing to the first decorator of the specified node
-		explicit FDecoratorHandle(FNodeHandle NodeHandle)
-			: DecoratorIndex(0)
-			, SharedOffset(NodeHandle.GetSharedOffset())
-		{}
-
-		// Creates a decorator handle pointing to the specified decorator on the specified node
-		FDecoratorHandle(FNodeHandle NodeHandle, uint32 DecoratorIndex_)
-			: DecoratorIndex(DecoratorIndex_)
-			, SharedOffset(NodeHandle.GetSharedOffset())
-		{}
-
-		// Returns true if this decorator handle is valid, false otherwise
-		bool IsValid() const { return SharedOffset != INVALID_SHARED_OFFSET_VALUE; }
-
-		// Returns the decorator index
-		uint32 GetDecoratorIndex() const { return DecoratorIndex; }
-
-		// Returns a handle to the node in the shared data segment
-		FNodeHandle GetNodeHandle() const { return FNodeHandle(SharedOffset); }
-
-	private:
-		// We pack the shared offset on 24 bits, the top 8 bits are truncated
-		static constexpr uint32 INVALID_SHARED_OFFSET_VALUE = ~0u >> (32 - 24);
-
-		uint32	DecoratorIndex : 8;
-		uint32	SharedOffset : 24;		// relative to root of sub-graph
-	};
-}
+	GENERATED_BODY()
+
+	// Creates an invalid decorator handle
+	FAnimNextDecoratorHandle()
+		: PackedDecoratorIndexAndNodeSharedOffset(INVALID_SHARED_OFFSET_VALUE)
+	{}
+
+	// Creates a decorator handle pointing to the first decorator of the specified node
+	explicit FAnimNextDecoratorHandle(UE::AnimNext::FNodeHandle NodeHandle)
+		: PackedDecoratorIndexAndNodeSharedOffset(NodeHandle.GetSharedOffset())
+	{}
+
+	// Creates a decorator handle pointing to the specified decorator on the specified node
+	FAnimNextDecoratorHandle(UE::AnimNext::FNodeHandle NodeHandle, uint32 DecoratorIndex_)
+		: PackedDecoratorIndexAndNodeSharedOffset(NodeHandle.GetSharedOffset() | (DecoratorIndex_ << 24))
+	{}
+
+	// Returns true if this decorator handle is valid, false otherwise
+	bool IsValid() const { return (PackedDecoratorIndexAndNodeSharedOffset & SHARED_OFFSET_MASK) != INVALID_SHARED_OFFSET_VALUE; }
+
+	// Returns the decorator index
+	uint32 GetDecoratorIndex() const { return PackedDecoratorIndexAndNodeSharedOffset >> DECORATOR_SHIFT; }
+
+	// Returns a handle to the node in the shared data segment
+	UE::AnimNext::FNodeHandle GetNodeHandle() const { return UE::AnimNext::FNodeHandle(PackedDecoratorIndexAndNodeSharedOffset & SHARED_OFFSET_MASK); }
+
+private:
+	// Bottom 24 bits are used by the node shared offset while the top 8 bits by the decorator index
+	static constexpr uint32 DECORATOR_SHIFT = 24;
+	static constexpr uint32 SHARED_OFFSET_MASK = ~0u >> (32 - DECORATOR_SHIFT);
+	static constexpr uint32 INVALID_SHARED_OFFSET_VALUE = SHARED_OFFSET_MASK;
+
+	UPROPERTY()
+	uint32		PackedDecoratorIndexAndNodeSharedOffset;
+};
diff --git a/Engine/Plugins/Experimental/AnimNext/Source/AnimNext/Public/DecoratorBase/DecoratorInstance.h b/Engine/Plugins/Experimental/AnimNext/Source/AnimNext/Public/DecoratorBase/DecoratorInstance.h
deleted file mode 100644
diff --git a/Engine/Plugins/Experimental/AnimNext/Source/AnimNext/Public/DecoratorBase/DecoratorInstanceData.h b/Engine/Plugins/Experimental/AnimNext/Source/AnimNext/Public/DecoratorBase/DecoratorInstanceData.h
new file mode 100644
--- /dev/null
+++ b/Engine/Plugins/Experimental/AnimNext/Source/AnimNext/Public/DecoratorBase/DecoratorInstanceData.h
@@ -0,0 +1,44 @@
+// Copyright Epic Games, Inc. All Rights Reserved.
+
+#pragma once
+
+#include "CoreMinimal.h"
+#include "DecoratorBase/DecoratorPtr.h"
+
+struct FAnimNextDecoratorSharedData;
+
+namespace UE::AnimNext
+{
+	struct FExecutionContext;
+
+	/**
+	 * FDecoratorInstanceData
+	 * A decorator instance represents allocated data for specific decorator instance.
+	 * 
+	 * @see FNodeInstance
+	 * 
+	 * A FDecoratorInstanceData is the base type that decorator instance data derives from.
+	 */
+	struct FDecoratorInstanceData
+	{
+		// Called after the constructor has been called when a new instance is created.
+		// This is called after the default constructor.
+		// You can override this function by adding a new one with the same name on your
+		// derived type. You can also specialize the FAnimNextDecoratorSharedData with the derived
+		// version as long as the reference can coerce to 'const FAnimNextDecoratorSharedData&'.
+		// Decorators are constructed from the bottom to the top.
+		void Construct(FExecutionContext& Context, FWeakDecoratorPtr DecoratorPtr, const FAnimNextDecoratorSharedData& DecoratorDesc)
+		{
+		}
+
+		// Called before the destructor has been called when an instance is destroyed.
+		// This is called before the default destructor.
+		// You can override this function by adding a new one with the same name on your
+		// derived type. You can also specialize the FAnimNextDecoratorSharedData with the derived
+		// version as long as the reference can coerce to 'const FAnimNextDecoratorSharedData&'.
+		// Decorators are destructed from the top to the bottom.
+		void Destruct(FExecutionContext& Context, FWeakDecoratorPtr DecoratorPtr, const FAnimNextDecoratorSharedData& DecoratorDesc)
+		{
+		}
+	};
+}
diff --git a/Engine/Plugins/Experimental/AnimNext/Source/AnimNext/Public/DecoratorBase/DecoratorPtr.h b/Engine/Plugins/Experimental/AnimNext/Source/AnimNext/Public/DecoratorBase/DecoratorPtr.h
--- a/Engine/Plugins/Experimental/AnimNext/Source/AnimNext/Public/DecoratorBase/DecoratorPtr.h
+++ b/Engine/Plugins/Experimental/AnimNext/Source/AnimNext/Public/DecoratorBase/DecoratorPtr.h
@@ -13,14 +13,14 @@ namespace UE::AnimNext
 	 * Decorator Pointer
 	 * A decorator pointer represents a shared pointer to allocated instance data.
 	 * It manages reference counting.
-	 * It points to a FDecoratorInstance or FNodeInstance when resolved.
+	 * It points to a FDecoratorInstanceData or FNodeInstance when resolved.
 	 *
 	 * A node pointer can also be weak, meaning that it does not update reference counting.
 	 * Note that weak pointers will point to garbage if the decorator instance is destroyed.
 	 *
 	 * Weak decorator pointers should be used carefully!
 	 * 
-	 * @see FNodeInstance, FDecoratorInstance
+	 * @see FNodeInstance, FDecoratorInstanceData
 	 */
 	struct ANIMNEXT_API FDecoratorPtr
 	{
diff --git a/Engine/Plugins/Experimental/AnimNext/Source/AnimNext/Public/DecoratorBase/DecoratorSharedData.h b/Engine/Plugins/Experimental/AnimNext/Source/AnimNext/Public/DecoratorBase/DecoratorSharedData.h
new file mode 100644
--- /dev/null
+++ b/Engine/Plugins/Experimental/AnimNext/Source/AnimNext/Public/DecoratorBase/DecoratorSharedData.h
@@ -0,0 +1,24 @@
+// Copyright Epic Games, Inc. All Rights Reserved.
+
+#pragma once
+
+#include "CoreMinimal.h"
+
+#include "DecoratorSharedData.generated.h"
+
+/**
+ * FAnimNextDecoratorSharedData
+ * Decorator shared data represents a unique instance in the authored static graph.
+ * Each instance of a graph will share instances of the read-only shared data.
+ * Shared data typically contains hardcoded properties, RigVM latent pin information,
+ * or pooled properties shared between multiple decorators.
+ * 
+ * @see FNodeDescription
+ *
+ * A FAnimNextDecoratorSharedData is the base type that decorator shared data derives from.
+ */
+USTRUCT()
+struct FAnimNextDecoratorSharedData
+{
+	GENERATED_BODY()
+};
diff --git a/Engine/Plugins/Experimental/AnimNext/Source/AnimNext/Public/DecoratorBase/DecoratorUID.h b/Engine/Plugins/Experimental/AnimNext/Source/AnimNext/Public/DecoratorBase/DecoratorUID.h
--- a/Engine/Plugins/Experimental/AnimNext/Source/AnimNext/Public/DecoratorBase/DecoratorUID.h
+++ b/Engine/Plugins/Experimental/AnimNext/Source/AnimNext/Public/DecoratorBase/DecoratorUID.h
@@ -45,10 +45,6 @@ namespace UE::AnimNext
 		// Returns whether this UID is valid or not
 		constexpr bool IsValid() const { return UID != 0; }
 
-		// Compares for equality and inequality
-		constexpr bool operator==(FDecoratorUID RHS) const { return UID == RHS.UID; }
-		constexpr bool operator!=(FDecoratorUID RHS) const { return UID != RHS.UID; }
-
 	private:
 		uint32	UID;
 
@@ -56,4 +52,12 @@ namespace UE::AnimNext
 		const TCHAR* DecoratorName;
 #endif
 	};
+
+	// Compares for equality and inequality
+	constexpr bool operator==(FDecoratorUID LHS, FDecoratorUID RHS) { return LHS.GetUID() == RHS.GetUID(); }
+	constexpr bool operator!=(FDecoratorUID LHS, FDecoratorUID RHS) { return LHS.GetUID() != RHS.GetUID(); }
+	constexpr bool operator==(FDecoratorUID LHS, uint32 RHS) { return LHS.GetUID() == RHS; }
+	constexpr bool operator!=(FDecoratorUID LHS, uint32 RHS) { return LHS.GetUID() != RHS; }
+	constexpr bool operator==(uint32 LHS, FDecoratorUID RHS) { return LHS == RHS.GetUID(); }
+	constexpr bool operator!=(uint32 LHS, FDecoratorUID RHS) { return LHS != RHS.GetUID(); }
 }
diff --git a/Engine/Plugins/Experimental/AnimNext/Source/AnimNext/Public/DecoratorBase/ExecutionContext.h b/Engine/Plugins/Experimental/AnimNext/Source/AnimNext/Public/DecoratorBase/ExecutionContext.h
--- a/Engine/Plugins/Experimental/AnimNext/Source/AnimNext/Public/DecoratorBase/ExecutionContext.h
+++ b/Engine/Plugins/Experimental/AnimNext/Source/AnimNext/Public/DecoratorBase/ExecutionContext.h
@@ -53,11 +53,11 @@ namespace UE::AnimNext
 
 		// Allocates a new node instance from a decorator handle
 		// If the desired decorator lives in the current parent, a weak handle to it will be returned
-		FDecoratorPtr AllocateNodeInstance(const FDecoratorBinding& ParentBinding, FDecoratorHandle ChildDecoratorHandle);
+		FDecoratorPtr AllocateNodeInstance(const FDecoratorBinding& ParentBinding, FAnimNextDecoratorHandle ChildDecoratorHandle);
 
 		// Allocates a new node instance from a decorator handle
 		// If the desired decorator lives in the current parent, a weak handle to it will be returned
-		FDecoratorPtr AllocateNodeInstance(FWeakDecoratorPtr ParentBinding, FDecoratorHandle ChildDecoratorHandle);
+		FDecoratorPtr AllocateNodeInstance(FWeakDecoratorPtr ParentBinding, FAnimNextDecoratorHandle ChildDecoratorHandle);
 
 		// Releases a node instance that is no longer referenced
 		void ReleaseNodeInstance(FNodeInstance* Node);
@@ -115,7 +115,7 @@ namespace UE::AnimNext
 		return GetInterfaceSuper<DecoratorInterface>(Binding.GetDecoratorPtr(), SuperBinding);
 	}
 
-	inline FDecoratorPtr FExecutionContext::AllocateNodeInstance(const FDecoratorBinding& ParentBinding, FDecoratorHandle ChildDecoratorHandle)
+	inline FDecoratorPtr FExecutionContext::AllocateNodeInstance(const FDecoratorBinding& ParentBinding, FAnimNextDecoratorHandle ChildDecoratorHandle)
 	{
 		return AllocateNodeInstance(ParentBinding.GetDecoratorPtr(), ChildDecoratorHandle);
 	}
diff --git a/Engine/Plugins/Experimental/AnimNext/Source/AnimNext/Public/DecoratorBase/NodeDescription.h b/Engine/Plugins/Experimental/AnimNext/Source/AnimNext/Public/DecoratorBase/NodeDescription.h
--- a/Engine/Plugins/Experimental/AnimNext/Source/AnimNext/Public/DecoratorBase/NodeDescription.h
+++ b/Engine/Plugins/Experimental/AnimNext/Source/AnimNext/Public/DecoratorBase/NodeDescription.h
@@ -3,13 +3,14 @@
 #pragma once
 
 #include "CoreMinimal.h"
-#include "DecoratorBase/DecoratorTemplate.h"
 #include "DecoratorBase/NodeTemplateRegistryHandle.h"
 
+#include <limits>
+
+class FArchive;
+
 namespace UE::AnimNext
 {
-	struct FDecoratorDescription;
-
 	/**
 	 * Node Description
 	 * A node description represents a unique instance in the authored static graph.
@@ -25,6 +26,15 @@ namespace UE::AnimNext
 	 */
 	struct alignas(alignof(uint32)) FNodeDescription
 	{
+		// Largest allowed size for a node description
+		// Graphs are currently limited to 64 KB because we use 16 bit offsets within our node/decorator handles
+		// We have no particular limit on the node size besides the graph limit (for now)
+		static constexpr uint32 MAXIMUM_SIZE = 64 * 1024;
+
+		// The maximum number of nodes allowed within a single graph
+		// We use 16 bits to represent node UIDs
+		static constexpr uint32 MAXIMUM_COUNT = std::numeric_limits<uint16>::max();
+
 		FNodeDescription(uint16 UID_, FNodeTemplateRegistryHandle TemplateHandle_)
 			: UID(UID_)
 			, TemplateHandle(TemplateHandle_)
@@ -36,10 +46,13 @@ namespace UE::AnimNext
 		// Returns the handle of the node's template in the node template registry
 		FNodeTemplateRegistryHandle GetTemplateHandle() const { return TemplateHandle; }
 
+		// Serializes this node description instance and the shared data of each decorator that follows
+		ANIMNEXT_API void Serialize(FArchive& Ar);
+
 	private:
 		uint16							UID;				// assigned during export/cook, unique to current sub-graph
 		FNodeTemplateRegistryHandle		TemplateHandle;		// offset of the node template within the global list
 
-		// Followed by a list of [FDecoratorDescription] instances and optional padding
+		// Followed by a list of [FAnimNextDecoratorSharedData] instances and optional padding
 	};
 }
diff --git a/Engine/Plugins/Experimental/AnimNext/Source/AnimNext/Public/DecoratorBase/NodeHandle.h b/Engine/Plugins/Experimental/AnimNext/Source/AnimNext/Public/DecoratorBase/NodeHandle.h
--- a/Engine/Plugins/Experimental/AnimNext/Source/AnimNext/Public/DecoratorBase/NodeHandle.h
+++ b/Engine/Plugins/Experimental/AnimNext/Source/AnimNext/Public/DecoratorBase/NodeHandle.h
@@ -4,6 +4,8 @@
 
 #include "CoreMinimal.h"
 
+struct FAnimNextDecoratorHandle;
+
 namespace UE::AnimNext
 {
 	/**
@@ -42,6 +44,6 @@ namespace UE::AnimNext
 		uint32	Padding : 8;			// unused padding
 		uint32	SharedOffset : 24;		// relative to root of sub-graph
 
-		friend struct FDecoratorHandle;
+		friend FAnimNextDecoratorHandle;
 	};
 }
diff --git a/Engine/Plugins/Experimental/AnimNext/Source/AnimNext/Public/DecoratorBase/NodeInstance.h b/Engine/Plugins/Experimental/AnimNext/Source/AnimNext/Public/DecoratorBase/NodeInstance.h
--- a/Engine/Plugins/Experimental/AnimNext/Source/AnimNext/Public/DecoratorBase/NodeInstance.h
+++ b/Engine/Plugins/Experimental/AnimNext/Source/AnimNext/Public/DecoratorBase/NodeInstance.h
@@ -8,8 +8,6 @@
 
 namespace UE::AnimNext
 {
-	struct FDecoratorInstance;
-
 	/**
 	 * Node Instance
 	 * A node instance represents allocated data for specific node.
@@ -42,7 +40,7 @@ namespace UE::AnimNext
 		uint32		ReferenceCount;		// how many non-weak FNodePtr and FDecoratorPtr handles point to us, not thread safe
 		FNodeHandle	NodeHandle;			// relative to root of sub-graph, should this be a pointer?
 
-		// Followed by a list of [FDecoratorInstance] instances and optional padding
+		// Followed by a list of [FDecoratorInstanceData] instances and optional padding
 
 		friend struct FDecoratorPtr;
 		friend struct FExecutionContext;
diff --git a/Engine/Plugins/Experimental/AnimNext/Source/AnimNextTestSuite/Private/AnimNextDecoratorBaseTest.cpp b/Engine/Plugins/Experimental/AnimNext/Source/AnimNextTestSuite/Private/AnimNextDecoratorBaseTest.cpp
--- a/Engine/Plugins/Experimental/AnimNext/Source/AnimNextTestSuite/Private/AnimNextDecoratorBaseTest.cpp
+++ b/Engine/Plugins/Experimental/AnimNext/Source/AnimNextTestSuite/Private/AnimNextDecoratorBaseTest.cpp
@@ -1,14 +1,18 @@
 // Copyright Epic Games, Inc. All Rights Reserved.
 
-
+#include "AnimNextDecoratorBaseTest.h"
 #include "AnimNextRuntimeTest.h"
 
 #include "Misc/AutomationTest.h"
 
 #if WITH_DEV_AUTOMATION_TESTS
 
+#include "Serialization/MemoryReader.h"
+
 #include "DecoratorBase/Decorator.h"
+#include "DecoratorBase/DecoratorReader.h"
 #include "DecoratorBase/DecoratorRegistry.h"
+#include "DecoratorBase/DecoratorWriter.h"
 #include "DecoratorBase/ExecutionContext.h"
 #include "DecoratorBase/IDecoratorInterface.h"
 #include "DecoratorBase/NodeInstance.h"
@@ -116,14 +120,11 @@ namespace UE::AnimNext
 
 	//////////////////////////////////////////////////////////////////////////
 
-	struct FDecoratorA_Base : FDecorator, IInterfaceA
+	struct FDecoratorA_Base : FBaseDecorator, IInterfaceA
 	{
-		DECLARE_ANIM_DECORATOR(FDecoratorA_Base, 0x3a1861cf, FDecorator)
+		DECLARE_ANIM_DECORATOR(FDecoratorA_Base, 0x3a1861cf, FBaseDecorator)
 
-		struct FSharedData : FDecorator::FSharedData
-		{
-			FDecoratorUID DecoratorUID = FDecoratorA_Base::DecoratorUID;
-		};
+		using FSharedData = FDecoratorA_BaseSharedData;
 
 		struct FInstanceData : FDecorator::FInstanceData
 		{
@@ -146,9 +147,6 @@ namespace UE::AnimNext
 			}
 		};
 
-		// FDecorator impl
-		virtual EDecoratorMode GetMode() const override { return EDecoratorMode::Base; }
-
 		// IInterfaceA impl
 		virtual void FuncA(FExecutionContext& Context, const TDecoratorBinding<IInterfaceA>& Binding) const override
 		{
@@ -161,14 +159,11 @@ namespace UE::AnimNext
 
 	//////////////////////////////////////////////////////////////////////////
 
-	struct FDecoratorAB_Add : FDecorator, IInterfaceA, IInterfaceB
+	struct FDecoratorAB_Add : FAdditiveDecorator, IInterfaceA, IInterfaceB
 	{
-		DECLARE_ANIM_DECORATOR(FDecoratorAB_Add, 0xe205a0e1, FDecorator)
+		DECLARE_ANIM_DECORATOR(FDecoratorAB_Add, 0xe205a0e1, FAdditiveDecorator)
 
-		struct FSharedData : FDecorator::FSharedData
-		{
-			FDecoratorUID DecoratorUID = FDecoratorAB_Add::DecoratorUID;
-		};
+		using FSharedData = FDecoratorAB_AddSharedData;
 
 		struct FInstanceData : FDecorator::FInstanceData
 		{
@@ -191,9 +186,6 @@ namespace UE::AnimNext
 			}
 		};
 
-		// FDecorator impl
-		virtual EDecoratorMode GetMode() const override { return EDecoratorMode::Additive; }
-
 		// IInterfaceA impl
 		virtual void FuncA(FExecutionContext& Context, const TDecoratorBinding<IInterfaceA>& Binding) const override
 		{
@@ -212,14 +204,11 @@ namespace UE::AnimNext
 
 	//////////////////////////////////////////////////////////////////////////
 
-	struct FDecoratorAC_Add : FDecorator, IInterfaceA, IInterfaceC
+	struct FDecoratorAC_Add : FAdditiveDecorator, IInterfaceA, IInterfaceC
 	{
-		DECLARE_ANIM_DECORATOR(FDecoratorAC_Add, 0x26d83846, FDecorator)
+		DECLARE_ANIM_DECORATOR(FDecoratorAC_Add, 0x26d83846, FAdditiveDecorator)
 
-		struct FSharedData : FDecorator::FSharedData
-		{
-			FDecoratorUID DecoratorUID = FDecoratorAC_Add::DecoratorUID;
-		};
+		using FSharedData = FDecoratorAC_AddSharedData;
 
 		struct FInstanceData : FDecorator::FInstanceData
 		{
@@ -242,9 +231,6 @@ namespace UE::AnimNext
 			}
 		};
 
-		// FDecorator impl
-		virtual EDecoratorMode GetMode() const override { return EDecoratorMode::Additive; }
-
 		// IInterfaceA impl
 		virtual void FuncA(FExecutionContext& Context, const TDecoratorBinding<IInterfaceA>& Binding) const override
 		{
@@ -260,6 +246,46 @@ namespace UE::AnimNext
 		DEFINE_ANIM_DECORATOR_IMPLEMENTS_INTERFACE(IInterfaceA)
 		DEFINE_ANIM_DECORATOR_IMPLEMENTS_INTERFACE(IInterfaceC)
 	DEFINE_ANIM_DECORATOR_END(FDecoratorAC_Add)
+
+
+	//////////////////////////////////////////////////////////////////////////
+
+	struct FDecoratorSerialization_Base : FBaseDecorator, IInterfaceA
+	{
+		DECLARE_ANIM_DECORATOR(FDecoratorSerialization_Base, 0x39fe92be, FBaseDecorator)
+
+		using FSharedData = FDecoratorSerialization_BaseSharedData;
+	};
+
+	DEFINE_ANIM_DECORATOR_BEGIN(FDecoratorSerialization_Base)
+		DEFINE_ANIM_DECORATOR_IMPLEMENTS_INTERFACE(IInterfaceA)
+	DEFINE_ANIM_DECORATOR_END(FDecoratorSerialization_Base)
+
+	//////////////////////////////////////////////////////////////////////////
+
+	struct FDecoratorSerialization_Add : FAdditiveDecorator, IInterfaceB
+	{
+		DECLARE_ANIM_DECORATOR(FDecoratorSerialization_Add, 0xf7c412d6, FAdditiveDecorator)
+
+		using FSharedData = FDecoratorSerialization_AddSharedData;
+	};
+
+	DEFINE_ANIM_DECORATOR_BEGIN(FDecoratorSerialization_Add)
+		DEFINE_ANIM_DECORATOR_IMPLEMENTS_INTERFACE(IInterfaceB)
+	DEFINE_ANIM_DECORATOR_END(FDecoratorSerialization_Add)
+
+	//////////////////////////////////////////////////////////////////////////
+
+	struct FDecoratorNativeSerialization_Add : FAdditiveDecorator, IInterfaceC
+	{
+		DECLARE_ANIM_DECORATOR(FDecoratorNativeSerialization_Add, 0x727ae94b, FAdditiveDecorator)
+
+		using FSharedData = FDecoratorNativeSerialization_AddSharedData;
+	};
+
+	DEFINE_ANIM_DECORATOR_BEGIN(FDecoratorNativeSerialization_Add)
+		DEFINE_ANIM_DECORATOR_IMPLEMENTS_INTERFACE(IInterfaceC)
+	DEFINE_ANIM_DECORATOR_END(FDecoratorNativeSerialization_Add)
 }
 
 IMPLEMENT_SIMPLE_AUTOMATION_TEST(FAnimationAnimNextRuntimeTest_DecoratorRegistry, "Animation.AnimNext.Runtime.DecoratorRegistry", EAutomationTestFlags::EditorContext | EAutomationTestFlags::EngineFilter)
@@ -295,7 +321,7 @@ bool FAnimationAnimNextRuntimeTest_DecoratorRegistry::RunTest(const FString& InP
 				// Auto register another decorator
 				AUTO_REGISTER_ANIM_DECORATOR(FDecoratorAB_Add)
 
-					AddErrorIfFalse(Registry.GetNum() == 2, "FAnimationAnimNextRuntimeTest_DecoratorRegistry -> Registry should contain 2 decorators");
+				AddErrorIfFalse(Registry.GetNum() == 2, "FAnimationAnimNextRuntimeTest_DecoratorRegistry -> Registry should contain 2 decorators");
 
 				FDecoratorRegistryHandle HandleAB = Registry.FindHandle(FDecoratorAB_Add::DecoratorUID);
 				AddErrorIfFalse(HandleAB.IsValid(), "FAnimationAnimNextRuntimeTest_DecoratorRegistry -> Decorator should have registered automatically");
@@ -487,18 +513,53 @@ bool FAnimationAnimNextRuntimeTest_NodeLifetime::RunTest(const FString& InParame
 	AddErrorIfFalse(Registry.GetNum() == 1, "FAnimationAnimNextRuntimeTest_NodeLifetime -> Registry should contain 1 template");
 	AddErrorIfFalse(TemplateHandle0.IsValid(), "FAnimationAnimNextRuntimeTest_NodeLifetime -> Registry should contain our template");
 
-	// Build our graph
-	uint16 NodeUID = 0;
+	FNodeHandle Node0;
+	FNodeHandle Node1;
+
+	// Write our graph
+	TArray<uint8> GraphSharedDataArchiveBuffer;
+	{
+		FDecoratorWriter DecoratorWriter;
+
+		Node0 = DecoratorWriter.RegisterNode(*NodeTemplate0);
+		Node1 = DecoratorWriter.RegisterNode(*NodeTemplate0);
+
+		// We don't have decorator properties
+		TArray<TMap<FString, FString>> DecoratorProperties;
+		DecoratorProperties.AddDefaulted(NodeTemplateDecoratorList.Num());
+
+		DecoratorWriter.BeginNodeWriting();
+		DecoratorWriter.WriteNode(Node0,
+			[&DecoratorProperties](uint32 DecoratorIndex) -> const TMap<FString, FString>&
+			{
+				return DecoratorProperties[DecoratorIndex];
+			});
+		DecoratorWriter.WriteNode(Node1,
+			[&DecoratorProperties](uint32 DecoratorIndex) -> const TMap<FString, FString>&
+			{
+				return DecoratorProperties[DecoratorIndex];
+			});
+		DecoratorWriter.EndNodeWriting();
+
+		AddErrorIfFalse(DecoratorWriter.GetErrorState() == FDecoratorWriter::EErrorState::None, "FAnimationAnimNextRuntimeTest_NodeLifetime -> Failed to write decorators");
+		GraphSharedDataArchiveBuffer = DecoratorWriter.GetGraphSharedData();
+	}
+
+	// Read our graph
 	TArray<uint8> GraphSharedDataBuffer;
-	const FNodeHandle Node0 = AppendNodeToGraph(*NodeTemplate0, NodeUID, GraphSharedDataBuffer);
-	const FNodeHandle Node1 = AppendNodeToGraph(*NodeTemplate0, NodeUID, GraphSharedDataBuffer);
+	{
+		FMemoryReader GraphSharedDataArchive(GraphSharedDataArchiveBuffer);
+		FDecoratorReader DecoratorReader(GraphSharedDataArchive);
+
+		DecoratorReader.ReadGraphSharedData(GraphSharedDataBuffer);
+	}
 
 	FExecutionContext Context(GraphSharedDataBuffer);
 
 	// Validate handle bookkeeping
 	{
 		FDecoratorBinding RootBinding;					// Empty, no parent
-		FDecoratorHandle DecoratorHandle00(Node0, 0);	// Point to first node, first base decorator
+		FAnimNextDecoratorHandle DecoratorHandle00(Node0, 0);	// Point to first node, first base decorator
 
 		// Allocate a node
 		FDecoratorPtr DecoratorPtr00 = Context.AllocateNodeInstance(RootBinding, DecoratorHandle00);
@@ -548,9 +609,9 @@ bool FAnimationAnimNextRuntimeTest_NodeLifetime::RunTest(const FString& InParame
 	// Validate parent support
 	{
 		FDecoratorBinding RootBinding;					// Empty, no parent
-		FDecoratorHandle DecoratorHandle00(Node0, 0);	// Point to first node, first base decorator
-		FDecoratorHandle DecoratorHandle03(Node0, 3);	// Point to first node, second base decorator
-		FDecoratorHandle DecoratorHandle10(Node1, 0);	// Point to second node, first base decorator
+		FAnimNextDecoratorHandle DecoratorHandle00(Node0, 0);	// Point to first node, first base decorator
+		FAnimNextDecoratorHandle DecoratorHandle03(Node0, 3);	// Point to first node, second base decorator
+		FAnimNextDecoratorHandle DecoratorHandle10(Node1, 0);	// Point to second node, first base decorator
 
 		// Allocate our first node
 		FDecoratorPtr DecoratorPtr00 = Context.AllocateNodeInstance(RootBinding, DecoratorHandle00);
@@ -587,7 +648,7 @@ bool FAnimationAnimNextRuntimeTest_NodeLifetime::RunTest(const FString& InParame
 
 		{
 			FDecoratorBinding RootBinding;					// Empty, no parent
-			FDecoratorHandle DecoratorHandle00(Node0, 0);	// Point to first node, first base decorator
+			FAnimNextDecoratorHandle DecoratorHandle00(Node0, 0);	// Point to first node, first base decorator
 
 			// Allocate our node instance
 			FDecoratorPtr DecoratorPtr00 = Context.AllocateNodeInstance(RootBinding, DecoratorHandle00);
@@ -626,9 +687,9 @@ bool FAnimationAnimNextRuntimeTest_NodeLifetime::RunTest(const FString& InParame
 	return true;
 }
 
-IMPLEMENT_SIMPLE_AUTOMATION_TEST(FAnimationAnimNextRuntimeTest_GetInterface, "Animation.AnimNext.Runtime.GetInterface", EAutomationTestFlags::EditorContext | EAutomationTestFlags::EngineFilter)
+IMPLEMENT_SIMPLE_AUTOMATION_TEST(FAnimationAnimNextRuntimeTest_GetDecoratorInterface, "Animation.AnimNext.Runtime.GetDecoratorInterface", EAutomationTestFlags::EditorContext | EAutomationTestFlags::EngineFilter)
 
-bool FAnimationAnimNextRuntimeTest_GetInterface::RunTest(const FString& InParameters)
+bool FAnimationAnimNextRuntimeTest_GetDecoratorInterface::RunTest(const FString& InParameters)
 {
 	using namespace UE::AnimNext;
 
@@ -650,117 +711,140 @@ bool FAnimationAnimNextRuntimeTest_GetInterface::RunTest(const FString& InParame
 	const FNodeTemplate* NodeTemplate0 = BuildNodeTemplate(NodeTemplateDecoratorList, NodeTemplateBuffer0);
 
 	FNodeTemplateRegistryHandle TemplateHandle0 = Registry.FindOrAdd(NodeTemplate0);
-	AddErrorIfFalse(Registry.GetNum() == 1, "FAnimationAnimNextRuntimeTest_GetInterface -> Registry should contain 1 template");
-	AddErrorIfFalse(TemplateHandle0.IsValid(), "FAnimationAnimNextRuntimeTest_GetInterface -> Registry should contain our template");
+	AddErrorIfFalse(Registry.GetNum() == 1, "FAnimationAnimNextRuntimeTest_GetDecoratorInterface -> Registry should contain 1 template");
+	AddErrorIfFalse(TemplateHandle0.IsValid(), "FAnimationAnimNextRuntimeTest_GetDecoratorInterface -> Registry should contain our template");
+
+	FNodeHandle Node0;
+
+	// Write our graph
+	TArray<uint8> GraphSharedDataArchiveBuffer;
+	{
+		FDecoratorWriter DecoratorWriter;
 
-	// Build our graph
-	uint16 NodeUID = 0;
+		Node0 = DecoratorWriter.RegisterNode(*NodeTemplate0);
+
+		// We don't have decorator properties
+		TArray<TMap<FString, FString>> DecoratorProperties;
+		DecoratorProperties.AddDefaulted(NodeTemplateDecoratorList.Num());
+
+		DecoratorWriter.BeginNodeWriting();
+		DecoratorWriter.WriteNode(Node0,
+			[&DecoratorProperties](uint32 DecoratorIndex) -> const TMap<FString, FString>&
+			{
+				return DecoratorProperties[DecoratorIndex];
+			});
+		DecoratorWriter.EndNodeWriting();
+
+		AddErrorIfFalse(DecoratorWriter.GetErrorState() == FDecoratorWriter::EErrorState::None, "FAnimationAnimNextRuntimeTest_GetDecoratorInterface -> Failed to write decorators");
+		GraphSharedDataArchiveBuffer = DecoratorWriter.GetGraphSharedData();
+	}
+
+	// Read our graph
 	TArray<uint8> GraphSharedDataBuffer;
+	{
+		FMemoryReader GraphSharedDataArchive(GraphSharedDataArchiveBuffer);
+		FDecoratorReader DecoratorReader(GraphSharedDataArchive);
 
-	const FNodeHandle Node0 = AppendNodeToGraph(*NodeTemplate0, NodeUID, GraphSharedDataBuffer);
-	InitNodeDecorator(*NodeTemplate0, Node0, GraphSharedDataBuffer, 0, FDecoratorA_Base::FSharedData());
-	InitNodeDecorator(*NodeTemplate0, Node0, GraphSharedDataBuffer, 1, FDecoratorAB_Add::FSharedData());
-	InitNodeDecorator(*NodeTemplate0, Node0, GraphSharedDataBuffer, 2, FDecoratorAC_Add::FSharedData());
-	InitNodeDecorator(*NodeTemplate0, Node0, GraphSharedDataBuffer, 3, FDecoratorA_Base::FSharedData());
-	InitNodeDecorator(*NodeTemplate0, Node0, GraphSharedDataBuffer, 4, FDecoratorAC_Add::FSharedData());
+		DecoratorReader.ReadGraphSharedData(GraphSharedDataBuffer);
+	}
 
 	FExecutionContext Context(GraphSharedDataBuffer);
 
 	// Validate from the first base decorator
 	{
 		FDecoratorBinding ParentBinding;				// Empty, no parent
-		FDecoratorHandle DecoratorHandle00(Node0, 0);	// Point to first node, first base decorator
+		FAnimNextDecoratorHandle DecoratorHandle00(Node0, 0);	// Point to first node, first base decorator
 
 		FDecoratorPtr DecoratorPtr00 = Context.AllocateNodeInstance(ParentBinding, DecoratorHandle00);
-		AddErrorIfFalse(DecoratorPtr00.IsValid(), "FAnimationAnimNextRuntimeTest_GetInterface -> Failed to allocate a node instance");
+		AddErrorIfFalse(DecoratorPtr00.IsValid(), "FAnimationAnimNextRuntimeTest_GetDecoratorInterface -> Failed to allocate a node instance");
 
 		// Validate GetInterface from a decorator handle
 		TDecoratorBinding<IInterfaceC> Binding00C;
-		AddErrorIfFalse(Context.GetInterface(DecoratorPtr00, Binding00C), "FAnimationAnimNextRuntimeTest_GetInterface -> InterfaceC not found");
-		AddErrorIfFalse(Binding00C.IsValid(), "FAnimationAnimNextRuntimeTest_GetInterface -> InterfaceC binding not valid");
-		AddErrorIfFalse(Binding00C.GetInterfaceUID() == IInterfaceC::InterfaceUID, "FAnimationAnimNextRuntimeTest_GetInterface -> Unexpected interface UID found in decorator binding");
-		AddErrorIfFalse(Binding00C.GetDecoratorPtr().GetDecoratorIndex() == 2, "FAnimationAnimNextRuntimeTest_GetInterface -> InterfaceC not found on expected decorator");
-		AddErrorIfFalse(Binding00C.GetDecoratorPtr().GetNodeInstance()->GetNodeHandle() == Node0, "FAnimationAnimNextRuntimeTest_GetInterface -> InterfaceC not found on expected node");
-		AddErrorIfFalse(Binding00C.GetSharedData<FDecoratorAC_Add::FSharedData>()->DecoratorUID == FDecoratorAC_Add::DecoratorUID, "FAnimationAnimNextRuntimeTest_GetInterface -> Unexpected shared data in decorator binding");
-		AddErrorIfFalse(Binding00C.GetInstanceData<FDecoratorAC_Add::FInstanceData>()->DecoratorUID == FDecoratorAC_Add::DecoratorUID, "FAnimationAnimNextRuntimeTest_GetInterface -> Unexpected instance data in decorator binding");
+		AddErrorIfFalse(Context.GetInterface(DecoratorPtr00, Binding00C), "FAnimationAnimNextRuntimeTest_GetDecoratorInterface -> InterfaceC not found");
+		AddErrorIfFalse(Binding00C.IsValid(), "FAnimationAnimNextRuntimeTest_GetDecoratorInterface -> InterfaceC binding not valid");
+		AddErrorIfFalse(Binding00C.GetInterfaceUID() == IInterfaceC::InterfaceUID, "FAnimationAnimNextRuntimeTest_GetDecoratorInterface -> Unexpected interface UID found in decorator binding");
+		AddErrorIfFalse(Binding00C.GetDecoratorPtr().GetDecoratorIndex() == 2, "FAnimationAnimNextRuntimeTest_GetDecoratorInterface -> InterfaceC not found on expected decorator");
+		AddErrorIfFalse(Binding00C.GetDecoratorPtr().GetNodeInstance()->GetNodeHandle() == Node0, "FAnimationAnimNextRuntimeTest_GetDecoratorInterface -> InterfaceC not found on expected node");
+		AddErrorIfFalse(Binding00C.GetSharedData<FDecoratorAC_Add::FSharedData>()->DecoratorUID == FDecoratorAC_Add::DecoratorUID, "FAnimationAnimNextRuntimeTest_GetDecoratorInterface -> Unexpected shared data in decorator binding");
+		AddErrorIfFalse(Binding00C.GetInstanceData<FDecoratorAC_Add::FInstanceData>()->DecoratorUID == FDecoratorAC_Add::DecoratorUID, "FAnimationAnimNextRuntimeTest_GetDecoratorInterface -> Unexpected instance data in decorator binding");
 
 		TDecoratorBinding<IInterfaceB> Binding00B;
-		AddErrorIfFalse(Context.GetInterface(DecoratorPtr00, Binding00B), "FAnimationAnimNextRuntimeTest_GetInterface -> InterfaceB not found");
-		AddErrorIfFalse(Binding00B.IsValid(), "FAnimationAnimNextRuntimeTest_GetInterface -> InterfaceB binding not valid");
-		AddErrorIfFalse(Binding00B.GetInterfaceUID() == IInterfaceB::InterfaceUID, "FAnimationAnimNextRuntimeTest_GetInterface -> Unexpected interface UID found in decorator binding");
-		AddErrorIfFalse(Binding00B.GetDecoratorPtr().GetDecoratorIndex() == 1, "FAnimationAnimNextRuntimeTest_GetInterface -> InterfaceB not found on expected decorator");
-		AddErrorIfFalse(Binding00B.GetDecoratorPtr().GetNodeInstance()->GetNodeHandle() == Node0, "FAnimationAnimNextRuntimeTest_GetInterface -> InterfaceB not found on expected node");
-		AddErrorIfFalse(Binding00B.GetSharedData<FDecoratorAB_Add::FSharedData>()->DecoratorUID == FDecoratorAB_Add::DecoratorUID, "FAnimationAnimNextRuntimeTest_GetInterface -> Unexpected shared data in decorator binding");
-		AddErrorIfFalse(Binding00B.GetInstanceData<FDecoratorAB_Add::FInstanceData>()->DecoratorUID == FDecoratorAB_Add::DecoratorUID, "FAnimationAnimNextRuntimeTest_GetInterface -> Unexpected instance data in decorator binding");
+		AddErrorIfFalse(Context.GetInterface(DecoratorPtr00, Binding00B), "FAnimationAnimNextRuntimeTest_GetDecoratorInterface -> InterfaceB not found");
+		AddErrorIfFalse(Binding00B.IsValid(), "FAnimationAnimNextRuntimeTest_GetDecoratorInterface -> InterfaceB binding not valid");
+		AddErrorIfFalse(Binding00B.GetInterfaceUID() == IInterfaceB::InterfaceUID, "FAnimationAnimNextRuntimeTest_GetDecoratorInterface -> Unexpected interface UID found in decorator binding");
+		AddErrorIfFalse(Binding00B.GetDecoratorPtr().GetDecoratorIndex() == 1, "FAnimationAnimNextRuntimeTest_GetDecoratorInterface -> InterfaceB not found on expected decorator");
+		AddErrorIfFalse(Binding00B.GetDecoratorPtr().GetNodeInstance()->GetNodeHandle() == Node0, "FAnimationAnimNextRuntimeTest_GetDecoratorInterface -> InterfaceB not found on expected node");
+		AddErrorIfFalse(Binding00B.GetSharedData<FDecoratorAB_Add::FSharedData>()->DecoratorUID == FDecoratorAB_Add::DecoratorUID, "FAnimationAnimNextRuntimeTest_GetDecoratorInterface -> Unexpected shared data in decorator binding");
+		AddErrorIfFalse(Binding00B.GetInstanceData<FDecoratorAB_Add::FInstanceData>()->DecoratorUID == FDecoratorAB_Add::DecoratorUID, "FAnimationAnimNextRuntimeTest_GetDecoratorInterface -> Unexpected instance data in decorator binding");
 
 		TDecoratorBinding<IInterfaceA> Binding00A;
-		AddErrorIfFalse(Context.GetInterface(DecoratorPtr00, Binding00A), "FAnimationAnimNextRuntimeTest_GetInterface -> InterfaceA not found");
-		AddErrorIfFalse(Binding00A.IsValid(), "FAnimationAnimNextRuntimeTest_GetInterface -> InterfaceA binding not valid");
-		AddErrorIfFalse(Binding00A.GetInterfaceUID() == IInterfaceA::InterfaceUID, "FAnimationAnimNextRuntimeTest_GetInterface -> Unexpected interface UID found in decorator binding");
-		AddErrorIfFalse(Binding00A.GetDecoratorPtr().GetDecoratorIndex() == 2, "FAnimationAnimNextRuntimeTest_GetInterface -> InterfaceA not found on expected decorator");
-		AddErrorIfFalse(Binding00A.GetDecoratorPtr().GetNodeInstance()->GetNodeHandle() == Node0, "FAnimationAnimNextRuntimeTest_GetInterface -> InterfaceA not found on expected node");
-		AddErrorIfFalse(Binding00A.GetSharedData<FDecoratorAC_Add::FSharedData>()->DecoratorUID == FDecoratorAC_Add::DecoratorUID, "FAnimationAnimNextRuntimeTest_GetInterface -> Unexpected shared data in decorator binding");
-		AddErrorIfFalse(Binding00A.GetInstanceData<FDecoratorAC_Add::FInstanceData>()->DecoratorUID == FDecoratorAC_Add::DecoratorUID, "FAnimationAnimNextRuntimeTest_GetInterface -> Unexpected instance data in decorator binding");
+		AddErrorIfFalse(Context.GetInterface(DecoratorPtr00, Binding00A), "FAnimationAnimNextRuntimeTest_GetDecoratorInterface -> InterfaceA not found");
+		AddErrorIfFalse(Binding00A.IsValid(), "FAnimationAnimNextRuntimeTest_GetDecoratorInterface -> InterfaceA binding not valid");
+		AddErrorIfFalse(Binding00A.GetInterfaceUID() == IInterfaceA::InterfaceUID, "FAnimationAnimNextRuntimeTest_GetDecoratorInterface -> Unexpected interface UID found in decorator binding");
+		AddErrorIfFalse(Binding00A.GetDecoratorPtr().GetDecoratorIndex() == 2, "FAnimationAnimNextRuntimeTest_GetDecoratorInterface -> InterfaceA not found on expected decorator");
+		AddErrorIfFalse(Binding00A.GetDecoratorPtr().GetNodeInstance()->GetNodeHandle() == Node0, "FAnimationAnimNextRuntimeTest_GetDecoratorInterface -> InterfaceA not found on expected node");
+		AddErrorIfFalse(Binding00A.GetSharedData<FDecoratorAC_Add::FSharedData>()->DecoratorUID == FDecoratorAC_Add::DecoratorUID, "FAnimationAnimNextRuntimeTest_GetDecoratorInterface -> Unexpected shared data in decorator binding");
+		AddErrorIfFalse(Binding00A.GetInstanceData<FDecoratorAC_Add::FInstanceData>()->DecoratorUID == FDecoratorAC_Add::DecoratorUID, "FAnimationAnimNextRuntimeTest_GetDecoratorInterface -> Unexpected instance data in decorator binding");
 
 		// Validate GetInterface from a decorator binding
 		{
 			{
 				{
 					TDecoratorBinding<IInterfaceC> Binding00C_;
-					AddErrorIfFalse(Context.GetInterface(Binding00C, Binding00C_), "FAnimationAnimNextRuntimeTest_GetInterface -> InterfaceC not found");
-					AddErrorIfFalse(Binding00C == Binding00C_, "FAnimationAnimNextRuntimeTest_GetInterface -> GetInterface methods should return the same result");
+					AddErrorIfFalse(Context.GetInterface(Binding00C, Binding00C_), "FAnimationAnimNextRuntimeTest_GetDecoratorInterface -> InterfaceC not found");
+					AddErrorIfFalse(Binding00C == Binding00C_, "FAnimationAnimNextRuntimeTest_GetDecoratorInterface -> GetInterface methods should return the same result");
 				}
 
 				{
 					TDecoratorBinding<IInterfaceC> Binding00C_;
-					AddErrorIfFalse(Context.GetInterface(Binding00B, Binding00C_), "FAnimationAnimNextRuntimeTest_GetInterface -> InterfaceC not found");
-					AddErrorIfFalse(Binding00C == Binding00C_, "FAnimationAnimNextRuntimeTest_GetInterface -> GetInterface methods should return the same result");
+					AddErrorIfFalse(Context.GetInterface(Binding00B, Binding00C_), "FAnimationAnimNextRuntimeTest_GetDecoratorInterface -> InterfaceC not found");
+					AddErrorIfFalse(Binding00C == Binding00C_, "FAnimationAnimNextRuntimeTest_GetDecoratorInterface -> GetInterface methods should return the same result");
 				}
 
 				{
 					TDecoratorBinding<IInterfaceC> Binding00C_;
-					AddErrorIfFalse(Context.GetInterface(Binding00A, Binding00C_), "FAnimationAnimNextRuntimeTest_GetInterface -> InterfaceC not found");
-					AddErrorIfFalse(Binding00C == Binding00C_, "FAnimationAnimNextRuntimeTest_GetInterface -> GetInterface methods should return the same result");
+					AddErrorIfFalse(Context.GetInterface(Binding00A, Binding00C_), "FAnimationAnimNextRuntimeTest_GetDecoratorInterface -> InterfaceC not found");
+					AddErrorIfFalse(Binding00C == Binding00C_, "FAnimationAnimNextRuntimeTest_GetDecoratorInterface -> GetInterface methods should return the same result");
 				}
 			}
 
 			{
 				{
 					TDecoratorBinding<IInterfaceB> Binding00B_;
-					AddErrorIfFalse(Context.GetInterface(Binding00C, Binding00B_), "FAnimationAnimNextRuntimeTest_GetInterface -> InterfaceB not found");
-					AddErrorIfFalse(Binding00B == Binding00B_, "FAnimationAnimNextRuntimeTest_GetInterface -> GetInterface methods should return the same result");
+					AddErrorIfFalse(Context.GetInterface(Binding00C, Binding00B_), "FAnimationAnimNextRuntimeTest_GetDecoratorInterface -> InterfaceB not found");
+					AddErrorIfFalse(Binding00B == Binding00B_, "FAnimationAnimNextRuntimeTest_GetDecoratorInterface -> GetInterface methods should return the same result");
 				}
 
 				{
 					TDecoratorBinding<IInterfaceB> Binding00B_;
-					AddErrorIfFalse(Context.GetInterface(Binding00B, Binding00B_), "FAnimationAnimNextRuntimeTest_GetInterface -> InterfaceB not found");
-					AddErrorIfFalse(Binding00B == Binding00B_, "FAnimationAnimNextRuntimeTest_GetInterface -> GetInterface methods should return the same result");
+					AddErrorIfFalse(Context.GetInterface(Binding00B, Binding00B_), "FAnimationAnimNextRuntimeTest_GetDecoratorInterface -> InterfaceB not found");
+					AddErrorIfFalse(Binding00B == Binding00B_, "FAnimationAnimNextRuntimeTest_GetDecoratorInterface -> GetInterface methods should return the same result");
 				}
 
 				{
 					TDecoratorBinding<IInterfaceB> Binding00B_;
-					AddErrorIfFalse(Context.GetInterface(Binding00A, Binding00B_), "FAnimationAnimNextRuntimeTest_GetInterface -> InterfaceB not found");
-					AddErrorIfFalse(Binding00B == Binding00B_, "FAnimationAnimNextRuntimeTest_GetInterface -> GetInterface methods should return the same result");
+					AddErrorIfFalse(Context.GetInterface(Binding00A, Binding00B_), "FAnimationAnimNextRuntimeTest_GetDecoratorInterface -> InterfaceB not found");
+					AddErrorIfFalse(Binding00B == Binding00B_, "FAnimationAnimNextRuntimeTest_GetDecoratorInterface -> GetInterface methods should return the same result");
 				}
 			}
 
 			{
 				{
 					TDecoratorBinding<IInterfaceA> Binding00A_;
-					AddErrorIfFalse(Context.GetInterface(Binding00C, Binding00A_), "FAnimationAnimNextRuntimeTest_GetInterface -> InterfaceA not found");
-					AddErrorIfFalse(Binding00A == Binding00A_, "FAnimationAnimNextRuntimeTest_GetInterface -> GetInterface methods should return the same result");
+					AddErrorIfFalse(Context.GetInterface(Binding00C, Binding00A_), "FAnimationAnimNextRuntimeTest_GetDecoratorInterface -> InterfaceA not found");
+					AddErrorIfFalse(Binding00A == Binding00A_, "FAnimationAnimNextRuntimeTest_GetDecoratorInterface -> GetInterface methods should return the same result");
 				}
 
 				{
 					TDecoratorBinding<IInterfaceA> Binding00A_;
-					AddErrorIfFalse(Context.GetInterface(Binding00B, Binding00A_), "FAnimationAnimNextRuntimeTest_GetInterface -> InterfaceA not found");
-					AddErrorIfFalse(Binding00A == Binding00A_, "FAnimationAnimNextRuntimeTest_GetInterface -> GetInterface methods should return the same result");
+					AddErrorIfFalse(Context.GetInterface(Binding00B, Binding00A_), "FAnimationAnimNextRuntimeTest_GetDecoratorInterface -> InterfaceA not found");
+					AddErrorIfFalse(Binding00A == Binding00A_, "FAnimationAnimNextRuntimeTest_GetDecoratorInterface -> GetInterface methods should return the same result");
 				}
 
 				{
 					TDecoratorBinding<IInterfaceA> Binding00A_;
-					AddErrorIfFalse(Context.GetInterface(Binding00A, Binding00A_), "FAnimationAnimNextRuntimeTest_GetInterface -> InterfaceA not found");
-					AddErrorIfFalse(Binding00A == Binding00A_, "FAnimationAnimNextRuntimeTest_GetInterface -> GetInterface methods should return the same result");
+					AddErrorIfFalse(Context.GetInterface(Binding00A, Binding00A_), "FAnimationAnimNextRuntimeTest_GetDecoratorInterface -> InterfaceA not found");
+					AddErrorIfFalse(Binding00A == Binding00A_, "FAnimationAnimNextRuntimeTest_GetDecoratorInterface -> GetInterface methods should return the same result");
 				}
 			}
 		}
@@ -769,91 +853,91 @@ bool FAnimationAnimNextRuntimeTest_GetInterface::RunTest(const FString& InParame
 	// Validate from the second base decorator
 	{
 		FDecoratorBinding ParentBinding;				// Empty, no parent
-		FDecoratorHandle DecoratorHandle03(Node0, 3);	// Point to first node, second base decorator
+		FAnimNextDecoratorHandle DecoratorHandle03(Node0, 3);	// Point to first node, second base decorator
 
 		FDecoratorPtr DecoratorPtr03 = Context.AllocateNodeInstance(ParentBinding, DecoratorHandle03);
-		AddErrorIfFalse(DecoratorPtr03.IsValid(), "FAnimationAnimNextRuntimeTest_GetInterface -> Failed to allocate a node instance");
+		AddErrorIfFalse(DecoratorPtr03.IsValid(), "FAnimationAnimNextRuntimeTest_GetDecoratorInterface -> Failed to allocate a node instance");
 
 		// Validate GetInterface from a decorator handle
 		TDecoratorBinding<IInterfaceC> Binding03C;
-		AddErrorIfFalse(Context.GetInterface(DecoratorPtr03, Binding03C), "FAnimationAnimNextRuntimeTest_GetInterface -> InterfaceC not found");
-		AddErrorIfFalse(Binding03C.IsValid(), "FAnimationAnimNextRuntimeTest_GetInterface -> InterfaceC binding not valid");
-		AddErrorIfFalse(Binding03C.GetInterfaceUID() == IInterfaceC::InterfaceUID, "FAnimationAnimNextRuntimeTest_GetInterface -> Unexpected interface UID found in decorator binding");
-		AddErrorIfFalse(Binding03C.GetDecoratorPtr().GetDecoratorIndex() == 4, "FAnimationAnimNextRuntimeTest_GetInterface -> InterfaceC not found on expected decorator");
-		AddErrorIfFalse(Binding03C.GetDecoratorPtr().GetNodeInstance()->GetNodeHandle() == Node0, "FAnimationAnimNextRuntimeTest_GetInterface -> InterfaceC not found on expected node");
-		AddErrorIfFalse(Binding03C.GetSharedData<FDecoratorAC_Add::FSharedData>()->DecoratorUID == FDecoratorAC_Add::DecoratorUID, "FAnimationAnimNextRuntimeTest_GetInterface -> Unexpected shared data in decorator binding");
-		AddErrorIfFalse(Binding03C.GetInstanceData<FDecoratorAC_Add::FInstanceData>()->DecoratorUID == FDecoratorAC_Add::DecoratorUID, "FAnimationAnimNextRuntimeTest_GetInterface -> Unexpected instance data in decorator binding");
+		AddErrorIfFalse(Context.GetInterface(DecoratorPtr03, Binding03C), "FAnimationAnimNextRuntimeTest_GetDecoratorInterface -> InterfaceC not found");
+		AddErrorIfFalse(Binding03C.IsValid(), "FAnimationAnimNextRuntimeTest_GetDecoratorInterface -> InterfaceC binding not valid");
+		AddErrorIfFalse(Binding03C.GetInterfaceUID() == IInterfaceC::InterfaceUID, "FAnimationAnimNextRuntimeTest_GetDecoratorInterface -> Unexpected interface UID found in decorator binding");
+		AddErrorIfFalse(Binding03C.GetDecoratorPtr().GetDecoratorIndex() == 4, "FAnimationAnimNextRuntimeTest_GetDecoratorInterface -> InterfaceC not found on expected decorator");
+		AddErrorIfFalse(Binding03C.GetDecoratorPtr().GetNodeInstance()->GetNodeHandle() == Node0, "FAnimationAnimNextRuntimeTest_GetDecoratorInterface -> InterfaceC not found on expected node");
+		AddErrorIfFalse(Binding03C.GetSharedData<FDecoratorAC_Add::FSharedData>()->DecoratorUID == FDecoratorAC_Add::DecoratorUID, "FAnimationAnimNextRuntimeTest_GetDecoratorInterface -> Unexpected shared data in decorator binding");
+		AddErrorIfFalse(Binding03C.GetInstanceData<FDecoratorAC_Add::FInstanceData>()->DecoratorUID == FDecoratorAC_Add::DecoratorUID, "FAnimationAnimNextRuntimeTest_GetDecoratorInterface -> Unexpected instance data in decorator binding");
 
 		TDecoratorBinding<IInterfaceB> Binding03B;
-		AddErrorIfFalse(!Context.GetInterface(DecoratorPtr03, Binding03B), "FAnimationAnimNextRuntimeTest_GetInterface -> InterfaceB found");
-		AddErrorIfFalse(!Binding03B.IsValid(), "FAnimationAnimNextRuntimeTest_GetInterface -> InterfaceB binding should not be valid");
+		AddErrorIfFalse(!Context.GetInterface(DecoratorPtr03, Binding03B), "FAnimationAnimNextRuntimeTest_GetDecoratorInterface -> InterfaceB found");
+		AddErrorIfFalse(!Binding03B.IsValid(), "FAnimationAnimNextRuntimeTest_GetDecoratorInterface -> InterfaceB binding should not be valid");
 
 		TDecoratorBinding<IInterfaceA> Binding03A;
-		AddErrorIfFalse(Context.GetInterface(DecoratorPtr03, Binding03A), "FAnimationAnimNextRuntimeTest_GetInterface -> InterfaceA not found");
-		AddErrorIfFalse(Binding03A.IsValid(), "FAnimationAnimNextRuntimeTest_GetInterface -> InterfaceA binding not valid");
-		AddErrorIfFalse(Binding03A.GetInterfaceUID() == IInterfaceA::InterfaceUID, "FAnimationAnimNextRuntimeTest_GetInterface -> Unexpected interface UID found in decorator binding");
-		AddErrorIfFalse(Binding03A.GetDecoratorPtr().GetDecoratorIndex() == 4, "FAnimationAnimNextRuntimeTest_GetInterface -> InterfaceA not found on expected decorator");
-		AddErrorIfFalse(Binding03A.GetDecoratorPtr().GetNodeInstance()->GetNodeHandle() == Node0, "FAnimationAnimNextRuntimeTest_GetInterface -> InterfaceA not found on expected node");
-		AddErrorIfFalse(Binding03A.GetSharedData<FDecoratorAC_Add::FSharedData>()->DecoratorUID == FDecoratorAC_Add::DecoratorUID, "FAnimationAnimNextRuntimeTest_GetInterface -> Unexpected shared data in decorator binding");
-		AddErrorIfFalse(Binding03A.GetInstanceData<FDecoratorAC_Add::FInstanceData>()->DecoratorUID == FDecoratorAC_Add::DecoratorUID, "FAnimationAnimNextRuntimeTest_GetInterface -> Unexpected instance data in decorator binding");
+		AddErrorIfFalse(Context.GetInterface(DecoratorPtr03, Binding03A), "FAnimationAnimNextRuntimeTest_GetDecoratorInterface -> InterfaceA not found");
+		AddErrorIfFalse(Binding03A.IsValid(), "FAnimationAnimNextRuntimeTest_GetDecoratorInterface -> InterfaceA binding not valid");
+		AddErrorIfFalse(Binding03A.GetInterfaceUID() == IInterfaceA::InterfaceUID, "FAnimationAnimNextRuntimeTest_GetDecoratorInterface -> Unexpected interface UID found in decorator binding");
+		AddErrorIfFalse(Binding03A.GetDecoratorPtr().GetDecoratorIndex() == 4, "FAnimationAnimNextRuntimeTest_GetDecoratorInterface -> InterfaceA not found on expected decorator");
+		AddErrorIfFalse(Binding03A.GetDecoratorPtr().GetNodeInstance()->GetNodeHandle() == Node0, "FAnimationAnimNextRuntimeTest_GetDecoratorInterface -> InterfaceA not found on expected node");
+		AddErrorIfFalse(Binding03A.GetSharedData<FDecoratorAC_Add::FSharedData>()->DecoratorUID == FDecoratorAC_Add::DecoratorUID, "FAnimationAnimNextRuntimeTest_GetDecoratorInterface -> Unexpected shared data in decorator binding");
+		AddErrorIfFalse(Binding03A.GetInstanceData<FDecoratorAC_Add::FInstanceData>()->DecoratorUID == FDecoratorAC_Add::DecoratorUID, "FAnimationAnimNextRuntimeTest_GetDecoratorInterface -> Unexpected instance data in decorator binding");
 
 		// Validate GetInterface from a decorator binding
 		{
 			{
 				{
 					TDecoratorBinding<IInterfaceC> Binding03C_;
-					AddErrorIfFalse(Context.GetInterface(Binding03C, Binding03C_), "FAnimationAnimNextRuntimeTest_GetInterface -> InterfaceC not found");
-					AddErrorIfFalse(Binding03C == Binding03C_, "FAnimationAnimNextRuntimeTest_GetInterface -> GetInterface methods should return the same result");
+					AddErrorIfFalse(Context.GetInterface(Binding03C, Binding03C_), "FAnimationAnimNextRuntimeTest_GetDecoratorInterface -> InterfaceC not found");
+					AddErrorIfFalse(Binding03C == Binding03C_, "FAnimationAnimNextRuntimeTest_GetDecoratorInterface -> GetInterface methods should return the same result");
 				}
 
 				{
 					TDecoratorBinding<IInterfaceC> Binding03C_;
-					AddErrorIfFalse(!Context.GetInterface(Binding03B, Binding03C_), "FAnimationAnimNextRuntimeTest_GetInterface -> InterfaceC found");
+					AddErrorIfFalse(!Context.GetInterface(Binding03B, Binding03C_), "FAnimationAnimNextRuntimeTest_GetDecoratorInterface -> InterfaceC found");
 				}
 
 				{
 					TDecoratorBinding<IInterfaceC> Binding03C_;
-					AddErrorIfFalse(Context.GetInterface(Binding03A, Binding03C_), "FAnimationAnimNextRuntimeTest_GetInterface -> InterfaceC not found");
-					AddErrorIfFalse(Binding03C == Binding03C_, "FAnimationAnimNextRuntimeTest_GetInterface -> GetInterface methods should return the same result");
+					AddErrorIfFalse(Context.GetInterface(Binding03A, Binding03C_), "FAnimationAnimNextRuntimeTest_GetDecoratorInterface -> InterfaceC not found");
+					AddErrorIfFalse(Binding03C == Binding03C_, "FAnimationAnimNextRuntimeTest_GetDecoratorInterface -> GetInterface methods should return the same result");
 				}
 			}
 
 			{
 				{
 					TDecoratorBinding<IInterfaceB> Binding03B_;
-					AddErrorIfFalse(!Context.GetInterface(Binding03C, Binding03B_), "FAnimationAnimNextRuntimeTest_GetInterface -> InterfaceB found");
-					AddErrorIfFalse(Binding03B == Binding03B_, "FAnimationAnimNextRuntimeTest_GetInterface -> GetInterface methods should return the same result");
+					AddErrorIfFalse(!Context.GetInterface(Binding03C, Binding03B_), "FAnimationAnimNextRuntimeTest_GetDecoratorInterface -> InterfaceB found");
+					AddErrorIfFalse(Binding03B == Binding03B_, "FAnimationAnimNextRuntimeTest_GetDecoratorInterface -> GetInterface methods should return the same result");
 				}
 
 				{
 					TDecoratorBinding<IInterfaceB> Binding03B_;
-					AddErrorIfFalse(!Context.GetInterface(Binding03B, Binding03B_), "FAnimationAnimNextRuntimeTest_GetInterface -> InterfaceB found");
-					AddErrorIfFalse(Binding03B == Binding03B_, "FAnimationAnimNextRuntimeTest_GetInterface -> GetInterface methods should return the same result");
+					AddErrorIfFalse(!Context.GetInterface(Binding03B, Binding03B_), "FAnimationAnimNextRuntimeTest_GetDecoratorInterface -> InterfaceB found");
+					AddErrorIfFalse(Binding03B == Binding03B_, "FAnimationAnimNextRuntimeTest_GetDecoratorInterface -> GetInterface methods should return the same result");
 				}
 
 				{
 					TDecoratorBinding<IInterfaceB> Binding03B_;
-					AddErrorIfFalse(!Context.GetInterface(Binding03A, Binding03B_), "FAnimationAnimNextRuntimeTest_GetInterface -> InterfaceB found");
-					AddErrorIfFalse(Binding03B == Binding03B_, "FAnimationAnimNextRuntimeTest_GetInterface -> GetInterface methods should return the same result");
+					AddErrorIfFalse(!Context.GetInterface(Binding03A, Binding03B_), "FAnimationAnimNextRuntimeTest_GetDecoratorInterface -> InterfaceB found");
+					AddErrorIfFalse(Binding03B == Binding03B_, "FAnimationAnimNextRuntimeTest_GetDecoratorInterface -> GetInterface methods should return the same result");
 				}
 			}
 
 			{
 				{
 					TDecoratorBinding<IInterfaceA> Binding03A_;
-					AddErrorIfFalse(Context.GetInterface(Binding03C, Binding03A_), "FAnimationAnimNextRuntimeTest_GetInterface -> InterfaceA not found");
-					AddErrorIfFalse(Binding03A == Binding03A_, "FAnimationAnimNextRuntimeTest_GetInterface -> GetInterface methods should return the same result");
+					AddErrorIfFalse(Context.GetInterface(Binding03C, Binding03A_), "FAnimationAnimNextRuntimeTest_GetDecoratorInterface -> InterfaceA not found");
+					AddErrorIfFalse(Binding03A == Binding03A_, "FAnimationAnimNextRuntimeTest_GetDecoratorInterface -> GetInterface methods should return the same result");
 				}
 
 				{
 					TDecoratorBinding<IInterfaceA> Binding03A_;
-					AddErrorIfFalse(!Context.GetInterface(Binding03B, Binding03A_), "FAnimationAnimNextRuntimeTest_GetInterface -> InterfaceA found");
+					AddErrorIfFalse(!Context.GetInterface(Binding03B, Binding03A_), "FAnimationAnimNextRuntimeTest_GetDecoratorInterface -> InterfaceA found");
 				}
 
 				{
 					TDecoratorBinding<IInterfaceA> Binding03A_;
-					AddErrorIfFalse(Context.GetInterface(Binding03A, Binding03A_), "FAnimationAnimNextRuntimeTest_GetInterface -> InterfaceA not found");
-					AddErrorIfFalse(Binding03A == Binding03A_, "FAnimationAnimNextRuntimeTest_GetInterface -> GetInterface methods should return the same result");
+					AddErrorIfFalse(Context.GetInterface(Binding03A, Binding03A_), "FAnimationAnimNextRuntimeTest_GetDecoratorInterface -> InterfaceA not found");
+					AddErrorIfFalse(Binding03A == Binding03A_, "FAnimationAnimNextRuntimeTest_GetDecoratorInterface -> GetInterface methods should return the same result");
 				}
 			}
 		}
@@ -861,14 +945,14 @@ bool FAnimationAnimNextRuntimeTest_GetInterface::RunTest(const FString& InParame
 
 	Registry.Unregister(NodeTemplate0);
 
-	AddErrorIfFalse(Registry.GetNum() == 0, "FAnimationAnimNextRuntimeTest_GetInterface -> Registry should contain 0 templates");
+	AddErrorIfFalse(Registry.GetNum() == 0, "FAnimationAnimNextRuntimeTest_GetDecoratorInterface -> Registry should contain 0 templates");
 
 	return true;
 }
 
-IMPLEMENT_SIMPLE_AUTOMATION_TEST(FAnimationAnimNextRuntimeTest_GetInterfaceSuper, "Animation.AnimNext.Runtime.GetInterfaceSuper", EAutomationTestFlags::EditorContext | EAutomationTestFlags::EngineFilter)
+IMPLEMENT_SIMPLE_AUTOMATION_TEST(FAnimationAnimNextRuntimeTest_GetDecoratorInterfaceSuper, "Animation.AnimNext.Runtime.GetDecoratorInterfaceSuper", EAutomationTestFlags::EditorContext | EAutomationTestFlags::EngineFilter)
 
-bool FAnimationAnimNextRuntimeTest_GetInterfaceSuper::RunTest(const FString& InParameters)
+bool FAnimationAnimNextRuntimeTest_GetDecoratorInterfaceSuper::RunTest(const FString& InParameters)
 {
 	using namespace UE::AnimNext;
 
@@ -890,149 +974,558 @@ bool FAnimationAnimNextRuntimeTest_GetInterfaceSuper::RunTest(const FString& InP
 	const FNodeTemplate* NodeTemplate0 = BuildNodeTemplate(NodeTemplateDecoratorList, NodeTemplateBuffer0);
 
 	FNodeTemplateRegistryHandle TemplateHandle0 = Registry.FindOrAdd(NodeTemplate0);
-	AddErrorIfFalse(Registry.GetNum() == 1, "FAnimationAnimNextRuntimeTest_GetInterfaceSuper -> Registry should contain 1 template");
-	AddErrorIfFalse(TemplateHandle0.IsValid(), "FAnimationAnimNextRuntimeTest_GetInterfaceSuper -> Registry should contain our template");
+	AddErrorIfFalse(Registry.GetNum() == 1, "FAnimationAnimNextRuntimeTest_GetDecoratorInterfaceSuper -> Registry should contain 1 template");
+	AddErrorIfFalse(TemplateHandle0.IsValid(), "FAnimationAnimNextRuntimeTest_GetDecoratorInterfaceSuper -> Registry should contain our template");
+
+	FNodeHandle Node0;
+
+	// Write our graph
+	TArray<uint8> GraphSharedDataArchiveBuffer;
+	{
+		FDecoratorWriter DecoratorWriter;
+
+		Node0 = DecoratorWriter.RegisterNode(*NodeTemplate0);
+
+		// We don't have decorator properties
+		TArray<TMap<FString, FString>> DecoratorProperties;
+		DecoratorProperties.AddDefaulted(NodeTemplateDecoratorList.Num());
+
+		DecoratorWriter.BeginNodeWriting();
+		DecoratorWriter.WriteNode(Node0,
+			[&DecoratorProperties](uint32 DecoratorIndex) -> const TMap<FString, FString>&
+			{
+				return DecoratorProperties[DecoratorIndex];
+			});
+		DecoratorWriter.EndNodeWriting();
 
-	// Build our graph
-	uint16 NodeUID = 0;
+		AddErrorIfFalse(DecoratorWriter.GetErrorState() == FDecoratorWriter::EErrorState::None, "FAnimationAnimNextRuntimeTest_GetDecoratorInterfaceSuper -> Failed to write decorators");
+		GraphSharedDataArchiveBuffer = DecoratorWriter.GetGraphSharedData();
+	}
+
+	// Read our graph
 	TArray<uint8> GraphSharedDataBuffer;
+	{
+		FMemoryReader GraphSharedDataArchive(GraphSharedDataArchiveBuffer);
+		FDecoratorReader DecoratorReader(GraphSharedDataArchive);
 
-	const FNodeHandle Node0 = AppendNodeToGraph(*NodeTemplate0, NodeUID, GraphSharedDataBuffer);
-	InitNodeDecorator(*NodeTemplate0, Node0, GraphSharedDataBuffer, 0, FDecoratorA_Base::FSharedData());
-	InitNodeDecorator(*NodeTemplate0, Node0, GraphSharedDataBuffer, 1, FDecoratorAB_Add::FSharedData());
-	InitNodeDecorator(*NodeTemplate0, Node0, GraphSharedDataBuffer, 2, FDecoratorAC_Add::FSharedData());
-	InitNodeDecorator(*NodeTemplate0, Node0, GraphSharedDataBuffer, 3, FDecoratorA_Base::FSharedData());
-	InitNodeDecorator(*NodeTemplate0, Node0, GraphSharedDataBuffer, 4, FDecoratorAC_Add::FSharedData());
+		DecoratorReader.ReadGraphSharedData(GraphSharedDataBuffer);
+	}
 
 	FExecutionContext Context(GraphSharedDataBuffer);
 
 	// Validate from the first base decorator
 	{
 		FDecoratorBinding ParentBinding;				// Empty, no parent
-		FDecoratorHandle DecoratorHandle00(Node0, 0);	// Point to first node, first base decorator
+		FAnimNextDecoratorHandle DecoratorHandle00(Node0, 0);	// Point to first node, first base decorator
 
 		FDecoratorPtr DecoratorPtr00 = Context.AllocateNodeInstance(ParentBinding, DecoratorHandle00);
-		AddErrorIfFalse(DecoratorPtr00.IsValid(), "FAnimationAnimNextRuntimeTest_GetInterfaceSuper -> Failed to allocate a node instance");
+		AddErrorIfFalse(DecoratorPtr00.IsValid(), "FAnimationAnimNextRuntimeTest_GetDecoratorInterfaceSuper -> Failed to allocate a node instance");
 
 		{
 			// Get a valid decorator binding: FDecoratorAC_Add
 			TDecoratorBinding<IInterfaceC> Binding02C;
-			AddErrorIfFalse(Context.GetInterface(DecoratorPtr00, Binding02C), "FAnimationAnimNextRuntimeTest_GetInterfaceSuper -> InterfaceC not found");
+			AddErrorIfFalse(Context.GetInterface(DecoratorPtr00, Binding02C), "FAnimationAnimNextRuntimeTest_GetDecoratorInterfaceSuper -> InterfaceC not found");
 
 			// Validate GetInterfaceSuper from a decorator handle
 			TDecoratorBinding<IInterfaceC> SuperBinding02C;
-			AddErrorIfFalse(!Context.GetInterfaceSuper(Binding02C.GetDecoratorPtr(), SuperBinding02C), "FAnimationAnimNextRuntimeTest_GetInterfaceSuper -> InterfaceC found");
+			AddErrorIfFalse(!Context.GetInterfaceSuper(Binding02C.GetDecoratorPtr(), SuperBinding02C), "FAnimationAnimNextRuntimeTest_GetDecoratorInterfaceSuper -> InterfaceC found");
 
 			// Validate GetInterfaceSuper from a decorator binding
 			TDecoratorBinding<IInterfaceC> SuperBinding02C_;
-			AddErrorIfFalse(!Context.GetInterfaceSuper(Binding02C, SuperBinding02C_), "FAnimationAnimNextRuntimeTest_GetInterfaceSuper -> InterfaceC found");
-			AddErrorIfFalse(SuperBinding02C == SuperBinding02C_, "FAnimationAnimNextRuntimeTest_GetInterfaceSuper -> GetInterfaceSuper methods should return the same result");
+			AddErrorIfFalse(!Context.GetInterfaceSuper(Binding02C, SuperBinding02C_), "FAnimationAnimNextRuntimeTest_GetDecoratorInterfaceSuper -> InterfaceC found");
+			AddErrorIfFalse(SuperBinding02C == SuperBinding02C_, "FAnimationAnimNextRuntimeTest_GetDecoratorInterfaceSuper -> GetInterfaceSuper methods should return the same result");
 		}
 
 		{
 			// Get a valid decorator binding: FDecoratorAC_Add
 			TDecoratorBinding<IInterfaceA> Binding02A;
-			AddErrorIfFalse(Context.GetInterface(DecoratorPtr00, Binding02A), "FAnimationAnimNextRuntimeTest_GetInterfaceSuper -> InterfaceA not found");
+			AddErrorIfFalse(Context.GetInterface(DecoratorPtr00, Binding02A), "FAnimationAnimNextRuntimeTest_GetDecoratorInterfaceSuper -> InterfaceA not found");
 
 			// Validate GetInterfaceSuper from a decorator handle, FDecoratorAB_Add
 			TDecoratorBinding<IInterfaceA> SuperBinding02A;
-			AddErrorIfFalse(Context.GetInterfaceSuper(Binding02A.GetDecoratorPtr(), SuperBinding02A), "FAnimationAnimNextRuntimeTest_GetInterfaceSuper -> InterfaceA not found");
-			AddErrorIfFalse(SuperBinding02A.IsValid(), "FAnimationAnimNextRuntimeTest_GetInterfaceSuper -> InterfaceA binding not valid");
-			AddErrorIfFalse(SuperBinding02A.GetInterfaceUID() == IInterfaceA::InterfaceUID, "FAnimationAnimNextRuntimeTest_GetInterfaceSuper -> Unexpected interface UID found in decorator binding");
-			AddErrorIfFalse(SuperBinding02A.GetDecoratorPtr().GetDecoratorIndex() == 1, "FAnimationAnimNextRuntimeTest_GetInterfaceSuper -> InterfaceA not found on expected decorator");
-			AddErrorIfFalse(SuperBinding02A.GetDecoratorPtr().GetNodeInstance()->GetNodeHandle() == Node0, "FAnimationAnimNextRuntimeTest_GetInterfaceSuper -> InterfaceA not found on expected node");
-			AddErrorIfFalse(SuperBinding02A.GetSharedData<FDecoratorAB_Add::FSharedData>()->DecoratorUID == FDecoratorAB_Add::DecoratorUID, "FAnimationAnimNextRuntimeTest_GetInterfaceSuper -> Unexpected shared data in decorator binding");
-			AddErrorIfFalse(SuperBinding02A.GetInstanceData<FDecoratorAB_Add::FInstanceData>()->DecoratorUID == FDecoratorAB_Add::DecoratorUID, "FAnimationAnimNextRuntimeTest_GetInterfaceSuper -> Unexpected instance data in decorator binding");
+			AddErrorIfFalse(Context.GetInterfaceSuper(Binding02A.GetDecoratorPtr(), SuperBinding02A), "FAnimationAnimNextRuntimeTest_GetDecoratorInterfaceSuper -> InterfaceA not found");
+			AddErrorIfFalse(SuperBinding02A.IsValid(), "FAnimationAnimNextRuntimeTest_GetDecoratorInterfaceSuper -> InterfaceA binding not valid");
+			AddErrorIfFalse(SuperBinding02A.GetInterfaceUID() == IInterfaceA::InterfaceUID, "FAnimationAnimNextRuntimeTest_GetDecoratorInterfaceSuper -> Unexpected interface UID found in decorator binding");
+			AddErrorIfFalse(SuperBinding02A.GetDecoratorPtr().GetDecoratorIndex() == 1, "FAnimationAnimNextRuntimeTest_GetDecoratorInterfaceSuper -> InterfaceA not found on expected decorator");
+			AddErrorIfFalse(SuperBinding02A.GetDecoratorPtr().GetNodeInstance()->GetNodeHandle() == Node0, "FAnimationAnimNextRuntimeTest_GetDecoratorInterfaceSuper -> InterfaceA not found on expected node");
+			AddErrorIfFalse(SuperBinding02A.GetSharedData<FDecoratorAB_Add::FSharedData>()->DecoratorUID == FDecoratorAB_Add::DecoratorUID, "FAnimationAnimNextRuntimeTest_GetDecoratorInterfaceSuper -> Unexpected shared data in decorator binding");
+			AddErrorIfFalse(SuperBinding02A.GetInstanceData<FDecoratorAB_Add::FInstanceData>()->DecoratorUID == FDecoratorAB_Add::DecoratorUID, "FAnimationAnimNextRuntimeTest_GetDecoratorInterfaceSuper -> Unexpected instance data in decorator binding");
 
 			// Validate GetInterfaceSuper from a decorator binding, FDecoratorAB_Add
 			TDecoratorBinding<IInterfaceA> SuperBinding02A_;
-			AddErrorIfFalse(Context.GetInterfaceSuper(Binding02A, SuperBinding02A_), "FAnimationAnimNextRuntimeTest_GetInterfaceSuper -> InterfaceA not found");
-			AddErrorIfFalse(SuperBinding02A == SuperBinding02A_, "FAnimationAnimNextRuntimeTest_GetInterfaceSuper -> GetInterfaceSuper methods should return the same result");
+			AddErrorIfFalse(Context.GetInterfaceSuper(Binding02A, SuperBinding02A_), "FAnimationAnimNextRuntimeTest_GetDecoratorInterfaceSuper -> InterfaceA not found");
+			AddErrorIfFalse(SuperBinding02A == SuperBinding02A_, "FAnimationAnimNextRuntimeTest_GetDecoratorInterfaceSuper -> GetInterfaceSuper methods should return the same result");
 
 			// Validate GetInterfaceSuper from a decorator handle, FDecoratorA_Base
 			TDecoratorBinding<IInterfaceA> SuperBinding01A;
-			AddErrorIfFalse(Context.GetInterfaceSuper(SuperBinding02A.GetDecoratorPtr(), SuperBinding01A), "FAnimationAnimNextRuntimeTest_GetInterfaceSuper -> InterfaceA not found");
-			AddErrorIfFalse(SuperBinding01A.IsValid(), "FAnimationAnimNextRuntimeTest_GetInterfaceSuper -> InterfaceA binding not valid");
-			AddErrorIfFalse(SuperBinding01A.GetInterfaceUID() == IInterfaceA::InterfaceUID, "FAnimationAnimNextRuntimeTest_GetInterfaceSuper -> Unexpected interface UID found in decorator binding");
-			AddErrorIfFalse(SuperBinding01A.GetDecoratorPtr().GetDecoratorIndex() == 0, "FAnimationAnimNextRuntimeTest_GetInterfaceSuper -> InterfaceA not found on expected decorator");
-			AddErrorIfFalse(SuperBinding01A.GetDecoratorPtr().GetNodeInstance()->GetNodeHandle() == Node0, "FAnimationAnimNextRuntimeTest_GetInterfaceSuper -> InterfaceA not found on expected node");
-			AddErrorIfFalse(SuperBinding01A.GetSharedData<FDecoratorA_Base::FSharedData>()->DecoratorUID == FDecoratorA_Base::DecoratorUID, "FAnimationAnimNextRuntimeTest_GetInterfaceSuper -> Unexpected shared data in decorator binding");
-			AddErrorIfFalse(SuperBinding01A.GetInstanceData<FDecoratorA_Base::FInstanceData>()->DecoratorUID == FDecoratorA_Base::DecoratorUID, "FAnimationAnimNextRuntimeTest_GetInterfaceSuper -> Unexpected instance data in decorator binding");
+			AddErrorIfFalse(Context.GetInterfaceSuper(SuperBinding02A.GetDecoratorPtr(), SuperBinding01A), "FAnimationAnimNextRuntimeTest_GetDecoratorInterfaceSuper -> InterfaceA not found");
+			AddErrorIfFalse(SuperBinding01A.IsValid(), "FAnimationAnimNextRuntimeTest_GetDecoratorInterfaceSuper -> InterfaceA binding not valid");
+			AddErrorIfFalse(SuperBinding01A.GetInterfaceUID() == IInterfaceA::InterfaceUID, "FAnimationAnimNextRuntimeTest_GetDecoratorInterfaceSuper -> Unexpected interface UID found in decorator binding");
+			AddErrorIfFalse(SuperBinding01A.GetDecoratorPtr().GetDecoratorIndex() == 0, "FAnimationAnimNextRuntimeTest_GetDecoratorInterfaceSuper -> InterfaceA not found on expected decorator");
+			AddErrorIfFalse(SuperBinding01A.GetDecoratorPtr().GetNodeInstance()->GetNodeHandle() == Node0, "FAnimationAnimNextRuntimeTest_GetDecoratorInterfaceSuper -> InterfaceA not found on expected node");
+			AddErrorIfFalse(SuperBinding01A.GetSharedData<FDecoratorA_Base::FSharedData>()->DecoratorUID == FDecoratorA_Base::DecoratorUID, "FAnimationAnimNextRuntimeTest_GetDecoratorInterfaceSuper -> Unexpected shared data in decorator binding");
+			AddErrorIfFalse(SuperBinding01A.GetInstanceData<FDecoratorA_Base::FInstanceData>()->DecoratorUID == FDecoratorA_Base::DecoratorUID, "FAnimationAnimNextRuntimeTest_GetDecoratorInterfaceSuper -> Unexpected instance data in decorator binding");
 
 			// Validate GetInterfaceSuper from a decorator binding, FDecoratorA_Base
 			TDecoratorBinding<IInterfaceA> SuperBinding01A_;
-			AddErrorIfFalse(Context.GetInterfaceSuper(SuperBinding02A, SuperBinding01A_), "FAnimationAnimNextRuntimeTest_GetInterfaceSuper -> InterfaceA not found");
-			AddErrorIfFalse(SuperBinding01A == SuperBinding01A_, "FAnimationAnimNextRuntimeTest_GetInterfaceSuper -> GetInterfaceSuper methods should return the same result");
+			AddErrorIfFalse(Context.GetInterfaceSuper(SuperBinding02A, SuperBinding01A_), "FAnimationAnimNextRuntimeTest_GetDecoratorInterfaceSuper -> InterfaceA not found");
+			AddErrorIfFalse(SuperBinding01A == SuperBinding01A_, "FAnimationAnimNextRuntimeTest_GetDecoratorInterfaceSuper -> GetInterfaceSuper methods should return the same result");
 
 			// Validate GetInterfaceSuper from a decorator handle
 			TDecoratorBinding<IInterfaceA> SuperBinding00A;
-			AddErrorIfFalse(!Context.GetInterfaceSuper(SuperBinding01A.GetDecoratorPtr(), SuperBinding00A), "FAnimationAnimNextRuntimeTest_GetInterfaceSuper -> InterfaceA found");
+			AddErrorIfFalse(!Context.GetInterfaceSuper(SuperBinding01A.GetDecoratorPtr(), SuperBinding00A), "FAnimationAnimNextRuntimeTest_GetDecoratorInterfaceSuper -> InterfaceA found");
 
 			// Validate GetInterfaceSuper from a decorator binding
 			TDecoratorBinding<IInterfaceA> SuperBinding00A_;
-			AddErrorIfFalse(!Context.GetInterfaceSuper(SuperBinding01A, SuperBinding00A_), "FAnimationAnimNextRuntimeTest_GetInterfaceSuper -> InterfaceA found");
-			AddErrorIfFalse(SuperBinding00A == SuperBinding00A_, "FAnimationAnimNextRuntimeTest_GetInterfaceSuper -> GetInterfaceSuper methods should return the same result");
+			AddErrorIfFalse(!Context.GetInterfaceSuper(SuperBinding01A, SuperBinding00A_), "FAnimationAnimNextRuntimeTest_GetDecoratorInterfaceSuper -> InterfaceA found");
+			AddErrorIfFalse(SuperBinding00A == SuperBinding00A_, "FAnimationAnimNextRuntimeTest_GetDecoratorInterfaceSuper -> GetInterfaceSuper methods should return the same result");
 		}
 	}
 
 	// Validate from the second base decorator
 	{
 		FDecoratorBinding ParentBinding;				// Empty, no parent
-		FDecoratorHandle DecoratorHandle03(Node0, 3);	// Point to first node, second base decorator
+		FAnimNextDecoratorHandle DecoratorHandle03(Node0, 3);	// Point to first node, second base decorator
 
 		FDecoratorPtr DecoratorPtr03 = Context.AllocateNodeInstance(ParentBinding, DecoratorHandle03);
-		AddErrorIfFalse(DecoratorPtr03.IsValid(), "FAnimationAnimNextRuntimeTest_GetInterfaceSuper -> Failed to allocate a node instance");
+		AddErrorIfFalse(DecoratorPtr03.IsValid(), "FAnimationAnimNextRuntimeTest_GetDecoratorInterfaceSuper -> Failed to allocate a node instance");
 
 		{
 			// Get a valid decorator binding: FDecoratorAC_Add
 			TDecoratorBinding<IInterfaceC> Binding04C;
-			AddErrorIfFalse(Context.GetInterface(DecoratorPtr03, Binding04C), "FAnimationAnimNextRuntimeTest_GetInterfaceSuper -> InterfaceC not found");
+			AddErrorIfFalse(Context.GetInterface(DecoratorPtr03, Binding04C), "FAnimationAnimNextRuntimeTest_GetDecoratorInterfaceSuper -> InterfaceC not found");
 
 			// Validate GetInterfaceSuper from a decorator handle
 			TDecoratorBinding<IInterfaceC> SuperBinding04C;
-			AddErrorIfFalse(!Context.GetInterfaceSuper(Binding04C.GetDecoratorPtr(), SuperBinding04C), "FAnimationAnimNextRuntimeTest_GetInterfaceSuper -> InterfaceC found");
+			AddErrorIfFalse(!Context.GetInterfaceSuper(Binding04C.GetDecoratorPtr(), SuperBinding04C), "FAnimationAnimNextRuntimeTest_GetDecoratorInterfaceSuper -> InterfaceC found");
 
 			// Validate GetInterfaceSuper from a decorator binding
 			TDecoratorBinding<IInterfaceC> SuperBinding04C_;
-			AddErrorIfFalse(!Context.GetInterfaceSuper(Binding04C, SuperBinding04C_), "FAnimationAnimNextRuntimeTest_GetInterfaceSuper -> InterfaceC found");
-			AddErrorIfFalse(SuperBinding04C == SuperBinding04C_, "FAnimationAnimNextRuntimeTest_GetInterfaceSuper -> GetInterfaceSuper methods should return the same result");
+			AddErrorIfFalse(!Context.GetInterfaceSuper(Binding04C, SuperBinding04C_), "FAnimationAnimNextRuntimeTest_GetDecoratorInterfaceSuper -> InterfaceC found");
+			AddErrorIfFalse(SuperBinding04C == SuperBinding04C_, "FAnimationAnimNextRuntimeTest_GetDecoratorInterfaceSuper -> GetInterfaceSuper methods should return the same result");
 		}
 
 		{
 			// Get a valid decorator binding: FDecoratorAC_Add
 			TDecoratorBinding<IInterfaceA> Binding04A;
-			AddErrorIfFalse(Context.GetInterface(DecoratorPtr03, Binding04A), "FAnimationAnimNextRuntimeTest_GetInterfaceSuper -> InterfaceA not found");
+			AddErrorIfFalse(Context.GetInterface(DecoratorPtr03, Binding04A), "FAnimationAnimNextRuntimeTest_GetDecoratorInterfaceSuper -> InterfaceA not found");
 
 			// Validate GetInterfaceSuper from a decorator handle, FDecoratorA_Base
 			TDecoratorBinding<IInterfaceA> SuperBinding04A;
-			AddErrorIfFalse(Context.GetInterfaceSuper(Binding04A.GetDecoratorPtr(), SuperBinding04A), "FAnimationAnimNextRuntimeTest_GetInterfaceSuper -> InterfaceA not found");
-			AddErrorIfFalse(SuperBinding04A.IsValid(), "FAnimationAnimNextRuntimeTest_GetInterfaceSuper -> InterfaceA binding not valid");
-			AddErrorIfFalse(SuperBinding04A.GetInterfaceUID() == IInterfaceA::InterfaceUID, "FAnimationAnimNextRuntimeTest_GetInterfaceSuper -> Unexpected interface UID found in decorator binding");
-			AddErrorIfFalse(SuperBinding04A.GetDecoratorPtr().GetDecoratorIndex() == 3, "FAnimationAnimNextRuntimeTest_GetInterfaceSuper -> InterfaceA not found on expected decorator");
-			AddErrorIfFalse(SuperBinding04A.GetDecoratorPtr().GetNodeInstance()->GetNodeHandle() == Node0, "FAnimationAnimNextRuntimeTest_GetInterfaceSuper -> InterfaceA not found on expected node");
-			AddErrorIfFalse(SuperBinding04A.GetSharedData<FDecoratorA_Base::FSharedData>()->DecoratorUID == FDecoratorA_Base::DecoratorUID, "FAnimationAnimNextRuntimeTest_GetInterfaceSuper -> Unexpected shared data in decorator binding");
-			AddErrorIfFalse(SuperBinding04A.GetInstanceData<FDecoratorA_Base::FInstanceData>()->DecoratorUID == FDecoratorA_Base::DecoratorUID, "FAnimationAnimNextRuntimeTest_GetInterfaceSuper -> Unexpected instance data in decorator binding");
+			AddErrorIfFalse(Context.GetInterfaceSuper(Binding04A.GetDecoratorPtr(), SuperBinding04A), "FAnimationAnimNextRuntimeTest_GetDecoratorInterfaceSuper -> InterfaceA not found");
+			AddErrorIfFalse(SuperBinding04A.IsValid(), "FAnimationAnimNextRuntimeTest_GetDecoratorInterfaceSuper -> InterfaceA binding not valid");
+			AddErrorIfFalse(SuperBinding04A.GetInterfaceUID() == IInterfaceA::InterfaceUID, "FAnimationAnimNextRuntimeTest_GetDecoratorInterfaceSuper -> Unexpected interface UID found in decorator binding");
+			AddErrorIfFalse(SuperBinding04A.GetDecoratorPtr().GetDecoratorIndex() == 3, "FAnimationAnimNextRuntimeTest_GetDecoratorInterfaceSuper -> InterfaceA not found on expected decorator");
+			AddErrorIfFalse(SuperBinding04A.GetDecoratorPtr().GetNodeInstance()->GetNodeHandle() == Node0, "FAnimationAnimNextRuntimeTest_GetDecoratorInterfaceSuper -> InterfaceA not found on expected node");
+			AddErrorIfFalse(SuperBinding04A.GetSharedData<FDecoratorA_Base::FSharedData>()->DecoratorUID == FDecoratorA_Base::DecoratorUID, "FAnimationAnimNextRuntimeTest_GetDecoratorInterfaceSuper -> Unexpected shared data in decorator binding");
+			AddErrorIfFalse(SuperBinding04A.GetInstanceData<FDecoratorA_Base::FInstanceData>()->DecoratorUID == FDecoratorA_Base::DecoratorUID, "FAnimationAnimNextRuntimeTest_GetDecoratorInterfaceSuper -> Unexpected instance data in decorator binding");
 
 			// Validate GetInterfaceSuper from a decorator binding, FDecoratorA_Base
 			TDecoratorBinding<IInterfaceA> SuperBinding04A_;
-			AddErrorIfFalse(Context.GetInterfaceSuper(Binding04A, SuperBinding04A_), "FAnimationAnimNextRuntimeTest_GetInterfaceSuper -> InterfaceA not found");
-			AddErrorIfFalse(SuperBinding04A == SuperBinding04A_, "FAnimationAnimNextRuntimeTest_GetInterfaceSuper -> GetInterfaceSuper methods should return the same result");
+			AddErrorIfFalse(Context.GetInterfaceSuper(Binding04A, SuperBinding04A_), "FAnimationAnimNextRuntimeTest_GetDecoratorInterfaceSuper -> InterfaceA not found");
+			AddErrorIfFalse(SuperBinding04A == SuperBinding04A_, "FAnimationAnimNextRuntimeTest_GetDecoratorInterfaceSuper -> GetInterfaceSuper methods should return the same result");
 
 			// Validate GetInterfaceSuper from a decorator handle
 			TDecoratorBinding<IInterfaceA> SuperBinding03A;
-			AddErrorIfFalse(!Context.GetInterfaceSuper(SuperBinding04A.GetDecoratorPtr(), SuperBinding03A), "FAnimationAnimNextRuntimeTest_GetInterfaceSuper -> InterfaceA found");
+			AddErrorIfFalse(!Context.GetInterfaceSuper(SuperBinding04A.GetDecoratorPtr(), SuperBinding03A), "FAnimationAnimNextRuntimeTest_GetDecoratorInterfaceSuper -> InterfaceA found");
 
 			// Validate GetInterfaceSuper from a decorator binding
 			TDecoratorBinding<IInterfaceA> SuperBinding03A_;
-			AddErrorIfFalse(!Context.GetInterfaceSuper(SuperBinding04A, SuperBinding03A_), "FAnimationAnimNextRuntimeTest_GetInterfaceSuper -> InterfaceA found");
-			AddErrorIfFalse(SuperBinding03A == SuperBinding03A_, "FAnimationAnimNextRuntimeTest_GetInterfaceSuper -> GetInterfaceSuper methods should return the same result");
+			AddErrorIfFalse(!Context.GetInterfaceSuper(SuperBinding04A, SuperBinding03A_), "FAnimationAnimNextRuntimeTest_GetDecoratorInterfaceSuper -> InterfaceA found");
+			AddErrorIfFalse(SuperBinding03A == SuperBinding03A_, "FAnimationAnimNextRuntimeTest_GetDecoratorInterfaceSuper -> GetInterfaceSuper methods should return the same result");
 		}
 	}
 
 	Registry.Unregister(NodeTemplate0);
 
-	AddErrorIfFalse(Registry.GetNum() == 0, "FAnimationAnimNextRuntimeTest_GetInterface -> Registry should contain 0 templates");
+	AddErrorIfFalse(Registry.GetNum() == 0, "FAnimationAnimNextRuntimeTest_GetDecoratorInterface -> Registry should contain 0 templates");
 
 	return true;
 }
+
+IMPLEMENT_SIMPLE_AUTOMATION_TEST(FAnimationAnimNextRuntimeTest_DecoratorSerialization, "Animation.AnimNext.Runtime.DecoratorSerialization", EAutomationTestFlags::EditorContext | EAutomationTestFlags::EngineFilter)
+
+bool FAnimationAnimNextRuntimeTest_DecoratorSerialization::RunTest(const FString& InParameters)
+{
+	using namespace UE::AnimNext;
+
+	AUTO_REGISTER_ANIM_DECORATOR(FDecoratorSerialization_Base)
+	AUTO_REGISTER_ANIM_DECORATOR(FDecoratorSerialization_Add)
+	AUTO_REGISTER_ANIM_DECORATOR(FDecoratorNativeSerialization_Add)
+
+	FNodeTemplateRegistry& Registry = FNodeTemplateRegistry::Get();
+
+	TArray<FDecoratorUID> NodeTemplateDecoratorList;
+	NodeTemplateDecoratorList.Add(FDecoratorSerialization_Base::DecoratorUID);
+	NodeTemplateDecoratorList.Add(FDecoratorSerialization_Add::DecoratorUID);
+	NodeTemplateDecoratorList.Add(FDecoratorNativeSerialization_Add::DecoratorUID);
+
+	// Populate our node template registry
+	TArray<uint8> NodeTemplateBuffer0;
+	const FNodeTemplate* NodeTemplate0 = BuildNodeTemplate(NodeTemplateDecoratorList, NodeTemplateBuffer0);
+
+	FNodeTemplateRegistryHandle TemplateHandle0 = Registry.FindOrAdd(NodeTemplate0);
+	AddErrorIfFalse(Registry.GetNum() == 1, "FAnimationAnimNextRuntimeTest_DecoratorSerialization -> Registry should contain 1 template");
+	AddErrorIfFalse(TemplateHandle0.IsValid(), "FAnimationAnimNextRuntimeTest_DecoratorSerialization -> Registry should contain our template");
+
+	FDecoratorSerialization_Base::FSharedData DecoratorBaseRef0;
+	DecoratorBaseRef0.Integer = 1651;
+	DecoratorBaseRef0.IntegerArray[0] = 1071;
+	DecoratorBaseRef0.IntegerArray[1] = -158;
+	DecoratorBaseRef0.IntegerArray[2] = 88116;
+	DecoratorBaseRef0.IntegerArray[3] = 0x417;
+	DecoratorBaseRef0.IntegerTArray = { -8162, 88152, 0x8152f };
+	DecoratorBaseRef0.Vector = FVector(0.1917, 12435.1, -18200.1726);
+	DecoratorBaseRef0.VectorArray[0] = FVector(192.1716, -1927.115, 99176.12);
+	DecoratorBaseRef0.VectorArray[1] = FVector(961.811, -18956.117, 81673.44);
+	DecoratorBaseRef0.VectorTArray = { FVector(-1927.8771, 1826.9917, -123.1555), FVector(9177.011, -71.44, -917.88), FVector(123.91, 852.11, -81652.1) };
+	DecoratorBaseRef0.String = TEXT("sample string 123");
+	DecoratorBaseRef0.Name = FName(TEXT("sample name 999178"));
+
+	FDecoratorSerialization_Add::FSharedData DecoratorAddRef0;
+	DecoratorAddRef0.Integer = 16511;
+	DecoratorAddRef0.IntegerArray[0] = 10711;
+	DecoratorAddRef0.IntegerArray[1] = -1581;
+	DecoratorAddRef0.IntegerArray[2] = 881161;
+	DecoratorAddRef0.IntegerArray[3] = 0x4171;
+	DecoratorAddRef0.IntegerTArray = { -81621, 881521, 0x8152f1 };
+	DecoratorAddRef0.Vector = FVector(0.19171, 12435.11, -18200.17261);
+	DecoratorAddRef0.VectorArray[0] = FVector(192.17161, -1927.1151, 99176.121);
+	DecoratorAddRef0.VectorArray[1] = FVector(961.8111, -18956.1171, 81673.441);
+	DecoratorAddRef0.VectorTArray = { FVector(-1927.87711, 1826.99171, -123.15551), FVector(9177.0111, -71.441, -917.881), FVector(123.911, 852.111, -81652.11) };
+	DecoratorAddRef0.String = TEXT("sample string 1231");
+	DecoratorAddRef0.Name = FName(TEXT("sample name 9991781"));
+
+	FDecoratorNativeSerialization_Add::FSharedData DecoratorNativeRef0;
+	DecoratorNativeRef0.Integer = 16514;
+	DecoratorNativeRef0.IntegerArray[0] = 10714;
+	DecoratorNativeRef0.IntegerArray[1] = -1584;
+	DecoratorNativeRef0.IntegerArray[2] = 881164;
+	DecoratorNativeRef0.IntegerArray[3] = 0x4174;
+	DecoratorNativeRef0.IntegerTArray = { -81624, 881524, 0x8152f4 };
+	DecoratorNativeRef0.Vector = FVector(0.19174, 12435.14, -18200.17264);
+	DecoratorNativeRef0.VectorArray[0] = FVector(192.17164, -1927.1154, 99176.124);
+	DecoratorNativeRef0.VectorArray[1] = FVector(961.8114, -18956.1174, 81673.444);
+	DecoratorNativeRef0.VectorTArray = { FVector(-1927.87714, 1826.99174, -123.15554), FVector(9177.0114, -71.444, -917.884), FVector(123.914, 852.114, -81652.14) };
+	DecoratorNativeRef0.String = TEXT("sample string 1234");
+	DecoratorNativeRef0.Name = FName(TEXT("sample name 9991784"));
+
+	FDecoratorSerialization_Base::FSharedData DecoratorBaseRef1;
+	DecoratorBaseRef1.Integer = 16512;
+	DecoratorBaseRef1.IntegerArray[0] = 10712;
+	DecoratorBaseRef1.IntegerArray[1] = -1582;
+	DecoratorBaseRef1.IntegerArray[2] = 881162;
+	DecoratorBaseRef1.IntegerArray[3] = 0x4172;
+	DecoratorBaseRef1.IntegerTArray = { -81622, 881522, 0x8152f2 };
+	DecoratorBaseRef1.Vector = FVector(0.19172, 12435.12, -18200.17262);
+	DecoratorBaseRef1.VectorArray[0] = FVector(192.17162, -1927.1152, 99176.122);
+	DecoratorBaseRef1.VectorArray[1] = FVector(961.8112, -18956.1172, 81673.442);
+	DecoratorBaseRef1.VectorTArray = { FVector(-1927.87712, 1826.99172, -123.15552), FVector(9177.0112, -71.442, -917.882), FVector(123.912, 852.112, -81652.12) };
+	DecoratorBaseRef1.String = TEXT("sample string 1232");
+	DecoratorBaseRef1.Name = FName(TEXT("sample name 9991782"));
+
+	FDecoratorSerialization_Add::FSharedData DecoratorAddRef1;
+	DecoratorAddRef1.Integer = 16513;
+	DecoratorAddRef1.IntegerArray[0] = 10713;
+	DecoratorAddRef1.IntegerArray[1] = -1583;
+	DecoratorAddRef1.IntegerArray[2] = 881163;
+	DecoratorAddRef1.IntegerArray[3] = 0x4173;
+	DecoratorAddRef1.IntegerTArray = { -81623, 881523, 0x8152f3 };
+	DecoratorAddRef1.Vector = FVector(0.19173, 12435.13, -18200.17263);
+	DecoratorAddRef1.VectorArray[0] = FVector(192.17163, -1927.1153, 99176.123);
+	DecoratorAddRef1.VectorArray[1] = FVector(961.8113, -18956.1173, 81673.443);
+	DecoratorAddRef1.VectorTArray = { FVector(-1927.87713, 1826.99173, -123.15553), FVector(9177.0113, -71.443, -917.883), FVector(123.913, 852.113, -81652.13) };
+	DecoratorAddRef1.String = TEXT("sample string 1233");
+	DecoratorAddRef1.Name = FName(TEXT("sample name 9991783"));
+
+	FDecoratorNativeSerialization_Add::FSharedData DecoratorNativeRef1;
+	DecoratorNativeRef1.Integer = 16515;
+	DecoratorNativeRef1.IntegerArray[0] = 10715;
+	DecoratorNativeRef1.IntegerArray[1] = -1585;
+	DecoratorNativeRef1.IntegerArray[2] = 881165;
+	DecoratorNativeRef1.IntegerArray[3] = 0x4175;
+	DecoratorNativeRef1.IntegerTArray = { -81625, 881525, 0x8152f5 };
+	DecoratorNativeRef1.Vector = FVector(0.19175, 12435.15, -18200.17265);
+	DecoratorNativeRef1.VectorArray[0] = FVector(192.17165, -1927.1155, 99176.125);
+	DecoratorNativeRef1.VectorArray[1] = FVector(961.8115, -18956.1175, 81673.445);
+	DecoratorNativeRef1.VectorTArray = { FVector(-1927.87715, 1826.99175, -123.15555), FVector(9177.0115, -71.445, -917.885), FVector(123.915, 852.115, -81652.15) };
+	DecoratorNativeRef1.String = TEXT("sample string 1235");
+	DecoratorNativeRef1.Name = FName(TEXT("sample name 9991785"));
+
+	FNodeHandle Node0;
+	FNodeHandle Node1;
+
+	// Write our graph
+	TArray<uint8> GraphSharedDataArchiveBuffer;
+	{
+		FDecoratorWriter DecoratorWriter;
+
+		Node0 = DecoratorWriter.RegisterNode(*NodeTemplate0);
+		Node1 = DecoratorWriter.RegisterNode(*NodeTemplate0);
+
+		// We don't have decorator properties
+		TArray<TMap<FString, FString>> DecoratorProperties0;
+		DecoratorProperties0.AddDefaulted(NodeTemplateDecoratorList.Num());
+
+		DecoratorProperties0[0].Add(TEXT("Integer"), ToString<FDecoratorSerialization_Base::FSharedData>(TEXT("Integer"), DecoratorBaseRef0.Integer));
+		DecoratorProperties0[0].Add(TEXT("IntegerArray"), ToString<FDecoratorSerialization_Base::FSharedData>(TEXT("IntegerArray"), DecoratorBaseRef0.IntegerArray));
+		DecoratorProperties0[0].Add(TEXT("IntegerTArray"), ToString<FDecoratorSerialization_Base::FSharedData>(TEXT("IntegerTArray"), DecoratorBaseRef0.IntegerTArray));
+		DecoratorProperties0[0].Add(TEXT("Vector"), ToString<FDecoratorSerialization_Base::FSharedData>(TEXT("Vector"), DecoratorBaseRef0.Vector));
+		DecoratorProperties0[0].Add(TEXT("VectorArray"), ToString<FDecoratorSerialization_Base::FSharedData>(TEXT("VectorArray"), DecoratorBaseRef0.VectorArray));
+		DecoratorProperties0[0].Add(TEXT("VectorTArray"), ToString<FDecoratorSerialization_Base::FSharedData>(TEXT("VectorTArray"), DecoratorBaseRef0.VectorTArray));
+		DecoratorProperties0[0].Add(TEXT("String"), ToString<FDecoratorSerialization_Base::FSharedData>(TEXT("String"), DecoratorBaseRef0.String));
+		DecoratorProperties0[0].Add(TEXT("Name"), ToString<FDecoratorSerialization_Base::FSharedData>(TEXT("Name"), DecoratorBaseRef0.Name));
+
+		DecoratorProperties0[1].Add(TEXT("Integer"), ToString<FDecoratorSerialization_Add::FSharedData>(TEXT("Integer"), DecoratorAddRef0.Integer));
+		DecoratorProperties0[1].Add(TEXT("IntegerArray"), ToString<FDecoratorSerialization_Add::FSharedData>(TEXT("IntegerArray"), DecoratorAddRef0.IntegerArray));
+		DecoratorProperties0[1].Add(TEXT("IntegerTArray"), ToString<FDecoratorSerialization_Add::FSharedData>(TEXT("IntegerTArray"), DecoratorAddRef0.IntegerTArray));
+		DecoratorProperties0[1].Add(TEXT("Vector"), ToString<FDecoratorSerialization_Add::FSharedData>(TEXT("Vector"), DecoratorAddRef0.Vector));
+		DecoratorProperties0[1].Add(TEXT("VectorArray"), ToString<FDecoratorSerialization_Add::FSharedData>(TEXT("VectorArray"), DecoratorAddRef0.VectorArray));
+		DecoratorProperties0[1].Add(TEXT("VectorTArray"), ToString<FDecoratorSerialization_Add::FSharedData>(TEXT("VectorTArray"), DecoratorAddRef0.VectorTArray));
+		DecoratorProperties0[1].Add(TEXT("String"), ToString<FDecoratorSerialization_Add::FSharedData>(TEXT("String"), DecoratorAddRef0.String));
+		DecoratorProperties0[1].Add(TEXT("Name"), ToString<FDecoratorSerialization_Add::FSharedData>(TEXT("Name"), DecoratorAddRef0.Name));
+
+		DecoratorProperties0[2].Add(TEXT("Integer"), ToString<FDecoratorNativeSerialization_Add::FSharedData>(TEXT("Integer"), DecoratorNativeRef0.Integer));
+		DecoratorProperties0[2].Add(TEXT("IntegerArray"), ToString<FDecoratorNativeSerialization_Add::FSharedData>(TEXT("IntegerArray"), DecoratorNativeRef0.IntegerArray));
+		DecoratorProperties0[2].Add(TEXT("IntegerTArray"), ToString<FDecoratorNativeSerialization_Add::FSharedData>(TEXT("IntegerTArray"), DecoratorNativeRef0.IntegerTArray));
+		DecoratorProperties0[2].Add(TEXT("Vector"), ToString<FDecoratorNativeSerialization_Add::FSharedData>(TEXT("Vector"), DecoratorNativeRef0.Vector));
+		DecoratorProperties0[2].Add(TEXT("VectorArray"), ToString<FDecoratorNativeSerialization_Add::FSharedData>(TEXT("VectorArray"), DecoratorNativeRef0.VectorArray));
+		DecoratorProperties0[2].Add(TEXT("VectorTArray"), ToString<FDecoratorNativeSerialization_Add::FSharedData>(TEXT("VectorTArray"), DecoratorNativeRef0.VectorTArray));
+		DecoratorProperties0[2].Add(TEXT("String"), ToString<FDecoratorNativeSerialization_Add::FSharedData>(TEXT("String"), DecoratorNativeRef0.String));
+		DecoratorProperties0[2].Add(TEXT("Name"), ToString<FDecoratorNativeSerialization_Add::FSharedData>(TEXT("Name"), DecoratorNativeRef0.Name));
+
+		TArray<TMap<FString, FString>> DecoratorProperties1;
+		DecoratorProperties1.AddDefaulted(NodeTemplateDecoratorList.Num());
+
+		DecoratorProperties1[0].Add(TEXT("Integer"), ToString<FDecoratorSerialization_Base::FSharedData>(TEXT("Integer"), DecoratorBaseRef1.Integer));
+		DecoratorProperties1[0].Add(TEXT("IntegerArray"), ToString<FDecoratorSerialization_Base::FSharedData>(TEXT("IntegerArray"), DecoratorBaseRef1.IntegerArray));
+		DecoratorProperties1[0].Add(TEXT("IntegerTArray"), ToString<FDecoratorSerialization_Base::FSharedData>(TEXT("IntegerTArray"), DecoratorBaseRef1.IntegerTArray));
+		DecoratorProperties1[0].Add(TEXT("Vector"), ToString<FDecoratorSerialization_Base::FSharedData>(TEXT("Vector"), DecoratorBaseRef1.Vector));
+		DecoratorProperties1[0].Add(TEXT("VectorArray"), ToString<FDecoratorSerialization_Base::FSharedData>(TEXT("VectorArray"), DecoratorBaseRef1.VectorArray));
+		DecoratorProperties1[0].Add(TEXT("VectorTArray"), ToString<FDecoratorSerialization_Base::FSharedData>(TEXT("VectorTArray"), DecoratorBaseRef1.VectorTArray));
+		DecoratorProperties1[0].Add(TEXT("String"), ToString<FDecoratorSerialization_Base::FSharedData>(TEXT("String"), DecoratorBaseRef1.String));
+		DecoratorProperties1[0].Add(TEXT("Name"), ToString<FDecoratorSerialization_Base::FSharedData>(TEXT("Name"), DecoratorBaseRef1.Name));
+
+		DecoratorProperties1[1].Add(TEXT("Integer"), ToString<FDecoratorSerialization_Add::FSharedData>(TEXT("Integer"), DecoratorAddRef1.Integer));
+		DecoratorProperties1[1].Add(TEXT("IntegerArray"), ToString<FDecoratorSerialization_Add::FSharedData>(TEXT("IntegerArray"), DecoratorAddRef1.IntegerArray));
+		DecoratorProperties1[1].Add(TEXT("IntegerTArray"), ToString<FDecoratorSerialization_Add::FSharedData>(TEXT("IntegerTArray"), DecoratorAddRef1.IntegerTArray));
+		DecoratorProperties1[1].Add(TEXT("Vector"), ToString<FDecoratorSerialization_Add::FSharedData>(TEXT("Vector"), DecoratorAddRef1.Vector));
+		DecoratorProperties1[1].Add(TEXT("VectorArray"), ToString<FDecoratorSerialization_Add::FSharedData>(TEXT("VectorArray"), DecoratorAddRef1.VectorArray));
+		DecoratorProperties1[1].Add(TEXT("VectorTArray"), ToString<FDecoratorSerialization_Add::FSharedData>(TEXT("VectorTArray"), DecoratorAddRef1.VectorTArray));
+		DecoratorProperties1[1].Add(TEXT("String"), ToString<FDecoratorSerialization_Add::FSharedData>(TEXT("String"), DecoratorAddRef1.String));
+		DecoratorProperties1[1].Add(TEXT("Name"), ToString<FDecoratorSerialization_Add::FSharedData>(TEXT("Name"), DecoratorAddRef1.Name));
+
+		DecoratorProperties1[2].Add(TEXT("Integer"), ToString<FDecoratorNativeSerialization_Add::FSharedData>(TEXT("Integer"), DecoratorNativeRef1.Integer));
+		DecoratorProperties1[2].Add(TEXT("IntegerArray"), ToString<FDecoratorNativeSerialization_Add::FSharedData>(TEXT("IntegerArray"), DecoratorNativeRef1.IntegerArray));
+		DecoratorProperties1[2].Add(TEXT("IntegerTArray"), ToString<FDecoratorNativeSerialization_Add::FSharedData>(TEXT("IntegerTArray"), DecoratorNativeRef1.IntegerTArray));
+		DecoratorProperties1[2].Add(TEXT("Vector"), ToString<FDecoratorNativeSerialization_Add::FSharedData>(TEXT("Vector"), DecoratorNativeRef1.Vector));
+		DecoratorProperties1[2].Add(TEXT("VectorArray"), ToString<FDecoratorNativeSerialization_Add::FSharedData>(TEXT("VectorArray"), DecoratorNativeRef1.VectorArray));
+		DecoratorProperties1[2].Add(TEXT("VectorTArray"), ToString<FDecoratorNativeSerialization_Add::FSharedData>(TEXT("VectorTArray"), DecoratorNativeRef1.VectorTArray));
+		DecoratorProperties1[2].Add(TEXT("String"), ToString<FDecoratorNativeSerialization_Add::FSharedData>(TEXT("String"), DecoratorNativeRef1.String));
+		DecoratorProperties1[2].Add(TEXT("Name"), ToString<FDecoratorNativeSerialization_Add::FSharedData>(TEXT("Name"), DecoratorNativeRef1.Name));
+
+		DecoratorWriter.BeginNodeWriting();
+		DecoratorWriter.WriteNode(Node0,
+			[&DecoratorProperties0](uint32 DecoratorIndex) -> const TMap<FString, FString>&
+			{
+				return DecoratorProperties0[DecoratorIndex];
+			});
+		DecoratorWriter.WriteNode(Node1,
+			[&DecoratorProperties1](uint32 DecoratorIndex) -> const TMap<FString, FString>&
+			{
+				return DecoratorProperties1[DecoratorIndex];
+			});
+		DecoratorWriter.EndNodeWriting();
+
+		AddErrorIfFalse(DecoratorWriter.GetErrorState() == FDecoratorWriter::EErrorState::None, "FAnimationAnimNextRuntimeTest_DecoratorSerialization -> Failed to write decorators");
+		GraphSharedDataArchiveBuffer = DecoratorWriter.GetGraphSharedData();
+	}
+
+	// Clear out the node template registry to test registration on load
+	Registry.Clear();
+
+	AddErrorIfFalse(Registry.GetNum() == 0, "FAnimationAnimNextRuntimeTest_DecoratorSerialization -> Registry should contain 0 templates");
+
+	// Read our graph
+	TArray<uint8> GraphSharedDataBuffer;
+	{
+		FMemoryReader GraphSharedDataArchive(GraphSharedDataArchiveBuffer);
+		FDecoratorReader DecoratorReader(GraphSharedDataArchive);
+
+		DecoratorReader.ReadGraphSharedData(GraphSharedDataBuffer);
+
+		AddErrorIfFalse(Registry.GetNum() == 1, "FAnimationAnimNextRuntimeTest_DecoratorSerialization -> Registry should contain 1 template");
+	}
+
+	FExecutionContext Context(GraphSharedDataBuffer);
+
+	// Validate decorator serialization
+	{
+		FDecoratorBinding ParentBinding;				// Empty, no parent
+		FAnimNextDecoratorHandle DecoratorHandle0(Node0, 0);	// Point to first node, first base decorator
+		FAnimNextDecoratorHandle DecoratorHandle1(Node1, 0);	// Point to second node, first base decorator
+
+		FDecoratorPtr DecoratorPtr0 = Context.AllocateNodeInstance(ParentBinding, DecoratorHandle0);
+		AddErrorIfFalse(DecoratorPtr0.IsValid(), "FAnimationAnimNextRuntimeTest_DecoratorSerialization -> Failed to allocate a node instance");
+
+		FDecoratorPtr DecoratorPtr1 = Context.AllocateNodeInstance(ParentBinding, DecoratorHandle1);
+		AddErrorIfFalse(DecoratorPtr1.IsValid(), "FAnimationAnimNextRuntimeTest_DecoratorSerialization -> Failed to allocate a node instance");
+
+		// Validate shared data for base decorator on node 0
+		{
+			TDecoratorBinding<IInterfaceA> BindingA0;
+			AddErrorIfFalse(Context.GetInterface(DecoratorPtr0, BindingA0), "FAnimationAnimNextRuntimeTest_DecoratorSerialization -> InterfaceA not found");
+
+			const auto* SharedDataA0 = BindingA0.GetSharedData<FDecoratorSerialization_Base::FSharedData>();
+
+			AddErrorIfFalse(SharedDataA0->Integer == DecoratorBaseRef0.Integer, "FAnimationAnimNextRuntimeTest_DecoratorSerialization -> Unexpected serialized value");
+			AddErrorIfFalse(FMemory::Memcmp(SharedDataA0->IntegerArray, DecoratorBaseRef0.IntegerArray, sizeof(DecoratorBaseRef0.IntegerArray)) == 0, "FAnimationAnimNextRuntimeTest_DecoratorSerialization -> Unexpected serialized value");
+			AddErrorIfFalse(SharedDataA0->IntegerTArray == DecoratorBaseRef0.IntegerTArray, "FAnimationAnimNextRuntimeTest_DecoratorSerialization -> Unexpected serialized value");
+			AddErrorIfFalse(SharedDataA0->Vector == DecoratorBaseRef0.Vector, "FAnimationAnimNextRuntimeTest_DecoratorSerialization -> Unexpected serialized value");
+			AddErrorIfFalse(FMemory::Memcmp(SharedDataA0->VectorArray, DecoratorBaseRef0.VectorArray, sizeof(DecoratorBaseRef0.VectorArray)) == 0, "FAnimationAnimNextRuntimeTest_DecoratorSerialization -> Unexpected serialized value");
+			AddErrorIfFalse(SharedDataA0->VectorTArray == DecoratorBaseRef0.VectorTArray, "FAnimationAnimNextRuntimeTest_DecoratorSerialization -> Unexpected serialized value");
+			AddErrorIfFalse(SharedDataA0->String == DecoratorBaseRef0.String, "FAnimationAnimNextRuntimeTest_DecoratorSerialization -> Unexpected serialized value");
+			AddErrorIfFalse(SharedDataA0->Name == DecoratorBaseRef0.Name, "FAnimationAnimNextRuntimeTest_DecoratorSerialization -> Unexpected serialized value");
+		}
+
+		// Validate shared data for additive decorator on node 0
+		{
+			TDecoratorBinding<IInterfaceB> BindingB0;
+			AddErrorIfFalse(Context.GetInterface(DecoratorPtr0, BindingB0), "FAnimationAnimNextRuntimeTest_DecoratorSerialization -> InterfaceB not found");
+
+			const auto* SharedDataB0 = BindingB0.GetSharedData<FDecoratorSerialization_Add::FSharedData>();
+
+			AddErrorIfFalse(SharedDataB0->Integer == DecoratorAddRef0.Integer, "FAnimationAnimNextRuntimeTest_DecoratorSerialization -> Unexpected serialized value");
+			AddErrorIfFalse(FMemory::Memcmp(SharedDataB0->IntegerArray, DecoratorAddRef0.IntegerArray, sizeof(DecoratorAddRef0.IntegerArray)) == 0, "FAnimationAnimNextRuntimeTest_DecoratorSerialization -> Unexpected serialized value");
+			AddErrorIfFalse(SharedDataB0->IntegerTArray == DecoratorAddRef0.IntegerTArray, "FAnimationAnimNextRuntimeTest_DecoratorSerialization -> Unexpected serialized value");
+			AddErrorIfFalse(SharedDataB0->Vector == DecoratorAddRef0.Vector, "FAnimationAnimNextRuntimeTest_DecoratorSerialization -> Unexpected serialized value");
+			AddErrorIfFalse(FMemory::Memcmp(SharedDataB0->VectorArray, DecoratorAddRef0.VectorArray, sizeof(DecoratorAddRef0.VectorArray)) == 0, "FAnimationAnimNextRuntimeTest_DecoratorSerialization -> Unexpected serialized value");
+			AddErrorIfFalse(SharedDataB0->VectorTArray == DecoratorAddRef0.VectorTArray, "FAnimationAnimNextRuntimeTest_DecoratorSerialization -> Unexpected serialized value");
+			AddErrorIfFalse(SharedDataB0->String == DecoratorAddRef0.String, "FAnimationAnimNextRuntimeTest_DecoratorSerialization -> Unexpected serialized value");
+			AddErrorIfFalse(SharedDataB0->Name == DecoratorAddRef0.Name, "FAnimationAnimNextRuntimeTest_DecoratorSerialization -> Unexpected serialized value");
+		}
+
+		// Validate shared data for native decorator on node 0
+		{
+			TDecoratorBinding<IInterfaceC> BindingC0;
+			AddErrorIfFalse(Context.GetInterface(DecoratorPtr0, BindingC0), "FAnimationAnimNextRuntimeTest_DecoratorSerialization -> InterfaceC not found");
+
+			const auto* SharedDataC0 = BindingC0.GetSharedData<FDecoratorNativeSerialization_Add::FSharedData>();
+
+			AddErrorIfFalse(SharedDataC0->Integer == DecoratorNativeRef0.Integer, "FAnimationAnimNextRuntimeTest_DecoratorSerialization -> Unexpected serialized value");
+			AddErrorIfFalse(FMemory::Memcmp(SharedDataC0->IntegerArray, DecoratorNativeRef0.IntegerArray, sizeof(DecoratorNativeRef0.IntegerArray)) == 0, "FAnimationAnimNextRuntimeTest_DecoratorSerialization -> Unexpected serialized value");
+			AddErrorIfFalse(SharedDataC0->IntegerTArray == DecoratorNativeRef0.IntegerTArray, "FAnimationAnimNextRuntimeTest_DecoratorSerialization -> Unexpected serialized value");
+			AddErrorIfFalse(SharedDataC0->Vector == DecoratorNativeRef0.Vector, "FAnimationAnimNextRuntimeTest_DecoratorSerialization -> Unexpected serialized value");
+			AddErrorIfFalse(FMemory::Memcmp(SharedDataC0->VectorArray, DecoratorNativeRef0.VectorArray, sizeof(DecoratorNativeRef0.VectorArray)) == 0, "FAnimationAnimNextRuntimeTest_DecoratorSerialization -> Unexpected serialized value");
+			AddErrorIfFalse(SharedDataC0->VectorTArray == DecoratorNativeRef0.VectorTArray, "FAnimationAnimNextRuntimeTest_DecoratorSerialization -> Unexpected serialized value");
+			AddErrorIfFalse(SharedDataC0->String == DecoratorNativeRef0.String, "FAnimationAnimNextRuntimeTest_DecoratorSerialization -> Unexpected serialized value");
+			AddErrorIfFalse(SharedDataC0->Name == DecoratorNativeRef0.Name, "FAnimationAnimNextRuntimeTest_DecoratorSerialization -> Unexpected serialized value");
+			AddErrorIfFalse(SharedDataC0->bSerializeCalled, "FAnimationAnimNextRuntimeTest_DecoratorSerialization -> Unexpected serialized value");
+		}
+
+		// Validate shared data for base decorator on node 1
+		{
+			TDecoratorBinding<IInterfaceA> BindingA1;
+			AddErrorIfFalse(Context.GetInterface(DecoratorPtr1, BindingA1), "FAnimationAnimNextRuntimeTest_DecoratorSerialization -> InterfaceA not found");
+
+			const auto* SharedDataA1 = BindingA1.GetSharedData<FDecoratorSerialization_Base::FSharedData>();
+
+			AddErrorIfFalse(SharedDataA1->Integer == DecoratorBaseRef1.Integer, "FAnimationAnimNextRuntimeTest_DecoratorSerialization -> Unexpected serialized value");
+			AddErrorIfFalse(FMemory::Memcmp(SharedDataA1->IntegerArray, DecoratorBaseRef1.IntegerArray, sizeof(DecoratorBaseRef1.IntegerArray)) == 0, "FAnimationAnimNextRuntimeTest_DecoratorSerialization -> Unexpected serialized value");
+			AddErrorIfFalse(SharedDataA1->IntegerTArray == DecoratorBaseRef1.IntegerTArray, "FAnimationAnimNextRuntimeTest_DecoratorSerialization -> Unexpected serialized value");
+			AddErrorIfFalse(SharedDataA1->Vector == DecoratorBaseRef1.Vector, "FAnimationAnimNextRuntimeTest_DecoratorSerialization -> Unexpected serialized value");
+			AddErrorIfFalse(FMemory::Memcmp(SharedDataA1->VectorArray, DecoratorBaseRef1.VectorArray, sizeof(DecoratorBaseRef1.VectorArray)) == 0, "FAnimationAnimNextRuntimeTest_DecoratorSerialization -> Unexpected serialized value");
+			AddErrorIfFalse(SharedDataA1->VectorTArray == DecoratorBaseRef1.VectorTArray, "FAnimationAnimNextRuntimeTest_DecoratorSerialization -> Unexpected serialized value");
+			AddErrorIfFalse(SharedDataA1->String == DecoratorBaseRef1.String, "FAnimationAnimNextRuntimeTest_DecoratorSerialization -> Unexpected serialized value");
+			AddErrorIfFalse(SharedDataA1->Name == DecoratorBaseRef1.Name, "FAnimationAnimNextRuntimeTest_DecoratorSerialization -> Unexpected serialized value");
+		}
+
+		// Validate shared data for additive decorator on node 1
+		{
+			TDecoratorBinding<IInterfaceB> BindingB1;
+			AddErrorIfFalse(Context.GetInterface(DecoratorPtr1, BindingB1), "FAnimationAnimNextRuntimeTest_DecoratorSerialization -> InterfaceB not found");
+
+			const auto* SharedDataB1 = BindingB1.GetSharedData<FDecoratorSerialization_Add::FSharedData>();
+
+			AddErrorIfFalse(SharedDataB1->Integer == DecoratorAddRef1.Integer, "FAnimationAnimNextRuntimeTest_DecoratorSerialization -> Unexpected serialized value");
+			AddErrorIfFalse(FMemory::Memcmp(SharedDataB1->IntegerArray, DecoratorAddRef1.IntegerArray, sizeof(DecoratorAddRef1.IntegerArray)) == 0, "FAnimationAnimNextRuntimeTest_DecoratorSerialization -> Unexpected serialized value");
+			AddErrorIfFalse(SharedDataB1->IntegerTArray == DecoratorAddRef1.IntegerTArray, "FAnimationAnimNextRuntimeTest_DecoratorSerialization -> Unexpected serialized value");
+			AddErrorIfFalse(SharedDataB1->Vector == DecoratorAddRef1.Vector, "FAnimationAnimNextRuntimeTest_DecoratorSerialization -> Unexpected serialized value");
+			AddErrorIfFalse(FMemory::Memcmp(SharedDataB1->VectorArray, DecoratorAddRef1.VectorArray, sizeof(DecoratorAddRef1.VectorArray)) == 0, "FAnimationAnimNextRuntimeTest_DecoratorSerialization -> Unexpected serialized value");
+			AddErrorIfFalse(SharedDataB1->VectorTArray == DecoratorAddRef1.VectorTArray, "FAnimationAnimNextRuntimeTest_DecoratorSerialization -> Unexpected serialized value");
+			AddErrorIfFalse(SharedDataB1->String == DecoratorAddRef1.String, "FAnimationAnimNextRuntimeTest_DecoratorSerialization -> Unexpected serialized value");
+			AddErrorIfFalse(SharedDataB1->Name == DecoratorAddRef1.Name, "FAnimationAnimNextRuntimeTest_DecoratorSerialization -> Unexpected serialized value");
+		}
+
+		// Validate shared data for native decorator on node 1
+		{
+			TDecoratorBinding<IInterfaceC> BindingC1;
+			AddErrorIfFalse(Context.GetInterface(DecoratorPtr1, BindingC1), "FAnimationAnimNextRuntimeTest_DecoratorSerialization -> InterfaceC not found");
+
+			const auto* SharedDataC1 = BindingC1.GetSharedData<FDecoratorNativeSerialization_Add::FSharedData>();
+
+			AddErrorIfFalse(SharedDataC1->Integer == DecoratorNativeRef1.Integer, "FAnimationAnimNextRuntimeTest_DecoratorSerialization -> Unexpected serialized value");
+			AddErrorIfFalse(FMemory::Memcmp(SharedDataC1->IntegerArray, DecoratorNativeRef1.IntegerArray, sizeof(DecoratorNativeRef1.IntegerArray)) == 0, "FAnimationAnimNextRuntimeTest_DecoratorSerialization -> Unexpected serialized value");
+			AddErrorIfFalse(SharedDataC1->IntegerTArray == DecoratorNativeRef1.IntegerTArray, "FAnimationAnimNextRuntimeTest_DecoratorSerialization -> Unexpected serialized value");
+			AddErrorIfFalse(SharedDataC1->Vector == DecoratorNativeRef1.Vector, "FAnimationAnimNextRuntimeTest_DecoratorSerialization -> Unexpected serialized value");
+			AddErrorIfFalse(FMemory::Memcmp(SharedDataC1->VectorArray, DecoratorNativeRef1.VectorArray, sizeof(DecoratorNativeRef1.VectorArray)) == 0, "FAnimationAnimNextRuntimeTest_DecoratorSerialization -> Unexpected serialized value");
+			AddErrorIfFalse(SharedDataC1->VectorTArray == DecoratorNativeRef1.VectorTArray, "FAnimationAnimNextRuntimeTest_DecoratorSerialization -> Unexpected serialized value");
+			AddErrorIfFalse(SharedDataC1->String == DecoratorNativeRef1.String, "FAnimationAnimNextRuntimeTest_DecoratorSerialization -> Unexpected serialized value");
+			AddErrorIfFalse(SharedDataC1->Name == DecoratorNativeRef1.Name, "FAnimationAnimNextRuntimeTest_DecoratorSerialization -> Unexpected serialized value");
+			AddErrorIfFalse(SharedDataC1->bSerializeCalled, "FAnimationAnimNextRuntimeTest_DecoratorSerialization -> Unexpected serialized value");
+		}
+	}
+
+	Registry.Clear();
+
+	AddErrorIfFalse(Registry.GetNum() == 0, "FAnimationAnimNextRuntimeTest_DecoratorSerialization -> Registry should contain 0 templates");
+
+	return true;
+}
+
 #endif
+
+FDecoratorA_BaseSharedData::FDecoratorA_BaseSharedData()
+	: DecoratorUID(UE::AnimNext::FDecoratorA_Base::DecoratorUID.GetUID())
+{
+}
+
+FDecoratorAB_AddSharedData::FDecoratorAB_AddSharedData()
+	: DecoratorUID(UE::AnimNext::FDecoratorAB_Add::DecoratorUID.GetUID())
+{
+}
+
+FDecoratorAC_AddSharedData::FDecoratorAC_AddSharedData()
+	: DecoratorUID(UE::AnimNext::FDecoratorAC_Add::DecoratorUID.GetUID())
+{
+}
+
+bool FDecoratorNativeSerialization_AddSharedData::Serialize(FArchive& Ar)
+{
+	Ar << Integer;
+
+	int32 integerArrayCount = sizeof(IntegerArray) / sizeof(IntegerArray[0]);
+	Ar << integerArrayCount;
+	for (int32 i = 0; i < integerArrayCount; ++i)
+	{
+		Ar << IntegerArray[i];
+	}
+
+	Ar << IntegerTArray;
+	Ar << Vector;
+
+	int32 vectorArrayCount = sizeof(VectorArray) / sizeof(VectorArray[0]);
+	Ar << vectorArrayCount;
+	for (int32 i = 0; i < vectorArrayCount; ++i)
+	{
+		Ar << VectorArray[i];
+	}
+
+	Ar << VectorTArray;
+	Ar << String;
+	Ar << Name;
+
+	bSerializeCalled = true;
+
+	return true;
+}
diff --git a/Engine/Plugins/Experimental/AnimNext/Source/AnimNextTestSuite/Private/AnimNextDecoratorBaseTest.h b/Engine/Plugins/Experimental/AnimNext/Source/AnimNextTestSuite/Private/AnimNextDecoratorBaseTest.h
new file mode 100644
--- /dev/null
+++ b/Engine/Plugins/Experimental/AnimNext/Source/AnimNextTestSuite/Private/AnimNextDecoratorBaseTest.h
@@ -0,0 +1,145 @@
+// Copyright Epic Games, Inc. All Rights Reserved.
+
+#pragma once
+
+#include "CoreMinimal.h"
+
+#include "DecoratorBase/DecoratorSharedData.h"
+
+#include "AnimNextDecoratorBaseTest.generated.h"
+
+USTRUCT()
+struct FDecoratorA_BaseSharedData : public FAnimNextDecoratorSharedData
+{
+	GENERATED_BODY()
+
+	UPROPERTY()
+	uint32 DecoratorUID;
+
+	FDecoratorA_BaseSharedData();
+};
+
+USTRUCT()
+struct FDecoratorAB_AddSharedData : public FAnimNextDecoratorSharedData
+{
+	GENERATED_BODY()
+
+	UPROPERTY()
+	uint32 DecoratorUID;
+
+	FDecoratorAB_AddSharedData();
+};
+
+USTRUCT()
+struct FDecoratorAC_AddSharedData : public FAnimNextDecoratorSharedData
+{
+	GENERATED_BODY()
+
+	UPROPERTY()
+	uint32 DecoratorUID;
+
+	FDecoratorAC_AddSharedData();
+};
+
+USTRUCT()
+struct FDecoratorSerialization_BaseSharedData : public FAnimNextDecoratorSharedData
+{
+	GENERATED_BODY()
+
+	UPROPERTY()
+	int32 Integer;
+
+	UPROPERTY()
+	int32 IntegerArray[4];
+
+	UPROPERTY()
+	TArray<int32> IntegerTArray;
+
+	UPROPERTY()
+	FVector Vector;
+
+	UPROPERTY()
+	FVector VectorArray[2];
+
+	UPROPERTY()
+	TArray<FVector> VectorTArray;
+
+	UPROPERTY()
+	FString String;
+
+	UPROPERTY()
+	FName Name;
+};
+
+USTRUCT()
+struct FDecoratorSerialization_AddSharedData : public FAnimNextDecoratorSharedData
+{
+	GENERATED_BODY()
+
+	UPROPERTY()
+	int32 Integer;
+
+	UPROPERTY()
+	int32 IntegerArray[4];
+
+	UPROPERTY()
+	TArray<int32> IntegerTArray;
+
+	UPROPERTY()
+	FVector Vector;
+
+	UPROPERTY()
+	FVector VectorArray[2];
+
+	UPROPERTY()
+	TArray<FVector> VectorTArray;
+
+	UPROPERTY()
+	FString String;
+
+	UPROPERTY()
+	FName Name;
+};
+
+USTRUCT()
+struct FDecoratorNativeSerialization_AddSharedData : public FAnimNextDecoratorSharedData
+{
+	GENERATED_BODY()
+
+	UPROPERTY()
+	int32 Integer;
+
+	UPROPERTY()
+	int32 IntegerArray[4];
+
+	UPROPERTY()
+	TArray<int32> IntegerTArray;
+
+	UPROPERTY()
+	FVector Vector;
+
+	UPROPERTY()
+	FVector VectorArray[2];
+
+	UPROPERTY()
+	TArray<FVector> VectorTArray;
+
+	UPROPERTY()
+	FString String;
+
+	UPROPERTY()
+	FName Name;
+
+	bool bSerializeCalled = false;
+
+	bool Serialize(FArchive& Ar);
+};
+
+template<>
+struct TStructOpsTypeTraits<FDecoratorNativeSerialization_AddSharedData> : public TStructOpsTypeTraitsBase2<FDecoratorNativeSerialization_AddSharedData>
+{
+	enum
+	{
+		WithSerializer = true,
+	};
+};
diff --git a/Engine/Plugins/Experimental/AnimNext/Source/AnimNextTestSuite/Private/AnimNextDecoratorInterfacesTest.cpp b/Engine/Plugins/Experimental/AnimNext/Source/AnimNextTestSuite/Private/AnimNextDecoratorInterfacesTest.cpp
--- a/Engine/Plugins/Experimental/AnimNext/Source/AnimNextTestSuite/Private/AnimNextDecoratorInterfacesTest.cpp
+++ b/Engine/Plugins/Experimental/AnimNext/Source/AnimNextTestSuite/Private/AnimNextDecoratorInterfacesTest.cpp
@@ -1,14 +1,18 @@
 // Copyright Epic Games, Inc. All Rights Reserved.
 
-
+#include "AnimNextDecoratorInterfacesTest.h"
 #include "AnimNextRuntimeTest.h"
 
 #include "Misc/AutomationTest.h"
 
 #if WITH_DEV_AUTOMATION_TESTS
 
+#include "Serialization/MemoryReader.h"
+
 #include "DecoratorBase/Decorator.h"
+#include "DecoratorBase/DecoratorReader.h"
 #include "DecoratorBase/DecoratorRegistry.h"
+#include "DecoratorBase/DecoratorWriter.h"
 #include "DecoratorBase/ExecutionContext.h"
 #include "DecoratorBase/IDecoratorInterface.h"
 #include "DecoratorBase/NodeInstance.h"
@@ -29,12 +33,9 @@ namespace UE::AnimNext
 		static TArray<FDecoratorUID>* EvaluatedDecorators = nullptr;
 	}
 
-	struct FDecoratorWithNoChildren : FDecorator, IHierarchy, IUpdate, IEvaluate
+	struct FDecoratorWithNoChildren : FBaseDecorator, IHierarchy, IUpdate, IEvaluate
 	{
-		DECLARE_ANIM_DECORATOR(FDecoratorWithNoChildren, 0xe2400d2a, FDecorator)
-
-		// FDecorator impl
-		virtual EDecoratorMode GetMode() const override { return EDecoratorMode::Base; }
+		DECLARE_ANIM_DECORATOR(FDecoratorWithNoChildren, 0xe2400d2a, FBaseDecorator)
 
 		// IHierarchy impl
 		virtual void GetChildren(FExecutionContext& Context, const TDecoratorBinding<IHierarchy>& Binding, FChildrenArray& Children) const override
@@ -92,18 +93,11 @@ namespace UE::AnimNext
 	DEFINE_ANIM_DECORATOR_END(FDecoratorWithNoChildren)
 
 	// This decorator does not update or evaluate
-	struct FDecoratorWithOneChild : FDecorator, IHierarchy
+	struct FDecoratorWithOneChild : FBaseDecorator, IHierarchy
 	{
-		DECLARE_ANIM_DECORATOR(FDecoratorWithOneChild, 0xba24d224, FDecorator)
+		DECLARE_ANIM_DECORATOR(FDecoratorWithOneChild, 0xba24d224, FBaseDecorator)
 
-		struct FSharedData : FDecorator::FSharedData
-		{
-			FDecoratorHandle Child;
-
-			FSharedData(FDecoratorHandle Child_)
-				: Child(Child_)
-			{}
-		};
+		using FSharedData = FDecoratorWithOneChildSharedData;
 
 		struct FInstanceData : FDecorator::FInstanceData
 		{
@@ -115,9 +109,6 @@ namespace UE::AnimNext
 			}
 		};
 
-		// FDecorator impl
-		virtual EDecoratorMode GetMode() const override { return EDecoratorMode::Base; }
-
 		// IHierarchy impl
 		virtual void GetChildren(FExecutionContext& Context, const TDecoratorBinding<IHierarchy>& Binding, FChildrenArray& Children) const override
 		{
@@ -133,18 +124,11 @@ namespace UE::AnimNext
 		DEFINE_ANIM_DECORATOR_IMPLEMENTS_INTERFACE(IHierarchy)
 	DEFINE_ANIM_DECORATOR_END(FDecoratorWithOneChild)
 
-	struct FDecoratorWithChildren : FDecorator, IHierarchy, IUpdate, IEvaluate
+	struct FDecoratorWithChildren : FBaseDecorator, IHierarchy, IUpdate, IEvaluate
 	{
-		DECLARE_ANIM_DECORATOR(FDecoratorWithChildren, 0xa3ad93b9, FDecorator)
+		DECLARE_ANIM_DECORATOR(FDecoratorWithChildren, 0xa3ad93b9, FBaseDecorator)
 
-		struct FSharedData : FDecorator::FSharedData
-		{
-			FDecoratorHandle Children[2];
-
-			FSharedData(FDecoratorHandle Child0, FDecoratorHandle Child1)
-				: Children{ Child0, Child1 }
-			{}
-		};
+		using FSharedData = FDecoratorWithChildrenSharedData;
 
 		struct FInstanceData : FDecorator::FInstanceData
 		{
@@ -157,9 +141,6 @@ namespace UE::AnimNext
 			}
 		};
 
-		// FDecorator impl
-		virtual EDecoratorMode GetMode() const override { return EDecoratorMode::Base; }
-
 		// IHierarchy impl
 		virtual void GetChildren(FExecutionContext& Context, const TDecoratorBinding<IHierarchy>& Binding, FChildrenArray& Children) const override
 		{
@@ -264,21 +245,74 @@ bool FAnimationAnimNextRuntimeTest_IHierarchy::RunTest(const FString& InParamete
 	// NodeC has two children: NodeA and NodeB (but both decorators are base, only NodeB will be referenced)
 	// NodeD has two children: NodeA and NodeC
 
-	uint16 NodeUID = 0;
-	TArray<uint8> GraphSharedDataBuffer;
+	FNodeHandle NodeA;
+	FNodeHandle NodeB;
+	FNodeHandle NodeC;
+	FNodeHandle NodeD;
 
-	const FNodeHandle NodeA = AppendNodeToGraph(*NodeTemplateA, NodeUID, GraphSharedDataBuffer);
-	InitNodeDecorator(*NodeTemplateA, NodeA, GraphSharedDataBuffer, 0, FDecoratorWithNoChildren::FSharedData());
+	// Write our graph
+	TArray<uint8> GraphSharedDataArchiveBuffer;
+	{
+		FDecoratorWriter DecoratorWriter;
+
+		NodeA = DecoratorWriter.RegisterNode(*NodeTemplateA);
+		NodeB = DecoratorWriter.RegisterNode(*NodeTemplateB);
+		NodeC = DecoratorWriter.RegisterNode(*NodeTemplateC);
+		NodeD = DecoratorWriter.RegisterNode(*NodeTemplateD);
+
+		// We don't have decorator properties
+		TArray<TMap<FString, FString>> DecoratorPropertiesA;
+		DecoratorPropertiesA.AddDefaulted(NodeTemplateDecoratorListA.Num());
+
+		TArray<TMap<FString, FString>> DecoratorPropertiesB;
+		DecoratorPropertiesB.AddDefaulted(NodeTemplateDecoratorListB.Num());
+		DecoratorPropertiesB[0].Add(TEXT("Child"), ToString<FDecoratorWithOneChild::FSharedData>(TEXT("Child"), FAnimNextDecoratorHandle(NodeA)));
+
+		TArray<TMap<FString, FString>> DecoratorPropertiesC;
+		DecoratorPropertiesC.AddDefaulted(NodeTemplateDecoratorListC.Num());
+		DecoratorPropertiesC[0].Add(TEXT("Child"), ToString<FDecoratorWithOneChild::FSharedData>(TEXT("Child"), FAnimNextDecoratorHandle(NodeA)));
+		DecoratorPropertiesC[1].Add(TEXT("Child"), ToString<FDecoratorWithOneChild::FSharedData>(TEXT("Child"), FAnimNextDecoratorHandle(NodeB)));
+
+		TArray<TMap<FString, FString>> DecoratorPropertiesD;
+		DecoratorPropertiesD.AddDefaulted(NodeTemplateDecoratorListD.Num());
+		FAnimNextDecoratorHandle ChildrenHandlesD[2] = { FAnimNextDecoratorHandle(NodeA), FAnimNextDecoratorHandle(NodeC, 1) };
+		DecoratorPropertiesD[0].Add(TEXT("Children"), ToString<FDecoratorWithChildren::FSharedData>(TEXT("Children"), ChildrenHandlesD));
+
+		DecoratorWriter.BeginNodeWriting();
+		DecoratorWriter.WriteNode(NodeA,
+			[&DecoratorPropertiesA](uint32 DecoratorIndex) -> const TMap<FString, FString>&
+			{
+				return DecoratorPropertiesA[DecoratorIndex];
+			});
+		DecoratorWriter.WriteNode(NodeB,
+			[&DecoratorPropertiesB](uint32 DecoratorIndex) -> const TMap<FString, FString>&
+			{
+				return DecoratorPropertiesB[DecoratorIndex];
+			});
+		DecoratorWriter.WriteNode(NodeC,
+			[&DecoratorPropertiesC](uint32 DecoratorIndex) -> const TMap<FString, FString>&
+			{
+				return DecoratorPropertiesC[DecoratorIndex];
+			});
+		DecoratorWriter.WriteNode(NodeD,
+			[&DecoratorPropertiesD](uint32 DecoratorIndex) -> const TMap<FString, FString>&
+			{
+				return DecoratorPropertiesD[DecoratorIndex];
+			});
+		DecoratorWriter.EndNodeWriting();
 
-	const FNodeHandle NodeB = AppendNodeToGraph(*NodeTemplateB, NodeUID, GraphSharedDataBuffer);
-	InitNodeDecorator(*NodeTemplateB, NodeB, GraphSharedDataBuffer, 0, FDecoratorWithOneChild::FSharedData(FDecoratorHandle(NodeA)));
+		AddErrorIfFalse(DecoratorWriter.GetErrorState() == FDecoratorWriter::EErrorState::None, "FAnimationAnimNextRuntimeTest_IHierarchy -> Failed to write decorators");
+		GraphSharedDataArchiveBuffer = DecoratorWriter.GetGraphSharedData();
+	}
 
-	const FNodeHandle NodeC = AppendNodeToGraph(*NodeTemplateC, NodeUID, GraphSharedDataBuffer);
-	InitNodeDecorator(*NodeTemplateC, NodeC, GraphSharedDataBuffer, 0, FDecoratorWithOneChild::FSharedData(FDecoratorHandle(NodeA)));
-	InitNodeDecorator(*NodeTemplateC, NodeC, GraphSharedDataBuffer, 1, FDecoratorWithOneChild::FSharedData(FDecoratorHandle(NodeB)));
+	// Read our graph
+	TArray<uint8> GraphSharedDataBuffer;
+	{
+		FMemoryReader GraphSharedDataArchive(GraphSharedDataArchiveBuffer);
+		FDecoratorReader DecoratorReader(GraphSharedDataArchive);
 
-	const FNodeHandle NodeD = AppendNodeToGraph(*NodeTemplateD, NodeUID, GraphSharedDataBuffer);
-	InitNodeDecorator(*NodeTemplateD, NodeD, GraphSharedDataBuffer, 0, FDecoratorWithChildren::FSharedData(FDecoratorHandle(NodeA), FDecoratorHandle(NodeC, 1)));
+		DecoratorReader.ReadGraphSharedData(GraphSharedDataBuffer);
+	}
 
 	FExecutionContext Context(GraphSharedDataBuffer);
 
@@ -286,7 +320,7 @@ bool FAnimationAnimNextRuntimeTest_IHierarchy::RunTest(const FString& InParamete
 		FMemMark Mark(FMemStack::Get());
 
 		FWeakDecoratorPtr NullPtr;				// Empty, no parent
-		FDecoratorHandle RootHandle(NodeD);		// Point to NodeD, first base decorator
+		FAnimNextDecoratorHandle RootHandle(NodeD);		// Point to NodeD, first base decorator
 
 		FDecoratorPtr NodeDPtr = Context.AllocateNodeInstance(NullPtr, RootHandle);
 		AddErrorIfFalse(NodeDPtr.IsValid(), "FAnimationAnimNextRuntimeTest_IHierarchy -> Failed to allocate root node instance");
@@ -390,17 +424,62 @@ bool FAnimationAnimNextRuntimeTest_IUpdate::RunTest(const FString& InParameters)
 	// NodeB has one child: NodeA (it doesn't update)
 	// NodeD has two children: NodeA and NodeB
 
-	uint16 NodeUID = 0;
-	TArray<uint8> GraphSharedDataBuffer;
+	FNodeHandle NodeA;
+	FNodeHandle NodeB;
+	FNodeHandle NodeC;
+
+	// Write our graph
+	TArray<uint8> GraphSharedDataArchiveBuffer;
+	{
+		FDecoratorWriter DecoratorWriter;
+
+		NodeA = DecoratorWriter.RegisterNode(*NodeTemplateA);
+		NodeB = DecoratorWriter.RegisterNode(*NodeTemplateB);
+		NodeC = DecoratorWriter.RegisterNode(*NodeTemplateC);
+
+		// We don't have decorator properties
+		TArray<TMap<FString, FString>> DecoratorPropertiesA;
+		DecoratorPropertiesA.AddDefaulted(NodeTemplateDecoratorListA.Num());
 
-	const FNodeHandle NodeA = AppendNodeToGraph(*NodeTemplateA, NodeUID, GraphSharedDataBuffer);
-	InitNodeDecorator(*NodeTemplateA, NodeA, GraphSharedDataBuffer, 0, FDecoratorWithNoChildren::FSharedData());
+		TArray<TMap<FString, FString>> DecoratorPropertiesB;
+		DecoratorPropertiesB.AddDefaulted(NodeTemplateDecoratorListB.Num());
+		DecoratorPropertiesB[0].Add(TEXT("Child"), ToString<FDecoratorWithOneChild::FSharedData>(TEXT("Child"), FAnimNextDecoratorHandle(NodeA)));
 
-	const FNodeHandle NodeB = AppendNodeToGraph(*NodeTemplateB, NodeUID, GraphSharedDataBuffer);
-	InitNodeDecorator(*NodeTemplateB, NodeB, GraphSharedDataBuffer, 0, FDecoratorWithOneChild::FSharedData(FDecoratorHandle(NodeA)));
+		TArray<TMap<FString, FString>> DecoratorPropertiesC;
+		DecoratorPropertiesC.AddDefaulted(NodeTemplateDecoratorListC.Num());
+		FAnimNextDecoratorHandle ChildrenHandlesC[2] = { FAnimNextDecoratorHandle(NodeA), FAnimNextDecoratorHandle(NodeB) };
+		DecoratorPropertiesC[0].Add(TEXT("Children"), ToString<FDecoratorWithChildren::FSharedData>(TEXT("Children"), ChildrenHandlesC));
 
-	const FNodeHandle NodeC = AppendNodeToGraph(*NodeTemplateC, NodeUID, GraphSharedDataBuffer);
-	InitNodeDecorator(*NodeTemplateC, NodeC, GraphSharedDataBuffer, 0, FDecoratorWithChildren::FSharedData(FDecoratorHandle(NodeA), FDecoratorHandle(NodeB)));
+		DecoratorWriter.BeginNodeWriting();
+		DecoratorWriter.WriteNode(NodeA,
+			[&DecoratorPropertiesA](uint32 DecoratorIndex) -> const TMap<FString, FString>&
+			{
+				return DecoratorPropertiesA[DecoratorIndex];
+			});
+		DecoratorWriter.WriteNode(NodeB,
+			[&DecoratorPropertiesB](uint32 DecoratorIndex) -> const TMap<FString, FString>&
+			{
+				return DecoratorPropertiesB[DecoratorIndex];
+			});
+		DecoratorWriter.WriteNode(NodeC,
+			[&DecoratorPropertiesC](uint32 DecoratorIndex) -> const TMap<FString, FString>&
+			{
+				return DecoratorPropertiesC[DecoratorIndex];
+			});
+		DecoratorWriter.EndNodeWriting();
+
+		AddErrorIfFalse(DecoratorWriter.GetErrorState() == FDecoratorWriter::EErrorState::None, "FAnimationAnimNextRuntimeTest_IUpdate -> Failed to write decorators");
+		GraphSharedDataArchiveBuffer = DecoratorWriter.GetGraphSharedData();
+	}
+
+	// Read our graph
+	TArray<uint8> GraphSharedDataBuffer;
+	{
+		FMemoryReader GraphSharedDataArchive(GraphSharedDataArchiveBuffer);
+		FDecoratorReader DecoratorReader(GraphSharedDataArchive);
+
+		DecoratorReader.ReadGraphSharedData(GraphSharedDataBuffer);
+	}
 
 	FExecutionContext Context(GraphSharedDataBuffer);
 
@@ -410,7 +489,7 @@ bool FAnimationAnimNextRuntimeTest_IUpdate::RunTest(const FString& InParameters)
 		Private::UpdatedDecorators = &UpdatedDecorators;
 
 		FWeakDecoratorPtr NullPtr;				// Empty, no parent
-		FDecoratorHandle RootHandle(NodeC);		// Point to NodeC, first base decorator
+		FAnimNextDecoratorHandle RootHandle(NodeC);		// Point to NodeC, first base decorator
 
 		FDecoratorPtr NodeCPtr = Context.AllocateNodeInstance(NullPtr, RootHandle);
 		AddErrorIfFalse(NodeCPtr.IsValid(), "FAnimationAnimNextRuntimeTest_IUpdate -> Failed to allocate root node instance");
@@ -474,17 +553,63 @@ bool FAnimationAnimNextRuntimeTest_IEvaluate::RunTest(const FString& InParameter
 	// NodeB has one child: NodeA (it doesn't evaluate)
 	// NodeD has two children: NodeA and NodeB
 
-	uint16 NodeUID = 0;
-	TArray<uint8> GraphSharedDataBuffer;
+	FNodeHandle NodeA;
+	FNodeHandle NodeB;
+	FNodeHandle NodeC;
+
+	// Write our graph
+	TArray<uint8> GraphSharedDataArchiveBuffer;
+	{
+		FDecoratorWriter DecoratorWriter;
 
-	const FNodeHandle NodeA = AppendNodeToGraph(*NodeTemplateA, NodeUID, GraphSharedDataBuffer);
-	InitNodeDecorator(*NodeTemplateA, NodeA, GraphSharedDataBuffer, 0, FDecoratorWithNoChildren::FSharedData());
+		NodeA = DecoratorWriter.RegisterNode(*NodeTemplateA);
+		NodeB = DecoratorWriter.RegisterNode(*NodeTemplateB);
+		NodeC = DecoratorWriter.RegisterNode(*NodeTemplateC);
 
-	const FNodeHandle NodeB = AppendNodeToGraph(*NodeTemplateB, NodeUID, GraphSharedDataBuffer);
-	InitNodeDecorator(*NodeTemplateB, NodeB, GraphSharedDataBuffer, 0, FDecoratorWithOneChild::FSharedData(FDecoratorHandle(NodeA)));
+		// We don't have decorator properties
+		TArray<TMap<FString, FString>> DecoratorPropertiesA;
+		DecoratorPropertiesA.AddDefaulted(NodeTemplateDecoratorListA.Num());
 
-	const FNodeHandle NodeC = AppendNodeToGraph(*NodeTemplateC, NodeUID, GraphSharedDataBuffer);
-	InitNodeDecorator(*NodeTemplateC, NodeC, GraphSharedDataBuffer, 0, FDecoratorWithChildren::FSharedData(FDecoratorHandle(NodeA), FDecoratorHandle(NodeB)));
+		TArray<TMap<FString, FString>> DecoratorPropertiesB;
+		DecoratorPropertiesB.AddDefaulted(NodeTemplateDecoratorListB.Num());
+		DecoratorPropertiesB[0].Add(TEXT("Child"), ToString<FDecoratorWithOneChild::FSharedData>(TEXT("Child"), FAnimNextDecoratorHandle(NodeA)));
+
+		TArray<TMap<FString, FString>> DecoratorPropertiesC;
+		DecoratorPropertiesC.AddDefaulted(NodeTemplateDecoratorListC.Num());
+
+		FAnimNextDecoratorHandle ChildrenHandlesC[2] = { FAnimNextDecoratorHandle(NodeA), FAnimNextDecoratorHandle(NodeB) };
+		DecoratorPropertiesC[0].Add(TEXT("Children"), ToString<FDecoratorWithChildren::FSharedData>(TEXT("Children"), ChildrenHandlesC));
+
+		DecoratorWriter.BeginNodeWriting();
+		DecoratorWriter.WriteNode(NodeA,
+			[&DecoratorPropertiesA](uint32 DecoratorIndex) -> const TMap<FString, FString>&
+			{
+				return DecoratorPropertiesA[DecoratorIndex];
+			});
+		DecoratorWriter.WriteNode(NodeB,
+			[&DecoratorPropertiesB](uint32 DecoratorIndex) -> const TMap<FString, FString>&
+			{
+				return DecoratorPropertiesB[DecoratorIndex];
+			});
+		DecoratorWriter.WriteNode(NodeC,
+			[&DecoratorPropertiesC](uint32 DecoratorIndex) -> const TMap<FString, FString>&
+			{
+				return DecoratorPropertiesC[DecoratorIndex];
+			});
+		DecoratorWriter.EndNodeWriting();
+
+		AddErrorIfFalse(DecoratorWriter.GetErrorState() == FDecoratorWriter::EErrorState::None, "FAnimationAnimNextRuntimeTest_IEvaluate -> Failed to write decorators");
+		GraphSharedDataArchiveBuffer = DecoratorWriter.GetGraphSharedData();
+	}
+
+	// Read our graph
+	TArray<uint8> GraphSharedDataBuffer;
+	{
+		FMemoryReader GraphSharedDataArchive(GraphSharedDataArchiveBuffer);
+		FDecoratorReader DecoratorReader(GraphSharedDataArchive);
+
+		DecoratorReader.ReadGraphSharedData(GraphSharedDataBuffer);
+	}
 
 	FExecutionContext Context(GraphSharedDataBuffer);
 
@@ -494,7 +619,7 @@ bool FAnimationAnimNextRuntimeTest_IEvaluate::RunTest(const FString& InParameter
 		Private::EvaluatedDecorators = &EvaluatedDecorators;
 
 		FWeakDecoratorPtr NullPtr;				// Empty, no parent
-		FDecoratorHandle RootHandle(NodeC);		// Point to NodeC, first base decorator
+		FAnimNextDecoratorHandle RootHandle(NodeC);		// Point to NodeC, first base decorator
 
 		FDecoratorPtr NodeCPtr = Context.AllocateNodeInstance(NullPtr, RootHandle);
 		AddErrorIfFalse(NodeCPtr.IsValid(), "FAnimationAnimNextRuntimeTest_IEvaluate -> Failed to allocate root node instance");
diff --git a/Engine/Plugins/Experimental/AnimNext/Source/AnimNextTestSuite/Private/AnimNextDecoratorInterfacesTest.h b/Engine/Plugins/Experimental/AnimNext/Source/AnimNextTestSuite/Private/AnimNextDecoratorInterfacesTest.h
new file mode 100644
--- /dev/null
+++ b/Engine/Plugins/Experimental/AnimNext/Source/AnimNextTestSuite/Private/AnimNextDecoratorInterfacesTest.h
@@ -0,0 +1,28 @@
+// Copyright Epic Games, Inc. All Rights Reserved.
+
+#pragma once
+
+#include "CoreMinimal.h"
+
+#include "DecoratorBase/DecoratorHandle.h"
+#include "DecoratorBase/DecoratorSharedData.h"
+
+#include "AnimNextDecoratorInterfacesTest.generated.h"
+
+USTRUCT()
+struct FDecoratorWithOneChildSharedData : public FAnimNextDecoratorSharedData
+{
+	GENERATED_BODY()
+
+	UPROPERTY()
+	FAnimNextDecoratorHandle Child;
+};
+
+USTRUCT()
+struct FDecoratorWithChildrenSharedData : public FAnimNextDecoratorSharedData
+{
+	GENERATED_BODY()
+
+	UPROPERTY()
+	FAnimNextDecoratorHandle Children[2];
+};
diff --git a/Engine/Plugins/Experimental/AnimNext/Source/AnimNextTestSuite/Private/AnimNextRuntimeTest.cpp b/Engine/Plugins/Experimental/AnimNext/Source/AnimNextTestSuite/Private/AnimNextRuntimeTest.cpp
--- a/Engine/Plugins/Experimental/AnimNext/Source/AnimNextTestSuite/Private/AnimNextRuntimeTest.cpp
+++ b/Engine/Plugins/Experimental/AnimNext/Source/AnimNextTestSuite/Private/AnimNextRuntimeTest.cpp
@@ -9,8 +9,8 @@
 
 #include "DecoratorBase/Decorator.h"
 #include "DecoratorBase/DecoratorRegistry.h"
+#include "DecoratorBase/NodeDescription.h"
 #include "DecoratorBase/NodeInstance.h"
-#include "DecoratorBase/NodeTemplateRegistry.h"
 
 namespace UE::AnimNext
 {
@@ -49,7 +49,7 @@ namespace UE::AnimNext
 		const FDecoratorUID DecoratorUID = NodeTemplateDecoratorList[DecoratorIndex];
 		const FDecoratorRegistryHandle DecoratorHandle = Registry.FindHandle(DecoratorUID);
 		const FDecorator* Decorator = Registry.Find(DecoratorHandle);
-		const EDecoratorMode DecoratorMode = Decorator->GetMode();
+		const EDecoratorMode DecoratorMode = Decorator->GetDecoratorMode();
 
 		uint32 AdditiveIndexOrNumAdditive;
 		if (DecoratorMode == EDecoratorMode::Base)
@@ -59,7 +59,7 @@ namespace UE::AnimNext
 			for (int32 Index = DecoratorIndex + 1; Index < NodeTemplateDecoratorList.Num(); ++Index)	// Skip ourself
 			{
 				const FDecorator* ChildDecorator = Registry.Find(NodeTemplateDecoratorList[Index]);
-				if (ChildDecorator->GetMode() == EDecoratorMode::Base)
+				if (ChildDecorator->GetDecoratorMode() == EDecoratorMode::Base)
 				{
 					break;	// Found another base decorator, we are done
 				}
@@ -75,7 +75,7 @@ namespace UE::AnimNext
 			for (int32 Index = DecoratorIndex - 1; Index >= 0; --Index)
 			{
 				const FDecorator* ParentDecorator = Registry.Find(NodeTemplateDecoratorList[Index]);
-				if (ParentDecorator->GetMode() == EDecoratorMode::Base)
+				if (ParentDecorator->GetDecoratorMode() == EDecoratorMode::Base)
 				{
 					break;	// Found our base decorator, we are done
 				}
@@ -85,14 +85,17 @@ namespace UE::AnimNext
 			}
 		}
 
+		const FDecoratorMemoryLayout MemoryLayout = Decorator->GetDecoratorMemoryDescription();
+
+		// Align our data
+		SharedDataOffset = Align(SharedDataOffset, MemoryLayout.SharedDataAlignment);
+		InstanceDataOffset = Align(InstanceDataOffset, MemoryLayout.InstanceDataAlignment);
+
 		// Append and update our offsets
 		const int32 BufferIndex = NodeTemplateBuffer.AddUninitialized(sizeof(FDecoratorTemplate));
 		new(&NodeTemplateBuffer[BufferIndex]) FDecoratorTemplate(DecoratorUID, DecoratorHandle, DecoratorMode, AdditiveIndexOrNumAdditive, SharedDataOffset, InstanceDataOffset);
 
-		const FDecoratorMemoryLayout MemoryLayout = Decorator->GetDecoratorMemoryDescription();
-		SharedDataOffset = Align(SharedDataOffset, MemoryLayout.SharedDataAlignment);
 		SharedDataOffset += MemoryLayout.SharedDataSize;
-		InstanceDataOffset = Align(InstanceDataOffset, MemoryLayout.InstanceDataAlignment);
 		InstanceDataOffset += MemoryLayout.InstanceDataSize;
 	}
 
@@ -111,37 +114,5 @@ namespace UE::AnimNext
 
 		return reinterpret_cast<const FNodeTemplate*>(&NodeTemplateBuffer[0]);
 	}
-
-	uint32 GetNodeSharedSize(const FNodeTemplate& NodeTemplate)
-	{
-		const FDecoratorRegistry& Registry = FDecoratorRegistry::Get();
-
-		const uint32 NumDecorators = NodeTemplate.GetNumDecorators();
-		const FDecoratorTemplate* DecoratorTemplates = NodeTemplate.GetDecorators();
-
-		uint32 NodeSharedSize = sizeof(FNodeDescription);
-		for (uint32_t DecoratorIndex = 0; DecoratorIndex < NumDecorators; ++DecoratorIndex)
-		{
-			const FDecorator* Decorator = Registry.Find(DecoratorTemplates[DecoratorIndex].GetRegistryHandle());
-
-			const FDecoratorMemoryLayout MemoryLayout = Decorator->GetDecoratorMemoryDescription();
-			NodeSharedSize = Align(NodeSharedSize, MemoryLayout.SharedDataAlignment);
-			NodeSharedSize += MemoryLayout.SharedDataSize;
-		}
-
-		return NodeSharedSize;
-	}
-
-	// Adds a node template to a shared data graph and returns a handle to it
-	FNodeHandle AppendNodeToGraph(const FNodeTemplate& NodeTemplate, uint16& NodeUID, TArray<uint8>& GraphSharedDataBuffer)
-	{
-		FNodeTemplateRegistry& Registry = FNodeTemplateRegistry::Get();
-
-		const uint32 NodeSharedDataSize = GetNodeSharedSize(NodeTemplate);
-		const uint32 BufferIndex = GraphSharedDataBuffer.AddZeroed(NodeSharedDataSize);
-		new(&GraphSharedDataBuffer[BufferIndex]) FNodeDescription(NodeUID++, Registry.FindOrAdd(&NodeTemplate));
-
-		return FNodeHandle(BufferIndex);
-	}
 }
 #endif
diff --git a/Engine/Plugins/Experimental/AnimNext/Source/AnimNextTestSuite/Private/AnimNextRuntimeTest.h b/Engine/Plugins/Experimental/AnimNext/Source/AnimNextTestSuite/Private/AnimNextRuntimeTest.h
--- a/Engine/Plugins/Experimental/AnimNext/Source/AnimNextTestSuite/Private/AnimNextRuntimeTest.h
+++ b/Engine/Plugins/Experimental/AnimNext/Source/AnimNextTestSuite/Private/AnimNextRuntimeTest.h
@@ -3,32 +3,71 @@
 #pragma once
 
 #include "CoreMinimal.h"
-//#include "AnimNextRuntimeTest.generated.h"
 
 #if WITH_DEV_AUTOMATION_TESTS
 
+#include "UObject/Class.h"
+#include "UObject/UnrealType.h"
+
 #include "DecoratorBase/DecoratorUID.h"
-#include "DecoratorBase/DecoratorTemplate.h"
-#include "DecoratorBase/NodeDescription.h"
-#include "DecoratorBase/NodeHandle.h"
 #include "DecoratorBase/NodeTemplate.h"
 
-// TODO: Add USTRUCT and other similar types here, if required for tests
-
 namespace UE::AnimNext
 {
+	// Builds a node template from a list of decorator UID
 	const FNodeTemplate* BuildNodeTemplate(const TArray<FDecoratorUID>& NodeTemplateDecoratorList, TArray<uint8>& NodeTemplateBuffer);
 
-	// Adds a node template to a shared data graph and returns a handle to it
-	FNodeHandle AppendNodeToGraph(const FNodeTemplate& NodeTemplate, uint16& NodeUID, TArray<uint8>& GraphSharedDataBuffer);
-
-	template<class SharedDataType>
-	void InitNodeDecorator(const FNodeTemplate& NodeTemplate, const FNodeHandle NodeHandle, TArray<uint8>& GraphSharedDataBuffer, uint32 DecoratorIndex, const SharedDataType& SharedData)
+	// Converts a property value into its string representation using UE reflection
+	template<class DecoratorSharedDataType, typename PropertyType>
+	static FString ToString(const FString& PropertyName, PropertyType PropertyValue)
 	{
-		FNodeDescription& NodeDesc = *reinterpret_cast<FNodeDescription*>(&GraphSharedDataBuffer[NodeHandle.GetSharedOffset()]);
+		const UScriptStruct* SharedDataStruct = DecoratorSharedDataType::StaticStruct();
+
+		for (const FProperty* Property = SharedDataStruct->PropertyLink; Property != nullptr; Property = Property->PropertyLinkNext)
+		{
+			if (Property->GetName() == PropertyName)
+			{
+				void* PropertyDefaults = Property->AllocateAndInitializeValue();
+
+				FString Result;
+
+				if constexpr (std::is_pointer_v<PropertyType>)
+				{
+					// C-style array properties aren't handled by ExportText, we need to handle it manually
+
+					const bool bIsCArray = Property->ArrayDim > 1;
+					if (bIsCArray)
+					{
+						Result += TEXT("(");
+					}
+
+					for (int32 Index = 0; Index < Property->ArrayDim; ++Index)
+					{
+						Property->ExportText_Direct(Result, PropertyValue + Index, (PropertyType)PropertyDefaults + Index, nullptr, PPF_None, nullptr);
+
+						if (Index + 1 < Property->ArrayDim)
+						{
+							Result += TEXT(",");
+						}
+					}
+
+					if (bIsCArray)
+					{
+						Result += TEXT(")");
+					}
+				}
+				else
+				{
+					Property->ExportText_Direct(Result, &PropertyValue, PropertyDefaults, nullptr, PPF_None);
+				}
+
+				Property->DestroyAndFreeValue(PropertyDefaults);
+
+				return Result;
+			}
+		}
 
-		const FDecoratorTemplate& DecoratorTemplate = NodeTemplate.GetDecorators()[DecoratorIndex];
-		*static_cast<SharedDataType*>(DecoratorTemplate.GetDecoratorDescription(NodeDesc)) = SharedData;
+		return FString();
 	}
 }
 
