Fix test build compilation failure
Fix static analysis issue around nullptr usage in tests

diff --git a/Engine/Plugins/Experimental/AnimNext/Source/AnimNext/Internal/DecoratorBase/DecoratorWriter.h b/Engine/Plugins/Experimental/AnimNext/Source/AnimNext/Internal/DecoratorBase/DecoratorWriter.h
--- a/Engine/Plugins/Experimental/AnimNext/Source/AnimNext/Internal/DecoratorBase/DecoratorWriter.h
+++ b/Engine/Plugins/Experimental/AnimNext/Source/AnimNext/Internal/DecoratorBase/DecoratorWriter.h
@@ -34,8 +34,10 @@ namespace UE::AnimNext
 			None,					// All good, no error
 			TooManyNodes,			// Exceeded the maximum number of nodes in a graph, @see FNodeDescription::MAXIMUM_COUNT
 			GraphTooLarge,			// Exceeded the maximum graph size, @see FDecoratorWriter::MAXIMUM_GRAPH_SHARED_DATA_SIZE
+			NodeTemplateNotFound,	// Failed to find a necessary node template
 			NodeTemplateTooLarge,	// Exceeded the maximum node template size, @see FNodeTemplate::MAXIMUM_SIZE
 			NodeSharedDataTooLarge,	// Exceeded the maximum node shared data size, @see FNodeDescription::MAXIMUM_SIZE
+			NodeHandleNotFound,		// Failed to find the mapping for a node handle, it was likely not registered
 		};
 
 		FDecoratorWriter();
diff --git a/Engine/Plugins/Experimental/AnimNext/Source/AnimNext/Private/DecoratorBase/DecoratorWriter.cpp b/Engine/Plugins/Experimental/AnimNext/Source/AnimNext/Private/DecoratorBase/DecoratorWriter.cpp
--- a/Engine/Plugins/Experimental/AnimNext/Source/AnimNext/Private/DecoratorBase/DecoratorWriter.cpp
+++ b/Engine/Plugins/Experimental/AnimNext/Source/AnimNext/Private/DecoratorBase/DecoratorWriter.cpp
@@ -178,10 +178,18 @@ namespace UE::AnimNext
 		FNodeTemplateRegistry& NodeTemplateRegistry = FNodeTemplateRegistry::Get();
 
 		const FNodeMapping* NodeMapping = NodeMappings.FindByPredicate([NodeHandle](const FNodeMapping& It) { return It.NodeHandle == NodeHandle; });
-		ensure(NodeMapping != nullptr);
+		if (NodeMapping == nullptr)
+		{
+			ErrorState = EErrorState::NodeHandleNotFound;
+			return;
+		}
 
 		const FNodeTemplate* NodeTemplate = NodeTemplateRegistry.Find(NodeMapping->NodeTemplateHandle);
-		ensure(NodeTemplate);
+		if (NodeTemplate == nullptr)
+		{
+			ErrorState = EErrorState::NodeTemplateNotFound;
+			return;
+		}
 
 		// Populate our node description into a temporary buffer
 		alignas(16) uint8 Buffer[64 * 1024];	// Max node size
diff --git a/Engine/Plugins/Experimental/AnimNext/Source/AnimNextTestSuite/Private/AnimNextDecoratorBaseTest.cpp b/Engine/Plugins/Experimental/AnimNext/Source/AnimNextTestSuite/Private/AnimNextDecoratorBaseTest.cpp
--- a/Engine/Plugins/Experimental/AnimNext/Source/AnimNextTestSuite/Private/AnimNextDecoratorBaseTest.cpp
+++ b/Engine/Plugins/Experimental/AnimNext/Source/AnimNextTestSuite/Private/AnimNextDecoratorBaseTest.cpp
@@ -1486,17 +1486,23 @@ bool FAnimationAnimNextRuntimeTest_DecoratorSerialization::RunTest(const FString
 #endif
 
 FDecoratorA_BaseSharedData::FDecoratorA_BaseSharedData()
+#if WITH_DEV_AUTOMATION_TESTS
 	: DecoratorUID(UE::AnimNext::FDecoratorA_Base::DecoratorUID.GetUID())
+#endif
 {
 }
 
 FDecoratorAB_AddSharedData::FDecoratorAB_AddSharedData()
+#if WITH_DEV_AUTOMATION_TESTS
 	: DecoratorUID(UE::AnimNext::FDecoratorAB_Add::DecoratorUID.GetUID())
+#endif
 {
 }
 
 FDecoratorAC_AddSharedData::FDecoratorAC_AddSharedData()
+#if WITH_DEV_AUTOMATION_TESTS
 	: DecoratorUID(UE::AnimNext::FDecoratorAC_Add::DecoratorUID.GetUID())
+#endif
 {
 }
 
