Parameter block runtime basis

Added parameter stack implementation.
- Stack avoids map-styple lookups in favor of O(n) lookup cost (but with the potential for larger memory use) by mapping parameters to a linear range of indices (FParamId). If FParamIds can be constructed ahead of time
- Still TODO: Would like to improve memory locality of stack layers by allocating everything in one block and avoiding copies when constructing simple layers
Removed IAnimNextInterface and all code that uses it. This is no longer used.
Const-corrected some PropertyBag and RigVM UObject usage

diff --git a/Engine/Plugins/Experimental/AnimNext/Source/AnimNext/Internal/Graph/AnimGraph/AnimNode_AnimNextGraph.h b/Engine/Plugins/Experimental/AnimNext/Source/AnimNext/Internal/Graph/AnimGraph/AnimNode_AnimNextGraph.h
--- a/Engine/Plugins/Experimental/AnimNext/Source/AnimNext/Internal/Graph/AnimGraph/AnimNode_AnimNextGraph.h
+++ b/Engine/Plugins/Experimental/AnimNext/Source/AnimNext/Internal/Graph/AnimGraph/AnimNode_AnimNextGraph.h
@@ -4,10 +4,11 @@
 
 #include "Animation/AnimNode_CustomProperty.h"
 #include "Graph/AnimNextGraph.h"
-#include "Interface/InterfaceContext.h"
+#include "Context.h"
 #include "AnimNode_AnimNextGraph.generated.h"
 
 class UNodeMappingContainer;
+struct FAnimSequencePlayerState;
 
 /**
  * Animation node that allows a AnimNextGraph output to be used in an animation graph
@@ -48,6 +49,8 @@ private:
 	UPROPERTY(EditAnywhere, Category = Settings, meta = (PinHiddenByDefault, DisallowedClasses = "/Script/Engine.AnimMontage"))
 	TObjectPtr<UAnimSequence> TestSequence = nullptr;
 
+	FAnimSequencePlayerState* SequencePlayerState;	// TEST - until we can allocate per-node state again
+
 	/*
 	 * Max LOD that this node is allowed to run
 	 * For example if you have LODThreshold to be 2, it will run until LOD 2 (based on 0 index)
@@ -57,9 +60,6 @@ private:
 	UPROPERTY(EditAnywhere, Category = Performance, meta = (DisplayName = "LOD Threshold"))
 	int32 LODThreshold;
 
-	// This should come from the FPoseContext, but as I need persistence, I store the state here
-	UE::AnimNext::FState RootState; // TODO : Get rid of num elements ? Let the param itself have it
-
 	// Delta time received accumulated in update and used at Evaluate (so we can receive multiple calls to Evaluate)
 	float GraphDeltaTime = 0.f;
 
diff --git a/Engine/Plugins/Experimental/AnimNext/Source/AnimNext/Internal/Graph/AnimNextGraph.h b/Engine/Plugins/Experimental/AnimNext/Source/AnimNext/Internal/Graph/AnimNextGraph.h
--- a/Engine/Plugins/Experimental/AnimNext/Source/AnimNext/Internal/Graph/AnimNextGraph.h
+++ b/Engine/Plugins/Experimental/AnimNext/Source/AnimNext/Internal/Graph/AnimNextGraph.h
@@ -3,13 +3,20 @@
 #pragma once
 
 #include "CoreMinimal.h"
-#include "Interface/IAnimNextInterface.h"
 #include "RigVMCore/RigVM.h"
-#include "Param/ParamType.h"
+
 #include "AnimNextGraph.generated.h"
 
 class UEdGraph;
 class UAnimNextGraph;
+class UAnimGraphNode_AnimNextGraph;
+struct FAnimNode_AnimNextGraph;
+
+namespace UE::AnimNext
+{
+	struct FContext;
+}
+
 namespace UE::AnimNext::UncookedOnly
 {
 	struct FUtils;
@@ -28,7 +35,7 @@ namespace UE::AnimNext::Graph
 
 // A user-created graph of logic used to supply data
 UCLASS(BlueprintType)
-class ANIMNEXT_API UAnimNextGraph : public UObject, public IAnimNextInterface
+class ANIMNEXT_API UAnimNextGraph : public UObject
 {
 	GENERATED_BODY()
 
@@ -36,12 +43,8 @@ class ANIMNEXT_API UAnimNextGraph : public UObject, public IAnimNextInterface
 	virtual void PostRename(UObject* OldOuter, const FName OldName) override;
 	virtual void GetPreloadDependencies(TArray<UObject*>& OutDeps) override;
 
-	// IAnimNextInterface interface
-	virtual UE::AnimNext::FParamTypeHandle GetReturnTypeHandleImpl() const final override;
-	virtual bool GetDataImpl(const UE::AnimNext::FContext& Context) const final override;
-
-	/** Set the return type of this graph */
-	void SetReturnTypeHandle(UE::AnimNext::FParamTypeHandle InHandle);
+	// Run this graph with the given context
+	void Run(const UE::AnimNext::FContext& Context) const;
 
 	// Support rig VM execution
 	TArray<FRigVMExternalVariable> GetRigVMExternalVariables();
@@ -52,7 +55,8 @@ class ANIMNEXT_API UAnimNextGraph : public UObject, public IAnimNextInterface
 	friend class UE::AnimNext::Editor::FGraphEditor;
 	friend class UAnimNextGraph;
 	friend class UAnimGraphNode_AnimNextGraph;
-	
+	friend struct FAnimNode_AnimNextGraph;
+
 	UPROPERTY()
 	TObjectPtr<URigVM> RigVM;
 
@@ -62,9 +66,6 @@ class ANIMNEXT_API UAnimNextGraph : public UObject, public IAnimNextInterface
 	UPROPERTY()
 	FRigVMRuntimeSettings VMRuntimeSettings;
 
-	UPROPERTY()
-	FAnimNextParamType ReturnType;
-
 #if WITH_EDITORONLY_DATA
 	UPROPERTY(VisibleAnywhere, Instanced, Category = "Graph", meta = (ShowInnerProperties))
 	TObjectPtr<UObject> EditorData;
diff --git a/Engine/Plugins/Experimental/AnimNext/Source/AnimNext/Internal/Graph/AnimNext_LODPose.h b/Engine/Plugins/Experimental/AnimNext/Source/AnimNext/Internal/Graph/AnimNext_LODPose.h
--- a/Engine/Plugins/Experimental/AnimNext/Source/AnimNext/Internal/Graph/AnimNext_LODPose.h
+++ b/Engine/Plugins/Experimental/AnimNext/Source/AnimNext/Internal/Graph/AnimNext_LODPose.h
@@ -3,9 +3,6 @@
 #pragma once
 
 #include "CoreMinimal.h"
-#include "Interface/InterfaceTypes.h"
-#include "Interface/IAnimNextInterface.h"
-#include "Interface/InterfaceContext.h"
 #include "ReferencePose.h"
 #include "LODPose.h"
 #include "AnimNext_LODPose.generated.h"
@@ -43,36 +40,3 @@ struct FAnimNextGraphLODPose
 
 	UE::AnimNext::FLODPose LODPose;
 };
-
-UCLASS()
-class UAnimNextInterfaceGraphLODPose : public UObject, public IAnimNextInterface
-{
-	GENERATED_BODY()
-
-	virtual UE::AnimNext::FParamTypeHandle GetReturnTypeHandleImpl() const final override
-	{
-		return UE::AnimNext::FParamTypeHandle::GetHandle<FAnimNextGraphLODPose>();
-	}
-
-	virtual bool GetDataImpl(const UE::AnimNext::FContext& Context) const override
-	{
-		checkf(false, TEXT("UAnimNextInterfaceLODPose::GetDataImpl must be overridden"));
-		return false;
-	}
-};
-
-UCLASS()
-class UAnimNextInterface_LODPose_Literal : public UAnimNextInterfaceGraphLODPose
-{
-	GENERATED_BODY()
-
-	// IAnimNextInterface interface
-	virtual bool GetDataImpl(const UE::AnimNext::FContext& Context) const final override
-	{
-		Context.SetResult(Value);
-		return true;
-	}
-	
-	UPROPERTY(EditAnywhere, Category = "Parameters")
-	FAnimNextGraphLODPose Value;
-};
diff --git a/Engine/Plugins/Experimental/AnimNext/Source/AnimNext/Internal/Graph/GraphExecuteContext.h b/Engine/Plugins/Experimental/AnimNext/Source/AnimNext/Internal/Graph/GraphExecuteContext.h
--- a/Engine/Plugins/Experimental/AnimNext/Source/AnimNext/Internal/Graph/GraphExecuteContext.h
+++ b/Engine/Plugins/Experimental/AnimNext/Source/AnimNext/Internal/Graph/GraphExecuteContext.h
@@ -7,8 +7,6 @@
 #include "Units/RigUnit.h"
 #include "GraphExecuteContext.generated.h"
 
-class IAnimNextInterface;
-
 namespace UE::AnimNext
 {
 	struct FContext;
@@ -21,52 +19,31 @@ struct FAnimNextGraphExecuteContext : public FRigVMExecuteContext
 
 	FAnimNextGraphExecuteContext()
 		: FRigVMExecuteContext()
-		, InterfaceContext(nullptr)
-		, Interface(nullptr)
-		, ResultPtr(nullptr)
+		, Context(nullptr)
 	{
 	}
 
 	const UE::AnimNext::FContext& GetContext() const
 	{
-		check(InterfaceContext);
-		return *InterfaceContext;
-	}
-
-	void SetContextData(const IAnimNextInterface* InInterface, const UE::AnimNext::FContext& InInterfaceContext, bool& bInResult)
-	{
-		Interface = InInterface;
-		InterfaceContext = &InInterfaceContext;
-		ResultPtr = &bInResult;
+		check(Context);
+		return *Context;
 	}
 
-	void SetResult(bool bInResult) const
+	void SetContextData(const UE::AnimNext::FContext& InContext)
 	{
-		check(ResultPtr);
-		*ResultPtr &= bInResult;
+		Context = &InContext;
 	}
 
-	const IAnimNextInterface* GetInterface() const
-	{
-		check(Interface);
-		return Interface;
-	}
-	
 	virtual void Copy(const FRigVMExecuteContext* InOtherContext) override
 	{
 		Super::Copy(InOtherContext);
 
 		const FAnimNextGraphExecuteContext* OtherContext = (const FAnimNextGraphExecuteContext*)InOtherContext; 
-		InterfaceContext = OtherContext->InterfaceContext;
-		Interface = OtherContext->Interface;
-		ResultPtr = OtherContext->ResultPtr;
+		Context = OtherContext->Context;
 	}
 
-
 private:
-	const UE::AnimNext::FContext* InterfaceContext;
-	const IAnimNextInterface* Interface;
-	bool* ResultPtr;
+	const UE::AnimNext::FContext* Context;
 };
 
 USTRUCT(meta=(ExecuteContext="FAnimNextGraphExecuteContext"))
diff --git a/Engine/Plugins/Experimental/AnimNext/Source/AnimNext/Internal/Graph/RigUnit_AnimNextEndExecution.h b/Engine/Plugins/Experimental/AnimNext/Source/AnimNext/Internal/Graph/RigUnit_AnimNextEndExecution.h
--- a/Engine/Plugins/Experimental/AnimNext/Source/AnimNext/Internal/Graph/RigUnit_AnimNextEndExecution.h
+++ b/Engine/Plugins/Experimental/AnimNext/Source/AnimNext/Internal/Graph/RigUnit_AnimNextEndExecution.h
@@ -15,13 +15,6 @@ struct ANIMNEXT_API FRigUnit_AnimNextEndExecution : public FRigUnit_AnimNextBase
 {
 	GENERATED_BODY()
 
-	template<typename ValueType>
-	static void SetResult(const FRigVMExecuteContext& InExecuteContext, const ValueType& InResult)
-	{
-		const FAnimNextGraphExecuteContext& ExecuteContext = static_cast<const FAnimNextGraphExecuteContext&>(InExecuteContext);
-		ExecuteContext.GetContext().SetResult<ValueType>(InResult);
-	}
-
 public:
 	// The execution result
 	UPROPERTY(EditAnywhere, Transient, DisplayName = "End Execute", Category = "EndExecution", meta = (Input))
diff --git a/Engine/Plugins/Experimental/AnimNext/Source/AnimNext/Internal/Param/ParametersExecuteContext.h b/Engine/Plugins/Experimental/AnimNext/Source/AnimNext/Internal/Param/ParametersExecuteContext.h
--- a/Engine/Plugins/Experimental/AnimNext/Source/AnimNext/Internal/Param/ParametersExecuteContext.h
+++ b/Engine/Plugins/Experimental/AnimNext/Source/AnimNext/Internal/Param/ParametersExecuteContext.h
@@ -5,57 +5,23 @@
 #include "CoreMinimal.h"
 #include "RigVMCore/RigVMExecuteContext.h"
 #include "Units/RigUnit.h"
+#include "Graph/GraphExecuteContext.h"
 #include "ParametersExecuteContext.generated.h"
 
-class IAnimNextInterface;
-
 namespace UE::AnimNext
 {
 	struct FContext;
 }
 
 USTRUCT(BlueprintType)
-struct FAnimNextParametersExecuteContext : public FRigVMExecuteContext
+struct FAnimNextParametersExecuteContext : public FAnimNextGraphExecuteContext
 {
 	GENERATED_BODY()
 
 	FAnimNextParametersExecuteContext()
-		: FRigVMExecuteContext()
-	{
-	}
-
-	void SetContextData(TConstArrayView<TArrayView<uint8>> InValues)
+		: FAnimNextGraphExecuteContext()
 	{
-		Values = InValues;
 	}
-
-	void SetCurrentValueIndex(int32 InIndex)
-	{
-		check(Values.IsValidIndex(InIndex));
-		Index = InIndex;
-	}
-
-	TArrayView<uint8> GetData() const
-	{
-		check(Values.Num() > 0);
-		check(Values.IsValidIndex(Index));
-		return Values[Index]; 
-	}
-	
-	virtual void Copy(const FRigVMExecuteContext* InOtherContext) override
-	{
-		Super::Copy(InOtherContext);
-
-		const FAnimNextParametersExecuteContext* OtherContext = (const FAnimNextParametersExecuteContext*)InOtherContext; 
-		Values = OtherContext->Values;
-	}
-
-private:
-	// The parameter values to set
-	TConstArrayView<TArrayView<uint8>> Values;
-
-	// Current value that is being set
-	int32 Index = 0;
 };
 
 USTRUCT(meta=(ExecuteContext="FAnimNextParametersExecuteContext"))
diff --git a/Engine/Plugins/Experimental/AnimNext/Source/AnimNext/Internal/Param/RigVMDispatch_GetParameter.h b/Engine/Plugins/Experimental/AnimNext/Source/AnimNext/Internal/Param/RigVMDispatch_GetParameter.h
new file mode 100644
--- /dev/null
+++ b/Engine/Plugins/Experimental/AnimNext/Source/AnimNext/Internal/Param/RigVMDispatch_GetParameter.h
@@ -0,0 +1,47 @@
+// Copyright Epic Games, Inc. All Rights Reserved.
+
+#pragma once
+
+#include "RigVMCore/RigVMDispatchFactory.h"
+#include "Graph/GraphExecuteContext.h"
+#include "RigVMDispatch_GetParameter.generated.h"
+
+namespace UE::AnimNext::UncookedOnly
+{
+	struct FUtils;
+}
+
+/*
+ * Gets a parameter's value
+ */
+USTRUCT(meta = (DisplayName = "Get Parameter", NodeColor = "0.8, 0, 0.2, 1"))
+struct ANIMNEXT_API FRigVMDispatch_GetParameter : public FRigVMDispatchFactory
+{
+	GENERATED_BODY()
+
+	FRigVMDispatch_GetParameter();
+
+private:
+	friend struct UE::AnimNext::UncookedOnly::FUtils;
+
+	virtual UScriptStruct* GetExecuteContextStruct() const { return FAnimNextGraphExecuteContext::StaticStruct(); }
+	virtual FName GetArgumentNameForOperandIndex(int32 InOperandIndex, int32 InTotalOperands) const override;
+	virtual TArray<FRigVMTemplateArgument> GetArguments() const override;
+#if WITH_EDITOR
+	virtual FString GetArgumentDefaultValue(const FName& InArgumentName, TRigVMTypeIndex InTypeIndex) const override;
+#endif
+	virtual FRigVMTemplateTypeMap OnNewArgumentType(const FName& InArgumentName, TRigVMTypeIndex InTypeIndex) const override;
+	virtual bool IsSingleton() const override { return true; }
+
+	virtual FRigVMFunctionPtr GetDispatchFunctionImpl(const FRigVMTemplateTypeMap& InTypes) const override
+	{
+		return &FRigVMDispatch_GetParameter::Execute;
+	}
+	static void Execute(FRigVMExtendedExecuteContext& InContext, FRigVMMemoryHandleArray Handles, FRigVMPredicateBranchArray RigVMBranches);
+
+	static const FName ExecuteContextName;
+	static const FName ParameterName;
+	static const FName ParameterIdName;
+	static const FName TypeHandleName;
+	static const FName ValueName;
+};
diff --git a/Engine/Plugins/Experimental/AnimNext/Source/AnimNext/Internal/Param/RigVMDispatch_SetParameter.h b/Engine/Plugins/Experimental/AnimNext/Source/AnimNext/Internal/Param/RigVMDispatch_SetParameter.h
--- a/Engine/Plugins/Experimental/AnimNext/Source/AnimNext/Internal/Param/RigVMDispatch_SetParameter.h
+++ b/Engine/Plugins/Experimental/AnimNext/Source/AnimNext/Internal/Param/RigVMDispatch_SetParameter.h
@@ -1,8 +1,9 @@
-﻿// Copyright Epic Games, Inc. All Rights Reserved.
+// Copyright Epic Games, Inc. All Rights Reserved.
 
 #pragma once
 
 #include "RigVMCore/RigVMDispatchFactory.h"
+#include "Graph/GraphExecuteContext.h"
 #include "RigVMDispatch_SetParameter.generated.h"
 
 namespace UE::AnimNext::UncookedOnly
@@ -23,8 +24,12 @@ struct ANIMNEXT_API FRigVMDispatch_SetParameter : public FRigVMDispatchFactory
 private:
 	friend struct UE::AnimNext::UncookedOnly::FUtils;
 
+	virtual UScriptStruct* GetExecuteContextStruct() const { return FAnimNextGraphExecuteContext::StaticStruct(); }
 	virtual FName GetArgumentNameForOperandIndex(int32 InOperandIndex, int32 InTotalOperands) const override;
 	virtual TArray<FRigVMTemplateArgument> GetArguments() const override;
+#if WITH_EDITOR
+	virtual FString GetArgumentDefaultValue(const FName& InArgumentName, TRigVMTypeIndex InTypeIndex) const override;
+#endif
 	virtual TArray<FRigVMExecuteArgument> GetExecuteArguments_Impl(const FRigVMDispatchContext& InContext) const override;
 	virtual FRigVMTemplateTypeMap OnNewArgumentType(const FName& InArgumentName, TRigVMTypeIndex InTypeIndex) const override;
 	virtual bool IsSingleton() const override { return true; } 
@@ -35,6 +40,9 @@ private:
 	}
 	static void Execute(FRigVMExtendedExecuteContext& InContext, FRigVMMemoryHandleArray Handles, FRigVMPredicateBranchArray RigVMBranches);
 
-	static const FName ValueName;
 	static const FName ExecuteContextName;
+	static const FName ParameterName;
+	static const FName ParameterIdName;
+	static const FName TypeHandleName;
+	static const FName ValueName;
 };
diff --git a/Engine/Plugins/Experimental/AnimNext/Source/AnimNext/Private/Context.cpp b/Engine/Plugins/Experimental/AnimNext/Source/AnimNext/Private/Context.cpp
new file mode 100644
--- /dev/null
+++ b/Engine/Plugins/Experimental/AnimNext/Source/AnimNext/Private/Context.cpp
@@ -0,0 +1,22 @@
+// Copyright Epic Games, Inc. All Rights Reserved.
+
+#include "Context.h"
+#include "Param/ParamType.h"
+#include "Param/ParamStack.h"
+
+namespace UE::AnimNext
+{
+
+FContext::FContext()
+	: Stack(FParamStack::Get())
+	, DeltaTime(0.0f)
+{
+}
+
+FContext::FContext(float InDeltaTime)
+	: Stack(FParamStack::Get())
+	, DeltaTime(InDeltaTime)
+{
+}
+
+}
\ No newline at end of file
diff --git a/Engine/Plugins/Experimental/AnimNext/Source/AnimNext/Private/Graph/AnimGraph/AnimNode_AnimNextGraph.cpp b/Engine/Plugins/Experimental/AnimNext/Source/AnimNext/Private/Graph/AnimGraph/AnimNode_AnimNextGraph.cpp
--- a/Engine/Plugins/Experimental/AnimNext/Source/AnimNext/Private/Graph/AnimGraph/AnimNode_AnimNextGraph.cpp
+++ b/Engine/Plugins/Experimental/AnimNext/Source/AnimNext/Private/Graph/AnimGraph/AnimNode_AnimNextGraph.cpp
@@ -7,16 +7,15 @@
 #include "Animation/AnimInstanceProxy.h"
 #include "GameFramework/Actor.h"
 #include "AnimationRuntime.h"
-#include "Interface/AnimNextInterface.h"
 #include "DataRegistry.h"
 #include "GenerationTools.h"
 #include "ReferencePose.h"
 #include "Graph/AnimNext_LODPose.h"
-#include "Graph/RigUnit_AnimNextAnimSequence.h" // TEST
-#include "Animation/AnimSequence.h"
 #include "Animation/AnimSequence.h" // TEST
+#include "Graph/RigUnit_AnimNextAnimSequence.h" // TEST
 #include "Engine/SkeletalMesh.h"
 #include "BoneContainer.h"
+#include "Param/ParamStack.h"
 
 #include UE_INLINE_GENERATED_CPP_BY_NAME(AnimNode_AnimNextGraph)
 
@@ -29,10 +28,12 @@ FAnimNode_AnimNextGraph::FAnimNode_AnimNextGraph()
 	, AnimNextGraph(nullptr)
 	, LODThreshold(INDEX_NONE)
 {
+	SequencePlayerState = new FAnimSequencePlayerState();
 }
 
 FAnimNode_AnimNextGraph::~FAnimNode_AnimNextGraph()
 {
+	delete SequencePlayerState;
 }
 
 void FAnimNode_AnimNextGraph::OnInitializeAnimInstance(const FAnimInstanceProxy* InProxy, const UAnimInstance* InAnimInstance)
@@ -113,43 +114,33 @@ void FAnimNode_AnimNextGraph::Evaluate_AnyThread(FPoseContext & Output)
 	const UE::AnimNext::FReferencePose& RefPose = RefPoseHandle.GetRef<UE::AnimNext::FReferencePose>();
 
 	const int32 LODLevel = Output.AnimInstanceProxy->GetLODLevel();
-
-	// TODO : In a normal graph, FState/FParamStorage/FContext would come from FAnimationUpdateContext, here I have to create one the fly for the AnimNext graph
-	UE::AnimNext::FParamStorage ParamStorage(5, 0, 0);	// TODO : see if we can have state and params together and optional. Also, default size
 	
 	// TODO : Using AnimInstanceProxy->GetDeltaSeconds() makes the preview to advance  multiple times when debug options are activated (i.e. ShowUncompressedAnim)
 	// See where to get / how to calculate the correct delta value
-	UE::AnimNext::FContext RootContext(GraphDeltaTime, RootState, ParamStorage);
+	UE::AnimNext::FContext Context(GraphDeltaTime);
 	GraphDeltaTime = 0.f; // Reset for the case that we receive multiple calls to Evaluate (debug options)
 
 	FAnimNextGraphReferencePose GraphReferencePose(&RefPose);
 	FAnimNextGraph_AnimSequence GraphTestSequence(TestSequence);  // TEST
 
 	FAnimNextGraphLODPose GraphSourceLODPose(FLODPose(RefPose, LODLevel, false, Output.ExpectsAdditivePose()));
+	FAnimNextGraphLODPose ResultPose(FLODPose(RefPose, LODLevel, true, Output.ExpectsAdditivePose()));
 	FGenerationTools::RemapPose(LODLevel, SourcePose, RefPose, GraphSourceLODPose.LODPose);
 
-	RootContext.AddInputReference(FName(TEXT("GraphReferencePose")), GraphReferencePose);			// TODO : Should this go into the AnimNextVM Context ?
-	RootContext.AddInputValue(FName(TEXT("GraphLODLevel")), LODLevel);								// TODO : Should this go into the AnimNextVM Context ?
-	RootContext.AddInputValue(FName(TEXT("GraphExpectsAdditive")), Output.ExpectsAdditivePose());	// TODO : Should this go into the AnimNextVM Context ?
+	Context.GetMutableParamStack().PushValues(
+		"AnimSequencePlayerState", *SequencePlayerState,
+		"GraphReferencePose", GraphReferencePose,
+		"ResultPose", ResultPose,
+		"GraphLODLevel", LODLevel,
+		"GraphExpectsAdditive", Output.ExpectsAdditivePose(),
+		"SourcePose", GraphSourceLODPose,						// TODO : Pass this as a external variable maybe ? When we have support for variables in the rigvm graph
+		"TestSequence", GraphTestSequence						// TEST anim decompression - Anim Sequence to decompress);
+	);
 
-	RootContext.AddInputReference(FName(TEXT("SourcePose")), GraphSourceLODPose);		// TODO : Pass this as a external variable maybe ? When we have support for variables in the rigvm graph
-	RootContext.AddInputReference(FName(TEXT("TestSequence")), GraphTestSequence);		// TEST anim decompression - Anim Sequence to decompress
-		
-	TScriptInterface<IAnimNextInterface> ScriptInterface = AnimNextGraph;
-	FAnimNextGraphLODPose GraphResultLODPose(FLODPose(RefPose, LODLevel, true, Output.ExpectsAdditivePose()));
-
-	const bool bOk = UE::AnimNext::Interface::GetDataSafe(ScriptInterface, RootContext, GraphResultLODPose);
-	if (bOk)
-	{
-		FGenerationTools::RemapPose(LODLevel, RefPose, GraphResultLODPose.LODPose, Output);
-	}
-	else
-	{
-		Output.ResetToRefPose();
-	}
+	AnimNextGraph->Run(Context);
 
+	FGenerationTools::RemapPose(LODLevel, RefPose, ResultPose.LODPose, Output);
 
-	// evaluate 
 	FAnimNode_CustomProperty::Evaluate_AnyThread(Output);
 }
 
@@ -164,7 +155,7 @@ void FAnimNode_AnimNextGraph::PostSerialize(const FArchive& Ar)
 	{
 		if (AnimNextGraph)
 		{
-			//AnimNextInterfaceGraph->Initialize();
+			//AnimNextGraph->Initialize();
 		}
 	}
 }
diff --git a/Engine/Plugins/Experimental/AnimNext/Source/AnimNext/Private/Graph/AnimNextGraph.cpp b/Engine/Plugins/Experimental/AnimNext/Source/AnimNext/Private/Graph/AnimNextGraph.cpp
--- a/Engine/Plugins/Experimental/AnimNext/Source/AnimNext/Private/Graph/AnimNextGraph.cpp
+++ b/Engine/Plugins/Experimental/AnimNext/Source/AnimNext/Private/Graph/AnimNextGraph.cpp
@@ -16,31 +16,16 @@ const FName EntryPointName("GetData");
 const FName ResultName("Result");
 }
 
-UE::AnimNext::FParamTypeHandle UAnimNextGraph::GetReturnTypeHandleImpl() const
+void UAnimNextGraph::Run(const UE::AnimNext::FContext& Context) const
 {
-	return ReturnType.GetHandle();
-}
-
-bool UAnimNextGraph::GetDataImpl(const UE::AnimNext::FContext& Context) const
-{
-	bool bResult = true;
-	
 	if(RigVM)
 	{
-		// we might have to set some sensible values to the extended execute context
-		FAnimNextGraphExecuteContext& AnimNextInterfaceContext = ExtendedExecuteContext.GetPublicDataSafe<FAnimNextGraphExecuteContext>();
-
-		AnimNextInterfaceContext.SetContextData(this, Context, bResult);
+		FRigVMExtendedExecuteContext RigVMExtendedExecuteContext;
+		FAnimNextGraphExecuteContext& AnimNextContext = RigVMExtendedExecuteContext.GetPublicDataSafe<FAnimNextGraphExecuteContext>();
+		AnimNextContext.SetContextData(Context);
 
-		bResult &= (RigVM->Execute(ExtendedExecuteContext, TArray<URigVMMemoryStorage*>(), FRigUnit_AnimNextBeginExecution::EventName) != ERigVMExecuteResult::Failed);
+		RigVM->Execute(RigVMExtendedExecuteContext, TArray<URigVMMemoryStorage*>(), FRigUnit_AnimNextBeginExecution::EventName);
 	}
-	
-	return bResult;
-}
-
-void UAnimNextGraph::SetReturnTypeHandle(UE::AnimNext::FParamTypeHandle InHandle)
-{
-	ReturnType = InHandle.GetType();
 }
 
 TArray<FRigVMExternalVariable> UAnimNextGraph::GetRigVMExternalVariables()
diff --git a/Engine/Plugins/Experimental/AnimNext/Source/AnimNext/Private/Graph/RigUnit_AnimNextAnimSequence.cpp b/Engine/Plugins/Experimental/AnimNext/Source/AnimNext/Private/Graph/RigUnit_AnimNextAnimSequence.cpp
--- a/Engine/Plugins/Experimental/AnimNext/Source/AnimNext/Private/Graph/RigUnit_AnimNextAnimSequence.cpp
+++ b/Engine/Plugins/Experimental/AnimNext/Source/AnimNext/Private/Graph/RigUnit_AnimNextAnimSequence.cpp
@@ -2,10 +2,11 @@
 
 #include "RigUnit_AnimNextAnimSequence.h"
 #include "Graph/GraphExecuteContext.h"
-#include "Interface/AnimNextInterface.h"
 #include "DecompressionTools.h"
 #include "Animation/AnimSequence.h"
 #include "Graph/AnimNext_LODPose.h"
+#include "Context.h"
+#include "Param/ParamStack.h"
 
 #include UE_INLINE_GENERATED_CPP_BY_NAME(RigUnit_AnimNextAnimSequence)
 
@@ -29,8 +30,8 @@ FRigUnit_AnimNext_AnimSequencePlayer_Initialize()
 	const UAnimSequenceBase* AnimSequenceBase = Sequence.AnimSequence;
 	if (AnimSequenceBase != nullptr)
 	{
-		// Get or Create the internal state
-		FAnimSequencePlayerState& AnimSequencePlayerState = InterfaceContext.GetState<FAnimSequencePlayerState>(ExecuteContext.GetInterface(), 0);
+		// Get internal state (as a param)
+		FAnimSequencePlayerState& AnimSequencePlayerState = InterfaceContext.GetMutableParamStack().GetMutableParam<FAnimSequencePlayerState>("AnimSequencePlayerState");
 
 		const float SequenceLength = AnimSequenceBase->GetPlayLength();
 		AnimSequencePlayerState.InternalTimeAccumulator = FMath::Clamp(Parameters.StartPosition, 0.f, SequenceLength);
@@ -47,9 +48,8 @@ FRigUnit_AnimNext_AnimSequencePlayer_Execute()
 	const UAnimSequenceBase* AnimSequenceBase = Sequence.AnimSequence;
 	if (AnimSequenceBase != nullptr)
 	{
-		// Get or Create the internal state
-		//const TParam<FAnimSequencePlayerState> AnimSequencePlayerState = InterfaceContext.GetState<FAnimSequencePlayerState>(ExecuteContext.GetInterface(), 0);
-		FAnimSequencePlayerState& AnimSequencePlayerState = InterfaceContext.GetState<FAnimSequencePlayerState, EStatePersistence::Permanent>(ExecuteContext.GetInterface(), 0);
+		// Get internal state (as a param)
+		FAnimSequencePlayerState& AnimSequencePlayerState = InterfaceContext.GetMutableParamStack().GetMutableParam<FAnimSequencePlayerState>("AnimSequencePlayerState");
 
 		const float DeltaTime = InterfaceContext.GetDeltaTime();
 		const float EffectiveDelta = FMath::IsNearlyZero(DeltaTime) || FMath::IsNearlyZero(Parameters.PlayRate) ? 0.f : DeltaTime * Parameters.PlayRate;
@@ -75,9 +75,9 @@ FRigUnit_AnimNext_AnimSequencePlayer_Execute()
 			, DeltaTimeRecord
 			, Parameters.bLoop);
 
-		const FAnimNextGraphReferencePose& GraphReferencePose = InterfaceContext.GetParameterChecked<FAnimNextGraphReferencePose>(TEXT("GraphReferencePose")).GetDataChecked();
-		const int32 GraphLODLevel = InterfaceContext.GetParameterChecked<int32>(TEXT("GraphLODLevel")).GetDataChecked();
-		const bool bGraphExpectsAdditive = InterfaceContext.GetParameterChecked<bool>(TEXT("GraphExpectsAdditive")).GetDataChecked();
+		const FAnimNextGraphReferencePose& GraphReferencePose = InterfaceContext.GetParamStack().GetParam<FAnimNextGraphReferencePose>("GraphReferencePose");
+		const int32 GraphLODLevel = InterfaceContext.GetParamStack().GetParam<int32>("GraphLODLevel");
+		const bool bGraphExpectsAdditive = InterfaceContext.GetParamStack().GetParam<bool>("GraphExpectsAdditive");
 
 		LODPose.LODPose.PrepareForLOD(*GraphReferencePose.ReferencePose, GraphLODLevel, true, bGraphExpectsAdditive);
 		
diff --git a/Engine/Plugins/Experimental/AnimNext/Source/AnimNext/Private/Graph/RigUnit_AnimNextEndExecution.cpp b/Engine/Plugins/Experimental/AnimNext/Source/AnimNext/Private/Graph/RigUnit_AnimNextEndExecution.cpp
--- a/Engine/Plugins/Experimental/AnimNext/Source/AnimNext/Private/Graph/RigUnit_AnimNextEndExecution.cpp
+++ b/Engine/Plugins/Experimental/AnimNext/Source/AnimNext/Private/Graph/RigUnit_AnimNextEndExecution.cpp
@@ -2,20 +2,27 @@
 
 #include "Graph/RigUnit_AnimNextEndExecution.h"
 #include "Units/RigUnitContext.h"
+#include "Param/ParamStack.h"
+#include "Context.h"
 
 #include UE_INLINE_GENERATED_CPP_BY_NAME(RigUnit_AnimNextEndExecution)
 
 FRigUnit_AnimNextEndExecution_Bool_Execute()
 {
-	SetResult(ExecuteContext, Result);
+
 }
 
 FRigUnit_AnimNextEndExecution_Float_Execute()
 {
-	SetResult(ExecuteContext, Result);
+
 }
 
 FRigUnit_AnimNextEndExecution_LODPose_Execute()
 {
-	SetResult(ExecuteContext, Result);
+	using namespace UE::AnimNext;
+
+	const FContext& InterfaceContext = ExecuteContext.GetContext();
+	FAnimNextGraphLODPose& ResultPose = InterfaceContext.GetMutableParamStack().GetMutableParam<FAnimNextGraphLODPose>("ResultPose");
+
+	ResultPose = Result;
 }
\ No newline at end of file
diff --git a/Engine/Plugins/Experimental/AnimNext/Source/AnimNext/Private/Graph/RigUnit_AnimNextParameter.cpp b/Engine/Plugins/Experimental/AnimNext/Source/AnimNext/Private/Graph/RigUnit_AnimNextParameter.cpp
deleted file mode 100644
diff --git a/Engine/Plugins/Experimental/AnimNext/Source/AnimNext/Private/Graph/RigUnit_AnimNextParameter.h b/Engine/Plugins/Experimental/AnimNext/Source/AnimNext/Private/Graph/RigUnit_AnimNextParameter.h
deleted file mode 100644
diff --git a/Engine/Plugins/Experimental/AnimNext/Source/AnimNext/Private/Interface/AnimNextInterface.cpp b/Engine/Plugins/Experimental/AnimNext/Source/AnimNext/Private/Interface/AnimNextInterface.cpp
deleted file mode 100644
diff --git a/Engine/Plugins/Experimental/AnimNext/Source/AnimNext/Private/Interface/InterfaceContext.cpp b/Engine/Plugins/Experimental/AnimNext/Source/AnimNext/Private/Interface/InterfaceContext.cpp
deleted file mode 100644
diff --git a/Engine/Plugins/Experimental/AnimNext/Source/AnimNext/Private/Interface/InterfaceState.cpp b/Engine/Plugins/Experimental/AnimNext/Source/AnimNext/Private/Interface/InterfaceState.cpp
deleted file mode 100644
diff --git a/Engine/Plugins/Experimental/AnimNext/Source/AnimNext/Private/Param/Param.cpp b/Engine/Plugins/Experimental/AnimNext/Source/AnimNext/Private/Param/Param.cpp
deleted file mode 100644
diff --git a/Engine/Plugins/Experimental/AnimNext/Source/AnimNext/Private/Param/ParamHelpers.cpp b/Engine/Plugins/Experimental/AnimNext/Source/AnimNext/Private/Param/ParamHelpers.cpp
--- a/Engine/Plugins/Experimental/AnimNext/Source/AnimNext/Private/Param/ParamHelpers.cpp
+++ b/Engine/Plugins/Experimental/AnimNext/Source/AnimNext/Private/Param/ParamHelpers.cpp
@@ -1,4 +1,4 @@
-﻿// Copyright Epic Games, Inc. All Rights Reserved.
+// Copyright Epic Games, Inc. All Rights Reserved.
 
 #include "Param/ParamHelpers.h"
 #include "Param/ParamType.h"
@@ -12,299 +12,503 @@ FParamHelpers::ECopyResult FParamHelpers::Copy(const FAnimNextParamType& InSourc
 {
 	if(InSourceType == InTargetType)
 	{
-		switch(InSourceType.ContainerType)
+		Copy(InSourceType, InSourceMemory, InTargetMemory);
+		return ECopyResult::Succeeded;
+	}
+
+	/** TODO: mismatched/conversions etc. */
+	check(false);
+	return ECopyResult::Failed;
+}
+
+
+FParamHelpers::ECopyResult FParamHelpers::Copy(const FParamTypeHandle& InSourceTypeHandle, const FParamTypeHandle& InTargetTypeHandle, TConstArrayView<uint8> InSourceMemory, TArrayView<uint8> InTargetMemory)
+{
+	if(InSourceTypeHandle == InTargetTypeHandle)
+	{
+		Copy(InSourceTypeHandle, InSourceMemory, InTargetMemory);
+		return ECopyResult::Succeeded;
+	}
+
+	/** TODO: mismatched/conversions etc. */
+	check(false);
+	return ECopyResult::Failed;
+}
+
+void FParamHelpers::Copy(const FAnimNextParamType& InType, TConstArrayView<uint8> InSourceMemory, TArrayView<uint8> InTargetMemory)
+{
+	switch (InType.ContainerType)
+	{
+	case FAnimNextParamType::EContainerType::Array:
+	{
+		const TScriptArray<FHeapAllocator>* SourceArray = reinterpret_cast<const TScriptArray<FHeapAllocator>*>(InSourceMemory.GetData());
+		TScriptArray<FHeapAllocator>* TargetArray = reinterpret_cast<TScriptArray<FHeapAllocator>*>(InTargetMemory.GetData());
+
+		const int32 NumElements = SourceArray->Num();
+		const size_t ValueTypeSize = InType.GetValueTypeSize();
+		const size_t ValueTypeAlignment = InType.GetValueTypeSize();
+
+		// Reallocate target array
+		TargetArray->SetNumUninitialized(NumElements, ValueTypeSize, ValueTypeAlignment);
+		check(TargetArray->GetAllocatedSize(ValueTypeSize) == SourceArray->GetAllocatedSize(ValueTypeSize));
+
+		// Perform the copy according to value type
+		switch (InType.ValueType)
+		{
+		case FAnimNextParamType::EValueType::None:
+			ensureMsgf(false, TEXT("Trying to copy parameter of type None"));
+			break;
+		case FAnimNextParamType::EValueType::Bool:
+		case FAnimNextParamType::EValueType::Byte:
+		case FAnimNextParamType::EValueType::Int32:
+		case FAnimNextParamType::EValueType::Int64:
+		case FAnimNextParamType::EValueType::Float:
+		case FAnimNextParamType::EValueType::Double:
+		case FAnimNextParamType::EValueType::Name:
+		case FAnimNextParamType::EValueType::Enum:
+		{
+			FMemory::Memcpy(TargetArray->GetData(), SourceArray->GetData(), SourceArray->GetAllocatedSize(ValueTypeSize));
+		}
+		break;
+		case FAnimNextParamType::EValueType::String:
+		{
+			const FString* SourceString = static_cast<const FString*>(SourceArray->GetData());
+			FString* TargetString = static_cast<FString*>(TargetArray->GetData());
+			for (int32 Index = 0; Index < NumElements; ++Index)
+			{
+				*TargetString = *SourceString;
+				++TargetString;
+				++SourceString;
+			}
+		}
+		break;
+		case FAnimNextParamType::EValueType::Text:
+		{
+			const FText* SourceText = static_cast<const FText*>(SourceArray->GetData());
+			FText* TargetText = static_cast<FText*>(TargetArray->GetData());
+			for (int32 Index = 0; Index < NumElements; ++Index)
+			{
+				*TargetText = *SourceText;
+				++TargetText;
+				++SourceText;
+			}
+		}
+		break;
+		case FAnimNextParamType::EValueType::Struct:
+			if (const UScriptStruct* ScriptStruct = Cast<UScriptStruct>(InType.ValueTypeObject.Get()))
+			{
+				ScriptStruct->CopyScriptStruct(TargetArray->GetData(), SourceArray->GetData(), NumElements);
+			}
+			else
+			{
+				checkf(false, TEXT("Error: FParameterHelpers::Copy: Unknown Struct Type"));
+			}
+			break;
+		case FAnimNextParamType::EValueType::Object:
+		{
+			const TObjectPtr<UObject>* SourceObjectPtr = reinterpret_cast<const TObjectPtr<UObject>*>(InSourceMemory.GetData());
+			TObjectPtr<UObject>* TargetObjectPtr = reinterpret_cast<TObjectPtr<UObject>*>(InTargetMemory.GetData());
+			for (int32 Index = 0; Index < NumElements; ++Index)
+			{
+				*TargetObjectPtr = *SourceObjectPtr;
+				++TargetObjectPtr;
+				++SourceObjectPtr;
+			}
+		}
+		break;
+		case FAnimNextParamType::EValueType::SoftObject:
 		{
-		case FAnimNextParamType::EContainerType::Array:
+			const TSoftObjectPtr<UObject>* SourceSoftObjectPtr = reinterpret_cast<const TSoftObjectPtr<UObject>*>(InSourceMemory.GetData());
+			TSoftObjectPtr<UObject>* TargetSoftObjectPtr = reinterpret_cast<TSoftObjectPtr<UObject>*>(InTargetMemory.GetData());
+			for (int32 Index = 0; Index < NumElements; ++Index)
 			{
-				const TScriptArray<FHeapAllocator>* SourceArray = reinterpret_cast<const TScriptArray<FHeapAllocator>*>(InSourceMemory.GetData());
-				TScriptArray<FHeapAllocator>* TargetArray = reinterpret_cast<TScriptArray<FHeapAllocator>*>(InTargetMemory.GetData());
+				*TargetSoftObjectPtr = *SourceSoftObjectPtr;
+				++TargetSoftObjectPtr;
+				++SourceSoftObjectPtr;
+			}
+		}
+		break;
+		case FAnimNextParamType::EValueType::Class:
+		{
+			const TSubclassOf<UObject>* SourceClassPtr = reinterpret_cast<const TSubclassOf<UObject>*>(InSourceMemory.GetData());
+			TSubclassOf<UObject>* TargetClassPtr = reinterpret_cast<TSubclassOf<UObject>*>(InTargetMemory.GetData());
+			for (int32 Index = 0; Index < NumElements; ++Index)
+			{
+				*TargetClassPtr = *SourceClassPtr;
+				++TargetClassPtr;
+				++SourceClassPtr;
+			}
+		}
+		break;
+		case FAnimNextParamType::EValueType::SoftClass:
+		{
+			const TSoftClassPtr<UObject>* SourceClassPtr = reinterpret_cast<const TSoftClassPtr<UObject>*>(InSourceMemory.GetData());
+			TSoftClassPtr<UObject>* TargetClassPtr = reinterpret_cast<TSoftClassPtr<UObject>*>(InTargetMemory.GetData());
+			for (int32 Index = 0; Index < NumElements; ++Index)
+			{
+				*TargetClassPtr = *SourceClassPtr;
+				++TargetClassPtr;
+				++SourceClassPtr;
+			}
+		}
+		break;
+		default:
+			checkf(false, TEXT("Error: FParameterHelpers::Copy of unknown type"));
+			break;
+		}
+	}
+	break;
+	case FAnimNextParamType::EContainerType::None:
+	{
+		switch (InType.ValueType)
+		{
+		case FAnimNextParamType::EValueType::None:
+			ensureMsgf(false, TEXT("Trying to copy parameter of type None"));
+			break;
+		case FAnimNextParamType::EValueType::Bool:
+		case FAnimNextParamType::EValueType::Byte:
+		case FAnimNextParamType::EValueType::Int32:
+		case FAnimNextParamType::EValueType::Int64:
+		case FAnimNextParamType::EValueType::Float:
+		case FAnimNextParamType::EValueType::Double:
+		case FAnimNextParamType::EValueType::Name:
+		case FAnimNextParamType::EValueType::Enum:
+		{
+			const int32 ParamAlignment = InType.GetAlignment();
+			const int32 ParamSize = InType.GetSize();
+			const int32 ParamAllocSize = Align(ParamSize, ParamAlignment);
+			check(InTargetMemory.Num() <= ParamAllocSize);
+			check(InSourceMemory.Num() <= ParamAllocSize);
+			check(IsAligned(InTargetMemory.GetData(), ParamAlignment));
+			check(IsAligned(InSourceMemory.GetData(), ParamAlignment));
+
+			FMemory::Memcpy(InTargetMemory.GetData(), InSourceMemory.GetData(), ParamAllocSize);
+		}
+		break;
+		case FAnimNextParamType::EValueType::String:
+		{
+			const FString* SourceString = reinterpret_cast<const FString*>(InSourceMemory.GetData());
+			FString* TargetString = reinterpret_cast<FString*>(InTargetMemory.GetData());
+			*TargetString = *SourceString;
+		}
+		break;
+		case FAnimNextParamType::EValueType::Text:
+		{
+			const FText* SourceText = reinterpret_cast<const FText*>(InSourceMemory.GetData());
+			FText* TargetText = reinterpret_cast<FText*>(InTargetMemory.GetData());
+			*TargetText = *SourceText;
+		}
+		break;
+		case FAnimNextParamType::EValueType::Struct:
+			if (const UScriptStruct* ScriptStruct = Cast<UScriptStruct>(InType.ValueTypeObject.Get()))
+			{
+				const int32 StructSize = ScriptStruct->GetStructureSize();
+				check(StructSize <= InTargetMemory.Num());
+
+				ScriptStruct->CopyScriptStruct(InTargetMemory.GetData(), InSourceMemory.GetData(), 1);
+			}
+			else
+			{
+				checkf(false, TEXT("Error: FParameterHelpers::Copy: Unknown Struct Type"));
+			}
+			break;
+		case FAnimNextParamType::EValueType::Object:
+		{
+			const TObjectPtr<UObject>* SourceObjectPtr = reinterpret_cast<const TObjectPtr<UObject>*>(InSourceMemory.GetData());
+			TObjectPtr<UObject>* TargetObjectPtr = reinterpret_cast<TObjectPtr<UObject>*>(InTargetMemory.GetData());
+			*TargetObjectPtr = *SourceObjectPtr;
+		}
+		break;
+		case FAnimNextParamType::EValueType::SoftObject:
+		{
+			const TSoftObjectPtr<UObject>* SourceSoftObjectPtr = reinterpret_cast<const TSoftObjectPtr<UObject>*>(InSourceMemory.GetData());
+			TSoftObjectPtr<UObject>* TargetSoftObjectPtr = reinterpret_cast<TSoftObjectPtr<UObject>*>(InTargetMemory.GetData());
+			*TargetSoftObjectPtr = *SourceSoftObjectPtr;
+		}
+		break;
+		case FAnimNextParamType::EValueType::Class:
+		{
+			const TSubclassOf<UObject>* SourceClassPtr = reinterpret_cast<const TSubclassOf<UObject>*>(InSourceMemory.GetData());
+			TSubclassOf<UObject>* TargetClassPtr = reinterpret_cast<TSubclassOf<UObject>*>(InTargetMemory.GetData());
+			*TargetClassPtr = *SourceClassPtr;
+		}
+		break;
+		case FAnimNextParamType::EValueType::SoftClass:
+		{
+			const TSoftClassPtr<UObject>* SourceClassPtr = reinterpret_cast<const TSoftClassPtr<UObject>*>(InSourceMemory.GetData());
+			TSoftClassPtr<UObject>* TargetClassPtr = reinterpret_cast<TSoftClassPtr<UObject>*>(InTargetMemory.GetData());
+			*TargetClassPtr = *SourceClassPtr;
+		}
+		break;
+		default:
+			checkf(false, TEXT("Error: FParameterHelpers::Copy of unknown type"));
+			break;
+		}
+	}
+	break;
+	}
+}
+
+void FParamHelpers::Copy(const FParamTypeHandle& InTypeHandle, TConstArrayView<uint8> InSourceMemory, TArrayView<uint8> InTargetMemory)
+{
+	auto SimpleCopy = [&InTargetMemory, &InSourceMemory](int32 InParamSize, int32 InAlignment)
+	{
+		const int32 ParamAllocSize = Align(InParamSize, InAlignment);
+		check(InTargetMemory.Num() >= ParamAllocSize);
+		check(InSourceMemory.Num() >= ParamAllocSize);
+		check(IsAligned(InTargetMemory.GetData(), InAlignment));
+		check(IsAligned(InSourceMemory.GetData(), InAlignment));
 
-				const int32 NumElements = SourceArray->Num(); 
-				const size_t ValueTypeSize = InSourceType.GetValueTypeSize();
-				const size_t ValueTypeAlignment = InSourceType.GetValueTypeSize();
+		FMemory::Memmove(InTargetMemory.GetData(), InSourceMemory.GetData(), ParamAllocSize);
+	};
 
-				// Reallocate target array
-				TargetArray->SetNumUninitialized(NumElements, ValueTypeSize, ValueTypeAlignment);
-				check(TargetArray->GetAllocatedSize(ValueTypeSize) == SourceArray->GetAllocatedSize(ValueTypeSize));
+	switch (InTypeHandle.GetParameterType())
+	{
+	case FParamTypeHandle::EParamType::None:
+		ensureMsgf(false, TEXT("Trying to copy parameter of type None"));
+		break;
+	case FParamTypeHandle::EParamType::Bool:
+		SimpleCopy(sizeof(bool), alignof(bool));
+		break;
+	case FParamTypeHandle::EParamType::Byte:
+		SimpleCopy(sizeof(uint8), alignof(uint8));
+		break;
+	case FParamTypeHandle::EParamType::Int32:
+		SimpleCopy(sizeof(int32), alignof(int32));
+		break;
+	case FParamTypeHandle::EParamType::Int64:
+		SimpleCopy(sizeof(int64), alignof(int64));
+		break;
+	case FParamTypeHandle::EParamType::Float:
+		SimpleCopy(sizeof(float), alignof(float));
+		break;
+	case FParamTypeHandle::EParamType::Double:
+		SimpleCopy(sizeof(double), alignof(double));
+		break;
+	case FParamTypeHandle::EParamType::Name:
+		SimpleCopy(sizeof(FName), alignof(FName));
+		break;
+	case FParamTypeHandle::EParamType::String:
+		{
+			const FString* SourceString = reinterpret_cast<const FString*>(InSourceMemory.GetData());
+			FString* TargetString = reinterpret_cast<FString*>(InTargetMemory.GetData());
+			*TargetString = *SourceString;
+		}
+		break;
+	case FParamTypeHandle::EParamType::Text:
+		{
+			const FText* SourceText = reinterpret_cast<const FText*>(InSourceMemory.GetData());
+			FText* TargetText = reinterpret_cast<FText*>(InTargetMemory.GetData());
+			*TargetText = *SourceText;
+		}
+		break;
+	case FParamTypeHandle::EParamType::Vector:
+		SimpleCopy(sizeof(FVector), alignof(FVector));
+		break;
+	case FParamTypeHandle::EParamType::Vector4:
+		SimpleCopy(sizeof(FVector4), alignof(FVector4));
+		break;
+	case FParamTypeHandle::EParamType::Quat:
+		SimpleCopy(sizeof(FQuat), alignof(FQuat));
+		break;
+	case FParamTypeHandle::EParamType::Transform:
+		SimpleCopy(sizeof(FTransform), alignof(FTransform));
+		break;
+	case FParamTypeHandle::EParamType::Custom:
+		{
+			FAnimNextParamType Type = InTypeHandle.GetType();
+			Copy(Type, InSourceMemory, InTargetMemory);
+		}
+		break;
+	default:
+		checkf(false, TEXT("Error: FParameterHelpers::Copy of unknown type"));
+		break;
+	}
+}
+
+void FParamHelpers::Destroy(const FAnimNextParamType& InType, TArrayView<uint8> InMemory)
+{
+	switch (InType.ContainerType)
+	{
+	case FAnimNextParamType::EContainerType::Array:
+		{
+			TScriptArray<FHeapAllocator>* Array = reinterpret_cast<TScriptArray<FHeapAllocator>*>(InMemory.GetData());
+			const int32 NumElements = Array->Num();
 
-				// Perform the copy according to value type
-				switch (InSourceType.ValueType)
+			// Perform the destroy according to value type
+			switch (InType.ValueType)
+			{
+			case FAnimNextParamType::EValueType::None:
+			case FAnimNextParamType::EValueType::Bool:
+			case FAnimNextParamType::EValueType::Byte:
+			case FAnimNextParamType::EValueType::Int32:
+			case FAnimNextParamType::EValueType::Int64:
+			case FAnimNextParamType::EValueType::Float:
+			case FAnimNextParamType::EValueType::Double:
+			case FAnimNextParamType::EValueType::Name:
+			case FAnimNextParamType::EValueType::Enum:
+				break;
+			case FAnimNextParamType::EValueType::String:
 				{
-				case FAnimNextParamType::EValueType::None:
-					ensureMsgf(false, TEXT("Trying to copy parameter of type None"));
-					break;
-				case FAnimNextParamType::EValueType::Bool:
-				case FAnimNextParamType::EValueType::Byte:
-				case FAnimNextParamType::EValueType::Int32:
-				case FAnimNextParamType::EValueType::Int64:
-				case FAnimNextParamType::EValueType::Float:
-				case FAnimNextParamType::EValueType::Double:
-				case FAnimNextParamType::EValueType::Name:
-				case FAnimNextParamType::EValueType::Enum:
-					{
-						FMemory::Memcpy(TargetArray->GetData(), SourceArray->GetData(), SourceArray->GetAllocatedSize(ValueTypeSize));
-					}
-					break;
-				case FAnimNextParamType::EValueType::String:
-					{
-						const FString* SourceString = static_cast<const FString*>(SourceArray->GetData());
-						FString* TargetString = static_cast<FString*>(TargetArray->GetData());
-						for(int32 Index = 0; Index < NumElements; ++Index)
-						{
-							*TargetString = *SourceString;
-							++TargetString;
-							++SourceString;
-						}
-					}
-					break;
-				case FAnimNextParamType::EValueType::Text:
-					{
-						const FText* SourceText = static_cast<const FText*>(SourceArray->GetData());
-						FText* TargetText = static_cast<FText*>(TargetArray->GetData());
-						for(int32 Index = 0; Index < NumElements; ++Index)
-						{
-							*TargetText = *SourceText;
-							++TargetText;
-							++SourceText;
-						}
-					}
-					break;
-				case FAnimNextParamType::EValueType::Struct:
-					if(UScriptStruct* ScriptStruct = Cast<UScriptStruct>(InSourceType.ValueTypeObject.Get()))
-					{
-						ScriptStruct->CopyScriptStruct(TargetArray->GetData(), SourceArray->GetData(), NumElements);
-					}
-					else
-					{
-						checkf(false, TEXT("Error: FParameterHelpers::Copy: Unknown Struct Type"));
-					}
-					break;
-				case FAnimNextParamType::EValueType::Object:
+					FString* String = static_cast<FString*>(Array->GetData());
+					for (int32 Index = 0; Index < NumElements; ++Index)
 					{
-						const TObjectPtr<UObject>* SourceObjectPtr = reinterpret_cast<const TObjectPtr<UObject>*>(InSourceMemory.GetData());
-						TObjectPtr<UObject>* TargetObjectPtr = reinterpret_cast<TObjectPtr<UObject>*>(InTargetMemory.GetData());
-						for(int32 Index = 0; Index < NumElements; ++Index)
-						{
-							*TargetObjectPtr = *SourceObjectPtr;
-							++TargetObjectPtr;
-							++SourceObjectPtr;
-						}
+						String->~FString();
+						++String;
 					}
-					break;
-				case FAnimNextParamType::EValueType::SoftObject:
+				}
+				break;
+			case FAnimNextParamType::EValueType::Text:
+				{
+					FText* Text = static_cast<FText*>(Array->GetData());
+					for (int32 Index = 0; Index < NumElements; ++Index)
 					{
-						const TSoftObjectPtr<UObject>* SourceSoftObjectPtr = reinterpret_cast<const TSoftObjectPtr<UObject>*>(InSourceMemory.GetData());
-						TSoftObjectPtr<UObject>* TargetSoftObjectPtr = reinterpret_cast<TSoftObjectPtr<UObject>*>(InTargetMemory.GetData());
-						for(int32 Index = 0; Index < NumElements; ++Index)
-						{
-							*TargetSoftObjectPtr = *SourceSoftObjectPtr;
-							++TargetSoftObjectPtr;
-							++SourceSoftObjectPtr;
-						}
+						Text->~FText();
+						++Text;
 					}
-					break;
-				case FAnimNextParamType::EValueType::Class:
+				}
+				break;
+			case FAnimNextParamType::EValueType::Struct:
+				if (const UScriptStruct* ScriptStruct = Cast<UScriptStruct>(InType.ValueTypeObject.Get()))
+				{
+					ScriptStruct->DestroyStruct(Array->GetData(), NumElements);
+				}
+				else
+				{
+					checkf(false, TEXT("Error: FParameterHelpers::Destroy: Unknown Struct Type"));
+				}
+				break;
+			case FAnimNextParamType::EValueType::Object:
+				{
+					TObjectPtr<UObject>* ObjectPtr = reinterpret_cast<TObjectPtr<UObject>*>(InMemory.GetData());
+					for (int32 Index = 0; Index < NumElements; ++Index)
 					{
-						const TSubclassOf<UObject>* SourceClassPtr = reinterpret_cast<const TSubclassOf<UObject>*>(InSourceMemory.GetData());
-						TSubclassOf<UObject>* TargetClassPtr = reinterpret_cast<TSubclassOf<UObject>*>(InTargetMemory.GetData());
-						for(int32 Index = 0; Index < NumElements; ++Index)
-						{
-							*TargetClassPtr = *SourceClassPtr;
-							++TargetClassPtr;
-							++SourceClassPtr;
-						}
+						*ObjectPtr = nullptr;
+						++ObjectPtr;
 					}
-					break;
-				case FAnimNextParamType::EValueType::SoftClass:
+				}
+				break;
+			case FAnimNextParamType::EValueType::SoftObject:
+				break;
+			case FAnimNextParamType::EValueType::Class:
+				{
+					TSubclassOf<UObject>* ClassPtr = reinterpret_cast<TSubclassOf<UObject>*>(InMemory.GetData());
+					for (int32 Index = 0; Index < NumElements; ++Index)
 					{
-						const TSoftClassPtr<UObject>* SourceClassPtr = reinterpret_cast<const TSoftClassPtr<UObject>*>(InSourceMemory.GetData());
-						TSoftClassPtr<UObject>* TargetClassPtr = reinterpret_cast<TSoftClassPtr<UObject>*>(InTargetMemory.GetData());
-						for(int32 Index = 0; Index < NumElements; ++Index)
-						{
-							*TargetClassPtr = *SourceClassPtr;
-							++TargetClassPtr;
-							++SourceClassPtr;
-						}
+						*ClassPtr = nullptr;
+						++ClassPtr;
 					}
-					break;
-				default:
-					checkf(false, TEXT("Error: FParameterHelpers::Copy of unknown type"));
-					break;
 				}
+				break;
+			case FAnimNextParamType::EValueType::SoftClass:
+				break;
+			default:
+				checkf(false, TEXT("Error: FParameterHelpers::Destroy of unknown type"));
+				break;
 			}
-			break;
-		case FAnimNextParamType::EContainerType::None:
+		}
+		break;
+	case FAnimNextParamType::EContainerType::None:
+		{
+			switch (InType.ValueType)
 			{
-				switch (InSourceType.ValueType)
+			case FAnimNextParamType::EValueType::None:
+			case FAnimNextParamType::EValueType::Bool:
+			case FAnimNextParamType::EValueType::Byte:
+			case FAnimNextParamType::EValueType::Int32:
+			case FAnimNextParamType::EValueType::Int64:
+			case FAnimNextParamType::EValueType::Float:
+			case FAnimNextParamType::EValueType::Double:
+			case FAnimNextParamType::EValueType::Name:
+			case FAnimNextParamType::EValueType::Enum:
+				break;
+			case FAnimNextParamType::EValueType::String:
 				{
-				case FAnimNextParamType::EValueType::None:
-					ensureMsgf(false, TEXT("Trying to copy parameter of type None"));
-					break;
-				case FAnimNextParamType::EValueType::Bool:
-				case FAnimNextParamType::EValueType::Byte:
-				case FAnimNextParamType::EValueType::Int32:
-				case FAnimNextParamType::EValueType::Int64:
-				case FAnimNextParamType::EValueType::Float:
-				case FAnimNextParamType::EValueType::Double:
-				case FAnimNextParamType::EValueType::Name:
-				case FAnimNextParamType::EValueType::Enum:
-					{
-						const int32 ParamAlignment = InSourceType.GetAlignment();
-						const int32 ParamSize = InSourceType.GetSize();
-						const int32 ParamAllocSize = Align(ParamSize, ParamAlignment);
-						check(InTargetMemory.Num() <= ParamAllocSize);
-						check(InSourceMemory.Num() <= ParamAllocSize);
-						check(IsAligned(InTargetMemory.GetData(), ParamAlignment));
-						check(IsAligned(InSourceMemory.GetData(), ParamAlignment));
-
-						FMemory::Memcpy(InTargetMemory.GetData(), InSourceMemory.GetData(), ParamAllocSize);
-					}
-					break;
-				case FAnimNextParamType::EValueType::String:
-					{
-						const FString* SourceString = reinterpret_cast<const FString*>(InSourceMemory.GetData());
-						FString* TargetString = reinterpret_cast<FString*>(InTargetMemory.GetData());
-						*TargetString = *SourceString;
-					}
-					break;
-				case FAnimNextParamType::EValueType::Text:
-					{
-						const FText* SourceText = reinterpret_cast<const FText*>(InSourceMemory.GetData());
-						FText* TargetText = reinterpret_cast<FText*>(InTargetMemory.GetData());
-						*TargetText = *SourceText;
-					}
-					break;
-				case FAnimNextParamType::EValueType::Struct:
-					if(UScriptStruct* ScriptStruct = Cast<UScriptStruct>(InSourceType.ValueTypeObject.Get()))
-					{
-						const int32 StructSize = ScriptStruct->GetStructureSize();
-						check(StructSize <= InTargetMemory.Num());
-
-						ScriptStruct->CopyScriptStruct(InTargetMemory.GetData(), InSourceMemory.GetData(), 1);
-					}
-					else
-					{
-						checkf(false, TEXT("Error: FParameterHelpers::Copy: Unknown Struct Type"));
-					}
-					break;
-				case FAnimNextParamType::EValueType::Object:
-					{
-						const TObjectPtr<UObject>* SourceObjectPtr = reinterpret_cast<const TObjectPtr<UObject>*>(InSourceMemory.GetData());
-						TObjectPtr<UObject>* TargetObjectPtr = reinterpret_cast<TObjectPtr<UObject>*>(InTargetMemory.GetData());
-						*TargetObjectPtr = *SourceObjectPtr;
-					}
-					break;
-				case FAnimNextParamType::EValueType::SoftObject:
-					{
-						const TSoftObjectPtr<UObject>* SourceSoftObjectPtr = reinterpret_cast<const TSoftObjectPtr<UObject>*>(InSourceMemory.GetData());
-						TSoftObjectPtr<UObject>* TargetSoftObjectPtr = reinterpret_cast<TSoftObjectPtr<UObject>*>(InTargetMemory.GetData());
-						*TargetSoftObjectPtr = *SourceSoftObjectPtr;
-					}
-					break;
-				case FAnimNextParamType::EValueType::Class:
-					{
-						const TSubclassOf<UObject>* SourceClassPtr = reinterpret_cast<const TSubclassOf<UObject>*>(InSourceMemory.GetData());
-						TSubclassOf<UObject>* TargetClassPtr = reinterpret_cast<TSubclassOf<UObject>*>(InTargetMemory.GetData());
-						*TargetClassPtr = *SourceClassPtr;
-					}
-					break;
-				case FAnimNextParamType::EValueType::SoftClass:
-					{
-						const TSoftClassPtr<UObject>* SourceClassPtr = reinterpret_cast<const TSoftClassPtr<UObject>*>(InSourceMemory.GetData());
-						TSoftClassPtr<UObject>* TargetClassPtr = reinterpret_cast<TSoftClassPtr<UObject>*>(InTargetMemory.GetData());
-						*TargetClassPtr = *SourceClassPtr;
-					}
-					break;
-				default:
-					checkf(false, TEXT("Error: FParameterHelpers::Copy of unknown type"));
-					break;
+					FString* String = reinterpret_cast<FString*>(InMemory.GetData());
+					String->~FString();
+				}
+				break;
+			case FAnimNextParamType::EValueType::Text:
+				{
+					FText* Text = reinterpret_cast<FText*>(InMemory.GetData());
+					Text->~FText();
+				}
+				break;
+			case FAnimNextParamType::EValueType::Struct:
+				if (const UScriptStruct* ScriptStruct = Cast<UScriptStruct>(InType.ValueTypeObject.Get()))
+				{
+					ScriptStruct->DestroyStruct(InMemory.GetData(), 1);
+				}
+				else
+				{
+					checkf(false, TEXT("Error: FParameterHelpers::Destroy: Unknown Struct Type"));
+				}
+				break;
+			case FAnimNextParamType::EValueType::Object:
+				{
+					TObjectPtr<UObject>* ObjectPtr = reinterpret_cast<TObjectPtr<UObject>*>(InMemory.GetData());
+					*ObjectPtr = nullptr;
 				}
+				break;
+			case FAnimNextParamType::EValueType::SoftObject:
+				break;
+			case FAnimNextParamType::EValueType::Class:
+				{
+					TSubclassOf<UObject>* ClassPtr = reinterpret_cast<TSubclassOf<UObject>*>(InMemory.GetData());
+					*ClassPtr = nullptr;
+				}
+				break;
+			case FAnimNextParamType::EValueType::SoftClass:
+				break;
+			default:
+				checkf(false, TEXT("Error: FParameterHelpers::Destroy of unknown type"));
+				break;
 			}
-			break;
 		}
-		
-		return ECopyResult::Succeeded;
-	}
-	else
-	{
-		/** TODO: mismatched/conversions etc. */
-		check(false);
-		return ECopyResult::Failed;
+		break;
 	}
 }
 
-
-FParamHelpers::ECopyResult FParamHelpers::Copy(const FParamTypeHandle& InSourceTypeHandle, const FParamTypeHandle& InTargetTypeHandle, TConstArrayView<uint8> InSourceMemory, TArrayView<uint8> InTargetMemory)
+void FParamHelpers::Destroy(const FParamTypeHandle& InTypeHandle, TArrayView<uint8> InMemory)
 {
-	if(InSourceTypeHandle == InTargetTypeHandle)
+	switch (InTypeHandle.GetParameterType())
 	{
-		auto SimpleCopy = [&InTargetMemory, &InSourceMemory](int32 InParamSize, int32 InAlignment)
+	case FParamTypeHandle::EParamType::None:
+	case FParamTypeHandle::EParamType::Bool:
+	case FParamTypeHandle::EParamType::Byte:
+	case FParamTypeHandle::EParamType::Int32:
+	case FParamTypeHandle::EParamType::Int64:
+	case FParamTypeHandle::EParamType::Float:
+	case FParamTypeHandle::EParamType::Double:
+	case FParamTypeHandle::EParamType::Name:
+		break;
+	case FParamTypeHandle::EParamType::String:
 		{
-			const int32 ParamAllocSize = Align(InParamSize, InAlignment);
-			check(InTargetMemory.Num() <= ParamAllocSize);
-			check(InSourceMemory.Num() <= ParamAllocSize);
-			check(IsAligned(InTargetMemory.GetData(), InAlignment));
-			check(IsAligned(InSourceMemory.GetData(), InAlignment));
-
-			FMemory::Memcpy(InTargetMemory.GetData(), InSourceMemory.GetData(), ParamAllocSize);
-
-			return ECopyResult::Succeeded;
-		};
-
-		switch(InSourceTypeHandle.GetParameterType())
+			FString* String = reinterpret_cast<FString*>(InMemory.GetData());
+			String->~FString();
+		}
+		break;
+	case FParamTypeHandle::EParamType::Text:
 		{
-		case FParamTypeHandle::EParamType::None:
-			ensureMsgf(false, TEXT("Trying to copy parameter of type None"));
-			break;
-		case FParamTypeHandle::EParamType::Bool:
-			return SimpleCopy(sizeof(bool), alignof(bool));
-		case FParamTypeHandle::EParamType::Byte:
-			return SimpleCopy(sizeof(uint8), alignof(uint8));
-		case FParamTypeHandle::EParamType::Int32:
-			return SimpleCopy(sizeof(int32), alignof(int32));
-		case FParamTypeHandle::EParamType::Int64:
-			return SimpleCopy(sizeof(int64), alignof(int64));
-		case FParamTypeHandle::EParamType::Float:
-			return SimpleCopy(sizeof(float), alignof(float));
-		case FParamTypeHandle::EParamType::Double:
-			return SimpleCopy(sizeof(double), alignof(double));
-		case FParamTypeHandle::EParamType::Name:
-			return SimpleCopy(sizeof(FName), alignof(FName));
-		case FParamTypeHandle::EParamType::String:
-			{
-				const FString* SourceString = reinterpret_cast<const FString*>(InSourceMemory.GetData());
-				FString* TargetString = reinterpret_cast<FString*>(InTargetMemory.GetData());
-				*TargetString = *SourceString;
-				return ECopyResult::Succeeded;
-			}
-		case FParamTypeHandle::EParamType::Text:
-			{
-				const FText* SourceText = reinterpret_cast<const FText*>(InSourceMemory.GetData());
-				FText* TargetText = reinterpret_cast<FText*>(InTargetMemory.GetData());
-				*TargetText = *SourceText;
-			}
-		case FParamTypeHandle::EParamType::Vector:
-			return SimpleCopy(sizeof(FVector), alignof(FVector));
-		case FParamTypeHandle::EParamType::Vector4:
-			return SimpleCopy(sizeof(FVector4), alignof(FVector4));
-		case FParamTypeHandle::EParamType::Quat:
-			return SimpleCopy(sizeof(FQuat), alignof(FQuat));
-		case FParamTypeHandle::EParamType::Transform:
-			return SimpleCopy(sizeof(FTransform), alignof(FTransform));
-		case FParamTypeHandle::EParamType::Custom:
-			{
-				FAnimNextParamType Type = InSourceTypeHandle.GetType();
-				return Copy(Type, Type, InSourceMemory, InTargetMemory);
-			}
-		default:
-			checkf(false, TEXT("Error: FParameterHelpers::Copy of unknown type"));
-			break;
+			FText* Text = reinterpret_cast<FText*>(InMemory.GetData());
+			Text->~FText();
 		}
+		break;
+	case FParamTypeHandle::EParamType::Vector:
+	case FParamTypeHandle::EParamType::Vector4:
+	case FParamTypeHandle::EParamType::Quat:
+	case FParamTypeHandle::EParamType::Transform:
+		break;
+	case FParamTypeHandle::EParamType::Custom:
+		{
+			FAnimNextParamType Type = InTypeHandle.GetType();
+			Destroy(Type, InMemory);
+		}
+		break;
+	default:
+		checkf(false, TEXT("Error: FParameterHelpers::Destroy of unknown type"));
+		break;
 	}
-
-	/** TODO: mismatched/conversions etc. */
-	check(false);
-	return ECopyResult::Failed;
 }
 
 }
\ No newline at end of file
diff --git a/Engine/Plugins/Experimental/AnimNext/Source/AnimNext/Private/Param/ParamHelpers.h b/Engine/Plugins/Experimental/AnimNext/Source/AnimNext/Private/Param/ParamHelpers.h
new file mode 100644
--- /dev/null
+++ b/Engine/Plugins/Experimental/AnimNext/Source/AnimNext/Private/Param/ParamHelpers.h
@@ -0,0 +1,41 @@
+// Copyright Epic Games, Inc. All Rights Reserved.
+
+#pragma once
+
+#include "CoreMinimal.h"
+
+struct FAnimNextParamType;
+
+namespace UE::AnimNext
+{
+
+struct FParamTypeHandle;
+
+struct FParamHelpers
+{
+	enum class ECopyResult
+	{
+		Failed,
+		Succeeded,
+	};
+
+	/** Copy a parameter from one type to another, validating type, container and any polymorphism rules if the types differ at all. */
+	static ECopyResult Copy(const FAnimNextParamType& InSourceType, const FAnimNextParamType& InTargetType, TConstArrayView<uint8> InSourceMemory, TArrayView<uint8> InTargetMemory);
+
+	/** Copy a parameter from one type to another, validating type, container and any polymorphism rules if the types differ at all. */
+	static ECopyResult Copy(const FParamTypeHandle& InSourceTypeHandle, const FParamTypeHandle& InTargetTypeHandle, TConstArrayView<uint8> InSourceMemory, TArrayView<uint8> InTargetMemory);
+
+	/** Copy a parameter of the specified type from one location to another */
+	static void Copy(const FAnimNextParamType& InType, TConstArrayView<uint8> InSourceMemory, TArrayView<uint8> InTargetMemory);
+
+	/** Copy a parameter of the specified type from one location to another */
+	static void Copy(const FParamTypeHandle& InTypeHandle, TConstArrayView<uint8> InSourceMemory, TArrayView<uint8> InTargetMemory);
+
+	/** Destroy a parameter of the specified type */
+	static void Destroy(const FAnimNextParamType& InType, TArrayView<uint8> InTargetMemory);
+
+	/** Destroy a parameter of the specified type */
+	static void Destroy(const FParamTypeHandle& InTypeHandle,TArrayView<uint8> InTargetMemory);
+};
+
+}
\ No newline at end of file
diff --git a/Engine/Plugins/Experimental/AnimNext/Source/AnimNext/Private/Param/ParamId.cpp b/Engine/Plugins/Experimental/AnimNext/Source/AnimNext/Private/Param/ParamId.cpp
new file mode 100644
--- /dev/null
+++ b/Engine/Plugins/Experimental/AnimNext/Source/AnimNext/Private/Param/ParamId.cpp
@@ -0,0 +1,101 @@
+// Copyright Epic Games, Inc. All Rights Reserved.
+
+#include "Param/ParamId.h"
+#include "Misc/ScopeRWLock.h"
+#include "HAL/ThreadSingleton.h"
+
+namespace UE::AnimNext
+{
+
+static FRWLock GParamIdLock;
+static TArray<FName> GParamIdToName;
+static TMap<FName, uint32> GNameToParamId;
+#if WITH_DEV_AUTOMATION_TESTS
+static bool bGTestSandbox = false;
+struct FSandboxThreadData : TThreadSingleton<FSandboxThreadData>
+{
+	TArray<FName> ParamIdToName;
+	TMap<FName, uint32> NameToParamId;
+};
+#endif
+
+FParamId::FParamId(FName InName)
+{
+#if WITH_DEV_AUTOMATION_TESTS
+	if (bGTestSandbox)
+	{
+		if (const uint32* FoundIndex = FSandboxThreadData::Get().NameToParamId.Find(InName))
+		{
+			ParameterIndex = *FoundIndex;
+		}
+		else
+		{
+			ParameterIndex = FSandboxThreadData::Get().ParamIdToName.Num();
+			FSandboxThreadData::Get().NameToParamId.Add(InName, ParameterIndex);
+			FSandboxThreadData::Get().ParamIdToName.Add(InName);
+		}
+	}
+	else
+#endif
+	{
+		FRWScopeLock Lock(GParamIdLock, SLT_Write);
+		if (const uint32* FoundIndex = GNameToParamId.Find(InName))
+		{
+			ParameterIndex = *FoundIndex;
+		}
+		else
+		{
+			ParameterIndex = GParamIdToName.Num();
+			GNameToParamId.Add(InName, ParameterIndex);
+			GParamIdToName.Add(InName);
+		}
+	}
+}
+
+FName FParamId::ToName() const
+{
+#if WITH_DEV_AUTOMATION_TESTS
+	if (bGTestSandbox)
+	{
+		return FSandboxThreadData::Get().ParamIdToName[ParameterIndex];
+	}
+	else
+#endif
+	{
+		FRWScopeLock Lock(GParamIdLock, SLT_ReadOnly);
+		return GParamIdToName[ParameterIndex];
+	}
+}
+
+FParamId FParamId::GetMaxParamId()
+{
+#if WITH_DEV_AUTOMATION_TESTS
+	if (bGTestSandbox)
+	{
+		return FParamId((uint32)FSandboxThreadData::Get().ParamIdToName.Num());
+	}
+	else
+#endif
+	{
+		FRWScopeLock Lock(GParamIdLock, SLT_ReadOnly);
+		return FParamId((uint32)GParamIdToName.Num());
+	}
+}
+
+#if WITH_DEV_AUTOMATION_TESTS
+void FParamId::BeginTestSandbox()
+{
+	bGTestSandbox = true;
+	FSandboxThreadData::Get().NameToParamId.Empty();
+	FSandboxThreadData::Get().ParamIdToName.Empty();
+}
+
+void FParamId::EndTestSandbox()
+{
+	FSandboxThreadData::Get().NameToParamId.Empty();
+	FSandboxThreadData::Get().ParamIdToName.Empty();
+	bGTestSandbox = false;
+}
+#endif
+
+}
\ No newline at end of file
diff --git a/Engine/Plugins/Experimental/AnimNext/Source/AnimNext/Private/Param/ParamStack.cpp b/Engine/Plugins/Experimental/AnimNext/Source/AnimNext/Private/Param/ParamStack.cpp
new file mode 100644
--- /dev/null
+++ b/Engine/Plugins/Experimental/AnimNext/Source/AnimNext/Private/Param/ParamStack.cpp
@@ -0,0 +1,472 @@
+// Copyright Epic Games, Inc. All Rights Reserved.
+
+#include "Param/ParamStack.h"
+#include "Param/ParamHelpers.h"
+#include "PropertyBag.h"
+#include "EngineLogs.h"
+
+#define LOCTEXT_NAMESPACE "AnimNextParamStack"
+
+namespace UE::AnimNext
+{
+
+FParamStack::FParam::FParam(const FParamTypeHandle& InTypeHandle, TArrayView<uint8> InData, bool bInIsReference, bool bInIsMutable)
+	: Data(nullptr)
+	, TypeHandle(InTypeHandle)
+	, Size(InData.Num())
+	, Flags(EParamFlags::None)
+{
+	check(TypeHandle.IsValid());
+	check(InData.Num() > 0 && InData.Num() < 0xffff);
+
+	// If we can store our data inside of a ptr, we do
+	if (!bInIsReference && InData.Num() <= sizeof(void*))
+	{
+		FParamHelpers::Copy(InTypeHandle, InData, TArrayView<uint8>(reinterpret_cast<uint8*>(&Data), sizeof(void*)));
+		Flags |= EParamFlags::Embedded;
+	}
+	else
+	{
+		Data = InData.GetData();
+	}
+
+	if (bInIsReference)
+	{
+		Flags |= EParamFlags::Reference;
+	}
+
+	if (bInIsMutable)
+	{
+		Flags |= EParamFlags::Mutable;
+	}
+}
+
+FParamStack::FParam::~FParam()
+{
+	if (Size > 0)
+	{
+		if (IsEmbedded())
+		{
+			FParamHelpers::Destroy(TypeHandle, TArrayView<uint8>(reinterpret_cast<uint8*>(&Data), sizeof(void*)));
+		}
+	}
+}
+
+FParamStack::FParam::FParam(const FParam& InOtherParam)
+	: TypeHandle(InOtherParam.TypeHandle)
+	, Size(InOtherParam.Size)
+	, Flags(InOtherParam.Flags)
+{
+	if (IsEmbedded())
+	{
+		FParamHelpers::Copy(InOtherParam.TypeHandle, TConstArrayView<uint8>(reinterpret_cast<const uint8*>(&InOtherParam.Data), sizeof(void*)), TArrayView<uint8>(reinterpret_cast<uint8*>(&Data), sizeof(void*)));
+	}
+	else
+	{
+		Data = InOtherParam.Data;
+	}
+}
+
+FParamStack::FParam& FParamStack::FParam::operator=(const FParam& InOtherParam)
+{
+	TypeHandle = InOtherParam.TypeHandle;
+	Size = InOtherParam.Size;
+	Flags = InOtherParam.Flags;
+
+	if (IsEmbedded())
+	{
+		FParamHelpers::Copy(InOtherParam.TypeHandle, TConstArrayView<uint8>(reinterpret_cast<const uint8*>(&InOtherParam.Data), sizeof(void*)), TArrayView<uint8>(reinterpret_cast<uint8*>(&Data), sizeof(void*)));
+	}
+	else
+	{
+		Data = InOtherParam.Data;
+	}
+
+	return *this;
+}
+
+FParamStack::FParam::FParam(FParam&& InOtherParam)
+	: TypeHandle(InOtherParam.TypeHandle)
+	, Size(InOtherParam.Size)
+	, Flags(InOtherParam.Flags)
+{
+	if (IsEmbedded())
+	{
+		FParamHelpers::Copy(InOtherParam.TypeHandle, TConstArrayView<uint8>(reinterpret_cast<const uint8*>(&InOtherParam.Data), sizeof(void*)), TArrayView<uint8>(reinterpret_cast<uint8*>(&Data), sizeof(void*)));
+	}
+	else
+	{
+		Data = InOtherParam.Data;
+	}
+}
+
+FParamStack::FParam& FParamStack::FParam::operator=(FParam&& InOtherParam)
+{
+	TypeHandle = InOtherParam.TypeHandle;
+	Size = InOtherParam.Size;
+	Flags = InOtherParam.Flags;
+
+	if (IsEmbedded())
+	{
+		FParamHelpers::Copy(InOtherParam.TypeHandle, TConstArrayView<uint8>(reinterpret_cast<const uint8*>(&InOtherParam.Data), sizeof(void*)), TArrayView<uint8>(reinterpret_cast<uint8*>(&Data), sizeof(void*)));
+	}
+	else
+	{
+		Data = InOtherParam.Data;
+	}
+
+	return *this;
+}
+
+FParamStackLayer::FParamStackLayer(const FInstancedPropertyBag& InPropertyBag)
+{
+	TConstArrayView<FPropertyBagPropertyDesc> Descs = InPropertyBag.GetPropertyBagStruct()->GetPropertyDescs();
+	
+	// Determine param ID range for this layer
+	TArray<FParamId> CachedIds;
+	CachedIds.SetNumUninitialized(Descs.Num());
+	MinParamId = MAX_uint32;
+	uint32 MaxParamId = 0;
+	for (uint32 DescIndex = 0; DescIndex < static_cast<uint32>(Descs.Num()); ++DescIndex)
+	{
+		const FPropertyBagPropertyDesc& Desc = Descs[DescIndex];
+		const FParamId& ParamId = CachedIds[DescIndex] = FParamId(Desc.Name);
+		MinParamId = FMath::Min(ParamId.ToInt(), MinParamId);
+		MaxParamId = FMath::Max(ParamId.ToInt(), MaxParamId);
+	}
+
+	if (MinParamId <= MaxParamId)
+	{
+		const uint32 ParamRangeSize = (MaxParamId - MinParamId) + 1;
+		Params.SetNumZeroed(ParamRangeSize);
+		FConstStructView StructView = InPropertyBag.GetValue();
+		for (uint32 DescIndex = 0; DescIndex < static_cast<uint32>(Descs.Num()); ++DescIndex)
+		{
+			const FPropertyBagPropertyDesc& Desc = Descs[DescIndex];
+			const FParamId& ParamId = CachedIds[DescIndex];
+			const uint8* DataPtr = StructView.GetMemory() + Desc.CachedProperty->GetOffset_ForInternal();
+			const uint32 LocalParamIndex = ParamId.ToInt() - MinParamId;
+			Params[LocalParamIndex] = FParamStack::FParam(FParamTypeHandle::FromPropertyBagPropertyDesc(Desc), TArrayView<uint8>(const_cast<uint8*>(DataPtr), Desc.CachedProperty->GetSize()), true, false);
+		}
+	}
+}
+
+FParamStackLayer::FParamStackLayer(TConstArrayView<TPair<FParamId, FParamStack::FParam>> InParams)
+{
+	MinParamId = MAX_uint32;
+	uint32 MaxParamId = 0;
+	for (uint32 ParamIndex = 0; ParamIndex < static_cast<uint32>(InParams.Num()); ++ParamIndex)
+	{
+		MinParamId = FMath::Min(InParams[ParamIndex].Key.ToInt(), MinParamId);
+		MaxParamId = FMath::Max(InParams[ParamIndex].Key.ToInt(), MaxParamId);
+	}
+
+	if (MinParamId <= MaxParamId)
+	{
+		const uint32 ParamRangeSize = (MaxParamId - MinParamId) + 1;
+		Params.SetNumZeroed(ParamRangeSize);
+		for (uint32 ParamIndex = 0; ParamIndex < static_cast<uint32>(InParams.Num()); ++ParamIndex)
+		{
+			const TPair<FParamId, FParamStack::FParam>& Pair = InParams[ParamIndex];
+			const uint32 LocalParamIndex = Pair.Key.ToInt() - MinParamId;
+			Params[LocalParamIndex] = Pair.Value;
+		}
+	}
+}
+
+FParamStack::FPushedLayer::FPushedLayer(FParamStackLayer& InLayer, FParamStack& InStack)
+	: Layer(InLayer)
+{
+	if (Layer.Params.Num())
+	{
+		InStack.ResizeLayerIndices();
+
+		for (uint32 LocalParamIndex = 0; LocalParamIndex < (uint32)InLayer.Params.Num(); ++LocalParamIndex)
+		{
+			const uint32 GlobalParamIndex = Layer.MinParamId + LocalParamIndex;
+			InStack.LayerIndices[GlobalParamIndex] = InLayer.Params[LocalParamIndex].IsValid() ? 0 : MAX_uint16;
+		}
+	}
+}
+
+FParamStack::FPushedLayer::FPushedLayer(const FPushedLayer& InPreviousLayer, FParamStackLayer& InLayer, FParamStack& InStack)
+	: Layer(InLayer)
+{
+	if (Layer.Params.Num())
+	{
+		InStack.ResizeLayerIndices();
+
+		const uint32 NumParams = Layer.Params.Num();
+		const uint32 BaseLayerIndex = InStack.PreviousLayerIndices.Num();
+		InStack.PreviousLayerIndices.SetNum(BaseLayerIndex + NumParams);
+		PreviousLayerIndexStart = BaseLayerIndex;
+		
+		// Update layer indices and build previous layer indices
+		for (uint32 LocalParamIndex = 0; LocalParamIndex < NumParams; ++LocalParamIndex)
+		{
+			const uint32 GlobalParamIndex = Layer.MinParamId + LocalParamIndex;
+			InStack.PreviousLayerIndices[BaseLayerIndex + LocalParamIndex] = InStack.LayerIndices[GlobalParamIndex];
+			if(InLayer.Params[LocalParamIndex].IsValid())
+			{
+				InStack.LayerIndices[GlobalParamIndex] = InStack.Layers.Num();
+			}
+		}
+	}
+}
+
+struct FParamStackThreadData : TThreadSingleton<FParamStackThreadData>
+{
+	FParamStack Stack;
+};
+
+FParamStack::FParamStack()
+{
+	Layers.Reserve(8);
+	LayerIndices.Reserve(FParamId::GetMaxParamId().ToInt());
+	PreviousLayerIndices.Reserve(FParamId::GetMaxParamId().ToInt());
+}
+
+FParamStack& FParamStack::Get()
+{
+	return FParamStackThreadData::Get().Stack;
+}
+
+void FParamStack::PushLayer(FParamStackLayer& InLayer)
+{
+	if (Layers.Num() < MAX_uint16)
+	{
+		if (Layers.Num())
+		{
+			Layers.Push(FPushedLayer(Layers.Top(), InLayer, *this));
+		}
+		else
+		{
+			Layers.Push(FPushedLayer(InLayer, *this));
+		}
+	}
+	else
+	{
+		UE_LOG(LogAnimation, Warning, TEXT("FParamStack: Could not push a layer: Maximum 65535 stack layers."))
+	}
+}
+
+void FParamStack::PushLayer(TConstArrayView<TPair<FParamId, FParamStack::FParam>> InParams)
+{
+	if(Layers.Num() < MAX_uint16)
+	{
+		FParamStackLayer& OwnedLayer = OwnedStackLayers.Add_GetRef(FParamStackLayer(InParams));
+		OwnedLayer.OwnedStorageOffset = AllocAndCopyOwnedParamStorage(OwnedLayer.Params);
+		PushLayer(OwnedLayer);
+	}
+	else
+	{
+		UE_LOG(LogAnimation, Warning, TEXT("FParamStack: Could not push a layer: Maximum 65535 stack layers."))
+	}
+}
+
+void FParamStack::PopLayer()
+{
+	if(Layers.Num() > 0)
+	{
+		const FPushedLayer& TopLayer = Layers.Top();
+
+		// Fixup layer indices to previous, if any
+		const uint32 NumParams = TopLayer.Layer.Params.Num();
+		if(NumParams > 0 && PreviousLayerIndices.Num() > 0)
+		{
+			for (uint32 LocalParamIndex = 0; LocalParamIndex < (uint32)NumParams; ++LocalParamIndex)
+			{
+				const uint16 PreviousLayerIndex = PreviousLayerIndices[TopLayer.PreviousLayerIndexStart + LocalParamIndex];
+				if (PreviousLayerIndex != MAX_uint16)
+				{
+					const uint32 GlobalParamIndex = TopLayer.Layer.MinParamId + LocalParamIndex;
+					LayerIndices[GlobalParamIndex] = PreviousLayerIndex;
+				}
+			}
+
+			PreviousLayerIndices.SetNum(PreviousLayerIndices.Num() - NumParams);
+		}
+
+		// Dont shrink allocs to avoid thrashing
+		constexpr bool bAllowShrinking = false;
+
+		// If we own the layer, pop the owned stack
+		if(TopLayer.Layer.OwnedStorageOffset != MAX_uint32)
+		{
+			check(OwnedStackLayers.Num() && &TopLayer.Layer == &OwnedStackLayers[OwnedStackLayers.Num() - 1]);
+			OwnedStackLayers.Pop(bAllowShrinking);
+
+			// Free any owned storage
+			FreeOwnedParamStorage(TopLayer.Layer.OwnedStorageOffset);
+		}
+
+		// Pop the layer itself
+		Layers.Pop(bAllowShrinking);
+	}
+}
+
+TUniquePtr<FParamStackLayer> FParamStack::MakeLayer(const FInstancedPropertyBag& InPropertyBag)
+{
+	TUniquePtr<FParamStackLayer> Layer = TUniquePtr<FParamStackLayer>(new FParamStackLayer(InPropertyBag));
+	return Layer;
+}
+
+TUniquePtr<FParamStackLayer> FParamStack::MakeLayer(TConstArrayView<TPair<FParamId, FParamStack::FParam>> InParams)
+{
+	TUniquePtr<FParamStackLayer> Layer = TUniquePtr<FParamStackLayer>(new FParamStackLayer(InParams));
+	return Layer;
+}
+
+FParamStack::EGetParamResult FParamStack::GetParamData(FParamId InId, FParamTypeHandle InTypeHandle, TConstArrayView<uint8>& OutParamData) const
+{
+	if (InId.ToInt() >= (uint32)LayerIndices.Num() || LayerIndices[InId.ToInt()] == MAX_uint16)
+	{
+		return EGetParamResult::NotInScope;
+	}
+
+	const FPushedLayer& Layer = Layers[LayerIndices[InId.ToInt()]];
+	const uint32 LocalParamIndex = InId.ToInt() - Layer.Layer.MinParamId;
+
+	const FParam& Param = Layer.Layer.Params[LocalParamIndex];
+	if (!Param.IsValid())
+	{
+		return EGetParamResult::NotInScope;
+	}
+
+	if (Param.GetTypeHandle() != InTypeHandle)
+	{
+		return EGetParamResult::IncorrectType;
+	}
+
+	OutParamData = Param.GetData();
+
+	return EGetParamResult::Succeeded;
+}
+
+FParamStack::EGetParamResult FParamStack::GetMutableParamData(FParamId InId, FParamTypeHandle InTypeHandle, TArrayView<uint8>& OutParamData)
+{
+	if (InId.ToInt() >= (uint32)LayerIndices.Num() || LayerIndices[InId.ToInt()] == MAX_uint16)
+	{
+		return EGetParamResult::NotInScope;
+	}
+
+	FPushedLayer& Layer = Layers[LayerIndices[InId.ToInt()]];
+	const uint32 LocalParamIndex = InId.ToInt() - Layer.Layer.MinParamId;
+
+	FParam& Param = Layer.Layer.Params[LocalParamIndex];
+	if (!Param.IsValid())
+	{
+		return EGetParamResult::NotInScope;
+	}
+
+	EGetParamResult AccessResult = Param.GetTypeHandle() != InTypeHandle ? EGetParamResult::IncorrectType : EGetParamResult::Succeeded;
+	AccessResult |= !Param.IsMutable() ? EGetParamResult::Immutable : EGetParamResult::Succeeded;
+	if (AccessResult != EGetParamResult::Succeeded)
+	{
+		return AccessResult;
+	}
+
+	OutParamData = Param.GetMutableData();
+
+	return EGetParamResult::Succeeded;
+}
+
+bool FParamStack::IsMutableParam(FParamId InId) const
+{
+	if (InId.ToInt() >= (uint32)LayerIndices.Num() || LayerIndices[InId.ToInt()] == MAX_uint16)
+	{
+		return false;
+	}
+
+	const FPushedLayer& Layer = Layers[LayerIndices[InId.ToInt()]];
+	const uint32 LocalParamIndex = InId.ToInt() - Layer.Layer.MinParamId;
+
+	const FParam& Param = Layer.Layer.Params[LocalParamIndex];
+
+	return Param.IsMutable();
+}
+
+bool FParamStack::IsReferenceParam(FParamId InId) const
+{
+	if (InId.ToInt() >= (uint32)LayerIndices.Num() || LayerIndices[InId.ToInt()] == MAX_uint16)
+	{
+		return false;
+	}
+
+	const FPushedLayer& Layer = Layers[LayerIndices[InId.ToInt()]];
+	const uint32 LocalParamIndex = InId.ToInt() - Layer.Layer.MinParamId;
+
+	const FParam& Param = Layer.Layer.Params[LocalParamIndex];
+
+	return Param.IsReference();
+}
+
+uint32 FParamStack::AllocAndCopyOwnedParamStorage(TArrayView<FParam> InParams)
+{
+	const uint32 CurrentOffset = OwnedLayerParamStorage.Num();
+	const uint32 PageSize = OwnedLayerParamStorage.MaxPerPage();
+
+	for (FParam& Param : InParams)
+	{
+		if(Param.IsValid())
+		{
+			FParamTypeHandle TypeHandle = Param.GetTypeHandle();
+			const uint32 Size = TypeHandle.GetSize();
+
+			// Dont copy references or embdedded values to internal storage
+			// - References should refer back to the originally-passed in data
+			// - Embedded do not need their own storage 
+			if (!Param.IsReference() && !Param.IsEmbedded())
+			{
+				const uint32 Alignment = TypeHandle.GetAlignment();
+				uint32 AllocSize = Align(Size, Alignment);
+				const uint32 NumPages = OwnedLayerParamStorage.NumPages();
+				check(AllocSize < PageSize);
+
+				// Do we need a new page? If so extend allocated size to the next page boundary
+				const uint32 PageMax = (NumPages + 1) * PageSize;
+				uint32 BaseOffset = Align(CurrentOffset, Alignment);
+				if (BaseOffset + AllocSize > PageMax)
+				{
+					OwnedLayerParamStorage.SetNum((NumPages + 1) * PageSize);
+				}
+
+				// Extend to encompass requested size
+				BaseOffset = Align(OwnedLayerParamStorage.Num(), Alignment);
+				OwnedLayerParamStorage.SetNum(BaseOffset + AllocSize);
+		
+				// Copy param data
+				TArrayView<uint8> TargetMemory(&OwnedLayerParamStorage[BaseOffset], Size);
+				FParamHelpers::Copy(TypeHandle, TypeHandle, Param.GetData(), TargetMemory);
+
+				// Update param to reference new storage
+				Param.Data = TargetMemory.GetData();
+			}
+		}
+	}
+
+	return CurrentOffset;
+}
+
+void FParamStack::FreeOwnedParamStorage(uint32 InOffset)
+{
+	check(InOffset <= (uint32)OwnedLayerParamStorage.Num());
+
+	OwnedLayerParamStorage.SetNum(InOffset, false);
+}
+
+void FParamStack::ResizeLayerIndices()
+{
+	const uint32 NumParams = FParamId::GetMaxParamId().ToInt();
+	const uint32 NumLayerIndices = LayerIndices.Num();
+
+	if (NumParams > NumLayerIndices)
+	{
+		LayerIndices.SetNum(NumParams);
+		FMemory::Memset(&LayerIndices[NumLayerIndices], 0xff, (NumParams - NumLayerIndices) * sizeof(uint16));
+	}
+}
+
+}
+
+#undef LOCTEXT_NAMESPACE
\ No newline at end of file
diff --git a/Engine/Plugins/Experimental/AnimNext/Source/AnimNext/Private/Param/ParamStorage.cpp b/Engine/Plugins/Experimental/AnimNext/Source/AnimNext/Private/Param/ParamStorage.cpp
deleted file mode 100644
diff --git a/Engine/Plugins/Experimental/AnimNext/Source/AnimNext/Private/Param/ParamType.cpp b/Engine/Plugins/Experimental/AnimNext/Source/AnimNext/Private/Param/ParamType.cpp
--- a/Engine/Plugins/Experimental/AnimNext/Source/AnimNext/Private/Param/ParamType.cpp
+++ b/Engine/Plugins/Experimental/AnimNext/Source/AnimNext/Private/Param/ParamType.cpp
@@ -6,7 +6,7 @@
 #include "UObject/Class.h"
 #include "Templates/SubclassOf.h"
 
-FAnimNextParamType::FAnimNextParamType(EValueType InValueType, EContainerType InContainerType, UObject* InValueTypeObject)
+FAnimNextParamType::FAnimNextParamType(EValueType InValueType, EContainerType InContainerType, const UObject* InValueTypeObject)
 	: ValueTypeObject(InValueTypeObject)
 	, ValueType(InValueType)
 	, ContainerType(InContainerType)
@@ -93,7 +93,7 @@ UE::AnimNext::FParamTypeHandle FAnimNextParamType::GetHandle() const
 			Handle.SetParameterType(FParamTypeHandle::EParamType::Text);
 			break;
 		case EValueType::Struct:
-			if(UScriptStruct* ScriptStruct = Cast<UScriptStruct>(ValueTypeObject.Get()))
+			if(const UScriptStruct* ScriptStruct = Cast<UScriptStruct>(ValueTypeObject.Get()))
 			{
 				if(ScriptStruct == TBaseStructure<FVector>::Get())
 				{
@@ -123,7 +123,7 @@ UE::AnimNext::FParamTypeHandle FAnimNextParamType::GetHandle() const
 			}
 			break;
 		case EValueType::Enum:
-			if(UEnum* Enum = Cast<UEnum>(ValueTypeObject.Get()))
+			if(const UEnum* Enum = Cast<UEnum>(ValueTypeObject.Get()))
 			{
 				Handle.SetParameterType(FParamTypeHandle::EParamType::Custom);
 				Handle.SetCustomTypeIndex(FParamTypeHandle::GetOrAllocateCustomTypeIndex(ValueType, EContainerType::None, Enum));
@@ -137,7 +137,7 @@ UE::AnimNext::FParamTypeHandle FAnimNextParamType::GetHandle() const
 		case EValueType::SoftObject:
 		case EValueType::Class:
 		case EValueType::SoftClass:
-			if(UClass* Class = Cast<UClass>(ValueTypeObject.Get()))
+			if(const UClass* Class = Cast<UClass>(ValueTypeObject.Get()))
 			{
 				Handle.SetParameterType(FParamTypeHandle::EParamType::Custom);
 				Handle.SetCustomTypeIndex(FParamTypeHandle::GetOrAllocateCustomTypeIndex(ValueType, EContainerType::None, Class));
@@ -150,7 +150,6 @@ UE::AnimNext::FParamTypeHandle FAnimNextParamType::GetHandle() const
 		}
 		break;
 	case EPropertyBagContainerType::Array:
-	default:
 		switch(ValueType)
 		{
 		default:
@@ -170,7 +169,7 @@ UE::AnimNext::FParamTypeHandle FAnimNextParamType::GetHandle() const
 			Handle.SetCustomTypeIndex(FParamTypeHandle::GetOrAllocateCustomTypeIndex(ValueType, ContainerType, nullptr));
 			break;
 		case EValueType::Struct:
-			if(UScriptStruct* ScriptStruct = Cast<UScriptStruct>(ValueTypeObject.Get()))
+			if(const UScriptStruct* ScriptStruct = Cast<UScriptStruct>(ValueTypeObject.Get()))
 			{
 				Handle.SetParameterType(FParamTypeHandle::EParamType::Custom);
 				Handle.SetCustomTypeIndex(FParamTypeHandle::GetOrAllocateCustomTypeIndex(ValueType, ContainerType, ScriptStruct));
@@ -181,7 +180,7 @@ UE::AnimNext::FParamTypeHandle FAnimNextParamType::GetHandle() const
 			}
 			break;
 		case EValueType::Enum:
-			if(UEnum* Enum = Cast<UEnum>(ValueTypeObject.Get()))
+			if(const UEnum* Enum = Cast<UEnum>(ValueTypeObject.Get()))
 			{
 				Handle.SetParameterType(FParamTypeHandle::EParamType::Custom);
 				Handle.SetCustomTypeIndex(FParamTypeHandle::GetOrAllocateCustomTypeIndex(ValueType, ContainerType, Enum));
@@ -195,7 +194,7 @@ UE::AnimNext::FParamTypeHandle FAnimNextParamType::GetHandle() const
 		case EValueType::SoftObject:
 		case EValueType::Class:
 		case EValueType::SoftClass:
-			if(UClass* Class = Cast<UClass>(ValueTypeObject.Get()))
+			if(const UClass* Class = Cast<UClass>(ValueTypeObject.Get()))
 			{
 				Handle.SetParameterType(FParamTypeHandle::EParamType::Custom);
 				Handle.SetCustomTypeIndex(FParamTypeHandle::GetOrAllocateCustomTypeIndex(ValueType, ContainerType, Class));
@@ -207,6 +206,8 @@ UE::AnimNext::FParamTypeHandle FAnimNextParamType::GetHandle() const
 			break;
 		}
 		break;
+	default:
+		break;
 	}
 
 	return Handle;
@@ -256,7 +257,7 @@ size_t FAnimNextParamType::GetValueTypeSize() const
 		// TODO: seems to be no way to find the size of a UEnum?
 		return sizeof(uint8);
 	case EValueType::Struct:
-		if(UScriptStruct* ScriptStruct = Cast<UScriptStruct>(ValueTypeObject.Get()))
+		if(const UScriptStruct* ScriptStruct = Cast<UScriptStruct>(ValueTypeObject.Get()))
 		{
 			return ScriptStruct->GetStructureSize();
 		}
@@ -325,7 +326,7 @@ size_t FAnimNextParamType::GetValueTypeAlignment() const
 		// TODO: seems to be no way to find the alignment of a UEnum?
 		return alignof(uint8);
 	case EValueType::Struct:
-		if(UScriptStruct* ScriptStruct = Cast<UScriptStruct>(ValueTypeObject.Get()))
+		if(const UScriptStruct* ScriptStruct = Cast<UScriptStruct>(ValueTypeObject.Get()))
 		{
 			return ScriptStruct->GetMinAlignment();
 		}
@@ -387,7 +388,7 @@ FString FAnimNextParamType::ToString() const
 			InStringBuilder.Append(TEXT("FText"));
 			break;
 		case EValueType::Enum:
-			if(UEnum* Enum = Cast<UEnum>(ValueTypeObject.Get()))
+			if(const UEnum* Enum = Cast<UEnum>(ValueTypeObject.Get()))
 			{
 				InStringBuilder.Append(TEXT("U"));
 				InStringBuilder.Append(Enum->GetName());
@@ -398,7 +399,7 @@ FString FAnimNextParamType::ToString() const
 			}
 			break;
 		case EValueType::Struct:
-			if(UScriptStruct* Struct = Cast<UScriptStruct>(ValueTypeObject.Get()))
+			if(const UScriptStruct* Struct = Cast<UScriptStruct>(ValueTypeObject.Get()))
 			{
 				InStringBuilder.Append(TEXT("F"));
 				InStringBuilder.Append(Struct->GetName());
@@ -409,7 +410,7 @@ FString FAnimNextParamType::ToString() const
 			}
 			break;
 		case EValueType::Object:
-			if(UClass* Class = Cast<UClass>(ValueTypeObject.Get()))
+			if(const UClass* Class = Cast<UClass>(ValueTypeObject.Get()))
 			{
 				InStringBuilder.Append(TEXT("TObjectPtr<U"));
 				InStringBuilder.Append(Class->GetName());
@@ -421,7 +422,7 @@ FString FAnimNextParamType::ToString() const
 			}
 			break;
 		case EValueType::SoftObject:
-			if(UClass* Class = Cast<UClass>(ValueTypeObject.Get()))
+			if(const UClass* Class = Cast<UClass>(ValueTypeObject.Get()))
 			{
 				InStringBuilder.Append(TEXT("TSoftObjectPtr<U"));
 				InStringBuilder.Append(Class->GetName());
@@ -433,7 +434,7 @@ FString FAnimNextParamType::ToString() const
 			}
 			break;
 		case EValueType::Class:
-			if(UClass* Class = Cast<UClass>(ValueTypeObject.Get()))
+			if(const UClass* Class = Cast<UClass>(ValueTypeObject.Get()))
 			{
 				InStringBuilder.Append(TEXT("TSubClassOf<U"));
 				InStringBuilder.Append(Class->GetName());
@@ -445,7 +446,7 @@ FString FAnimNextParamType::ToString() const
 			}
 			break;
 		case EValueType::SoftClass:
-			if(UClass* Class = Cast<UClass>(ValueTypeObject.Get()))
+			if(const UClass* Class = Cast<UClass>(ValueTypeObject.Get()))
 			{
 				InStringBuilder.Append(TEXT("TSoftClassPtr<U"));
 				InStringBuilder.Append(Class->GetName());
diff --git a/Engine/Plugins/Experimental/AnimNext/Source/AnimNext/Private/Param/ParamTypeHandle.cpp b/Engine/Plugins/Experimental/AnimNext/Source/AnimNext/Private/Param/ParamTypeHandle.cpp
--- a/Engine/Plugins/Experimental/AnimNext/Source/AnimNext/Private/Param/ParamTypeHandle.cpp
+++ b/Engine/Plugins/Experimental/AnimNext/Source/AnimNext/Private/Param/ParamTypeHandle.cpp
@@ -4,6 +4,7 @@
 #include "Param/ParamType.h"
 #include "Misc/ScopeRWLock.h"
 #include "UObject/Class.h"
+#include "UObject/TextProperty.h"
 
 namespace UE::AnimNext
 {
@@ -25,7 +26,7 @@ void FParamTypeHandle::ResetCustomTypes()
 	GTypeToIndexMap.Empty();
 }
 
-uint32 FParamTypeHandle::GetOrAllocateCustomTypeIndex(FAnimNextParamType::EValueType InValueType, FAnimNextParamType::EContainerType InContainerType, UObject* InValueTypeObject)
+uint32 FParamTypeHandle::GetOrAllocateCustomTypeIndex(FAnimNextParamType::EValueType InValueType, FAnimNextParamType::EContainerType InContainerType, const UObject* InValueTypeObject)
 {
 	FAnimNextParamType ParameterType;
 	ParameterType.ValueTypeObject = InValueTypeObject;
@@ -61,6 +62,338 @@ bool FParamTypeHandle::ValidateCustomTypeIndex(uint32 InCustomTypeIndex)
 	return GCustomTypes.IsValidIndex(InCustomTypeIndex - 1);
 }
 
+FParamTypeHandle FParamTypeHandle::FromPropertyBagPropertyDesc(const FPropertyBagPropertyDesc& Desc)
+{
+	FParamTypeHandle Handle;
+	
+	if (Desc.ContainerTypes.Num() <= 1)
+	{
+		switch(Desc.ContainerTypes.GetFirstContainerType())
+		{
+		case EPropertyBagContainerType::None:
+			switch(Desc.ValueType)
+			{
+			default:
+			case EPropertyBagPropertyType::None:
+				Handle.SetParameterType(EParamType::None);
+				break;
+			case EPropertyBagPropertyType::Bool:
+				Handle.SetParameterType(EParamType::Bool);
+				break;
+			case EPropertyBagPropertyType::Byte:
+				Handle.SetParameterType(EParamType::Byte);
+				break;
+			case EPropertyBagPropertyType::Int32:
+				Handle.SetParameterType(EParamType::Int32);
+				break;
+			case EPropertyBagPropertyType::Int64:
+				Handle.SetParameterType(EParamType::Int64);
+				break;
+			case EPropertyBagPropertyType::Float:
+				Handle.SetParameterType(EParamType::Float);
+				break;
+			case EPropertyBagPropertyType::Double:
+				Handle.SetParameterType(EParamType::Double);
+				break;
+			case EPropertyBagPropertyType::Name:
+				Handle.SetParameterType(EParamType::Name);
+				break;
+			case EPropertyBagPropertyType::String:
+				Handle.SetParameterType(EParamType::String);
+				break;
+			case EPropertyBagPropertyType::Text:
+				Handle.SetParameterType(EParamType::Text);
+				break;
+			case EPropertyBagPropertyType::Struct:
+				if(const UScriptStruct* ScriptStruct = Cast<UScriptStruct>(Desc.ValueTypeObject.Get()))
+				{
+					if(ScriptStruct == TBaseStructure<FVector>::Get())
+					{
+						Handle.SetParameterType(EParamType::Vector);
+					}
+					else if(ScriptStruct == TBaseStructure<FVector4>::Get())
+					{
+						Handle.SetParameterType(EParamType::Vector4);
+					}
+					else if(ScriptStruct == TBaseStructure<FQuat>::Get())
+					{
+						Handle.SetParameterType(EParamType::Quat);
+					}
+					else if(ScriptStruct == TBaseStructure<FTransform>::Get())
+					{
+						Handle.SetParameterType(EParamType::Transform);
+					}
+					else
+					{
+						Handle.SetParameterType(EParamType::Custom);
+						Handle.SetCustomTypeIndex(GetOrAllocateCustomTypeIndex(Desc.ValueType, Desc.ContainerTypes[0], ScriptStruct));
+					}
+				}
+				else
+				{
+					Handle.SetParameterType(EParamType::None);
+				}
+				break;
+			case EPropertyBagPropertyType::Enum:
+				if(const UEnum* Enum = Cast<UEnum>(Desc.ValueTypeObject.Get()))
+				{
+					Handle.SetParameterType(EParamType::Custom);
+					Handle.SetCustomTypeIndex(GetOrAllocateCustomTypeIndex(Desc.ValueType, Desc.ContainerTypes[0], Enum));
+				}
+				else
+				{
+					Handle.SetParameterType(EParamType::None);
+				}
+				break;
+			case EPropertyBagPropertyType::Object:
+			case EPropertyBagPropertyType::SoftObject:
+			case EPropertyBagPropertyType::Class:
+			case EPropertyBagPropertyType::SoftClass:
+				if(const UClass* Class = Cast<UClass>(Desc.ValueTypeObject.Get()))
+				{
+					Handle.SetParameterType(EParamType::Custom);
+					Handle.SetCustomTypeIndex(GetOrAllocateCustomTypeIndex(Desc.ValueType, Desc.ContainerTypes[0], Class));
+				}
+				else
+				{
+					Handle.SetParameterType(EParamType::None);
+				}
+				break;
+			}
+			break;
+		case EPropertyBagContainerType::Array:
+			switch(Desc.ValueType)
+			{
+			default:
+			case EPropertyBagPropertyType::None:
+				Handle.SetParameterType(EParamType::None);
+				break;
+			case EPropertyBagPropertyType::Bool:
+			case EPropertyBagPropertyType::Byte:
+			case EPropertyBagPropertyType::Int32:
+			case EPropertyBagPropertyType::Int64:
+			case EPropertyBagPropertyType::Float:
+			case EPropertyBagPropertyType::Double:
+			case EPropertyBagPropertyType::Name:
+			case EPropertyBagPropertyType::String:
+			case EPropertyBagPropertyType::Text:
+				Handle.SetParameterType(EParamType::Custom);
+				Handle.SetCustomTypeIndex(GetOrAllocateCustomTypeIndex(Desc.ValueType, Desc.ContainerTypes[0], nullptr));
+				break;
+			case EPropertyBagPropertyType::Struct:
+				if(const UScriptStruct* ScriptStruct = Cast<UScriptStruct>(Desc.ValueTypeObject.Get()))
+				{
+					Handle.SetParameterType(EParamType::Custom);
+					Handle.SetCustomTypeIndex(GetOrAllocateCustomTypeIndex(Desc.ValueType, Desc.ContainerTypes[0], ScriptStruct));
+				}
+				else
+				{
+					Handle.SetParameterType(EParamType::None);
+				}
+				break;
+			case EPropertyBagPropertyType::Enum:
+				if(const UEnum* Enum = Cast<UEnum>(Desc.ValueTypeObject.Get()))
+				{
+					Handle.SetParameterType(EParamType::Custom);
+					Handle.SetCustomTypeIndex(GetOrAllocateCustomTypeIndex(Desc.ValueType, Desc.ContainerTypes[0], Enum));
+				}
+				else
+				{
+					Handle.SetParameterType(EParamType::None);
+				}
+				break;
+			case EPropertyBagPropertyType::Object:
+			case EPropertyBagPropertyType::SoftObject:
+			case EPropertyBagPropertyType::Class:
+			case EPropertyBagPropertyType::SoftClass:
+				if(const UClass* Class = Cast<UClass>(Desc.ValueTypeObject.Get()))
+				{
+					Handle.SetParameterType(EParamType::Custom);
+					Handle.SetCustomTypeIndex(GetOrAllocateCustomTypeIndex(Desc.ValueType, Desc.ContainerTypes[0], Class));
+				}
+				else
+				{
+					Handle.SetParameterType(EParamType::None);
+				}
+				break;
+			}
+			break;
+		default:
+			break;
+		}
+	}
+
+	return Handle;
+}
+
+FParamTypeHandle FParamTypeHandle::FromProperty(const FProperty* InProperty)
+{
+	FParamTypeHandle Handle;
+
+	if (InProperty->IsA<FBoolProperty>())
+	{
+		Handle.SetParameterType(EParamType::Bool);
+	}
+	else if (InProperty->IsA<FByteProperty>())
+	{
+		Handle.SetParameterType(EParamType::Byte);
+	}
+	else if (InProperty->IsA<FIntProperty>())
+	{
+		Handle.SetParameterType(EParamType::Int32);
+	}
+	else if (InProperty->IsA<FInt64Property>())
+	{
+		Handle.SetParameterType(EParamType::Int64);
+	}
+	else if (InProperty->IsA<FFloatProperty>())
+	{
+		Handle.SetParameterType(EParamType::Float);
+	}
+	else if (InProperty->IsA<FDoubleProperty>())
+	{
+		Handle.SetParameterType(EParamType::Double);
+	}
+	else if (InProperty->IsA<FNameProperty>())
+	{
+		Handle.SetParameterType(EParamType::Name);
+	}
+	else if (InProperty->IsA<FStrProperty>())
+	{
+		Handle.SetParameterType(EParamType::String);
+	}
+	else if (InProperty->IsA<FTextProperty>())
+	{
+		Handle.SetParameterType(EParamType::Text);
+	}
+	else if (InProperty->IsA<FStructProperty>())
+	{
+		UScriptStruct* ScriptStruct = CastField<FStructProperty>(InProperty)->Struct;
+		if (ScriptStruct == TBaseStructure<FVector>::Get())
+		{
+			Handle.SetParameterType(EParamType::Vector);
+		}
+		else if (ScriptStruct == TBaseStructure<FVector4>::Get())
+		{
+			Handle.SetParameterType(EParamType::Vector4);
+		}
+		else if (ScriptStruct == TBaseStructure<FQuat>::Get())
+		{
+			Handle.SetParameterType(EParamType::Quat);
+		}
+		else if (ScriptStruct == TBaseStructure<FTransform>::Get())
+		{
+			Handle.SetParameterType(EParamType::Transform);
+		}
+	}
+
+	// Not found - custom type
+	if (!Handle.IsValid())
+	{
+		if (InProperty->IsA<FStructProperty>())
+		{
+			UScriptStruct* ScriptStruct = CastField<FStructProperty>(InProperty)->Struct;
+			Handle.SetParameterType(EParamType::Custom);
+			Handle.SetCustomTypeIndex(GetOrAllocateCustomTypeIndex(FAnimNextParamType::EValueType::Struct, FAnimNextParamType::EContainerType::None, ScriptStruct));
+		}
+		else if (InProperty->IsA<FObjectPropertyBase>())
+		{
+			UClass* Class = CastField<FObjectPropertyBase>(InProperty)->PropertyClass;
+			if (InProperty->IsA<FObjectProperty>() || InProperty->IsA<FObjectPtrProperty>())
+			{
+				Handle.SetParameterType(EParamType::Custom);
+				Handle.SetCustomTypeIndex(GetOrAllocateCustomTypeIndex(FAnimNextParamType::EValueType::Object, FAnimNextParamType::EContainerType::None, Class));
+			}
+			else if (InProperty->IsA<FSoftObjectProperty>())
+			{
+				Handle.SetParameterType(EParamType::Custom);
+				Handle.SetCustomTypeIndex(GetOrAllocateCustomTypeIndex(FAnimNextParamType::EValueType::SoftObject, FAnimNextParamType::EContainerType::None, Class));
+			}
+			else if (InProperty->IsA<FSoftClassProperty>())
+			{
+				Handle.SetParameterType(EParamType::Custom);
+				Handle.SetCustomTypeIndex(GetOrAllocateCustomTypeIndex(FAnimNextParamType::EValueType::SoftClass, FAnimNextParamType::EContainerType::None, Class));
+			}
+		}
+		else if(InProperty->IsA<FArrayProperty>())
+		{
+			FProperty* InnerProperty = CastField<FArrayProperty>(InProperty)->Inner;
+			if (InnerProperty->IsA<FBoolProperty>())
+			{
+				Handle.SetParameterType(EParamType::Custom);
+				Handle.SetCustomTypeIndex(GetOrAllocateCustomTypeIndex(FAnimNextParamType::EValueType::Bool, FAnimNextParamType::EContainerType::Array, nullptr));
+			}
+			else if (InnerProperty->IsA<FByteProperty>())
+			{
+				Handle.SetParameterType(EParamType::Custom);
+				Handle.SetCustomTypeIndex(GetOrAllocateCustomTypeIndex(FAnimNextParamType::EValueType::Byte, FAnimNextParamType::EContainerType::Array, nullptr));
+			}
+			else if (InnerProperty->IsA<FIntProperty>())
+			{
+				Handle.SetParameterType(EParamType::Custom);
+				Handle.SetCustomTypeIndex(GetOrAllocateCustomTypeIndex(FAnimNextParamType::EValueType::Int32, FAnimNextParamType::EContainerType::Array, nullptr));
+			}
+			else if (InnerProperty->IsA<FInt64Property>())
+			{
+				Handle.SetParameterType(EParamType::Custom);
+				Handle.SetCustomTypeIndex(GetOrAllocateCustomTypeIndex(FAnimNextParamType::EValueType::Int64, FAnimNextParamType::EContainerType::Array, nullptr));
+			}
+			else if (InnerProperty->IsA<FFloatProperty>())
+			{
+				Handle.SetParameterType(EParamType::Custom);
+				Handle.SetCustomTypeIndex(GetOrAllocateCustomTypeIndex(FAnimNextParamType::EValueType::Float, FAnimNextParamType::EContainerType::Array, nullptr));
+			}
+			else if (InnerProperty->IsA<FDoubleProperty>())
+			{
+				Handle.SetParameterType(EParamType::Custom);
+				Handle.SetCustomTypeIndex(GetOrAllocateCustomTypeIndex(FAnimNextParamType::EValueType::Double, FAnimNextParamType::EContainerType::Array, nullptr));
+			}
+			else if (InnerProperty->IsA<FNameProperty>())
+			{
+				Handle.SetParameterType(EParamType::Custom);
+				Handle.SetCustomTypeIndex(GetOrAllocateCustomTypeIndex(FAnimNextParamType::EValueType::Name, FAnimNextParamType::EContainerType::Array, nullptr));
+			}
+			else if (InnerProperty->IsA<FStrProperty>())
+			{
+				Handle.SetParameterType(EParamType::Custom);
+				Handle.SetCustomTypeIndex(GetOrAllocateCustomTypeIndex(FAnimNextParamType::EValueType::String, FAnimNextParamType::EContainerType::Array, nullptr));
+			}
+			else if (InnerProperty->IsA<FTextProperty>())
+			{
+				Handle.SetParameterType(EParamType::Custom);
+				Handle.SetCustomTypeIndex(GetOrAllocateCustomTypeIndex(FAnimNextParamType::EValueType::Text, FAnimNextParamType::EContainerType::Array, nullptr));
+			}
+			else if (InnerProperty->IsA<FStructProperty>())
+			{
+				UScriptStruct* ScriptStruct = CastField<FStructProperty>(InnerProperty)->Struct;
+				Handle.SetParameterType(EParamType::Custom);
+				Handle.SetCustomTypeIndex(GetOrAllocateCustomTypeIndex(FAnimNextParamType::EValueType::Struct, FAnimNextParamType::EContainerType::Array, ScriptStruct));
+			}
+			else if (InnerProperty->IsA<FObjectPropertyBase>())
+			{
+				UClass* Class = CastField<FObjectPropertyBase>(InnerProperty)->PropertyClass;
+				if (InnerProperty->IsA<FObjectProperty>() || InnerProperty->IsA<FObjectPtrProperty>())
+				{
+					Handle.SetParameterType(EParamType::Custom);
+					Handle.SetCustomTypeIndex(GetOrAllocateCustomTypeIndex(FAnimNextParamType::EValueType::Object, FAnimNextParamType::EContainerType::Array, Class));
+				}
+				else if (InnerProperty->IsA<FSoftObjectProperty>())
+				{
+					Handle.SetParameterType(EParamType::Custom);
+					Handle.SetCustomTypeIndex(GetOrAllocateCustomTypeIndex(FAnimNextParamType::EValueType::SoftObject, FAnimNextParamType::EContainerType::Array, Class));
+				}
+				else if (InnerProperty->IsA<FSoftClassProperty>())
+				{
+					Handle.SetParameterType(EParamType::Custom);
+					Handle.SetCustomTypeIndex(GetOrAllocateCustomTypeIndex(FAnimNextParamType::EValueType::SoftClass, FAnimNextParamType::EContainerType::Array, Class));
+				}
+			}
+		}
+	}
+
+	return Handle;
+}
+
 FAnimNextParamType FParamTypeHandle::GetType() const
 {
 	FAnimNextParamType ParameterType;
@@ -137,6 +470,7 @@ FAnimNextParamType FParamTypeHandle::GetType() const
 	return ParameterType;
 }
 
+
 size_t FParamTypeHandle::GetSize() const
 {
 	switch(GetParameterType())
diff --git a/Engine/Plugins/Experimental/AnimNext/Source/AnimNext/Private/Param/RigVMDispatch_GetParameter.cpp b/Engine/Plugins/Experimental/AnimNext/Source/AnimNext/Private/Param/RigVMDispatch_GetParameter.cpp
new file mode 100644
--- /dev/null
+++ b/Engine/Plugins/Experimental/AnimNext/Source/AnimNext/Private/Param/RigVMDispatch_GetParameter.cpp
@@ -0,0 +1,101 @@
+// Copyright Epic Games, Inc. All Rights Reserved.
+
+#include "Param/RigVMDispatch_GetParameter.h"
+
+#include "Param/ParametersExecuteContext.h"
+#include "RigVMCore/RigVMStruct.h"
+#include "RigVMCore/RigVM.h"
+#include "Graph/GraphExecuteContext.h"
+#include "Context.h"
+#include "Param/ParamStack.h"
+#include "Param/ParamStack.h"
+
+class UAnimNextParameter;
+
+const FName FRigVMDispatch_GetParameter::ValueName = TEXT("Value");
+const FName FRigVMDispatch_GetParameter::TypeHandleName = TEXT("Type");
+const FName FRigVMDispatch_GetParameter::ParameterName = TEXT("Parameter");
+const FName FRigVMDispatch_GetParameter::ParameterIdName = TEXT("ParameterId");
+
+FRigVMDispatch_GetParameter::FRigVMDispatch_GetParameter()
+{
+	FactoryScriptStruct = StaticStruct();
+}
+
+FName FRigVMDispatch_GetParameter::GetArgumentNameForOperandIndex(int32 InOperandIndex, int32 InTotalOperands) const
+{
+	static const FName ArgumentNames[] =
+	{
+		ParameterName,
+		ValueName,
+		ParameterIdName,
+		TypeHandleName
+	};
+	check(InTotalOperands == UE_ARRAY_COUNT(ArgumentNames));
+	return ArgumentNames[InOperandIndex];
+}
+
+#if WITH_EDITOR
+FString FRigVMDispatch_GetParameter::GetArgumentDefaultValue(const FName& InArgumentName, TRigVMTypeIndex InTypeIndex) const
+{
+	if (InArgumentName == ParameterIdName)
+	{
+		return FString::Printf(TEXT("%u"), 0xffffffff);
+	}
+
+	return Super::GetArgumentDefaultValue(InArgumentName, InTypeIndex);
+}
+#endif
+
+TArray<FRigVMTemplateArgument> FRigVMDispatch_GetParameter::GetArguments() const
+{
+	const TArray<FRigVMTemplateArgument::ETypeCategory> ValueCategories =
+	{
+		FRigVMTemplateArgument::ETypeCategory_SingleAnyValue,
+		FRigVMTemplateArgument::ETypeCategory_ArrayAnyValue
+	};
+	return
+	{
+		FRigVMTemplateArgument(ParameterName, ERigVMPinDirection::Input, RigVMTypeUtils::TypeIndex::FName),
+		FRigVMTemplateArgument(ValueName, ERigVMPinDirection::Output, ValueCategories),
+		FRigVMTemplateArgument(ParameterIdName, ERigVMPinDirection::Hidden, RigVMTypeUtils::TypeIndex::UInt32),
+		FRigVMTemplateArgument(TypeHandleName, ERigVMPinDirection::Hidden, RigVMTypeUtils::TypeIndex::UInt32),
+	};
+}
+
+FRigVMTemplateTypeMap FRigVMDispatch_GetParameter::OnNewArgumentType(const FName& InArgumentName, TRigVMTypeIndex InTypeIndex) const
+{
+	FRigVMTemplateTypeMap Types;
+	Types.Add(ValueName, InTypeIndex);
+	return Types;
+}
+
+void FRigVMDispatch_GetParameter::Execute(FRigVMExtendedExecuteContext& InContext, FRigVMMemoryHandleArray Handles, FRigVMPredicateBranchArray RigVMBranches)
+{
+	using namespace UE::AnimNext;
+
+	const FName Parameter = *(FName*)Handles[0].GetData();
+	const FProperty* ValueProperty = Handles[1].GetResolvedProperty();
+	check(ValueProperty);
+	uint8* TargetData = Handles[1].GetData();
+
+	uint32& ParameterId = *(uint32*)Handles[2].GetData();
+	if (ParameterId != FParamId::InvalidIndex)
+	{
+		ParameterId = FParamId(Parameter).ToInt();
+	}
+
+	uint32& TypeHandle = *(uint32*)Handles[3].GetData();
+	if (TypeHandle == 0)
+	{
+		TypeHandle = FParamTypeHandle::FromProperty(ValueProperty).ToRaw();
+	}
+
+	FAnimNextGraphExecuteContext& Context = InContext.GetPublicData<FAnimNextGraphExecuteContext>();
+	TConstArrayView<uint8> SourceData;
+	if (Context.GetContext().GetParamStack().GetParamData(FParamId(ParameterId), FParamTypeHandle::FromRaw(TypeHandle), SourceData) == FParamStack::EGetParamResult::Succeeded)
+	{
+		ValueProperty->CopyCompleteValue(TargetData, SourceData.GetData());
+	}
+}
+
diff --git a/Engine/Plugins/Experimental/AnimNext/Source/AnimNext/Private/Param/RigVMDispatch_SetParameter.cpp b/Engine/Plugins/Experimental/AnimNext/Source/AnimNext/Private/Param/RigVMDispatch_SetParameter.cpp
--- a/Engine/Plugins/Experimental/AnimNext/Source/AnimNext/Private/Param/RigVMDispatch_SetParameter.cpp
+++ b/Engine/Plugins/Experimental/AnimNext/Source/AnimNext/Private/Param/RigVMDispatch_SetParameter.cpp
@@ -1,14 +1,22 @@
-﻿// Copyright Epic Games, Inc. All Rights Reserved.
+// Copyright Epic Games, Inc. All Rights Reserved.
 
 #include "Param/RigVMDispatch_SetParameter.h"
 
 #include "Param/ParametersExecuteContext.h"
 #include "RigVMCore/RigVMStruct.h"
 #include "RigVMCore/RigVM.h"
+#include "Graph/GraphExecuteContext.h"
+#include "Context.h"
+#include "Param/ParamStack.h"
+#include "Param/ParamStack.h"
 
 class UAnimNextParameter;
-const FName FRigVMDispatch_SetParameter::ValueName = TEXT("Value");
+
 const FName FRigVMDispatch_SetParameter::ExecuteContextName = TEXT("ExecuteContext");
+const FName FRigVMDispatch_SetParameter::ValueName = TEXT("Value");
+const FName FRigVMDispatch_SetParameter::TypeHandleName = TEXT("Type");
+const FName FRigVMDispatch_SetParameter::ParameterName = TEXT("Parameter");
+const FName FRigVMDispatch_SetParameter::ParameterIdName = TEXT("ParameterId");
 
 FRigVMDispatch_SetParameter::FRigVMDispatch_SetParameter()
 {
@@ -19,12 +27,27 @@ FName FRigVMDispatch_SetParameter::GetArgumentNameForOperandIndex(int32 InOperan
 {
 	static const FName ArgumentNames[] = 
 	{
-		ValueName
+		ParameterName,
+		ValueName,
+		ParameterIdName,
+		TypeHandleName
 	};
 	check(InTotalOperands == UE_ARRAY_COUNT(ArgumentNames));
 	return ArgumentNames[InOperandIndex];
 }
 
+#if WITH_EDITOR
+FString FRigVMDispatch_SetParameter::GetArgumentDefaultValue(const FName& InArgumentName, TRigVMTypeIndex InTypeIndex) const
+{
+	if (InArgumentName == ParameterIdName)
+	{
+		return FString::Printf(TEXT("%u"), 0xffffffff);
+	}
+
+	return Super::GetArgumentDefaultValue(InArgumentName, InTypeIndex);
+}
+#endif
+
 TArray<FRigVMTemplateArgument> FRigVMDispatch_SetParameter::GetArguments() const
 {
 	const TArray<FRigVMTemplateArgument::ETypeCategory> ValueCategories =
@@ -32,7 +55,13 @@ TArray<FRigVMTemplateArgument> FRigVMDispatch_SetParameter::GetArguments() const
 		FRigVMTemplateArgument::ETypeCategory_SingleAnyValue,
 		FRigVMTemplateArgument::ETypeCategory_ArrayAnyValue
 	};
-	return { FRigVMTemplateArgument(ValueName, ERigVMPinDirection::Input, ValueCategories) };
+	return 
+	{ 
+		FRigVMTemplateArgument(ParameterName, ERigVMPinDirection::Input, RigVMTypeUtils::TypeIndex::FName),
+		FRigVMTemplateArgument(ValueName, ERigVMPinDirection::Input, ValueCategories),
+		FRigVMTemplateArgument(ParameterIdName, ERigVMPinDirection::Hidden, RigVMTypeUtils::TypeIndex::UInt32),
+		FRigVMTemplateArgument(TypeHandleName, ERigVMPinDirection::Hidden, RigVMTypeUtils::TypeIndex::UInt32),
+	};
 }
 
 TArray<FRigVMExecuteArgument> FRigVMDispatch_SetParameter::GetExecuteArguments_Impl(const FRigVMDispatchContext& InContext) const
@@ -49,13 +78,30 @@ FRigVMTemplateTypeMap FRigVMDispatch_SetParameter::OnNewArgumentType(const FName
 
 void FRigVMDispatch_SetParameter::Execute(FRigVMExtendedExecuteContext& InContext, FRigVMMemoryHandleArray Handles, FRigVMPredicateBranchArray RigVMBranches)
 {
-	const FProperty* ValueProperty = Handles[0].GetResolvedProperty(); 
+	using namespace UE::AnimNext;
+
+	const FName Parameter = *(FName*)Handles[0].GetData();
+	const FProperty* ValueProperty = Handles[1].GetResolvedProperty();
 	check(ValueProperty);
-	const uint8* SourceData =  Handles[0].GetData();
+	const uint8* SourceData =  Handles[1].GetData();
 
-	FAnimNextParametersExecuteContext& Context = InContext.GetPublicData<FAnimNextParametersExecuteContext>();
-	TArrayView<uint8> TargetData = Context.GetData();
+	uint32& ParameterId = *(uint32*)Handles[2].GetData();
+	if (ParameterId != FParamId::InvalidIndex)
+	{
+		ParameterId = FParamId(Parameter).ToInt();
+	}
 
-	ValueProperty->CopyCompleteValue(TargetData.GetData(), SourceData);
+	uint32& TypeHandle = *(uint32*)Handles[3].GetData();
+	if (TypeHandle == 0)
+	{
+		TypeHandle = FParamTypeHandle::FromProperty(ValueProperty).ToRaw();
+	}
+
+	const FAnimNextGraphExecuteContext& Context = InContext.GetPublicData<FAnimNextGraphExecuteContext>();
+	TArrayView<uint8> TargetData;
+	if (Context.GetContext().GetMutableParamStack().GetMutableParamData(FParamId(ParameterId), FParamTypeHandle::FromRaw(TypeHandle), TargetData) == FParamStack::EGetParamResult::Succeeded)
+	{
+		ValueProperty->CopyCompleteValue(TargetData.GetData(), SourceData);
+	}	
 }
 
diff --git a/Engine/Plugins/Experimental/AnimNext/Source/AnimNext/Public/Context.h b/Engine/Plugins/Experimental/AnimNext/Source/AnimNext/Public/Context.h
new file mode 100644
--- /dev/null
+++ b/Engine/Plugins/Experimental/AnimNext/Source/AnimNext/Public/Context.h
@@ -0,0 +1,39 @@
+// Copyright Epic Games, Inc. All Rights Reserved.
+
+#pragma once
+
+#include "CoreMinimal.h"
+#include "UObject/Class.h"
+#include "Param/ParamType.h"
+#include "Param/ParamId.h"
+
+namespace UE::AnimNext
+{
+
+struct FContext;
+struct FParamStack;
+struct FParamStackLayer;
+
+// Context providing methods for mutating & interrogating the anim interface runtime
+struct ANIMNEXT_API FContext
+{
+public:
+	FContext();
+	explicit FContext(float InDeltaTime);
+	FContext(const FContext& Other) = delete;
+	FContext& operator=(const FContext&) = delete;
+	FContext(FContext&& Other) = delete;
+	FContext& operator=(FContext&& Other) = delete;
+
+	// Access the parameter stack, cached on construction
+	FParamStack& GetMutableParamStack() const { return Stack; }
+	const FParamStack& GetParamStack() const { return Stack; }
+
+	float GetDeltaTime() const { return DeltaTime; }
+
+private:
+	FParamStack& Stack;
+	float DeltaTime = 0.0f;
+};
+
+}
diff --git a/Engine/Plugins/Experimental/AnimNext/Source/AnimNext/Public/DataRegistry.h b/Engine/Plugins/Experimental/AnimNext/Source/AnimNext/Public/DataRegistry.h
--- a/Engine/Plugins/Experimental/AnimNext/Source/AnimNext/Public/DataRegistry.h
+++ b/Engine/Plugins/Experimental/AnimNext/Source/AnimNext/Public/DataRegistry.h
@@ -5,7 +5,6 @@
 #include "CoreMinimal.h"
 #include "DataRegistryTypes.h"
 #include "UObject/WeakObjectPtr.h"
-#include "Param/Param.h"
 
 class USkeletalMeshComponent;
 struct FReferenceSkeleton;
diff --git a/Engine/Plugins/Experimental/AnimNext/Source/AnimNext/Public/Interface/AnimNextInterface.h b/Engine/Plugins/Experimental/AnimNext/Source/AnimNext/Public/Interface/AnimNextInterface.h
deleted file mode 100644
diff --git a/Engine/Plugins/Experimental/AnimNext/Source/AnimNext/Public/Interface/IAnimNextInterface.h b/Engine/Plugins/Experimental/AnimNext/Source/AnimNext/Public/Interface/IAnimNextInterface.h
deleted file mode 100644
diff --git a/Engine/Plugins/Experimental/AnimNext/Source/AnimNext/Public/Interface/InterfaceContext.h b/Engine/Plugins/Experimental/AnimNext/Source/AnimNext/Public/Interface/InterfaceContext.h
deleted file mode 100644
diff --git a/Engine/Plugins/Experimental/AnimNext/Source/AnimNext/Public/Interface/InterfaceKey.h b/Engine/Plugins/Experimental/AnimNext/Source/AnimNext/Public/Interface/InterfaceKey.h
deleted file mode 100644
diff --git a/Engine/Plugins/Experimental/AnimNext/Source/AnimNext/Public/Interface/InterfaceState.h b/Engine/Plugins/Experimental/AnimNext/Source/AnimNext/Public/Interface/InterfaceState.h
deleted file mode 100644
diff --git a/Engine/Plugins/Experimental/AnimNext/Source/AnimNext/Public/Interface/InterfaceTypes.h b/Engine/Plugins/Experimental/AnimNext/Source/AnimNext/Public/Interface/InterfaceTypes.h
deleted file mode 100644
diff --git a/Engine/Plugins/Experimental/AnimNext/Source/AnimNext/Public/Param/IAnimNextParamInterface.h b/Engine/Plugins/Experimental/AnimNext/Source/AnimNext/Public/Param/IAnimNextParamInterface.h
deleted file mode 100644
diff --git a/Engine/Plugins/Experimental/AnimNext/Source/AnimNext/Public/Param/Param.h b/Engine/Plugins/Experimental/AnimNext/Source/AnimNext/Public/Param/Param.h
deleted file mode 100644
diff --git a/Engine/Plugins/Experimental/AnimNext/Source/AnimNext/Public/Param/ParamHelpers.h b/Engine/Plugins/Experimental/AnimNext/Source/AnimNext/Public/Param/ParamHelpers.h
deleted file mode 100644
diff --git a/Engine/Plugins/Experimental/AnimNext/Source/AnimNext/Public/Param/ParamId.h b/Engine/Plugins/Experimental/AnimNext/Source/AnimNext/Public/Param/ParamId.h
new file mode 100644
--- /dev/null
+++ b/Engine/Plugins/Experimental/AnimNext/Source/AnimNext/Public/Param/ParamId.h
@@ -0,0 +1,68 @@
+// Copyright Epic Games, Inc. All Rights Reserved.
+
+#pragma once
+
+#include "CoreMinimal.h"
+
+struct FRigVMDispatch_SetParameter;
+struct FRigVMDispatch_GetParameter;
+
+namespace UE::AnimNext::Tests
+{
+	class FParamStackTest;
+}
+
+namespace UE::AnimNext
+{
+
+// Global identifier used to index into dense parameter arrays
+struct ANIMNEXT_API FParamId
+{
+	friend struct FParamStack;
+	friend class Tests::FParamStackTest;
+	friend struct ::FRigVMDispatch_SetParameter;
+	friend struct ::FRigVMDispatch_GetParameter;
+
+	static constexpr uint32 InvalidIndex = MAX_uint32;
+
+	FParamId() = delete;
+	FParamId(const FParamId& InName) = default;
+	FParamId& operator=(const FParamId& InName) = default;
+
+
+	// Make a parameter ID from an FName
+	explicit FParamId(FName InName);
+
+	// Get the index of this param
+	uint32 ToInt() const
+	{
+		return ParameterIndex;
+	}
+
+	// Get the name that this parameter was created from
+	FName ToName() const;
+
+private:
+	// Make a parameter ID from an index (for internal use)
+	explicit FParamId(uint32 InParameterIndex)
+		: ParameterIndex(InParameterIndex)
+	{
+	}
+
+	// Get the maximum parameter ID that can exist at present
+	// Note that this can change via concurrent modifications when new parameters are created by 
+	// different threads
+	static FParamId GetMaxParamId();
+
+#if WITH_DEV_AUTOMATION_TESTS
+	// Used to isolate global param IDs from automated tests
+	static void BeginTestSandbox();
+	static void EndTestSandbox();
+#endif
+
+private:
+	// Stable index
+	uint32 ParameterIndex = InvalidIndex;
+};
+
+} // end namespace UE::AnimNext
diff --git a/Engine/Plugins/Experimental/AnimNext/Source/AnimNext/Public/Param/ParamStack.h b/Engine/Plugins/Experimental/AnimNext/Source/AnimNext/Public/Param/ParamStack.h
new file mode 100644
--- /dev/null
+++ b/Engine/Plugins/Experimental/AnimNext/Source/AnimNext/Public/Param/ParamStack.h
@@ -0,0 +1,509 @@
+// Copyright Epic Games, Inc. All Rights Reserved.
+
+#pragma once
+
+#include "UObject/NameTypes.h"
+#include "Containers/ContainersFwd.h"
+#include "Containers/Array.h"
+#include "Containers/Map.h"
+#include "Containers/PagedArray.h"
+#include "Param/ParamId.h"
+#include "Param/ParamTypeHandle.h"
+#include "Misc/MemStack.h"
+
+struct FInstancedPropertyBag;
+
+namespace UE::AnimNext::Tests
+{
+	class FParamStackTest;
+}
+
+namespace UE::AnimNext
+{
+
+struct FParamStackLayer;
+
+// Stack of parameter layers.
+// Acts as an associative container - allows retrieval of parameter calues (by ID) that have been
+// pushed onto the stack in 'layers'.
+// Parameter values (and notably types) on higher stack layers override those on lower layers. 
+// Older values and types are restored when a layer is popped.
+// Parameter values, once pushed, can be overridden by subsequent layers, but they cannot be 
+// removed until the layer that initially introduces the parameter is popped.
+struct FParamStack
+{
+	friend class Tests::FParamStackTest;
+	friend struct FParamStackLayer;
+
+private:
+	enum class EParamFlags : uint8
+	{
+		None = 0,				// No flags
+		Mutable = 1 << 0,		// Parameter is mutable, so can be mutated after it has been created
+		Embedded = 1 << 1,		// Parameter will be stored as a Value, but stored directly on the Data pointer
+		Reference = 1 << 2,		// Parameter is a reference, so will not be embdedded even if it fits in the Data pointer
+	};
+
+	FRIEND_ENUM_CLASS_FLAGS(EParamFlags);
+
+	// Parameter memory wrapper
+	struct  FParam
+	{
+	public:
+		friend struct FParamStack;
+		friend struct FParamStackLayer;
+
+		FParam() = default;
+		ANIMNEXT_API FParam(const FParam& InOtherParam);
+		ANIMNEXT_API FParam& operator=(const FParam& InOtherParam);
+		ANIMNEXT_API FParam(FParam&& InOtherParam);
+		ANIMNEXT_API FParam& operator=(FParam&& InOtherParam);
+		ANIMNEXT_API ~FParam();
+
+	private:
+		ANIMNEXT_API FParam(const FParamTypeHandle& InTypeHandle, TArrayView<uint8> InData, bool bInIsReference, bool bInIsMutable);
+
+		// Get the type handle of this param 
+		FParamTypeHandle GetTypeHandle() const { return TypeHandle; }
+
+		// Check whether this parameter is able to be mutated
+		bool IsMutable() const { return EnumHasAnyFlags(Flags, EParamFlags::Mutable); }
+
+		// Check whether this parameter is of reference type (i.e. the parameter refers to user data owned outside of the parameter)
+		bool IsReference() const { return EnumHasAnyFlags(Flags, EParamFlags::Reference); }
+
+		// Check whether this parameter is embedded in the parameter (stored internally rather than as a ptr)
+		bool IsEmbedded() const { return EnumHasAnyFlags(Flags, EParamFlags::Embedded); }
+
+		// Get an immutable view of the parameter's data
+		TConstArrayView<uint8> GetData() const
+		{
+			if (EnumHasAnyFlags(Flags, EParamFlags::Embedded))
+			{
+				return TConstArrayView<uint8>(reinterpret_cast<const uint8*>(&Data), Size);
+			}
+			else
+			{
+				return TConstArrayView<uint8>(static_cast<const uint8*>(Data), Size);
+			}
+		}
+
+		// Get an mutable view of the parameter's data, returns an empty array view if this parameter is immutable
+		TArrayView<uint8> GetMutableData()
+		{
+			if (IsMutable())
+			{
+				if (EnumHasAnyFlags(Flags, EParamFlags::Embedded))
+				{
+					return TArrayView<uint8>(reinterpret_cast<uint8*>(&Data), Size);
+				}
+				else
+				{
+					return TArrayView<uint8>(static_cast<uint8*>(Data), Size);
+				}
+			}
+			return TArrayView<uint8>();
+		}
+
+		// Check whether this param represents valid data. Note that this doesnt check the type hande for validity.
+		bool IsValid() const
+		{
+			return Size > 0;
+		}
+
+		// Raw ptr to the data, or the data itself if we have EFlags::Embedded
+		void* Data = nullptr;
+
+		// The type of the param
+		FParamTypeHandle TypeHandle;
+
+		// Size of the data
+		uint16 Size = 0;
+
+		// Internal flags
+		EParamFlags Flags = EParamFlags::None;
+	};
+
+public:
+	ANIMNEXT_API FParamStack();
+
+	// Get the param stack for this thread
+	static ANIMNEXT_API FParamStack& Get();
+
+	// Push a value as a parameter layer
+	// @param	InParamId			Parameter ID for the parameter
+	// @param	InValue				Value to push
+	template<typename... ValueType>
+	void PushValue(FParamId InParamId, ValueType&&... InValue)
+	{
+		PushValues(InParamId, Forward<ValueType>(InValue)...);
+	}
+
+	// Push a value as a parameter layer
+	// @param	InParamId			Parameter name for the parameter
+	// @param	InValue				Value to push
+	template<typename... ValueType>
+	void PushValue(FName InName, ValueType&&... InValue)
+	{
+		PushValues(FParamId(InName), Forward<ValueType>(InValue)...);
+	}
+
+	// Push an interleaved parameter list of keys and values as a parameter layer
+	// @param	InValues	Values to push
+	template<typename... Args>
+	void PushValues(Args&&... InValues)
+	{
+		constexpr int32 NumItems = sizeof...(InValues) / 2;
+		TArray<TPair<FParamId, FParam>, TInlineAllocator<NumItems>> ParamIdValues;
+		ParamIdValues.Reserve(NumItems);
+		PushValuesHelper(ParamIdValues, Forward<Args>(InValues)...);
+	}
+
+	// Push a layer
+	ANIMNEXT_API void PushLayer(FParamStackLayer& InLayer);
+
+	// Push an internally-owned layer. Copies parameter data to internal storage.
+	ANIMNEXT_API void PushLayer(TConstArrayView<TPair<FParamId, FParam>> InParams);
+
+	// Create a cached parameter layer from an instanced property bag
+	static ANIMNEXT_API TUniquePtr<FParamStackLayer> MakeLayer(const FInstancedPropertyBag& InInstancedPropertyBag);
+
+	// Create a cached parameter layer from set of params
+	static ANIMNEXT_API TUniquePtr<FParamStackLayer> MakeLayer(TConstArrayView<TPair<FParamId, FParam>> InParams);
+
+	// Make a parameter layer from a value
+	// @param	InParamId			Parameter ID for the parameter
+	// @param	InValue				Value to push
+	template<typename... ValueType>
+	static TUniquePtr<FParamStackLayer> MakeValueLayer(FParamId InParamId, ValueType&&... InValue)
+	{
+		return MakeValuesLayer(InParamId, Forward<ValueType>(InValue)...);
+	}
+
+	// Make a parameter layer from a value
+	// @param	InParamId			Parameter name for the parameter
+	// @param	InValue				Value to push
+	template<typename... ValueType>
+	static TUniquePtr<FParamStackLayer> MakeValueLayer(FName InName, ValueType&&... InValue)
+	{
+		return MakeValuesLayer(FParamId(InName), Forward<ValueType>(InValue)...);
+	}
+
+	// Make a parameter layer from an interleaved parameter list of keys and values
+	// @param	InValues	Values to push
+	template<typename... Args>
+	static TUniquePtr<FParamStackLayer> MakeValuesLayer(Args&&... InValues)
+	{
+		constexpr int32 NumItems = sizeof...(InValues) / 2;
+		TArray<TPair<FParamId, FParam>, TInlineAllocator<NumItems>> ParamIdValues;
+		ParamIdValues.Reserve(NumItems);
+		return MakeValuesLayerHelper(ParamIdValues, Forward<Args>(InValues)...);
+	}
+
+	// Pop a parameter layer
+	ANIMNEXT_API void PopLayer();
+
+	// Results for FParamStack::GetParam
+	enum class EGetParamResult : uint8
+	{
+		// The requested parameter is present in the current scope
+		Succeeded		= 0x00,
+
+		// The requested parameter is not present in the current scope
+		NotInScope		= 0x01,
+
+		// The requested parameter is present in the current scope, but is of a different type to the requested type
+		IncorrectType	= 0x02,
+
+		// The requested parameter is immutable but a mutable request was made of it
+		Immutable		= 0x04,
+	};
+
+	FRIEND_ENUM_CLASS_FLAGS(FParamStack::EGetParamResult);
+
+	// Get a pointer to a parameter's value given a FParamId.
+	// @param	InParamId			Parameter ID to find the currently-pushed value for
+	// @param	OutResult			Optional ptr to aresult that will be filled if an error occurs
+	// @return nullptr if the parameter is not mapped
+	template<typename ValueType>
+	const ValueType* GetParamPtr(FParamId InParamId, FParamStack::EGetParamResult* OutResult = nullptr) const
+	{
+		TConstArrayView<uint8> Data;
+		FParamStack::EGetParamResult Result = GetParamData(InParamId, FParamTypeHandle::GetHandle<ValueType>(), Data);
+		if (OutResult)
+		{
+			*OutResult = Result;
+		}
+		return reinterpret_cast<const ValueType*>(Data.GetData());
+	}
+
+	// Get a pointer to a parameter's mutable value given a FParamId.
+	// @param	InParamId			Parameter ID to find the currently-pushed value for
+	// @param	OutResult			Optional ptr to a result that will be filled if an error occurs
+	// @return nullptr if the parameter is not mapped or mutable
+	template<typename ValueType>
+	ValueType* GetMutableParamPtr(FParamId InParamId, FParamStack::EGetParamResult* OutResult = nullptr)
+	{
+		TArrayView<uint8> Data;
+		FParamStack::EGetParamResult Result = GetMutableParamData(InParamId, FParamTypeHandle::GetHandle<ValueType>(), Data);
+		if (OutResult)
+		{
+			*OutResult = Result;
+		}
+		return reinterpret_cast<ValueType*>(Data.GetData());
+	}
+
+	// Get a pointer to a parameter's value given a FName.
+	// Using the FParamId overload will avoid having to hash the FName to recover the ID.
+	// @param	InKey				Parameter name to find the currently-pushed value for
+	// @param	OutResult			Optional ptr to a result that will be filled if an error occurs
+	// @return nullptr if the parameter is not mapped
+	template<typename ValueType>
+	const ValueType* GetParamPtr(FName InKey, FParamStack::EGetParamResult* OutResult = nullptr) const
+	{
+		return GetParamPtr<ValueType>(FParamId(InKey), OutResult);
+	}
+
+	// Get a pointer to a parameter's mutable value given a FName.
+	// Using the FParamId overload will avoid having to hash the FName to recover the ID.
+	// @param	InKey				Parameter name to find the currently-pushed value for
+	// @param	OutResult			Optional ptr to a result that will be filled if an error occurs
+	// @return nullptr if the parameter is not mapped
+	template<typename ValueType>
+	ValueType* GetMutableParamPtr(FName InKey, FParamStack::EGetParamResult* OutResult = nullptr)
+	{
+		return GetMutableParamPtr<ValueType>(FParamId(InKey), OutResult);
+	}
+
+	// Get the value of a parameter given a FParamId. Potentially performs a copy of the value.
+	// @param	InParamId			Parameter ID to find the currently-pushed value for
+	// @param	OutResult			Optional ptr to a result that will be filled if an error occurs
+	// @return the parameter's value, or a default value if the parameter is not present
+	template<typename ValueType>
+	ValueType GetParamValue(FParamId InParamId, FParamStack::EGetParamResult* OutResult = nullptr) const
+	{
+		if (const ValueType* ParamPtr = GetParamPtr<ValueType>(InParamId, OutResult))
+		{
+			return *ParamPtr;
+		}
+		return ValueType();
+	}
+
+	// Get the value of a parameter given a FName. Potentially performs a copy of the value.
+	// @param	InKey				Parameter name to find the currently-pushed value for
+	// @param	OutResult			Optional ptr to a result that will be filled if an error occurs
+	// @return the parameter's value, or a default value if the parameter is not present
+	template<typename ValueType>
+	ValueType GetParamValue(FName InKey, FParamStack::EGetParamResult* OutResult = nullptr) const
+	{
+		if (const ValueType* ParamPtr = GetParamPtr<ValueType>(FParamId(InKey), OutResult))
+		{
+			return *ParamPtr;
+		}
+		return ValueType();
+	}
+
+	// Get a const reference to the value of a parameter given a FParamId.
+	// @param	InParamId			Parameter ID to find the currently-pushed value for
+	// @param	OutResult			Optional ptr to a result that will be filled if an error occurs
+	// @return a const reference to the parameter's value. Function asserts if the parameter is not present
+	template<typename ValueType>
+	const ValueType& GetParam(FParamId InParamId, FParamStack::EGetParamResult* OutResult = nullptr) const
+	{
+		const ValueType* ParamPtr = GetParamPtr<ValueType>(InParamId, OutResult);
+		check(ParamPtr);
+		return *ParamPtr;
+	}
+
+	// Get a const reference to the value of a parameter given a FName.
+	// Using the FParamId overload will avoid having to hash the FName to recover the ID.
+	// @param	InKey				Parameter name to find the currently-pushed value for
+	// @param	OutResult			Optional ptr to a result that will be filled if an error occurs
+	// @return a const reference to the parameter's value. Function asserts if the parameter is not present
+	template<typename ValueType>
+	const ValueType& GetParam(FName InKey, FParamStack::EGetParamResult* OutResult = nullptr) const
+	{
+		const ValueType* ParamPtr = GetParamPtr<ValueType>(FParamId(InKey), OutResult);
+		check(ParamPtr);
+		return *ParamPtr;
+	}
+
+	// Get a mutable reference to the value of a parameter given a FParamId.
+	// @param	InParamId			Parameter ID to find the currently-pushed value for
+	// @param	OutResult			Optional ptr to a result that will be filled if an error occurs
+	// @return a mutable reference to the parameter's value. Function asserts if the parameter is not present
+	template<typename ValueType>
+	ValueType& GetMutableParam(FParamId InParamId, FParamStack::EGetParamResult* OutResult = nullptr)
+	{
+		ValueType* ParamPtr = GetMutableParamPtr<ValueType>(InParamId, OutResult);
+		check(ParamPtr);
+		return *ParamPtr;
+	}
+
+	// Get a mutable reference to the value of a parameter given a FName.
+	// Using the FParamId overload will avoid having to hash the FName to recover the ID.
+	// @param	InKey				Parameter name to find the currently-pushed value for
+	// @param	OutResult			Optional ptr to a result that will be filled if an error occurs
+	// @return a mutable reference to the parameter's value. Function asserts if the parameter is not present
+	template<typename ValueType>
+	ValueType& GetMutableParam(FName InKey, FParamStack::EGetParamResult* OutResult = nullptr)
+	{
+		ValueType* ParamPtr = GetMutableParamPtr<ValueType>(FParamId(InKey), OutResult);
+		check(ParamPtr);
+		return *ParamPtr;
+	}
+
+	// Get a parameter's data given an FParamId.
+	// @param	InParamId			Parameter ID to find the currently-pushed value for
+	// @param	InParamTypeHandle	Handle to the type of the parameter, which must match the stored type for a value to
+	//								be returned
+	// @param	OutParamData		View into the parameters data that will be filled
+	// @return an enum describing the result of the operation @see FParamStack::EGetParamResult
+	ANIMNEXT_API EGetParamResult GetParamData(FParamId InId, FParamTypeHandle InTypeHandle, TConstArrayView<uint8>& OutParamData) const;
+
+	// Get a parameter's mutable data given an FParamId.
+	// @param	InParamId			Parameter ID to find the currently-pushed value for
+	// @param	InParamTypeHandle	Handle to the type of the parameter, which must match the stored type for a value to
+	//								be returned
+	// @param	OutParamData		View into the parameters data that will be filled
+	// @return an enum describing the result of the operation @see FParamStack::EGetParamResult
+	ANIMNEXT_API EGetParamResult GetMutableParamData(FParamId InId, FParamTypeHandle InTypeHandle, TArrayView<uint8>& OutParamData);
+
+private:
+	// A layer held on the stack, may be owned by the stack or not
+	struct FPushedLayer
+	{
+		FPushedLayer(FParamStackLayer& InLayer, FParamStack& InStack);
+
+		FPushedLayer(const FPushedLayer& InPreviousLayer, FParamStackLayer& InLayer, FParamStack& InStack);
+
+		FParamStackLayer& Layer;
+
+		// Index offset for the previous layer for each active param.
+		// Offset into FParamStack::PreviousLayerIndices.
+		uint32 PreviousLayerIndexStart = 0;
+	};
+
+	// Recursive helper function for PushValues
+	template <uint32 NumItems, typename FirstType, typename SecondType, typename... OtherTypes>
+	void PushValuesHelper(TArray<TPair<FParamId, FParam>, TInlineAllocator<NumItems>>& InArray, FirstType&& InFirst, SecondType&& InSecond, OtherTypes&&... InOthers)
+	{
+		InArray.Emplace(FParamId(InFirst),
+			FParam(
+				FParamTypeHandle::GetHandle<std::remove_reference_t<SecondType>>(),
+				TArrayView<uint8>(const_cast<uint8*>(reinterpret_cast<const uint8*>(&InSecond)), sizeof(std::remove_reference_t<SecondType>)),
+				std::is_reference_v<SecondType>,
+				!std::is_const_v<std::remove_reference_t<SecondType>>
+			)
+		);
+
+		if constexpr (sizeof...(InOthers) > 0)
+		{
+			PushValuesHelper(InArray, Forward<OtherTypes>(InOthers)...);
+		}
+		else
+		{
+			PushLayer(InArray);
+		}
+	}
+
+	// Recursive helper function for MakeValuesLayer
+	template <uint32 NumItems, typename FirstType, typename SecondType, typename... OtherTypes>
+	static TUniquePtr<FParamStackLayer> MakeValuesLayerHelper(TArray<TPair<FParamId, FParam>, TInlineAllocator<NumItems>>& InArray, FirstType&& InFirst, SecondType&& InSecond, OtherTypes&&... InOthers)
+	{
+		InArray.Emplace(FParamId(InFirst), 
+			FParam(
+				FParamTypeHandle::GetHandle<std::remove_reference_t<SecondType>>(),
+				TArrayView<uint8>(const_cast<uint8*>(reinterpret_cast<const uint8*>(&InSecond)), sizeof(std::remove_reference_t<SecondType>)),
+				std::is_reference_v<SecondType>,
+				!std::is_const_v<std::remove_reference_t<SecondType>>
+			)
+		);
+
+		if constexpr (sizeof...(InOthers) > 0)
+		{
+			return MakeValuesLayerHelper(InArray, Forward<OtherTypes>(InOthers)...);
+		}
+		else
+		{
+			return MakeLayer(InArray);
+		}
+	}
+
+	// Check whether a given parameter is mutable
+	ANIMNEXT_API bool IsMutableParam(FParamId InId) const;
+
+	// Check whether a given parameter is mutable
+	bool IsMutableParam(FName InKey) const
+	{
+		return IsMutableParam(FParamId(InKey));
+	}
+
+	// Check whether a given parameter is a reference
+	ANIMNEXT_API bool IsReferenceParam(FParamId InId) const;
+
+	// Check whether a given parameter is a reference
+	bool IsReferenceParam(FName InKey) const
+	{
+		return IsReferenceParam(FParamId(InKey));
+	}
+
+	// Allocates any necessary owned storage for non-reference, non-embeded parameters that a layer holds
+	// @return the base offset of the storage area
+	uint32 AllocAndCopyOwnedParamStorage(TArrayView<FParam> InParams);
+
+	// Frees any owned storage, restores owned storage offset
+	void FreeOwnedParamStorage(uint32 InOffset);
+
+	// Resize layer indices to deal with any new params we have seen since the stack was created
+	void ResizeLayerIndices();
+
+	// Layer stack
+	TArray<FPushedLayer> Layers;
+
+	// Owned layers, paged for stable addresses as layers reference them by ptr
+	// Grows with each new pushed owned layer
+	TPagedArray<FParamStackLayer, 4096> OwnedStackLayers;
+
+	// Parameter storage for owned layers, paged for stable addresses as layer parameters reference them by ptr
+	// Grows with each new pushed layer, free'd when the param stack is destroyed
+	TPagedArray<uint8, 4096> OwnedLayerParamStorage;
+
+	// Indicies into layers. These are the indices into the layer stack where the current 'top' value resides.
+	// Grows on creation and when new parameter IDs are encountered
+	TArray<uint16> LayerIndices;
+
+	// Indices for the previous layer for each active param. MAX_uint16 if there is no previous layer for the param.
+	// Pushed layers hold views into this array.
+	// Grows with each new pushed layer, free'd when the param stack is destroyed
+	TArray<uint16> PreviousLayerIndices;
+};
+
+ENUM_CLASS_FLAGS(FParamStack::EGetParamResult);
+ENUM_CLASS_FLAGS(FParamStack::EParamFlags);
+
+// Stack layers are what actually get pushed/popped onto the layer stack.
+// They are designed to be held on an instance, their data updated in place.
+// Memory ownership for params is assumed to be outside of the stack and layers
+struct FParamStackLayer
+{
+private:
+	friend struct FParamStack;
+
+	FParamStackLayer() = delete;
+
+	explicit FParamStackLayer(const FInstancedPropertyBag& InPropertyBag);
+
+	explicit FParamStackLayer(TConstArrayView<TPair<FParamId, FParamStack::FParam>> InParams);
+
+	// Params that this layer supplies
+	TArray<FParamStack::FParam> Params;
+
+	// ID that the param indices start at. Maps the global param ID range into the range for this layer
+	uint32 MinParamId = 0;
+
+	// Storage offset for this layer if it is owned internally, otherwise MAX_uint32
+	uint32 OwnedStorageOffset = MAX_uint32;
+};
+
+}
diff --git a/Engine/Plugins/Experimental/AnimNext/Source/AnimNext/Public/Param/ParamStorage.h b/Engine/Plugins/Experimental/AnimNext/Source/AnimNext/Public/Param/ParamStorage.h
deleted file mode 100644
diff --git a/Engine/Plugins/Experimental/AnimNext/Source/AnimNext/Public/Param/ParamType.h b/Engine/Plugins/Experimental/AnimNext/Source/AnimNext/Public/Param/ParamType.h
--- a/Engine/Plugins/Experimental/AnimNext/Source/AnimNext/Public/Param/ParamType.h
+++ b/Engine/Plugins/Experimental/AnimNext/Source/AnimNext/Public/Param/ParamType.h
@@ -68,12 +68,12 @@ public:
 	FAnimNextParamType() = default;
 
 	/** Construct a parameter type from the passed in value, container and object type. */
-	FAnimNextParamType(EValueType InValueType, EContainerType InContainerType = EContainerType::None, UObject* InValueTypeObject = nullptr);
+	FAnimNextParamType(EValueType InValueType, EContainerType InContainerType = EContainerType::None, const UObject* InValueTypeObject = nullptr);
 
 private:
 	/** Pointer to object that defines the Enum, Struct, or Class. */
 	UPROPERTY()
-	TObjectPtr<UObject> ValueTypeObject = nullptr;
+	TObjectPtr<const UObject> ValueTypeObject = nullptr;
 
 	/** Type of the value described by this parameter. */
 	UPROPERTY()
@@ -215,7 +215,7 @@ public:
 	UE::AnimNext::FParamTypeHandle GetHandle() const;
 
 	/** Get the pointer to the object that defines the Enum, Struct, or Class. */
-	UObject* GetValueTypeObject() const
+	const UObject* GetValueTypeObject() const
 	{
 		return ValueTypeObject;
 	}
diff --git a/Engine/Plugins/Experimental/AnimNext/Source/AnimNext/Public/Param/ParamTypeHandle.h b/Engine/Plugins/Experimental/AnimNext/Source/AnimNext/Public/Param/ParamTypeHandle.h
--- a/Engine/Plugins/Experimental/AnimNext/Source/AnimNext/Public/Param/ParamTypeHandle.h
+++ b/Engine/Plugins/Experimental/AnimNext/Source/AnimNext/Public/Param/ParamTypeHandle.h
@@ -6,7 +6,10 @@
 #include "Concepts/BaseStructureProvider.h"
 #include "Param/ParamType.h"
 
-class FAnimationAnimNextParametersTest_Types;
+namespace UE::AnimNext::Tests
+{
+	class FParamTypesTest;
+}
 
 namespace UE::AnimNext
 {
@@ -49,7 +52,7 @@ struct ANIMNEXT_API FParamTypeHandle
 
 	friend struct ::FAnimNextParamType;
 	friend struct FParamHelpers;
-	friend class ::FAnimationAnimNextParametersTest_Types;
+	friend class UE::AnimNext::Tests::FParamTypesTest;
 
 private:
 	union
@@ -85,7 +88,7 @@ private:
 	/** Get the custom type index */
 	void SetCustomTypeIndex(uint32 InIndex)
 	{
-		checkf(InIndex < (1 << 24), TEXT("FTypeHandle::SetCustomTypeIndex: Type Index ouf of range"));
+		checkf(InIndex < (1 << 24), TEXT("FTypeHandle::SetCustomTypeIndex: Type Index out of range"));
 		Fields.CustomTypeIndex = InIndex;
 	}
 
@@ -93,7 +96,7 @@ private:
 	static void ResetCustomTypes();
 
 	/** Get a custom type index for the passed-in type information */
-	static uint32 GetOrAllocateCustomTypeIndex(FAnimNextParamType::EValueType InValueType, FAnimNextParamType::EContainerType InContainerType, UObject* InValueTypeObject);
+	static uint32 GetOrAllocateCustomTypeIndex(FAnimNextParamType::EValueType InValueType, FAnimNextParamType::EContainerType InContainerType, const UObject* InValueTypeObject);
 
 	/** Validate the passed-in custom type index */
 	static bool ValidateCustomTypeIndex(uint32 InCustomTypeIndex);
@@ -302,6 +305,26 @@ public:
 		return TypeHandle;
 	}
 
+	/** Make a parameter handle from a raw uint32 */
+	static FParamTypeHandle FromRaw(uint32 InRawHandle)
+	{
+		FParamTypeHandle Handle;
+		Handle.Value = InRawHandle;
+		return Handle;
+	}
+
+	/** Make a parameter type handle from a property bag property desc */
+	static FParamTypeHandle FromPropertyBagPropertyDesc(const FPropertyBagPropertyDesc& Desc);
+	
+	/** Make a parameter type handle from a FProperty */
+	static FParamTypeHandle FromProperty(const FProperty* InProperty);
+
+	/** Return the raw value that this handle uses to represent itself */
+	uint32 ToRaw() const
+	{
+		return Value;
+	}
+
 	/** Get the parameter type that this handle corresponds to */
 	FAnimNextParamType GetType() const;
 
diff --git a/Engine/Plugins/Experimental/AnimNext/Source/AnimNextEditor/Internal/AnimNextEditorModule.h b/Engine/Plugins/Experimental/AnimNext/Source/AnimNextEditor/Internal/AnimNextEditorModule.h
--- a/Engine/Plugins/Experimental/AnimNext/Source/AnimNextEditor/Internal/AnimNextEditorModule.h
+++ b/Engine/Plugins/Experimental/AnimNext/Source/AnimNextEditor/Internal/AnimNextEditorModule.h
@@ -11,8 +11,6 @@ namespace UE::AnimNext::Editor
 {
 
 class FAssetTypeActions_AnimNextGraph;
-class FAnimNextPropertyTypeCustomization;
-class FPropertyTypeIdentifier;
 struct FParameterPickerArgs;
 
 class FModule : public IModuleInterface
@@ -26,7 +24,6 @@ private:
 	virtual void ShutdownModule() override;
 
 	TSharedPtr<FAssetTypeActions_AnimNextGraph> AssetTypeActions_AnimNextGraph;
-	TSharedPtr<FPropertyTypeIdentifier> AnimNextPropertyTypeIdentifier;
 };
 
 }
\ No newline at end of file
diff --git a/Engine/Plugins/Experimental/AnimNext/Source/AnimNextEditor/Private/AnimNextEditorModule.cpp b/Engine/Plugins/Experimental/AnimNext/Source/AnimNextEditor/Private/AnimNextEditorModule.cpp
--- a/Engine/Plugins/Experimental/AnimNext/Source/AnimNextEditor/Private/AnimNextEditorModule.cpp
+++ b/Engine/Plugins/Experimental/AnimNext/Source/AnimNextEditor/Private/AnimNextEditorModule.cpp
@@ -4,7 +4,6 @@
 #include "AssetToolsModule.h"
 #include "IAssetTools.h"
 #include "Graph/AssetTypeActions.h"
-#include "Graph/PropertyTypeCustomization.h"
 #include "Param/ParamTypePropertyCustomization.h"
 #include "Param/ParameterPickerArgs.h"
 #include "Param/SParameterPicker.h"
@@ -19,12 +18,7 @@ void FModule::StartupModule()
 	AssetTypeActions_AnimNextGraph = MakeShared<FAssetTypeActions_AnimNextGraph>();
 	AssetTools.RegisterAssetTypeActions(AssetTypeActions_AnimNextGraph.ToSharedRef());
 
-	AnimNextPropertyTypeIdentifier = MakeShared<FPropertyTypeIdentifier>();
 	FPropertyEditorModule& PropertyModule = FModuleManager::LoadModuleChecked<FPropertyEditorModule>("PropertyEditor");
-	PropertyModule.RegisterCustomPropertyTypeLayout(
-		"InterfaceProperty",
-		FOnGetPropertyTypeCustomizationInstance::CreateLambda([] { return MakeShared<FPropertyTypeCustomization>(); }),
-		AnimNextPropertyTypeIdentifier);
 
 	PropertyModule.RegisterCustomPropertyTypeLayout(
 		"AnimNextParamType",
@@ -42,7 +36,6 @@ void FModule::ShutdownModule()
 	if(FModuleManager::Get().IsModuleLoaded("PropertyEditor"))
 	{
 		FPropertyEditorModule& PropertyModule = FModuleManager::LoadModuleChecked<FPropertyEditorModule>("PropertyEditor");
-		PropertyModule.UnregisterCustomPropertyTypeLayout("InterfaceProperty", AnimNextPropertyTypeIdentifier);
 		PropertyModule.UnregisterCustomPropertyTypeLayout("AnimNextParamType");
 	}
 }
diff --git a/Engine/Plugins/Experimental/AnimNext/Source/AnimNextEditor/Private/Graph/AnimNextGraphEditor.cpp b/Engine/Plugins/Experimental/AnimNext/Source/AnimNextEditor/Private/Graph/AnimNextGraphEditor.cpp
--- a/Engine/Plugins/Experimental/AnimNext/Source/AnimNextEditor/Private/Graph/AnimNextGraphEditor.cpp
+++ b/Engine/Plugins/Experimental/AnimNext/Source/AnimNextEditor/Private/Graph/AnimNextGraphEditor.cpp
@@ -61,9 +61,7 @@ void FGraphEditor::InitEditor(const EToolkitMode::Type InMode, const TSharedPtr<
 	RegenerateMenusAndToolbars();
 
 	// Open initial document
-	//DocumentManager->OpenDocument(FTabPayload_UObject::Make(AnimNextInterfaceGraph_EditorData->EntryPointGraph), FDocumentTracker::EOpenDocumentCause::OpenNewDocument);
 	DocumentManager->OpenDocument(FTabPayload_UObject::Make(AnimNextGraph_EditorData->RootGraph), FDocumentTracker::EOpenDocumentCause::OpenNewDocument);
-	//DocumentManager->OpenDocument(FTabPayload_UObject::Make(AnimNextInterfaceGraph_EditorData->FunctionLibraryEdGraph), FDocumentTracker::EOpenDocumentCause::OpenNewDocument);
 }
 
 void FGraphEditor::BindCommands()
@@ -94,17 +92,17 @@ void FGraphEditor::UnregisterTabSpawners(const TSharedRef<FTabManager>& InTabMan
 
 FName FGraphEditor::GetToolkitFName() const
 {
-	return FName("AnimNextInterfaceGraphEditor");
+	return FName("AnimNextGraphEditor");
 }
 
 FText FGraphEditor::GetBaseToolkitName() const
 {
-	return LOCTEXT("AppLabel", "AnimNextInterfaceGraphEditor");
+	return LOCTEXT("AppLabel", "AnimNextGraphEditor");
 }
 
 FString FGraphEditor::GetWorldCentricTabPrefix() const
 {
-	return LOCTEXT("WorldCentricTabPrefix", "AnimNextInterfaceGraphEditor ").ToString();
+	return LOCTEXT("WorldCentricTabPrefix", "AnimNextGraphEditor ").ToString();
 }
 
 FLinearColor FGraphEditor::GetWorldCentricTabColorScale() const
diff --git a/Engine/Plugins/Experimental/AnimNext/Source/AnimNextEditor/Private/Graph/AnimNextGraphEditor.h b/Engine/Plugins/Experimental/AnimNext/Source/AnimNextEditor/Private/Graph/AnimNextGraphEditor.h
--- a/Engine/Plugins/Experimental/AnimNext/Source/AnimNextEditor/Private/Graph/AnimNextGraphEditor.h
+++ b/Engine/Plugins/Experimental/AnimNext/Source/AnimNextEditor/Private/Graph/AnimNextGraphEditor.h
@@ -36,7 +36,7 @@ class FGraphEditor : public FWorkflowCentricApplication
 {
 public:
 	/** Edits the specified asset */
-	void InitEditor(const EToolkitMode::Type InMode, const TSharedPtr<IToolkitHost>& InInitToolkitHost, UAnimNextGraph* InAnimNextInterfaceGraph);
+	void InitEditor(const EToolkitMode::Type InMode, const TSharedPtr<IToolkitHost>& InInitToolkitHost, UAnimNextGraph* InAnimNextGraph);
 
 private:
 	friend class FGraphEditorMode;
diff --git a/Engine/Plugins/Experimental/AnimNext/Source/AnimNextEditor/Private/Graph/ClassFilter.h b/Engine/Plugins/Experimental/AnimNext/Source/AnimNextEditor/Private/Graph/ClassFilter.h
deleted file mode 100644
diff --git a/Engine/Plugins/Experimental/AnimNext/Source/AnimNextEditor/Private/Graph/GraphEditorSchemaActions.cpp b/Engine/Plugins/Experimental/AnimNext/Source/AnimNextEditor/Private/Graph/GraphEditorSchemaActions.cpp
--- a/Engine/Plugins/Experimental/AnimNext/Source/AnimNextEditor/Private/Graph/GraphEditorSchemaActions.cpp
+++ b/Engine/Plugins/Experimental/AnimNext/Source/AnimNextEditor/Private/Graph/GraphEditorSchemaActions.cpp
@@ -57,3 +57,46 @@ UEdGraphNode* FAnimNextGraphSchemaAction_RigUnit::PerformAction(UEdGraph* Parent
 
 	return NewNode;
 }
+
+
+UEdGraphNode* FAnimNextGraphSchemaAction_DispatchFactory::PerformAction(UEdGraph* ParentGraph, TArray<UEdGraphPin*>& FromPins, const FVector2D Location, bool bSelectNewNode)
+{
+	UAnimNextGraph_EditorData* EditorData = ParentGraph->GetTypedOuter<UAnimNextGraph_EditorData>();
+	UAnimNextGraph_EdGraphNode* NewNode = nullptr;
+	UAnimNextGraph_EdGraph* EdGraph = Cast<UAnimNextGraph_EdGraph>(ParentGraph);
+
+	if (EditorData != nullptr && EdGraph != nullptr)
+	{
+		const FRigVMTemplate* Template = FRigVMRegistry::Get().FindTemplate(Notation);
+		if (Template == nullptr)
+		{
+			return nullptr;
+		}
+
+		const int32 NotationHash = (int32)GetTypeHash(Notation);
+		const FString TemplateName = TEXT("RigVMTemplate_") + FString::FromInt(NotationHash);
+
+		FName Name = UE::AnimNext::Editor::FUtils::ValidateName(EditorData, Template->GetName().ToString());
+		URigVMController* Controller = EditorData->GetRigVMClient()->GetController(ParentGraph);
+
+		Controller->OpenUndoBracket(FString::Printf(TEXT("Add '%s' Node"), *Name.ToString()));
+
+		if (URigVMTemplateNode* ModelNode = Controller->AddTemplateNode(Notation, Location, Name.ToString(), true, false))
+		{
+			NewNode = Cast<UAnimNextGraph_EdGraphNode>(EdGraph->FindNodeForModelNodeName(ModelNode->GetFName()));
+
+			if (NewNode)
+			{
+				Controller->ClearNodeSelection(true);
+				Controller->SelectNode(ModelNode, true, true);
+			}
+
+			Controller->CloseUndoBracket();
+		}
+		else
+		{
+			Controller->CancelUndoBracket();
+		}
+	}
+	return NewNode;
+}
\ No newline at end of file
diff --git a/Engine/Plugins/Experimental/AnimNext/Source/AnimNextEditor/Private/Graph/GraphEditorSchemaActions.h b/Engine/Plugins/Experimental/AnimNext/Source/AnimNextEditor/Private/Graph/GraphEditorSchemaActions.h
--- a/Engine/Plugins/Experimental/AnimNext/Source/AnimNextEditor/Private/Graph/GraphEditorSchemaActions.h
+++ b/Engine/Plugins/Experimental/AnimNext/Source/AnimNextEditor/Private/Graph/GraphEditorSchemaActions.h
@@ -53,3 +53,25 @@ private:
 	// The script struct for our rig unit
 	UScriptStruct* StructTemplate = nullptr;
 };
+
+
+USTRUCT()
+struct FAnimNextGraphSchemaAction_DispatchFactory : public FAnimNextGraphSchemaAction
+{
+	GENERATED_BODY()
+
+	FAnimNextGraphSchemaAction_DispatchFactory() = default;
+
+	FAnimNextGraphSchemaAction_DispatchFactory(FName InNotation, FText InNodeCategory, FText InMenuDesc, FText InToolTip, FText InKeywords = FText::GetEmpty())
+		: FAnimNextGraphSchemaAction(InNodeCategory, InMenuDesc, InToolTip, InKeywords)
+		, Notation(InNotation)
+	{}
+
+	// FEdGraphSchemaAction Interface
+	virtual UEdGraphNode* PerformAction(UEdGraph* ParentGraph, TArray<UEdGraphPin*>& FromPins, const FVector2D Location, bool bSelectNewNode = true) override;
+	virtual UEdGraphNode* PerformAction(UEdGraph* ParentGraph, UEdGraphPin* FromPin, const FVector2D Location, bool bSelectNewNode = true) { return nullptr; }
+
+private:
+	// Notation for dispatch factory
+	FName Notation;
+};
diff --git a/Engine/Plugins/Experimental/AnimNext/Source/AnimNextEditor/Private/Graph/GraphFactory.cpp b/Engine/Plugins/Experimental/AnimNext/Source/AnimNextEditor/Private/Graph/GraphFactory.cpp
--- a/Engine/Plugins/Experimental/AnimNext/Source/AnimNextEditor/Private/Graph/GraphFactory.cpp
+++ b/Engine/Plugins/Experimental/AnimNext/Source/AnimNextEditor/Private/Graph/GraphFactory.cpp
@@ -74,7 +74,7 @@ UObject* UAnimNextGraphFactory::FactoryCreateNew(UClass* Class, UObject* InParen
 
 	//// Add end-execution unit of the correct type
 	//// TODO: using a float for now, but need to use a registry to determine correct type
-	//URigVMUnitNode* MainExitPointNode = Controller->AddUnitNode(FRigUnit_AnimNextInterfaceEndExecution_Float::StaticStruct(), FRigUnit::GetMethodName(), FVector2D(400.0f, 0.0f), FString(), false);
+	//URigVMUnitNode* MainExitPointNode = Controller->AddUnitNode(FRigUnit_AnimNextEndExecution_Float::StaticStruct(), FRigUnit::GetMethodName(), FVector2D(400.0f, 0.0f), FString(), false);
 	//URigVMPin* EndExecutePin = MainExitPointNode->FindPin(GET_MEMBER_NAME_STRING_CHECKED(FRigUnit_AnimNextEndExecution, ExecuteContext));
 	//check(EndExecutePin);
 	//check(EndExecutePin->GetDirection() == ERigVMPinDirection::Input);
diff --git a/Engine/Plugins/Experimental/AnimNext/Source/AnimNextEditor/Private/Graph/PropertyTypeCustomization.cpp b/Engine/Plugins/Experimental/AnimNext/Source/AnimNextEditor/Private/Graph/PropertyTypeCustomization.cpp
deleted file mode 100644
diff --git a/Engine/Plugins/Experimental/AnimNext/Source/AnimNextEditor/Private/Graph/PropertyTypeCustomization.h b/Engine/Plugins/Experimental/AnimNext/Source/AnimNextEditor/Private/Graph/PropertyTypeCustomization.h
deleted file mode 100644
diff --git a/Engine/Plugins/Experimental/AnimNext/Source/AnimNextEditor/Private/Graph/SActionMenu.cpp b/Engine/Plugins/Experimental/AnimNext/Source/AnimNextEditor/Private/Graph/SActionMenu.cpp
--- a/Engine/Plugins/Experimental/AnimNext/Source/AnimNextEditor/Private/Graph/SActionMenu.cpp
+++ b/Engine/Plugins/Experimental/AnimNext/Source/AnimNextEditor/Private/Graph/SActionMenu.cpp
@@ -30,12 +30,8 @@ static void CollectAllAnimNextGraphActions(FGraphContextMenuBuilder& MenuBuilder
 		FAnimNextGraphExecuteContext::StaticStruct()
 	};
 
-	const TChunkedArray<FRigVMFunction>& Functions = FRigVMRegistry::Get().GetFunctions();
-	const int32 NumFunctions = Functions.Num();
-	for(int i=0; i < NumFunctions; i++)
+	for(const FRigVMFunction& Function : FRigVMRegistry::Get().GetFunctions())
 	{
-		const FRigVMFunction& Function = Functions[i];
-
 		const UScriptStruct* FunctionContext = Function.GetExecuteContextStruct();
 		if (FunctionContext == nullptr || !AllowedExecuteContexts.Contains(FunctionContext))
 		{
@@ -66,6 +62,26 @@ static void CollectAllAnimNextGraphActions(FGraphContextMenuBuilder& MenuBuilder
 		}
 
 		MenuBuilder.AddAction(MakeShared<FAnimNextGraphSchemaAction_RigUnit>(Struct, NodeCategory, MenuDesc, ToolTip));
+	}
+
+	for (const FRigVMDispatchFactory* Factory : FRigVMRegistry::Get().GetFactories())
+	{
+		if (!Factory->SupportsExecuteContextStruct(FRigVMExecuteContext::StaticStruct()) && !Factory->SupportsExecuteContextStruct(FAnimNextGraphExecuteContext::StaticStruct()))
+		{
+			continue;
+		}
+
+		const FRigVMTemplate* Template = Factory->GetTemplate();
+		if (Template == nullptr)
+		{
+			continue;
+		}
+
+		FText NodeCategory = FText::FromString(Factory->GetCategory());
+		FText MenuDesc = FText::FromString(Factory->GetNodeTitle(FRigVMTemplateTypeMap()));
+		FText ToolTip = Factory->GetNodeTooltip(FRigVMTemplateTypeMap());
+
+		MenuBuilder.AddAction(MakeShared<FAnimNextGraphSchemaAction_DispatchFactory>(Template->GetNotation(), NodeCategory, MenuDesc, ToolTip));
 	};
 }
 
@@ -108,8 +124,8 @@ void SActionMenu::Construct(const FArguments& InArgs)
 						.ToolTip(IDocumentation::Get()->CreateToolTip(
 							LOCTEXT("ActionMenuContextTextTooltip", "Describes the current context of the action list"),
 							nullptr,
-							TEXT("Shared/Editors/AnimNextInterfaceEditor"),
-							TEXT("AnimNextInterfaceActionMenuContextText")))
+							TEXT("Shared/Editors/AnimNextGraphEditor"),
+							TEXT("AnimNextActionMenuContextText")))
 						.WrapTextAt(280)
 					]
 				]
diff --git a/Engine/Plugins/Experimental/AnimNext/Source/AnimNextEditor/Private/Graph/WidgetFactories.cpp b/Engine/Plugins/Experimental/AnimNext/Source/AnimNextEditor/Private/Graph/WidgetFactories.cpp
deleted file mode 100644
diff --git a/Engine/Plugins/Experimental/AnimNext/Source/AnimNextEditor/Private/Graph/WidgetFactories.h b/Engine/Plugins/Experimental/AnimNext/Source/AnimNextEditor/Private/Graph/WidgetFactories.h
deleted file mode 100644
diff --git a/Engine/Plugins/Experimental/AnimNext/Source/AnimNextTestSuite/Private/AnimNextParametersTest.cpp b/Engine/Plugins/Experimental/AnimNext/Source/AnimNextTestSuite/Private/AnimNextParametersTest.cpp
--- a/Engine/Plugins/Experimental/AnimNext/Source/AnimNextTestSuite/Private/AnimNextParametersTest.cpp
+++ b/Engine/Plugins/Experimental/AnimNext/Source/AnimNextTestSuite/Private/AnimNextParametersTest.cpp
@@ -1,22 +1,25 @@
 ﻿// Copyright Epic Games, Inc. All Rights Reserved.
 
 #include "CoreMinimal.h"
+#include "Context.h"
 #include "Misc/AutomationTest.h"
 #include "Param/ParamType.h"
 #include "Param/ParamTypeHandle.h"
 #include "Animation/AnimSequence.h"
 #include "Async/ParallelFor.h"
+#include "Param/ParamStack.h"
 
 // AnimNext Parameters Tests
 
 #if WITH_DEV_AUTOMATION_TESTS
 
-IMPLEMENT_SIMPLE_AUTOMATION_TEST(FAnimationAnimNextParametersTest_Types, "Animation.AnimNext.Parameters.Types", EAutomationTestFlags::EditorContext | EAutomationTestFlags::EngineFilter)
-
-bool FAnimationAnimNextParametersTest_Types::RunTest(const FString& InParameters)
+namespace UE::AnimNext::Tests
 {
-	using namespace UE::AnimNext;
 
+IMPLEMENT_SIMPLE_AUTOMATION_TEST(FParamTypesTest, "Animation.AnimNext.Parameters.Types", EAutomationTestFlags::EditorContext | EAutomationTestFlags::EngineFilter)
+
+bool FParamTypesTest::RunTest(const FString& InParameters)
+{
 	// None is invalid
 	FAnimNextParamType ParameterTypeValueNone(FAnimNextParamType::EValueType::None); 
 	AddErrorIfFalse(!ParameterTypeValueNone.IsValid(), TEXT("Parameter type None is valid."));
@@ -285,4 +288,249 @@ bool FAnimationAnimNextParametersTest_Types::RunTest(const FString& InParameters
 	return true;
 }
 
+IMPLEMENT_SIMPLE_AUTOMATION_TEST(FParamStackTest, "Animation.AnimNext.Parameters.Stack", EAutomationTestFlags::EditorContext | EAutomationTestFlags::EngineFilter)
+
+bool FParamStackTest::RunTest(const FString& InParameters)
+{
+	FParamId::BeginTestSandbox();
+
+	FParamId ParamIds[] = { FParamId("Param0"), FParamId("Param1"), FParamId("Param2"), FParamId("Param3") };
+
+	// Check immediate value interfaces
+	{
+		FParamStack Stack;
+
+		Stack.PushValues("Param0", 1.0f);
+		Stack.PushValues("Param1", 2.0f);
+		Stack.PushValues("Param1", 3.0f, "Param2", true);
+		Stack.PushValues("Param1", 4.0f, "Param2", false);
+		Stack.PushValues("Param1", 5.0f, "Param2", false, "Param3", 3);
+		Stack.PopLayer();
+		Stack.PopLayer();
+		Stack.PopLayer();
+		Stack.PopLayer();
+		Stack.PopLayer();
+
+		Stack.PushValue(ParamIds[0], 1.0f);
+		Stack.PushValue("Param1", 1.0f);
+		Stack.PushValue(ParamIds[2], true);
+		Stack.PushValue(ParamIds[3], 3);
+		Stack.PopLayer();
+		Stack.PopLayer();
+		Stack.PopLayer();
+		Stack.PopLayer();
+	}
+
+	// Check const/mutable
+	{
+		FParamStack Stack;
+
+		const FDateTime Now = FDateTime::Now();
+		Stack.PushValue("Param0", Now);
+
+		FParamStack::EGetParamResult Result;
+		AddErrorIfFalse(Stack.GetParamPtr<FDateTime>("Param0", &Result) != nullptr, "Const ptr cannot be accessed");
+		AddErrorIfFalse(Result == FParamStack::EGetParamResult::Succeeded, "Result != FParamStack::EGetParamResult::Succeeded");
+		AddErrorIfFalse(Stack.GetMutableParamPtr<FDateTime>("Param0", &Result) == nullptr, "Mutable ptr can be accessed");
+		AddErrorIfFalse(Result == FParamStack::EGetParamResult::Immutable, "Result != FParamStack::EGetParamResult::Immutable");
+	}
+
+	// Check type variance
+	{
+		FParamStack Stack;
+
+		const FLinearColor Blue = FLinearColor::Blue;
+		Stack.PushValue("Param0", Blue);
+
+		FParamStack::EGetParamResult Result;
+		AddErrorIfFalse(Stack.GetParamPtr<FDateTime>("Param0", &Result) == nullptr, "Incorrectly typed access is non-null");
+		AddErrorIfFalse(Result == FParamStack::EGetParamResult::IncorrectType, "Result != FParamStack::EGetParamResult::IncorrectType");
+		AddErrorIfFalse(Stack.GetMutableParamPtr<FDateTime>("Param0", &Result) == nullptr, "Incorrectly typed mutable access is non-null");
+		AddErrorIfFalse(Result == (FParamStack::EGetParamResult::IncorrectType | FParamStack::EGetParamResult::Immutable), "Result != (FParamStack::EGetParamResult::IncorrectType | FParamStack::EGetParamResult::Immutable)");
+		AddErrorIfFalse(Stack.GetParamPtr<FLinearColor>("Param0", &Result) != nullptr, "Correctly typed access is null");
+		AddErrorIfFalse(Result == FParamStack::EGetParamResult::Succeeded, "Result != FParamStack::EGetParamResult::Succeeded");
+	}
+
+	// Check external layer interfaces
+	{
+		FParamStack Stack;
+
+		TUniquePtr<FParamStackLayer> Layer0 = FParamStack::MakeValueLayer("Param0", 1.0f);
+		TUniquePtr<FParamStackLayer> Layer1 = FParamStack::MakeValuesLayer("Param1", 2.0, "Param3", 5);
+		TUniquePtr<FParamStackLayer> Layer2 = FParamStack::MakeValueLayer(ParamIds[2], true);
+
+		Stack.PushLayer(*Layer0);
+		Stack.PushLayer(*Layer1);
+		Stack.PushLayer(*Layer2);
+
+		AddErrorIfFalse(Stack.GetParam<float>("Param0") == 1.0f, "Param0 != 1.0f");
+		AddErrorIfFalse(Stack.GetParam<double>("Param1") == 2.0, "Param1 != 2.0");
+		AddErrorIfFalse(Stack.GetParam<bool>("Param2") == true, "Param2 != true");
+		AddErrorIfFalse(Stack.GetParam<int32>("Param3") == 5, "Param3 != 5");
+	}
+
+	// Check basic push/pop of plain ol' values
+	{
+		FParamStack Stack;
+
+		Stack.PushValues("Param1", 1.0f);
+
+		AddErrorIfFalse(Stack.GetParam<float>("Param1") == 1.0f, "Param1 != 1.0f");
+
+		Stack.PushValues("Param0", 2.0f);
+
+		AddErrorIfFalse(Stack.GetParam<float>("Param0") == 2.0f, "Param0 != 2.0f");
+		AddErrorIfFalse(Stack.GetParam<float>("Param1") == 1.0f, "Param1 != 1.0f");
+	
+		Stack.PushValues("Param3", 3.0f, "Param0", 3.0f);
+
+		AddErrorIfFalse(Stack.GetParam<float>("Param0") == 3.0f, "Param0 != 3.0f");
+		AddErrorIfFalse(Stack.GetParam<float>("Param1") == 1.0f, "Param1 != 1.0f");
+		AddErrorIfFalse(Stack.GetParam<float>("Param3") == 3.0f, "Param3 != 3.0f");
+
+		Stack.PushValues("Param1", 4.0f, "Param2", 4.0f);
+
+		AddErrorIfFalse(Stack.GetParam<float>("Param0") == 3.0f, "Param0 != 3.0f");
+		AddErrorIfFalse(Stack.GetParam<float>("Param1") == 4.0f, "Param1 != 4.0f");
+		AddErrorIfFalse(Stack.GetParam<float>("Param2") == 4.0f, "Param2 != 4.0f");
+		AddErrorIfFalse(Stack.GetParam<float>("Param3") == 3.0f, "Param3 != 3.0f");
+
+		Stack.PopLayer();
+
+		AddErrorIfFalse(Stack.GetParam<float>("Param0") == 3.0f, "Param0 != 3.0f");
+		AddErrorIfFalse(Stack.GetParam<float>("Param1") == 1.0f, "Param1 != 1.0f");
+		AddErrorIfFalse(Stack.GetParam<float>("Param3") == 3.0f, "Param3 != 3.0f");
+
+		Stack.PushValues("Param3", 4.0f, "Param0", 4.0f);
+
+		AddErrorIfFalse(Stack.GetParam<float>("Param0") == 4.0f, "Param0 != 4.0f");
+		AddErrorIfFalse(Stack.GetParam<float>("Param1") == 1.0f, "Param1 != 1.0f");
+		AddErrorIfFalse(Stack.GetParam<float>("Param3") == 4.0f, "Param3 != 4.0f");
+
+		Stack.PopLayer();
+		Stack.PopLayer();
+		Stack.PopLayer();
+
+		AddErrorIfFalse(Stack.GetParam<float>("Param1") == 1.0f, "Param1 != 1.0f");
+
+		// Add some more parameters to check internal buffer resize logic
+		FParamId NewParamIds[] = { FParamId("Param4"), FParamId("Param5"), FParamId("Param6"), FParamId("Param7") };
+
+		Stack.PushValues("Param4", 1.0f, "Param5", 1.0f, "Param7", 1.0f);
+
+		AddErrorIfFalse(Stack.GetParam<float>("Param4") == 1.0f, "Param4 != 1.0f");
+		AddErrorIfFalse(Stack.GetParam<float>("Param5") == 1.0f, "Param5 != 1.0f");
+		AddErrorIfFalse(Stack.GetParam<float>("Param7") == 1.0f, "Param7 != 1.0f");
+
+		Stack.PushValues("Param6", 2.0f);
+
+		AddErrorIfFalse(Stack.GetParam<float>("Param4") == 1.0f, "Param4 != 1.0f");
+		AddErrorIfFalse(Stack.GetParam<float>("Param5") == 1.0f, "Param5 != 1.0f");
+		AddErrorIfFalse(Stack.GetParam<float>("Param6") == 2.0f, "Param6 != 2.0f");
+		AddErrorIfFalse(Stack.GetParam<float>("Param7") == 1.0f, "Param7 != 1.0f");
+	}
+
+	// Check push/pop of referenced values
+	{
+		FParamStack Stack;
+		
+		float Value1 = 1.0f;
+		const float ConstValue2 = 1.0f;
+		Stack.PushValues("Param1", Value1, "Param2", ConstValue2);
+
+		AddErrorIfFalse(Stack.GetParam<float>("Param1") == 1.0f, "Param1 != 1.0f");
+		AddErrorIfFalse(Stack.IsMutableParam("Param1"), "Param1 is not mutable");
+		AddErrorIfFalse(Stack.IsReferenceParam("Param1"), "Param1 is not a reference");
+
+		AddErrorIfFalse(Stack.GetParam<float>("Param2") == 1.0f, "Param2 != 1.0f");
+		AddErrorIfFalse(!Stack.IsMutableParam("Param2"), "Param2 is mutable");
+		AddErrorIfFalse(Stack.IsReferenceParam("Param2"), "Param2 is not a reference");
+
+		Value1 = 2.0f;
+
+		AddErrorIfFalse(Stack.GetParam<float>("Param1") == 2.0f, "Param1 != 2.0f");
+		AddErrorIfFalse(Stack.GetParam<float>("Param2") == 1.0f, "Param2 != 1.0f");
+
+		Stack.PushValues("Param4", FVector(1.0f, 1.0f, 1.0f), "Param5", FQuat::MakeFromEuler(FVector(1.0f, 1.0f, 1.0f)));
+
+		const FVector Value4 = Stack.GetParam<FVector>("Param4");
+		AddErrorIfFalse(Value4 == FVector(1.0f, 1.0f, 1.0f), "Param4 != FVector(1.0f, 1.0f, 1.0f)");
+		AddErrorIfFalse(Stack.IsMutableParam("Param4"), "Param4 is not mutable");
+		AddErrorIfFalse(!Stack.IsReferenceParam("Param4"), "Param4 is a reference");
+
+		const FQuat Value5 = Stack.GetParam<FQuat>("Param5");
+		AddErrorIfFalse(Value5 == FQuat::MakeFromEuler(FVector(1.0f, 1.0f, 1.0f)), "Param5 != FQuat::MakeFromEuler(FVector(1.0f, 1.0f, 1.0f)");
+		AddErrorIfFalse(Stack.IsMutableParam("Param5"), "Param5 is not mutable");
+		AddErrorIfFalse(!Stack.IsReferenceParam("Param5"), "Param5 is a reference");
+
+		FTransform Value6 = FTransform::Identity;
+		Value6.SetScale3D(FVector(2.0f, 2.0f, 2.0f));
+		FTransform Value6Copy = Value6;
+		Stack.PushValues("Param6", Value6);
+
+		AddErrorIfFalse(Stack.GetParam<FTransform>("Param6").Equals(Value6Copy, 0.0f), "Param6 != FTransform::Identity");
+		AddErrorIfFalse(Stack.IsMutableParam("Param6"), "Param6 is not mutable");
+		AddErrorIfFalse(Stack.IsReferenceParam("Param6"), "Param6 is not a reference");
+	}
+
+	// Tests for instanced property bag
+	{
+		FInstancedPropertyBag PropertyBag;
+
+		FPropertyBagPropertyDesc PropertyDescs[] =
+		{
+			FPropertyBagPropertyDesc("Param0", EPropertyBagPropertyType::Float),
+			FPropertyBagPropertyDesc("Param1", EPropertyBagPropertyType::Bool),
+			FPropertyBagPropertyDesc("Param2", EPropertyBagPropertyType::Int32)
+		};
+		
+		PropertyBag.AddProperties(PropertyDescs);
+
+		PropertyBag.SetValueFloat("Param0", 1.5f);
+		PropertyBag.SetValueBool("Param1", true);
+		PropertyBag.SetValueInt32("Param2", 5);
+
+		TUniquePtr<FParamStackLayer> Layer = FParamStack::MakeLayer(PropertyBag);
+
+		FParamStack Stack;
+
+		Stack.PushLayer(*Layer);
+
+		AddErrorIfFalse(Stack.GetParam<float>("Param0") == 1.5f, "Param0 != 1.5f");
+		AddErrorIfFalse(Stack.GetParam<bool>("Param1") == true, "Param1 != true");
+		AddErrorIfFalse(Stack.GetParam<int32>("Param2") == 5, "Param2 != 5");
+	}
+
+	// Test to numeric limits
+	{
+		FParamStack Stack;
+
+		const uint32 NumParams = MAX_uint16;
+		for (uint32 Index = 0; Index < NumParams; ++Index)
+		{
+			Stack.PushValue(FName("Param", Index), (float)Index);
+		}
+
+		for (uint32 Index = 0; Index < NumParams; ++Index)
+		{
+			FName ParamName("Param", Index);
+			AddErrorIfFalse(Stack.GetParam<float>(ParamName) == (float)Index, FString::Printf(TEXT("%s != %.0f"), *ParamName.ToString(), (float)Index));
+		}
+
+		Stack.PushValue(FName("Param", MAX_uint16), (float)MAX_uint16);
+		AddErrorIfFalse(Stack.GetParamPtr<float>(FName("Param", MAX_uint16)) == nullptr, TEXT("Overflowed stack value is non-null"));
+
+		for (uint32 Index = 0; Index < NumParams; ++Index)
+		{
+			Stack.PopLayer();
+		}
+	}
+
+	FParamId::EndTestSandbox();
+
+	return true;
+}
+
+}
+
 #endif	// WITH_DEV_AUTOMATION_TESTS
\ No newline at end of file
diff --git a/Engine/Plugins/Experimental/AnimNext/Source/AnimNextTestSuite/Private/AnimNextTest.cpp b/Engine/Plugins/Experimental/AnimNext/Source/AnimNextTestSuite/Private/AnimNextTest.cpp
--- a/Engine/Plugins/Experimental/AnimNext/Source/AnimNextTestSuite/Private/AnimNextTest.cpp
+++ b/Engine/Plugins/Experimental/AnimNext/Source/AnimNextTestSuite/Private/AnimNextTest.cpp
@@ -3,17 +3,18 @@
 
 #include "AnimNextTest.h"
 #include "CoreMinimal.h"
-#include "Interface/AnimNextInterface.h"
-#include "Param/Param.h"
-#include "Interface/InterfaceContext.h"
-#include "Param/ParamStorage.h"
+#include "Context.h"
 #include "Misc/AutomationTest.h"
 #include "DataRegistry.h"
 #include "ReferencePose.h"
-
 #include <chrono>
 
-namespace 
+#if WITH_DEV_AUTOMATION_TESTS
+
+namespace UE::AnimNext::Tests
+{
+
+namespace Private
 {
 
 class FHighResTimer
@@ -36,352 +37,11 @@ public:
 	}
 };
 
-} // end namespace
-
-//----------------------------------------------------------------------------
-
-/*static*/ FName UAnimNextInterface_TestData_Multiply::NameParamA(TEXT("TestDataA"));
-/*static*/ FName UAnimNextInterface_TestData_Multiply::NameParamB(TEXT("TestDataB"));
-/*static*/ FName UAnimNextInterface_TestData_Multiply::NameParamResult(TEXT("NameParamAxB"));
-
-
-bool UAnimNextInterfaceTestDataLiteral::GetDataImpl(const UE::AnimNext::FContext& Context) const
-{
-	Context.SetResult(Value);
-
-	return true;
-}
-
-bool UAnimNextInterface_TestData_Multiply::GetDataImpl(const UE::AnimNext::FContext& Context) const
-{
-	using namespace UE::AnimNext;
-
-	bool bResult = false;
-
-	const FAnimNextTestData& ParamA = Context.GetParameterChecked<const FAnimNextTestData>(NameParamA);
-	const FAnimNextTestData& ParamB = Context.GetParameterChecked<const FAnimNextTestData>(NameParamB);
-	FAnimNextTestData& Result = Context.GetResult<FAnimNextTestData>();
-
-	bResult = true;
-
-	Result.A = ParamA.A * ParamB.A;
-	Result.B = ParamA.B * ParamB.B;
-
-	return bResult;
 }
 
-//----------------------------------------------------------------------------
+IMPLEMENT_SIMPLE_AUTOMATION_TEST(FDataRegistryTransformsTest, "Animation.AnimNext.AnimationDataRegistry.TransformsTest", EAutomationTestFlags::EditorContext | EAutomationTestFlags::EngineFilter)
 
-/*static*/ FName UAnimNextInterface_TestData_Split::InputName_AB(TEXT("Input_AB"));
-/*static*/ FName UAnimNextInterface_TestData_Split::OutputName_A(TEXT("Output_A"));
-/*static*/ FName UAnimNextInterface_TestData_Split::OutputName_B(TEXT("Output_B"));
-
-bool UAnimNextInterface_TestData_Split::GetDataImpl(const UE::AnimNext::FContext& Context) const
-{
-	using namespace UE::AnimNext;
-
-	bool bResult = false;
-
-	const FAnimNextTestData& Input_AB = Context.GetParameterChecked<const FAnimNextTestData>(InputName_AB);
-	float& Output_A = Context.GetParameterChecked<float>(OutputName_A);
-	float& Output_B = Context.GetParameterChecked<float>(OutputName_B);
-
-	bResult = true;
-
-	Output_A = Input_AB.A;
-	Output_B = Input_AB.B;
-
-	return bResult;
-}
-
-
-//****************************************************************************
-// AnimNext Interface Test
-//****************************************************************************
-#if WITH_DEV_AUTOMATION_TESTS
-
-IMPLEMENT_SIMPLE_AUTOMATION_TEST(FAnimationAnimNextInterfaceTest_SingleElement, "Animation.AnimNext.Interface.Single", EAutomationTestFlags::EditorContext | EAutomationTestFlags::EngineFilter)
-
-bool FAnimationAnimNextInterfaceTest_SingleElement::RunTest(const FString& InParameters)
-{
-	using namespace UE::AnimNext;
-
-	UAnimNextInterface_TestData_Multiply* AnimNextInterface_TestData_Multiply = Cast<UAnimNextInterface_TestData_Multiply>(UAnimNextInterface_TestData_Multiply::StaticClass()->GetDefaultObject());
-	TScriptInterface<IAnimNextInterface> TestDataMultiply = AnimNextInterface_TestData_Multiply;
-
-	AddErrorIfFalse(TestDataMultiply.GetInterface() != nullptr, "UAnimNextInterface_TestData_Multiply -> Interface is Null.");
-
-	constexpr FAnimNextTestData TestDataA = { 1.f, 1.f };
-	constexpr FAnimNextTestData TestDataB = { 1.f, 2.f };
-	FAnimNextTestData TestDataResult = { 0.f, 0.f };
-
-	FState State;
-	FParamStorage ParamStorage(16, 256, 2);
-	FContext Context(0.f, State, ParamStorage);
-
-	const FContext TestContext = Context.WithParameters({
-		TPairInitializer(UAnimNextInterface_TestData_Multiply::NameParamA, TWrapParam<const FAnimNextTestData>(TestDataA)),
-		TPairInitializer(UAnimNextInterface_TestData_Multiply::NameParamB, TWrapParam<const FAnimNextTestData>(TestDataB))
-		});
-
-	Interface::GetDataSafe(TestDataMultiply, TestContext, TestDataResult);
-
-	// --- Check the results from the batched operation ---
-	AddErrorIfFalse(TestDataResult.A == TestDataA.A * TestDataB.A, FString::Printf(TEXT("TestDataResult.A is : [%.04f] Expected : [%.04f]"), TestDataResult.A, TestDataA.A * TestDataB.A));
-	AddErrorIfFalse(TestDataResult.B == TestDataA.B * TestDataB.B, FString::Printf(TEXT("TestDataResult.B is : [%.04f] Expected : [%.04f]"), TestDataResult.B, TestDataA.B * TestDataB.B));
-
-	// Final GC to make sure everything is cleaned up
-	CollectGarbage(GARBAGE_COLLECTION_KEEPFLAGS);
-
-	return true;
-}
-
-IMPLEMENT_SIMPLE_AUTOMATION_TEST(FAnimationAnimNextInterfaceTest_Chained, "Animation.AnimNext.Interface.Chained", EAutomationTestFlags::EditorContext | EAutomationTestFlags::EngineFilter)
-
-bool FAnimationAnimNextInterfaceTest_Chained::RunTest(const FString& InParameters)
-{
-	using namespace UE::AnimNext;
-
-	FState RootState;
-	FParamStorage ParamStorage(16, 2048, 8);
-	const FContext RootContext(0.f, RootState, ParamStorage);
-	
-	// Simulate we pass a const context and have to return data
-	TParamValue<FAnimNextTestData> RootContextResultParam;
-	const FContext InContext = RootContext.WithResult(RootContextResultParam);
-
-	FAnimNextTestData LiteralResult;
-	TParamValue<FAnimNextTestData> ParamLiteralResult;
-	//TContextStorageParam<FTestData> ParamLiteralResult;
-
-	// Retrieve the literal value
-	{
-		TScriptInterface<IAnimNextInterface> TestDataLiteral = Cast<UAnimNextInterfaceTestDataLiteral>(UAnimNextInterfaceTestDataLiteral::StaticClass()->GetDefaultObject());
-		AddErrorIfFalse(TestDataLiteral.GetInterface() != nullptr, "UAnimNextInterfaceTestDataLiteral -> Interface is Null.");
-
-		const FContext LiteralContext = InContext.WithResult(ParamLiteralResult);
-
-		Interface::GetDataSafe(TestDataLiteral, LiteralContext);
-
-		AddErrorIfFalse(ParamLiteralResult->A == 1, FString::Printf(TEXT("LiteralResult.A is : [%.04f] Expected : [%.04f]"), ParamLiteralResult->A, 1));
-		AddErrorIfFalse(ParamLiteralResult->B == 1, FString::Printf(TEXT("LiteralResult.B is : [%.04f] Expected : [%.04f]"), ParamLiteralResult->B, 1));
-	}
-
-	const FAnimNextTestData TestDataB = { 2.f, 2.f };
-
-	// now chain an operation using the literal result as a parameter
-	{
-		UAnimNextInterface_TestData_Multiply* AnimNextInterface_TestData_Multiply = Cast<UAnimNextInterface_TestData_Multiply>(UAnimNextInterface_TestData_Multiply::StaticClass()->GetDefaultObject());
-		TScriptInterface<IAnimNextInterface> TestDataMultiply = AnimNextInterface_TestData_Multiply;
-		AddErrorIfFalse(TestDataMultiply.GetInterface() != nullptr, "UAnimNextInterface_TestData_Multiply -> Interface is Null.");
-
-		const FContext ChainContext = InContext.WithResultAndParameters(InContext.GetResultParam(), {
-			TPairInitializer(UAnimNextInterface_TestData_Multiply::NameParamA, ParamLiteralResult),						// chain the result of the literal call as the input
-			TPairInitializer(UAnimNextInterface_TestData_Multiply::NameParamB, TWrapParam<const FAnimNextTestData>(TestDataB))	// set the second operand
-			});
-
-		Interface::GetDataSafe(TestDataMultiply, ChainContext);
-	}
-
-	// --- Check the results from the chained operation on the root component ---
-	AddErrorIfFalse(RootContextResultParam->A == ParamLiteralResult->A * TestDataB.A, FString::Printf(TEXT("TestDataResult.A is : [%.04f] Expected : [%.04f]"), RootContextResultParam->A, ParamLiteralResult->A * TestDataB.A));
-	AddErrorIfFalse(RootContextResultParam->B == ParamLiteralResult->B * TestDataB.B, FString::Printf(TEXT("TestDataResult.B is : [%.04f] Expected : [%.04f]"), RootContextResultParam->B, ParamLiteralResult->B * TestDataB.B));
-
-	// now chain an operation using the literal result as a parameter but copying the input parameters on the Context Parameter storage
-	{
-		UAnimNextInterface_TestData_Multiply* AnimNextInterface_TestData_Multiply = Cast<UAnimNextInterface_TestData_Multiply>(UAnimNextInterface_TestData_Multiply::StaticClass()->GetDefaultObject());
-		TScriptInterface<IAnimNextInterface> TestDataMultiply = AnimNextInterface_TestData_Multiply;
-		AddErrorIfFalse(TestDataMultiply.GetInterface() != nullptr, "UAnimNextInterface_TestData_Multiply -> Interface is Null.");
-
-		const FContext ChainContext = InContext.WithResultAndParameters(InContext.GetResultParam(), {
-			TPairInitializer(UAnimNextInterface_TestData_Multiply::NameParamA, TContextStorageParam<const FAnimNextTestData>(ParamLiteralResult)),	// chain the result of the literal call as the input
-			TPairInitializer(UAnimNextInterface_TestData_Multiply::NameParamB, TContextStorageParam<const FAnimNextTestData>(TestDataB))			// set the second operand
-			});
-
-		Interface::GetDataSafe(TestDataMultiply, ChainContext);
-	}
-
-	// --- Check the results from the chained operation on the root component ---
-	AddErrorIfFalse(RootContextResultParam->A == ParamLiteralResult->A * TestDataB.A, FString::Printf(TEXT("TestDataResult.A is : [%.04f] Expected : [%.04f]"), RootContextResultParam->A, ParamLiteralResult->A * TestDataB.A));
-	AddErrorIfFalse(RootContextResultParam->B == ParamLiteralResult->B * TestDataB.B, FString::Printf(TEXT("TestDataResult.B is : [%.04f] Expected : [%.04f]"), RootContextResultParam->B, ParamLiteralResult->B * TestDataB.B));
-
-	// Final GC to make sure everything is cleaned up
-	CollectGarbage(GARBAGE_COLLECTION_KEEPFLAGS);
-
-	return true;
-}
-
-IMPLEMENT_SIMPLE_AUTOMATION_TEST(FAnimationAnimNextInterfaceTest_AltParamStorage, "Animation.AnimNext.Interface.AltParamStorage", EAutomationTestFlags::EditorContext | EAutomationTestFlags::EngineFilter)
-
-bool FAnimationAnimNextInterfaceTest_AltParamStorage::RunTest(const FString& InParameters)
-{
-	using namespace UE::AnimNext;
-
-	FState RootState;
-	FParamStorage ParamStorage(16, 2048, 8);
-	const FContext RootContext(0.f, RootState, ParamStorage);
-
-	UAnimNextInterface_TestData_Multiply* AnimNextInterface_TestData_Multiply = Cast<UAnimNextInterface_TestData_Multiply>(UAnimNextInterface_TestData_Multiply::StaticClass()->GetDefaultObject());
-	TScriptInterface<IAnimNextInterface> TestDataMultiply = AnimNextInterface_TestData_Multiply;
-
-	AddErrorIfFalse(TestDataMultiply.GetInterface() != nullptr, "UAnimNextInterface_TestData_Multiply -> Interface is Null.");
-
-	constexpr FAnimNextTestData TestDataA = { 1.f, 1.f };
-	constexpr FAnimNextTestData TestDataB = { 1.f, 2.f };
-	FAnimNextTestData TestDataResult = { 0.f, 0.f };
-
-	FContext TestContext = RootContext.CreateSubContext();
-
-	TestContext.AddValue(FContext::EParamType::Input, UAnimNextInterface_TestData_Multiply::NameParamA, TestDataA);
-	TestContext.AddValue(FContext::EParamType::Input, UAnimNextInterface_TestData_Multiply::NameParamB, TestDataB);
-	FParamHandle HOutput = TestContext.AddReference(FContext::EParamType::Output, UAnimNextInterface_TestData_Multiply::NameParamResult, TestDataResult);
-
-	// for compatibility reasons, as the current code expects to have a return value set
-	TestContext.SetHParamAsResult(HOutput);  
-
-	const bool bOk = TestDataMultiply->GetData(TestContext);
-
-	AddErrorIfFalse(bOk, "TestDataMultiply -> GetData returned FAILURE.");
-
-	// --- Check the results from the operation ---
-	AddErrorIfFalse(TestDataResult.A == TestDataA.A * TestDataB.A, FString::Printf(TEXT("TestDataResult.A is : [%.04f] Expected : [%.04f]"), TestDataResult.A, TestDataA.A * TestDataB.A));
-	AddErrorIfFalse(TestDataResult.B == TestDataA.B * TestDataB.B, FString::Printf(TEXT("TestDataResult.B is : [%.04f] Expected : [%.04f]"), TestDataResult.B, TestDataA.B * TestDataB.B));
-
-	// Final GC to make sure everything is cleaned up
-	CollectGarbage(GARBAGE_COLLECTION_KEEPFLAGS);
-
-	return true;
-}
-
-
-IMPLEMENT_SIMPLE_AUTOMATION_TEST(FAnimationAnimNextInterfaceTest_AltParamStorageMultipleOutputs, "Animation.AnimNext.Interface.AltParamStorageMultipleOutputs", EAutomationTestFlags::EditorContext | EAutomationTestFlags::EngineFilter)
-
-bool FAnimationAnimNextInterfaceTest_AltParamStorageMultipleOutputs::RunTest(const FString& InParameters)
-{
-	using namespace UE::AnimNext;
-
-	FState RootState;
-	FParamStorage ParamStorage(16, 2048, 8);
-	const FContext RootContext(0.f, RootState, ParamStorage);
-
-	TScriptInterface<IAnimNextInterface> TestDataSplit(Cast<UAnimNextInterface_TestData_Split>(UAnimNextInterface_TestData_Split::StaticClass()->GetDefaultObject()));
-
-	AddErrorIfFalse(TestDataSplit.GetInterface() != nullptr, "UAnimNextInterface_TestData_Split -> Interface is Null.");
-
-	constexpr FAnimNextTestData Input_AB = { 1.f, 2.f };
-	float Output_A = 0.f;
-	float Output_B = 0.f;
-
-	FContext TestContext = RootContext.CreateSubContext();
-
-	// Adding an input as a value (will be copied and the copy will become mutable)
-	FParamHandle RetValCompatiblility = TestContext.AddValue(FContext::EParamType::Input, UAnimNextInterface_TestData_Split::InputName_AB, Input_AB);
-	// Adding first output as a reference
-	TestContext.AddReference(FContext::EParamType::Output, UAnimNextInterface_TestData_Split::OutputName_A, Output_A);
-	// Adding second output as a reference
-	TestContext.AddReference(FContext::EParamType::Output, UAnimNextInterface_TestData_Split::OutputName_B, Output_B);
-
-	// Not used, but added for compatibility reasons, as the current code expects to have a return value set of the interface type
-	TestContext.SetHParamAsResult(RetValCompatiblility);
-
-	const bool bOk = TestDataSplit->GetData(TestContext);
-
-	AddErrorIfFalse(bOk, "TestDataSplit -> GetData returned FAILURE.");
-
-	// --- Check the results from the operation ---
-	AddErrorIfFalse(Output_A == Input_AB.A, FString::Printf(TEXT("Output_A is : [%.04f] Expected : [%.04f]"), Output_A, Input_AB.A));
-	AddErrorIfFalse(Output_B == Input_AB.B, FString::Printf(TEXT("Output_B is : [%.04f] Expected : [%.04f]"), Output_B, Input_AB.B));
-
-	// Final GC to make sure everything is cleaned up
-	CollectGarbage(GARBAGE_COLLECTION_KEEPFLAGS);
-
-	return true;
-}
-
-IMPLEMENT_SIMPLE_AUTOMATION_TEST(FAnimationAnimNextInterfaceTest_AltParamStorageChained, "Animation.AnimNext.Interface.AltParamStorageChained", EAutomationTestFlags::EditorContext | EAutomationTestFlags::EngineFilter)
-
-bool FAnimationAnimNextInterfaceTest_AltParamStorageChained::RunTest(const FString& InParameters)
-{
-	using namespace UE::AnimNext;
-
-	FState RootState;
-	FParamStorage ParamStorage(16, 2048, 8);
-	const FContext RootContext(0.f, RootState, ParamStorage);
-
-	UAnimNextInterface_TestData_Multiply* AnimNextInterface_TestData_Multiply = Cast<UAnimNextInterface_TestData_Multiply>(UAnimNextInterface_TestData_Multiply::StaticClass()->GetDefaultObject());
-	TScriptInterface<IAnimNextInterface> TestDataMultiply = AnimNextInterface_TestData_Multiply;
-
-	AddErrorIfFalse(TestDataMultiply.GetInterface() != nullptr, "UAnimNextInterface_TestData_Multiply -> Interface is Null.");
-
-	// --- Initial Test : Multiply two TestData and store the value at Result ---
-
-	FParamHandle TestOutput_AxB_Param;
-
-	{
-		constexpr FAnimNextTestData TestDataA = { 1.f, 1.f };
-		constexpr FAnimNextTestData TestDataB = { 1.f, 2.f };
-
-		FContext TestContext = RootContext.CreateSubContext();
-
-		// Adding a const value. Note that adding a value will convert const into mutable, as the copy will not be const anymore
-		TestContext.AddValue(FContext::EParamType::Input, UAnimNextInterface_TestData_Multiply::NameParamA, TestDataA);
-		// Adding a reference to const. This param will keep the const
-		TestContext.AddReference(FContext::EParamType::Input, UAnimNextInterface_TestData_Multiply::NameParamB, TestDataB);
-		// Adding a value using a temp taht will be copied into the value
-		TestOutput_AxB_Param = TestContext.AddValue(FContext::EParamType::Output, UAnimNextInterface_TestData_Multiply::NameParamResult, FAnimNextTestData{ 0.f, 0.f });
-
-		// for compatibility reasons, as the current code expects to have a return value set
-		TestContext.SetHParamAsResult(TestOutput_AxB_Param);  
-
-		const bool bOk = TestDataMultiply->GetData(TestContext);
-
-		AddErrorIfFalse(bOk, "TestDataMultiply -> GetData returned FAILURE.");
-
-		// --- Check the results from the operation ---
-		const FAnimNextTestData& TestDataResult = TestContext.GetParameterAs<const FAnimNextTestData>(TestOutput_AxB_Param);
-		AddErrorIfFalse(TestDataResult.A == TestDataA.A * TestDataB.A, FString::Printf(TEXT("TestDataResult.A is : [%.04f] Expected : [%.04f]"), TestDataResult.A, TestDataA.A * TestDataB.A));
-		AddErrorIfFalse(TestDataResult.B == TestDataA.B * TestDataB.B, FString::Printf(TEXT("TestDataResult.B is : [%.04f] Expected : [%.04f]"), TestDataResult.B, TestDataA.B * TestDataB.B));
-	}
-
-	// --- Chained Test : Use the previous call Result as an input ---
-
-	{
-		FContext TestContext = RootContext.CreateSubContext();
-
-		constexpr FAnimNextTestData TestDataC = { 2.f, 2.f };
-
-		// Add the result of the previous call as an input for this operation
-		TestContext.AddValue(FContext::EParamType::Input, UAnimNextInterface_TestData_Multiply::NameParamA, TestOutput_AxB_Param);
-		// Add a reference to const data
-		TestContext.AddReference(FContext::EParamType::Input, UAnimNextInterface_TestData_Multiply::NameParamB, TestDataC);
-		// Add a value
-		const FParamHandle ChainedResultParam = TestContext.AddValue(FContext::EParamType::Output, UAnimNextInterface_TestData_Multiply::NameParamResult, FAnimNextTestData{ 0.f, 0.f });
-
-		// for compatibility reasons, as the current code expects to have a return value set
-		TestContext.SetHParamAsResult(ChainedResultParam);
-
-		const bool bOk = TestDataMultiply->GetData(TestContext);
-
-		AddErrorIfFalse(bOk, "TestDataMultiply -> GetData returned FAILURE.");
-
-		// --- Check the results from the operation ---
-		const FAnimNextTestData& TestOutput_AxB = TestContext.GetParameterAs<const FAnimNextTestData>(TestOutput_AxB_Param);
-		const FAnimNextTestData& ChainedResult = TestContext.GetParameterAs<const FAnimNextTestData>(ChainedResultParam);
-		
-		AddErrorIfFalse(ChainedResult.A == TestOutput_AxB.A * TestDataC.A, FString::Printf(TEXT("Chained TestDataResult.A is : [%.04f] Expected : [%.04f]"), ChainedResult.A, TestOutput_AxB.A * TestDataC.A));
-		AddErrorIfFalse(ChainedResult.B == TestOutput_AxB.B * TestDataC.B, FString::Printf(TEXT("Chained TestDataResult.B is : [%.04f] Expected : [%.04f]"), ChainedResult.B, TestOutput_AxB.B * TestDataC.B));
-	}
-
-
-	// Final GC to make sure everything is cleaned up
-	CollectGarbage(GARBAGE_COLLECTION_KEEPFLAGS);
-
-	return true;
-}
-
-IMPLEMENT_SIMPLE_AUTOMATION_TEST(FAnimationAnimNextInterfaceTest_AnimationDataRegistry, "Animation.AnimNext.AnimationDataRegistry.TransformsTest", EAutomationTestFlags::EditorContext | EAutomationTestFlags::EngineFilter)
-
-bool FAnimationAnimNextInterfaceTest_AnimationDataRegistry::RunTest(const FString& InParameters)
+bool FDataRegistryTransformsTest::RunTest(const FString& InParameters)
 {
 	using namespace UE::AnimNext;
 
@@ -433,11 +93,11 @@ bool FAnimationAnimNextInterfaceTest_AnimationDataRegistry::RunTest(const FStrin
 				Transform.SetScale3D(Half);
 			}
 
-			FHighResTimer::FTimePoint StartTime = FHighResTimer::GetTimeMark();
+			Private::FHighResTimer::FTimePoint StartTime = Private::FHighResTimer::GetTimeMark();
 
 			ResultAoS.Blend(TransformsA, TransformsB, BlendFactor);
 
-			TimeDiffAoS[i] = FHighResTimer::GetTimeDiffNanoSec(StartTime);
+			TimeDiffAoS[i] = Private::FHighResTimer::GetTimeDiffNanoSec(StartTime);
 		}
 	}
 	for (auto& Transform : ResultAoS.Transforms)
@@ -471,11 +131,11 @@ bool FAnimationAnimNextInterfaceTest_AnimationDataRegistry::RunTest(const FStrin
 				Scale3D = Half;
 			}
 
-			FHighResTimer::FTimePoint StartTime = FHighResTimer::GetTimeMark();
+			Private::FHighResTimer::FTimePoint StartTime = Private::FHighResTimer::GetTimeMark();
 
 			ResultSoA.Blend(TransformsA, TransformsB, BlendFactor);
 
-			TimeDiffSoA[i] = FHighResTimer::GetTimeDiffNanoSec(StartTime);
+			TimeDiffSoA[i] = Private::FHighResTimer::GetTimeDiffNanoSec(StartTime);
 		}
 	}
 	for (const auto& Translation : ResultSoA.Translations)
@@ -521,4 +181,6 @@ bool FAnimationAnimNextInterfaceTest_AnimationDataRegistry::RunTest(const FStrin
 	return true;
 }
 
+}
+
 #endif // WITH_DEV_AUTOMATION_TESTS
\ No newline at end of file
diff --git a/Engine/Plugins/Experimental/AnimNext/Source/AnimNextTestSuite/Private/AnimNextTest.h b/Engine/Plugins/Experimental/AnimNext/Source/AnimNextTestSuite/Private/AnimNextTest.h
--- a/Engine/Plugins/Experimental/AnimNext/Source/AnimNextTestSuite/Private/AnimNextTest.h
+++ b/Engine/Plugins/Experimental/AnimNext/Source/AnimNextTestSuite/Private/AnimNextTest.h
@@ -2,9 +2,6 @@
 
 #pragma once
 
-#include "Interface/InterfaceTypes.h"
-#include "Interface/IAnimNextInterface.h"
-
 #include "AnimNextTest.generated.h"
 
 //// --- Raw type ---
@@ -16,60 +13,3 @@ struct FAnimNextTestData
 	float A = 0.f;
 	float B = 0.f;
 };
-
-// -- Defines the TestData as an interface that will be used by any variant operating on TestData --
-UCLASS()
-class UAnimNextInterface_TestData : public UObject, public IAnimNextInterface
-{
-	GENERATED_BODY()
-
-	ANIM_NEXT_INTERFACE_RETURN_TYPE(FAnimNextTestData)
-
-	virtual bool GetDataImpl(const UE::AnimNext::FContext& Context) const override
-	{
-		checkf(false, TEXT("UAnimNextInterfaceFloat::GetDataImpl must be overridden"));
-		return false;
-	}
-};
-
-// --- Constant user defined value ---
-UCLASS()
-class UAnimNextInterfaceTestDataLiteral : public UAnimNextInterface_TestData
-{
-	GENERATED_BODY()
-
-	// IAnimNextInterface interface
-	virtual bool GetDataImpl(const UE::AnimNext::FContext& Context) const final override;
-
-	FAnimNextTestData Value = {1.f, 1.f};
-};
-
-// --- TestData multiply operator  ---
-UCLASS()
-class UAnimNextInterface_TestData_Multiply : public UAnimNextInterface_TestData
-{
-	GENERATED_BODY()
-
-	// IAnimNextInterface interface
-	virtual bool GetDataImpl(const UE::AnimNext::FContext& Context) const final override;
-
-public:
-	static FName NameParamA;
-	static FName NameParamB;
-	static FName NameParamResult;
-};
-
-// --- TestData split operator : creates two outputs from a single input (multi output test) ---
-UCLASS()
-class UAnimNextInterface_TestData_Split : public UAnimNextInterface_TestData
-{
-	GENERATED_BODY()
-
-	// IAnimNextInterface interface
-	virtual bool GetDataImpl(const UE::AnimNext::FContext& Context) const final override;
-
-public:
-	static FName InputName_AB;
-	static FName OutputName_A;
-	static FName OutputName_B;
-};
diff --git a/Engine/Plugins/Experimental/AnimNext/Source/AnimNextUncookedOnly/Internal/Graph/AnimNextGraph_EditorData.h b/Engine/Plugins/Experimental/AnimNext/Source/AnimNextUncookedOnly/Internal/Graph/AnimNextGraph_EditorData.h
--- a/Engine/Plugins/Experimental/AnimNext/Source/AnimNextUncookedOnly/Internal/Graph/AnimNextGraph_EditorData.h
+++ b/Engine/Plugins/Experimental/AnimNext/Source/AnimNextUncookedOnly/Internal/Graph/AnimNextGraph_EditorData.h
@@ -44,6 +44,7 @@ class UAnimNextGraph_EditorData : public UObject, public IRigVMClientHost, publi
 	friend struct UE::AnimNext::UncookedOnly::FUtils;
 	friend class UE::AnimNext::Editor::FGraphEditor;
 	friend struct FAnimNextGraphSchemaAction_RigUnit;
+	friend struct FAnimNextGraphSchemaAction_DispatchFactory;
 
 	// UObject interface
 	virtual void Serialize(FArchive& Ar) override;
@@ -85,10 +86,6 @@ class UAnimNextGraph_EditorData : public UObject, public IRigVMClientHost, publi
 
 	void CreateEdGraphForCollapseNode(URigVMCollapseNode* InNode);
 
-	void UpdateGraphReturnType();
-
-	FEdGraphPinType FindGraphReturnPinType() const;
-
 	bool IsNodeExecConnected(const UEdGraphNode* Node) const;
 
 	UPROPERTY()
diff --git a/Engine/Plugins/Experimental/AnimNext/Source/AnimNextUncookedOnly/Internal/Param/AnimNextParameterLibrary.h b/Engine/Plugins/Experimental/AnimNext/Source/AnimNextUncookedOnly/Internal/Param/AnimNextParameterLibrary.h
--- a/Engine/Plugins/Experimental/AnimNext/Source/AnimNextUncookedOnly/Internal/Param/AnimNextParameterLibrary.h
+++ b/Engine/Plugins/Experimental/AnimNext/Source/AnimNextUncookedOnly/Internal/Param/AnimNextParameterLibrary.h
@@ -78,7 +78,7 @@ class UAnimNextParameterLibrary : public UObject
 	ANIMNEXTUNCOOKEDONLY_API UAnimNextParameter* AddParameter(FName InName, const FAnimNextParamType& InType, bool bSetupUndoRedo = true, bool bPrintPythonCommand = true);
 
 	UFUNCTION(BlueprintCallable, Category = "AnimNext|Parameter Library")
-	ANIMNEXTUNCOOKEDONLY_API UAnimNextParameter* AddParameter(FName InName, const EPropertyBagPropertyType& InValueType, const EPropertyBagContainerType& InContainerType, UObject* InValueTypeObject = nullptr, bool bSetupUndoRedo = true, bool bPrintPythonCommand = true);
+	ANIMNEXTUNCOOKEDONLY_API UAnimNextParameter* AddParameter(FName InName, const EPropertyBagPropertyType& InValueType, const EPropertyBagContainerType& InContainerType, const UObject* InValueTypeObject = nullptr, bool bSetupUndoRedo = true, bool bPrintPythonCommand = true);
 
 	UFUNCTION(BlueprintCallable, Category = "AnimNext|Parameter Library")
 	ANIMNEXTUNCOOKEDONLY_API bool RemoveParameter(FName InName, bool bSetupUndoRedo = true, bool bPrintPythonCommand = true);
diff --git a/Engine/Plugins/Experimental/AnimNext/Source/AnimNextUncookedOnly/Internal/UncookedOnlyUtils.h b/Engine/Plugins/Experimental/AnimNext/Source/AnimNextUncookedOnly/Internal/UncookedOnlyUtils.h
--- a/Engine/Plugins/Experimental/AnimNext/Source/AnimNextUncookedOnly/Internal/UncookedOnlyUtils.h
+++ b/Engine/Plugins/Experimental/AnimNext/Source/AnimNextUncookedOnly/Internal/UncookedOnlyUtils.h
@@ -24,7 +24,7 @@ struct ANIMNEXTUNCOOKEDONLY_API FUtils
 {
 	static void Compile(UAnimNextGraph* InGraph);
 	
-	static UAnimNextGraph_EditorData* GetEditorData(const UAnimNextGraph* InAnimNextInterfaceGraph);
+	static UAnimNextGraph_EditorData* GetEditorData(const UAnimNextGraph* InAnimNextGraph);
 	
 	static void RecreateVM(UAnimNextGraph* InGraph);
 
diff --git a/Engine/Plugins/Experimental/AnimNext/Source/AnimNextUncookedOnly/Private/Graph/AnimGraph/AnimGraphNode_AnimNextGraph.cpp b/Engine/Plugins/Experimental/AnimNext/Source/AnimNextUncookedOnly/Private/Graph/AnimGraph/AnimGraphNode_AnimNextGraph.cpp
--- a/Engine/Plugins/Experimental/AnimNext/Source/AnimNextUncookedOnly/Private/Graph/AnimGraph/AnimGraphNode_AnimNextGraph.cpp
+++ b/Engine/Plugins/Experimental/AnimNext/Source/AnimNextUncookedOnly/Private/Graph/AnimGraph/AnimGraphNode_AnimNextGraph.cpp
@@ -128,7 +128,7 @@ void UAnimGraphNode_AnimNextGraph::CustomizeDetails(IDetailLayoutBuilder& Detail
 	// ***************************************************************
 
 
-	TSharedRef<IPropertyHandle> ClassHandle = DetailBuilder.GetProperty(TEXT("Node.AnimNextInterfaceGraph"), GetClass());
+	TSharedRef<IPropertyHandle> ClassHandle = DetailBuilder.GetProperty(TEXT("Node.AnimNextGraph"), GetClass());
 	if (ClassHandle->IsValidHandle())
 	{
 		ClassHandle->SetOnPropertyValueChanged(FSimpleDelegate::CreateUObject(this, &UAnimGraphNode_AnimNextGraph::OnInstanceClassChanged, &DetailBuilder));
diff --git a/Engine/Plugins/Experimental/AnimNext/Source/AnimNextUncookedOnly/Private/Graph/AnimNextGraph_EditorData.cpp b/Engine/Plugins/Experimental/AnimNext/Source/AnimNextUncookedOnly/Private/Graph/AnimNextGraph_EditorData.cpp
--- a/Engine/Plugins/Experimental/AnimNext/Source/AnimNextUncookedOnly/Private/Graph/AnimNextGraph_EditorData.cpp
+++ b/Engine/Plugins/Experimental/AnimNext/Source/AnimNextUncookedOnly/Private/Graph/AnimNextGraph_EditorData.cpp
@@ -389,7 +389,6 @@ void UAnimNextGraph_EditorData::HandleModifiedEvent(ERigVMGraphNotifType InNotif
 	case ERigVMGraphNotifType::PinArraySizeChanged:
 	case ERigVMGraphNotifType::PinDirectionChanged:
 		{
-			UpdateGraphReturnType();
 			RecompileVM();
 			break;
 		}
@@ -475,38 +474,6 @@ void UAnimNextGraph_EditorData::CreateEdGraphForCollapseNode(URigVMCollapseNode*
 	}
 }
 
-FEdGraphPinType UAnimNextGraph_EditorData::FindGraphReturnPinType() const
-{
-	static const TCHAR* EndExecStr = TEXT("AnimNextEndExec");
-	static const FName ResultName = TEXT(".Result");
-
-	const UAnimNextGraph_EdGraph* AnimNextInterfaceEdGraph = RootGraph.Get();
-	if (AnimNextInterfaceEdGraph != nullptr)
-	{
-		for (const TObjectPtr<UEdGraphNode>& Node : AnimNextInterfaceEdGraph->Nodes)
-		{
-			if (IsNodeExecConnected(Node))
-			{
-				const FString& NodeName = Node->GetName();
-				if (NodeName.Contains(EndExecStr))
-				{
-					for (const UEdGraphPin* Pin : Node->Pins)
-					{
-						const FString PinName = Pin->PinName.ToString();
-
-						if (PinName.Contains(ResultName.ToString()) && PinName.Contains(EndExecStr))
-						{
-							return Pin->PinType;
-						}
-					}
-				}
-			}
-		}
-	}
-
-	return FEdGraphPinType();
-}
-
 bool UAnimNextGraph_EditorData::IsNodeExecConnected(const UEdGraphNode* Node) const
 {
 	static const TCHAR* ExecuteContextStr = TEXT(".ExecuteContext");
@@ -523,17 +490,3 @@ bool UAnimNextGraph_EditorData::IsNodeExecConnected(const UEdGraphNode* Node) co
 
 	return false;
 }
-
-void UAnimNextGraph_EditorData::UpdateGraphReturnType()
-{
-	if (UAnimNextGraph* Graph = GetTypedOuter<UAnimNextGraph>())
-	{
-		FEdGraphPinType GraphReturnPinType = FindGraphReturnPinType();
-		UE::AnimNext::FParamTypeHandle Handle = UE::AnimNext::UncookedOnly::FUtils::GetParameterHandleFromPin(GraphReturnPinType);
-		if (Handle.IsValid() && Graph->GetReturnTypeHandle() != Handle)
-		{
-			Graph->Modify();
-			Graph->SetReturnTypeHandle(Handle);
-		}
-	}
-}
diff --git a/Engine/Plugins/Experimental/AnimNext/Source/AnimNextUncookedOnly/Private/Param/AnimNextParameterLibrary.cpp b/Engine/Plugins/Experimental/AnimNext/Source/AnimNextUncookedOnly/Private/Param/AnimNextParameterLibrary.cpp
--- a/Engine/Plugins/Experimental/AnimNext/Source/AnimNextUncookedOnly/Private/Param/AnimNextParameterLibrary.cpp
+++ b/Engine/Plugins/Experimental/AnimNext/Source/AnimNextUncookedOnly/Private/Param/AnimNextParameterLibrary.cpp
@@ -33,7 +33,7 @@ UAnimNextParameter* UAnimNextParameterLibrary::AddParameter(FName InName, const
 	return AddParameter(InName, InType.GetValueType(), InType.GetContainerType(), InType.GetValueTypeObject(), bSetupUndoRedo, bPrintPythonCommand);
 }
 
-UAnimNextParameter* UAnimNextParameterLibrary::AddParameter(FName InName, const EPropertyBagPropertyType& InValueType, const EPropertyBagContainerType& InContainerType, UObject* InValueTypeObject, bool bSetupUndoRedo, bool bPrintPythonCommand)
+UAnimNextParameter* UAnimNextParameterLibrary::AddParameter(FName InName, const EPropertyBagPropertyType& InValueType, const EPropertyBagContainerType& InContainerType, const UObject* InValueTypeObject, bool bSetupUndoRedo, bool bPrintPythonCommand)
 {
 	if(InName == NAME_None)
 	{
diff --git a/Engine/Plugins/Experimental/AnimNext/Source/AnimNextUncookedOnly/Private/UncookedOnlyUtils.cpp b/Engine/Plugins/Experimental/AnimNext/Source/AnimNextUncookedOnly/Private/UncookedOnlyUtils.cpp
--- a/Engine/Plugins/Experimental/AnimNext/Source/AnimNextUncookedOnly/Private/UncookedOnlyUtils.cpp
+++ b/Engine/Plugins/Experimental/AnimNext/Source/AnimNextUncookedOnly/Private/UncookedOnlyUtils.cpp
@@ -571,15 +571,15 @@ FRigVMTemplateArgumentType FUtils::GetRigVMArgTypeFromParamType(const FAnimNextP
 		break;
 	case EPropertyBagPropertyType::Enum:
 		CPPTypeString = RigVMTypeUtils::CPPTypeFromEnum(Cast<UEnum>(InParamType.ValueTypeObject.Get()));
-		ArgType.CPPTypeObject = InParamType.ValueTypeObject.Get();
+		ArgType.CPPTypeObject = const_cast<UObject*>(InParamType.ValueTypeObject.Get());
 		break;
 	case EPropertyBagPropertyType::Struct:
 		CPPTypeString = RigVMTypeUtils::GetUniqueStructTypeName(Cast<UScriptStruct>(InParamType.ValueTypeObject.Get()));
-		ArgType.CPPTypeObject = InParamType.ValueTypeObject.Get();
+		ArgType.CPPTypeObject = const_cast<UObject*>(InParamType.ValueTypeObject.Get());
 		break;
 	case EPropertyBagPropertyType::Object:
 		CPPTypeString = RigVMTypeUtils::CPPTypeFromObject(Cast<UClass>(InParamType.ValueTypeObject.Get()));
-		ArgType.CPPTypeObject = InParamType.ValueTypeObject.Get();
+		ArgType.CPPTypeObject = const_cast<UObject*>(InParamType.ValueTypeObject.Get());
 		break;
 	case EPropertyBagPropertyType::SoftObject:
 		ensureMsgf(false, TEXT("Unhandled value type %d"), InParamType.ValueType);
diff --git a/Engine/Plugins/Experimental/StructUtils/Source/StructUtils/Public/PropertyBag.h b/Engine/Plugins/Experimental/StructUtils/Source/StructUtils/Public/PropertyBag.h
--- a/Engine/Plugins/Experimental/StructUtils/Source/StructUtils/Public/PropertyBag.h
+++ b/Engine/Plugins/Experimental/StructUtils/Source/StructUtils/Public/PropertyBag.h
@@ -223,13 +223,13 @@ struct STRUCTUTILS_API FPropertyBagPropertyDesc
 
 	FPropertyBagPropertyDesc() = default;
 	FPropertyBagPropertyDesc(const FName InName, const FProperty* InSourceProperty);
-	FPropertyBagPropertyDesc(const FName InName, const EPropertyBagPropertyType InValueType, UObject* InValueTypeObject = nullptr)
+	FPropertyBagPropertyDesc(const FName InName, const EPropertyBagPropertyType InValueType, const UObject* InValueTypeObject = nullptr)
 		: ValueTypeObject(InValueTypeObject)
 		, Name(InName)
 		, ValueType(InValueType)
 	{
 	}
-	FPropertyBagPropertyDesc(const FName InName, const EPropertyBagContainerType InContainerType, const EPropertyBagPropertyType InValueType, UObject* InValueTypeObject = nullptr)
+	FPropertyBagPropertyDesc(const FName InName, const EPropertyBagContainerType InContainerType, const EPropertyBagPropertyType InValueType, const UObject* InValueTypeObject = nullptr)
 		: ValueTypeObject(InValueTypeObject)
 		, Name(InName)
 		, ValueType(InValueType)
diff --git a/Engine/Plugins/Runtime/RigVM/Source/RigVM/Public/RigVMTypeUtils.h b/Engine/Plugins/Runtime/RigVM/Source/RigVM/Public/RigVMTypeUtils.h
--- a/Engine/Plugins/Runtime/RigVM/Source/RigVM/Public/RigVMTypeUtils.h
+++ b/Engine/Plugins/Runtime/RigVM/Source/RigVM/Public/RigVMTypeUtils.h
@@ -104,7 +104,7 @@ namespace RigVMTypeUtils
 		return InScriptStruct->GetStructCPPName();
 	}
 
-	inline FString CPPTypeFromEnum(UEnum* InEnum)
+	inline FString CPPTypeFromEnum(const UEnum* InEnum)
 	{
 		check(InEnum);
 
@@ -274,7 +274,7 @@ namespace RigVMTypeUtils
 		return Object;
 	}
 
-	static FString CPPTypeFromObject(UObject* InCPPTypeObject)
+	static FString CPPTypeFromObject(const UObject* InCPPTypeObject)
 	{
 		if (const UClass* Class = Cast<UClass>(InCPPTypeObject))
 		{
@@ -291,7 +291,7 @@ namespace RigVMTypeUtils
 		{
 			return GetUniqueStructTypeName(ScriptStruct);
 		}
-		else if (UEnum* Enum = Cast<UEnum>(InCPPTypeObject))
+		else if (const UEnum* Enum = Cast<UEnum>(InCPPTypeObject))
 		{
 			return RigVMTypeUtils::CPPTypeFromEnum(Enum);
 		}
