Avoid having uninitialized properties

diff --git a/Engine/Plugins/Experimental/AnimNext/Source/AnimNextTestSuite/Private/AnimNextDecoratorBaseTest.h b/Engine/Plugins/Experimental/AnimNext/Source/AnimNextTestSuite/Private/AnimNextDecoratorBaseTest.h
--- a/Engine/Plugins/Experimental/AnimNext/Source/AnimNextTestSuite/Private/AnimNextDecoratorBaseTest.h
+++ b/Engine/Plugins/Experimental/AnimNext/Source/AnimNextTestSuite/Private/AnimNextDecoratorBaseTest.h
@@ -47,19 +47,19 @@ struct FDecoratorSerialization_BaseSharedData : public FAnimNextDecoratorSharedD
 	GENERATED_BODY()
 
 	UPROPERTY()
-	int32 Integer;
+	int32 Integer = 0;
 
 	UPROPERTY()
-	int32 IntegerArray[4];
+	int32 IntegerArray[4] = { 0, 0, 0, 0 };
 
 	UPROPERTY()
 	TArray<int32> IntegerTArray;
 
 	UPROPERTY()
-	FVector Vector;
+	FVector Vector = FVector::ZeroVector;
 
 	UPROPERTY()
-	FVector VectorArray[2];
+	FVector VectorArray[2] = { FVector::ZeroVector, FVector::ZeroVector };
 
 	UPROPERTY()
 	TArray<FVector> VectorTArray;
@@ -77,19 +77,19 @@ struct FDecoratorSerialization_AddSharedData : public FAnimNextDecoratorSharedDa
 	GENERATED_BODY()
 
 	UPROPERTY()
-	int32 Integer;
+	int32 Integer = 0;
 
 	UPROPERTY()
-	int32 IntegerArray[4];
+	int32 IntegerArray[4] = { 0, 0, 0, 0 };
 
 	UPROPERTY()
 	TArray<int32> IntegerTArray;
 
 	UPROPERTY()
-	FVector Vector;
+	FVector Vector = FVector::ZeroVector;
 
 	UPROPERTY()
-	FVector VectorArray[2];
+	FVector VectorArray[2] = { FVector::ZeroVector, FVector::ZeroVector };
 
 	UPROPERTY()
 	TArray<FVector> VectorTArray;
@@ -107,19 +107,19 @@ struct FDecoratorNativeSerialization_AddSharedData : public FAnimNextDecoratorSh
 	GENERATED_BODY()
 
 	UPROPERTY()
-	int32 Integer;
+	int32 Integer = 0;
 
 	UPROPERTY()
-	int32 IntegerArray[4];
+	int32 IntegerArray[4] = { 0, 0, 0, 0 };
 
 	UPROPERTY()
 	TArray<int32> IntegerTArray;
 
 	UPROPERTY()
-	FVector Vector;
+	FVector Vector = FVector::ZeroVector;
 
 	UPROPERTY()
-	FVector VectorArray[2];
+	FVector VectorArray[2] = { FVector::ZeroVector, FVector::ZeroVector };
 
 	UPROPERTY()
 	TArray<FVector> VectorTArray;
