Fix double-free when compiling an anim BP with an AnimNextGraph node in it

diff --git a/Engine/Plugins/Experimental/AnimNext/Source/AnimNext/Internal/Graph/AnimGraph/AnimNode_AnimNextGraph.h b/Engine/Plugins/Experimental/AnimNext/Source/AnimNext/Internal/Graph/AnimGraph/AnimNode_AnimNextGraph.h
--- a/Engine/Plugins/Experimental/AnimNext/Source/AnimNext/Internal/Graph/AnimGraph/AnimNode_AnimNextGraph.h
+++ b/Engine/Plugins/Experimental/AnimNext/Source/AnimNext/Internal/Graph/AnimGraph/AnimNode_AnimNextGraph.h
@@ -8,7 +8,19 @@
 #include "AnimNode_AnimNextGraph.generated.h"
 
 class UNodeMappingContainer;
-struct FAnimSequencePlayerState;
+
+// TEST - until we can allocate per-node state again
+USTRUCT()
+struct FAnimSequencePlayerState
+{
+	GENERATED_BODY()
+
+	UPROPERTY()
+	float InternalTimeAccumulator = 0.0f;
+
+	UPROPERTY()
+	float PrevInternalTimeAccumulator = 0.0f;
+};
 
 /**
  * Animation node that allows a AnimNextGraph output to be used in an animation graph
@@ -19,7 +31,6 @@ struct ANIMNEXT_API FAnimNode_AnimNextGraph : public FAnimNode_CustomProperty
 	GENERATED_BODY()
 
 	FAnimNode_AnimNextGraph();
-	~FAnimNode_AnimNextGraph();
 
 	// FAnimNode_Base interface
 	virtual void OnInitializeAnimInstance(const FAnimInstanceProxy* InProxy, const UAnimInstance* InAnimInstance) override;
@@ -49,7 +60,7 @@ private:
 	UPROPERTY(EditAnywhere, Category = Settings, meta = (PinHiddenByDefault, DisallowedClasses = "/Script/Engine.AnimMontage"))
 	TObjectPtr<UAnimSequence> TestSequence = nullptr;
 
-	FAnimSequencePlayerState* SequencePlayerState;	// TEST - until we can allocate per-node state again
+	FAnimSequencePlayerState SequencePlayerState;	// TEST - until we can allocate per-node state again
 
 	/*
 	 * Max LOD that this node is allowed to run
diff --git a/Engine/Plugins/Experimental/AnimNext/Source/AnimNext/Private/Graph/AnimGraph/AnimNode_AnimNextGraph.cpp b/Engine/Plugins/Experimental/AnimNext/Source/AnimNext/Private/Graph/AnimGraph/AnimNode_AnimNextGraph.cpp
--- a/Engine/Plugins/Experimental/AnimNext/Source/AnimNext/Private/Graph/AnimGraph/AnimNode_AnimNextGraph.cpp
+++ b/Engine/Plugins/Experimental/AnimNext/Source/AnimNext/Private/Graph/AnimGraph/AnimNode_AnimNextGraph.cpp
@@ -28,12 +28,7 @@ FAnimNode_AnimNextGraph::FAnimNode_AnimNextGraph()
 	, AnimNextGraph(nullptr)
 	, LODThreshold(INDEX_NONE)
 {
-	SequencePlayerState = new FAnimSequencePlayerState();
-}
 
-FAnimNode_AnimNextGraph::~FAnimNode_AnimNextGraph()
-{
-	delete SequencePlayerState;
 }
 
 void FAnimNode_AnimNextGraph::OnInitializeAnimInstance(const FAnimInstanceProxy* InProxy, const UAnimInstance* InAnimInstance)
@@ -128,7 +123,7 @@ void FAnimNode_AnimNextGraph::Evaluate_AnyThread(FPoseContext & Output)
 	FGenerationTools::RemapPose(LODLevel, SourcePose, RefPose, GraphSourceLODPose.LODPose);
 
 	Context.GetMutableParamStack().PushValues(
-		"AnimSequencePlayerState", *SequencePlayerState,
+		"AnimSequencePlayerState", SequencePlayerState,
 		"GraphReferencePose", GraphReferencePose,
 		"ResultPose", ResultPose,
 		"GraphLODLevel", LODLevel,
diff --git a/Engine/Plugins/Experimental/AnimNext/Source/AnimNext/Private/Graph/RigUnit_AnimNextAnimSequence.h b/Engine/Plugins/Experimental/AnimNext/Source/AnimNext/Private/Graph/RigUnit_AnimNextAnimSequence.h
--- a/Engine/Plugins/Experimental/AnimNext/Source/AnimNext/Private/Graph/RigUnit_AnimNextAnimSequence.h
+++ b/Engine/Plugins/Experimental/AnimNext/Source/AnimNext/Private/Graph/RigUnit_AnimNextAnimSequence.h
@@ -4,22 +4,11 @@
 
 #include "Graph/GraphExecuteContext.h"
 #include "Graph/AnimNext_LODPose.h"
+#include "Graph/AnimGraph/AnimNode_AnimNextGraph.h"	// TEST - until we can allocate per-node state again
 #include "RigUnit_AnimNextAnimSequence.generated.h"
 
 class UAnimSequenceBase;
 
-USTRUCT()
-struct FAnimSequencePlayerState
-{
-	GENERATED_BODY()
-
-	UPROPERTY()
-	float InternalTimeAccumulator = 0.0f;
-
-	UPROPERTY()
-	float PrevInternalTimeAccumulator = 0.0f;
-};
-
 USTRUCT(BlueprintType, meta = (DisplayName = "AnimSequence"))
 struct FAnimNextGraph_AnimSequence
 {
