AnimNext params in Anim BPs

Added a new anim node: AnimNext Parameters. This injects parameters into the AnimNext stack for all leafwards nodes.

Added a new way of binding to parameters on anim nodes.
- Added indirection to UAnimGraphNode_Base to allow different 'binding types' to be authored by deriving from UAnimGraphNodeBinding
- Added new binding type for AnimNext parameters
- Moved existing binding code (inc UI widget creation) into UAnimGraphNodeBinding_Base
- Binding types can be selected on a per-node or per-anim BP basis

Reworked FParamStack API a little
- Allow for better error checking when pushing/popping stack layers. Pushed layers can now only be subsequently popped by passing in the handle of the pushed layer.
- Standalone layers are now wrapped in an opaque handle rather than returning a unique ptr
- GetParamData APIs now perform more involved type checking, allowing derived object types & type conversions to be implemented

Improved parameter type sandboxing for automated tests. If running low-level type tests while execution was happening on another thread, the editor could crash because of invalidating already-existing types.

Lots of other small fixes to get workflows nicer and end-to-end functionality working

diff --git a/Engine/Plugins/Animation/ControlRig/Source/ControlRigDeveloper/Private/AnimGraphNode_ControlRig.cpp b/Engine/Plugins/Animation/ControlRig/Source/ControlRigDeveloper/Private/AnimGraphNode_ControlRig.cpp
--- a/Engine/Plugins/Animation/ControlRig/Source/ControlRigDeveloper/Private/AnimGraphNode_ControlRig.cpp
+++ b/Engine/Plugins/Animation/ControlRig/Source/ControlRigDeveloper/Private/AnimGraphNode_ControlRig.cpp
@@ -768,7 +768,7 @@ void UAnimGraphNode_ControlRig::PostEditChangeProperty(FPropertyChangedEvent& Pr
 					if (Node.AlphaInputType != EAnimAlphaInputType::Float)
 					{
 						Pin->BreakAllPinLinks();
-						PropertyBindings.Remove(Pin->PinName);
+						RemoveBindings(Pin->PinName);
 					}
 				}
 				else if (Pin->PinName == GET_MEMBER_NAME_STRING_CHECKED(FAnimNode_ControlRig, bAlphaBoolEnabled))
@@ -776,7 +776,7 @@ void UAnimGraphNode_ControlRig::PostEditChangeProperty(FPropertyChangedEvent& Pr
 					if (Node.AlphaInputType != EAnimAlphaInputType::Bool)
 					{
 						Pin->BreakAllPinLinks();
-						PropertyBindings.Remove(Pin->PinName);
+						RemoveBindings(Pin->PinName);
 					}
 				}
 				else if (Pin->PinName == GET_MEMBER_NAME_STRING_CHECKED(FAnimNode_ControlRig, AlphaCurveName))
@@ -784,7 +784,7 @@ void UAnimGraphNode_ControlRig::PostEditChangeProperty(FPropertyChangedEvent& Pr
 					if (Node.AlphaInputType != EAnimAlphaInputType::Curve)
 					{
 						Pin->BreakAllPinLinks();
-						PropertyBindings.Remove(Pin->PinName);
+						RemoveBindings(Pin->PinName);
 					}
 				}
 			}
diff --git a/Engine/Plugins/Animation/IKRig/Source/IKRigDeveloper/Private/AnimGraphNode_IKRig.cpp b/Engine/Plugins/Animation/IKRig/Source/IKRigDeveloper/Private/AnimGraphNode_IKRig.cpp
--- a/Engine/Plugins/Animation/IKRig/Source/IKRigDeveloper/Private/AnimGraphNode_IKRig.cpp
+++ b/Engine/Plugins/Animation/IKRig/Source/IKRigDeveloper/Private/AnimGraphNode_IKRig.cpp
@@ -376,7 +376,7 @@ void UAnimGraphNode_IKRig::PostEditChangeProperty(struct FPropertyChangedEvent&
 				if (Node.AlphaInputType != EAnimAlphaInputType::Float)
 				{
 					Pin->BreakAllPinLinks();
-					PropertyBindings.Remove(Pin->PinName);
+					RemoveBindings(Pin->PinName);
 				}
 			}
 			else if (Pin->PinName == GET_MEMBER_NAME_STRING_CHECKED(FAnimNode_IKRig, bAlphaBoolEnabled))
@@ -384,7 +384,7 @@ void UAnimGraphNode_IKRig::PostEditChangeProperty(struct FPropertyChangedEvent&
 				if (Node.AlphaInputType != EAnimAlphaInputType::Bool)
 				{
 					Pin->BreakAllPinLinks();
-					PropertyBindings.Remove(Pin->PinName);
+					RemoveBindings(Pin->PinName);
 				}
 			}
 			else if (Pin->PinName == GET_MEMBER_NAME_STRING_CHECKED(FAnimNode_IKRig, AlphaCurveName))
@@ -392,7 +392,7 @@ void UAnimGraphNode_IKRig::PostEditChangeProperty(struct FPropertyChangedEvent&
 				if (Node.AlphaInputType != EAnimAlphaInputType::Curve)
 				{
 					Pin->BreakAllPinLinks();
-					PropertyBindings.Remove(Pin->PinName);
+					RemoveBindings(Pin->PinName);
 				}
 			}
 		}
diff --git a/Engine/Plugins/Experimental/AnimNext/Source/AnimNext/Internal/Graph/AnimGraph/AnimNodeExposedValueHandler_AnimNextParameters.h b/Engine/Plugins/Experimental/AnimNext/Source/AnimNext/Internal/Graph/AnimGraph/AnimNodeExposedValueHandler_AnimNextParameters.h
new file mode 100644
--- /dev/null
+++ b/Engine/Plugins/Experimental/AnimNext/Source/AnimNext/Internal/Graph/AnimGraph/AnimNodeExposedValueHandler_AnimNextParameters.h
@@ -0,0 +1,55 @@
+// Copyright Epic Games, Inc. All Rights Reserved.
+
+#pragma once
+
+#include "Animation/ExposedValueHandler.h"
+#include "PropertyAccess.h"
+#include "Param/ParamId.h"
+#include "Param/ParamType.h"
+#include "Param/ParamTypeHandle.h"
+#include "AnimNodeExposedValueHandler_AnimNextParameters.generated.h"
+
+USTRUCT()
+struct FAnimNodeExposedValueHandler_AnimNextParameters_Entry
+{
+	GENERATED_BODY()
+
+	// Parameter to get
+	UPROPERTY()
+	FName ParameterName;
+
+	// Cached ID for the parameter
+	UE::AnimNext::FParamId ParamId;
+
+	// ParamType of the property that this entry will write to
+	UPROPERTY()
+	FAnimNextParamType PropertyParamType;
+
+	// Cached type handle created from ParamType
+	UE::AnimNext::FParamTypeHandle PropertyParamTypeHandle;
+
+	// Property access index
+	UPROPERTY()
+	int32 AccessIndex;
+
+	// Access operation to perform
+	UPROPERTY()
+	EPropertyAccessCopyType AccessType;
+};
+
+USTRUCT()
+struct FAnimNodeExposedValueHandler_AnimNextParameters : public FAnimNodeExposedValueHandler_Base
+{
+	GENERATED_BODY()
+
+	// FAnimNodeExposedValueHandler interface
+	virtual void Initialize(const UClass* InClass) override;
+	virtual void Execute(const FAnimationBaseContext& InContext) const override;
+
+	UPROPERTY()
+	TArray<FAnimNodeExposedValueHandler_AnimNextParameters_Entry> Entries;
+
+	// Cached property access library ptr
+	const FPropertyAccessLibrary* PropertyAccessLibrary;
+};
+
diff --git a/Engine/Plugins/Experimental/AnimNext/Source/AnimNext/Internal/Graph/AnimGraph/AnimNode_AnimNextParameters.h b/Engine/Plugins/Experimental/AnimNext/Source/AnimNext/Internal/Graph/AnimGraph/AnimNode_AnimNextParameters.h
new file mode 100644
--- /dev/null
+++ b/Engine/Plugins/Experimental/AnimNext/Source/AnimNext/Internal/Graph/AnimGraph/AnimNode_AnimNextParameters.h
@@ -0,0 +1,54 @@
+// Copyright Epic Games, Inc. All Rights Reserved.
+
+#pragma once
+
+#include "CoreMinimal.h"
+#include "Animation/AnimNodeBase.h"
+#include "Param/ParamStack.h"
+#include "AnimNode_AnimNextParameters.generated.h"
+
+class UAnimNextParameterBlock;
+
+USTRUCT(BlueprintInternalUseOnly)
+struct FAnimNode_AnimNextParameters : public FAnimNode_Base
+{
+	GENERATED_BODY()
+
+	friend class UAnimGraphNode_AnimNextParameters;
+
+	FAnimNode_AnimNextParameters() = default;
+	FAnimNode_AnimNextParameters(const FAnimNode_AnimNextParameters& InOther);
+	FAnimNode_AnimNextParameters& operator=(const FAnimNode_AnimNextParameters& InOther);
+	FAnimNode_AnimNextParameters(FAnimNode_AnimNextParameters&& InOther);
+	FAnimNode_AnimNextParameters& operator=(FAnimNode_AnimNextParameters&& InOther);
+	~FAnimNode_AnimNextParameters() = default;
+
+private:
+	UPROPERTY(EditAnywhere, Category = Links)
+	FPoseLink Source;
+
+#if WITH_EDITORONLY_DATA
+	UPROPERTY(EditAnywhere, Category = Settings, meta=(FoldProperty, PinHiddenByDefault))
+	TObjectPtr<UAnimNextParameterBlock> ParameterBlock;
+#endif
+
+	// Cache previous param block so we know when it changes via pin
+	UAnimNextParameterBlock* PreviousParameterBlock = nullptr;
+
+	// Cached layer
+	UE::AnimNext::FParamStack::FLayerHandle ParamLayerHandle;
+
+	// Property bag for layer backing storage
+	UPROPERTY(Transient)
+	FInstancedPropertyBag PropertyBag;
+
+private:
+	// FAnimNode_Base
+	ANIMNEXT_API virtual void Initialize_AnyThread(const FAnimationInitializeContext& Context) override;
+	ANIMNEXT_API virtual void Update_AnyThread(const FAnimationUpdateContext& Context) override;
+	ANIMNEXT_API virtual void CacheBones_AnyThread(const FAnimationCacheBonesContext& Context) override;
+	ANIMNEXT_API virtual void Evaluate_AnyThread(FPoseContext& Output) override;
+	ANIMNEXT_API virtual void GatherDebugData(FNodeDebugData& DebugData) override;
+
+	UAnimNextParameterBlock* GetParameterBlock() const;
+};
diff --git a/Engine/Plugins/Experimental/AnimNext/Source/AnimNext/Internal/Graph/AnimGraph/AnimSubsystem_AnimNextParameters.h b/Engine/Plugins/Experimental/AnimNext/Source/AnimNext/Internal/Graph/AnimGraph/AnimSubsystem_AnimNextParameters.h
new file mode 100644
--- /dev/null
+++ b/Engine/Plugins/Experimental/AnimNext/Source/AnimNext/Internal/Graph/AnimGraph/AnimSubsystem_AnimNextParameters.h
@@ -0,0 +1,14 @@
+// Copyright Epic Games, Inc. All Rights Reserved.
+
+#pragma once
+
+#include "Animation/AnimSubsystem.h"
+#include "CoreMinimal.h"
+#include "AnimSubsystem_AnimNextParameters.generated.h"
+
+USTRUCT()
+struct FAnimSubsystem_AnimNextParameters : public FAnimSubsystem
+{
+	GENERATED_BODY()
+
+};
\ No newline at end of file
diff --git a/Engine/Plugins/Experimental/AnimNext/Source/AnimNext/Internal/Graph/AnimNextExecuteContext.h b/Engine/Plugins/Experimental/AnimNext/Source/AnimNext/Internal/Graph/AnimNextExecuteContext.h
new file mode 100644
--- /dev/null
+++ b/Engine/Plugins/Experimental/AnimNext/Source/AnimNext/Internal/Graph/AnimNextExecuteContext.h
@@ -0,0 +1,53 @@
+﻿// Copyright Epic Games, Inc. All Rights Reserved.
+
+#pragma once
+
+#include "CoreMinimal.h"
+#include "RigVMCore/RigVMExecuteContext.h"
+#include "Units/RigUnit.h"
+#include "AnimNextExecuteContext.generated.h"
+
+namespace UE::AnimNext
+{
+	struct FContext;
+}
+
+USTRUCT(BlueprintType)
+struct FAnimNextExecuteContext : public FRigVMExecuteContext
+{
+	GENERATED_BODY()
+
+	FAnimNextExecuteContext()
+		: FRigVMExecuteContext()
+		, Context(nullptr)
+	{
+	}
+
+	const UE::AnimNext::FContext& GetContext() const
+	{
+		check(Context);
+		return *Context;
+	}
+
+	void SetContextData(const UE::AnimNext::FContext& InContext)
+	{
+		Context = &InContext;
+	}
+
+	virtual void Copy(const FRigVMExecuteContext* InOtherContext) override
+	{
+		Super::Copy(InOtherContext);
+
+		const FAnimNextExecuteContext* OtherContext = (const FAnimNextExecuteContext*)InOtherContext;
+		Context = OtherContext->Context;
+	}
+
+private:
+	const UE::AnimNext::FContext* Context;
+};
+
+USTRUCT(meta=(ExecuteContext="FAnimNextExecuteContext"))
+struct FRigUnit_AnimNextBase : public FRigUnit
+{
+	GENERATED_BODY()
+};
diff --git a/Engine/Plugins/Experimental/AnimNext/Source/AnimNext/Internal/Graph/RigUnit_AnimNextBeginExecution.h b/Engine/Plugins/Experimental/AnimNext/Source/AnimNext/Internal/Graph/RigUnit_AnimNextBeginExecution.h
--- a/Engine/Plugins/Experimental/AnimNext/Source/AnimNext/Internal/Graph/RigUnit_AnimNextBeginExecution.h
+++ b/Engine/Plugins/Experimental/AnimNext/Source/AnimNext/Internal/Graph/RigUnit_AnimNextBeginExecution.h
@@ -3,7 +3,7 @@
 #pragma once
 
 #include "Units/RigUnit.h"
-#include "GraphExecuteContext.h"
+#include "AnimNextExecuteContext.h"
 #include "RigUnit_AnimNextBeginExecution.generated.h"
 
 /**
@@ -22,7 +22,7 @@ struct ANIMNEXT_API FRigUnit_AnimNextBeginExecution : public FRigUnit_AnimNextBa
 
 	// The execution result
 	UPROPERTY(EditAnywhere, DisplayName = "Execute", Category = "BeginExecution", meta = (Output))
-	FAnimNextGraphExecuteContext ExecuteContext;
+	FAnimNextExecuteContext ExecuteContext;
 
 	static FName EventName;
 };
diff --git a/Engine/Plugins/Experimental/AnimNext/Source/AnimNext/Internal/Graph/RigUnit_AnimNextEndExecution.h b/Engine/Plugins/Experimental/AnimNext/Source/AnimNext/Internal/Graph/RigUnit_AnimNextEndExecution.h
--- a/Engine/Plugins/Experimental/AnimNext/Source/AnimNext/Internal/Graph/RigUnit_AnimNextEndExecution.h
+++ b/Engine/Plugins/Experimental/AnimNext/Source/AnimNext/Internal/Graph/RigUnit_AnimNextEndExecution.h
@@ -3,7 +3,7 @@
 #pragma once
 
 #include "Units/RigUnit.h"
-#include "GraphExecuteContext.h"
+#include "AnimNextExecuteContext.h"
 #include "AnimNext_LODPose.h"
 #include "RigUnit_AnimNextEndExecution.generated.h"
 
@@ -18,7 +18,7 @@ struct ANIMNEXT_API FRigUnit_AnimNextEndExecution : public FRigUnit_AnimNextBase
 public:
 	// The execution result
 	UPROPERTY(EditAnywhere, Transient, DisplayName = "End Execute", Category = "EndExecution", meta = (Input))
-	FAnimNextGraphExecuteContext ExecuteContext;
+	FAnimNextExecuteContext ExecuteContext;
 };
 
 /** Event for writing back a calculated bool */
diff --git a/Engine/Plugins/Experimental/AnimNext/Source/AnimNext/Internal/Param/AnimNextParameterBlock.h b/Engine/Plugins/Experimental/AnimNext/Source/AnimNext/Internal/Param/AnimNextParameterBlock.h
--- a/Engine/Plugins/Experimental/AnimNext/Source/AnimNext/Internal/Param/AnimNextParameterBlock.h
+++ b/Engine/Plugins/Experimental/AnimNext/Source/AnimNext/Internal/Param/AnimNextParameterBlock.h
@@ -9,6 +9,11 @@
 
 class UEdGraph;
 
+namespace UE::AnimNext
+{
+	struct FContext;
+}
+
 namespace UE::AnimNext::UncookedOnly
 {
 	struct FUtils;
@@ -63,6 +68,9 @@ class UAnimNextParameterBlock : public UObject
 	friend struct UE::AnimNext::UncookedOnly::FUtilsPrivate;
 	friend class UE::AnimNext::Editor::FParametersEditor;
 	friend struct UE::AnimNext::Editor::FUtils;
+	friend struct FAnimNode_AnimNextParameters;
+
+	void Run(const UE::AnimNext::FContext& Context) const;
 
 	// UObject interface
 	virtual void PostRename(UObject* OldOuter, const FName OldName) override;
diff --git a/Engine/Plugins/Experimental/AnimNext/Source/AnimNext/Internal/Param/ParamUtils.h b/Engine/Plugins/Experimental/AnimNext/Source/AnimNext/Internal/Param/ParamUtils.h
new file mode 100644
--- /dev/null
+++ b/Engine/Plugins/Experimental/AnimNext/Source/AnimNext/Internal/Param/ParamUtils.h
@@ -0,0 +1,28 @@
+// Copyright Epic Games, Inc. All Rights Reserved.
+
+#pragma once
+
+#include "CoreMinimal.h"
+
+struct FAnimNextParamType;
+
+namespace UE::AnimNext
+{
+
+struct FParamTypeHandle;
+struct FParamCompatibility;
+
+struct FParamUtils
+{
+	// Check to see if a parameter type handle is compatible with another. The order of the parameters
+	// imply directionality, e.g. IsCompatible((int64)InLHS, (int32)InRHS) is allowed as no data loss 
+	// occurs, but IsCompatible((int32)InLHS, (int64)InRHS) is not as B could be truncated.
+	static ANIMNEXT_API FParamCompatibility GetCompatibility(const FParamTypeHandle& InLHS, const FParamTypeHandle& InRHS);
+
+	// Check to see if a parameter handle is compatible with another. The order of the parameters
+	// imply directionality, e.g. IsCompatible((int64)InLHS, (int32)InRHS) is allowed as no data loss 
+	// occurs, but IsCompatible((int32)InLHS, (int64)InRHS) is not as B could be truncated.
+	static ANIMNEXT_API FParamCompatibility GetCompatibility(const FAnimNextParamType& InLHS, const FAnimNextParamType& InRHS);
+};
+
+}
\ No newline at end of file
diff --git a/Engine/Plugins/Experimental/AnimNext/Source/AnimNext/Internal/Param/RigVMDispatch_GetParameter.h b/Engine/Plugins/Experimental/AnimNext/Source/AnimNext/Internal/Param/RigVMDispatch_GetParameter.h
--- a/Engine/Plugins/Experimental/AnimNext/Source/AnimNext/Internal/Param/RigVMDispatch_GetParameter.h
+++ b/Engine/Plugins/Experimental/AnimNext/Source/AnimNext/Internal/Param/RigVMDispatch_GetParameter.h
@@ -3,7 +3,7 @@
 #pragma once
 
 #include "RigVMCore/RigVMDispatchFactory.h"
-#include "Graph/GraphExecuteContext.h"
+#include "Graph/AnimNextExecuteContext.h"
 #include "RigVMDispatch_GetParameter.generated.h"
 
 namespace UE::AnimNext::UncookedOnly
@@ -24,10 +24,11 @@ struct ANIMNEXT_API FRigVMDispatch_GetParameter : public FRigVMDispatchFactory
 private:
 	friend struct UE::AnimNext::UncookedOnly::FUtils;
 
-	virtual UScriptStruct* GetExecuteContextStruct() const { return FAnimNextGraphExecuteContext::StaticStruct(); }
+	virtual UScriptStruct* GetExecuteContextStruct() const { return FAnimNextExecuteContext::StaticStruct(); }
 	virtual FName GetArgumentNameForOperandIndex(int32 InOperandIndex, int32 InTotalOperands) const override;
 	virtual TArray<FRigVMTemplateArgument> GetArguments() const override;
 #if WITH_EDITOR
+	virtual FString GetArgumentMetaData(const FName& InArgumentName, const FName& InMetaDataKey) const override;
 	virtual FString GetArgumentDefaultValue(const FName& InArgumentName, TRigVMTypeIndex InTypeIndex) const override;
 #endif
 	virtual FRigVMTemplateTypeMap OnNewArgumentType(const FName& InArgumentName, TRigVMTypeIndex InTypeIndex) const override;
diff --git a/Engine/Plugins/Experimental/AnimNext/Source/AnimNext/Internal/Param/RigVMDispatch_SetParameter.h b/Engine/Plugins/Experimental/AnimNext/Source/AnimNext/Internal/Param/RigVMDispatch_SetParameter.h
--- a/Engine/Plugins/Experimental/AnimNext/Source/AnimNext/Internal/Param/RigVMDispatch_SetParameter.h
+++ b/Engine/Plugins/Experimental/AnimNext/Source/AnimNext/Internal/Param/RigVMDispatch_SetParameter.h
@@ -3,7 +3,7 @@
 #pragma once
 
 #include "RigVMCore/RigVMDispatchFactory.h"
-#include "Graph/GraphExecuteContext.h"
+#include "Graph/AnimNextExecuteContext.h"
 #include "RigVMDispatch_SetParameter.generated.h"
 
 namespace UE::AnimNext::UncookedOnly
@@ -24,10 +24,11 @@ struct ANIMNEXT_API FRigVMDispatch_SetParameter : public FRigVMDispatchFactory
 private:
 	friend struct UE::AnimNext::UncookedOnly::FUtils;
 
-	virtual UScriptStruct* GetExecuteContextStruct() const { return FAnimNextGraphExecuteContext::StaticStruct(); }
+	virtual UScriptStruct* GetExecuteContextStruct() const { return FAnimNextExecuteContext::StaticStruct(); }
 	virtual FName GetArgumentNameForOperandIndex(int32 InOperandIndex, int32 InTotalOperands) const override;
 	virtual TArray<FRigVMTemplateArgument> GetArguments() const override;
 #if WITH_EDITOR
+	virtual FString GetArgumentMetaData(const FName& InArgumentName, const FName& InMetaDataKey) const override;
 	virtual FString GetArgumentDefaultValue(const FName& InArgumentName, TRigVMTypeIndex InTypeIndex) const override;
 #endif
 	virtual TArray<FRigVMExecuteArgument> GetExecuteArguments_Impl(const FRigVMDispatchContext& InContext) const override;
diff --git a/Engine/Plugins/Experimental/AnimNext/Source/AnimNext/Private/Graph/AnimGraph/AnimNodeExposedValueHandler_AnimNextParameters.cpp b/Engine/Plugins/Experimental/AnimNext/Source/AnimNext/Private/Graph/AnimGraph/AnimNodeExposedValueHandler_AnimNextParameters.cpp
new file mode 100644
--- /dev/null
+++ b/Engine/Plugins/Experimental/AnimNext/Source/AnimNext/Private/Graph/AnimGraph/AnimNodeExposedValueHandler_AnimNextParameters.cpp
@@ -0,0 +1,460 @@
+// Copyright Epic Games, Inc. All Rights Reserved.
+
+#include "Graph/AnimGraph/AnimNodeExposedValueHandler_AnimNextParameters.h"
+#include "Param/ParamStack.h"
+#include "Animation/AnimNodeBase.h"
+#include "Animation/AnimSubsystem_PropertyAccess.h"
+#include "Animation/AnimClassInterface.h"
+
+void FAnimNodeExposedValueHandler_AnimNextParameters::Initialize(const UClass* InClass)
+{
+	using namespace UE::AnimNext;
+
+	Super::Initialize(InClass);
+
+	if (const FAnimSubsystem_PropertyAccess* Subsystem = IAnimClassInterface::GetFromClass(InClass)->FindSubsystem<FAnimSubsystem_PropertyAccess>())
+	{
+		PropertyAccessLibrary = &Subsystem->GetLibrary();
+	}
+
+	for (FAnimNodeExposedValueHandler_AnimNextParameters_Entry& Entry : Entries)
+	{
+		Entry.ParamId = FParamId(Entry.ParameterName);
+		Entry.PropertyParamTypeHandle = Entry.PropertyParamType.GetHandle();
+	}
+}
+
+namespace UE::AnimNext::Private
+{
+	template <typename SourceType, typename DestinationType>
+	static void CopyAndCastFloatingPointArray(
+		const TArray<SourceType>& SourceArray,
+		const FArrayProperty* DestinationArrayProperty,
+		void* DestinationAddress)
+	{
+		checkSlow(DestinationArrayProperty);
+		checkSlow(DestinationAddress);
+
+		FScriptArrayHelper DestinationArrayHelper(DestinationArrayProperty, DestinationAddress);
+
+		DestinationArrayHelper.Resize(SourceArray.Num());
+		for (int32 i = 0; i < SourceArray.Num(); ++i)
+		{
+			const SourceType& SourceData = SourceArray[i];
+			DestinationType* DestinationData = reinterpret_cast<DestinationType*>(DestinationArrayHelper.GetRawPtr(i));
+
+			*DestinationData = static_cast<DestinationType>(SourceData);
+		}
+	}
+
+	static void PerformCopy(EPropertyAccessCopyType InCopyType, const FProperty* InDestProperty, void* InDestAddr, const void* InSrcAddr)
+	{
+		switch (InCopyType)
+		{
+		case EPropertyAccessCopyType::None:
+			break;
+		case EPropertyAccessCopyType::Plain:
+			checkSlow(InDestProperty->PropertyFlags & CPF_IsPlainOldData);
+			FMemory::Memcpy(InDestAddr, InSrcAddr, InDestProperty->ElementSize);
+			break;
+		case EPropertyAccessCopyType::Complex:
+			InDestProperty->CopyCompleteValue(InDestAddr, InSrcAddr);
+			break;
+		case EPropertyAccessCopyType::Bool:
+			checkSlow(InDestProperty->IsA<FBoolProperty>());
+			static_cast<const FBoolProperty*>(InDestProperty)->SetPropertyValue(InDestAddr, *reinterpret_cast<const bool*>(InSrcAddr));
+			break;
+		case EPropertyAccessCopyType::Struct:
+			checkSlow(InDestProperty->IsA<FStructProperty>());
+			static_cast<const FStructProperty*>(InDestProperty)->Struct->CopyScriptStruct(InDestAddr, InSrcAddr);
+			break;
+		case EPropertyAccessCopyType::Name:
+			checkSlow(InDestProperty->IsA<FNameProperty>());
+			static_cast<const FNameProperty*>(InDestProperty)->SetPropertyValue(InDestAddr, *reinterpret_cast<const FName*>(InSrcAddr));
+			break;
+		case EPropertyAccessCopyType::Object:
+			checkSlow(InDestProperty->IsA<FObjectPropertyBase>());
+			static_cast<const FObjectPropertyBase*>(InDestProperty)->SetObjectPropertyValue(InDestAddr, *reinterpret_cast<UObject* const*>(InSrcAddr));
+			break;
+		case EPropertyAccessCopyType::Array:
+		{
+			checkSlow(InDestProperty->IsA<FArrayProperty>());
+			const FArrayProperty* ArrayProperty = static_cast<const FArrayProperty*>(InDestProperty);
+			FScriptArrayHelper SourceArrayHelper(ArrayProperty, InSrcAddr);
+			FScriptArrayHelper DestArrayHelper(ArrayProperty, InDestAddr);
+
+			// Copy the minimum number of elements to the destination array without resizing
+			const int32 MinSize = FMath::Min(SourceArrayHelper.Num(), DestArrayHelper.Num());
+			for (int32 ElementIndex = 0; ElementIndex < MinSize; ++ElementIndex)
+			{
+				ArrayProperty->Inner->CopySingleValue(DestArrayHelper.GetRawPtr(ElementIndex), SourceArrayHelper.GetRawPtr(ElementIndex));
+			}
+			break;
+		}
+		case EPropertyAccessCopyType::PromoteBoolToByte:
+			checkSlow(InDestProperty->IsA<FByteProperty>());
+			static_cast<const FByteProperty*>(InDestProperty)->SetPropertyValue(InDestAddr, (uint8)(*reinterpret_cast<const bool*>(InSrcAddr)));
+			break;
+		case EPropertyAccessCopyType::PromoteBoolToInt32:
+			checkSlow(InDestProperty->IsA<FIntProperty>());
+			static_cast<const FIntProperty*>(InDestProperty)->SetPropertyValue(InDestAddr, (int32)(*reinterpret_cast<const bool*>(InSrcAddr)));
+			break;
+		case EPropertyAccessCopyType::PromoteBoolToInt64:
+			checkSlow(InDestProperty->IsA<FInt64Property>());
+			static_cast<const FInt64Property*>(InDestProperty)->SetPropertyValue(InDestAddr, (int64)(*reinterpret_cast<const bool*>(InSrcAddr)));
+			break;
+		case EPropertyAccessCopyType::PromoteBoolToFloat:
+			checkSlow(InDestProperty->IsA<FFloatProperty>());
+			static_cast<const FFloatProperty*>(InDestProperty)->SetPropertyValue(InDestAddr, (float)(*reinterpret_cast<const bool*>(InSrcAddr)));
+			break;
+		case EPropertyAccessCopyType::PromoteBoolToDouble:
+			checkSlow(InDestProperty->IsA<FDoubleProperty>());
+			static_cast<const FDoubleProperty*>(InDestProperty)->SetPropertyValue(InDestAddr, (double)(*reinterpret_cast<const bool*>(InSrcAddr)));
+			break;
+		case EPropertyAccessCopyType::PromoteByteToInt32:
+			checkSlow(InDestProperty->IsA<FIntProperty>());
+			static_cast<const FIntProperty*>(InDestProperty)->SetPropertyValue(InDestAddr, (int32)(*reinterpret_cast<const uint8*>(InSrcAddr)));
+			break;
+		case EPropertyAccessCopyType::PromoteByteToInt64:
+			checkSlow(InDestProperty->IsA<FInt64Property>());
+			static_cast<const FInt64Property*>(InDestProperty)->SetPropertyValue(InDestAddr, (int64)(*reinterpret_cast<const uint8*>(InSrcAddr)));
+			break;
+		case EPropertyAccessCopyType::PromoteByteToFloat:
+			checkSlow(InDestProperty->IsA<FFloatProperty>());
+			static_cast<const FFloatProperty*>(InDestProperty)->SetPropertyValue(InDestAddr, (float)(*reinterpret_cast<const uint8*>(InSrcAddr)));
+			break;
+		case EPropertyAccessCopyType::PromoteByteToDouble:
+			checkSlow(InDestProperty->IsA<FDoubleProperty>());
+			static_cast<const FDoubleProperty*>(InDestProperty)->SetPropertyValue(InDestAddr, (double)(*reinterpret_cast<const uint8*>(InSrcAddr)));
+			break;
+		case EPropertyAccessCopyType::PromoteInt32ToInt64:
+			checkSlow(InDestProperty->IsA<FInt64Property>());
+			static_cast<const FInt64Property*>(InDestProperty)->SetPropertyValue(InDestAddr, (int64)(*reinterpret_cast<const int32*>(InSrcAddr)));
+			break;
+		case EPropertyAccessCopyType::PromoteInt32ToFloat:
+			checkSlow(InDestProperty->IsA<FFloatProperty>());
+			static_cast<const FFloatProperty*>(InDestProperty)->SetPropertyValue(InDestAddr, (float)(*reinterpret_cast<const int32*>(InSrcAddr)));
+			break;
+		case EPropertyAccessCopyType::PromoteInt32ToDouble:
+			checkSlow(InDestProperty->IsA<FDoubleProperty>());
+			static_cast<const FDoubleProperty*>(InDestProperty)->SetPropertyValue(InDestAddr, (double)(*reinterpret_cast<const int32*>(InSrcAddr)));
+			break;
+		case EPropertyAccessCopyType::PromoteFloatToDouble:
+			checkSlow(InDestProperty->IsA<FDoubleProperty>());
+			static_cast<const FDoubleProperty*>(InDestProperty)->SetPropertyValue(InDestAddr, (double)(*reinterpret_cast<const float*>(InSrcAddr)));
+			break;
+		case EPropertyAccessCopyType::DemoteDoubleToFloat:
+			checkSlow(InDestProperty->IsA<FFloatProperty>());
+			static_cast<const FFloatProperty*>(InDestProperty)->SetPropertyValue(InDestAddr, (float)(*reinterpret_cast<const double*>(InSrcAddr)));
+			break;
+		case EPropertyAccessCopyType::PromoteArrayFloatToDouble:
+		{
+			checkSlow(InDestProperty->IsA<FArrayProperty>());
+			const FArrayProperty* DestArrayProperty = ExactCastField<const FArrayProperty>(InDestProperty);
+			Private::CopyAndCastFloatingPointArray<float, double>(*reinterpret_cast<const TArray<float>*>(InSrcAddr), DestArrayProperty, InDestAddr);
+			break;
+		}
+		case EPropertyAccessCopyType::DemoteArrayDoubleToFloat:
+		{
+			checkSlow(InDestProperty->IsA<FArrayProperty>());
+			const FArrayProperty* DestArrayProperty = ExactCastField<const FArrayProperty>(InDestProperty);
+			Private::CopyAndCastFloatingPointArray<double, float>(*reinterpret_cast<const TArray<double>*>(InSrcAddr), DestArrayProperty, InDestAddr);
+			break;
+		}
+		// TODO: no map support in AnimNext parameter at present
+		case EPropertyAccessCopyType::PromoteMapValueFloatToDouble:
+		case EPropertyAccessCopyType::DemoteMapValueDoubleToFloat:
+		default:
+			check(false);
+			break;
+		}
+	}
+
+	static EPropertyAccessCopyType GetCopyTypeForCompatibility(EPropertyAccessCopyType InAccessType, const FParamTypeHandle& InDestTypeHandle, const FParamTypeHandle& InSrcTypeHandle)
+	{
+		switch (InSrcTypeHandle.GetParameterType())
+		{
+		case FParamTypeHandle::EParamType::Bool:
+			switch (InDestTypeHandle.GetParameterType())
+			{
+			case FParamTypeHandle::EParamType::Bool:
+				return InAccessType;
+			}
+			break;
+		case FParamTypeHandle::EParamType::Byte:
+			switch (InDestTypeHandle.GetParameterType())
+			{
+			case FParamTypeHandle::EParamType::Byte:
+				return InAccessType;
+			case FParamTypeHandle::EParamType::Int32:
+				return EPropertyAccessCopyType::PromoteByteToInt32;
+			case FParamTypeHandle::EParamType::Int64:
+				return EPropertyAccessCopyType::PromoteByteToInt64;
+			case FParamTypeHandle::EParamType::Float:
+				return EPropertyAccessCopyType::PromoteByteToFloat;
+			case FParamTypeHandle::EParamType::Double:
+				return EPropertyAccessCopyType::PromoteByteToDouble;
+			}
+			break;
+		case FParamTypeHandle::EParamType::Int32:
+			switch (InDestTypeHandle.GetParameterType())
+			{
+			case FParamTypeHandle::EParamType::Byte:
+				return EPropertyAccessCopyType::None;
+			case FParamTypeHandle::EParamType::Int32:
+				return InAccessType;
+			case FParamTypeHandle::EParamType::Int64:
+				return EPropertyAccessCopyType::PromoteInt32ToInt64;
+			case FParamTypeHandle::EParamType::Float:
+				return EPropertyAccessCopyType::None;
+			case FParamTypeHandle::EParamType::Double:
+				return EPropertyAccessCopyType::PromoteInt32ToDouble;
+			}
+			break;
+		case FParamTypeHandle::EParamType::Int64:
+			switch (InDestTypeHandle.GetParameterType())
+			{
+			case FParamTypeHandle::EParamType::Byte:
+				return EPropertyAccessCopyType::None;
+			case FParamTypeHandle::EParamType::Int32:
+				return EPropertyAccessCopyType::None;
+			case FParamTypeHandle::EParamType::Int64:
+				return InAccessType;
+			case FParamTypeHandle::EParamType::Float:
+				return EPropertyAccessCopyType::None;
+			case FParamTypeHandle::EParamType::Double:
+				return EPropertyAccessCopyType::None;
+			}
+			break;
+		case FParamTypeHandle::EParamType::Float:
+			switch (InDestTypeHandle.GetParameterType())
+			{
+			case FParamTypeHandle::EParamType::Byte:
+				return EPropertyAccessCopyType::None;
+			case FParamTypeHandle::EParamType::Int32:
+				return EPropertyAccessCopyType::None;
+			case FParamTypeHandle::EParamType::Int64:
+				return EPropertyAccessCopyType::None;
+			case FParamTypeHandle::EParamType::Float:
+				return InAccessType;
+			case FParamTypeHandle::EParamType::Double:
+				return EPropertyAccessCopyType::PromoteFloatToDouble;
+			}
+			break;
+		case FParamTypeHandle::EParamType::Double:
+			switch (InDestTypeHandle.GetParameterType())
+			{
+			case FParamTypeHandle::EParamType::Byte:
+				return EPropertyAccessCopyType::None;
+			case FParamTypeHandle::EParamType::Int32:
+				return EPropertyAccessCopyType::None;
+			case FParamTypeHandle::EParamType::Int64:
+				return EPropertyAccessCopyType::None;
+			case FParamTypeHandle::EParamType::Float:
+				return EPropertyAccessCopyType::DemoteDoubleToFloat;
+			case FParamTypeHandle::EParamType::Double:
+				return InAccessType;
+			}
+			break;
+		case FParamTypeHandle::EParamType::Name:
+			switch (InDestTypeHandle.GetParameterType())
+			{
+			case FParamTypeHandle::EParamType::Name:
+				return InAccessType;
+			}
+			break;
+		case FParamTypeHandle::EParamType::String:
+			switch (InDestTypeHandle.GetParameterType())
+			{
+			case FParamTypeHandle::EParamType::String:
+				return InAccessType;;
+			}
+			break;
+		case FParamTypeHandle::EParamType::Text:
+			switch (InDestTypeHandle.GetParameterType())
+			{
+			case FParamTypeHandle::EParamType::Text:
+				return InAccessType;
+			}
+			break;
+		case FParamTypeHandle::EParamType::Vector:
+			switch (InDestTypeHandle.GetParameterType())
+			{
+			case FParamTypeHandle::EParamType::Vector:
+				return InAccessType;
+			}
+			break;
+		case FParamTypeHandle::EParamType::Vector4:
+			switch (InDestTypeHandle.GetParameterType())
+			{
+			case FParamTypeHandle::EParamType::Vector4:
+				return InAccessType;
+			}
+			break;
+		case FParamTypeHandle::EParamType::Quat:
+			switch (InDestTypeHandle.GetParameterType())
+			{
+			case FParamTypeHandle::EParamType::Quat:
+				return InAccessType;
+			}
+			break;
+		case FParamTypeHandle::EParamType::Transform:
+			switch (InDestTypeHandle.GetParameterType())
+			{
+			case FParamTypeHandle::EParamType::Transform:
+				return InAccessType;
+			}
+			break;
+		case FParamTypeHandle::EParamType::Custom:
+			switch (InDestTypeHandle.GetParameterType())
+			{
+			case FParamTypeHandle::EParamType::Custom:
+				{
+					FAnimNextParamType::EValueType ValueTypeDest, ValueTypeSrc;
+					FAnimNextParamType::EContainerType ContainerTypeDest, ContainerTypeSrc;
+					const UObject* ValueTypeObjectDest, * ValueTypeObjectSrc;
+
+					InDestTypeHandle.GetCustomTypeInfo(ValueTypeDest, ContainerTypeDest, ValueTypeObjectDest);
+					InSrcTypeHandle.GetCustomTypeInfo(ValueTypeSrc, ContainerTypeSrc, ValueTypeObjectSrc);
+
+					if (ContainerTypeDest == ContainerTypeSrc)
+					{
+						if (ContainerTypeDest == FAnimNextParamType::EContainerType::Array)
+						{
+							if (ValueTypeDest != ValueTypeSrc)
+							{
+								switch (ValueTypeDest)
+								{
+								case FAnimNextParamType::EValueType::Float:
+									if (ValueTypeSrc == FAnimNextParamType::EValueType::Double)
+									{
+										return EPropertyAccessCopyType::DemoteArrayDoubleToFloat;
+									}
+									break;
+								case FAnimNextParamType::EValueType::Double:
+									if (ValueTypeSrc == FAnimNextParamType::EValueType::Float)
+									{
+										return EPropertyAccessCopyType::PromoteArrayFloatToDouble;
+									}
+									break;
+								case FAnimNextParamType::EValueType::Struct:
+								case FAnimNextParamType::EValueType::Object:
+								case FAnimNextParamType::EValueType::SoftObject:
+								case FAnimNextParamType::EValueType::Class:
+								case FAnimNextParamType::EValueType::SoftClass:
+									if (ValueTypeObjectDest == ValueTypeObjectSrc)
+									{
+										return InAccessType;
+									}
+									break;
+								}
+							}
+							else
+							{
+								switch (ValueTypeDest)
+								{
+								case FAnimNextParamType::EValueType::Bool:
+								case FAnimNextParamType::EValueType::Byte:
+								case FAnimNextParamType::EValueType::Int32:
+								case FAnimNextParamType::EValueType::Int64:
+								case FAnimNextParamType::EValueType::Float:
+								case FAnimNextParamType::EValueType::Double:
+								case FAnimNextParamType::EValueType::Name:
+								case FAnimNextParamType::EValueType::String:
+								case FAnimNextParamType::EValueType::Text:
+									return InAccessType;
+								case FAnimNextParamType::EValueType::Enum:
+								case FAnimNextParamType::EValueType::Struct:
+								case FAnimNextParamType::EValueType::Object:
+								case FAnimNextParamType::EValueType::SoftObject:
+								case FAnimNextParamType::EValueType::Class:
+								case FAnimNextParamType::EValueType::SoftClass:
+									if (ValueTypeObjectDest == ValueTypeObjectSrc)
+									{
+										return InAccessType;
+									}
+								}
+							}
+						}
+						else
+						{
+							if (ValueTypeDest == ValueTypeSrc)
+							{
+								switch (ValueTypeDest)
+								{
+								case FAnimNextParamType::EValueType::Enum:
+									if (ValueTypeObjectDest == ValueTypeObjectSrc)
+									{
+										return InAccessType;
+									}
+									break;
+								case FAnimNextParamType::EValueType::Struct:
+									if (ValueTypeObjectDest == ValueTypeObjectSrc)
+									{
+										return InAccessType;
+									}
+									else if (CastChecked<UScriptStruct>(ValueTypeObjectDest)->IsChildOf(CastChecked<UScriptStruct>(ValueTypeObjectSrc)))
+									{
+										return InAccessType;
+									}
+									break;
+								case FAnimNextParamType::EValueType::Object:
+								case FAnimNextParamType::EValueType::SoftObject:
+								case FAnimNextParamType::EValueType::Class:
+								case FAnimNextParamType::EValueType::SoftClass:
+									if (ValueTypeObjectDest == ValueTypeObjectSrc)
+									{
+										return InAccessType;
+									}
+									else if (CastChecked<UClass>(ValueTypeObjectDest)->IsChildOf(CastChecked<UClass>(ValueTypeObjectSrc)))
+									{
+										return InAccessType;
+									}
+									break;
+								}
+							}
+						}
+					}
+				}
+				break;
+			}
+			break;
+		}
+
+		return EPropertyAccessCopyType::None;
+	}
+}
+
+
+void FAnimNodeExposedValueHandler_AnimNextParameters::Execute(const FAnimationBaseContext& InContext) const
+{
+	using namespace UE::AnimNext;
+
+	Super::Execute(InContext);
+
+	FParamStack& ParamStack = FParamStack::Get();
+	UObject* Object = InContext.GetAnimInstanceObject();
+
+	for (const FAnimNodeExposedValueHandler_AnimNextParameters_Entry& Entry : Entries)
+	{
+		EPropertyAccessCopyType CopyType = Entry.AccessType;
+		FParamTypeHandle FoundTypeHandle;
+		TConstArrayView<uint8> Value;
+		FParamStack::EGetParamResult Result = ParamStack.GetParamData(Entry.ParamId, Entry.PropertyParamTypeHandle, Value, FoundTypeHandle, FParamCompatibility::IncompatibleWithDataLoss());
+		if (Result == FParamStack::EGetParamResult::CompatibleType)
+		{
+			// Compatibility indicates we need to potentially modify the access type
+			CopyType = Private::GetCopyTypeForCompatibility(CopyType, Entry.PropertyParamTypeHandle, FoundTypeHandle);
+		}
+
+		if(Result == FParamStack::EGetParamResult::CompatibleType || Result == FParamStack::EGetParamResult::Succeeded)
+		{
+			const uint8* SrcAddr = Value.GetData();
+
+			PropertyAccess::GetAccessAddress(Object, *PropertyAccessLibrary, Entry.AccessIndex, [CopyType, SrcAddr](const FProperty* InProperty, void* InAddress)
+			{
+				Private::PerformCopy(CopyType, InProperty, InAddress, SrcAddr);
+			});
+		}
+	}
+}
\ No newline at end of file
diff --git a/Engine/Plugins/Experimental/AnimNext/Source/AnimNext/Private/Graph/AnimGraph/AnimNode_AnimNextGraph.cpp b/Engine/Plugins/Experimental/AnimNext/Source/AnimNext/Private/Graph/AnimGraph/AnimNode_AnimNextGraph.cpp
--- a/Engine/Plugins/Experimental/AnimNext/Source/AnimNext/Private/Graph/AnimGraph/AnimNode_AnimNextGraph.cpp
+++ b/Engine/Plugins/Experimental/AnimNext/Source/AnimNext/Private/Graph/AnimGraph/AnimNode_AnimNextGraph.cpp
@@ -122,7 +122,7 @@ void FAnimNode_AnimNextGraph::Evaluate_AnyThread(FPoseContext & Output)
 	FAnimNextGraphLODPose ResultPose(FLODPose(RefPose, LODLevel, true, Output.ExpectsAdditivePose()));
 	FGenerationTools::RemapPose(LODLevel, SourcePose, RefPose, GraphSourceLODPose.LODPose);
 
-	Context.GetMutableParamStack().PushValues(
+	FParamStack::FPushedLayerHandle LayerHandle = Context.GetMutableParamStack().PushValues(
 		"AnimSequencePlayerState", SequencePlayerState,
 		"GraphReferencePose", GraphReferencePose,
 		"ResultPose", ResultPose,
@@ -137,6 +137,8 @@ void FAnimNode_AnimNextGraph::Evaluate_AnyThread(FPoseContext & Output)
 	FGenerationTools::RemapPose(LODLevel, RefPose, ResultPose.LODPose, Output);
 
 	FAnimNode_CustomProperty::Evaluate_AnyThread(Output);
+
+	Context.GetMutableParamStack().PopLayer(LayerHandle);
 }
 
 void FAnimNode_AnimNextGraph::PostSerialize(const FArchive& Ar)
diff --git a/Engine/Plugins/Experimental/AnimNext/Source/AnimNext/Private/Graph/AnimGraph/AnimNode_AnimNextParameters.cpp b/Engine/Plugins/Experimental/AnimNext/Source/AnimNext/Private/Graph/AnimGraph/AnimNode_AnimNextParameters.cpp
new file mode 100644
--- /dev/null
+++ b/Engine/Plugins/Experimental/AnimNext/Source/AnimNext/Private/Graph/AnimGraph/AnimNode_AnimNextParameters.cpp
@@ -0,0 +1,120 @@
+// Copyright Epic Games, Inc. All Rights Reserved.
+
+#include "Graph/AnimGraph/AnimNode_AnimNextParameters.h"
+#include "Param/ParamStack.h"
+#include "Param/AnimNextParameterBlock.h"
+#include "Context.h"
+
+#include UE_INLINE_GENERATED_CPP_BY_NAME(AnimNode_AnimNextParameters)
+
+FAnimNode_AnimNextParameters::FAnimNode_AnimNextParameters(const FAnimNode_AnimNextParameters& InOther)
+	: Source(InOther.Source)
+#if WITH_EDITORONLY_DATA
+	, ParameterBlock(InOther.ParameterBlock)
+#endif
+	, PreviousParameterBlock(nullptr)
+	, ParamLayerHandle()
+	, PropertyBag()
+{
+}
+
+FAnimNode_AnimNextParameters& FAnimNode_AnimNextParameters::operator=(const FAnimNode_AnimNextParameters& InOther)
+{
+	Source = InOther.Source;
+#if WITH_EDITORONLY_DATA
+	ParameterBlock = InOther.ParameterBlock;
+#endif
+	PreviousParameterBlock = nullptr;
+	ParamLayerHandle.Invalidate();
+	PropertyBag.Reset();
+	return *this;
+}
+
+FAnimNode_AnimNextParameters::FAnimNode_AnimNextParameters(FAnimNode_AnimNextParameters&& InOther)
+{
+	Source = InOther.Source;
+#if WITH_EDITORONLY_DATA
+	ParameterBlock = InOther.ParameterBlock;
+#endif
+	PreviousParameterBlock = nullptr;
+	ParamLayerHandle.Invalidate();
+	PropertyBag.Reset();
+}
+
+FAnimNode_AnimNextParameters& FAnimNode_AnimNextParameters::operator=(FAnimNode_AnimNextParameters&& InOther)
+{
+	Source = InOther.Source;
+#if WITH_EDITORONLY_DATA
+	ParameterBlock = InOther.ParameterBlock;
+#endif
+	PreviousParameterBlock = nullptr;
+	ParamLayerHandle.Invalidate();
+	PropertyBag.Reset();
+	return *this;
+}
+
+void FAnimNode_AnimNextParameters::Initialize_AnyThread(const FAnimationInitializeContext& Context)
+{
+	Source.Initialize(Context);
+}
+
+void FAnimNode_AnimNextParameters::Update_AnyThread(const FAnimationUpdateContext& Context)
+{
+	using namespace UE::AnimNext;
+
+	GetEvaluateGraphExposedInputs().Execute(Context);
+
+	UAnimNextParameterBlock* CurrentParameterBlock = GetParameterBlock();
+
+	// Reconstruct param block's cached layer if required
+	if (CurrentParameterBlock != PreviousParameterBlock || !ParamLayerHandle.IsValid())
+	{
+		ParamLayerHandle.Invalidate();
+
+		if (CurrentParameterBlock)
+		{
+			PropertyBag = CurrentParameterBlock->PropertyBag;
+			ParamLayerHandle = FParamStack::MakeMutableLayer(PropertyBag);
+		}
+
+		PreviousParameterBlock = CurrentParameterBlock;
+	}
+
+	FParamStack& ParamStack = FParamStack::Get();
+	FParamStack::FPushedLayerHandle PushedLayerHandle;
+	if (CurrentParameterBlock && ParamLayerHandle.IsValid())
+	{
+		PushedLayerHandle = ParamStack.PushLayer(ParamLayerHandle);
+		UE::AnimNext::FContext AnimNextContext;
+		CurrentParameterBlock->Run(AnimNextContext);
+	}
+
+	Source.Update(Context);
+
+	if (PushedLayerHandle.IsValid())
+	{
+		ParamStack.PopLayer(PushedLayerHandle);
+	}
+}
+
+void FAnimNode_AnimNextParameters::CacheBones_AnyThread(const FAnimationCacheBonesContext& Context)
+{
+	Source.CacheBones(Context);
+}
+
+void FAnimNode_AnimNextParameters::Evaluate_AnyThread(FPoseContext& Output)
+{
+	Source.Evaluate(Output);
+}
+
+void FAnimNode_AnimNextParameters::GatherDebugData(FNodeDebugData& DebugData)
+{
+	DebugData.AddDebugItem(DebugData.GetNodeName(this));
+
+	Source.GatherDebugData(DebugData);
+}
+
+UAnimNextParameterBlock* FAnimNode_AnimNextParameters::GetParameterBlock() const
+{
+	return GET_ANIM_NODE_DATA(TObjectPtr<UAnimNextParameterBlock>, ParameterBlock);
+}
\ No newline at end of file
diff --git a/Engine/Plugins/Experimental/AnimNext/Source/AnimNext/Private/Graph/AnimNextGraph.cpp b/Engine/Plugins/Experimental/AnimNext/Source/AnimNext/Private/Graph/AnimNextGraph.cpp
--- a/Engine/Plugins/Experimental/AnimNext/Source/AnimNext/Private/Graph/AnimNextGraph.cpp
+++ b/Engine/Plugins/Experimental/AnimNext/Source/AnimNext/Private/Graph/AnimNextGraph.cpp
@@ -21,7 +21,7 @@ void UAnimNextGraph::Run(const UE::AnimNext::FContext& Context) const
 	if(RigVM)
 	{
 		FRigVMExtendedExecuteContext RigVMExtendedExecuteContext;
-		FAnimNextGraphExecuteContext& AnimNextContext = RigVMExtendedExecuteContext.GetPublicDataSafe<FAnimNextGraphExecuteContext>();
+		FAnimNextExecuteContext& AnimNextContext = RigVMExtendedExecuteContext.GetPublicDataSafe<FAnimNextExecuteContext>();
 		AnimNextContext.SetContextData(Context);
 
 		RigVM->Execute(RigVMExtendedExecuteContext, TArray<URigVMMemoryStorage*>(), FRigUnit_AnimNextBeginExecution::EventName);
diff --git a/Engine/Plugins/Experimental/AnimNext/Source/AnimNext/Private/Graph/RigUnit_AnimNextAnimSequence.cpp b/Engine/Plugins/Experimental/AnimNext/Source/AnimNext/Private/Graph/RigUnit_AnimNextAnimSequence.cpp
--- a/Engine/Plugins/Experimental/AnimNext/Source/AnimNext/Private/Graph/RigUnit_AnimNextAnimSequence.cpp
+++ b/Engine/Plugins/Experimental/AnimNext/Source/AnimNext/Private/Graph/RigUnit_AnimNextAnimSequence.cpp
@@ -1,7 +1,7 @@
 // Copyright Epic Games, Inc. All Rights Reserved.
 
 #include "RigUnit_AnimNextAnimSequence.h"
-#include "Graph/GraphExecuteContext.h"
+#include "Graph/AnimNextExecuteContext.h"
 #include "DecompressionTools.h"
 #include "Animation/AnimSequence.h"
 #include "Graph/AnimNext_LODPose.h"
diff --git a/Engine/Plugins/Experimental/AnimNext/Source/AnimNext/Private/Graph/RigUnit_AnimNextAnimSequence.h b/Engine/Plugins/Experimental/AnimNext/Source/AnimNext/Private/Graph/RigUnit_AnimNextAnimSequence.h
--- a/Engine/Plugins/Experimental/AnimNext/Source/AnimNext/Private/Graph/RigUnit_AnimNextAnimSequence.h
+++ b/Engine/Plugins/Experimental/AnimNext/Source/AnimNext/Private/Graph/RigUnit_AnimNextAnimSequence.h
@@ -2,7 +2,7 @@
 
 #pragma once
 
-#include "Graph/GraphExecuteContext.h"
+#include "Graph/AnimNextExecuteContext.h"
 #include "Graph/AnimNext_LODPose.h"
 #include "Graph/AnimGraph/AnimNode_AnimNextGraph.h"	// TEST - until we can allocate per-node state again
 #include "RigUnit_AnimNextAnimSequence.generated.h"
diff --git a/Engine/Plugins/Experimental/AnimNext/Source/AnimNext/Private/Module.cpp b/Engine/Plugins/Experimental/AnimNext/Source/AnimNext/Private/Module.cpp
--- a/Engine/Plugins/Experimental/AnimNext/Source/AnimNext/Private/Module.cpp
+++ b/Engine/Plugins/Experimental/AnimNext/Source/AnimNext/Private/Module.cpp
@@ -7,6 +7,8 @@
 #include "DataRegistry.h"
 #include "DecoratorBase/DecoratorRegistry.h"
 #include "DecoratorBase/NodeTemplateRegistry.h"
+#include "RigVMCore/RigVMRegistry.h"
+#include "Animation/AnimSequence.h"
 
 namespace UE::AnimNext
 {
@@ -16,6 +18,13 @@ class FModule : public IModuleInterface
 public:
 	virtual void StartupModule() override
 	{
+		static UClass* const AllowedAssetTypes[] =
+		{
+			UAnimSequence::StaticClass()
+		};
+
+		FRigVMRegistry::Get().RegisterObjectTypes(AllowedAssetTypes);
+
 		FDataRegistry::Init();
 		FDecoratorRegistry::Init();
 		FNodeTemplateRegistry::Init();
diff --git a/Engine/Plugins/Experimental/AnimNext/Source/AnimNext/Private/Param/AnimNextParameterBlock.cpp b/Engine/Plugins/Experimental/AnimNext/Source/AnimNext/Private/Param/AnimNextParameterBlock.cpp
--- a/Engine/Plugins/Experimental/AnimNext/Source/AnimNext/Private/Param/AnimNextParameterBlock.cpp
+++ b/Engine/Plugins/Experimental/AnimNext/Source/AnimNext/Private/Param/AnimNextParameterBlock.cpp
@@ -1,9 +1,11 @@
-﻿// Copyright Epic Games, Inc. All Rights Reserved.
+// Copyright Epic Games, Inc. All Rights Reserved.
 
 #include "Param/AnimNextParameterBlock.h"
 #include "RigVMCore/RigVMMemoryStorage.h"
 #include "UObject/ObjectSaveContext.h"
 #include "UObject/Package.h"
+#include "Graph/RigUnit_AnimNextBeginExecution.h"
+#include "Graph/AnimNextExecuteContext.h"
 
 #include UE_INLINE_GENERATED_CPP_BY_NAME(AnimNextParameterBlock)
 
@@ -34,6 +36,17 @@ static TArray<UClass*> GetClassObjectsInPackage(UPackage* InPackage)
 
 }
 
+void UAnimNextParameterBlock::Run(const UE::AnimNext::FContext& Context) const
+{
+	if (RigVM)
+	{
+		FRigVMExtendedExecuteContext RigVMExtendedExecuteContext;
+		FAnimNextExecuteContext& AnimNextContext = RigVMExtendedExecuteContext.GetPublicDataSafe<FAnimNextExecuteContext>();
+		AnimNextContext.SetContextData(Context);
+		RigVM->Execute(RigVMExtendedExecuteContext, TArray<URigVMMemoryStorage*>(), FRigUnit_AnimNextBeginExecution::EventName);
+	}
+}
+
 void UAnimNextParameterBlock::PostRename(UObject* OldOuter, const FName OldName)
 {
 	Super::PostRename(OldOuter, OldName);
diff --git a/Engine/Plugins/Experimental/AnimNext/Source/AnimNext/Private/Param/ParamId.cpp b/Engine/Plugins/Experimental/AnimNext/Source/AnimNext/Private/Param/ParamId.cpp
--- a/Engine/Plugins/Experimental/AnimNext/Source/AnimNext/Private/Param/ParamId.cpp
+++ b/Engine/Plugins/Experimental/AnimNext/Source/AnimNext/Private/Param/ParamId.cpp
@@ -8,93 +8,80 @@ namespace UE::AnimNext
 {
 
 static FRWLock GParamIdLock;
-static TArray<FName> GParamIdToName;
-static TMap<FName, uint32> GNameToParamId;
-#if WITH_DEV_AUTOMATION_TESTS
-static bool bGTestSandbox = false;
-struct FSandboxThreadData : TThreadSingleton<FSandboxThreadData>
+
+struct FParamIdGlobalData
 {
 	TArray<FName> ParamIdToName;
 	TMap<FName, uint32> NameToParamId;
 };
+
+static FParamIdGlobalData GParamIdGlobalData;
+
+#if WITH_DEV_AUTOMATION_TESTS
+static FParamIdGlobalData GSandboxedParamIdGlobalData;
+static uint32 GParamIdSandboxedThreadId = MAX_uint32;
+static bool bGParamIdSandboxed = false;
 #endif
 
-FParamId::FParamId(FName InName)
+static FParamIdGlobalData& GetParamIdData()
 {
-#if WITH_DEV_AUTOMATION_TESTS
-	if (bGTestSandbox)
+#if WITH_DEV_AUTOMATION_TESTS		
+	if (bGParamIdSandboxed && GParamIdSandboxedThreadId == FPlatformTLS::GetCurrentThreadId())
 	{
-		if (const uint32* FoundIndex = FSandboxThreadData::Get().NameToParamId.Find(InName))
-		{
-			ParameterIndex = *FoundIndex;
-		}
-		else
-		{
-			ParameterIndex = FSandboxThreadData::Get().ParamIdToName.Num();
-			FSandboxThreadData::Get().NameToParamId.Add(InName, ParameterIndex);
-			FSandboxThreadData::Get().ParamIdToName.Add(InName);
-		}
+		return GSandboxedParamIdGlobalData;
 	}
 	else
 #endif
 	{
-		FRWScopeLock Lock(GParamIdLock, SLT_Write);
-		if (const uint32* FoundIndex = GNameToParamId.Find(InName))
-		{
-			ParameterIndex = *FoundIndex;
-		}
-		else
-		{
-			ParameterIndex = GParamIdToName.Num();
-			GNameToParamId.Add(InName, ParameterIndex);
-			GParamIdToName.Add(InName);
-		}
+		return GParamIdGlobalData;
 	}
 }
 
-FName FParamId::ToName() const
+FParamId::FParamId(FName InName)
 {
-#if WITH_DEV_AUTOMATION_TESTS
-	if (bGTestSandbox)
+	FRWScopeLock Lock(GParamIdLock, SLT_Write);
+	if (const uint32* FoundIndex = GetParamIdData().NameToParamId.Find(InName))
 	{
-		return FSandboxThreadData::Get().ParamIdToName[ParameterIndex];
+		ParameterIndex = *FoundIndex;
 	}
 	else
-#endif
 	{
-		FRWScopeLock Lock(GParamIdLock, SLT_ReadOnly);
-		return GParamIdToName[ParameterIndex];
+		ParameterIndex = GetParamIdData().ParamIdToName.Num();
+		GetParamIdData().NameToParamId.Add(InName, ParameterIndex);
+		GetParamIdData().ParamIdToName.Add(InName);
 	}
 }
 
+FName FParamId::ToName() const
+{
+	FRWScopeLock Lock(GParamIdLock, SLT_ReadOnly);
+	return GetParamIdData().ParamIdToName[ParameterIndex];
+}
+
 FParamId FParamId::GetMaxParamId()
 {
-#if WITH_DEV_AUTOMATION_TESTS
-	if (bGTestSandbox)
-	{
-		return FParamId((uint32)FSandboxThreadData::Get().ParamIdToName.Num());
-	}
-	else
-#endif
-	{
-		FRWScopeLock Lock(GParamIdLock, SLT_ReadOnly);
-		return FParamId((uint32)GParamIdToName.Num());
-	}
+	FRWScopeLock Lock(GParamIdLock, SLT_ReadOnly);
+	return FParamId((uint32)GetParamIdData().ParamIdToName.Num());
 }
 
 #if WITH_DEV_AUTOMATION_TESTS
 void FParamId::BeginTestSandbox()
 {
-	bGTestSandbox = true;
-	FSandboxThreadData::Get().NameToParamId.Empty();
-	FSandboxThreadData::Get().ParamIdToName.Empty();
+	FRWScopeLock Lock(GParamIdLock, SLT_Write);
+	check(bGParamIdSandboxed == false);
+	bGParamIdSandboxed = true;
+	GParamIdSandboxedThreadId = FPlatformTLS::GetCurrentThreadId();
+	GSandboxedParamIdGlobalData.NameToParamId.Empty();
+	GSandboxedParamIdGlobalData.ParamIdToName.Empty();
 }
 
 void FParamId::EndTestSandbox()
 {
-	FSandboxThreadData::Get().NameToParamId.Empty();
-	FSandboxThreadData::Get().ParamIdToName.Empty();
-	bGTestSandbox = false;
+	FRWScopeLock Lock(GParamIdLock, SLT_Write);
+	check(bGParamIdSandboxed == true);
+	bGParamIdSandboxed = false;
+	GSandboxedParamIdGlobalData.NameToParamId.Empty();
+	GSandboxedParamIdGlobalData.ParamIdToName.Empty();
 }
 #endif
 
diff --git a/Engine/Plugins/Experimental/AnimNext/Source/AnimNext/Private/Param/ParamStack.cpp b/Engine/Plugins/Experimental/AnimNext/Source/AnimNext/Private/Param/ParamStack.cpp
--- a/Engine/Plugins/Experimental/AnimNext/Source/AnimNext/Private/Param/ParamStack.cpp
+++ b/Engine/Plugins/Experimental/AnimNext/Source/AnimNext/Private/Param/ParamStack.cpp
@@ -4,6 +4,7 @@
 #include "Param/ParamHelpers.h"
 #include "PropertyBag.h"
 #include "EngineLogs.h"
+#include "Param/ParamUtils.h"
 
 #define LOCTEXT_NAMESPACE "AnimNextParamStack"
 
@@ -118,40 +119,43 @@ FParamStack::FParam& FParamStack::FParam::operator=(FParam&& InOtherParam)
 	return *this;
 }
 
-FParamStackLayer::FParamStackLayer(const FInstancedPropertyBag& InPropertyBag)
+FParamStack::FLayer::FLayer(FInstancedPropertyBag& InPropertyBag, bool bInIsMutable)
 {
-	TConstArrayView<FPropertyBagPropertyDesc> Descs = InPropertyBag.GetPropertyBagStruct()->GetPropertyDescs();
-	
-	// Determine param ID range for this layer
-	TArray<FParamId> CachedIds;
-	CachedIds.SetNumUninitialized(Descs.Num());
-	MinParamId = MAX_uint32;
-	uint32 MaxParamId = 0;
-	for (uint32 DescIndex = 0; DescIndex < static_cast<uint32>(Descs.Num()); ++DescIndex)
+	if (const UPropertyBag* PropertyBagStruct = InPropertyBag.GetPropertyBagStruct())
 	{
-		const FPropertyBagPropertyDesc& Desc = Descs[DescIndex];
-		const FParamId& ParamId = CachedIds[DescIndex] = FParamId(Desc.Name);
-		MinParamId = FMath::Min(ParamId.ToInt(), MinParamId);
-		MaxParamId = FMath::Max(ParamId.ToInt(), MaxParamId);
-	}
+		TConstArrayView<FPropertyBagPropertyDesc> Descs = PropertyBagStruct->GetPropertyDescs();
 
-	if (MinParamId <= MaxParamId)
-	{
-		const uint32 ParamRangeSize = (MaxParamId - MinParamId) + 1;
-		Params.SetNumZeroed(ParamRangeSize);
-		FConstStructView StructView = InPropertyBag.GetValue();
+		// Determine param ID range for this layer
+		TArray<FParamId> CachedIds;
+		CachedIds.SetNumUninitialized(Descs.Num());
+		MinParamId = MAX_uint32;
+		uint32 MaxParamId = 0;
 		for (uint32 DescIndex = 0; DescIndex < static_cast<uint32>(Descs.Num()); ++DescIndex)
 		{
 			const FPropertyBagPropertyDesc& Desc = Descs[DescIndex];
-			const FParamId& ParamId = CachedIds[DescIndex];
-			const uint8* DataPtr = StructView.GetMemory() + Desc.CachedProperty->GetOffset_ForInternal();
-			const uint32 LocalParamIndex = ParamId.ToInt() - MinParamId;
-			Params[LocalParamIndex] = FParamStack::FParam(FParamTypeHandle::FromPropertyBagPropertyDesc(Desc), TArrayView<uint8>(const_cast<uint8*>(DataPtr), Desc.CachedProperty->GetSize()), true, false);
+			const FParamId& ParamId = CachedIds[DescIndex] = FParamId(Desc.Name);
+			MinParamId = FMath::Min(ParamId.ToInt(), MinParamId);
+			MaxParamId = FMath::Max(ParamId.ToInt(), MaxParamId);
+		}
+
+		if (MinParamId <= MaxParamId)
+		{
+			const uint32 ParamRangeSize = (MaxParamId - MinParamId) + 1;
+			Params.SetNumZeroed(ParamRangeSize);
+			FConstStructView StructView = InPropertyBag.GetValue();
+			for (uint32 DescIndex = 0; DescIndex < static_cast<uint32>(Descs.Num()); ++DescIndex)
+			{
+				const FPropertyBagPropertyDesc& Desc = Descs[DescIndex];
+				const FParamId& ParamId = CachedIds[DescIndex];
+				const uint8* DataPtr = StructView.GetMemory() + Desc.CachedProperty->GetOffset_ForInternal();
+				const uint32 LocalParamIndex = ParamId.ToInt() - MinParamId;
+				Params[LocalParamIndex] = FParamStack::FParam(FParamTypeHandle::FromPropertyBagPropertyDesc(Desc), TArrayView<uint8>(const_cast<uint8*>(DataPtr), Desc.CachedProperty->GetSize()), true, bInIsMutable);
+			}
 		}
 	}
 }
 
-FParamStackLayer::FParamStackLayer(TConstArrayView<TPair<FParamId, FParamStack::FParam>> InParams)
+FParamStack::FLayer::FLayer(TConstArrayView<TPair<FParamId, FParamStack::FParam>> InParams)
 {
 	MinParamId = MAX_uint32;
 	uint32 MaxParamId = 0;
@@ -174,9 +178,11 @@ FParamStackLayer::FParamStackLayer(TConstArrayView<TPair<FParamId, FParamStack::
 	}
 }
 
-FParamStack::FPushedLayer::FPushedLayer(FParamStackLayer& InLayer, FParamStack& InStack)
+FParamStack::FPushedLayer::FPushedLayer(FLayer& InLayer, FParamStack& InStack)
 	: Layer(InLayer)
 {
+	SerialNumber = InStack.MakeSerialNumber();
+
 	if (Layer.Params.Num())
 	{
 		InStack.ResizeLayerIndices();
@@ -189,9 +195,11 @@ FParamStack::FPushedLayer::FPushedLayer(FParamStackLayer& InLayer, FParamStack&
 	}
 }
 
-FParamStack::FPushedLayer::FPushedLayer(const FPushedLayer& InPreviousLayer, FParamStackLayer& InLayer, FParamStack& InStack)
+FParamStack::FPushedLayer::FPushedLayer(const FPushedLayer& InPreviousLayer, FLayer& InLayer, FParamStack& InStack)
 	: Layer(InLayer)
 {
+	SerialNumber = InStack.MakeSerialNumber();
+
 	if (Layer.Params.Num())
 	{
 		InStack.ResizeLayerIndices();
@@ -208,7 +216,7 @@ FParamStack::FPushedLayer::FPushedLayer(const FPushedLayer& InPreviousLayer, FPa
 			InStack.PreviousLayerIndices[BaseLayerIndex + LocalParamIndex] = InStack.LayerIndices[GlobalParamIndex];
 			if(InLayer.Params[LocalParamIndex].IsValid())
 			{
-				InStack.LayerIndices[GlobalParamIndex] = InStack.Layers.Num();
+				InStack.LayerIndices[GlobalParamIndex] = InStack.Layers.Num() - 1;
 			}
 		}
 	}
@@ -231,45 +239,58 @@ FParamStack& FParamStack::Get()
 	return FParamStackThreadData::Get().Stack;
 }
 
-void FParamStack::PushLayer(FParamStackLayer& InLayer)
+FParamStack::FPushedLayerHandle FParamStack::PushLayer(const FLayerHandle& InLayerHandle)
 {
-	if (Layers.Num() < MAX_uint16)
+	return PushLayerInternal(*InLayerHandle.Layer.Get());
+}
+
+FParamStack::FPushedLayerHandle FParamStack::PushLayer(TConstArrayView<TPair<FParamId, FParamStack::FParam>> InParams)
+{
+	if(Layers.Num() < MAX_uint16)
 	{
-		if (Layers.Num())
-		{
-			Layers.Push(FPushedLayer(Layers.Top(), InLayer, *this));
-		}
-		else
-		{
-			Layers.Push(FPushedLayer(InLayer, *this));
-		}
+		FLayer& OwnedLayer = OwnedStackLayers.Add_GetRef(FLayer(InParams));
+		OwnedLayer.OwnedStorageOffset = AllocAndCopyOwnedParamStorage(OwnedLayer.Params);
+		return PushLayerInternal(OwnedLayer);
 	}
 	else
 	{
-		UE_LOG(LogAnimation, Warning, TEXT("FParamStack: Could not push a layer: Maximum 65535 stack layers."))
+		UE_LOG(LogAnimation, Warning, TEXT("FParamStack: Could not push a layer: Maximum 65535 stack layers."));
+		return FPushedLayerHandle();
 	}
 }
 
-void FParamStack::PushLayer(TConstArrayView<TPair<FParamId, FParamStack::FParam>> InParams)
+FParamStack::FPushedLayerHandle FParamStack::PushLayerInternal(FLayer& InLayer)
 {
-	if(Layers.Num() < MAX_uint16)
+	if (Layers.Num() < MAX_uint16)
 	{
-		FParamStackLayer& OwnedLayer = OwnedStackLayers.Add_GetRef(FParamStackLayer(InParams));
-		OwnedLayer.OwnedStorageOffset = AllocAndCopyOwnedParamStorage(OwnedLayer.Params);
-		PushLayer(OwnedLayer);
+		if (Layers.Num())
+		{
+			FPushedLayer& NewPushedLayer = Layers.Emplace_GetRef(Layers.Top(), InLayer, *this);
+			return FPushedLayerHandle(Layers.Num() - 1, NewPushedLayer.SerialNumber);
+		}
+		else
+		{
+			FPushedLayer& NewPushedLayer = Layers.Emplace_GetRef(InLayer, *this);
+			return FPushedLayerHandle(Layers.Num() - 1, NewPushedLayer.SerialNumber);
+		}
 	}
 	else
 	{
-		UE_LOG(LogAnimation, Warning, TEXT("FParamStack: Could not push a layer: Maximum 65535 stack layers."))
+		UE_LOG(LogAnimation, Warning, TEXT("FParamStack: Could not push a layer: Maximum 65535 stack layers."));
+		return FPushedLayerHandle();
 	}
 }
 
-void FParamStack::PopLayer()
+void FParamStack::PopLayer(FPushedLayerHandle InHandle)
 {
 	if(Layers.Num() > 0)
 	{
 		const FPushedLayer& TopLayer = Layers.Top();
 
+		checkf(TopLayer.SerialNumber == InHandle.SerialNumber && (uint32)Layers.Num() - 1 == InHandle.Index, 
+			TEXT("UE::AnimNext::FParamStack::PopLayer: Invalid layer handle supplied (Have: %u, %u, Expected: %u, %u)"), 
+			InHandle.Index, InHandle.SerialNumber, (uint32)Layers.Num() - 1, TopLayer.SerialNumber);
+
 		// Fixup layer indices to previous, if any
 		const uint32 NumParams = TopLayer.Layer.Params.Num();
 		if(NumParams > 0 && PreviousLayerIndices.Num() > 0)
@@ -305,21 +326,33 @@ void FParamStack::PopLayer()
 	}
 }
 
-TUniquePtr<FParamStackLayer> FParamStack::MakeLayer(const FInstancedPropertyBag& InPropertyBag)
+FParamStack::FLayerHandle FParamStack::MakeLayer(const FInstancedPropertyBag& InPropertyBag)
+{
+	TUniquePtr<FLayer> Layer = TUniquePtr<FLayer>(new FLayer(const_cast<FInstancedPropertyBag&>(InPropertyBag), false));
+	return FLayerHandle(MoveTemp(Layer));
+}
+
+FParamStack::FLayerHandle FParamStack::MakeMutableLayer(FInstancedPropertyBag& InPropertyBag)
 {
-	TUniquePtr<FParamStackLayer> Layer = TUniquePtr<FParamStackLayer>(new FParamStackLayer(InPropertyBag));
-	return Layer;
+	TUniquePtr<FLayer> Layer = TUniquePtr<FLayer>(new FLayer(InPropertyBag, true));
+	return FLayerHandle(MoveTemp(Layer));
 }
 
-TUniquePtr<FParamStackLayer> FParamStack::MakeLayer(TConstArrayView<TPair<FParamId, FParamStack::FParam>> InParams)
+FParamStack::FLayerHandle FParamStack::MakeLayer(TConstArrayView<TPair<FParamId, FParamStack::FParam>> InParams)
 {
-	TUniquePtr<FParamStackLayer> Layer = TUniquePtr<FParamStackLayer>(new FParamStackLayer(InParams));
-	return Layer;
+	TUniquePtr<FLayer> Layer = TUniquePtr<FLayer>(new FLayer(InParams));
+	return FLayerHandle(MoveTemp(Layer));
 }
 
-FParamStack::EGetParamResult FParamStack::GetParamData(FParamId InId, FParamTypeHandle InTypeHandle, TConstArrayView<uint8>& OutParamData) const
+FParamStack::EGetParamResult FParamStack::GetParamData(FParamId InId, FParamTypeHandle InTypeHandle, TConstArrayView<uint8>& OutParamData, FParamCompatibility InRequiredCompatibility) const
 {
-	if (InId.ToInt() >= (uint32)LayerIndices.Num() || LayerIndices[InId.ToInt()] == MAX_uint16)
+	FParamTypeHandle ParamTypeHandle;
+	return GetParamData(InId, InTypeHandle, OutParamData, ParamTypeHandle, InRequiredCompatibility);
+}
+
+FParamStack::EGetParamResult FParamStack::GetParamData(FParamId InId, FParamTypeHandle InTypeHandle, TConstArrayView<uint8>& OutParamData, FParamTypeHandle& OutParamTypeHandle, FParamCompatibility InRequiredCompatibility) const
+{
+	if (Layers.Num() == 0 || InId.ToInt() >= (uint32)LayerIndices.Num() || LayerIndices[InId.ToInt()] == MAX_uint16)
 	{
 		return EGetParamResult::NotInScope;
 	}
@@ -333,19 +366,33 @@ FParamStack::EGetParamResult FParamStack::GetParamData(FParamId InId, FParamType
 		return EGetParamResult::NotInScope;
 	}
 
-	if (Param.GetTypeHandle() != InTypeHandle)
+	OutParamTypeHandle = Param.GetTypeHandle();
+
+	FParamCompatibility Compatibility = UE::AnimNext::FParamUtils::GetCompatibility(InTypeHandle, OutParamTypeHandle);
+	if (Compatibility < InRequiredCompatibility)
 	{
-		return EGetParamResult::IncorrectType;
+		return EGetParamResult::IncompatibleType;
 	}
-
+	
 	OutParamData = Param.GetData();
 
+	if (Compatibility == InRequiredCompatibility)
+	{
+		return EGetParamResult::CompatibleType;
+	}
+
 	return EGetParamResult::Succeeded;
 }
 
-FParamStack::EGetParamResult FParamStack::GetMutableParamData(FParamId InId, FParamTypeHandle InTypeHandle, TArrayView<uint8>& OutParamData)
+FParamStack::EGetParamResult FParamStack::GetMutableParamData(FParamId InId, FParamTypeHandle InTypeHandle, TArrayView<uint8>& OutParamData, FParamCompatibility InRequiredCompatibility)
+{
+	FParamTypeHandle ParamTypeHandle;
+	return GetMutableParamData(InId, InTypeHandle, OutParamData, ParamTypeHandle, InRequiredCompatibility);
+}
+
+FParamStack::EGetParamResult FParamStack::GetMutableParamData(FParamId InId, FParamTypeHandle InTypeHandle, TArrayView<uint8>& OutParamData, FParamTypeHandle& OutParamTypeHandle, FParamCompatibility InRequiredCompatibility)
 {
-	if (InId.ToInt() >= (uint32)LayerIndices.Num() || LayerIndices[InId.ToInt()] == MAX_uint16)
+	if (Layers.Num() == 0 || InId.ToInt() >= (uint32)LayerIndices.Num() || LayerIndices[InId.ToInt()] == MAX_uint16)
 	{
 		return EGetParamResult::NotInScope;
 	}
@@ -359,9 +406,21 @@ FParamStack::EGetParamResult FParamStack::GetMutableParamData(FParamId InId, FPa
 		return EGetParamResult::NotInScope;
 	}
 
-	EGetParamResult AccessResult = Param.GetTypeHandle() != InTypeHandle ? EGetParamResult::IncorrectType : EGetParamResult::Succeeded;
+	OutParamTypeHandle = Param.GetTypeHandle();
+
+	EGetParamResult AccessResult = EGetParamResult::Succeeded;
+	FParamCompatibility Compatibility = UE::AnimNext::FParamUtils::GetCompatibility(InTypeHandle, OutParamTypeHandle);
+	if (Compatibility < InRequiredCompatibility)
+	{
+		AccessResult |= EGetParamResult::IncompatibleType;
+	}
+	else if (Compatibility == InRequiredCompatibility)
+	{
+		AccessResult |= EGetParamResult::CompatibleType;
+	}
+
 	AccessResult |= !Param.IsMutable() ? EGetParamResult::Immutable : EGetParamResult::Succeeded;
-	if (AccessResult != EGetParamResult::Succeeded)
+	if (AccessResult != EGetParamResult::Succeeded && AccessResult != EGetParamResult::CompatibleType)
 	{
 		return AccessResult;
 	}
@@ -373,7 +432,7 @@ FParamStack::EGetParamResult FParamStack::GetMutableParamData(FParamId InId, FPa
 
 bool FParamStack::IsMutableParam(FParamId InId) const
 {
-	if (InId.ToInt() >= (uint32)LayerIndices.Num() || LayerIndices[InId.ToInt()] == MAX_uint16)
+	if (Layers.Num() == 0 || InId.ToInt() >= (uint32)LayerIndices.Num() || LayerIndices[InId.ToInt()] == MAX_uint16)
 	{
 		return false;
 	}
@@ -388,7 +447,7 @@ bool FParamStack::IsMutableParam(FParamId InId) const
 
 bool FParamStack::IsReferenceParam(FParamId InId) const
 {
-	if (InId.ToInt() >= (uint32)LayerIndices.Num() || LayerIndices[InId.ToInt()] == MAX_uint16)
+	if (Layers.Num() == 0 || InId.ToInt() >= (uint32)LayerIndices.Num() || LayerIndices[InId.ToInt()] == MAX_uint16)
 	{
 		return false;
 	}
@@ -467,6 +526,16 @@ void FParamStack::ResizeLayerIndices()
 	}
 }
 
+uint32 FParamStack::MakeSerialNumber()
+{
+	++SerialNumber;
+	if (SerialNumber == 0)
+	{
+		++SerialNumber;
+	}
+	return SerialNumber;
+}
+
 }
 
 #undef LOCTEXT_NAMESPACE
\ No newline at end of file
diff --git a/Engine/Plugins/Experimental/AnimNext/Source/AnimNext/Private/Param/ParamTypeHandle.cpp b/Engine/Plugins/Experimental/AnimNext/Source/AnimNext/Private/Param/ParamTypeHandle.cpp
--- a/Engine/Plugins/Experimental/AnimNext/Source/AnimNext/Private/Param/ParamTypeHandle.cpp
+++ b/Engine/Plugins/Experimental/AnimNext/Source/AnimNext/Private/Param/ParamTypeHandle.cpp
@@ -9,22 +9,61 @@
 namespace UE::AnimNext
 {
 
-// Array of all non built-in types. Index into this array is (CustomTypeIndex - 1)
-static TArray<FAnimNextParamType> GCustomTypes;
+// RW lock for global data
+static FRWLock GParamTypeHandleLock;
 
-// Map of type to index in GNonBuiltInTypes array
-static TMap<FAnimNextParamType, uint32> GTypeToIndexMap;
+struct FTypeHandleGlobalData
+{
+	// Array of all non built-in types. Index into this array is (CustomTypeIndex - 1)
+	TArray<FAnimNextParamType> CustomTypes;
+
+	// Map of type to index in GNonBuiltInTypes array
+	TMap<FAnimNextParamType, uint32> TypeToIndexMap;
+};
+
+static FTypeHandleGlobalData GTypeHandleGlobalData;
 
-// RW lock for types array/map
-FRWLock TypesLock;
+#if WITH_DEV_AUTOMATION_TESTS	
+static FTypeHandleGlobalData GSandboxedTypeHandleGlobalData;
+static uint32 GTypeHandleSandboxedThreadId = MAX_uint32;
+static bool bGTypeHandleSandboxed = false;
+#endif
+
+static FTypeHandleGlobalData& GetTypeHandleData()
+{
+#if WITH_DEV_AUTOMATION_TESTS		
+	if(bGTypeHandleSandboxed && GTypeHandleSandboxedThreadId == FPlatformTLS::GetCurrentThreadId())
+	{
+		return GSandboxedTypeHandleGlobalData;
+	}
+	else
+#endif
+	{
+		return GTypeHandleGlobalData;
+	}
+}
 
-void FParamTypeHandle::ResetCustomTypes()
+#if WITH_DEV_AUTOMATION_TESTS
+void FParamTypeHandle::BeginTestSandbox()
 {
-	FRWScopeLock ScopeLock(TypesLock, SLT_Write);
+	FRWScopeLock Lock(GParamTypeHandleLock, SLT_Write);
+	check(bGTypeHandleSandboxed == false);
+	bGTypeHandleSandboxed = true;
+	GTypeHandleSandboxedThreadId = FPlatformTLS::GetCurrentThreadId();
+	GSandboxedTypeHandleGlobalData.CustomTypes.Empty();
+	GSandboxedTypeHandleGlobalData.TypeToIndexMap.Empty();
+}
 
-	GCustomTypes.Empty();
-	GTypeToIndexMap.Empty();
+void FParamTypeHandle::EndTestSandbox()
+{
+	FRWScopeLock Lock(GParamTypeHandleLock, SLT_Write);
+	check(bGTypeHandleSandboxed == true);
+	bGTypeHandleSandboxed = false;
+	GTypeHandleSandboxedThreadId = MAX_uint32;
+	GSandboxedTypeHandleGlobalData.CustomTypes.Empty();
+	GSandboxedTypeHandleGlobalData.TypeToIndexMap.Empty();
 }
+#endif
 
 uint32 FParamTypeHandle::GetOrAllocateCustomTypeIndex(FAnimNextParamType::EValueType InValueType, FAnimNextParamType::EContainerType InContainerType, const UObject* InValueTypeObject)
 {
@@ -40,26 +79,39 @@ uint32 FParamTypeHandle::GetOrAllocateCustomTypeIndex(FAnimNextParamType::EValue
 	// FRWScopeLock (@see FRWScopeLock comments), hence the write lock around the map/array access here.
 	{
 		// See if the type already exists in the map
-		FRWScopeLock ScopeLock(TypesLock, SLT_Write);
+		FRWScopeLock ScopeLock(GParamTypeHandleLock, SLT_Write);
 
-		if(const uint32* IndexPtr = GTypeToIndexMap.FindByHash(Hash, ParameterType))
+		if(const uint32* IndexPtr = GetTypeHandleData().TypeToIndexMap.FindByHash(Hash, ParameterType))
 		{
 			return *IndexPtr + 1;
 		}
 
 		// Add a new custom type
-		const uint32 Index = GCustomTypes.Add(ParameterType);
-		GTypeToIndexMap.AddByHash(Hash, ParameterType, Index);
+		const uint32 Index = GetTypeHandleData().CustomTypes.Add(ParameterType);
+		GetTypeHandleData().TypeToIndexMap.AddByHash(Hash, ParameterType, Index);
 
 		checkf((Index + 1) < (1 << 24), TEXT("FParamTypeHandle::GetCustomTypeIndex: Type index overflowed"));
 		return Index + 1;
 	}
 }
 
+void FParamTypeHandle::GetCustomTypeInfo(FAnimNextParamType::EValueType& OutValueType, FAnimNextParamType::EContainerType& OutContainerType, const UObject*& OutValueTypeObject) const
+{
+	if (GetParameterType() == EParamType::Custom)
+	{
+		FRWScopeLock ScopeLock(GParamTypeHandleLock, SLT_ReadOnly);
+
+		const FAnimNextParamType& ParamType = GetTypeHandleData().CustomTypes[GetCustomTypeIndex() - 1];
+		OutValueType = ParamType.GetValueType();
+		OutContainerType = ParamType.GetContainerType();
+		OutValueTypeObject = ParamType.GetValueTypeObject();
+	}
+}
+
 bool FParamTypeHandle::ValidateCustomTypeIndex(uint32 InCustomTypeIndex)
 {
-	FRWScopeLock ScopeLock(TypesLock, SLT_ReadOnly);
-	return GCustomTypes.IsValidIndex(InCustomTypeIndex - 1);
+	FRWScopeLock ScopeLock(GParamTypeHandleLock, SLT_ReadOnly);
+	return GetTypeHandleData().CustomTypes.IsValidIndex(InCustomTypeIndex - 1);
 }
 
 FParamTypeHandle FParamTypeHandle::FromPropertyBagPropertyDesc(const FPropertyBagPropertyDesc& Desc)
@@ -126,7 +178,7 @@ FParamTypeHandle FParamTypeHandle::FromPropertyBagPropertyDesc(const FPropertyBa
 					else
 					{
 						Handle.SetParameterType(EParamType::Custom);
-						Handle.SetCustomTypeIndex(GetOrAllocateCustomTypeIndex(Desc.ValueType, Desc.ContainerTypes[0], ScriptStruct));
+						Handle.SetCustomTypeIndex(GetOrAllocateCustomTypeIndex(Desc.ValueType, Desc.ContainerTypes.GetFirstContainerType(), ScriptStruct));
 					}
 				}
 				else
@@ -138,7 +190,7 @@ FParamTypeHandle FParamTypeHandle::FromPropertyBagPropertyDesc(const FPropertyBa
 				if(const UEnum* Enum = Cast<UEnum>(Desc.ValueTypeObject.Get()))
 				{
 					Handle.SetParameterType(EParamType::Custom);
-					Handle.SetCustomTypeIndex(GetOrAllocateCustomTypeIndex(Desc.ValueType, Desc.ContainerTypes[0], Enum));
+					Handle.SetCustomTypeIndex(GetOrAllocateCustomTypeIndex(Desc.ValueType, Desc.ContainerTypes.GetFirstContainerType(), Enum));
 				}
 				else
 				{
@@ -152,7 +204,7 @@ FParamTypeHandle FParamTypeHandle::FromPropertyBagPropertyDesc(const FPropertyBa
 				if(const UClass* Class = Cast<UClass>(Desc.ValueTypeObject.Get()))
 				{
 					Handle.SetParameterType(EParamType::Custom);
-					Handle.SetCustomTypeIndex(GetOrAllocateCustomTypeIndex(Desc.ValueType, Desc.ContainerTypes[0], Class));
+					Handle.SetCustomTypeIndex(GetOrAllocateCustomTypeIndex(Desc.ValueType, Desc.ContainerTypes.GetFirstContainerType(), Class));
 				}
 				else
 				{
@@ -178,13 +230,13 @@ FParamTypeHandle FParamTypeHandle::FromPropertyBagPropertyDesc(const FPropertyBa
 			case EPropertyBagPropertyType::String:
 			case EPropertyBagPropertyType::Text:
 				Handle.SetParameterType(EParamType::Custom);
-				Handle.SetCustomTypeIndex(GetOrAllocateCustomTypeIndex(Desc.ValueType, Desc.ContainerTypes[0], nullptr));
+				Handle.SetCustomTypeIndex(GetOrAllocateCustomTypeIndex(Desc.ValueType, Desc.ContainerTypes.GetFirstContainerType(), nullptr));
 				break;
 			case EPropertyBagPropertyType::Struct:
 				if(const UScriptStruct* ScriptStruct = Cast<UScriptStruct>(Desc.ValueTypeObject.Get()))
 				{
 					Handle.SetParameterType(EParamType::Custom);
-					Handle.SetCustomTypeIndex(GetOrAllocateCustomTypeIndex(Desc.ValueType, Desc.ContainerTypes[0], ScriptStruct));
+					Handle.SetCustomTypeIndex(GetOrAllocateCustomTypeIndex(Desc.ValueType, Desc.ContainerTypes.GetFirstContainerType(), ScriptStruct));
 				}
 				else
 				{
@@ -195,7 +247,7 @@ FParamTypeHandle FParamTypeHandle::FromPropertyBagPropertyDesc(const FPropertyBa
 				if(const UEnum* Enum = Cast<UEnum>(Desc.ValueTypeObject.Get()))
 				{
 					Handle.SetParameterType(EParamType::Custom);
-					Handle.SetCustomTypeIndex(GetOrAllocateCustomTypeIndex(Desc.ValueType, Desc.ContainerTypes[0], Enum));
+					Handle.SetCustomTypeIndex(GetOrAllocateCustomTypeIndex(Desc.ValueType, Desc.ContainerTypes.GetFirstContainerType(), Enum));
 				}
 				else
 				{
@@ -209,7 +261,7 @@ FParamTypeHandle FParamTypeHandle::FromPropertyBagPropertyDesc(const FPropertyBa
 				if(const UClass* Class = Cast<UClass>(Desc.ValueTypeObject.Get()))
 				{
 					Handle.SetParameterType(EParamType::Custom);
-					Handle.SetCustomTypeIndex(GetOrAllocateCustomTypeIndex(Desc.ValueType, Desc.ContainerTypes[0], Class));
+					Handle.SetCustomTypeIndex(GetOrAllocateCustomTypeIndex(Desc.ValueType, Desc.ContainerTypes.GetFirstContainerType(), Class));
 				}
 				else
 				{
@@ -461,8 +513,8 @@ FAnimNextParamType FParamTypeHandle::GetType() const
 		break;
 	case EParamType::Custom:
 		{
-			FRWScopeLock ScopeLock(TypesLock, SLT_ReadOnly);
-			ParameterType = GCustomTypes[GetCustomTypeIndex() - 1];
+			FRWScopeLock ScopeLock(GParamTypeHandleLock, SLT_ReadOnly);
+			ParameterType = GetTypeHandleData().CustomTypes[GetCustomTypeIndex() - 1];
 		}
 		break;
 	}
diff --git a/Engine/Plugins/Experimental/AnimNext/Source/AnimNext/Private/Param/ParamUtils.cpp b/Engine/Plugins/Experimental/AnimNext/Source/AnimNext/Private/Param/ParamUtils.cpp
new file mode 100644
--- /dev/null
+++ b/Engine/Plugins/Experimental/AnimNext/Source/AnimNext/Private/Param/ParamUtils.cpp
@@ -0,0 +1,274 @@
+// Copyright Epic Games, Inc. All Rights Reserved.
+
+#include "Param/ParamUtils.h"
+#include "Param/ParamType.h"
+#include "Param/ParamTypeHandle.h"
+#include "Param/ParamCompatibility.h"
+
+namespace UE::AnimNext
+{
+
+FParamCompatibility FParamUtils::GetCompatibility(const FParamTypeHandle& InLHS, const FParamTypeHandle& InRHS)
+{
+	switch (InRHS.GetParameterType())
+	{
+	case FParamTypeHandle::EParamType::Bool:
+		switch (InLHS.GetParameterType())
+		{
+		case FParamTypeHandle::EParamType::Bool:
+			return EParamCompatibility::Compatible_Equal;
+		}
+		break;
+	case FParamTypeHandle::EParamType::Byte:
+		switch (InLHS.GetParameterType())
+		{
+		case FParamTypeHandle::EParamType::Byte:
+			return EParamCompatibility::Compatible_Equal;
+		case FParamTypeHandle::EParamType::Int32:
+			return EParamCompatibility::Compatible_Promotion;
+		case FParamTypeHandle::EParamType::Int64:
+			return EParamCompatibility::Compatible_Promotion;
+		case FParamTypeHandle::EParamType::Float:
+			return EParamCompatibility::Compatible_Promotion;
+		case FParamTypeHandle::EParamType::Double:
+			return EParamCompatibility::Compatible_Promotion;
+		}
+		break;
+	case FParamTypeHandle::EParamType::Int32:
+		switch (InLHS.GetParameterType())
+		{
+		case FParamTypeHandle::EParamType::Byte:
+			return EParamCompatibility::Incompatible_DataLoss;
+		case FParamTypeHandle::EParamType::Int32:
+			return EParamCompatibility::Compatible_Equal;
+		case FParamTypeHandle::EParamType::Int64:
+			return EParamCompatibility::Compatible_Promotion;
+		case FParamTypeHandle::EParamType::Float:
+			return EParamCompatibility::Incompatible_DataLoss;
+		case FParamTypeHandle::EParamType::Double:
+			return EParamCompatibility::Compatible_Promotion;
+		}
+		break;
+	case FParamTypeHandle::EParamType::Int64:
+		switch (InLHS.GetParameterType())
+		{
+		case FParamTypeHandle::EParamType::Byte:
+			return EParamCompatibility::Incompatible_DataLoss;
+		case FParamTypeHandle::EParamType::Int32:
+			return EParamCompatibility::Incompatible_DataLoss;
+		case FParamTypeHandle::EParamType::Int64:
+			return EParamCompatibility::Compatible_Equal;
+		case FParamTypeHandle::EParamType::Float:
+			return EParamCompatibility::Incompatible_DataLoss;
+		case FParamTypeHandle::EParamType::Double:
+			return EParamCompatibility::Incompatible_DataLoss;
+		}
+		break;
+	case FParamTypeHandle::EParamType::Float:
+		switch (InLHS.GetParameterType())
+		{
+		case FParamTypeHandle::EParamType::Byte:
+			return EParamCompatibility::Incompatible_DataLoss;
+		case FParamTypeHandle::EParamType::Int32:
+			return EParamCompatibility::Incompatible_DataLoss;
+		case FParamTypeHandle::EParamType::Int64:
+			return EParamCompatibility::Incompatible_DataLoss;
+		case FParamTypeHandle::EParamType::Float:
+			return EParamCompatibility::Compatible_Equal;
+		case FParamTypeHandle::EParamType::Double:
+			return EParamCompatibility::Compatible_Promotion;
+		}
+		break;
+	case FParamTypeHandle::EParamType::Double:
+		switch (InLHS.GetParameterType())
+		{
+		case FParamTypeHandle::EParamType::Byte:
+			return EParamCompatibility::Incompatible_DataLoss;
+		case FParamTypeHandle::EParamType::Int32:
+			return EParamCompatibility::Incompatible_DataLoss;
+		case FParamTypeHandle::EParamType::Int64:
+			return EParamCompatibility::Incompatible_DataLoss;
+		case FParamTypeHandle::EParamType::Float:
+			return EParamCompatibility::Incompatible_DataLoss;
+		case FParamTypeHandle::EParamType::Double:
+			return EParamCompatibility::Compatible_Equal;
+		}
+		break;
+	case FParamTypeHandle::EParamType::Name:
+		switch (InLHS.GetParameterType())
+		{
+		case FParamTypeHandle::EParamType::Name:
+			return EParamCompatibility::Compatible_Equal;
+		}
+		break;
+	case FParamTypeHandle::EParamType::String:
+		switch (InLHS.GetParameterType())
+		{
+		case FParamTypeHandle::EParamType::String:
+			return EParamCompatibility::Compatible_Equal;
+		}
+		break;
+	case FParamTypeHandle::EParamType::Text:
+		switch (InLHS.GetParameterType())
+		{
+		case FParamTypeHandle::EParamType::Text:
+			return EParamCompatibility::Compatible_Equal;
+		}
+		break;
+	case FParamTypeHandle::EParamType::Vector:
+		switch (InLHS.GetParameterType())
+		{
+		case FParamTypeHandle::EParamType::Vector:
+			return EParamCompatibility::Compatible_Equal;
+		}
+		break;
+	case FParamTypeHandle::EParamType::Vector4:
+		switch (InLHS.GetParameterType())
+		{
+		case FParamTypeHandle::EParamType::Vector4:
+			return EParamCompatibility::Compatible_Equal;
+		}
+		break;
+	case FParamTypeHandle::EParamType::Quat:
+		switch (InLHS.GetParameterType())
+		{
+		case FParamTypeHandle::EParamType::Quat:
+			return EParamCompatibility::Compatible_Equal;
+		}
+		break;
+	case FParamTypeHandle::EParamType::Transform:
+		switch (InLHS.GetParameterType())
+		{
+		case FParamTypeHandle::EParamType::Transform:
+			return EParamCompatibility::Compatible_Equal;
+		}
+		break;
+	case FParamTypeHandle::EParamType::Custom:
+		switch (InLHS.GetParameterType())
+		{
+		case FParamTypeHandle::EParamType::Custom:
+			{
+				FAnimNextParamType::EValueType ValueTypeLHS, ValueTypeRHS;
+				FAnimNextParamType::EContainerType ContainerTypeLHS, ContainerTypeRHS;
+				const UObject* ValueTypeObjectLHS, *ValueTypeObjectRHS;
+
+				InLHS.GetCustomTypeInfo(ValueTypeLHS, ContainerTypeLHS, ValueTypeObjectLHS);
+				InRHS.GetCustomTypeInfo(ValueTypeRHS, ContainerTypeRHS, ValueTypeObjectRHS);
+
+				if (ContainerTypeLHS == ContainerTypeRHS)
+				{
+					if (ContainerTypeLHS == FAnimNextParamType::EContainerType::Array)
+					{
+						if (ValueTypeLHS != ValueTypeRHS)
+						{
+							switch (ValueTypeLHS)
+							{
+							case FAnimNextParamType::EValueType::Float:
+								if(ValueTypeRHS == FAnimNextParamType::EValueType::Double)
+								{
+									return EParamCompatibility::Incompatible_DataLoss;
+								}
+								break;
+							case FAnimNextParamType::EValueType::Double:
+								if (ValueTypeRHS == FAnimNextParamType::EValueType::Float)
+								{
+									return EParamCompatibility::Compatible_Promotion;
+								}
+								break;
+							case FAnimNextParamType::EValueType::Struct:
+							case FAnimNextParamType::EValueType::Object:
+							case FAnimNextParamType::EValueType::SoftObject:
+							case FAnimNextParamType::EValueType::Class:
+							case FAnimNextParamType::EValueType::SoftClass:
+								if (ValueTypeObjectLHS == ValueTypeObjectRHS)
+								{
+									return EParamCompatibility::Compatible_Equal;
+								}
+								break;
+							}
+						}
+						else
+						{
+							switch (ValueTypeLHS)
+							{
+							case FAnimNextParamType::EValueType::Bool:
+							case FAnimNextParamType::EValueType::Byte:
+							case FAnimNextParamType::EValueType::Int32:
+							case FAnimNextParamType::EValueType::Int64:
+							case FAnimNextParamType::EValueType::Float:
+							case FAnimNextParamType::EValueType::Double:
+							case FAnimNextParamType::EValueType::Name:
+							case FAnimNextParamType::EValueType::String:
+							case FAnimNextParamType::EValueType::Text:
+								return EParamCompatibility::Compatible_Equal;
+							case FAnimNextParamType::EValueType::Enum:
+							case FAnimNextParamType::EValueType::Struct:
+							case FAnimNextParamType::EValueType::Object:
+							case FAnimNextParamType::EValueType::SoftObject:
+							case FAnimNextParamType::EValueType::Class:
+							case FAnimNextParamType::EValueType::SoftClass:
+								if (ValueTypeObjectLHS == ValueTypeObjectRHS)
+								{
+									return EParamCompatibility::Compatible_Equal;
+								}
+								break;
+							}
+						}
+					}
+					else
+					{
+						if (ValueTypeLHS == ValueTypeRHS)
+						{
+							switch (ValueTypeLHS)
+							{
+							default:
+							case FAnimNextParamType::EValueType::None:
+								return EParamCompatibility::Incompatible;
+							case FAnimNextParamType::EValueType::Enum:
+								if (ValueTypeObjectLHS == ValueTypeObjectRHS)
+								{
+									return EParamCompatibility::Compatible_Equal;
+								}
+								break;
+							case FAnimNextParamType::EValueType::Struct:
+								if (ValueTypeObjectLHS == ValueTypeObjectRHS)
+								{
+									return EParamCompatibility::Compatible_Equal;
+								}
+								else if (CastChecked<UScriptStruct>(ValueTypeObjectRHS)->IsChildOf(CastChecked<UScriptStruct>(ValueTypeObjectLHS)))
+								{
+									return EParamCompatibility::Compatible_Cast;
+								}
+								break;
+							case FAnimNextParamType::EValueType::Object:
+							case FAnimNextParamType::EValueType::SoftObject:
+							case FAnimNextParamType::EValueType::Class:
+							case FAnimNextParamType::EValueType::SoftClass:
+								if (ValueTypeObjectLHS == ValueTypeObjectRHS)
+								{
+									return EParamCompatibility::Compatible_Equal;
+								}
+								else if (CastChecked<UClass>(ValueTypeObjectRHS)->IsChildOf(CastChecked<UClass>(ValueTypeObjectLHS)))
+								{
+									return EParamCompatibility::Compatible_Cast;
+								}
+								break;
+							}
+						}
+					}
+				}
+			}
+			break;
+		}
+		break;
+	}
+
+	return EParamCompatibility::Incompatible;
+}
+
+FParamCompatibility FParamUtils::GetCompatibility(const FAnimNextParamType& InLHS, const FAnimNextParamType& InRHS)
+{
+	return GetCompatibility(InLHS.GetHandle(), InRHS.GetHandle());
+}
+
+}
\ No newline at end of file
diff --git a/Engine/Plugins/Experimental/AnimNext/Source/AnimNext/Private/Param/RigVMDispatch_GetParameter.cpp b/Engine/Plugins/Experimental/AnimNext/Source/AnimNext/Private/Param/RigVMDispatch_GetParameter.cpp
--- a/Engine/Plugins/Experimental/AnimNext/Source/AnimNext/Private/Param/RigVMDispatch_GetParameter.cpp
+++ b/Engine/Plugins/Experimental/AnimNext/Source/AnimNext/Private/Param/RigVMDispatch_GetParameter.cpp
@@ -1,11 +1,9 @@
 // Copyright Epic Games, Inc. All Rights Reserved.
 
 #include "Param/RigVMDispatch_GetParameter.h"
-
-#include "Param/ParametersExecuteContext.h"
 #include "RigVMCore/RigVMStruct.h"
 #include "RigVMCore/RigVM.h"
-#include "Graph/GraphExecuteContext.h"
+#include "Graph/AnimNextExecuteContext.h"
 #include "Context.h"
 #include "Param/ParamStack.h"
 #include "Param/ParamStack.h"
@@ -36,6 +34,17 @@ FName FRigVMDispatch_GetParameter::GetArgumentNameForOperandIndex(int32 InOperan
 }
 
 #if WITH_EDITOR
+FString FRigVMDispatch_GetParameter::GetArgumentMetaData(const FName& InArgumentName, const FName& InMetaDataKey) const
+{
+	if ((InArgumentName == TypeHandleName || InArgumentName == ParameterIdName) &&
+		InMetaDataKey == FRigVMStruct::SingletonMetaName)
+	{
+		return TEXT("True");
+	}
+
+	return Super::GetArgumentMetaData(InArgumentName, InMetaDataKey);
+}
+
 FString FRigVMDispatch_GetParameter::GetArgumentDefaultValue(const FName& InArgumentName, TRigVMTypeIndex InTypeIndex) const
 {
 	if (InArgumentName == ParameterIdName)
@@ -80,7 +89,7 @@ void FRigVMDispatch_GetParameter::Execute(FRigVMExtendedExecuteContext& InContex
 	uint8* TargetData = Handles[1].GetData();
 
 	uint32& ParameterId = *(uint32*)Handles[2].GetData();
-	if (ParameterId != FParamId::InvalidIndex)
+	if (ParameterId == FParamId::InvalidIndex)
 	{
 		ParameterId = FParamId(Parameter).ToInt();
 	}
@@ -91,9 +100,8 @@ void FRigVMDispatch_GetParameter::Execute(FRigVMExtendedExecuteContext& InContex
 		TypeHandle = FParamTypeHandle::FromProperty(ValueProperty).ToRaw();
 	}
 
-	FAnimNextGraphExecuteContext& Context = InContext.GetPublicData<FAnimNextGraphExecuteContext>();
 	TConstArrayView<uint8> SourceData;
-	if (Context.GetContext().GetParamStack().GetParamData(FParamId(ParameterId), FParamTypeHandle::FromRaw(TypeHandle), SourceData) == FParamStack::EGetParamResult::Succeeded)
+	if (FParamStack::Get().GetParamData(FParamId(ParameterId), FParamTypeHandle::FromRaw(TypeHandle), SourceData) == FParamStack::EGetParamResult::Succeeded)
 	{
 		ValueProperty->CopyCompleteValue(TargetData, SourceData.GetData());
 	}
diff --git a/Engine/Plugins/Experimental/AnimNext/Source/AnimNext/Private/Param/RigVMDispatch_SetParameter.cpp b/Engine/Plugins/Experimental/AnimNext/Source/AnimNext/Private/Param/RigVMDispatch_SetParameter.cpp
--- a/Engine/Plugins/Experimental/AnimNext/Source/AnimNext/Private/Param/RigVMDispatch_SetParameter.cpp
+++ b/Engine/Plugins/Experimental/AnimNext/Source/AnimNext/Private/Param/RigVMDispatch_SetParameter.cpp
@@ -1,11 +1,9 @@
 // Copyright Epic Games, Inc. All Rights Reserved.
 
 #include "Param/RigVMDispatch_SetParameter.h"
-
-#include "Param/ParametersExecuteContext.h"
 #include "RigVMCore/RigVMStruct.h"
 #include "RigVMCore/RigVM.h"
-#include "Graph/GraphExecuteContext.h"
+#include "Graph/AnimNextExecuteContext.h"
 #include "Context.h"
 #include "Param/ParamStack.h"
 #include "Param/ParamStack.h"
@@ -37,6 +35,17 @@ FName FRigVMDispatch_SetParameter::GetArgumentNameForOperandIndex(int32 InOperan
 }
 
 #if WITH_EDITOR
+FString FRigVMDispatch_SetParameter::GetArgumentMetaData(const FName& InArgumentName, const FName& InMetaDataKey) const
+{
+	if ((InArgumentName == TypeHandleName || InArgumentName == ParameterIdName) &&
+		InMetaDataKey == FRigVMStruct::SingletonMetaName)
+	{
+		return TEXT("True");
+	}
+
+	return Super::GetArgumentMetaData(InArgumentName, InMetaDataKey);
+}
+
 FString FRigVMDispatch_SetParameter::GetArgumentDefaultValue(const FName& InArgumentName, TRigVMTypeIndex InTypeIndex) const
 {
 	if (InArgumentName == ParameterIdName)
@@ -86,7 +95,7 @@ void FRigVMDispatch_SetParameter::Execute(FRigVMExtendedExecuteContext& InContex
 	const uint8* SourceData =  Handles[1].GetData();
 
 	uint32& ParameterId = *(uint32*)Handles[2].GetData();
-	if (ParameterId != FParamId::InvalidIndex)
+	if (ParameterId == FParamId::InvalidIndex)
 	{
 		ParameterId = FParamId(Parameter).ToInt();
 	}
@@ -97,9 +106,8 @@ void FRigVMDispatch_SetParameter::Execute(FRigVMExtendedExecuteContext& InContex
 		TypeHandle = FParamTypeHandle::FromProperty(ValueProperty).ToRaw();
 	}
 
-	const FAnimNextGraphExecuteContext& Context = InContext.GetPublicData<FAnimNextGraphExecuteContext>();
 	TArrayView<uint8> TargetData;
-	if (Context.GetContext().GetMutableParamStack().GetMutableParamData(FParamId(ParameterId), FParamTypeHandle::FromRaw(TypeHandle), TargetData) == FParamStack::EGetParamResult::Succeeded)
+	if (FParamStack::Get().GetMutableParamData(FParamId(ParameterId), FParamTypeHandle::FromRaw(TypeHandle), TargetData) == FParamStack::EGetParamResult::Succeeded)
 	{
 		ValueProperty->CopyCompleteValue(TargetData.GetData(), SourceData);
 	}	
diff --git a/Engine/Plugins/Experimental/AnimNext/Source/AnimNext/Public/Param/ParamId.h b/Engine/Plugins/Experimental/AnimNext/Source/AnimNext/Public/Param/ParamId.h
--- a/Engine/Plugins/Experimental/AnimNext/Source/AnimNext/Public/Param/ParamId.h
+++ b/Engine/Plugins/Experimental/AnimNext/Source/AnimNext/Public/Param/ParamId.h
@@ -6,6 +6,7 @@
 
 struct FRigVMDispatch_SetParameter;
 struct FRigVMDispatch_GetParameter;
+class UAnimNextParametersLibrary;
 
 namespace UE::AnimNext::Tests
 {
@@ -22,10 +23,11 @@ struct ANIMNEXT_API FParamId
 	friend class Tests::FParamStackTest;
 	friend struct ::FRigVMDispatch_SetParameter;
 	friend struct ::FRigVMDispatch_GetParameter;
+	friend class ::UAnimNextParametersLibrary;
 
 	static constexpr uint32 InvalidIndex = MAX_uint32;
 
-	FParamId() = delete;
+	FParamId() = default;
 	FParamId(const FParamId& InName) = default;
 	FParamId& operator=(const FParamId& InName) = default;
 
