AnimNext params in Anim BPs

diff --git a/Engine/Plugins/Experimental/AnimNext/Source/AnimNext/Public/Param/ParamStack.h b/Engine/Plugins/Experimental/AnimNext/Source/AnimNext/Public/Param/ParamStack.h
--- a/Engine/Plugins/Experimental/AnimNext/Source/AnimNext/Public/Param/ParamStack.h
+++ b/Engine/Plugins/Experimental/AnimNext/Source/AnimNext/Public/Param/ParamStack.h
@@ -9,6 +9,7 @@
 #include "Containers/PagedArray.h"
 #include "Param/ParamId.h"
 #include "Param/ParamTypeHandle.h"
+#include "Param/ParamCompatibility.h"
 #include "Misc/MemStack.h"
 
 struct FInstancedPropertyBag;
@@ -21,8 +22,6 @@ namespace UE::AnimNext::Tests
 namespace UE::AnimNext
 {
 
-struct FParamStackLayer;
-
 // Stack of parameter layers.
 // Acts as an associative container - allows retrieval of parameter calues (by ID) that have been
 // pushed onto the stack in 'layers'.
@@ -33,7 +32,6 @@ struct FParamStackLayer;
 struct FParamStack
 {
 	friend class Tests::FParamStackTest;
-	friend struct FParamStackLayer;
 
 private:
 	enum class EParamFlags : uint8
@@ -47,12 +45,8 @@ private:
 	FRIEND_ENUM_CLASS_FLAGS(EParamFlags);
 
 	// Parameter memory wrapper
-	struct  FParam
+	struct FParam
 	{
-	public:
-		friend struct FParamStack;
-		friend struct FParamStackLayer;
-
 		FParam() = default;
 		ANIMNEXT_API FParam(const FParam& InOtherParam);
 		ANIMNEXT_API FParam& operator=(const FParam& InOtherParam);
@@ -60,7 +54,6 @@ private:
 		ANIMNEXT_API FParam& operator=(FParam&& InOtherParam);
 		ANIMNEXT_API ~FParam();
 
-	private:
 		ANIMNEXT_API FParam(const FParamTypeHandle& InTypeHandle, TArrayView<uint8> InData, bool bInIsReference, bool bInIsMutable);
 
 		// Get the type handle of this param 
@@ -124,7 +117,101 @@ private:
 		EParamFlags Flags = EParamFlags::None;
 	};
 
+	// Stack layers are what actually get pushed/popped onto the layer stack.
+	// They are designed to be held on an instance, their data updated in place.
+	// Memory ownership for params is assumed to be outside of the stack and layers
+	struct FLayer
+	{
+		FLayer() = delete;
+
+		explicit FLayer(FInstancedPropertyBag& InPropertyBag, bool bInIsMutable);
+
+		explicit FLayer(TConstArrayView<TPair<FParamId, FParamStack::FParam>> InParams);
+
+		// Params that this layer supplies
+		TArray<FParamStack::FParam> Params;
+
+		// ID that the param indices start at. Maps the global param ID range into the range for this layer
+		uint32 MinParamId = 0;
+
+		// Storage offset for this layer if it is owned internally, otherwise MAX_uint32
+		uint32 OwnedStorageOffset = MAX_uint32;
+	};
+
+	// A layer held on the stack, may be owned by the stack or not
+	struct FPushedLayer
+	{
+		FPushedLayer(FLayer& InLayer, FParamStack& InStack);
+
+		FPushedLayer(const FPushedLayer& InPreviousLayer, FLayer& InLayer, FParamStack& InStack);
+
+		FLayer& Layer;
+
+		// Index offset for the previous layer for each active param.
+		// Offset into FParamStack::PreviousLayerIndices.
+		uint32 PreviousLayerIndexStart = MAX_uint32;
+
+		// Serial number used for identifying layers to pop
+		uint32 SerialNumber = 0;
+	};
+
 public:
+	// Opaque handle to a layer on the stack
+	struct FPushedLayerHandle
+	{
+		FPushedLayerHandle() = default;
+
+		bool IsValid() const
+		{
+			return Index != MAX_uint32 && SerialNumber != 0;
+		}
+
+		void Invalidate()
+		{
+			Index = MAX_uint32;
+			SerialNumber = 0;
+		}
+
+	private:
+		friend struct FParamStack;
+
+		FPushedLayerHandle(uint32 InIndex, uint32 InSerialNumber)
+			: Index(InIndex)
+			, SerialNumber(InSerialNumber)
+		{}
+
+		// Index into the stack
+		uint32 Index = MAX_uint32;
+
+		// Serial number of the layer
+		uint32 SerialNumber = 0;
+	};
+
+	// Opaque handle to a layer that can be held by external systems
+	struct FLayerHandle
+	{
+		FLayerHandle() = default;
+
+		bool IsValid() const
+		{
+			return Layer.IsValid();
+		}
+
+		void Invalidate()
+		{
+			Layer = nullptr;
+		}
+
+	private:
+		friend struct FParamStack;
+
+		explicit FLayerHandle(TUniquePtr<FLayer>&& InLayer)
+			: Layer(MoveTemp(InLayer))
+		{}
+
+		TUniquePtr<FLayer> Layer;
+	};
+
 	ANIMNEXT_API FParamStack();
 
 	// Get the param stack for this thread
@@ -134,48 +221,51 @@ public:
 	// @param	InParamId			Parameter ID for the parameter
 	// @param	InValue				Value to push
 	template<typename... ValueType>
-	void PushValue(FParamId InParamId, ValueType&&... InValue)
+	FPushedLayerHandle PushValue(FParamId InParamId, ValueType&&... InValue)
 	{
-		PushValues(InParamId, Forward<ValueType>(InValue)...);
+		return PushValues(InParamId, Forward<ValueType>(InValue)...);
 	}
 
 	// Push a value as a parameter layer
 	// @param	InParamId			Parameter name for the parameter
 	// @param	InValue				Value to push
 	template<typename... ValueType>
-	void PushValue(FName InName, ValueType&&... InValue)
+	FPushedLayerHandle PushValue(FName InName, ValueType&&... InValue)
 	{
-		PushValues(FParamId(InName), Forward<ValueType>(InValue)...);
+		return PushValues(FParamId(InName), Forward<ValueType>(InValue)...);
 	}
 
 	// Push an interleaved parameter list of keys and values as a parameter layer
 	// @param	InValues	Values to push
 	template<typename... Args>
-	void PushValues(Args&&... InValues)
+	FPushedLayerHandle PushValues(Args&&... InValues)
 	{
 		constexpr int32 NumItems = sizeof...(InValues) / 2;
 		TArray<TPair<FParamId, FParam>, TInlineAllocator<NumItems>> ParamIdValues;
 		ParamIdValues.Reserve(NumItems);
-		PushValuesHelper(ParamIdValues, Forward<Args>(InValues)...);
+		return PushValuesHelper(ParamIdValues, Forward<Args>(InValues)...);
 	}
 
 	// Push a layer
-	ANIMNEXT_API void PushLayer(FParamStackLayer& InLayer);
+	ANIMNEXT_API FPushedLayerHandle PushLayer(const FLayerHandle& InLayerHandle);
 
 	// Push an internally-owned layer. Copies parameter data to internal storage.
-	ANIMNEXT_API void PushLayer(TConstArrayView<TPair<FParamId, FParam>> InParams);
+	ANIMNEXT_API FPushedLayerHandle PushLayer(TConstArrayView<TPair<FParamId, FParam>> InParams);
 
 	// Create a cached parameter layer from an instanced property bag
-	static ANIMNEXT_API TUniquePtr<FParamStackLayer> MakeLayer(const FInstancedPropertyBag& InInstancedPropertyBag);
+	static ANIMNEXT_API FLayerHandle MakeLayer(const FInstancedPropertyBag& InInstancedPropertyBag);
 
-	// Create a cached parameter layer from set of params
-	static ANIMNEXT_API TUniquePtr<FParamStackLayer> MakeLayer(TConstArrayView<TPair<FParamId, FParam>> InParams);
+	// Create a cached mutable parameter layer from an instanced property bag
+	static ANIMNEXT_API FLayerHandle MakeMutableLayer(FInstancedPropertyBag& InPropertyBag);
+
+	// Create a cached parameter layer from set of params. Mutabilty is on a per-parameter basis.
+	static ANIMNEXT_API FLayerHandle MakeLayer(TConstArrayView<TPair<FParamId, FParam>> InParams);
 
 	// Make a parameter layer from a value
 	// @param	InParamId			Parameter ID for the parameter
 	// @param	InValue				Value to push
 	template<typename... ValueType>
-	static TUniquePtr<FParamStackLayer> MakeValueLayer(FParamId InParamId, ValueType&&... InValue)
+	static FLayerHandle MakeValueLayer(FParamId InParamId, ValueType&&... InValue)
 	{
 		return MakeValuesLayer(InParamId, Forward<ValueType>(InValue)...);
 	}
@@ -184,7 +274,7 @@ public:
 	// @param	InParamId			Parameter name for the parameter
 	// @param	InValue				Value to push
 	template<typename... ValueType>
-	static TUniquePtr<FParamStackLayer> MakeValueLayer(FName InName, ValueType&&... InValue)
+	static FLayerHandle MakeValueLayer(FName InName, ValueType&&... InValue)
 	{
 		return MakeValuesLayer(FParamId(InName), Forward<ValueType>(InValue)...);
 	}
@@ -192,7 +282,7 @@ public:
 	// Make a parameter layer from an interleaved parameter list of keys and values
 	// @param	InValues	Values to push
 	template<typename... Args>
-	static TUniquePtr<FParamStackLayer> MakeValuesLayer(Args&&... InValues)
+	static FLayerHandle MakeValuesLayer(Args&&... InValues)
 	{
 		constexpr int32 NumItems = sizeof...(InValues) / 2;
 		TArray<TPair<FParamId, FParam>, TInlineAllocator<NumItems>> ParamIdValues;
@@ -200,23 +290,26 @@ public:
 		return MakeValuesLayerHelper(ParamIdValues, Forward<Args>(InValues)...);
 	}
 
-	// Pop a parameter layer
-	ANIMNEXT_API void PopLayer();
+	// Pop a parameter layer - asserts if the layer supplied is not the top layer
+	ANIMNEXT_API void PopLayer(FPushedLayerHandle InLayer);
 
 	// Results for FParamStack::GetParam
 	enum class EGetParamResult : uint8
 	{
 		// The requested parameter is present in the current scope
-		Succeeded		= 0x00,
+		Succeeded			= 0x00,
 
-		// The requested parameter is not present in the current scope
-		NotInScope		= 0x01,
+		// The requested parameter is present in the current scope, and its type is compatible according to the supplied requirements
+		CompatibleType		= 0x01,
 
-		// The requested parameter is present in the current scope, but is of a different type to the requested type
-		IncorrectType	= 0x02,
+		// The requested parameter is present in the current scope, but its type is not compatible according to the supplied requirements
+		IncompatibleType	= 0x02,
+
+		// The requested parameter is not present in the current scope
+		NotInScope			= 0x04,
 
 		// The requested parameter is immutable but a mutable request was made of it
-		Immutable		= 0x04,
+		Immutable			= 0x08,
 	};
 
 	FRIEND_ENUM_CLASS_FLAGS(FParamStack::EGetParamResult);
@@ -229,7 +322,8 @@ public:
 	const ValueType* GetParamPtr(FParamId InParamId, FParamStack::EGetParamResult* OutResult = nullptr) const
 	{
 		TConstArrayView<uint8> Data;
-		FParamStack::EGetParamResult Result = GetParamData(InParamId, FParamTypeHandle::GetHandle<ValueType>(), Data);
+		FParamTypeHandle TypeHandle;
+		FParamStack::EGetParamResult Result = GetParamData(InParamId, FParamTypeHandle::GetHandle<ValueType>(), Data, TypeHandle);
 		if (OutResult)
 		{
 			*OutResult = Result;
@@ -245,7 +339,8 @@ public:
 	ValueType* GetMutableParamPtr(FParamId InParamId, FParamStack::EGetParamResult* OutResult = nullptr)
 	{
 		TArrayView<uint8> Data;
-		FParamStack::EGetParamResult Result = GetMutableParamData(InParamId, FParamTypeHandle::GetHandle<ValueType>(), Data);
+		FParamTypeHandle TypeHandle;
+		FParamStack::EGetParamResult Result = GetMutableParamData(InParamId, FParamTypeHandle::GetHandle<ValueType>(), Data, TypeHandle);
 		if (OutResult)
 		{
 			*OutResult = Result;
@@ -354,39 +449,33 @@ public:
 	}
 
 	// Get a parameter's data given an FParamId.
-	// @param	InParamId			Parameter ID to find the currently-pushed value for
-	// @param	InParamTypeHandle	Handle to the type of the parameter, which must match the stored type for a value to
-	//								be returned
-	// @param	OutParamData		View into the parameters data that will be filled
+	// @param	InParamId				Parameter ID to find the currently-pushed value for
+	// @param	InParamTypeHandle		Handle to the type of the parameter, which must match the stored type for a value to
+	//									be returned according to the supplied compatibility
+	// @param	OutParamData			View into the parameters data that will be filled
+	// @param	OutParamTypeHandle		Handle to the type of the parameter if the parameter is present in the stack.
+	//									This will be valid even if the types are deemed incompatible.
+	// @param	InRequiredCompatibility	The minimum required compatibility level
 	// @return an enum describing the result of the operation @see FParamStack::EGetParamResult
-	ANIMNEXT_API EGetParamResult GetParamData(FParamId InId, FParamTypeHandle InTypeHandle, TConstArrayView<uint8>& OutParamData) const;
+	ANIMNEXT_API EGetParamResult GetParamData(FParamId InId, FParamTypeHandle InTypeHandle, TConstArrayView<uint8>& OutParamData, FParamCompatibility InRequiredCompatibility = FParamCompatibility::Equal()) const;
+	ANIMNEXT_API EGetParamResult GetParamData(FParamId InId, FParamTypeHandle InTypeHandle, TConstArrayView<uint8>& OutParamData, FParamTypeHandle& OutParamTypeHandle, FParamCompatibility InRequiredCompatibility = FParamCompatibility::Equal()) const;
 
 	// Get a parameter's mutable data given an FParamId.
-	// @param	InParamId			Parameter ID to find the currently-pushed value for
-	// @param	InParamTypeHandle	Handle to the type of the parameter, which must match the stored type for a value to
-	//								be returned
-	// @param	OutParamData		View into the parameters data that will be filled
+	// @param	InParamId				Parameter ID to find the currently-pushed value for
+	// @param	InParamTypeHandle		Handle to the type of the parameter, which must match the stored type for a value to
+	//									be returned according to the supplied compatibility
+	// @param	OutParamData			View into the parameters data that will be filled
+	// @param	OutParamTypeHandle		Handle to the type of the parameter if the parameter is present in the stack.
+	//									This will be valid even if the types are deemed incompatible.
+	// @param	InRequiredCompatibility	The minimum required compatibility level
 	// @return an enum describing the result of the operation @see FParamStack::EGetParamResult
-	ANIMNEXT_API EGetParamResult GetMutableParamData(FParamId InId, FParamTypeHandle InTypeHandle, TArrayView<uint8>& OutParamData);
+	ANIMNEXT_API EGetParamResult GetMutableParamData(FParamId InId, FParamTypeHandle InTypeHandle, TArrayView<uint8>& OutParamData, FParamCompatibility InRequiredCompatibility = FParamCompatibility::Equal());
+	ANIMNEXT_API EGetParamResult GetMutableParamData(FParamId InId, FParamTypeHandle InTypeHandle, TArrayView<uint8>& OutParamData, FParamTypeHandle& OutParamTypeHandle, FParamCompatibility InRequiredCompatibility = FParamCompatibility::Equal());
 
 private:
-	// A layer held on the stack, may be owned by the stack or not
-	struct FPushedLayer
-	{
-		FPushedLayer(FParamStackLayer& InLayer, FParamStack& InStack);
-
-		FPushedLayer(const FPushedLayer& InPreviousLayer, FParamStackLayer& InLayer, FParamStack& InStack);
-
-		FParamStackLayer& Layer;
-
-		// Index offset for the previous layer for each active param.
-		// Offset into FParamStack::PreviousLayerIndices.
-		uint32 PreviousLayerIndexStart = 0;
-	};
-
 	// Recursive helper function for PushValues
 	template <uint32 NumItems, typename FirstType, typename SecondType, typename... OtherTypes>
-	void PushValuesHelper(TArray<TPair<FParamId, FParam>, TInlineAllocator<NumItems>>& InArray, FirstType&& InFirst, SecondType&& InSecond, OtherTypes&&... InOthers)
+	FPushedLayerHandle PushValuesHelper(TArray<TPair<FParamId, FParam>, TInlineAllocator<NumItems>>& InArray, FirstType&& InFirst, SecondType&& InSecond, OtherTypes&&... InOthers)
 	{
 		InArray.Emplace(FParamId(InFirst),
 			FParam(
@@ -399,17 +488,17 @@ private:
 
 		if constexpr (sizeof...(InOthers) > 0)
 		{
-			PushValuesHelper(InArray, Forward<OtherTypes>(InOthers)...);
+			return PushValuesHelper(InArray, Forward<OtherTypes>(InOthers)...);
 		}
 		else
 		{
-			PushLayer(InArray);
+			return PushLayer(InArray);
 		}
 	}
 
 	// Recursive helper function for MakeValuesLayer
 	template <uint32 NumItems, typename FirstType, typename SecondType, typename... OtherTypes>
-	static TUniquePtr<FParamStackLayer> MakeValuesLayerHelper(TArray<TPair<FParamId, FParam>, TInlineAllocator<NumItems>>& InArray, FirstType&& InFirst, SecondType&& InSecond, OtherTypes&&... InOthers)
+	static FLayerHandle MakeValuesLayerHelper(TArray<TPair<FParamId, FParam>, TInlineAllocator<NumItems>>& InArray, FirstType&& InFirst, SecondType&& InSecond, OtherTypes&&... InOthers)
 	{
 		InArray.Emplace(FParamId(InFirst), 
 			FParam(
@@ -430,6 +519,9 @@ private:
 		}
 	}
 
+	// Helper function for layer pushes
+	FPushedLayerHandle PushLayerInternal(FLayer& InLayer);
+
 	// Check whether a given parameter is mutable
 	ANIMNEXT_API bool IsMutableParam(FParamId InId) const;
 
@@ -458,12 +550,15 @@ private:
 	// Resize layer indices to deal with any new params we have seen since the stack was created
 	void ResizeLayerIndices();
 
+	// Get a new serial number for a pushed layer
+	uint32 MakeSerialNumber();
+
 	// Layer stack
 	TArray<FPushedLayer> Layers;
 
 	// Owned layers, paged for stable addresses as layers reference them by ptr
 	// Grows with each new pushed owned layer
-	TPagedArray<FParamStackLayer, 4096> OwnedStackLayers;
+	TPagedArray<FLayer, 4096> OwnedStackLayers;
 
 	// Parameter storage for owned layers, paged for stable addresses as layer parameters reference them by ptr
 	// Grows with each new pushed layer, free'd when the param stack is destroyed
@@ -477,33 +572,12 @@ private:
 	// Pushed layers hold views into this array.
 	// Grows with each new pushed layer, free'd when the param stack is destroyed
 	TArray<uint16> PreviousLayerIndices;
+
+	// Serial number for stack layers
+	uint32 SerialNumber = 0;
 };
 
 ENUM_CLASS_FLAGS(FParamStack::EGetParamResult);
 ENUM_CLASS_FLAGS(FParamStack::EParamFlags);
 
-// Stack layers are what actually get pushed/popped onto the layer stack.
-// They are designed to be held on an instance, their data updated in place.
-// Memory ownership for params is assumed to be outside of the stack and layers
-struct FParamStackLayer
-{
-private:
-	friend struct FParamStack;
-
-	FParamStackLayer() = delete;
-
-	explicit FParamStackLayer(const FInstancedPropertyBag& InPropertyBag);
-
-	explicit FParamStackLayer(TConstArrayView<TPair<FParamId, FParamStack::FParam>> InParams);
-
-	// Params that this layer supplies
-	TArray<FParamStack::FParam> Params;
-
-	// ID that the param indices start at. Maps the global param ID range into the range for this layer
-	uint32 MinParamId = 0;
-
-	// Storage offset for this layer if it is owned internally, otherwise MAX_uint32
-	uint32 OwnedStorageOffset = MAX_uint32;
-};
-
 }
diff --git a/Engine/Plugins/Experimental/AnimNext/Source/AnimNext/Public/Param/ParamType.h b/Engine/Plugins/Experimental/AnimNext/Source/AnimNext/Public/Param/ParamType.h
--- a/Engine/Plugins/Experimental/AnimNext/Source/AnimNext/Public/Param/ParamType.h
+++ b/Engine/Plugins/Experimental/AnimNext/Source/AnimNext/Public/Param/ParamType.h
@@ -1,4 +1,4 @@
-﻿// Copyright Epic Games, Inc. All Rights Reserved.
+// Copyright Epic Games, Inc. All Rights Reserved.
 
 #pragma once
 
@@ -49,7 +49,7 @@ namespace UE::AnimNext::Private
 
 /**
  * Representation of a parameter's type. Serializable, but fairly heavyweight to pass around and compare.
- * Faster comparisons and other operations can be performed on UE::AnimNext::FTypeHandle, but they cannot be
+ * Faster comparisons and other operations can be performed on UE::AnimNext::FParamTypeHandle, but they cannot be
  * serialized as they are not stable across runs.
  */
 USTRUCT()
diff --git a/Engine/Plugins/Experimental/AnimNext/Source/AnimNext/Public/Param/ParamTypeHandle.h b/Engine/Plugins/Experimental/AnimNext/Source/AnimNext/Public/Param/ParamTypeHandle.h
--- a/Engine/Plugins/Experimental/AnimNext/Source/AnimNext/Public/Param/ParamTypeHandle.h
+++ b/Engine/Plugins/Experimental/AnimNext/Source/AnimNext/Public/Param/ParamTypeHandle.h
@@ -52,6 +52,7 @@ struct ANIMNEXT_API FParamTypeHandle
 
 	friend struct ::FAnimNextParamType;
 	friend struct FParamHelpers;
+	friend struct FParamUtils;
 	friend class UE::AnimNext::Tests::FParamTypesTest;
 
 private:
@@ -67,12 +68,6 @@ private:
 	};
 
 private:
-	/** Get the built-in parameter type */
-	EParamType GetParameterType() const
-	{
-		return (EParamType)Fields.BuiltInType;
-	}
-
 	/** Set the built-in parameter type */
 	void SetParameterType(EParamType InType)
 	{
@@ -91,9 +86,12 @@ private:
 		checkf(InIndex < (1 << 24), TEXT("FTypeHandle::SetCustomTypeIndex: Type Index out of range"));
 		Fields.CustomTypeIndex = InIndex;
 	}
-
-	/** Reset the custom type registry, used in tests */
-	static void ResetCustomTypes();
+	
+#if WITH_DEV_AUTOMATION_TESTS
+	// Used to isolate param type handles from automated tests
+	static void BeginTestSandbox();
+	static void EndTestSandbox();
+#endif
 
 	/** Get a custom type index for the passed-in type information */
 	static uint32 GetOrAllocateCustomTypeIndex(FAnimNextParamType::EValueType InValueType, FAnimNextParamType::EContainerType InContainerType, const UObject* InValueTypeObject);
@@ -209,6 +207,15 @@ public:
 	{
 	}
 
+	/** Get the built-in parameter type */
+	EParamType GetParameterType() const
+	{
+		return (EParamType)Fields.BuiltInType;
+	}
+
+	/** Get the custom type info */
+	void GetCustomTypeInfo(FAnimNextParamType::EValueType& OutValueType, FAnimNextParamType::EContainerType& OutContainerType, const UObject*& OutValueTypeObject) const;
+
 	/** Check whether this describes a built-in type (i.e. not a custom type) */
 	bool IsBuiltInType() const
 	{
diff --git a/Engine/Plugins/Experimental/AnimNext/Source/AnimNextEditor/Internal/AnimNextEditorModule.h b/Engine/Plugins/Experimental/AnimNext/Source/AnimNextEditor/Internal/AnimNextEditorModule.h
--- a/Engine/Plugins/Experimental/AnimNext/Source/AnimNextEditor/Internal/AnimNextEditorModule.h
+++ b/Engine/Plugins/Experimental/AnimNext/Source/AnimNextEditor/Internal/AnimNextEditorModule.h
@@ -10,20 +10,13 @@ class SWidget;
 namespace UE::AnimNext::Editor
 {
 
-class FAssetTypeActions_AnimNextGraph;
 struct FParameterPickerArgs;
 
-class FModule : public IModuleInterface
+class IModule : public IModuleInterface
 {
 public:
 	// Create a parameter picker
-	TSharedPtr<SWidget> CreateParameterPicker(const FParameterPickerArgs& InArgs);
-	
-private:
-	virtual void StartupModule() override;
-	virtual void ShutdownModule() override;
-
-	TSharedPtr<FAssetTypeActions_AnimNextGraph> AssetTypeActions_AnimNextGraph;
+	virtual TSharedRef<SWidget> CreateParameterPicker(const FParameterPickerArgs& InArgs) = 0;
 };
 
 }
\ No newline at end of file
diff --git a/Engine/Plugins/Experimental/AnimNext/Source/AnimNextEditor/Internal/Param/ParameterPickerArgs.h b/Engine/Plugins/Experimental/AnimNext/Source/AnimNextEditor/Internal/Param/ParameterPickerArgs.h
--- a/Engine/Plugins/Experimental/AnimNext/Source/AnimNextEditor/Internal/Param/ParameterPickerArgs.h
+++ b/Engine/Plugins/Experimental/AnimNext/Source/AnimNextEditor/Internal/Param/ParameterPickerArgs.h
@@ -5,6 +5,7 @@
 #include "UObject/NameTypes.h"
 #include "Delegates/DelegateCombinations.h"
 #include "AssetRegistry/AssetData.h"
+#include "Param/ParamTypeHandle.h"
 
 namespace UE::AnimNext::Editor
 {
@@ -34,6 +35,9 @@ struct FParameterBindingReference
 // Delegate called when a parameter has been picked. Block argument is invalid when an unbound parameter is chosen.
 DECLARE_DELEGATE_OneParam(FOnGetParameterBindings, TArray<FParameterBindingReference>& /*OutParameterBindings*/);
 
+// Delegate called when a parameter has been picked. Block argument is invalid when an unbound parameter is chosen.
+DECLARE_DELEGATE_OneParam(FOnParameterPicked, const FParameterBindingReference& /*OutParameterBinding*/);
+
 // Result of a filter operation via FOnFilterParameter
 enum class EFilterParameterResult : int32
 {
@@ -44,6 +48,9 @@ enum class EFilterParameterResult : int32
 // Delegate called to filter parameters for display to the user
 DECLARE_DELEGATE_RetVal_OneParam(EFilterParameterResult, FOnFilterParameter, const FParameterBindingReference& /*InParameterBinding*/);
 
+// Delegate called to filter parameters by type for display to the user
+DECLARE_DELEGATE_RetVal_OneParam(EFilterParameterResult, FOnFilterParameterType, const FAnimNextParamType& /*InParameterType*/);
+
 struct FParameterPickerArgs
 {
 	FParameterPickerArgs() = default;
@@ -55,14 +62,29 @@ struct FParameterPickerArgs
 	// Delegate used to signal whether selection has changed
 	FSimpleDelegate OnSelectionChanged;
 
+	// Delegate called when a single parameter has been picked
+	FOnParameterPicked OnParameterPicked;
+
 	// Delegate called to filter parameters for display to the user
 	FOnFilterParameter OnFilterParameter;
 
+	// Delegate called to filter parameters by type for display to the user
+	FOnFilterParameterType OnFilterParameterType;
+
+	// Whether we allow selecting multiple parameters or just one
+	bool bMultiSelect = true;
+
 	// Whether we should show parameters that are bound in a parameter block
 	bool bShowBoundParameters = true;
 
-	// Whether we should show parameters that are not bound in a parameter block
+	// Whether we should show parameters that are not bound in a parameter block (if bShowBoundParameters is false this will show all parameters)
 	bool bShowUnboundParameters = true;	
+
+	// Whether we should show the library alongside parameters
+	bool bShowLibraries = true;
+
+	// Whether we should show the block alongside bound parameters
+	bool bShowBlocks = true;
 };
 
 }
diff --git a/Engine/Plugins/Experimental/AnimNext/Source/AnimNextEditor/Private/AnimNextEditorModule.cpp b/Engine/Plugins/Experimental/AnimNext/Source/AnimNextEditor/Private/AnimNextEditorModule.cpp
--- a/Engine/Plugins/Experimental/AnimNext/Source/AnimNextEditor/Private/AnimNextEditorModule.cpp
+++ b/Engine/Plugins/Experimental/AnimNext/Source/AnimNextEditor/Private/AnimNextEditorModule.cpp
@@ -12,39 +12,45 @@
 namespace UE::AnimNext::Editor
 {
 
-void FModule::StartupModule()
+class FModule : public IModule
 {
-	IAssetTools& AssetTools = FModuleManager::LoadModuleChecked<FAssetToolsModule>("AssetTools").Get();
-	AssetTypeActions_AnimNextGraph = MakeShared<FAssetTypeActions_AnimNextGraph>();
-	AssetTools.RegisterAssetTypeActions(AssetTypeActions_AnimNextGraph.ToSharedRef());
 
-	FPropertyEditorModule& PropertyModule = FModuleManager::LoadModuleChecked<FPropertyEditorModule>("PropertyEditor");
+	virtual void StartupModule() override
+	{
+		IAssetTools& AssetTools = FModuleManager::LoadModuleChecked<FAssetToolsModule>("AssetTools").Get();
+		AssetTypeActions_AnimNextGraph = MakeShared<FAssetTypeActions_AnimNextGraph>();
+		AssetTools.RegisterAssetTypeActions(AssetTypeActions_AnimNextGraph.ToSharedRef());
 
-	PropertyModule.RegisterCustomPropertyTypeLayout(
-		"AnimNextParamType",
-		FOnGetPropertyTypeCustomizationInstance::CreateLambda([] { return MakeShared<FParamTypePropertyTypeCustomization>(); }));
-}
+		FPropertyEditorModule& PropertyModule = FModuleManager::LoadModuleChecked<FPropertyEditorModule>("PropertyEditor");
 
-void FModule::ShutdownModule()
-{
-	if(FModuleManager::Get().IsModuleLoaded("AssetTools"))
-	{
-		IAssetTools& AssetTools = FModuleManager::GetModuleChecked<FAssetToolsModule>("AssetTools").Get();
-		AssetTools.UnregisterAssetTypeActions(AssetTypeActions_AnimNextGraph.ToSharedRef());
+		PropertyModule.RegisterCustomPropertyTypeLayout(
+			"AnimNextParamType",
+			FOnGetPropertyTypeCustomizationInstance::CreateLambda([] { return MakeShared<FParamTypePropertyTypeCustomization>(); }));
 	}
+
+	virtual void ShutdownModule() override
+	{
+		if(FModuleManager::Get().IsModuleLoaded("AssetTools"))
+		{
+			IAssetTools& AssetTools = FModuleManager::GetModuleChecked<FAssetToolsModule>("AssetTools").Get();
+			AssetTools.UnregisterAssetTypeActions(AssetTypeActions_AnimNextGraph.ToSharedRef());
+		}
 	
-	if(FModuleManager::Get().IsModuleLoaded("PropertyEditor"))
+		if(FModuleManager::Get().IsModuleLoaded("PropertyEditor"))
+		{
+			FPropertyEditorModule& PropertyModule = FModuleManager::LoadModuleChecked<FPropertyEditorModule>("PropertyEditor");
+			PropertyModule.UnregisterCustomPropertyTypeLayout("AnimNextParamType");
+		}
+	}
+
+	virtual TSharedRef<SWidget> CreateParameterPicker(const FParameterPickerArgs& InArgs) override
 	{
-		FPropertyEditorModule& PropertyModule = FModuleManager::LoadModuleChecked<FPropertyEditorModule>("PropertyEditor");
-		PropertyModule.UnregisterCustomPropertyTypeLayout("AnimNextParamType");
+		return SNew(SParameterPicker)
+			.Args(InArgs);
 	}
-}
 
-TSharedPtr<SWidget> FModule::CreateParameterPicker(const FParameterPickerArgs& InArgs)
-{
-	return SNew(SParameterPicker)
-		.Args(InArgs);
-}
+	TSharedPtr<FAssetTypeActions_AnimNextGraph> AssetTypeActions_AnimNextGraph;
+};
 
 }
 
diff --git a/Engine/Plugins/Experimental/AnimNext/Source/AnimNextEditor/Private/Graph/GraphFactory.cpp b/Engine/Plugins/Experimental/AnimNext/Source/AnimNextEditor/Private/Graph/GraphFactory.cpp
--- a/Engine/Plugins/Experimental/AnimNext/Source/AnimNextEditor/Private/Graph/GraphFactory.cpp
+++ b/Engine/Plugins/Experimental/AnimNext/Source/AnimNextEditor/Private/Graph/GraphFactory.cpp
@@ -30,7 +30,7 @@ UObject* UAnimNextGraphFactory::FactoryCreateNew(UClass* Class, UObject* InParen
 	UAnimNextGraph_EditorData* EditorData = NewObject<UAnimNextGraph_EditorData>(NewGraph, TEXT("EditorData"));
 	NewGraph->EditorData = EditorData;
 	EditorData->Initialize(/*bRecompileVM*/false);
-	EditorData->GetRigVMClient()->SetExecuteContextStruct(FAnimNextGraphExecuteContext::StaticStruct());
+	EditorData->GetRigVMClient()->SetExecuteContextStruct(FAnimNextExecuteContext::StaticStruct());
 
 	// Add initial execution unit
 	URigVMController* Controller = EditorData->GetRigVMClient()->GetController(EditorData->RootGraph);
diff --git a/Engine/Plugins/Experimental/AnimNext/Source/AnimNextEditor/Private/Graph/SActionMenu.cpp b/Engine/Plugins/Experimental/AnimNext/Source/AnimNextEditor/Private/Graph/SActionMenu.cpp
--- a/Engine/Plugins/Experimental/AnimNext/Source/AnimNextEditor/Private/Graph/SActionMenu.cpp
+++ b/Engine/Plugins/Experimental/AnimNext/Source/AnimNextEditor/Private/Graph/SActionMenu.cpp
@@ -15,7 +15,7 @@
 #include "RigVMCore/RigVMRegistry.h"
 #include "Units/RigUnit.h"
 #include "Widgets/SToolTip.h"
-#include "Graph/GraphExecuteContext.h"
+#include "Graph/AnimNextExecuteContext.h"
 
 #define LOCTEXT_NAMESPACE "AnimNextEditor"
 
@@ -27,7 +27,7 @@ static void CollectAllAnimNextGraphActions(FGraphContextMenuBuilder& MenuBuilder
 	static const TArray<UScriptStruct*> AllowedExecuteContexts =
 	{
 		FRigVMExecuteContext::StaticStruct(),
-		FAnimNextGraphExecuteContext::StaticStruct()
+		FAnimNextExecuteContext::StaticStruct()
 	};
 
 	for(const FRigVMFunction& Function : FRigVMRegistry::Get().GetFunctions())
@@ -66,7 +66,7 @@ static void CollectAllAnimNextGraphActions(FGraphContextMenuBuilder& MenuBuilder
 
 	for (const FRigVMDispatchFactory* Factory : FRigVMRegistry::Get().GetFactories())
 	{
-		if (!Factory->SupportsExecuteContextStruct(FRigVMExecuteContext::StaticStruct()) && !Factory->SupportsExecuteContextStruct(FAnimNextGraphExecuteContext::StaticStruct()))
+		if (!Factory->SupportsExecuteContextStruct(FRigVMExecuteContext::StaticStruct()) && !Factory->SupportsExecuteContextStruct(FAnimNextExecuteContext::StaticStruct()))
 		{
 			continue;
 		}
@@ -119,7 +119,7 @@ void SActionMenu::Construct(const FArguments& InArgs)
 					.VAlign(VAlign_Center)
 					[
 						SNew(STextBlock)
-						.Text(LOCTEXT("ContextText", "All AnimNext Interface Node Classes"))
+						.Text(LOCTEXT("ContextText", "All AnimNext Nodes"))
 						.Font(FAppStyle::Get().GetFontStyle("BlueprintEditor.ActionMenu.ContextDescriptionFont"))
 						.ToolTip(IDocumentation::Get()->CreateToolTip(
 							LOCTEXT("ActionMenuContextTextTooltip", "Describes the current context of the action list"),
diff --git a/Engine/Plugins/Experimental/AnimNext/Source/AnimNextEditor/Private/Param/AnimNextParameterBlockEditor.cpp b/Engine/Plugins/Experimental/AnimNext/Source/AnimNextEditor/Private/Param/AnimNextParameterBlockEditor.cpp
--- a/Engine/Plugins/Experimental/AnimNext/Source/AnimNextEditor/Private/Param/AnimNextParameterBlockEditor.cpp
+++ b/Engine/Plugins/Experimental/AnimNext/Source/AnimNextEditor/Private/Param/AnimNextParameterBlockEditor.cpp
@@ -6,7 +6,7 @@
 #include "Param/AnimNextParameterBlock_EditorData.h"
 #include "EdGraphNode_Comment.h"
 #include "ExternalPackageHelper.h"
-#include "Graph/SActionMenu.h"
+#include "Param/SParametersActionMenu.h"
 #include "UncookedOnlyUtils.h"
 #include "Framework/Commands/GenericCommands.h"
 #include "RigVMModel/RigVMController.h"
@@ -15,6 +15,7 @@
 #include "Param/AnimNextParameterBlockEntry.h"
 #include "Param/AnimNextParameterSettings.h"
 #include "ParameterBlockEditorMode.h"
+#include "Widgets/Docking/SDockTab.h"
 
 #define LOCTEXT_NAMESPACE "AnimNextParameterBlockEditor"
 
@@ -61,6 +62,7 @@ void FParameterBlockEditor::InitEditor(const EToolkitMode::Type InMode, const TS
 	DocumentSummoner->OnCreateGraphEditorWidget().BindSP(this, &FParameterBlockEditor::CreateGraphEditorWidget);
 	DocumentSummoner->OnGraphEditorFocused().BindSP(this, &FParameterBlockEditor::OnGraphEditorFocused);
 	DocumentSummoner->OnGraphEditorBackgrounded().BindSP(this, &FParameterBlockEditor::OnGraphEditorBackgrounded);
+	DocumentSummoner->OnSaveGraphState().BindSP(this, &FParameterBlockEditor::HandleSaveGraphState);
 	GraphEditorTabFactoryPtr = DocumentSummoner;
 	DocumentManager->RegisterDocumentFactory(DocumentSummoner);
 
@@ -76,20 +78,32 @@ void FParameterBlockEditor::InitEditor(const EToolkitMode::Type InMode, const TS
 	ExtendMenu();
 	ExtendToolbar();
 	RegenerateMenusAndToolbars();
-	RestoreLastEditedState();
 }
 
-void FParameterBlockEditor::RestoreLastEditedState()
+void FParameterBlockEditor::RestoreEditedObjectState()
 {
 	for (const FEditedDocumentInfo& Document : EditorData->LastEditedDocuments)
 	{
 		if (UObject* Obj = Document.EditedObjectPath.ResolveObject())
 		{
-			OpenDocument(Obj, FDocumentTracker::RestorePreviousDocument);
+			if(TSharedPtr<SDockTab> DockTab = OpenDocument(Obj, FDocumentTracker::RestorePreviousDocument))
+			{
+				TSharedRef<SGraphEditor> GraphEditor = StaticCastSharedRef<SGraphEditor>(DockTab->GetContent());
+				GraphEditor->SetViewLocation(Document.SavedViewOffset, Document.SavedZoomAmount);
+			}
 		}
 	}
 }
 
+void FParameterBlockEditor::SaveEditedObjectState()
+{
+	// Clear currently edited documents
+	EditorData->LastEditedDocuments.Empty();
+
+	// Ask all open documents to save their state, which will update LastEditedDocuments
+	DocumentManager->SaveAllState();
+}
+
 TSharedPtr<SDockTab> FParameterBlockEditor::OpenDocument(const UObject* InForObject, FDocumentTracker::EOpenDocumentCause InCause)
 {
 	if(InCause != FDocumentTracker::RestorePreviousDocument)
@@ -210,7 +224,7 @@ TSharedRef<SGraphEditor> FParameterBlockEditor::CreateGraphEditorWidget(TSharedR
 
 FActionMenuContent FParameterBlockEditor::OnCreateGraphActionMenu(UEdGraph* InGraph, const FVector2D& InNodePosition, const TArray<UEdGraphPin*>& InDraggedPins, bool bAutoExpand, SGraphEditor::FActionMenuClosed InOnMenuClosed)
 {
-	TSharedRef<SActionMenu> ActionMenu = SNew(SActionMenu)
+	TSharedRef<SParametersActionMenu> ActionMenu = SNew(SParametersActionMenu)
 		.AutoExpandActionMenu(bAutoExpand)
 		.Graph(InGraph)
 		.NewNodePosition(InNodePosition)
@@ -231,6 +245,8 @@ void FParameterBlockEditor::OnNodeTitleCommitted(const FText& NewText, ETextComm
 
 void FParameterBlockEditor::CloseDocumentTab(const UObject* DocumentID)
 {
+	EditorData->LastEditedDocuments.Remove(const_cast<UObject*>(DocumentID));
+
 	TSharedRef<FTabPayload_UObject> Payload = FTabPayload_UObject::Make(DocumentID);
 	DocumentManager->CloseTab(Payload);
 }
@@ -373,6 +389,11 @@ void FParameterBlockEditor::GetSaveableObjects(TArray<UObject*>& OutObjects) con
 	}
 }
 
+void FParameterBlockEditor::HandleSaveGraphState(UEdGraph* InGraph, FVector2D InViewOffset, float InZoomAmount)
+{
+	EditorData->LastEditedDocuments.AddUnique(FEditedDocumentInfo(InGraph, InViewOffset, InZoomAmount));
+}
+
 }
 
 #undef LOCTEXT_NAMESPACE
\ No newline at end of file
diff --git a/Engine/Plugins/Experimental/AnimNext/Source/AnimNextEditor/Private/Param/AnimNextParameterBlockEditor.h b/Engine/Plugins/Experimental/AnimNext/Source/AnimNextEditor/Private/Param/AnimNextParameterBlockEditor.h
--- a/Engine/Plugins/Experimental/AnimNext/Source/AnimNextEditor/Private/Param/AnimNextParameterBlockEditor.h
+++ b/Engine/Plugins/Experimental/AnimNext/Source/AnimNextEditor/Private/Param/AnimNextParameterBlockEditor.h
@@ -106,10 +106,14 @@ private:
 		DetailsView = InDetailsView;
 	}
 
-	void RestoreLastEditedState();
+	void RestoreEditedObjectState();
+
+	void SaveEditedObjectState();
 
 	TSharedPtr<SDockTab> OpenDocument(const UObject* InForObject, FDocumentTracker::EOpenDocumentCause InCause);
 
+	void HandleSaveGraphState(UEdGraph* InGraph, FVector2D InViewOffset, float InZoomAmount);
+
 	// The asset we are editing
 	UAnimNextParameterBlock* ParameterBlock = nullptr;
 
diff --git a/Engine/Plugins/Experimental/AnimNext/Source/AnimNextEditor/Private/Param/AnimNextParameterLibraryEditor.cpp b/Engine/Plugins/Experimental/AnimNext/Source/AnimNextEditor/Private/Param/AnimNextParameterLibraryEditor.cpp
--- a/Engine/Plugins/Experimental/AnimNext/Source/AnimNextEditor/Private/Param/AnimNextParameterLibraryEditor.cpp
+++ b/Engine/Plugins/Experimental/AnimNext/Source/AnimNextEditor/Private/Param/AnimNextParameterLibraryEditor.cpp
@@ -1,10 +1,9 @@
-﻿// Copyright Epic Games, Inc. All Rights Reserved.
+// Copyright Epic Games, Inc. All Rights Reserved.
 
 #include "AnimNextParameterLibraryEditor.h"
 #include "ParameterLibraryEditorMode.h"
 #include "Param/AnimNextParameterLibrary.h"
 #include "ExternalPackageHelper.h"
-#include "Graph/SActionMenu.h"
 #include "UncookedOnlyUtils.h"
 #include "Framework/Commands/GenericCommands.h"
 #include "ScopedTransaction.h"
diff --git a/Engine/Plugins/Experimental/AnimNext/Source/AnimNextEditor/Private/Param/ParameterBlockEditorMode.cpp b/Engine/Plugins/Experimental/AnimNext/Source/AnimNextEditor/Private/Param/ParameterBlockEditorMode.cpp
--- a/Engine/Plugins/Experimental/AnimNext/Source/AnimNextEditor/Private/Param/ParameterBlockEditorMode.cpp
+++ b/Engine/Plugins/Experimental/AnimNext/Source/AnimNextEditor/Private/Param/ParameterBlockEditorMode.cpp
@@ -11,6 +11,9 @@
 #include "SParameterBlockView.h"
 #include "Param/AnimNextParameterBlock_EditorData.h"
 #include "RigVMModel/RigVMGraph.h"
+#include "Param/IAnimNextParameterBlockGraphInterface.h"
+#include "Param/AnimNextParameterBlockEntry.h"
+#include "Param/AnimNextParameterBlockEditor.h"
 
 #define LOCTEXT_NAMESPACE "ParameterBlockEditorMode"
 
@@ -228,7 +231,8 @@ private:
 	{
 		return SNew(SParameterBlockView, StaticCastSharedPtr<FParameterBlockEditor>(HostingApp.Pin())->EditorData)
 			.OnSelectionChanged(this, &FParameterBlockTabSummoner::HandleSelectionChanged)
-			.OnOpenGraph(this, &FParameterBlockTabSummoner::HandleOpenGraph);
+			.OnOpenGraph(this, &FParameterBlockTabSummoner::HandleOpenGraph)
+			.OnDeleteEntries(this, &FParameterBlockTabSummoner::HandleDeleteEntries);
 	}
 	
 	virtual FText GetTabToolTipText(const FWorkflowTabSpawnInfo& Info) const override
@@ -251,14 +255,32 @@ private:
 			ParameterBlockEditor->OpenDocument(EditorObject, FDocumentTracker::EOpenDocumentCause::OpenNewDocument);
 		}
 	}
+
+	void HandleDeleteEntries(const TArray<UAnimNextParameterBlockEntry*>& InEntries) const
+	{
+		TSharedPtr<FParameterBlockEditor> ParameterBlockEditor = StaticCastSharedPtr<FParameterBlockEditor>(HostingApp.Pin());
+		UAnimNextParameterBlock_EditorData* EditorData = StaticCastSharedPtr<FParameterBlockEditor>(HostingApp.Pin())->EditorData;
+
+		for(UAnimNextParameterBlockEntry* Entry : InEntries)
+		{
+			if(IAnimNextParameterBlockGraphInterface* GraphInterface = Cast<IAnimNextParameterBlockGraphInterface>(Entry))
+			{
+				if(URigVMGraph* RigVMGraph = GraphInterface->GetGraph())
+				{
+					if (UObject* EditorObject = EditorData->GetEditorObjectForRigVMGraph(RigVMGraph))
+					{
+						ParameterBlockEditor->CloseDocumentTab(EditorObject);
+					}
+				}
+			}
+		}
+	}
 };
 
-FParameterBlockEditorMode::FParameterBlockEditorMode(TSharedRef<FWorkflowCentricApplication> InHostingApp)
+FParameterBlockEditorMode::FParameterBlockEditorMode(TSharedRef<FParameterBlockEditor> InHostingApp)
 	: FApplicationMode(ParameterBlockModes::ParameterBlockEditor)
-	, HostingAppPtr(InHostingApp)
+	, ParameterBlockEditorPtr(InHostingApp)
 {
-	HostingAppPtr = InHostingApp;
-
 	TSharedRef<FParameterBlockEditor> ParametersEditor = StaticCastSharedRef<FParameterBlockEditor>(InHostingApp);
 	
 	TabFactories.RegisterFactory(MakeShared<FParameterBlockDetailsTabSummoner>(ParametersEditor, FOnDetailsViewCreated::CreateSP(&ParametersEditor.Get(), &FParameterBlockEditor::HandleDetailsViewCreated)));
@@ -303,7 +325,7 @@ FParameterBlockEditorMode::FParameterBlockEditorMode(TSharedRef<FWorkflowCentric
 
 void FParameterBlockEditorMode::RegisterTabFactories(TSharedPtr<FTabManager> InTabManager)
 {
-	TSharedPtr<FWorkflowCentricApplication> HostingApp = HostingAppPtr.Pin();
+	TSharedPtr<FParameterBlockEditor> HostingApp = ParameterBlockEditorPtr.Pin();
 	HostingApp->PushTabFactories(TabFactories);
 
 	FApplicationMode::RegisterTabFactories(InTabManager);
@@ -313,7 +335,7 @@ void FParameterBlockEditorMode::AddTabFactory(FCreateWorkflowTabFactory FactoryC
 {
 	if (FactoryCreator.IsBound())
 	{
-		TabFactories.RegisterFactory(FactoryCreator.Execute(HostingAppPtr.Pin()));
+		TabFactories.RegisterFactory(FactoryCreator.Execute(ParameterBlockEditorPtr.Pin()));
 	}
 }
 
@@ -322,6 +344,25 @@ void FParameterBlockEditorMode::RemoveTabFactory(FName TabFactoryID)
 	TabFactories.UnregisterFactory(TabFactoryID);
 }
 
+void FParameterBlockEditorMode::PreDeactivateMode()
+{
+	FApplicationMode::PreDeactivateMode();
+
+	TSharedPtr<FParameterBlockEditor> ParameterBlockEditor = ParameterBlockEditorPtr.Pin();
+
+	ParameterBlockEditor->SaveEditedObjectState();
+}
+
+void FParameterBlockEditorMode::PostActivateMode()
+{
+	// Reopen any documents that were open when the blueprint was last saved
+	TSharedPtr<FParameterBlockEditor> ParameterBlockEditor = ParameterBlockEditorPtr.Pin();
+	ParameterBlockEditor->RestoreEditedObjectState();
+
+	FApplicationMode::PostActivateMode();
+}
+
+
 }
 
 #undef LOCTEXT_NAMESPACE
\ No newline at end of file
diff --git a/Engine/Plugins/Experimental/AnimNext/Source/AnimNextEditor/Private/Param/ParameterBlockEditorMode.h b/Engine/Plugins/Experimental/AnimNext/Source/AnimNextEditor/Private/Param/ParameterBlockEditorMode.h
--- a/Engine/Plugins/Experimental/AnimNext/Source/AnimNextEditor/Private/Param/ParameterBlockEditorMode.h
+++ b/Engine/Plugins/Experimental/AnimNext/Source/AnimNextEditor/Private/Param/ParameterBlockEditorMode.h
@@ -11,6 +11,11 @@
 class FWorkflowCentricApplication;
 class SGraphEditor;
 
+namespace UE::AnimNext::Editor
+{
+	class FParameterBlockEditor;
+}
+
 namespace UE::AnimNext::Editor
 {
 
@@ -77,16 +82,18 @@ private:
 class FParameterBlockEditorMode : public FApplicationMode
 {
 public:
-	FParameterBlockEditorMode(TSharedRef<FWorkflowCentricApplication> InHostingApp);
+	FParameterBlockEditorMode(TSharedRef<FParameterBlockEditor> InHostingApp);
 
 private:
 	// FApplicationMode interface
 	virtual void RegisterTabFactories(TSharedPtr<FTabManager> InTabManager) override;
 	virtual void AddTabFactory(FCreateWorkflowTabFactory FactoryCreator) override;
 	virtual void RemoveTabFactory(FName TabFactoryID) override;
-	
+	virtual void PreDeactivateMode() override;
+	virtual void PostActivateMode() override;
+
 	// The hosting app
-	TWeakPtr<FWorkflowCentricApplication> HostingAppPtr;
+	TWeakPtr<FParameterBlockEditor> ParameterBlockEditorPtr;
 
 	// The tab factories we support
 	FWorkflowAllowedTabSet TabFactories;
diff --git a/Engine/Plugins/Experimental/AnimNext/Source/AnimNextEditor/Private/Param/ParameterBlockFactory.cpp b/Engine/Plugins/Experimental/AnimNext/Source/AnimNextEditor/Private/Param/ParameterBlockFactory.cpp
--- a/Engine/Plugins/Experimental/AnimNext/Source/AnimNextEditor/Private/Param/ParameterBlockFactory.cpp
+++ b/Engine/Plugins/Experimental/AnimNext/Source/AnimNextEditor/Private/Param/ParameterBlockFactory.cpp
@@ -3,7 +3,7 @@
 #include "Param/ParameterBlockFactory.h"
 #include "Param/AnimNextParameterBlock.h"
 #include "Param/AnimNextParameterBlock_EditorData.h"
-#include "Param/ParametersExecuteContext.h"
+#include "Graph/AnimNextExecuteContext.h"
 #include "UncookedOnlyUtils.h"
 
 UAnimNextParameterBlockFactory::UAnimNextParameterBlockFactory()
@@ -25,7 +25,7 @@ UObject* UAnimNextParameterBlockFactory::FactoryCreateNew(UClass* Class, UObject
 	UAnimNextParameterBlock_EditorData* EditorData = NewObject<UAnimNextParameterBlock_EditorData>(NewBlock, TEXT("EditorData"), RF_Transactional);
 	NewBlock->EditorData = EditorData;
 	EditorData->Initialize(/*bRecompileVM*/false);
-	EditorData->GetRigVMClient()->SetExecuteContextStruct(FAnimNextParametersExecuteContext::StaticStruct());
+	EditorData->GetRigVMClient()->SetExecuteContextStruct(FAnimNextExecuteContext::StaticStruct());
 
 	// Compile the initial skeleton
 	UE::AnimNext::UncookedOnly::FUtils::Compile(NewBlock);
diff --git a/Engine/Plugins/Experimental/AnimNext/Source/AnimNextEditor/Private/Param/ParametersEditorSchemaActions.cpp b/Engine/Plugins/Experimental/AnimNext/Source/AnimNextEditor/Private/Param/ParametersEditorSchemaActions.cpp
--- a/Engine/Plugins/Experimental/AnimNext/Source/AnimNextEditor/Private/Param/ParametersEditorSchemaActions.cpp
+++ b/Engine/Plugins/Experimental/AnimNext/Source/AnimNextEditor/Private/Param/ParametersEditorSchemaActions.cpp
@@ -57,3 +57,46 @@ UEdGraphNode* FAnimNextParameterSchemaAction_RigUnit::PerformAction(UEdGraph* Pa
 
 	return NewNode;
 }
+
+
+UEdGraphNode* FAnimNextParameterSchemaAction_DispatchFactory::PerformAction(UEdGraph* ParentGraph, TArray<UEdGraphPin*>& FromPins, const FVector2D Location, bool bSelectNewNode)
+{
+	UAnimNextParameterBlock_EditorData* EditorData = ParentGraph->GetTypedOuter<UAnimNextParameterBlock_EditorData>();
+	UAnimNextParameterBlock_EdGraphNode* NewNode = nullptr;
+	UAnimNextParameterBlock_EdGraph* EdGraph = Cast<UAnimNextParameterBlock_EdGraph>(ParentGraph);
+
+	if (EditorData != nullptr && EdGraph != nullptr)
+	{
+		const FRigVMTemplate* Template = FRigVMRegistry::Get().FindTemplate(Notation);
+		if (Template == nullptr)
+		{
+			return nullptr;
+		}
+
+		const int32 NotationHash = (int32)GetTypeHash(Notation);
+		const FString TemplateName = TEXT("RigVMTemplate_") + FString::FromInt(NotationHash);
+
+		FName Name = UE::AnimNext::Editor::FUtils::ValidateName(EditorData, Template->GetName().ToString());
+		URigVMController* Controller = EditorData->GetRigVMClient()->GetController(ParentGraph);
+
+		Controller->OpenUndoBracket(FString::Printf(TEXT("Add '%s' Node"), *Name.ToString()));
+
+		if (URigVMTemplateNode* ModelNode = Controller->AddTemplateNode(Notation, Location, Name.ToString(), true, false))
+		{
+			NewNode = Cast<UAnimNextParameterBlock_EdGraphNode>(EdGraph->FindNodeForModelNodeName(ModelNode->GetFName()));
+
+			if (NewNode)
+			{
+				Controller->ClearNodeSelection(true);
+				Controller->SelectNode(ModelNode, true, true);
+			}
+
+			Controller->CloseUndoBracket();
+		}
+		else
+		{
+			Controller->CancelUndoBracket();
+		}
+	}
+	return NewNode;
+}
\ No newline at end of file
diff --git a/Engine/Plugins/Experimental/AnimNext/Source/AnimNextEditor/Private/Param/ParametersEditorSchemaActions.h b/Engine/Plugins/Experimental/AnimNext/Source/AnimNextEditor/Private/Param/ParametersEditorSchemaActions.h
--- a/Engine/Plugins/Experimental/AnimNext/Source/AnimNextEditor/Private/Param/ParametersEditorSchemaActions.h
+++ b/Engine/Plugins/Experimental/AnimNext/Source/AnimNextEditor/Private/Param/ParametersEditorSchemaActions.h
@@ -53,3 +53,24 @@ private:
 	// The script struct for our rig unit
 	UScriptStruct* StructTemplate = nullptr;
 };
+
+USTRUCT()
+struct FAnimNextParameterSchemaAction_DispatchFactory : public FAnimNextParameterSchemaAction
+{
+	GENERATED_BODY()
+
+	FAnimNextParameterSchemaAction_DispatchFactory() = default;
+
+	FAnimNextParameterSchemaAction_DispatchFactory(FName InNotation, FText InNodeCategory, FText InMenuDesc, FText InToolTip, FText InKeywords = FText::GetEmpty())
+		: FAnimNextParameterSchemaAction(InNodeCategory, InMenuDesc, InToolTip, InKeywords)
+		, Notation(InNotation)
+	{}
+
+	// FEdGraphSchemaAction Interface
+	virtual UEdGraphNode* PerformAction(UEdGraph* ParentGraph, TArray<UEdGraphPin*>& FromPins, const FVector2D Location, bool bSelectNewNode = true) override;
+	virtual UEdGraphNode* PerformAction(UEdGraph* ParentGraph, UEdGraphPin* FromPin, const FVector2D Location, bool bSelectNewNode = true) { return nullptr; }
+
+private:
+	// Notation for dispatch factory
+	FName Notation;
+};
diff --git a/Engine/Plugins/Experimental/AnimNext/Source/AnimNextEditor/Private/Param/SAddParametersDialog.cpp b/Engine/Plugins/Experimental/AnimNext/Source/AnimNextEditor/Private/Param/SAddParametersDialog.cpp
--- a/Engine/Plugins/Experimental/AnimNext/Source/AnimNextEditor/Private/Param/SAddParametersDialog.cpp
+++ b/Engine/Plugins/Experimental/AnimNext/Source/AnimNextEditor/Private/Param/SAddParametersDialog.cpp
@@ -204,10 +204,7 @@ class SParameterToAdd : public SMultiColumnTableRow<TSharedRef<FParameterToAdd>>
 					})
 					.OnTextCommitted_Lambda([this](const FText& InText, ETextCommit::Type InCommitType)
 					{
-						if(InCommitType == ETextCommit::OnEnter)
-						{
-							Entry->Name = *InText.ToString();
-						}
+						Entry->Name = *InText.ToString();
 					})
 					.OnVerifyTextChanged_Lambda([this](const FText& InNewText, FText& OutErrorText)
 					{
diff --git a/Engine/Plugins/Experimental/AnimNext/Source/AnimNextEditor/Private/Param/SParameterBlockView.cpp b/Engine/Plugins/Experimental/AnimNext/Source/AnimNextEditor/Private/Param/SParameterBlockView.cpp
--- a/Engine/Plugins/Experimental/AnimNext/Source/AnimNextEditor/Private/Param/SParameterBlockView.cpp
+++ b/Engine/Plugins/Experimental/AnimNext/Source/AnimNextEditor/Private/Param/SParameterBlockView.cpp
@@ -1,4 +1,4 @@
-﻿// Copyright Epic Games, Inc. All Rights Reserved.
+// Copyright Epic Games, Inc. All Rights Reserved.
 
 #include "SParameterBlockView.h"
 
@@ -89,6 +89,7 @@ void SParameterBlockView::Construct(const FArguments& InArgs, UAnimNextParameter
 
 	OnSelectionChangedDelegate = InArgs._OnSelectionChanged;
 	OnOpenGraphDelegate = InArgs._OnOpenGraph;
+	OnDeleteEntriesDelegate = InArgs._OnDeleteEntries;
 
 	EditorData->ModifiedDelegate.AddSP(this, &SParameterBlockView::HandleBlockModified);
 
@@ -372,12 +373,13 @@ void SParameterBlockView::HandleDelete()
 	if(EntriesList->GetNumItemsSelected() > 0)
 	{
 		TArray<TSharedRef<FParameterBlockViewEntry>> SelectedItems = EntriesList->GetSelectedItems();
+		TArray<UAnimNextParameterBlockEntry*> EntriesToRemove;
+		Algo::Transform(SelectedItems, EntriesToRemove, [](const TSharedRef<FParameterBlockViewEntry>& InEntry) { return InEntry->WeakEntry.Get(); });
+
+		OnDeleteEntriesDelegate.ExecuteIfBound(EntriesToRemove);
 
 		{
 			FScopedTransaction Transaction(FText::FormatOrdered(LOCTEXT("DeleteParameterBlockEntry", "Delete parameter block {0}|plural(one=entry,other=entries)"), EntriesList->GetNumItemsSelected()));
-
-			TArray<UAnimNextParameterBlockEntry*> EntriesToRemove;
-			Algo::Transform(SelectedItems, EntriesToRemove, [](const TSharedRef<FParameterBlockViewEntry>& InEntry){ return InEntry->WeakEntry.Get(); });
 			EditorData->RemoveEntries(EntriesToRemove);
 		}
 	}
diff --git a/Engine/Plugins/Experimental/AnimNext/Source/AnimNextEditor/Private/Param/SParameterBlockView.h b/Engine/Plugins/Experimental/AnimNext/Source/AnimNextEditor/Private/Param/SParameterBlockView.h
--- a/Engine/Plugins/Experimental/AnimNext/Source/AnimNextEditor/Private/Param/SParameterBlockView.h
+++ b/Engine/Plugins/Experimental/AnimNext/Source/AnimNextEditor/Private/Param/SParameterBlockView.h
@@ -1,4 +1,4 @@
-﻿// Copyright Epic Games, Inc. All Rights Reserved.
+// Copyright Epic Games, Inc. All Rights Reserved.
 
 #pragma once
 
@@ -27,6 +27,8 @@ public:
 	DECLARE_DELEGATE_OneParam(FOnSelectionChanged, const TArray<UObject*>& /*InEntries*/);
 
 	DECLARE_DELEGATE_OneParam(FOnOpenGraph, URigVMGraph* /*InGraph*/);
+
+	DECLARE_DELEGATE_OneParam(FOnDeleteEntries, const TArray<UAnimNextParameterBlockEntry*>& /*InEntries*/);
 	
 	SLATE_BEGIN_ARGS(SParameterBlockView) {}
 
@@ -34,6 +36,8 @@ public:
 	
 	SLATE_EVENT(SParameterBlockView::FOnOpenGraph, OnOpenGraph)
 
+	SLATE_EVENT(SParameterBlockView::FOnDeleteEntries, OnDeleteEntries)
+
 	SLATE_END_ARGS()
 
 	void Construct(const FArguments& InArgs, UAnimNextParameterBlock_EditorData* InEditorData);
@@ -94,6 +98,8 @@ private:
 
 	FOnOpenGraph OnOpenGraphDelegate;
 
+	FOnDeleteEntries OnDeleteEntriesDelegate;
+
 	FAssetData BlockAssetData;
 
 	TArray<UAnimNextParameterBlockEntry*> PendingSelection;
diff --git a/Engine/Plugins/Experimental/AnimNext/Source/AnimNextEditor/Private/Param/SParameterPicker.cpp b/Engine/Plugins/Experimental/AnimNext/Source/AnimNextEditor/Private/Param/SParameterPicker.cpp
--- a/Engine/Plugins/Experimental/AnimNext/Source/AnimNextEditor/Private/Param/SParameterPicker.cpp
+++ b/Engine/Plugins/Experimental/AnimNext/Source/AnimNextEditor/Private/Param/SParameterPicker.cpp
@@ -14,6 +14,9 @@
 #include "EditorUtils.h"
 #include "Param/ParameterPickerArgs.h"
 #include "Widgets/Input/SSearchBox.h"
+#include "String/ParseTokens.h"
+#include "Widgets/Layout/SBox.h"
+#include "Widgets/Images/SImage.h"
 
 #define LOCTEXT_NAMESPACE "SParameterPicker"
 
@@ -25,10 +28,20 @@ namespace ParameterPicker
 static FName Column_Parameter(TEXT("Parameter"));
 static FName Column_Library(TEXT("Library"));
 static FName Column_Block(TEXT("Block"));
+static FName Column_Type(TEXT("Type"));
 }
 
 struct FParameterPickerEntry
 {
+	enum class EFilterResult : uint8
+	{
+		DoesNotPassFilter	= 0x00,
+		PassesFilter		= 0x01,
+		ChildPassesFilter	= 0x02,
+	};
+
+	FRIEND_ENUM_CLASS_FLAGS(EFilterResult);
+
 	FParameterPickerEntry() = default;
 
 	FParameterPickerEntry(const FParameterBindingReference& InBinding, const FAnimNextParamType& InParamType)
@@ -38,6 +51,25 @@ struct FParameterPickerEntry
 		PinType = UE::AnimNext::UncookedOnly::FUtils::GetPinTypeFromParamType(ParamType);
 		PinIcon = FBlueprintEditorUtils::GetIconFromPin(PinType, /* bIsLarge = */true);
 		PinColor = GetDefault<UEdGraphSchema_K2>()->GetPinTypeColor(PinType);
+		FString ParameterString = Binding.Parameter.ToString();
+		int32 LastUnderscoreLoc = INDEX_NONE;
+		if (ParameterString.FindLastChar(TEXT('_'), LastUnderscoreLoc))
+		{
+			ParameterString.RightChopInline(LastUnderscoreLoc + 1);
+		}
+		DisplayString = MoveTemp(ParameterString);
+	}
+
+	FParameterPickerEntry(FName InName)
+	{
+		Binding.Parameter = InName;
+		FString ParameterString = Binding.Parameter.ToString();
+		int32 LastUnderscoreLoc = INDEX_NONE;
+		if (ParameterString.FindLastChar(TEXT('_'), LastUnderscoreLoc))
+		{
+			ParameterString.RightChopInline(LastUnderscoreLoc + 1);
+		}
+		DisplayString = MoveTemp(ParameterString);
 	}
 
 	bool PassesFilter(const FString& InFilterText) const
@@ -45,17 +77,29 @@ struct FParameterPickerEntry
 		return Binding.Parameter.ToString().Contains(InFilterText);
 	}
 
+	FString DisplayString;
+
 	FParameterBindingReference Binding;
 
 	FAnimNextParamType ParamType;
 
 	FEdGraphPinType PinType;
 
-	const FSlateBrush* PinIcon;
+	const FSlateBrush* PinIcon = nullptr;
+
+	FLinearColor PinColor = FLinearColor::White;
 
-	FLinearColor PinColor;
+	TSharedPtr<FParameterPickerEntry> Parent;
+
+	TArray<TSharedRef<FParameterPickerEntry>> Children;
+
+	TArray<TSharedRef<FParameterPickerEntry>> FilteredChildren;
+
+	EFilterResult FilterResult = EFilterResult::PassesFilter;
 };
 
+ENUM_CLASS_FLAGS(FParameterPickerEntry::EFilterResult);
+
 void SParameterPicker::Construct(const FArguments& InArgs)
 {
 	using namespace ParameterPicker;
@@ -67,6 +111,8 @@ void SParameterPicker::Construct(const FArguments& InArgs)
 		Args.OnGetParameterBindings->BindSP(this, &SParameterPicker::HandleGetParameterBindings);
 	}
 
+	TSharedPtr<SHeaderRow> HeaderRow;
+
 	ChildSlot
 	[
 		SNew(SVerticalBox)
@@ -85,31 +131,60 @@ void SParameterPicker::Construct(const FArguments& InArgs)
 		.FillHeight(1.0f)
 		.Padding(2.0f)
 		[
-			SAssignNew(EntriesList, SListView<TSharedRef<FParameterPickerEntry>>)
-			.ListItemsSource(&FilteredEntries)
+			SAssignNew(EntriesList, STreeView<TSharedRef<FParameterPickerEntry>>)
+			.TreeItemsSource(&FilteredHierarchy)
+			.SelectionMode(Args.bMultiSelect ? ESelectionMode::Multi : ESelectionMode::Single)
 			.OnGenerateRow(this, &SParameterPicker::HandleGenerateRow)
 			.OnSelectionChanged(this, &SParameterPicker::HandleSelectionChanged)
+			.OnGetChildren(this, &SParameterPicker::HandleGetChildren)
+			.OnIsSelectableOrNavigable(this, &SParameterPicker::HandleIsSelectableOrNavigable)
 			.ItemHeight(20.0f)
 			.HeaderRow(
-				SNew(SHeaderRow)
+				SAssignNew(HeaderRow, SHeaderRow)
+				+SHeaderRow::Column(Column_Type)
+				.DefaultLabel(FText::GetEmpty())
+				.FixedWidth(24.0f)
+				.HeaderContent()
+				[
+					SNew(SBox)
+					.WidthOverride(16.0f)
+					.HeightOverride(16.0f)
+					.VAlign(VAlign_Center)
+					.HAlign(HAlign_Center)
+					[
+						SNew(SImage)
+						.ColorAndOpacity(FSlateColor::UseForeground())
+						.Image(FAppStyle::GetBrush("Kismet.VariableList.TypeIcon"))
+						.ToolTipText(LOCTEXT("TypeColumnHeaderTooltip", "The parameter's type"))
+					]
+				]
+
 				+SHeaderRow::Column(Column_Parameter)
 				.DefaultLabel(LOCTEXT("ParameterColumnHeader", "Parameter"))
-				.ToolTipText(LOCTEXT("ParameterColumnHeaderTooltip", "The parameter"))
-				.FillWidth(0.33f)
-				
-				+SHeaderRow::Column(Column_Library)
-				.DefaultLabel(LOCTEXT("LibraryColumnHeader", "Library"))
-				.ToolTipText(LOCTEXT("LibraryColumnHeaderTooltip", "The library that the parameter is declared in"))
-				.FillWidth(0.33f)
-				
-				+SHeaderRow::Column(Column_Block)
-				.DefaultLabel(LOCTEXT("BlockColumnHeader", "Block"))
-				.ToolTipText(LOCTEXT("BlockColumnHeaderTooltip", "The parameter block that has a binding to the parameter"))
+				.ToolTipText(LOCTEXT("ParameterColumnHeaderTooltip", "The parameter's name"))
 				.FillWidth(0.33f)
 			)
 		]
 	];
 
+	if (Args.bShowLibraries)
+	{
+		HeaderRow->AddColumn(
+			SHeaderRow::Column(Column_Library)
+			.DefaultLabel(LOCTEXT("LibraryColumnHeader", "Library"))
+			.ToolTipText(LOCTEXT("LibraryColumnHeaderTooltip", "The library that the parameter is declared in"))
+			.FillWidth(0.33f));
+	}
+	
+	if (Args.bShowLibraries)
+	{
+		HeaderRow->AddColumn(
+			SHeaderRow::Column(Column_Block)
+			.DefaultLabel(LOCTEXT("BlockColumnHeader", "Block"))
+			.ToolTipText(LOCTEXT("BlockColumnHeaderTooltip", "The parameter block that has a binding to the parameter"))
+			.FillWidth(0.33f));
+	}
+
 	RefreshEntries();
 }
 
@@ -163,9 +238,11 @@ void SParameterPicker::RefreshEntries()
 						{
 							const FAnimNextParameterLibraryAssetRegistryExports& LibraryExports = GetExportsForLibrary(LibraryAsset);
 							FAnimNextParamType ParamType = FUtils::GetParameterTypeFromLibraryExports(Export.Name, LibraryExports);
-
-							TSharedRef<FParameterPickerEntry> NewEntry = MakeShared<FParameterPickerEntry>(NewReference, ParamType);
-							Entries.Add(NewEntry);
+							if(!Args.OnFilterParameterType.IsBound() || Args.OnFilterParameterType.Execute(ParamType) == EFilterParameterResult::Include)
+							{
+								TSharedRef<FParameterPickerEntry> NewEntry = MakeShared<FParameterPickerEntry>(NewReference, ParamType);
+								Entries.Add(NewEntry);
+							}
 						}
 					}
 				}
@@ -191,31 +268,173 @@ void SParameterPicker::RefreshEntries()
 					FParameterBindingReference NewReference(Export.Name, LibraryAsset);
 					if(!Args.OnFilterParameter.IsBound() || Args.OnFilterParameter.Execute(NewReference) == EFilterParameterResult::Include)
 					{
-						TSharedRef<FParameterPickerEntry> NewEntry = MakeShared<FParameterPickerEntry>(NewReference, Export.Type);
-						Entries.Add(NewEntry);
+						if (!Args.OnFilterParameterType.IsBound() || Args.OnFilterParameterType.Execute(Export.Type) == EFilterParameterResult::Include)
+						{
+							TSharedRef<FParameterPickerEntry> NewEntry = MakeShared<FParameterPickerEntry>(NewReference, Export.Type);
+							Entries.Add(NewEntry);
+						}
 					}
 				}
 			}
 		}
 	}
 
+	BuildHierarchy();
+
 	RefreshFilter();
 }
 
-void SParameterPicker::RefreshFilter()
+void SParameterPicker::BuildHierarchy()
 {
-	FilteredEntries.Empty();
+	Hierarchy.Reset();
+
+	TMap<FName, TSharedRef<FParameterPickerEntry>> NameMap;
+	NameMap.Reserve(Entries.Num());
+
+	for (const TSharedRef<FParameterPickerEntry>& Entry : Entries)
+	{
+		NameMap.Add(Entry->Binding.Parameter, Entry);
+	}
+
+	TArray< TSharedRef<FParameterPickerEntry>> NewEntries;
+	for (const TSharedRef<FParameterPickerEntry>& Entry : Entries)
+	{
+		// Parse name into separators
+		TStringBuilder<256> WholeParameterString;
+		Entry->Binding.Parameter.ToString(WholeParameterString);
+
+		TStringBuilder<256> PartialParameterString;
+
+		// We use '_' as a seperator here as:
+		// - Each param is a UObject in editor and uses its object name, so we cannot use '.'
+		// - This maps nicely to Verse tags that use '_' to hierarchically define their relationship
+		TSharedPtr<FParameterPickerEntry> Parent = nullptr;
+		UE::String::ParseTokens(WholeParameterString, TEXT('_'), [this, &PartialParameterString, &NameMap, &Parent, &NewEntries](FStringView InStringView)
+		{
+			if(Parent.IsValid())
+			{
+				// Have a parent, so add child if it doesnt exist already
+				PartialParameterString += TEXT('_');
+				PartialParameterString += InStringView;
+
+				const FName ParameterName(InStringView);
+				const FName PartialParameterName(PartialParameterString);
+				if (TSharedRef<FParameterPickerEntry>* ExistingEntry = NameMap.Find(PartialParameterName))
+				{
+					(*ExistingEntry)->Parent = Parent;
+					if (!Parent->Children.ContainsByPredicate([&ExistingEntry](const TSharedRef<FParameterPickerEntry>& InEntry) { return (*ExistingEntry)->Binding.Parameter == InEntry->Binding.Parameter; }))
+					{
+						Parent->Children.Add(*ExistingEntry);
+					}
+					Parent = *ExistingEntry;
+				}
+				else
+				{
+					TSharedRef<FParameterPickerEntry> NewEntry = MakeShared<FParameterPickerEntry>(PartialParameterName);
+					NewEntry->Parent = Parent;
+					if (!Parent->Children.ContainsByPredicate([&NewEntry](const TSharedRef<FParameterPickerEntry>& InEntry) { return NewEntry->Binding.Parameter == InEntry->Binding.Parameter; }))
+					{
+						Parent->Children.Add(NewEntry);
+					}
+					Parent = NameMap.Add(PartialParameterName, NewEntry);
+					NewEntries.Add(NewEntry);
+				}
+			}
+			else
+			{
+				PartialParameterString += InStringView;
+
+				// Add root item if it doesnt exist already
+				const FName ParameterName(InStringView);
+				if (TSharedRef<FParameterPickerEntry>* ExistingParent = NameMap.Find(ParameterName))
+				{
+					Parent = *ExistingParent;
+				}
+				else
+				{
+					TSharedRef<FParameterPickerEntry> NewEntry = MakeShared<FParameterPickerEntry>(ParameterName);
+					Parent = NameMap.Add(ParameterName, NewEntry);
+					NewEntries.Add(NewEntry);
+				}
+
+				if (!Hierarchy.ContainsByPredicate([&Parent](const TSharedRef<FParameterPickerEntry>& InEntry){ return Parent->Binding.Parameter == InEntry->Binding.Parameter; }))
+				{
+					Hierarchy.Add(Parent.ToSharedRef());
+				}
+			}
+		}, UE::String::EParseTokensOptions::SkipEmpty);
+	}
+
+	Entries.Append(NewEntries);
+}
 
+void SParameterPicker::RefreshFilter()
+{
+	FilteredEntries.Reset();
+	FilteredHierarchy.Reset();
 	const FString FilterTextAsString = FilterText.ToString();
+
+	for (const TSharedRef<FParameterPickerEntry>& Entry : Entries)
+	{
+		Entry->FilterResult = FParameterPickerEntry::EFilterResult::DoesNotPassFilter;
+	}
+
 	for(const TSharedRef<FParameterPickerEntry>& Entry : Entries)
 	{
+		Entry->FilteredChildren.Reset();
 		if(Entry->PassesFilter(FilterTextAsString))
 		{
+			Entry->FilterResult |= FParameterPickerEntry::EFilterResult::PassesFilter;
 			FilteredEntries.Add(Entry);
+
+			if(FilterTextAsString.Len() > 0)
+			{
+				TSharedPtr< FParameterPickerEntry> ParentEntry = Entry->Parent;
+				while (ParentEntry.IsValid())
+				{
+					ParentEntry->FilterResult |= FParameterPickerEntry::EFilterResult::ChildPassesFilter;
+					ParentEntry = ParentEntry->Parent;
+				}
+			}
 		}
 	}
 
-	EntriesList->RequestListRefresh();
+	TArray<TSharedRef<FParameterPickerEntry>> Stack;
+	Stack.Reserve(16);
+	for (const TSharedRef<FParameterPickerEntry>& Entry : Hierarchy)
+	{
+		Stack.Add(Entry);
+	}
+
+	while (Stack.Num() > 0)
+	{
+		TSharedRef<FParameterPickerEntry> Top = Stack.Top();
+		Stack.Pop(false);
+
+		if (Top->FilterResult != FParameterPickerEntry::EFilterResult::DoesNotPassFilter)
+		{
+			if (!Top->Parent.IsValid())
+			{
+				FilteredHierarchy.Add(Top);
+			}
+			else
+			{
+				Top->Parent->FilteredChildren.Add(Top);
+			}
+
+			for (const TSharedRef<FParameterPickerEntry>& ChildEntry : Top->Children)
+			{
+				Stack.Add(ChildEntry);
+			}
+
+			if (EnumHasAnyFlags(Top->FilterResult, FParameterPickerEntry::EFilterResult::ChildPassesFilter))
+			{
+				EntriesList->SetItemExpansion(Top, true);
+			}
+		}
+	}
+
+	EntriesList->RequestTreeRefresh();
 }
 
 class SParameterPickerRow : public SMultiColumnTableRow<TSharedRef<FParameterPickerEntry>>
@@ -223,10 +442,10 @@ class SParameterPickerRow : public SMultiColumnTableRow<TSharedRef<FParameterPic
 	SLATE_BEGIN_ARGS(SParameterPickerRow) {}
 	SLATE_END_ARGS()
 
-	void Construct(const FArguments& InArgs, const TSharedRef<STableViewBase>& InOwnerTableView, TSharedRef<FParameterPickerEntry> InEntry)
+	void Construct(const FArguments& InArgs, const TSharedRef<STableViewBase>& InOwnerTableView, TSharedRef<FParameterPickerEntry> InEntry, TSharedRef<const SParameterPicker> InParameterPicker)
 	{
 		Entry = InEntry;
-
+		ParameterPicker = InParameterPicker;
 		SMultiColumnTableRow<TSharedRef<FParameterPickerEntry>>::Construct( SMultiColumnTableRow<TSharedRef<FParameterPickerEntry>>::FArguments(), InOwnerTableView);
 	}
 
@@ -234,11 +453,12 @@ class SParameterPickerRow : public SMultiColumnTableRow<TSharedRef<FParameterPic
 	{
 		using namespace ParameterPicker;
 
-		if(InColumnName == Column_Parameter)
+		if (InColumnName == Column_Type)
 		{
 			return
 				SNew(SHorizontalBox)
-				+SHorizontalBox::Slot()
+				.Visibility(Entry->ParamType.IsValid() ? EVisibility::Visible : EVisibility::Hidden)
+				+ SHorizontalBox::Slot()
 				.HAlign(HAlign_Center)
 				.VAlign(VAlign_Center)
 				.Padding(2.0f)
@@ -247,6 +467,20 @@ class SParameterPickerRow : public SMultiColumnTableRow<TSharedRef<FParameterPic
 					SNew(SImage)
 					.Image(Entry->PinIcon)
 					.ColorAndOpacity(Entry->PinColor)
+				];
+		}
+		else if(InColumnName == Column_Parameter)
+		{
+			return
+				SNew(SHorizontalBox)
+				+SHorizontalBox::Slot()
+				.HAlign(HAlign_Center)
+				.VAlign(VAlign_Center)
+				.Padding(2.0f)
+				.AutoWidth()
+				[
+					SNew(SExpanderArrow, SharedThis(this))
+					.IndentAmount(8)
 				]
 				+SHorizontalBox::Slot()
 				.VAlign(VAlign_Center)
@@ -255,8 +489,12 @@ class SParameterPickerRow : public SMultiColumnTableRow<TSharedRef<FParameterPic
 				[
 					SNew(STextBlock)
 					.Font(IDetailLayoutBuilder::GetDetailFont())
-					.Text(FText::FromName(Entry->Binding.Parameter))
-					.ToolTipText(FText::FromName(Entry->Binding.Parameter))
+					.Text(FText::FromString(Entry->DisplayString))
+					.ToolTipText(UncookedOnly::FUtils::GetParameterDisplayNameText(Entry->Binding.Parameter))
+					.HighlightText_Lambda([this]()
+					{
+						return ParameterPicker.Pin()->FilterText;
+					})
 				];
 		}
 		else if(InColumnName == Column_Library)
@@ -294,17 +532,34 @@ class SParameterPickerRow : public SMultiColumnTableRow<TSharedRef<FParameterPic
 		return SNullWidget::NullWidget;
 	}
 	
+	TWeakPtr<const SParameterPicker> ParameterPicker;
 	TSharedPtr<FParameterPickerEntry> Entry;
 };
 
 TSharedRef<ITableRow> SParameterPicker::HandleGenerateRow(TSharedRef<FParameterPickerEntry> InEntry, const TSharedRef<STableViewBase>& InOwnerTable) const
 {
-	return SNew(SParameterPickerRow, InOwnerTable, InEntry);
+	return SNew(SParameterPickerRow, InOwnerTable, InEntry, SharedThis(this));
+}
+
+void SParameterPicker::HandleGetChildren(TSharedRef<FParameterPickerEntry> InEntry, TArray<TSharedRef<FParameterPickerEntry>>& OutChildren) const
+{
+	OutChildren = InEntry->FilteredChildren;
 }
 
 void SParameterPicker::HandleSelectionChanged(TSharedPtr<FParameterPickerEntry> InEntry, ESelectInfo::Type InSelectInfo)
 {
 	Args.OnSelectionChanged.ExecuteIfBound();
+
+	if(!Args.bMultiSelect && EntriesList->GetNumItemsSelected() == 1 && Args.OnParameterPicked.IsBound())
+	{
+		TArray<TSharedRef<FParameterPickerEntry>> SelectedEntries;
+		EntriesList->GetSelectedItems(SelectedEntries);
+
+		if(SelectedEntries[0]->ParamType.IsValid())
+		{
+			Args.OnParameterPicked.ExecuteIfBound(SelectedEntries[0]->Binding);
+		}
+	}
 }
 
 void SParameterPicker::HandleGetParameterBindings(TArray<FParameterBindingReference>& OutParameterBindings) const
@@ -318,6 +573,12 @@ void SParameterPicker::HandleGetParameterBindings(TArray<FParameterBindingRefere
 	}
 }
 
+bool SParameterPicker::HandleIsSelectableOrNavigable(TSharedRef<FParameterPickerEntry> InEntry) const
+{
+	// Only allow selecting items that have valid parameters
+	return InEntry->ParamType.IsValid();
+}
+
 }
 
 #undef LOCTEXT_NAMESPACE
\ No newline at end of file
diff --git a/Engine/Plugins/Experimental/AnimNext/Source/AnimNextEditor/Private/Param/SParameterPicker.h b/Engine/Plugins/Experimental/AnimNext/Source/AnimNextEditor/Private/Param/SParameterPicker.h
--- a/Engine/Plugins/Experimental/AnimNext/Source/AnimNextEditor/Private/Param/SParameterPicker.h
+++ b/Engine/Plugins/Experimental/AnimNext/Source/AnimNextEditor/Private/Param/SParameterPicker.h
@@ -4,7 +4,7 @@
 
 #include "Widgets/SCompoundWidget.h"
 #include "Param/ParameterPickerArgs.h"
-#include "Widgets/Views/SListView.h"
+#include "Widgets/Views/STreeView.h"
 
 namespace UE::AnimNext::Editor
 {
@@ -25,23 +25,35 @@ public:
 private:
 	void RefreshEntries();
 
+	void BuildHierarchy();
+
 	void RefreshFilter();
 
 	TSharedRef<ITableRow> HandleGenerateRow(TSharedRef<FParameterPickerEntry> InEntry, const TSharedRef<STableViewBase>& InOwnerTable) const;
 
+	void HandleGetChildren(TSharedRef<FParameterPickerEntry> InEntry, TArray<TSharedRef<FParameterPickerEntry>>& OutChildren) const;
+
 	void HandleSelectionChanged(TSharedPtr<FParameterPickerEntry> InEntry, ESelectInfo::Type InSelectInfo);
-	
+
 	void HandleGetParameterBindings(TArray<FParameterBindingReference>& OutParameterBindings) const;
 
+	bool HandleIsSelectableOrNavigable(TSharedRef<FParameterPickerEntry> InEntry) const;
+
 private:
+	friend class SParameterPickerRow;
+
 	FParameterPickerArgs Args;
 
-	TSharedPtr<SListView<TSharedRef<FParameterPickerEntry>>> EntriesList;
+	TSharedPtr<STreeView<TSharedRef<FParameterPickerEntry>>> EntriesList;
 
 	TArray<TSharedRef<FParameterPickerEntry>> Entries;
 
 	TArray<TSharedRef<FParameterPickerEntry>> FilteredEntries;
 
+	TArray<TSharedRef<FParameterPickerEntry>> Hierarchy;
+
+	TArray<TSharedRef<FParameterPickerEntry>> FilteredHierarchy;
+
 	FText FilterText;
 };
 
diff --git a/Engine/Plugins/Experimental/AnimNext/Source/AnimNextEditor/Private/Param/SParametersActionMenu.cpp b/Engine/Plugins/Experimental/AnimNext/Source/AnimNextEditor/Private/Param/SParametersActionMenu.cpp
--- a/Engine/Plugins/Experimental/AnimNext/Source/AnimNextEditor/Private/Param/SParametersActionMenu.cpp
+++ b/Engine/Plugins/Experimental/AnimNext/Source/AnimNextEditor/Private/Param/SParametersActionMenu.cpp
@@ -14,7 +14,7 @@
 #include "RigVMCore/RigVMRegistry.h"
 #include "Units/RigUnit.h"
 #include "Widgets/SToolTip.h"
-#include "Param/ParametersExecuteContext.h"
+#include "Graph/AnimNextExecuteContext.h"
 
 #define LOCTEXT_NAMESPACE "AnimNextParametersEditor"
 
@@ -25,16 +25,11 @@ static void CollectAllAnimNextParameterActions(FGraphContextMenuBuilder& MenuBui
 {
 	static const TArray<UScriptStruct*> AllowedExecuteContexts =
 	{
-		FRigVMExecuteContext::StaticStruct(),
-		FAnimNextParametersExecuteContext::StaticStruct()
+		FRigVMExecuteContext::StaticStruct()
 	};
 
-	const TChunkedArray<FRigVMFunction>& Functions = FRigVMRegistry::Get().GetFunctions();
-	const int32 NumFunctions = Functions.Num();
-	for(int i=0; i < NumFunctions; i++)
+	for (const FRigVMFunction& Function : FRigVMRegistry::Get().GetFunctions())
 	{
-		const FRigVMFunction& Function = Functions[i];
-
 		const UScriptStruct* FunctionContext = Function.GetExecuteContextStruct();
 		if (FunctionContext == nullptr || !AllowedExecuteContexts.Contains(FunctionContext))
 		{
@@ -65,6 +60,26 @@ static void CollectAllAnimNextParameterActions(FGraphContextMenuBuilder& MenuBui
 		}
 
 		MenuBuilder.AddAction(MakeShared<FAnimNextParameterSchemaAction_RigUnit>(Struct, NodeCategory, MenuDesc, ToolTip));
+	}
+
+	for (const FRigVMDispatchFactory* Factory : FRigVMRegistry::Get().GetFactories())
+	{
+		if (!Factory->SupportsExecuteContextStruct(FRigVMExecuteContext::StaticStruct()))
+		{
+			continue;
+		}
+
+		const FRigVMTemplate* Template = Factory->GetTemplate();
+		if (Template == nullptr)
+		{
+			continue;
+		}
+
+		FText NodeCategory = FText::FromString(Factory->GetCategory());
+		FText MenuDesc = FText::FromString(Factory->GetNodeTitle(FRigVMTemplateTypeMap()));
+		FText ToolTip = Factory->GetNodeTooltip(FRigVMTemplateTypeMap());
+
+		MenuBuilder.AddAction(MakeShared<FAnimNextParameterSchemaAction_DispatchFactory>(Template->GetNotation(), NodeCategory, MenuDesc, ToolTip));
 	};
 }
 
@@ -102,7 +117,7 @@ void SParametersActionMenu::Construct(const FArguments& InArgs)
 					.VAlign(VAlign_Center)
 					[
 						SNew(STextBlock)
-						.Text(LOCTEXT("ContextText", "All AnimNext Node Classes"))
+						.Text(LOCTEXT("ContextText", "All AnimNext Nodes"))
 						.Font(FAppStyle::Get().GetFontStyle("BlueprintEditor.ActionMenu.ContextDescriptionFont"))
 						.ToolTip(IDocumentation::Get()->CreateToolTip(
 							LOCTEXT("ActionMenuContextTextTooltip", "Describes the current context of the action list"),
diff --git a/Engine/Plugins/Experimental/AnimNext/Source/AnimNextTestSuite/Private/AnimNextEditorTest.cpp b/Engine/Plugins/Experimental/AnimNext/Source/AnimNextTestSuite/Private/AnimNextEditorTest.cpp
--- a/Engine/Plugins/Experimental/AnimNext/Source/AnimNextTestSuite/Private/AnimNextEditorTest.cpp
+++ b/Engine/Plugins/Experimental/AnimNext/Source/AnimNextTestSuite/Private/AnimNextEditorTest.cpp
@@ -1,4 +1,4 @@
-﻿// Copyright Epic Games, Inc. All Rights Reserved.
+// Copyright Epic Games, Inc. All Rights Reserved.
 
 #include "CoreMinimal.h"
 #include "UncookedOnlyUtils.h"
@@ -77,7 +77,7 @@ bool FAnimationAnimNextParametersEditorTest_Block::RunTest(const FString& InPara
 
 	AddErrorIfFalse(Library->AddParameter(TEXT("TestParam"), FAnimNextParamType::GetType<bool>()) != nullptr, TEXT("Could not create new parameter in library."));
 
-	UFactory* BlockFactory = NewObject<UAnimNextParameterBlockFactory>();
+	TStrongObjectPtr<UFactory> BlockFactory(NewObject<UAnimNextParameterBlockFactory>());
 	UAnimNextParameterBlock* Block = Cast<UAnimNextParameterBlock>(BlockFactory->FactoryCreateNew(UAnimNextParameterBlock::StaticClass(), GetTransientPackage(), TEXT("TestAnimNextParameterBlock"), RF_Transient, nullptr, nullptr, NAME_None));
 	if(Block == nullptr)
 	{
diff --git a/Engine/Plugins/Experimental/AnimNext/Source/AnimNextTestSuite/Private/AnimNextParametersTest.cpp b/Engine/Plugins/Experimental/AnimNext/Source/AnimNextTestSuite/Private/AnimNextParametersTest.cpp
--- a/Engine/Plugins/Experimental/AnimNext/Source/AnimNextTestSuite/Private/AnimNextParametersTest.cpp
+++ b/Engine/Plugins/Experimental/AnimNext/Source/AnimNextTestSuite/Private/AnimNextParametersTest.cpp
@@ -6,6 +6,7 @@
 #include "Param/ParamType.h"
 #include "Param/ParamTypeHandle.h"
 #include "Animation/AnimSequence.h"
+#include "Animation/AnimMontage.h"
 #include "Async/ParallelFor.h"
 #include "Param/ParamStack.h"
 
@@ -273,7 +274,7 @@ bool FParamTypesTest::RunTest(const FString& InParameters)
 	AddErrorIfFalse(!FParamTypeHandle::GetHandle<TSoftClassPtr<UObject>>().IsBuiltInType(), TEXT("TSoftClassPtr<UObject> parameter handle is built in."));
 
 	// Check adding types across multiple threads is invariant
-	FParamTypeHandle::ResetCustomTypes();
+	FParamTypeHandle::BeginTestSandbox();
 	FParamTypeHandle Handles[100];
 	ParallelFor(100, [&Handles](int32 InIndex)
 	{
@@ -284,6 +285,7 @@ bool FParamTypesTest::RunTest(const FString& InParameters)
 	{
 		AddErrorIfFalse(Handles[Index] == Handles[0], FString::Printf(TEXT("Different type (index %d) has been defined when generating a handle across threads (0x%08x vs 0x%08x)."), Index, Handles[Index].Value, Handles[0].Value));
 	}
+	FParamTypeHandle::EndTestSandbox();
 
 	return true;
 }
@@ -300,25 +302,25 @@ bool FParamStackTest::RunTest(const FString& InParameters)
 	{
 		FParamStack Stack;
 
-		Stack.PushValues("Param0", 1.0f);
-		Stack.PushValues("Param1", 2.0f);
-		Stack.PushValues("Param1", 3.0f, "Param2", true);
-		Stack.PushValues("Param1", 4.0f, "Param2", false);
-		Stack.PushValues("Param1", 5.0f, "Param2", false, "Param3", 3);
-		Stack.PopLayer();
-		Stack.PopLayer();
-		Stack.PopLayer();
-		Stack.PopLayer();
-		Stack.PopLayer();
-
-		Stack.PushValue(ParamIds[0], 1.0f);
-		Stack.PushValue("Param1", 1.0f);
-		Stack.PushValue(ParamIds[2], true);
-		Stack.PushValue(ParamIds[3], 3);
-		Stack.PopLayer();
-		Stack.PopLayer();
-		Stack.PopLayer();
-		Stack.PopLayer();
+		FParamStack::FPushedLayerHandle Handle0 = Stack.PushValues("Param0", 1.0f);
+		FParamStack::FPushedLayerHandle Handle1 = Stack.PushValues("Param1", 2.0f);
+		FParamStack::FPushedLayerHandle Handle2 = Stack.PushValues("Param1", 3.0f, "Param2", true);
+		FParamStack::FPushedLayerHandle Handle3 = Stack.PushValues("Param1", 4.0f, "Param2", false);
+		FParamStack::FPushedLayerHandle Handle4 = Stack.PushValues("Param1", 5.0f, "Param2", false, "Param3", 3);
+		Stack.PopLayer(Handle4);
+		Stack.PopLayer(Handle3);
+		Stack.PopLayer(Handle2);
+		Stack.PopLayer(Handle1);
+		Stack.PopLayer(Handle0);
+
+		Handle0 = Stack.PushValue(ParamIds[0], 1.0f);
+		Handle1 = Stack.PushValue("Param1", 1.0f);
+		Handle2 = Stack.PushValue(ParamIds[2], true);
+		Handle3 = Stack.PushValue(ParamIds[3], 3);
+		Stack.PopLayer(Handle3);
+		Stack.PopLayer(Handle2);
+		Stack.PopLayer(Handle1);
+		Stack.PopLayer(Handle0);
 	}
 
 	// Check const/mutable
@@ -338,30 +340,57 @@ bool FParamStackTest::RunTest(const FString& InParameters)
 	// Check type variance
 	{
 		FParamStack Stack;
+		FParamStack::FPushedLayerHandle Handle;
 
-		const FLinearColor Blue = FLinearColor::Blue;
-		Stack.PushValue("Param0", Blue);
+		const float Float = 1.0f;
+		Handle = Stack.PushValue("Param0", Float);
 
 		FParamStack::EGetParamResult Result;
+		AddErrorIfFalse(Stack.GetParamPtr<double>("Param0", &Result) == nullptr, "Promoted typed access is non-null");
+		AddErrorIfFalse(Result == FParamStack::EGetParamResult::IncompatibleType, "Result != FParamStack::EGetParamResult::IncompatibleType");
+		AddErrorIfFalse(Stack.GetParamPtr<uint8>("Param0", &Result) == nullptr, "Demoted typed access is non-null");
+		AddErrorIfFalse(Result == FParamStack::EGetParamResult::IncompatibleType, "Result != FParamStack::EGetParamResult::IncompatibleType");
+
+		Stack.PopLayer(Handle);
+
+		const FLinearColor Blue = FLinearColor::Blue;
+		Handle = Stack.PushValue("Param0", Blue);
+
 		AddErrorIfFalse(Stack.GetParamPtr<FDateTime>("Param0", &Result) == nullptr, "Incorrectly typed access is non-null");
-		AddErrorIfFalse(Result == FParamStack::EGetParamResult::IncorrectType, "Result != FParamStack::EGetParamResult::IncorrectType");
+		AddErrorIfFalse(Result == FParamStack::EGetParamResult::IncompatibleType, "Result != FParamStack::EGetParamResult::IncompatibleType");
 		AddErrorIfFalse(Stack.GetMutableParamPtr<FDateTime>("Param0", &Result) == nullptr, "Incorrectly typed mutable access is non-null");
-		AddErrorIfFalse(Result == (FParamStack::EGetParamResult::IncorrectType | FParamStack::EGetParamResult::Immutable), "Result != (FParamStack::EGetParamResult::IncorrectType | FParamStack::EGetParamResult::Immutable)");
+		AddErrorIfFalse(Result == (FParamStack::EGetParamResult::IncompatibleType | FParamStack::EGetParamResult::Immutable), "Result != (FParamStack::EGetParamResult::IncompatibleType | FParamStack::EGetParamResult::Immutable)");
 		AddErrorIfFalse(Stack.GetParamPtr<FLinearColor>("Param0", &Result) != nullptr, "Correctly typed access is null");
 		AddErrorIfFalse(Result == FParamStack::EGetParamResult::Succeeded, "Result != FParamStack::EGetParamResult::Succeeded");
+
+		Stack.PopLayer(Handle);
+
+		TObjectPtr<const UAnimSequence> AnimSequence = GetDefault<UAnimSequence>();
+		Handle = Stack.PushValue("Param0", AnimSequence);
+
+		AddErrorIfFalse(Stack.GetParamPtr<UAnimSequence>("Param0", &Result) != nullptr, "Correctly typed access is null");
+		AddErrorIfFalse(Result == FParamStack::EGetParamResult::Succeeded, "Result != FParamStack::EGetParamResult::Succeeded");
+		AddErrorIfFalse(Stack.GetParamPtr<UAnimSequenceBase>("Param0", &Result) != nullptr, "Base class access is null");
+		AddErrorIfFalse(Result == FParamStack::EGetParamResult::CompatibleType, "Result != FParamStack::EGetParamResult::Succeeded");
+		AddErrorIfFalse(Stack.GetParamPtr<UAnimationAsset>("Param0", &Result) != nullptr, "Ancestor class access is null");
+		AddErrorIfFalse(Result == FParamStack::EGetParamResult::CompatibleType, "Result != FParamStack::EGetParamResult::Succeeded");
+		AddErrorIfFalse(Stack.GetParamPtr<UAnimMontage>("Param0", &Result) == nullptr, "Sibling class access is non-null");
+		AddErrorIfFalse(Result == FParamStack::EGetParamResult::IncompatibleType, "Result != FParamStack::EGetParamResult::IncompatibleType");
+
+		Stack.PopLayer(Handle);
 	}
 
 	// Check external layer interfaces
 	{
 		FParamStack Stack;
 
-		TUniquePtr<FParamStackLayer> Layer0 = FParamStack::MakeValueLayer("Param0", 1.0f);
-		TUniquePtr<FParamStackLayer> Layer1 = FParamStack::MakeValuesLayer("Param1", 2.0, "Param3", 5);
-		TUniquePtr<FParamStackLayer> Layer2 = FParamStack::MakeValueLayer(ParamIds[2], true);
+		FParamStack::FLayerHandle Layer0 = FParamStack::MakeValueLayer("Param0", 1.0f);
+		FParamStack::FLayerHandle Layer1 = FParamStack::MakeValuesLayer("Param1", 2.0, "Param3", 5);
+		FParamStack::FLayerHandle Layer2 = FParamStack::MakeValueLayer(ParamIds[2], true);
 
-		Stack.PushLayer(*Layer0);
-		Stack.PushLayer(*Layer1);
-		Stack.PushLayer(*Layer2);
+		Stack.PushLayer(Layer0);
+		Stack.PushLayer(Layer1);
+		Stack.PushLayer(Layer2);
 
 		AddErrorIfFalse(Stack.GetParam<float>("Param0") == 1.0f, "Param0 != 1.0f");
 		AddErrorIfFalse(Stack.GetParam<double>("Param1") == 2.0, "Param1 != 2.0");
@@ -373,56 +402,56 @@ bool FParamStackTest::RunTest(const FString& InParameters)
 	{
 		FParamStack Stack;
 
-		Stack.PushValues("Param1", 1.0f);
+		FParamStack::FPushedLayerHandle Handle0 = Stack.PushValues("Param1", 1.0f);
 
 		AddErrorIfFalse(Stack.GetParam<float>("Param1") == 1.0f, "Param1 != 1.0f");
 
-		Stack.PushValues("Param0", 2.0f);
+		FParamStack::FPushedLayerHandle Handle1 = Stack.PushValues("Param0", 2.0f);
 
 		AddErrorIfFalse(Stack.GetParam<float>("Param0") == 2.0f, "Param0 != 2.0f");
 		AddErrorIfFalse(Stack.GetParam<float>("Param1") == 1.0f, "Param1 != 1.0f");
 	
-		Stack.PushValues("Param3", 3.0f, "Param0", 3.0f);
+		FParamStack::FPushedLayerHandle Handle2 = Stack.PushValues("Param3", 3.0f, "Param0", 3.0f);
 
 		AddErrorIfFalse(Stack.GetParam<float>("Param0") == 3.0f, "Param0 != 3.0f");
 		AddErrorIfFalse(Stack.GetParam<float>("Param1") == 1.0f, "Param1 != 1.0f");
 		AddErrorIfFalse(Stack.GetParam<float>("Param3") == 3.0f, "Param3 != 3.0f");
 
-		Stack.PushValues("Param1", 4.0f, "Param2", 4.0f);
+		FParamStack::FPushedLayerHandle Handle3 = Stack.PushValues("Param1", 4.0f, "Param2", 4.0f);
 
 		AddErrorIfFalse(Stack.GetParam<float>("Param0") == 3.0f, "Param0 != 3.0f");
 		AddErrorIfFalse(Stack.GetParam<float>("Param1") == 4.0f, "Param1 != 4.0f");
 		AddErrorIfFalse(Stack.GetParam<float>("Param2") == 4.0f, "Param2 != 4.0f");
 		AddErrorIfFalse(Stack.GetParam<float>("Param3") == 3.0f, "Param3 != 3.0f");
 
-		Stack.PopLayer();
+		Stack.PopLayer(Handle3);
 
 		AddErrorIfFalse(Stack.GetParam<float>("Param0") == 3.0f, "Param0 != 3.0f");
 		AddErrorIfFalse(Stack.GetParam<float>("Param1") == 1.0f, "Param1 != 1.0f");
 		AddErrorIfFalse(Stack.GetParam<float>("Param3") == 3.0f, "Param3 != 3.0f");
 
-		Stack.PushValues("Param3", 4.0f, "Param0", 4.0f);
+		Handle3 = Stack.PushValues("Param3", 4.0f, "Param0", 4.0f);
 
 		AddErrorIfFalse(Stack.GetParam<float>("Param0") == 4.0f, "Param0 != 4.0f");
 		AddErrorIfFalse(Stack.GetParam<float>("Param1") == 1.0f, "Param1 != 1.0f");
 		AddErrorIfFalse(Stack.GetParam<float>("Param3") == 4.0f, "Param3 != 4.0f");
 
-		Stack.PopLayer();
-		Stack.PopLayer();
-		Stack.PopLayer();
+		Stack.PopLayer(Handle3);
+		Stack.PopLayer(Handle2);
+		Stack.PopLayer(Handle1);
 
 		AddErrorIfFalse(Stack.GetParam<float>("Param1") == 1.0f, "Param1 != 1.0f");
 
 		// Add some more parameters to check internal buffer resize logic
 		FParamId NewParamIds[] = { FParamId("Param4"), FParamId("Param5"), FParamId("Param6"), FParamId("Param7") };
 
-		Stack.PushValues("Param4", 1.0f, "Param5", 1.0f, "Param7", 1.0f);
+		Handle1 = Stack.PushValues("Param4", 1.0f, "Param5", 1.0f, "Param7", 1.0f);
 
 		AddErrorIfFalse(Stack.GetParam<float>("Param4") == 1.0f, "Param4 != 1.0f");
 		AddErrorIfFalse(Stack.GetParam<float>("Param5") == 1.0f, "Param5 != 1.0f");
 		AddErrorIfFalse(Stack.GetParam<float>("Param7") == 1.0f, "Param7 != 1.0f");
 
-		Stack.PushValues("Param6", 2.0f);
+		Handle2 = Stack.PushValues("Param6", 2.0f);
 
 		AddErrorIfFalse(Stack.GetParam<float>("Param4") == 1.0f, "Param4 != 1.0f");
 		AddErrorIfFalse(Stack.GetParam<float>("Param5") == 1.0f, "Param5 != 1.0f");
@@ -490,11 +519,11 @@ bool FParamStackTest::RunTest(const FString& InParameters)
 		PropertyBag.SetValueBool("Param1", true);
 		PropertyBag.SetValueInt32("Param2", 5);
 
-		TUniquePtr<FParamStackLayer> Layer = FParamStack::MakeLayer(PropertyBag);
+		FParamStack::FLayerHandle Layer = FParamStack::MakeLayer(PropertyBag);
 
 		FParamStack Stack;
 
-		Stack.PushLayer(*Layer);
+		Stack.PushLayer(Layer);
 
 		AddErrorIfFalse(Stack.GetParam<float>("Param0") == 1.5f, "Param0 != 1.5f");
 		AddErrorIfFalse(Stack.GetParam<bool>("Param1") == true, "Param1 != true");
@@ -505,10 +534,12 @@ bool FParamStackTest::RunTest(const FString& InParameters)
 	{
 		FParamStack Stack;
 
+		TArray<FParamStack::FPushedLayerHandle> Handles;
+		Handles.SetNum(MAX_uint16);
 		const uint32 NumParams = MAX_uint16;
 		for (uint32 Index = 0; Index < NumParams; ++Index)
 		{
-			Stack.PushValue(FName("Param", Index), (float)Index);
+			Handles[Index] = Stack.PushValue(FName("Param", Index), (float)Index);
 		}
 
 		for (uint32 Index = 0; Index < NumParams; ++Index)
@@ -520,9 +551,9 @@ bool FParamStackTest::RunTest(const FString& InParameters)
 		Stack.PushValue(FName("Param", MAX_uint16), (float)MAX_uint16);
 		AddErrorIfFalse(Stack.GetParamPtr<float>(FName("Param", MAX_uint16)) == nullptr, TEXT("Overflowed stack value is non-null"));
 
-		for (uint32 Index = 0; Index < NumParams; ++Index)
+		for (int32 Index = NumParams - 1; Index >= 0; --Index)
 		{
-			Stack.PopLayer();
+			Stack.PopLayer(Handles[Index]);
 		}
 	}
 
diff --git a/Engine/Plugins/Experimental/AnimNext/Source/AnimNextUncookedOnly/AnimNextUncookedOnly.Build.cs b/Engine/Plugins/Experimental/AnimNext/Source/AnimNextUncookedOnly/AnimNextUncookedOnly.Build.cs
--- a/Engine/Plugins/Experimental/AnimNext/Source/AnimNextUncookedOnly/AnimNextUncookedOnly.Build.cs
+++ b/Engine/Plugins/Experimental/AnimNext/Source/AnimNextUncookedOnly/AnimNextUncookedOnly.Build.cs
@@ -41,6 +41,20 @@ namespace UnrealBuildTool.Rules
 						"UnrealEd",
 					}
 				);
+
+				PrivateIncludePathModuleNames.AddRange(
+					new string[]
+					{
+						"AnimNextEditor",
+					}
+				);
+
+				DynamicallyLoadedModuleNames.AddRange(
+					new string[]
+					{
+						"AnimNextEditor",
+					}
+				);
 			}
 		}
 	}
diff --git a/Engine/Plugins/Experimental/AnimNext/Source/AnimNextUncookedOnly/Internal/Graph/AnimGraph/AnimGraphNodeBinding_AnimNextParameters.h b/Engine/Plugins/Experimental/AnimNext/Source/AnimNextUncookedOnly/Internal/Graph/AnimGraph/AnimGraphNodeBinding_AnimNextParameters.h
new file mode 100644
--- /dev/null
+++ b/Engine/Plugins/Experimental/AnimNext/Source/AnimNextUncookedOnly/Internal/Graph/AnimGraph/AnimGraphNodeBinding_AnimNextParameters.h
@@ -0,0 +1,49 @@
+// Copyright Epic Games, Inc. All Rights Reserved.
+
+#pragma once
+
+#include "AnimGraphNodeBinding.h"
+#include "AnimGraphNodeBinding_AnimNextParameters.generated.h"
+
+USTRUCT()
+struct FAnimNextAnimGraphNodeParameterBinding
+{
+	GENERATED_BODY()
+
+	UPROPERTY()
+	FName BindingName;
+
+	UPROPERTY()
+	int32 ArrayIndex;
+
+	UPROPERTY()
+	FName ParameterName;
+
+	UPROPERTY()
+	FText CachedParameterNameText;
+};
+
+UCLASS(EditInlineNew, DisplayName = "Anim Next Parameter Binding")
+class UAnimGraphNodeBinding_AnimNextParameters : public UAnimGraphNodeBinding
+{
+	GENERATED_BODY()
+
+private:
+	friend class UAnimBlueprintExtension_AnimNextParameters;
+
+	// UAnimGraphNodeBinding interface
+	virtual UScriptStruct* GetAnimNodeHandlerStruct() const override;
+	virtual void OnExpandNode(IAnimBlueprintCompilationContext& InCompilationContext, UAnimGraphNode_Base* InNode, UEdGraph* InSourceGraph) override;
+	virtual bool HasBinding(FName InBindingName, bool bCheckArrayIndexName) const override;
+	virtual void RemoveBindings(FName InBindingName) override;
+	virtual void AddPinSearchMetaDataInfo(const UEdGraphPin* InPin, FName InBindingName, TArray<struct FSearchTagDataPair>& OutTaggedMetaData) const override;
+	virtual void UpdateBindingNames(TFunctionRef<FString(const FString& InOldName)> InModifierFunction) override;
+	virtual void GetRequiredExtensions(TArray<TSubclassOf<UAnimBlueprintExtension>>& OutExtensions) const override;
+	virtual void ProcessDuringCompilation(IAnimBlueprintCompilationContext& InCompilationContext, IAnimBlueprintGeneratedClassCompiledData& OutCompiledData) override;
+#if WITH_EDITOR
+	virtual TSharedRef<SWidget> MakePropertyBindingWidget(const UAnimGraphNode_Base::FAnimPropertyBindingWidgetArgs& InArgs) override;
+#endif
+
+	UPROPERTY()
+	TMap<FName, FAnimNextAnimGraphNodeParameterBinding> PropertyBindings;
+};
\ No newline at end of file
diff --git a/Engine/Plugins/Experimental/AnimNext/Source/AnimNextUncookedOnly/Internal/Graph/AnimGraph/AnimGraphNode_AnimNextParameters.h b/Engine/Plugins/Experimental/AnimNext/Source/AnimNextUncookedOnly/Internal/Graph/AnimGraph/AnimGraphNode_AnimNextParameters.h
new file mode 100644
--- /dev/null
+++ b/Engine/Plugins/Experimental/AnimNext/Source/AnimNextUncookedOnly/Internal/Graph/AnimGraph/AnimGraphNode_AnimNextParameters.h
@@ -0,0 +1,22 @@
+// Copyright Epic Games, Inc. All Rights Reserved.
+
+#pragma once
+
+#include "CoreMinimal.h"
+#include "Graph/AnimGraph/AnimNode_AnimNextParameters.h"
+#include "AnimGraphNode_Base.h"
+#include "AnimGraphNode_AnimNextParameters.generated.h"
+
+UCLASS(MinimalAPI)
+class UAnimGraphNode_AnimNextParameters : public UAnimGraphNode_Base
+{
+	GENERATED_BODY()
+
+	UPROPERTY(EditAnywhere, Category = Settings)
+	FAnimNode_AnimNextParameters Node;
+
+	// UAnimGraphNode_Base interface
+	virtual FText GetNodeTitle(ENodeTitleType::Type TitleType) const override;
+	virtual FText GetTooltipText() const override;
+	virtual FText GetMenuCategory() const override;
+};
\ No newline at end of file
diff --git a/Engine/Plugins/Experimental/AnimNext/Source/AnimNextUncookedOnly/Internal/Param/AnimNextParameterBlock_EditorData.h b/Engine/Plugins/Experimental/AnimNext/Source/AnimNextUncookedOnly/Internal/Param/AnimNextParameterBlock_EditorData.h
--- a/Engine/Plugins/Experimental/AnimNext/Source/AnimNextUncookedOnly/Internal/Param/AnimNextParameterBlock_EditorData.h
+++ b/Engine/Plugins/Experimental/AnimNext/Source/AnimNextUncookedOnly/Internal/Param/AnimNextParameterBlock_EditorData.h
@@ -94,6 +94,7 @@ class UAnimNextParameterBlock_EditorData : public UObject, public IRigVMClientHo
 	friend struct UE::AnimNext::UncookedOnly::FUtilsPrivate;
 	friend class UE::AnimNext::Editor::FParameterBlockEditor;
 	friend struct FAnimNextParameterSchemaAction_RigUnit;
+	friend struct FAnimNextParameterSchemaAction_DispatchFactory;
 	friend class UE::AnimNext::Editor::SParameterBlockView;
 	friend class UE::AnimNext::Editor::SParameterPicker;
 	friend class UE::AnimNext::Editor::FParameterBlockTabSummoner;
diff --git a/Engine/Plugins/Experimental/AnimNext/Source/AnimNextUncookedOnly/Internal/Param/IAnimNextParameterBlockGraphInterface.h b/Engine/Plugins/Experimental/AnimNext/Source/AnimNextUncookedOnly/Internal/Param/IAnimNextParameterBlockGraphInterface.h
--- a/Engine/Plugins/Experimental/AnimNext/Source/AnimNextUncookedOnly/Internal/Param/IAnimNextParameterBlockGraphInterface.h
+++ b/Engine/Plugins/Experimental/AnimNext/Source/AnimNextUncookedOnly/Internal/Param/IAnimNextParameterBlockGraphInterface.h
@@ -1,4 +1,4 @@
-﻿// Copyright Epic Games, Inc. All Rights Reserved.
+// Copyright Epic Games, Inc. All Rights Reserved.
 
 #pragma once
 
@@ -13,6 +13,7 @@ namespace UE::AnimNext::Editor
 {
 	class SParameterBlockView;
 	class SParameterBlockViewRow;
+	class FParameterBlockTabSummoner;
 }
 
 namespace UE::AnimNext::UncookedOnly
@@ -34,6 +35,7 @@ class ANIMNEXTUNCOOKEDONLY_API IAnimNextParameterBlockGraphInterface
 	friend class UE::AnimNext::Editor::SParameterBlockView;
 	friend class UE::AnimNext::Editor::SParameterBlockViewRow;
 	friend struct UE::AnimNext::UncookedOnly::FUtilsPrivate;
+	friend class UE::AnimNext::Editor::FParameterBlockTabSummoner;
 
 	// Get the graph
 	virtual URigVMGraph* GetGraph() const = 0;
diff --git a/Engine/Plugins/Experimental/AnimNext/Source/AnimNextUncookedOnly/Internal/UncookedOnlyUtils.h b/Engine/Plugins/Experimental/AnimNext/Source/AnimNextUncookedOnly/Internal/UncookedOnlyUtils.h
--- a/Engine/Plugins/Experimental/AnimNext/Source/AnimNextUncookedOnly/Internal/UncookedOnlyUtils.h
+++ b/Engine/Plugins/Experimental/AnimNext/Source/AnimNextUncookedOnly/Internal/UncookedOnlyUtils.h
@@ -68,7 +68,10 @@ struct ANIMNEXTUNCOOKEDONLY_API FUtils
 	static FRigVMTemplateArgumentType GetRigVMArgTypeFromParamType(const FAnimNextParamType& InParamType);
 
 	/** Set up a binding graph given the type to set */
-	static void SetupBindingGraphForLiteral(URigVMController* InController, const FAnimNextParamType& InParamType);
+	static void SetupBindingGraphForLiteral(URigVMController* InController, FName InParameterName, const FAnimNextParamType& InParamType);
+
+	/** Converts the Verse-tag-like snake_case_parameter_name to a period-separated display name similar to a gameplay tag */
+	static FText GetParameterDisplayNameText(FName InParameterName);
 };
 
 }
\ No newline at end of file
