Implement animnext graph decorator authoring
Moved decorator offset computation to runtime to support #ifdef and adding/removing members
Added a node template builder to ensure node templates are always fully formed when used
Add support for RigVM decorator dynamic pins
Other minor imrovements to decorator runtime

diff --git a/Engine/Plugins/Experimental/AnimNext/Source/AnimNext/AnimNext.Build.cs b/Engine/Plugins/Experimental/AnimNext/Source/AnimNext/AnimNext.Build.cs
--- a/Engine/Plugins/Experimental/AnimNext/Source/AnimNext/AnimNext.Build.cs
+++ b/Engine/Plugins/Experimental/AnimNext/Source/AnimNext/AnimNext.Build.cs
@@ -5,7 +5,7 @@ namespace UnrealBuildTool.Rules
 	public class AnimNext : ModuleRules
 	{
 		public AnimNext(ReadOnlyTargetRules Target) : base(Target)
-		{	
+		{
 			PublicDependencyModuleNames.AddRange(
 				new string[]
 				{
@@ -17,12 +17,23 @@ namespace UnrealBuildTool.Rules
 					"Engine"
 				}
 			);
+
 			PrivateDependencyModuleNames.AddRange(
 				new string[]
 				{
 					"AssetRegistry"
 				}
 			);
+
+			if (Target.bBuildEditor == true)
+			{
+				PrivateDependencyModuleNames.AddRange(
+					new string[]
+					{
+						"RigVMDeveloper"
+					}
+				);
+			}
 		}
 	}
 }
\ No newline at end of file
diff --git a/Engine/Plugins/Experimental/AnimNext/Source/AnimNext/Internal/DecoratorBase/DecoratorReader.h b/Engine/Plugins/Experimental/AnimNext/Source/AnimNext/Internal/DecoratorBase/DecoratorReader.h
--- a/Engine/Plugins/Experimental/AnimNext/Source/AnimNext/Internal/DecoratorBase/DecoratorReader.h
+++ b/Engine/Plugins/Experimental/AnimNext/Source/AnimNext/Internal/DecoratorBase/DecoratorReader.h
@@ -3,6 +3,9 @@
 #pragma once
 
 #include "CoreMinimal.h"
+#include "DecoratorBase/DecoratorHandle.h"
+#include "DecoratorBase/EntryPointHandle.h"
+#include "DecoratorBase/NodeHandle.h"
 #include "Serialization/ArchiveProxy.h"
 
 namespace UE::AnimNext
@@ -15,11 +18,53 @@ namespace UE::AnimNext
 	  *     - A list of FNodeTemplates that the nodes use
 	  *     - The graph shared data (FNodeDescription for every node)
 	  */
-	class ANIMNEXT_API FDecoratorReader : public FArchiveProxy
+	class ANIMNEXT_API FDecoratorReader final : public FArchiveProxy
 	{
 	public:
-		FDecoratorReader(FArchive& Ar);
+		// The largest size allowed for the shared data of a single graph
+		// Node/decorator handles use an unsigned 24 bits value to represent graph offsets
+		// which limits us to 16 MB
+		// We could extend this range by leveraging the fact that nodes have a minimum required alignment
+		static constexpr uint32 MAXIMUM_GRAPH_SHARED_DATA_SIZE = (1 << 24) - 1;
 
-		void ReadGraphSharedData(TArray<uint8>& GraphSharedData);
+		enum class EErrorState
+		{
+			None,						// All good, no error
+			GraphTooLarge,				// Exceeded the maximum graph size, @see FDecoratorReader::MAXIMUM_GRAPH_SHARED_DATA_SIZE
+			NodeSharedDataTooLarge,		// Exceeded the maximum node shared data size, @see FNodeDescription::MAXIMUM_NODE_SHARED_DATA_SIZE
+			NodeInstanceDataTooLarge,	// Exceeded the maximum node instance data size, @see FNodeInstance::MAXIMUM_NODE_INSTANCE_DATA_SIZE
+		};
+
+		explicit FDecoratorReader(FArchive& Ar);
+
+		[[nodiscard]] EErrorState ReadGraph(TArray<uint8>& GraphSharedData, TArray<TObjectPtr<UObject>>& TrackedObjectsForGC);
+
+		// Takes a node handle representing a node index and resolves it into a node handle representing a shared data offset
+		// Must be called after ReadGraphSharedData as it populates the necessary data
+		[[nodiscard]] FNodeHandle ResolveNodeHandle(FNodeHandle NodeHandle) const;
+
+		// Takes a decorator handle representing a node index and resolves it into a node handle representing a shared data offset
+		// Must be called after ReadGraphSharedData as it populates the necessary data
+		[[nodiscard]] FAnimNextDecoratorHandle ResolveDecoratorHandle(FAnimNextDecoratorHandle DecoratorHandle) const;
+
+		// Takes an entry point handle representing a node index and resolves it into a node handle representing a shared data offset
+		// Must be called after ReadGraphSharedData as it populates the necessary data
+		[[nodiscard]] FAnimNextDecoratorHandle ResolveEntryPointHandle(FAnimNextEntryPointHandle EntryPointHandle) const;
+
+		// FArchive implementation
+		virtual FArchive& operator<<(UObject*& Obj) override;
+		virtual FArchive& operator<<(FObjectPtr& Obj) override;
+
+	private:
+		// Call first to read the graph shared data
+		[[nodiscard]] EErrorState ReadGraphSharedData(TArray<uint8>& GraphSharedData);
+
+		// Call second to read the tracked UObjects for GC
+		[[nodiscard]] EErrorState ReadTrackedObjectsForGC(TArray<TObjectPtr<UObject>>& TrackedObjectsForGC);
+
+		// A list of node handles for each node within the archive
+		TArray<FNodeHandle> NodeHandles;
+
+		friend FAnimNextDecoratorHandle;
 	};
 }
diff --git a/Engine/Plugins/Experimental/AnimNext/Source/AnimNext/Internal/DecoratorBase/DecoratorRegistry.h b/Engine/Plugins/Experimental/AnimNext/Source/AnimNext/Internal/DecoratorBase/DecoratorRegistry.h
--- a/Engine/Plugins/Experimental/AnimNext/Source/AnimNext/Internal/DecoratorBase/DecoratorRegistry.h
+++ b/Engine/Plugins/Experimental/AnimNext/Source/AnimNext/Internal/DecoratorBase/DecoratorRegistry.h
@@ -21,33 +21,40 @@ namespace UE::AnimNext
 		// Access the global registry
 		static FDecoratorRegistry& Get();
 
-		// Finds and returns the decorator handle for the provided decorator UID or an invalid
-		// handle if that decorator hasn't been registered yet.
-		FDecoratorRegistryHandle FindHandle(uint32 DecoratorUID) const;
-
 		// Finds and returns the decorator handle for the provided decorator UID or an invalid
 		// handle if that decorator hasn't been registered yet.
 		FDecoratorRegistryHandle FindHandle(FDecoratorUID DecoratorUID) const;
 
-		// Finds and returns the decorator associated with the provided handle. If the handle
-		// is invalid, nullptr is returned.
+		// Finds and returns the decorator associated with the provided handle.
+		// If the handle is not valid, nullptr is returned.
 		const FDecorator* Find(FDecoratorRegistryHandle DecoratorHandle) const;
 
-		// Finds and returns the decorator associated with the provided decorator UID. If the decorator
-		// is registered, nullptr is returned.
+		// Finds and returns the decorator associated with the provided decorator UID.
+		// If the decorator is not registered, nullptr is returned.
 		const FDecorator* Find(FDecoratorUID DecoratorUID) const;
 
+		// Finds and returns the decorator associated with the provided decorator shared data UScriptStruct.
+		// If the matching decorator is not registered, nullptr is returned.
+		const FDecorator* Find(const UScriptStruct* DecoratorSharedDataStruct) const;
+
 		// Registers a decorator dynamically
 		void Register(FDecorator* Decorator);
 
 		// Unregisters a decorator dynamically
 		void Unregister(FDecorator* Decorator);
 
+		// Returns a list of all registered decorators
+		TArray<const FDecorator*> GetDecorators() const;
+
 		// Returns the number of registered decorators
 		uint32 GetNum() const;
 
 	private:
 		FDecoratorRegistry() = default;
+		FDecoratorRegistry(const FDecoratorRegistry&) = delete;
+		FDecoratorRegistry(FDecoratorRegistry&&) = default;
+		FDecoratorRegistry& operator=(const FDecoratorRegistry&) = delete;
+		FDecoratorRegistry& operator=(FDecoratorRegistry&&) = default;
 
 		// Static init lifetime functions
 		static void StaticRegister(DecoratorConstructorFunc DecoratorConstructor);
diff --git a/Engine/Plugins/Experimental/AnimNext/Source/AnimNext/Internal/DecoratorBase/DecoratorTemplate.h b/Engine/Plugins/Experimental/AnimNext/Source/AnimNext/Internal/DecoratorBase/DecoratorTemplate.h
--- a/Engine/Plugins/Experimental/AnimNext/Source/AnimNext/Internal/DecoratorBase/DecoratorTemplate.h
+++ b/Engine/Plugins/Experimental/AnimNext/Source/AnimNext/Internal/DecoratorBase/DecoratorTemplate.h
@@ -26,17 +26,8 @@ namespace UE::AnimNext
 	 */
 	struct FDecoratorTemplate
 	{
-		FDecoratorTemplate(FDecoratorUID UID_, FDecoratorRegistryHandle RegistryHandle_, EDecoratorMode Mode_, uint32 AdditiveIndexOrNumAdditive_, uint32 NodeSharedOffset_, uint32 NodeInstanceOffset_)
-			: UID(UID_.GetUID())
-			, RegistryHandle(RegistryHandle_)
-			, Mode(static_cast<uint8>(Mode_))
-			, AdditiveIndexOrNumAdditive(AdditiveIndexOrNumAdditive_)
-			, NodeSharedOffset(NodeSharedOffset_)
-			, NodeInstanceOffset(NodeInstanceOffset_)
-		{}
-
 		// Returns the globally unique identifier for the decorator
-		uint32 GetUID() const { return UID; }
+		FDecoratorUID GetUID() const { return FDecoratorUID(UID); }
 
 		// Returns the decorator registry handle
 		FDecoratorRegistryHandle GetRegistryHandle() const { return RegistryHandle; }
@@ -92,7 +83,19 @@ namespace UE::AnimNext
 		ANIMNEXT_API void Serialize(FArchive& Ar);
 
 	private:
-		uint32	UID;							// decorator globally unique identifier
+		friend struct FNodeTemplate;
+		friend struct FNodeTemplateBuilder;
+
+		FDecoratorTemplate(FDecoratorUID InUID, FDecoratorRegistryHandle InRegistryHandle, EDecoratorMode InMode, uint32 InAdditiveIndexOrNumAdditive)
+			: UID(InUID.GetUID())
+			, RegistryHandle(InRegistryHandle)
+			, Mode(static_cast<uint8>(InMode))
+			, AdditiveIndexOrNumAdditive(InAdditiveIndexOrNumAdditive)
+			, NodeSharedOffset(0)
+			, NodeInstanceOffset(0)				// For instance data, 0 is an invalid offset since the data follows an instance of FNodeInstance
+		{}
+
+		FDecoratorUIDRaw			UID;			// decorator globally unique identifier
 
 		FDecoratorRegistryHandle	RegistryHandle;	// decorator registry handle
 		uint8						Mode;			// the decorator mode (we only need 1 bit)
@@ -103,8 +106,8 @@ namespace UE::AnimNext
 		uint8	AdditiveIndexOrNumAdditive;
 
 		// Offsets into the shared read-only and instance data portions of a node
-		uint16	NodeSharedOffset;				// relative to root of node description data (max 64 KB per node)
-		uint16	NodeInstanceOffset;				// relative to root of node instance data (max 64 KB per node)
+		uint16	NodeSharedOffset;				// relative to root of node description data (max 64 KB per node) (not serialized, @see FNodeTemplate::Finalize)
+		uint16	NodeInstanceOffset;				// relative to root of node instance data (max 64 KB per node) (not serialized, @see FNodeTemplate::Finalize)
 
 		// TODO: We could cache which parent/base decorator handles which interface
 		// This would avoid the need to iterate on every decorator to look up a 'Super'. Perhaps only common interfaces could be cached.
diff --git a/Engine/Plugins/Experimental/AnimNext/Source/AnimNext/Internal/DecoratorBase/DecoratorWriter.h b/Engine/Plugins/Experimental/AnimNext/Source/AnimNext/Internal/DecoratorBase/DecoratorWriter.h
--- a/Engine/Plugins/Experimental/AnimNext/Source/AnimNext/Internal/DecoratorBase/DecoratorWriter.h
+++ b/Engine/Plugins/Experimental/AnimNext/Source/AnimNext/Internal/DecoratorBase/DecoratorWriter.h
@@ -22,34 +22,22 @@ namespace UE::AnimNext
 	  *     - A list of FNodeTemplates that the nodes use
 	  *     - The graph shared data (FNodeDescription for every node)
 	  */
-	class ANIMNEXT_API FDecoratorWriter : public FMemoryWriter
+	class ANIMNEXT_API FDecoratorWriter final : public FMemoryWriter
 	{
 	public:
-		// The largest size allowed for the shared data of a single graph
-		// We use 16 bit offsets to index into the shared data
-		static constexpr uint32 MAXIMUM_GRAPH_SHARED_DATA_SIZE = 64 * 1024;
-
 		enum class EErrorState
 		{
 			None,					// All good, no error
 			TooManyNodes,			// Exceeded the maximum number of nodes in a graph, @see FNodeDescription::MAXIMUM_COUNT
-			GraphTooLarge,			// Exceeded the maximum graph size, @see FDecoratorWriter::MAXIMUM_GRAPH_SHARED_DATA_SIZE
 			NodeTemplateNotFound,	// Failed to find a necessary node template
 			NodeTemplateTooLarge,	// Exceeded the maximum node template size, @see FNodeTemplate::MAXIMUM_SIZE
-			NodeSharedDataTooLarge,	// Exceeded the maximum node shared data size, @see FNodeDescription::MAXIMUM_SIZE
 			NodeHandleNotFound,		// Failed to find the mapping for a node handle, it was likely not registered
 		};
 
 		FDecoratorWriter();
 
 		// Registers an instance of the provided node template and assigns a node handle and node UID to it
-		FNodeHandle RegisterNode(const FNodeTemplate& NodeTemplate);
-
-		// Returns the node UID associated with a node handle
-		uint16 GetNodeUID(const FNodeHandle NodeHandle) const;
-
-		// Returns the node handle associated with a node UID
-		FNodeHandle GetNodeHandle(uint16 NodeUID) const;
+		[[nodiscard]] FNodeHandle RegisterNode(const FNodeTemplate& NodeTemplate);
 
 		// Called before node writing can begin
 		void BeginNodeWriting();
@@ -59,31 +47,39 @@ namespace UE::AnimNext
 
 		// Writes out the provided node using the decorator properties
 		// Nodes must be written in the same order they were registered in
-		void WriteNode(const FNodeHandle NodeHandle, const TFunction<const TMap<FString, FString>& (uint32 DecoratorIndex)>& GetDecoratorProperties);
+		void WriteNode(const FNodeHandle NodeHandle, const TFunction<FString(uint32 DecoratorIndex, const FString& PropertyName)>& GetDecoratorProperty);
 
 		// Returns the error state
-		EErrorState GetErrorState() const;
+		[[nodiscard]] EErrorState GetErrorState() const;
 
 		// Returns the populated raw graph shared data buffer
-		const TArray<uint8>& GetGraphSharedData() const;
+		[[nodiscard]] const TArray<uint8>& GetGraphSharedData() const;
+
+		// FArchive implementation
+		virtual FArchive& operator<<(UObject*& Obj) override;
+		virtual FArchive& operator<<(FObjectPtr& Obj) override;
 
 	private:
 		struct FNodeMapping
 		{
+			// The node handle for this entry (encoded as a node ID)
 			FNodeHandle NodeHandle;
-			uint32 NodeSize;
-			uint16 NodeUID;
+
+			// The node template handle the node uses
 			FNodeTemplateRegistryHandle NodeTemplateHandle;
+
+			// The unique node template index that we'll serialize
+			uint32 NodeTemplateIndex;
 		};
 
 		TArray<uint8> GraphSharedDataArchiveBuffer;
 		TArray<FNodeMapping> NodeMappings;
 
 		// To track the node registration process
-		uint32 GraphSharedDataSize;
-		uint16 NextNodeUID;
+		FNodeID NextNodeID;
 
 		// To track node writing
+		TArray<UObject*> TrackedObjectsForGC;
 		uint32 NumNodesWritten;
 		bool bIsNodeWriting;
 
diff --git a/Engine/Plugins/Experimental/AnimNext/Source/AnimNext/Internal/DecoratorBase/NodeTemplate.h b/Engine/Plugins/Experimental/AnimNext/Source/AnimNext/Internal/DecoratorBase/NodeTemplate.h
--- a/Engine/Plugins/Experimental/AnimNext/Source/AnimNext/Internal/DecoratorBase/NodeTemplate.h
+++ b/Engine/Plugins/Experimental/AnimNext/Source/AnimNext/Internal/DecoratorBase/NodeTemplate.h
@@ -21,7 +21,9 @@ namespace UE::AnimNext
 	 * Templates can exist in multiple sub-graphs on disk. On load, a single unique copy is retained in
 	 * memory and shared between ALL graphs.
 	 * 
-	 * @see FNodeDescription
+	 * Node templates are created through the FNodeTemplateBuilder.
+	 * 
+	 * @see FNodeDescription, FNodeTemplateBuilder
 	 */
 	struct FNodeTemplate
 	{
@@ -30,24 +32,24 @@ namespace UE::AnimNext
 		// We artificially specify a conservative upper bound
 		static constexpr uint32 MAXIMUM_SIZE = 64 * 1024;
 
-		FNodeTemplate(uint32 UID_, uint32 InstanceSize_, uint32 NumDecorators_)
-			: UID(UID_)
-			, InstanceSize(InstanceSize_)
-			, NumDecorators(NumDecorators_)
-		{}
-
 		// Returns the globally unique template identifier (a hash of all the decorator UIDs present in this node)
 		uint32 GetUID() const { return UID; }
 
 		// Returns the size in bytes of this node template
-		uint32 GetNodeTemplateSize() const { return sizeof(FNodeTemplate) + NumDecorators * sizeof(FDecoratorTemplate); }
+		uint32 GetNodeTemplateSize() const { return sizeof(FNodeTemplate) + (NumDecorators * sizeof(FDecoratorTemplate)); }
 
-		// Returns the size in bytes of a node instance
-		uint32 GetInstanceSize() const { return InstanceSize; }
+		// Returns the size in bytes of a node shared data
+		uint32 GetNodeSharedDataSize() const { return NodeSharedDataSize; }
+
+		// Returns the size in bytes of a node instance (how much space to allocate to hold a node instance of this template)
+		uint32 GetNodeInstanceDataSize() const { return NodeInstanceDataSize; }
 
 		// Returns the number of decorators present in the node template
 		uint32 GetNumDecorators() const { return NumDecorators; }
 
+		// Returns whether the node template is valid or not
+		bool IsValid() const { return NodeSharedDataSize != 0 && NodeInstanceDataSize != 0; }
+
 		// Returns a pointer to the list of decorator template descriptions
 		FDecoratorTemplate* GetDecorators() { return reinterpret_cast<FDecoratorTemplate*>(reinterpret_cast<uint8*>(this) + sizeof(FNodeTemplate)); }
 
@@ -58,17 +60,31 @@ namespace UE::AnimNext
 		ANIMNEXT_API void Serialize(FArchive& Ar);
 
 	private:
-		uint32	UID;				// globally unique template identifier or hash
+		friend struct FNodeTemplateBuilder;
+
+		FNodeTemplate(uint32 InUID, uint32 InNumDecorators)
+			: UID(InUID)
+			, NodeSharedDataSize(0)
+			, NodeInstanceDataSize(0)
+			, NumDecorators(static_cast<uint8>(InNumDecorators))
+		{}
+
+		// This function performs the necessary fix-ups once all decorator templates have been populated
+		// Call this function whenever a new node template has been constructed before using it
+		ANIMNEXT_API void Finalize();
+
+		uint32	UID;					// globally unique template identifier or hash
 
-		uint16	InstanceSize;		// size in bytes of a node instance
+		uint16	NodeSharedDataSize;		// size in bytes of a node shared data (not serialized, @see FNodeTemplate::Finalize)
+		uint16	NodeInstanceDataSize;	// size in bytes of a node instance (not serialized, @see FNodeTemplate::Finalize)
 
 		uint8	NumDecorators;
-		uint8	Padding[1];
+		uint8	Padding[3];
 
 		// TODO: We could use the padding (and extra space) here to cache which decorator handles which interface
 		// This would avoid the need to iterate on every decorator to look it up. Perhaps only common interfaces could be cached.
 
-		// Followed by a list of [FDecoratorTemplate] instances
+		// Followed in memory by a list of [FDecoratorTemplate] instances
 	};
 
 	static_assert(std::is_trivially_copyable_v<FNodeTemplate>, "FNodeTemplate needs to be trivially copyable");
diff --git a/Engine/Plugins/Experimental/AnimNext/Source/AnimNext/Internal/DecoratorBase/NodeTemplateBuilder.h b/Engine/Plugins/Experimental/AnimNext/Source/AnimNext/Internal/DecoratorBase/NodeTemplateBuilder.h
new file mode 100644
--- /dev/null
+++ b/Engine/Plugins/Experimental/AnimNext/Source/AnimNext/Internal/DecoratorBase/NodeTemplateBuilder.h
@@ -0,0 +1,43 @@
+// Copyright Epic Games, Inc. All Rights Reserved.
+
+#pragma once
+
+#include "CoreMinimal.h"
+#include "DecoratorBase/DecoratorUID.h"
+
+namespace UE::AnimNext
+{
+	struct FNodeTemplate;
+
+	/**
+	  * FNodeTemplateBuilder
+	  * 
+	  * Utility to help construct node templates.
+	  */
+	struct ANIMNEXT_API FNodeTemplateBuilder
+	{
+		FNodeTemplateBuilder() = default;
+
+		// Adds the specified decorator type to the node template
+		void AddDecorator(FDecoratorUID DecoratorUID);
+
+		// Returns a node template for the provided list of decorators
+		// The node template will be built into the provided buffer and a pointer to it is returned
+		FNodeTemplate* BuildNodeTemplate(TArray<uint8>& NodeTemplateBuffer) const;
+
+		// Returns a node template for the provided list of decorators
+		// The node template will be built into the provided buffer and a pointer to it is returned
+		static FNodeTemplate* BuildNodeTemplate(const TArray<FDecoratorUID>& InDecoratorUIDs, TArray<uint8>& NodeTemplateBuffer);
+
+		// Resets the node template builder
+		void Reset();
+
+	private:
+		static uint32 GetNodeTemplateUID(const TArray<FDecoratorUID>& InDecoratorUIDs);
+		static void AppendTemplateDecorator(
+			const TArray<FDecoratorUID>& InDecoratorUIDs, int32 DecoratorIndex,
+			TArray<uint8>& NodeTemplateBuffer);
+
+		TArray<FDecoratorUID> DecoratorUIDs;	// The list of decorators to use when building the node template
+	};
+}
diff --git a/Engine/Plugins/Experimental/AnimNext/Source/AnimNext/Internal/DecoratorBase/NodeTemplateRegistry.h b/Engine/Plugins/Experimental/AnimNext/Source/AnimNext/Internal/DecoratorBase/NodeTemplateRegistry.h
--- a/Engine/Plugins/Experimental/AnimNext/Source/AnimNext/Internal/DecoratorBase/NodeTemplateRegistry.h
+++ b/Engine/Plugins/Experimental/AnimNext/Source/AnimNext/Internal/DecoratorBase/NodeTemplateRegistry.h
@@ -27,9 +27,6 @@ namespace UE::AnimNext
 		// Finds or adds the specified node template and returns its offset
 		FNodeTemplateRegistryHandle FindOrAdd(const FNodeTemplate* NodeTemplate);
 
-		// Finds and returns a node template based on its handle or nullptr if the handle is invalid
-		FNodeTemplate* Find(FNodeTemplateRegistryHandle TemplateHandle);
-
 		// Finds and returns a node template based on its handle or nullptr if the handle is invalid
 		const FNodeTemplate* Find(FNodeTemplateRegistryHandle TemplateHandle) const;
 
@@ -39,11 +36,15 @@ namespace UE::AnimNext
 		// Removes the specified node template from the registry
 		void Unregister(const FNodeTemplate* NodeTemplate);
 
-		// Clears the registry by removing every node template
-		void Clear();
-
 	private:
 		FNodeTemplateRegistry() = default;
+		FNodeTemplateRegistry(const FNodeTemplateRegistry&) = delete;
+		FNodeTemplateRegistry(FNodeTemplateRegistry&&) = default;
+		FNodeTemplateRegistry& operator=(const FNodeTemplateRegistry&) = delete;
+		FNodeTemplateRegistry& operator=(FNodeTemplateRegistry&&) = default;
+
+		// Finds and returns a node template based on its handle or nullptr if the handle is invalid
+		FNodeTemplate* FindMutable(FNodeTemplateRegistryHandle TemplateHandle);
 
 		// Module lifetime functions
 		static void Init();
@@ -60,5 +61,7 @@ namespace UE::AnimNext
 		TMap<uint32, FNodeTemplateRegistryHandle>	TemplateUIDToHandleMap;
 
 		friend class FModule;
+		friend class FDecoratorWriter;
+		friend struct FScopedClearNodeTemplateRegistry;
 	};
 }
diff --git a/Engine/Plugins/Experimental/AnimNext/Source/AnimNext/Internal/Graph/AnimGraph/AnimNode_AnimNextGraph.h b/Engine/Plugins/Experimental/AnimNext/Source/AnimNext/Internal/Graph/AnimGraph/AnimNode_AnimNextGraph.h
--- a/Engine/Plugins/Experimental/AnimNext/Source/AnimNext/Internal/Graph/AnimGraph/AnimNode_AnimNextGraph.h
+++ b/Engine/Plugins/Experimental/AnimNext/Source/AnimNext/Internal/Graph/AnimGraph/AnimNode_AnimNextGraph.h
@@ -4,24 +4,10 @@
 
 #include "Animation/AnimNode_CustomProperty.h"
 #include "Graph/AnimNextGraph.h"
+#include "DecoratorBase/DecoratorPtr.h"
 #include "Context.h"
 #include "AnimNode_AnimNextGraph.generated.h"
 
-class UNodeMappingContainer;
-
-// TEST - until we can allocate per-node state again
-USTRUCT()
-struct FAnimSequencePlayerState
-{
-	GENERATED_BODY()
-
-	UPROPERTY()
-	float InternalTimeAccumulator = 0.0f;
-
-	UPROPERTY()
-	float PrevInternalTimeAccumulator = 0.0f;
-};
-
 /**
  * Animation node that allows a AnimNextGraph output to be used in an animation graph
  */
@@ -31,6 +17,7 @@ struct ANIMNEXT_API FAnimNode_AnimNextGraph : public FAnimNode_CustomProperty
 	GENERATED_BODY()
 
 	FAnimNode_AnimNextGraph();
+	~FAnimNode_AnimNextGraph();
 
 	// FAnimNode_Base interface
 	virtual void OnInitializeAnimInstance(const FAnimInstanceProxy* InProxy, const UAnimInstance* InAnimInstance) override;
@@ -57,10 +44,8 @@ private:
 	UPROPERTY(EditAnywhere, Category = Settings, meta = (PinHiddenByDefault))
 	TObjectPtr<UAnimNextGraph> AnimNextGraph;
 
-	UPROPERTY(EditAnywhere, Category = Settings, meta = (PinHiddenByDefault, DisallowedClasses = "/Script/Engine.AnimMontage"))
-	TObjectPtr<UAnimSequence> TestSequence = nullptr;
-
-	FAnimSequencePlayerState SequencePlayerState;	// TEST - until we can allocate per-node state again
+	// Shared pointer to our graph instance, we own it
+	UE::AnimNext::FDecoratorPtr GraphInstancePtr;
 
 	/*
 	 * Max LOD that this node is allowed to run
@@ -71,9 +56,6 @@ private:
 	UPROPERTY(EditAnywhere, Category = Performance, meta = (DisplayName = "LOD Threshold"))
 	int32 LODThreshold;
 
-	// Delta time received accumulated in update and used at Evaluate (so we can receive multiple calls to Evaluate)
-	float GraphDeltaTime = 0.f;
-
 protected:
 	virtual UClass* GetTargetClass() const override { return AnimNextGraph ? AnimNextGraph->StaticClass() : nullptr; }
 	
diff --git a/Engine/Plugins/Experimental/AnimNext/Source/AnimNext/Internal/Graph/AnimNextExecuteContext.h b/Engine/Plugins/Experimental/AnimNext/Source/AnimNext/Internal/Graph/AnimNextExecuteContext.h
--- a/Engine/Plugins/Experimental/AnimNext/Source/AnimNext/Internal/Graph/AnimNextExecuteContext.h
+++ b/Engine/Plugins/Experimental/AnimNext/Source/AnimNext/Internal/Graph/AnimNextExecuteContext.h
@@ -5,6 +5,8 @@
 #include "CoreMinimal.h"
 #include "RigVMCore/RigVMExecuteContext.h"
 #include "Units/RigUnit.h"
+#include "DecoratorBase/DecoratorPtr.h"
+
 #include "AnimNextExecuteContext.generated.h"
 
 namespace UE::AnimNext
@@ -12,6 +14,26 @@ namespace UE::AnimNext
 	struct FContext;
 }
 
+/**
+  * An enum to control which simulation steps should be performed when the animation graph is processed.
+  */
+enum class EAnimNextGraphSimulationSteps
+{
+	// No step is performed
+	None = 0x00,
+
+	// Performs the update step
+	Update = 0x01,
+
+	// Performs the evaluate step
+	Evaluate = 0x02,
+
+	// Perform all simulation steps
+	All = 0xff,
+};
+
+ENUM_CLASS_FLAGS(EAnimNextGraphSimulationSteps)
+
 USTRUCT(BlueprintType)
 struct FAnimNextExecuteContext : public FRigVMExecuteContext
 {
@@ -29,10 +51,11 @@ struct FAnimNextExecuteContext : public FRigVMExecuteContext
 		return *Context;
 	}
 
-	void SetContextData(const UE::AnimNext::FContext& InContext)
-	{
-		Context = &InContext;
-	}
+	const TArrayView<const uint8>& GetSharedDataBuffer() const { return SharedDataBuffer; }
+
+	UE::AnimNext::FWeakDecoratorPtr GetGraphInstancePtr() const { return GraphInstancePtr; }
+
+	EAnimNextGraphSimulationSteps GetSimulationSteps() const { return SimulationSteps; }
 
 	virtual void Copy(const FRigVMExecuteContext* InOtherContext) override
 	{
@@ -40,10 +63,35 @@ struct FAnimNextExecuteContext : public FRigVMExecuteContext
 
 		const FAnimNextExecuteContext* OtherContext = (const FAnimNextExecuteContext*)InOtherContext;
 		Context = OtherContext->Context;
+		SharedDataBuffer = OtherContext->SharedDataBuffer;
+		GraphInstancePtr = OtherContext->GraphInstancePtr;
 	}
 
 private:
+	void SetContextData(const UE::AnimNext::FContext& InContext)
+	{
+		Context = &InContext;
+	}
+
+	void InitializeWithGraph(TArrayView<const uint8> InSharedDataBuffer, UE::AnimNext::FWeakDecoratorPtr InGraphInstancePtr)
+	{
+		SharedDataBuffer = InSharedDataBuffer;
+		GraphInstancePtr = InGraphInstancePtr;
+	}
+
+	void SetSimulationSteps(EAnimNextGraphSimulationSteps InSimulationSteps)
+	{
+		SimulationSteps = InSimulationSteps;
+	}
+
 	const UE::AnimNext::FContext* Context;
+
+	TArrayView<const uint8> SharedDataBuffer;
+	UE::AnimNext::FWeakDecoratorPtr GraphInstancePtr;
+	EAnimNextGraphSimulationSteps SimulationSteps;
+
+	friend class UAnimNextGraph;
+	friend class UAnimNextParameterBlock;
 };
 
 USTRUCT(meta=(ExecuteContext="FAnimNextExecuteContext"))
diff --git a/Engine/Plugins/Experimental/AnimNext/Source/AnimNext/Internal/Graph/AnimNextGraph.h b/Engine/Plugins/Experimental/AnimNext/Source/AnimNext/Internal/Graph/AnimNextGraph.h
--- a/Engine/Plugins/Experimental/AnimNext/Source/AnimNext/Internal/Graph/AnimNextGraph.h
+++ b/Engine/Plugins/Experimental/AnimNext/Source/AnimNext/Internal/Graph/AnimNextGraph.h
@@ -4,6 +4,9 @@
 
 #include "CoreMinimal.h"
 #include "RigVMCore/RigVM.h"
+#include "DecoratorBase/DecoratorPtr.h"
+#include "DecoratorBase/DecoratorHandle.h"
+#include "DecoratorBase/EntryPointHandle.h"
 
 #include "AnimNextGraph.generated.h"
 
@@ -11,6 +14,7 @@ class UEdGraph;
 class UAnimNextGraph;
 class UAnimGraphNode_AnimNextGraph;
 struct FAnimNode_AnimNextGraph;
+enum class EAnimNextGraphSimulationSteps;
 
 namespace UE::AnimNext
 {
@@ -39,15 +43,28 @@ class ANIMNEXT_API UAnimNextGraph : public UObject
 {
 	GENERATED_BODY()
 
+public:
 	// UObject interface
 	virtual void PostRename(UObject* OldOuter, const FName OldName) override;
 	virtual void GetPreloadDependencies(TArray<UObject*>& OutDeps) override;
+	virtual void Serialize(FArchive& Ar) override;
+	static void AddReferencedObjects(UObject* InThis, FReferenceCollector& Collector);
 
-	// Run this graph with the given context
-	void Run(const UE::AnimNext::FContext& Context) const;
+	// Allocates an instance of the graph, retain the handle and use it with the Run() function to evaluate it
+	UE::AnimNext::FDecoratorPtr AllocateInstance() const;
 
+	// Releases an instance of the graph and clears the provided handle
+	void ReleaseInstance(UE::AnimNext::FDecoratorPtr& GraphInstancePtr) const;
+
+	// Run the specified simulation steps on the provided graph with the given context
+	void Run(const UE::AnimNext::FContext& Context, UE::AnimNext::FWeakDecoratorPtr GraphInstancePtr, EAnimNextGraphSimulationSteps SimulationSteps) const;
+
+protected:
 	// Support rig VM execution
 	TArray<FRigVMExternalVariable> GetRigVMExternalVariables();
+
+	// Loads the graph data from the provided archive buffer and returns true on success, false otherwise
+	bool LoadFromArchiveBuffer(const TArray<uint8>& SharedDataArchiveBuffer);
 	
 	friend class UAnimNextGraphFactory;
 	friend class UAnimNextGraph_EditorData;
@@ -55,11 +72,26 @@ class ANIMNEXT_API UAnimNextGraph : public UObject
 	friend class UE::AnimNext::Editor::FGraphEditor;
 	friend class UAnimNextGraph;
 	friend class UAnimGraphNode_AnimNextGraph;
-	friend struct FAnimNode_AnimNextGraph;
 
 	UPROPERTY()
 	TObjectPtr<URigVM> RigVM;
 
+	// This is a handle to the root decorator in our graph
+	UPROPERTY()
+	FAnimNextEntryPointHandle RootDecoratorHandle;
+
+	// This is a resolved handle to the root decorator in our graph
+	FAnimNextDecoratorHandle ResolvedRootDecoratorHandle;
+
+	// This is the graph shared data used by the decorator system, the output of FDecoratorReader
+	// We de-serialize manually into this buffer from the archive buffer, this is never saved on disk
+	TArray<uint8> SharedDataBuffer;
+
+	// This is the list of all referenced UObjects in the graph shared data
+	// We collect all the references here to make it quick and easy for the GC to query them
+	// It means that object references in the graph shared data are not visited at runtime by the GC (they are immutable)
+	TArray<TObjectPtr<UObject>> TrackedObjectsForGC;
+
 	UPROPERTY(transient)
 	mutable FRigVMExtendedExecuteContext ExtendedExecuteContext;
 
@@ -69,5 +101,9 @@ class ANIMNEXT_API UAnimNextGraph : public UObject
 #if WITH_EDITORONLY_DATA
 	UPROPERTY(VisibleAnywhere, Instanced, Category = "Graph", meta = (ShowInnerProperties))
 	TObjectPtr<UObject> EditorData;
+
+	// This buffer holds the output of the FDecoratorWriter post compilation
+	// We serialize it manually and it is discarded at runtime
+	TArray<uint8> SharedDataArchiveBuffer;
 #endif
 };
diff --git a/Engine/Plugins/Experimental/AnimNext/Source/AnimNext/Internal/Graph/RigDecorator_AnimNextCppDecorator.h b/Engine/Plugins/Experimental/AnimNext/Source/AnimNext/Internal/Graph/RigDecorator_AnimNextCppDecorator.h
new file mode 100644
--- /dev/null
+++ b/Engine/Plugins/Experimental/AnimNext/Source/AnimNext/Internal/Graph/RigDecorator_AnimNextCppDecorator.h
@@ -0,0 +1,38 @@
+// Copyright Epic Games, Inc. All Rights Reserved.
+
+#pragma once
+
+#include "RigVMCore/RigVMDecorator.h"
+
+#include "RigDecorator_AnimNextCppDecorator.generated.h"
+
+namespace UE::AnimNext { struct FDecorator; }
+
+/**
+ * AnimNext RigDecorator for all C++ decorators.
+ * The decorator shared data UScriptStruct determines which properties are exposed.
+ */
+USTRUCT(BlueprintType)
+struct ANIMNEXT_API FRigDecorator_AnimNextCppDecorator : public FRigVMDecorator
+{
+	GENERATED_BODY()
+
+	// The struct the decorator exposes with its shared data. Each one of its properties will be added as a pin.
+	UPROPERTY()
+	UScriptStruct* DecoratorSharedDataStruct = nullptr;
+
+#if WITH_EDITOR
+	virtual void GetProgrammaticPins(URigVMController* InController, int32 InParentPinIndex, struct FRigVMPinInfoArray& OutPinArray) const override;
+
+	const UE::AnimNext::FDecorator* GetDecorator() const;
+#endif
+};
+
+USTRUCT(BlueprintType)
+struct FAnimNextCppDecoratorWrapper
+{
+	GENERATED_BODY()
+
+	UPROPERTY()
+	FRigDecorator_AnimNextCppDecorator CppDecorator;
+};
diff --git a/Engine/Plugins/Experimental/AnimNext/Source/AnimNext/Internal/Graph/RigUnit_AnimNextDecoratorStack.h b/Engine/Plugins/Experimental/AnimNext/Source/AnimNext/Internal/Graph/RigUnit_AnimNextDecoratorStack.h
new file mode 100644
--- /dev/null
+++ b/Engine/Plugins/Experimental/AnimNext/Source/AnimNext/Internal/Graph/RigUnit_AnimNextDecoratorStack.h
@@ -0,0 +1,27 @@
+// Copyright Epic Games, Inc. All Rights Reserved.
+
+#pragma once
+
+#include "AnimNextExecuteContext.h"
+#include "DecoratorBase/DecoratorHandle.h"
+
+#include "RigUnit_AnimNextDecoratorStack.generated.h"
+
+/**
+ * Animation graph decorator container
+ * It contains a stack of decorators which are used during compilation to output a cooked graph
+ * This node is synthetic and only present in the editor. During compilation, it is disconnected from
+ * the graph and replaced by the cooked decorator metadata.
+ */
+USTRUCT(meta=(DisplayName="Decorator Stack", Category="Animation", NodeColor="0, 1, 1", Keywords="Decorator,Stack"))
+struct ANIMNEXT_API FRigUnit_AnimNextDecoratorStack : public FRigUnit_AnimNextBase
+{
+	GENERATED_BODY()
+
+	RIGVM_METHOD()
+	void Execute();
+
+	// The execution result
+	UPROPERTY(EditAnywhere, Category = Result, meta = (Output))
+	FAnimNextDecoratorHandle Result;
+};
diff --git a/Engine/Plugins/Experimental/AnimNext/Source/AnimNext/Internal/Graph/RigUnit_AnimNextGraphRoot.h b/Engine/Plugins/Experimental/AnimNext/Source/AnimNext/Internal/Graph/RigUnit_AnimNextGraphRoot.h
new file mode 100644
--- /dev/null
+++ b/Engine/Plugins/Experimental/AnimNext/Source/AnimNext/Internal/Graph/RigUnit_AnimNextGraphRoot.h
@@ -0,0 +1,35 @@
+﻿// Copyright Epic Games, Inc. All Rights Reserved.
+
+#pragma once
+
+#include "AnimNextExecuteContext.h"
+#include "DecoratorBase/DecoratorHandle.h"
+
+#include "RigUnit_AnimNextGraphRoot.generated.h"
+
+/**
+ * Animation graph output
+ */
+USTRUCT(meta=(DisplayName="Animation Output", Category="Events", NodeColor="1, 0, 0", Keywords="Root,Output"))
+struct ANIMNEXT_API FRigUnit_AnimNextGraphRoot : public FRigUnit_AnimNextBase
+{
+	GENERATED_BODY()
+
+	RIGVM_METHOD()
+	void Execute();
+
+	virtual FName GetEventName() const override { return EventName; }
+	virtual bool CanOnlyExistOnce() const override { return true; }
+
+	// The execution result
+	UPROPERTY(EditAnywhere, Category = Result, meta = (Input))
+	FAnimNextDecoratorHandle Result;
+
+	// In order for this node to be considered an executable RigUnit, it needs a pin to derive from FRigVMExecuteContext
+	// We keep it hidden it since we don't need it
+	UPROPERTY()
+	FAnimNextExecuteContext ExecuteContext;
+
+	// This unit is our graph entry point, it needs an event so we can call it
+	static FName EventName;
+};
diff --git a/Engine/Plugins/Experimental/AnimNext/Source/AnimNext/Private/DecoratorBase/Decorator.cpp b/Engine/Plugins/Experimental/AnimNext/Source/AnimNext/Private/DecoratorBase/Decorator.cpp
--- a/Engine/Plugins/Experimental/AnimNext/Source/AnimNext/Private/DecoratorBase/Decorator.cpp
+++ b/Engine/Plugins/Experimental/AnimNext/Source/AnimNext/Private/DecoratorBase/Decorator.cpp
@@ -18,7 +18,7 @@ namespace UE::AnimNext
 	}
 
 #if WITH_EDITOR
-	void FDecorator::SaveDecoratorSharedData(FDecoratorWriter& Writer, const TMap<FString, FString>& Properties, FAnimNextDecoratorSharedData& OutSharedData) const
+	void FDecorator::SaveDecoratorSharedData(FDecoratorWriter& Writer, const TFunction<FString(const FString& PropertyName)>& GetDecoratorProperty, FAnimNextDecoratorSharedData& OutSharedData) const
 	{
 		const UScriptStruct* SharedDataStruct = GetDecoratorSharedDataStruct();
 
@@ -31,9 +31,10 @@ namespace UE::AnimNext
 		// We convert every property from its string representation into its binary form
 		for (const FProperty* Property = SharedDataStruct->PropertyLink; Property != nullptr; Property = Property->PropertyLinkNext)
 		{
-			if (const FString* PropertyValue = Properties.Find(Property->GetName()))
+			const FString PropertyValue = GetDecoratorProperty(Property->GetName());
+			if (PropertyValue.Len() != 0)
 			{
-				const TCHAR* PropertyValuePtr = **PropertyValue;
+				const TCHAR* PropertyValuePtr = *PropertyValue;
 
 				// C-style array properties aren't handled by ExportText, we need to handle it manually
 				const bool bIsCArray = Property->ArrayDim > 1;
diff --git a/Engine/Plugins/Experimental/AnimNext/Source/AnimNext/Private/DecoratorBase/DecoratorHandle.cpp b/Engine/Plugins/Experimental/AnimNext/Source/AnimNext/Private/DecoratorBase/DecoratorHandle.cpp
new file mode 100644
--- /dev/null
+++ b/Engine/Plugins/Experimental/AnimNext/Source/AnimNext/Private/DecoratorBase/DecoratorHandle.cpp
@@ -0,0 +1,22 @@
+// Copyright Epic Games, Inc. All Rights Reserved.
+
+#include "DecoratorBase/DecoratorHandle.h"
+#include "DecoratorBase/DecoratorReader.h"
+#include "Serialization/Archive.h"
+
+bool FAnimNextDecoratorHandle::Serialize(FArchive& Ar)
+{
+	Ar << PackedDecoratorIndexAndNodeHandle;
+
+	if (Ar.IsLoading() && IsValid())
+	{
+		// On load, we hold a node ID that needs fix-up
+		check(GetNodeHandle().IsNodeID());
+
+		UE::AnimNext::FDecoratorReader& DecoratorReader = static_cast<UE::AnimNext::FDecoratorReader&>(Ar);
+		*this = DecoratorReader.ResolveDecoratorHandle(*this);
+		check(GetNodeHandle().IsSharedOffset());
+	}
+
+	return true;
+}
diff --git a/Engine/Plugins/Experimental/AnimNext/Source/AnimNext/Private/DecoratorBase/DecoratorPtr.cpp b/Engine/Plugins/Experimental/AnimNext/Source/AnimNext/Private/DecoratorBase/DecoratorPtr.cpp
--- a/Engine/Plugins/Experimental/AnimNext/Source/AnimNext/Private/DecoratorBase/DecoratorPtr.cpp
+++ b/Engine/Plugins/Experimental/AnimNext/Source/AnimNext/Private/DecoratorBase/DecoratorPtr.cpp
@@ -7,10 +7,10 @@ namespace UE::AnimNext
 {
 	FDecoratorPtr::FDecoratorPtr(FNodeInstance* NodeInstance, uint32 DecoratorIndex_)
 		: PackedPointerAndFlags(NodeInstance != nullptr ? reinterpret_cast<uintptr_t>(NodeInstance) : 0)
-		, DecoratorIndex(static_cast<uint8>(DecoratorIndex_))
+		, DecoratorIndex(DecoratorIndex_)
 	{
 		check((reinterpret_cast<uintptr_t>(NodeInstance) & FLAGS_MASK) == 0);	// Make sure we have enough alignment
-		check(DecoratorIndex == DecoratorIndex_);	// Make sure we didn't truncate
+		check(DecoratorIndex <= MAX_uint8);	// Make sure we don't truncate
 
 		if (NodeInstance != nullptr)
 		{
@@ -43,10 +43,10 @@ namespace UE::AnimNext
 
 	FDecoratorPtr::FDecoratorPtr(FNodeInstance* NodeInstance, EFlags Flags, uint32 DecoratorIndex_)
 		: PackedPointerAndFlags(NodeInstance != nullptr ? (reinterpret_cast<uintptr_t>(NodeInstance) | Flags) : 0)
-		, DecoratorIndex(static_cast<uint8>(DecoratorIndex_))
+		, DecoratorIndex(DecoratorIndex_)
 	{
 		check((reinterpret_cast<uintptr_t>(NodeInstance) & FLAGS_MASK) == 0);	// Make sure we have enough alignment
-		check(DecoratorIndex == DecoratorIndex_);	// Make sure we didn't truncate
+		check(DecoratorIndex <= MAX_uint8);	// Make sure we don't truncate
 
 		// Only increment the reference count if we aren't a weak handle
 		if (NodeInstance != nullptr && (Flags & IS_WEAK_BIT) == 0)
@@ -112,6 +112,13 @@ namespace UE::AnimNext
 		DecoratorIndex = 0;
 	}
 
+	FWeakDecoratorPtr::FWeakDecoratorPtr(FNodeInstance* NodeInstance_, uint32 DecoratorIndex_)
+		: NodeInstance(NodeInstance_)
+		, DecoratorIndex(DecoratorIndex_)
+	{
+		check(DecoratorIndex <= MAX_uint8);	// Make sure we don't truncate
+	}
+
 	void FWeakDecoratorPtr::Reset()
 	{
 		NodeInstance = nullptr;
diff --git a/Engine/Plugins/Experimental/AnimNext/Source/AnimNext/Private/DecoratorBase/DecoratorReader.cpp b/Engine/Plugins/Experimental/AnimNext/Source/AnimNext/Private/DecoratorBase/DecoratorReader.cpp
--- a/Engine/Plugins/Experimental/AnimNext/Source/AnimNext/Private/DecoratorBase/DecoratorReader.cpp
+++ b/Engine/Plugins/Experimental/AnimNext/Source/AnimNext/Private/DecoratorBase/DecoratorReader.cpp
@@ -3,8 +3,11 @@
 #include "DecoratorBase/DecoratorReader.h"
 
 #include "DecoratorBase/NodeDescription.h"
+#include "DecoratorBase/NodeInstance.h"
 #include "DecoratorBase/NodeTemplate.h"
 #include "DecoratorBase/NodeTemplateRegistry.h"
+#include "Serialization/ArchiveUObject.h"
+#include "UObject/UObjectGlobals.h"
 
 namespace UE::AnimNext
 {
@@ -13,48 +16,191 @@ namespace UE::AnimNext
 	{
 	}
 
-	void FDecoratorReader::ReadGraphSharedData(TArray<uint8>& GraphSharedData)
+	FDecoratorReader::EErrorState FDecoratorReader::ReadGraph(TArray<uint8>& GraphSharedData, TArray<TObjectPtr<UObject>>& TrackedObjectsForGC)
 	{
 		GraphSharedData.Empty(0);
+		TrackedObjectsForGC.Empty(0);
 
+		EErrorState ErrorState = ReadGraphSharedData(GraphSharedData);
+		if (ErrorState != EErrorState::None)
+		{
+			return ErrorState;
+		}
+
+		ErrorState = ReadTrackedObjectsForGC(TrackedObjectsForGC);
+
+		return ErrorState;
+	}
+
+	FDecoratorReader::EErrorState FDecoratorReader::ReadGraphSharedData(TArray<uint8>& GraphSharedData)
+	{
 		// Read the node templates and register them as needed
 		uint32 NumNodeTemplates = 0;
 		*this << NumNodeTemplates;
 
-		// We serialize each node template in the same buffer, they get copied into the registry and we don't need to retain them
-		alignas(16) uint8 NodeTemplateBuffer[64 * 1024];
-		FNodeTemplate* NodeTemplate = reinterpret_cast<FNodeTemplate*>(&NodeTemplateBuffer[0]);
+		TArray<FNodeTemplateRegistryHandle> NodeTemplateHandles;
+		NodeTemplateHandles.Reserve(NumNodeTemplates);
 
 		FNodeTemplateRegistry& NodeTemplateRegistry = FNodeTemplateRegistry::Get();
-		for (uint32 NodeTemplateIndex = 0; NodeTemplateIndex < NumNodeTemplates; ++NodeTemplateIndex)
+
 		{
-			NodeTemplate->Serialize(*this);
+			// We serialize each node template in the same buffer, they get copied into the registry and we don't need to retain them
+			alignas(16) uint8 NodeTemplateBuffer[64 * 1024];
+			FNodeTemplate* NodeTemplate = reinterpret_cast<FNodeTemplate*>(&NodeTemplateBuffer[0]);
 
-			// Register our node template
-			NodeTemplateRegistry.FindOrAdd(NodeTemplate);
+			for (uint32 NodeTemplateIndex = 0; NodeTemplateIndex < NumNodeTemplates; ++NodeTemplateIndex)
+			{
+				NodeTemplate->Serialize(*this);
+
+				if (!NodeTemplate->IsValid())
+				{
+					if (NodeTemplate->GetNodeSharedDataSize() == 0)
+					{
+						// This node shared data is too large
+						return EErrorState::NodeSharedDataTooLarge;
+					}
+
+					if (NodeTemplate->GetNodeInstanceDataSize() == 0)
+					{
+						// This node instance data is too large
+						return EErrorState::NodeInstanceDataTooLarge;
+					}
+				}
+
+				// Register our node template
+				NodeTemplateHandles.Add(NodeTemplateRegistry.FindOrAdd(NodeTemplate));
+			}
 		}
 
 		// Read our graph shared data
 		uint32 NumNodes = 0;
 		*this << NumNodes;
 
+		NodeHandles.Empty(0);
+		NodeHandles.Reserve(NumNodes);
+
+		// Calculate our shared data size and all node offsets
 		uint32 SharedDataSize = 0;
-		*this << SharedDataSize;
+		for (uint32 NodeIndex = 0; NodeIndex < NumNodes; ++NodeIndex)
+		{
+			SharedDataSize = Align(SharedDataSize, alignof(FNodeDescription));	// Make sure we respect our alignment constraints
+
+			if (SharedDataSize > MAXIMUM_GRAPH_SHARED_DATA_SIZE)
+			{
+				// This graph shared data is too large, we won't be able to create handles to this node
+				return EErrorState::GraphTooLarge;
+			}
+
+			NodeHandles.Add(FNodeHandle::FromSharedOffset(SharedDataSize));	// This node starts here
+			check(NodeHandles.Last().IsSharedOffset());
 
+			uint32 NodeTemplateIndex = 0;
+			*this << NodeTemplateIndex;
+
+			const FNodeTemplate* NodeTemplate = NodeTemplateRegistry.Find(NodeTemplateHandles[NodeTemplateIndex]);
+			check(NodeTemplate != nullptr);
+
+			SharedDataSize += NodeTemplate->GetNodeSharedDataSize();
+		}
+
+		// The shared data size might exceed MAXIMUM_GRAPH_SHARED_DATA_SIZE a little bit
+		// The only requirement is that the node begins before that threshold so that we can
+		// create handles to it
+
+		GraphSharedData.Empty(0);
 		GraphSharedData.AddZeroed(SharedDataSize);
 
-		uint32 SharedDataOffset = 0;
+		// Serialize our graph
 		for (uint32 NodeIndex = 0; NodeIndex < NumNodes; ++NodeIndex)
 		{
-			uint32 NodeSize = 0;
-			*this << NodeSize;
-
+			const uint32 SharedDataOffset = NodeHandles[NodeIndex].GetSharedOffset();
 			FNodeDescription* NodeDesc = reinterpret_cast<FNodeDescription*>(&GraphSharedData[SharedDataOffset]);
 			NodeDesc->Serialize(*this);
+		}
+
+		return EErrorState::None;
+	}
+
+	FDecoratorReader::EErrorState FDecoratorReader::ReadTrackedObjectsForGC(TArray<TObjectPtr<UObject>>& TrackedObjectsForGC)
+	{
+		int32 NumTrackedObjectsForGC = 0;
+		*this << NumTrackedObjectsForGC;
+
+		TrackedObjectsForGC.SetNumZeroed(NumTrackedObjectsForGC);
+
+		for (int32 TrackedIndex = 0; TrackedIndex < NumTrackedObjectsForGC; ++TrackedIndex)
+		{
+			// Load the path name to the object
+			FString LoadedString;
+			*this << LoadedString;
 
-			SharedDataOffset += NodeSize;
+			// Look up the object by fully qualified pathname
+			UObject* Obj = FindObject<UObject>(nullptr, *LoadedString, false);
+			TrackedObjectsForGC[TrackedIndex] = Obj;
+		}
+
+		return EErrorState::None;
+	}
+
+	FArchive& FDecoratorReader::operator<<(UObject*& Obj)
+	{
+		// Load the path name to the object
+		FString LoadedString;
+		*this << LoadedString;
+
+		// Look up the object by fully qualified pathname
+		Obj = FindObject<UObject>(nullptr, *LoadedString, false);
+
+		// If we couldn't find it, load it
+		if (Obj == nullptr)
+		{
+			Obj = LoadObject<UObject>(nullptr, *LoadedString);
+		}
+
+		return *this;
+	}
+
+	FArchive& FDecoratorReader::operator<<(FObjectPtr& Obj)
+	{
+		return FArchiveUObject::SerializeObjectPtr(*this, Obj);
+	}
+
+	FNodeHandle FDecoratorReader::ResolveNodeHandle(FNodeHandle NodeHandle) const
+	{
+		if (!NodeHandle.IsValid())
+		{
+			// The node handle is invalid, return it unchanged
+			return NodeHandle;
+		}
+
+		check(NodeHandle.IsNodeID());
+		const FNodeID NodeID = NodeHandle.GetNodeID();
+		check(NodeID.IsValid());
+
+		return NodeHandles[NodeID.GetNodeIndex()];
+	}
+
+	FAnimNextDecoratorHandle FDecoratorReader::ResolveDecoratorHandle(FAnimNextDecoratorHandle DecoratorHandle) const
+	{
+		if (!DecoratorHandle.IsValid())
+		{
+			// The decorator handle is invalid, return it unchanged
+			return DecoratorHandle;
+		}
+
+		const FNodeHandle NodeHandle = ResolveNodeHandle(DecoratorHandle.GetNodeHandle());
+		return FAnimNextDecoratorHandle(NodeHandle, DecoratorHandle.GetDecoratorIndex());
+	}
+
+	FAnimNextDecoratorHandle FDecoratorReader::ResolveEntryPointHandle(FAnimNextEntryPointHandle EntryPointHandle) const
+	{
+		if (!EntryPointHandle.IsValid())
+		{
+			// The decorator handle is invalid, return an invalid handle
+			return FAnimNextDecoratorHandle();
 		}
 
-		ensure(SharedDataOffset == SharedDataSize);
+		const FNodeHandle NodeHandle = ResolveNodeHandle(EntryPointHandle.GetNodeHandle());
+		return FAnimNextDecoratorHandle(NodeHandle, EntryPointHandle.GetDecoratorIndex());
 	}
 }
diff --git a/Engine/Plugins/Experimental/AnimNext/Source/AnimNext/Private/DecoratorBase/DecoratorRegistry.cpp b/Engine/Plugins/Experimental/AnimNext/Source/AnimNext/Private/DecoratorBase/DecoratorRegistry.cpp
--- a/Engine/Plugins/Experimental/AnimNext/Source/AnimNext/Private/DecoratorBase/DecoratorRegistry.cpp
+++ b/Engine/Plugins/Experimental/AnimNext/Source/AnimNext/Private/DecoratorBase/DecoratorRegistry.cpp
@@ -206,9 +206,14 @@ namespace UE::AnimNext
 		}
 	}
 
-	FDecoratorRegistryHandle FDecoratorRegistry::FindHandle(uint32 DecoratorUID) const
+	FDecoratorRegistryHandle FDecoratorRegistry::FindHandle(FDecoratorUID DecoratorUID) const
 	{
-		if (const FRegistryEntry* Entry = DecoratorUIDToEntryMap.Find(DecoratorUID))
+		if (!DecoratorUID.IsValid())
+		{
+			return FDecoratorRegistryHandle();
+		}
+
+		if (const FRegistryEntry* Entry = DecoratorUIDToEntryMap.Find(DecoratorUID.GetUID()))
 		{
 			return Entry->DecoratorHandle;
 		}
@@ -217,11 +222,6 @@ namespace UE::AnimNext
 		return FDecoratorRegistryHandle();
 	}
 
-	FDecoratorRegistryHandle FDecoratorRegistry::FindHandle(FDecoratorUID DecoratorUID) const
-	{
-		return FindHandle(DecoratorUID.GetUID());
-	}
-
 	const FDecorator* FDecoratorRegistry::Find(FDecoratorRegistryHandle DecoratorHandle) const
 	{
 		if (!DecoratorHandle.IsValid())
@@ -247,6 +247,25 @@ namespace UE::AnimNext
 		return Find(DecoratorHandle);
 	}
 
+	const FDecorator* FDecoratorRegistry::Find(const UScriptStruct* DecoratorSharedDataStruct) const
+	{
+		if (DecoratorSharedDataStruct == nullptr)
+		{
+			return nullptr;
+		}
+
+		for (const auto& KeyValuePair : DecoratorUIDToEntryMap)
+		{
+			const FDecorator* Decorator = KeyValuePair.Value.Decorator;
+			if (Decorator->GetDecoratorSharedDataStruct() == DecoratorSharedDataStruct)
+			{
+				return Decorator;
+			}
+		}
+
+		return nullptr;
+	}
+
 	void FDecoratorRegistry::Register(FDecorator* Decorator)
 	{
 		if (Decorator == nullptr)
@@ -322,6 +341,19 @@ namespace UE::AnimNext
 		}
 	}
 
+	TArray<const FDecorator*> FDecoratorRegistry::GetDecorators() const
+	{
+		TArray<const FDecorator*> Decorators;
+		Decorators.Reserve(DecoratorUIDToEntryMap.Num());
+		
+		for (const auto& It : DecoratorUIDToEntryMap)
+		{
+			Decorators.Add(It.Value.Decorator);
+		}
+
+		return Decorators;
+	}
+
 	uint32 FDecoratorRegistry::GetNum() const
 	{
 		return DecoratorUIDToEntryMap.Num();
diff --git a/Engine/Plugins/Experimental/AnimNext/Source/AnimNext/Private/DecoratorBase/DecoratorTemplate.cpp b/Engine/Plugins/Experimental/AnimNext/Source/AnimNext/Private/DecoratorBase/DecoratorTemplate.cpp
--- a/Engine/Plugins/Experimental/AnimNext/Source/AnimNext/Private/DecoratorBase/DecoratorTemplate.cpp
+++ b/Engine/Plugins/Experimental/AnimNext/Source/AnimNext/Private/DecoratorBase/DecoratorTemplate.cpp
@@ -10,11 +10,13 @@ namespace UE::AnimNext
 {
 	void FDecoratorTemplate::Serialize(FArchive& Ar)
 	{
+		const FDecoratorRegistry& DecoratorRegistry = FDecoratorRegistry::Get();
+
 		Ar << UID;
 
 		if (Ar.IsSaving())
 		{
-			const FDecorator* Decorator = FDecoratorRegistry::Get().Find(RegistryHandle);
+			const FDecorator* Decorator = DecoratorRegistry.Find(RegistryHandle);
 
 			uint32 DecoratorUID = Decorator->GetDecoratorUID().GetUID();
 			Ar << DecoratorUID;
@@ -24,7 +26,11 @@ namespace UE::AnimNext
 			uint32 DecoratorUID = 0;
 			Ar << DecoratorUID;
 
-			RegistryHandle = FDecoratorRegistry::Get().FindHandle(DecoratorUID);
+			// It is possible that we fail to find the decorator that we need
+			// This can happen if the decorator hasn't been loaded or registered
+			// When this happens, the decorator is a no-op and the runtime behavior
+			// may not be what is expected
+			RegistryHandle = DecoratorRegistry.FindHandle(FDecoratorUID(DecoratorUID));
 		}
 		else
 		{
@@ -43,7 +49,5 @@ namespace UE::AnimNext
 
 		Ar << Mode;
 		Ar << AdditiveIndexOrNumAdditive;
-		Ar << NodeSharedOffset;
-		Ar << NodeInstanceOffset;
 	}
 }
diff --git a/Engine/Plugins/Experimental/AnimNext/Source/AnimNext/Private/DecoratorBase/DecoratorWriter.cpp b/Engine/Plugins/Experimental/AnimNext/Source/AnimNext/Private/DecoratorBase/DecoratorWriter.cpp
--- a/Engine/Plugins/Experimental/AnimNext/Source/AnimNext/Private/DecoratorBase/DecoratorWriter.cpp
+++ b/Engine/Plugins/Experimental/AnimNext/Source/AnimNext/Private/DecoratorBase/DecoratorWriter.cpp
@@ -8,37 +8,15 @@
 #include "DecoratorBase/NodeDescription.h"
 #include "DecoratorBase/NodeTemplate.h"
 #include "DecoratorBase/NodeTemplateRegistry.h"
+#include "Serialization/ArchiveUObject.h"
+
+#include "DecoratorBase/DecoratorReader.h"
 
 namespace UE::AnimNext
 {
-	namespace Private
-	{
-		// Note: this does not account for inline/pooled/pinned properties
-		static uint32 GetNodeSharedSize(const FNodeTemplate& NodeTemplate)
-		{
-			const FDecoratorRegistry& Registry = FDecoratorRegistry::Get();
-
-			const uint32 NumDecorators = NodeTemplate.GetNumDecorators();
-			const FDecoratorTemplate* DecoratorTemplates = NodeTemplate.GetDecorators();
-
-			uint32 NodeSharedSize = sizeof(FNodeDescription);
-			for (uint32_t DecoratorIndex = 0; DecoratorIndex < NumDecorators; ++DecoratorIndex)
-			{
-				const FDecorator* Decorator = Registry.Find(DecoratorTemplates[DecoratorIndex].GetRegistryHandle());
-
-				const FDecoratorMemoryLayout MemoryLayout = Decorator->GetDecoratorMemoryDescription();
-				NodeSharedSize = Align(NodeSharedSize, MemoryLayout.SharedDataAlignment);
-				NodeSharedSize += MemoryLayout.SharedDataSize;
-			}
-
-			return NodeSharedSize;
-		}
-	}
-
 	FDecoratorWriter::FDecoratorWriter()
 		: FMemoryWriter(GraphSharedDataArchiveBuffer)
-		, GraphSharedDataSize(0)
-		, NextNodeUID(1)	// Node UID 0 is reserved for the invalid value (null)
+		, NextNodeID(FNodeID::GetFirstID())
 		, NumNodesWritten(0)
 		, bIsNodeWriting(false)
 		, ErrorState(EErrorState::None)
@@ -62,54 +40,28 @@ namespace UE::AnimNext
 			return FNodeHandle();
 		}
 
-		if (NextNodeUID == FNodeDescription::MAXIMUM_COUNT)
+		if (!NextNodeID.IsValid())
 		{
 			// We have too many nodes in the graph, we need to be able to represent them with 16 bits
+			// The node ID must have wrapped around
 			ErrorState = EErrorState::TooManyNodes;
 			return FNodeHandle();
 		}
 
-		const uint32 NodeSharedDataSize = Private::GetNodeSharedSize(NodeTemplate);
-		if (NodeSharedDataSize > FNodeDescription::MAXIMUM_SIZE)
-		{
-			// This node shared data is too large
-			ErrorState = EErrorState::NodeSharedDataTooLarge;
-			return FNodeHandle();
-		}
-
-		const uint32 NodeSharedDataOffset = GraphSharedDataSize;
-		if (GraphSharedDataSize + NodeSharedDataSize > MAXIMUM_GRAPH_SHARED_DATA_SIZE)
-		{
-			// The graph shared data size is too large
-			ErrorState = EErrorState::GraphTooLarge;
-			return FNodeHandle();
-		}
+		const FNodeHandle NodeHandle = FNodeHandle::FromNodeID(NextNodeID);
+		check(NodeHandle.IsValid() && NodeHandle.IsNodeID());
+		check(NodeMappings.Num() == NodeHandle.GetNodeID().GetNodeIndex());
 
-		GraphSharedDataSize += NodeSharedDataSize;
-
-		const FNodeHandle NodeHandle(NodeSharedDataOffset);
-		const uint16 NodeUID = NextNodeUID++;
+		NextNodeID = NextNodeID.GetNextID();
 
 		FNodeTemplateRegistry& NodeTemplateRegistry = FNodeTemplateRegistry::Get();
 		const FNodeTemplateRegistryHandle NodeTemplateHandle = NodeTemplateRegistry.FindOrAdd(&NodeTemplate);
 
-		NodeMappings.Add({ NodeHandle, NodeSharedDataSize, NodeUID, NodeTemplateHandle });
+		NodeMappings.Add({ NodeHandle, NodeTemplateHandle, 0 });
 
 		return NodeHandle;
 	}
 
-	uint16 FDecoratorWriter::GetNodeUID(const FNodeHandle NodeHandle) const
-	{
-		const FNodeMapping* NodeMapping = NodeMappings.FindByPredicate([NodeHandle](const FNodeMapping& It) { return It.NodeHandle == NodeHandle; });
-		return NodeMapping != nullptr ? NodeMapping->NodeUID : 0;
-	}
-
-	FNodeHandle FDecoratorWriter::GetNodeHandle(uint16 NodeUID) const
-	{
-		const FNodeMapping* NodeMapping = NodeMappings.FindByPredicate([NodeUID](const FNodeMapping& It) { return It.NodeUID == NodeUID; });
-		return NodeMapping != nullptr ? NodeMapping->NodeHandle : FNodeHandle();
-	}
-
 	void FDecoratorWriter::BeginNodeWriting()
 	{
 		ensure(!bIsNodeWriting);
@@ -122,14 +74,15 @@ namespace UE::AnimNext
 		}
 
 		bIsNodeWriting = true;
+		TrackedObjectsForGC.Reset();
 
 		// Serialize the node templates
 		TArray<FNodeTemplateRegistryHandle> NodeTemplateHandles;
 		NodeTemplateHandles.Reserve(NodeMappings.Num());
 
-		for (const FNodeMapping& NodeMapping : NodeMappings)
+		for (FNodeMapping& NodeMapping : NodeMappings)
 		{
-			NodeTemplateHandles.AddUnique(NodeMapping.NodeTemplateHandle);
+			NodeMapping.NodeTemplateIndex = NodeTemplateHandles.AddUnique(NodeMapping.NodeTemplateHandle);
 		}
 
 		uint32 NumNodeTemplates = NodeTemplateHandles.Num();
@@ -138,7 +91,7 @@ namespace UE::AnimNext
 		FNodeTemplateRegistry& NodeTemplateRegistry = FNodeTemplateRegistry::Get();
 		for (FNodeTemplateRegistryHandle NodeTemplateHandle : NodeTemplateHandles)
 		{
-			FNodeTemplate* NodeTemplate = NodeTemplateRegistry.Find(NodeTemplateHandle);
+			FNodeTemplate* NodeTemplate = NodeTemplateRegistry.FindMutable(NodeTemplateHandle);
 			NodeTemplate->Serialize(*this);
 		}
 
@@ -146,8 +99,12 @@ namespace UE::AnimNext
 		uint32 NumNodes = NodeMappings.Num();
 		*this << NumNodes;
 
-		uint32 SharedDataSize = GraphSharedDataSize;
-		*this << SharedDataSize;
+		// Serialize the node template indices that we'll use for each node
+		for (const FNodeMapping& NodeMapping : NodeMappings)
+		{
+			uint32 NodeTemplateIndex = NodeMapping.NodeTemplateIndex;
+			*this << NodeTemplateIndex;
+		}
 	}
 
 	void FDecoratorWriter::EndNodeWriting()
@@ -162,9 +119,20 @@ namespace UE::AnimNext
 		}
 
 		ensure(NumNodesWritten == NodeMappings.Num());
+
+		// Now write all the object reference we found, we'll use them for the GC callbacks at runtime to avoid travering the graph for the references
+		int32 NumTrackedObjectsForGC = TrackedObjectsForGC.Num();
+		*this << NumTrackedObjectsForGC;
+
+		for (UObject* Obj : TrackedObjectsForGC)
+		{
+			// Save out the fully qualified object name
+			FString SavedString(Obj->GetPathName());
+			*this << SavedString;
+		}
 	}
 
-	void FDecoratorWriter::WriteNode(const FNodeHandle NodeHandle, const TFunction<const TMap<FString, FString>& (uint32 DecoratorIndex)>& GetDecoratorProperties)
+	void FDecoratorWriter::WriteNode(const FNodeHandle NodeHandle, const TFunction<FString (uint32 DecoratorIndex, const FString& PropertyName)>& GetDecoratorProperty)
 	{
 		ensure(bIsNodeWriting);
 
@@ -194,7 +162,7 @@ namespace UE::AnimNext
 		// Populate our node description into a temporary buffer
 		alignas(16) uint8 Buffer[64 * 1024];	// Max node size
 
-		FNodeDescription* NodeDesc = new(Buffer) FNodeDescription(NodeMapping->NodeUID, NodeMapping->NodeTemplateHandle);
+		FNodeDescription* NodeDesc = new(Buffer) FNodeDescription(NodeMapping->NodeHandle.GetNodeID(), NodeMapping->NodeTemplateHandle);
 
 		// Populate our decorator properties
 		const uint32 NumDecorators = NodeTemplate->GetNumDecorators();
@@ -206,13 +174,14 @@ namespace UE::AnimNext
 
 			FAnimNextDecoratorSharedData* SharedData = DecoratorTemplates[DecoratorIndex].GetDecoratorDescription(*NodeDesc);
 
-			const TMap<FString, FString>& Properties = GetDecoratorProperties(DecoratorIndex);
-			Decorator->SaveDecoratorSharedData(*this, Properties, *SharedData);
-		}
+			// Curry our lambda with the decorator index
+			const auto GetDecoratorPropertyAt = [&GetDecoratorProperty, DecoratorIndex](const FString& PropertyName)
+			{
+				return GetDecoratorProperty(DecoratorIndex, PropertyName);
+			};
 
-		// Write our node size first so the reader knows how much to skip ahead to the next node
-		uint32 NodeSize = NodeMapping->NodeSize;
-		*this << NodeSize;
+			Decorator->SaveDecoratorSharedData(*this, GetDecoratorPropertyAt, *SharedData);
+		}
 
 		// Append it to our archive
 		NodeDesc->Serialize(*this);
@@ -229,5 +198,27 @@ namespace UE::AnimNext
 	{
 		return GraphSharedDataArchiveBuffer;
 	}
+
+	FArchive& FDecoratorWriter::operator<<(UObject*& Obj)
+	{
+		// Add our object for GC tracking
+		TrackedObjectsForGC.AddUnique(Obj);
+
+		// TODO: Instead of saving our object path, we could save the object index
+		//       Indices can be re-used if an object is referenced more than once (is it common?)
+		//       Allows all objects to be batch found/loaded before nodes are de-serialized
+		//       Requires seeking to read the objects first or an intermediary archive to write objects first
+
+		// Save out the fully qualified object name
+		FString SavedString(Obj->GetPathName());
+		*this << SavedString;
+
+		return *this;
+	}
+
+	FArchive& FDecoratorWriter::operator<<(FObjectPtr& Obj)
+	{
+		return FArchiveUObject::SerializeObjectPtr(*this, Obj);
+	}
 }
 #endif
diff --git a/Engine/Plugins/Experimental/AnimNext/Source/AnimNext/Private/DecoratorBase/ExecutionContext.cpp b/Engine/Plugins/Experimental/AnimNext/Source/AnimNext/Private/DecoratorBase/ExecutionContext.cpp
--- a/Engine/Plugins/Experimental/AnimNext/Source/AnimNext/Private/DecoratorBase/ExecutionContext.cpp
+++ b/Engine/Plugins/Experimental/AnimNext/Source/AnimNext/Private/DecoratorBase/ExecutionContext.cpp
@@ -72,7 +72,7 @@ namespace UE::AnimNext
 		// We need to allocate a new node instance
 		const FDecoratorTemplate* DecoratorDescs = NodeTemplate->GetDecorators();
 
-		const uint32 InstanceSize = NodeTemplate->GetInstanceSize();
+		const uint32 InstanceSize = NodeTemplate->GetNodeInstanceDataSize();
 		uint8* NodeInstanceBuffer = reinterpret_cast<uint8*>(FMemory::Malloc(InstanceSize, 16));
 		FNodeInstance* NodeInstance = new(NodeInstanceBuffer) FNodeInstance(ChildNodeHandle);
 
@@ -80,16 +80,12 @@ namespace UE::AnimNext
 		const FDecoratorTemplate* StartDesc = DecoratorDescs;
 		const FDecoratorTemplate* EndDesc = DecoratorDescs + NodeTemplate->GetNumDecorators();
 
-		const FDecoratorTemplate* FailedDecoratorDesc = nullptr;
 		for (const FDecoratorTemplate* DecoratorDesc = StartDesc; DecoratorDesc != EndDesc; ++DecoratorDesc)
 		{
 			const FDecorator* Decorator = GetDecorator(*DecoratorDesc);
-			ensure(Decorator != nullptr);
 			if (Decorator == nullptr)
 			{
-				// Failed to find the matching decorator, did it get unregistered or is the decorator descriptor corrupted?
-				FailedDecoratorDesc = DecoratorDesc;
-				break;
+				continue;	// Decorator hasn't been loaded or registered, skip it
 			}
 
 			const uint32 DecoratorIndex = DecoratorDesc - DecoratorDescs;
@@ -102,30 +98,6 @@ namespace UE::AnimNext
 			Decorator->ConstructDecoratorInstance(*this, DecoratorPtr, *SharedData, *InstanceData);
 		}
 
-		if (FailedDecoratorDesc != nullptr)
-		{
-			// We failed to construct our node instance, destroy it
-			// Start destruction with the decorator prior to the one that failed
-			StartDesc = FailedDecoratorDesc - 1;
-			EndDesc = DecoratorDescs - 1;
-			for (const FDecoratorTemplate* DecoratorDesc = StartDesc; DecoratorDesc != EndDesc; --DecoratorDesc)
-			{
-				const FDecorator* Decorator = GetDecorator(*DecoratorDesc);
-				const uint32 DecoratorIndex = DecoratorDesc - DecoratorDescs;
-
-				FWeakDecoratorPtr DecoratorPtr(NodeInstance, DecoratorIndex);
-
-				const FAnimNextDecoratorSharedData* SharedData = DecoratorDesc->GetDecoratorDescription(NodeDesc);
-				FDecoratorInstanceData* InstanceData = DecoratorDesc->GetDecoratorInstance(*NodeInstance);
-
-				Decorator->DestructDecoratorInstance(*this, DecoratorPtr, *SharedData, *InstanceData);
-			}
-
-			FMemory::Free(NodeInstance);
-
-			return FDecoratorPtr();
-		}
-
 		return FDecoratorPtr(NodeInstance, ChildDecoratorIndex);
 	}
 
@@ -159,8 +131,7 @@ namespace UE::AnimNext
 		const FDecoratorTemplate* EndDesc = DecoratorDescs - 1;
 		for (const FDecoratorTemplate* DecoratorDesc = StartDesc; DecoratorDesc != EndDesc; --DecoratorDesc)
 		{
-			const FDecorator* Decorator = GetDecorator(*DecoratorDesc);
-			if (ensure(Decorator != nullptr))
+			if (const FDecorator* Decorator = GetDecorator(*DecoratorDesc))
 			{
 				const uint32 DecoratorIndex = DecoratorDesc - DecoratorDescs;
 
@@ -207,10 +178,9 @@ namespace UE::AnimNext
 		for (const FDecoratorTemplate* DecoratorDesc = StartDesc; DecoratorDesc != EndDesc; --DecoratorDesc)
 		{
 			const FDecorator* Decorator = GetDecorator(*DecoratorDesc);
-			ensure(Decorator != nullptr);
 			if (Decorator == nullptr)
 			{
-				return false;	// Failed to find the matching decorator, did it get unregistered or is the decorator descriptor corrupted?
+				continue;	// Decorator hasn't been loaded or registered, skip it
 			}
 
 			if (const IDecoratorInterface* Interface = Decorator->GetDecoratorInterface(InterfaceUID))
@@ -265,10 +235,9 @@ namespace UE::AnimNext
 		for (const FDecoratorTemplate* DecoratorDesc = StartDesc; DecoratorDesc != EndDesc; --DecoratorDesc)
 		{
 			const FDecorator* Decorator = GetDecorator(*DecoratorDesc);
-			ensure(Decorator != nullptr);
 			if (Decorator == nullptr)
 			{
-				return false;	// Failed to find the matching decorator, did it get unregistered or is the decorator descriptor corrupted?
+				continue;	// Decorator hasn't been loaded or registered, skip it
 			}
 
 			if (const IDecoratorInterface* Interface = Decorator->GetDecoratorInterface(InterfaceUID))
@@ -297,13 +266,17 @@ namespace UE::AnimNext
 	const FNodeTemplate* FExecutionContext::GetNodeTemplate(const FNodeDescription& NodeDesc) const
 	{
 		check(NodeDesc.GetTemplateHandle().IsValid());
-		return FNodeTemplateRegistry::Get().Find(NodeDesc.GetTemplateHandle());
+
+		const FNodeTemplateRegistry& NodeTemplateRegistry = FNodeTemplateRegistry::Get();
+		return NodeTemplateRegistry.Find(NodeDesc.GetTemplateHandle());
 	}
 
 	const FDecorator* FExecutionContext::GetDecorator(const FDecoratorTemplate& Template) const
 	{
 		check(Template.GetRegistryHandle().IsValid());
-		return FDecoratorRegistry::Get().Find(Template.GetRegistryHandle());
+
+		const FDecoratorRegistry& DecoratorRegistry = FDecoratorRegistry::Get();
+		return DecoratorRegistry.Find(Template.GetRegistryHandle());
 	}
 
 	FExecutionContext* GetThreadExecutionContext()
diff --git a/Engine/Plugins/Experimental/AnimNext/Source/AnimNext/Private/DecoratorBase/NodeDescription.cpp b/Engine/Plugins/Experimental/AnimNext/Source/AnimNext/Private/DecoratorBase/NodeDescription.cpp
--- a/Engine/Plugins/Experimental/AnimNext/Source/AnimNext/Private/DecoratorBase/NodeDescription.cpp
+++ b/Engine/Plugins/Experimental/AnimNext/Source/AnimNext/Private/DecoratorBase/NodeDescription.cpp
@@ -11,11 +11,13 @@ namespace UE::AnimNext
 {
 	void FNodeDescription::Serialize(FArchive& Ar)
 	{
-		Ar << UID;
+		const FNodeTemplateRegistry& NodeTemplateRegistry = FNodeTemplateRegistry::Get();
+
+		Ar << NodeID;
 
 		if (Ar.IsSaving())
 		{
-			const FNodeTemplate* NodeTemplate = FNodeTemplateRegistry::Get().Find(TemplateHandle);
+			const FNodeTemplate* NodeTemplate = NodeTemplateRegistry.Find(TemplateHandle);
 
 			uint32 TemplateUID = NodeTemplate->GetUID();
 			Ar << TemplateUID;
@@ -25,7 +27,7 @@ namespace UE::AnimNext
 			uint32 TemplateUID = 0;
 			Ar << TemplateUID;
 
-			TemplateHandle = FNodeTemplateRegistry::Get().Find(TemplateUID);
+			TemplateHandle = NodeTemplateRegistry.Find(TemplateUID);
 		}
 		else
 		{
@@ -35,7 +37,7 @@ namespace UE::AnimNext
 		}
 
 		// Use our template to serialize our decorators
-		const FNodeTemplate* NodeTemplate = FNodeTemplateRegistry::Get().Find(TemplateHandle);
+		const FNodeTemplate* NodeTemplate = NodeTemplateRegistry.Find(TemplateHandle);
 
 		const uint32 NumDecorators = NodeTemplate->GetNumDecorators();
 		const FDecoratorTemplate* DecoratorTemplates = NodeTemplate->GetDecorators();
diff --git a/Engine/Plugins/Experimental/AnimNext/Source/AnimNext/Private/DecoratorBase/NodeTemplate.cpp b/Engine/Plugins/Experimental/AnimNext/Source/AnimNext/Private/DecoratorBase/NodeTemplate.cpp
--- a/Engine/Plugins/Experimental/AnimNext/Source/AnimNext/Private/DecoratorBase/NodeTemplate.cpp
+++ b/Engine/Plugins/Experimental/AnimNext/Source/AnimNext/Private/DecoratorBase/NodeTemplate.cpp
@@ -3,19 +3,82 @@
 #include "DecoratorBase/NodeTemplate.h"
 
 #include "Serialization/Archive.h"
+#include "DecoratorBase/Decorator.h"
+#include "DecoratorBase/DecoratorRegistry.h"
+#include "DecoratorBase/NodeDescription.h"
+#include "DecoratorBase/NodeInstance.h"
 
 namespace UE::AnimNext
 {
 	void FNodeTemplate::Serialize(FArchive& Ar)
 	{
 		Ar << UID;
-		Ar << InstanceSize;
 		Ar << NumDecorators;
 
+		const uint32 NumDecorators_ = GetNumDecorators();
 		FDecoratorTemplate* DecoratorTemplates = GetDecorators();
-		for (uint32 DecoratorIndex = 0; DecoratorIndex < GetNumDecorators(); ++DecoratorIndex)
+
+		for (uint32 DecoratorIndex = 0; DecoratorIndex < NumDecorators_; ++DecoratorIndex)
 		{
 			DecoratorTemplates[DecoratorIndex].Serialize(Ar);
 		}
+
+		if (Ar.IsLoading())
+		{
+			// When loading, make sure to recompute all runtime dependent values (e.g. sizes and offsets)
+			Finalize();
+		}
+	}
+
+	void FNodeTemplate::Finalize()
+	{
+		const FDecoratorRegistry& DecoratorRegistry = FDecoratorRegistry::Get();
+
+		const uint32 NumDecorators_ = GetNumDecorators();
+		FDecoratorTemplate* DecoratorTemplates = GetDecorators();
+
+		uint32 SharedDataOffset = sizeof(FNodeDescription);
+		uint32 InstanceDataOffset = sizeof(FNodeInstance);
+
+		for (uint32 DecoratorIndex = 0; DecoratorIndex < NumDecorators_; ++DecoratorIndex)
+		{
+			FDecoratorTemplate& DecoratorTemplate = DecoratorTemplates[DecoratorIndex];
+			const FDecoratorUID DecoratorUID = DecoratorTemplate.GetUID();
+
+			uint32 DecoratorSharedDataOffset = 0;
+			uint32 DecoratorInstanceDataOffset = 0;	// For instance data, 0 is an invalid offset since the data follows an instance of FNodeInstance
+
+			// Skip decorators that we can't find
+			// If a decorator isn't loaded and we attempt to run the graph, it will be a no-op entry
+			if (const FDecorator* Decorator = DecoratorRegistry.Find(DecoratorUID))
+			{
+				const FDecoratorMemoryLayout MemoryLayout = Decorator->GetDecoratorMemoryDescription();
+
+				// Align our data
+				SharedDataOffset = Align(SharedDataOffset, MemoryLayout.SharedDataAlignment);
+				InstanceDataOffset = Align(InstanceDataOffset, MemoryLayout.InstanceDataAlignment);
+
+				// Save our decorator offsets
+				DecoratorSharedDataOffset = SharedDataOffset;
+				DecoratorInstanceDataOffset = InstanceDataOffset;
+
+				// Include our decorator
+				SharedDataOffset += MemoryLayout.SharedDataSize;
+				InstanceDataOffset += MemoryLayout.InstanceDataSize;
+			}
+
+			check(DecoratorSharedDataOffset <= MAX_uint16);
+			check(DecoratorInstanceDataOffset <= MAX_uint16);
+
+			// Update our decorator offsets
+			DecoratorTemplate.NodeSharedOffset = static_cast<uint16>(DecoratorSharedDataOffset);
+			DecoratorTemplate.NodeInstanceOffset = static_cast<uint16>(DecoratorInstanceDataOffset);
+		}
+
+		// Our size is the offset of the decorator that would follow afterwards
+		// If the size is too large, we'll end up truncating the offsets/size
+		// Set a value of 0 to be able to detect it later
+		NodeSharedDataSize = SharedDataOffset > MAX_uint16 ? 0 : static_cast<uint16>(SharedDataOffset);
+		NodeInstanceDataSize = InstanceDataOffset > MAX_uint16 ? 0 : static_cast<uint16>(InstanceDataOffset);
 	}
 }
diff --git a/Engine/Plugins/Experimental/AnimNext/Source/AnimNext/Private/DecoratorBase/NodeTemplateBuilder.cpp b/Engine/Plugins/Experimental/AnimNext/Source/AnimNext/Private/DecoratorBase/NodeTemplateBuilder.cpp
new file mode 100644
--- /dev/null
+++ b/Engine/Plugins/Experimental/AnimNext/Source/AnimNext/Private/DecoratorBase/NodeTemplateBuilder.cpp
@@ -0,0 +1,110 @@
+// Copyright Epic Games, Inc. All Rights Reserved.
+
+#include "DecoratorBase/NodeTemplateBuilder.h"
+
+#include "DecoratorBase/Decorator.h"
+#include "DecoratorBase/DecoratorRegistry.h"
+#include "DecoratorBase/NodeTemplate.h"
+
+namespace UE::AnimNext
+{
+	void FNodeTemplateBuilder::AddDecorator(FDecoratorUID DecoratorUID)
+	{
+		DecoratorUIDs.Add(DecoratorUID);
+	}
+
+	FNodeTemplate* FNodeTemplateBuilder::BuildNodeTemplate(TArray<uint8>& NodeTemplateBuffer) const
+	{
+		return BuildNodeTemplate(DecoratorUIDs, NodeTemplateBuffer);
+	}
+
+	FNodeTemplate* FNodeTemplateBuilder::BuildNodeTemplate(const TArray<FDecoratorUID>& InDecoratorUIDs, TArray<uint8>& NodeTemplateBuffer)
+	{
+		NodeTemplateBuffer.Reset();
+
+		const uint32 NodeTemplateUID = GetNodeTemplateUID(InDecoratorUIDs);
+
+		NodeTemplateBuffer.AddUninitialized(sizeof(FNodeTemplate));
+		new(NodeTemplateBuffer.GetData()) FNodeTemplate(NodeTemplateUID, InDecoratorUIDs.Num());
+
+		for (int32 DecoratorIndex = 0; DecoratorIndex < InDecoratorUIDs.Num(); ++DecoratorIndex)
+		{
+			AppendTemplateDecorator(InDecoratorUIDs, DecoratorIndex, NodeTemplateBuffer);
+		}
+
+		// Grab pointer after everything is setup, we could re-alloc
+		FNodeTemplate* NodeTemplate = reinterpret_cast<FNodeTemplate*>(NodeTemplateBuffer.GetData());
+
+		// Perform all our finalizing work
+		NodeTemplate->Finalize();
+
+		return NodeTemplate;
+	}
+
+	void FNodeTemplateBuilder::Reset()
+	{
+		DecoratorUIDs.Reset();
+	}
+
+	uint32 FNodeTemplateBuilder::GetNodeTemplateUID(const TArray<FDecoratorUID>& InDecoratorUIDs)
+	{
+		uint32 NodeTemplateUID = 0;
+
+		for (const FDecoratorUID& DecoratorUID : InDecoratorUIDs)
+		{
+			NodeTemplateUID = HashCombineFast(NodeTemplateUID, DecoratorUID.GetUID());
+		}
+
+		return NodeTemplateUID;
+	}
+
+	void FNodeTemplateBuilder::AppendTemplateDecorator(
+		const TArray<FDecoratorUID>& InDecoratorUIDs, int32 DecoratorIndex,
+		TArray<uint8>& NodeTemplateBuffer)
+	{
+		const FDecoratorRegistry& DecoratorRegistry = FDecoratorRegistry::Get();
+
+		const FDecoratorUID DecoratorUID = InDecoratorUIDs[DecoratorIndex];
+		const FDecoratorRegistryHandle DecoratorHandle = DecoratorRegistry.FindHandle(DecoratorUID);
+		const FDecorator* Decorator = DecoratorRegistry.Find(DecoratorHandle);
+		const EDecoratorMode DecoratorMode = Decorator->GetDecoratorMode();
+
+		uint32 AdditiveIndexOrNumAdditive;
+		if (DecoratorMode == EDecoratorMode::Base)
+		{
+			// Find out how many additive decorators we have
+			AdditiveIndexOrNumAdditive = 0;
+			for (int32 Index = DecoratorIndex + 1; Index < InDecoratorUIDs.Num(); ++Index)	// Skip ourself
+			{
+				const FDecorator* ChildDecorator = DecoratorRegistry.Find(InDecoratorUIDs[Index]);
+				if (ChildDecorator->GetDecoratorMode() == EDecoratorMode::Base)
+				{
+					break;	// Found another base decorator, we are done
+				}
+
+				// We are additive
+				AdditiveIndexOrNumAdditive++;
+			}
+		}
+		else
+		{
+			// Find out our additive index
+			AdditiveIndexOrNumAdditive = 1;	// Skip ourself
+			for (int32 Index = DecoratorIndex - 1; Index >= 0; --Index)
+			{
+				const FDecorator* ParentDecorator = DecoratorRegistry.Find(InDecoratorUIDs[Index]);
+				if (ParentDecorator->GetDecoratorMode() == EDecoratorMode::Base)
+				{
+					break;	// Found our base decorator, we are done
+				}
+
+				// We are additive
+				AdditiveIndexOrNumAdditive++;
+			}
+		}
+
+		// Append our decorator template
+		const int32 BufferIndex = NodeTemplateBuffer.AddUninitialized(sizeof(FDecoratorTemplate));
+		new(&NodeTemplateBuffer[BufferIndex]) FDecoratorTemplate(DecoratorUID, DecoratorHandle, DecoratorMode, AdditiveIndexOrNumAdditive);
+	}
+}
diff --git a/Engine/Plugins/Experimental/AnimNext/Source/AnimNext/Private/DecoratorBase/NodeTemplateRegistry.cpp b/Engine/Plugins/Experimental/AnimNext/Source/AnimNext/Private/DecoratorBase/NodeTemplateRegistry.cpp
--- a/Engine/Plugins/Experimental/AnimNext/Source/AnimNext/Private/DecoratorBase/NodeTemplateRegistry.cpp
+++ b/Engine/Plugins/Experimental/AnimNext/Source/AnimNext/Private/DecoratorBase/NodeTemplateRegistry.cpp
@@ -81,7 +81,7 @@ namespace UE::AnimNext
 		TemplateUIDToHandleMap.Remove(TemplateUID);
 	}
 
-	FNodeTemplate* FNodeTemplateRegistry::Find(FNodeTemplateRegistryHandle TemplateHandle)
+	FNodeTemplate* FNodeTemplateRegistry::FindMutable(FNodeTemplateRegistryHandle TemplateHandle)
 	{
 		if (!TemplateHandle.IsValid())
 		{
@@ -105,10 +105,4 @@ namespace UE::AnimNext
 	{
 		return TemplateUIDToHandleMap.Num();
 	}
-
-	void FNodeTemplateRegistry::Clear()
-	{
-		TemplateBuffer.Empty(0);
-		TemplateUIDToHandleMap.Empty(0);
-	}
 }
diff --git a/Engine/Plugins/Experimental/AnimNext/Source/AnimNext/Private/DecoratorInterfaces/IUpdate.cpp b/Engine/Plugins/Experimental/AnimNext/Source/AnimNext/Private/DecoratorInterfaces/IUpdate.cpp
--- a/Engine/Plugins/Experimental/AnimNext/Source/AnimNext/Private/DecoratorInterfaces/IUpdate.cpp
+++ b/Engine/Plugins/Experimental/AnimNext/Source/AnimNext/Private/DecoratorInterfaces/IUpdate.cpp
@@ -7,6 +7,11 @@
 
 namespace UE::AnimNext
 {
+	FUpdateTraversalContext::FUpdateTraversalContext(float InDeltaTime)
+		: DeltaTime(InDeltaTime)
+	{
+	}
+
 	void IUpdate::PreUpdate(FExecutionContext& Context, const TDecoratorBinding<IUpdate>& Binding) const
 	{
 		TDecoratorBinding<IUpdate> SuperBinding;
@@ -37,7 +42,7 @@ namespace UE::AnimNext
 		EUpdateStep			DesiredStep = EUpdateStep::PreUpdate;
 	};
 
-	void UpdateGraph(FExecutionContext& Context, FWeakDecoratorPtr GraphRootPtr)
+	void UpdateGraph(FExecutionContext& Context, FWeakDecoratorPtr GraphRootPtr, float DeltaTime)
 	{
 		if (!GraphRootPtr.IsValid())
 		{
@@ -52,7 +57,7 @@ namespace UE::AnimNext
 		FChildrenArray Children;
 		Children.Reserve(64);
 
-		FUpdateTraversalContext TraversalContext;
+		FUpdateTraversalContext TraversalContext(DeltaTime);
 
 		FScopedTraversalContext ScopedTraversalContext(Context, TraversalContext);
 
diff --git a/Engine/Plugins/Experimental/AnimNext/Source/AnimNext/Private/Decorators/BlendTwoWay.cpp b/Engine/Plugins/Experimental/AnimNext/Source/AnimNext/Private/Decorators/BlendTwoWay.cpp
--- a/Engine/Plugins/Experimental/AnimNext/Source/AnimNext/Private/Decorators/BlendTwoWay.cpp
+++ b/Engine/Plugins/Experimental/AnimNext/Source/AnimNext/Private/Decorators/BlendTwoWay.cpp
@@ -6,9 +6,12 @@
 
 namespace UE::AnimNext
 {
+	AUTO_REGISTER_ANIM_DECORATOR(FBlendTwoWayDecorator)
+
 	DEFINE_ANIM_DECORATOR_BEGIN(FBlendTwoWayDecorator)
 		DEFINE_ANIM_DECORATOR_IMPLEMENTS_INTERFACE(IEvaluate)
 		DEFINE_ANIM_DECORATOR_IMPLEMENTS_INTERFACE(IUpdate)
+		DEFINE_ANIM_DECORATOR_IMPLEMENTS_INTERFACE(IHierarchy)
 	DEFINE_ANIM_DECORATOR_END(FBlendTwoWayDecorator)
 
 	void FBlendTwoWayDecorator::PostEvaluate(FExecutionContext& Context, const TDecoratorBinding<IEvaluate>& Binding) const
@@ -32,31 +35,31 @@ namespace UE::AnimNext
 
 		if (SharedData->BlendWeight < 1.0)
 		{
-			if (!InstanceData->Children[0].IsValid())
+			if (!InstanceData->ChildA.IsValid())
 			{
 				// We need to blend a child that isn't instanced yet, allocate it
-				InstanceData->Children[0] = Context.AllocateNodeInstance(Binding, SharedData->Children[0]);
+				InstanceData->ChildA = Context.AllocateNodeInstance(Binding, SharedData->ChildA);
 			}
 
 			if (SharedData->BlendWeight == 0.0)
 			{
 				// We no longer need this child, release it
-				InstanceData->Children[1].Reset();
+				InstanceData->ChildB.Reset();
 			}
 		}
 
 		if (SharedData->BlendWeight > 0.0)
 		{
-			if (!InstanceData->Children[1].IsValid())
+			if (!InstanceData->ChildB.IsValid())
 			{
 				// We need to blend a child that isn't instanced yet, allocate it
-				InstanceData->Children[1] = Context.AllocateNodeInstance(Binding, SharedData->Children[1]);
+				InstanceData->ChildB = Context.AllocateNodeInstance(Binding, SharedData->ChildB);
 			}
 
 			if (SharedData->BlendWeight == 1.0)
 			{
 				// We no longer need this child, release it
-				InstanceData->Children[0].Reset();
+				InstanceData->ChildA.Reset();
 			}
 		}
 	}
@@ -66,7 +69,7 @@ namespace UE::AnimNext
 		const FInstanceData* InstanceData = Binding.GetInstanceData<FInstanceData>();
 
 		// Add the two child handles, even if they are empty
-		Children.Add(InstanceData->Children[0]);
-		Children.Add(InstanceData->Children[1]);
+		Children.Add(InstanceData->ChildA);
+		Children.Add(InstanceData->ChildB);
 	}
 }
diff --git a/Engine/Plugins/Experimental/AnimNext/Source/AnimNext/Private/Decorators/BlendTwoWay.h b/Engine/Plugins/Experimental/AnimNext/Source/AnimNext/Private/Decorators/BlendTwoWay.h
--- a/Engine/Plugins/Experimental/AnimNext/Source/AnimNext/Private/Decorators/BlendTwoWay.h
+++ b/Engine/Plugins/Experimental/AnimNext/Source/AnimNext/Private/Decorators/BlendTwoWay.h
@@ -16,10 +16,16 @@ struct FAnimNextBlendTwoWayDecoratorSharedData : public FAnimNextDecoratorShared
 {
 	GENERATED_BODY()
 
-	UPROPERTY()
-	FAnimNextDecoratorHandle Children[2];
+	/** First output to be blended (full weight is 0.0). */
+	UPROPERTY(meta = (Input))
+	FAnimNextDecoratorHandle ChildA;
 
-	UPROPERTY()
+	/** Second output to be blended (full weight is 1.0). */
+	UPROPERTY(meta = (Input))
+	FAnimNextDecoratorHandle ChildB;
+
+	/** How much to blend our two children: 0.0 is fully child A while 1.0 is fully child B. */
+	UPROPERTY(meta = (Input))
 	double BlendWeight = 0.0;
 };
 
@@ -38,7 +44,8 @@ namespace UE::AnimNext
 
 		struct FInstanceData : FDecorator::FInstanceData
 		{
-			FDecoratorPtr Children[2];
+			FDecoratorPtr ChildA;
+			FDecoratorPtr ChildB;
 		};
 
 		// IEvaluate impl
diff --git a/Engine/Plugins/Experimental/AnimNext/Source/AnimNext/Private/Decorators/SequencePlayer.cpp b/Engine/Plugins/Experimental/AnimNext/Source/AnimNext/Private/Decorators/SequencePlayer.cpp
--- a/Engine/Plugins/Experimental/AnimNext/Source/AnimNext/Private/Decorators/SequencePlayer.cpp
+++ b/Engine/Plugins/Experimental/AnimNext/Source/AnimNext/Private/Decorators/SequencePlayer.cpp
@@ -3,20 +3,64 @@
 #include "Decorators/SequencePlayer.h"
 
 #include "DecoratorBase/ExecutionContext.h"
+#include "Param/ParamStack.h"
+#include "Graph/AnimNext_LODPose.h"
+#include "DecompressionTools.h"
 
 namespace UE::AnimNext
 {
+	AUTO_REGISTER_ANIM_DECORATOR(FSequencePlayerDecorator)
+
 	DEFINE_ANIM_DECORATOR_BEGIN(FSequencePlayerDecorator)
 		DEFINE_ANIM_DECORATOR_IMPLEMENTS_INTERFACE(IEvaluate)
 		DEFINE_ANIM_DECORATOR_IMPLEMENTS_INTERFACE(IUpdate)
 		DEFINE_ANIM_DECORATOR_IMPLEMENTS_INTERFACE(ITimeline)
 	DEFINE_ANIM_DECORATOR_END(FSequencePlayerDecorator)
 
+	void FSequencePlayerDecorator::FInstanceData::Construct(FExecutionContext& Context, FWeakDecoratorPtr DecoratorPtr, const FAnimNextSequencePlayerDecoratorSharedData& SharedData)
+	{
+		if (SharedData.AnimSequence != nullptr)
+		{
+			const float SequenceLength = SharedData.AnimSequence->GetPlayLength();
+			InternalTimeAccumulator = FMath::Clamp(SharedData.StartPosition, 0.f, SequenceLength);
+			PrevInternalTimeAccumulator = InternalTimeAccumulator;
+		}
+	}
+
 	void FSequencePlayerDecorator::PreEvaluate(FExecutionContext& Context, const TDecoratorBinding<IEvaluate>& Binding) const
 	{
+		const FSharedData* SharedData = Binding.GetSharedData<FSharedData>();
+		const FInstanceData* InstanceData = Binding.GetInstanceData<FInstanceData>();
+
 		// TODO: Sample pose
 		// FPose foo = ...
 		// Context.PushPose(foo);
+
+		FParamStack& ParamStack = FParamStack::Get();
+
+		FDeltaTimeRecord DeltaTimeRecord;
+		DeltaTimeRecord.Set(InstanceData->PrevInternalTimeAccumulator, InstanceData->InternalTimeAccumulator);
+
+		const FAnimExtractContext ExtractionContext(static_cast<double>(InstanceData->InternalTimeAccumulator)
+			, false /*Output.AnimInstanceProxy->ShouldExtractRootMotion()*/
+			, DeltaTimeRecord
+			, SharedData->bLoop);
+
+		const FAnimNextGraphReferencePose& GraphReferencePose = ParamStack.GetParam<FAnimNextGraphReferencePose>("GraphReferencePose");
+		const int32 GraphLODLevel = ParamStack.GetParam<int32>("GraphLODLevel");
+		const bool bGraphExpectsAdditive = ParamStack.GetParam<bool>("GraphExpectsAdditive");
+
+		// HACK: Write to our output directly, only works if we have a single node in the graph that outputs a pose
+		FAnimNextGraphLODPose& ResultPose = ParamStack.GetMutableParam<FAnimNextGraphLODPose>("ResultPose");
+
+		ResultPose.LODPose.PrepareForLOD(*GraphReferencePose.ReferencePose, GraphLODLevel, true, bGraphExpectsAdditive);
+
+		// Note : calling FDecompressionTools instead of UAnimSequence / UAnimSequenceBase, as I can not modify the engine code (so I extracted the function)
+		// TODO : Revisit this, so we can plug other sequence types
+		if (SharedData->AnimSequence != nullptr)
+		{
+			FDecompressionTools::GetAnimationPose(SharedData->AnimSequence, ResultPose.LODPose, ExtractionContext);
+		}
 	}
 
 	double FSequencePlayerDecorator::GetPlayRate(FExecutionContext& Context, const TDecoratorBinding<ITimeline>& Binding) const
@@ -30,17 +74,30 @@ namespace UE::AnimNext
 		const FSharedData* SharedData = Binding.GetSharedData<FSharedData>();
 		FInstanceData* InstanceData = Binding.GetInstanceData<FInstanceData>();
 
-		TDecoratorBinding<ITimeline> TimelineDecorator;
-		Context.GetInterface(Binding, TimelineDecorator);
+		if (SharedData->AnimSequence != nullptr)
+		{
+			FUpdateTraversalContext& TraversalContext = Context.GetTraversalContext<FUpdateTraversalContext>();
+
+			TDecoratorBinding<ITimeline> TimelineDecorator;
+			Context.GetInterface(Binding, TimelineDecorator);
+
+			const float DeltaTime = TraversalContext.GetDeltaTime();
+			const float PlayRate = (float)TimelineDecorator.GetPlayRate(Context);
 
-		FUpdateTraversalContext& TraversalContext = Context.GetTraversalContext<FUpdateTraversalContext>();
+			const float EffectiveDelta = FMath::IsNearlyZero(DeltaTime) || FMath::IsNearlyZero(PlayRate) ? 0.f : DeltaTime * PlayRate;
 
-		double DeltaTime = TraversalContext.GetDeltaTime();
+			const float SequenceLength = SharedData->AnimSequence->GetPlayLength();
+			float CurrentTime = SharedData->bLoop
+				? FMath::Fmod(InstanceData->InternalTimeAccumulator + EffectiveDelta, SequenceLength)
+				: FMath::Clamp(InstanceData->InternalTimeAccumulator + EffectiveDelta, 0.f, SequenceLength);
 
-		// Combine the execution play rate along with the timeline play rate, possibly overriden by another decorator
-		double PlayRate = TraversalContext.GetPlayRate() * TimelineDecorator.GetPlayRate(Context);
+			if (SharedData->bLoop && CurrentTime < 0.f)
+			{
+				CurrentTime += SequenceLength;
+			}
 
-		// TODO: need an actual impl, for illustrative purposes
-		InstanceData->CurrentTime += DeltaTime * PlayRate;
+			InstanceData->PrevInternalTimeAccumulator = InstanceData->InternalTimeAccumulator;
+			InstanceData->InternalTimeAccumulator = CurrentTime;
+		}
 	}
 }
diff --git a/Engine/Plugins/Experimental/AnimNext/Source/AnimNext/Private/Decorators/SequencePlayer.h b/Engine/Plugins/Experimental/AnimNext/Source/AnimNext/Private/Decorators/SequencePlayer.h
--- a/Engine/Plugins/Experimental/AnimNext/Source/AnimNext/Private/Decorators/SequencePlayer.h
+++ b/Engine/Plugins/Experimental/AnimNext/Source/AnimNext/Private/Decorators/SequencePlayer.h
@@ -3,26 +3,34 @@
 #pragma once
 
 #include "CoreMinimal.h"
-#include "UObject/ObjectPtr.h"
 #include "DecoratorBase/Decorator.h"
 #include "DecoratorInterfaces/IEvaluate.h"
 #include "DecoratorInterfaces/ITimeline.h"
 #include "DecoratorInterfaces/IUpdate.h"
+#include "Animation/AnimSequence.h"
 
 #include "SequencePlayer.generated.h"
 
-class UAnimSequence;
-
 USTRUCT()
 struct FAnimNextSequencePlayerDecoratorSharedData : public FAnimNextDecoratorSharedData
 {
 	GENERATED_BODY()
 
-	UPROPERTY()
-	TObjectPtr<UAnimSequence> AnimSeq;
+	/** The sequence to play. */
+	UPROPERTY(meta = (Input))
+	TObjectPtr<UAnimSequence> AnimSequence;
+
+	/** The play rate multiplier at which this sequence plays. */
+	UPROPERTY(meta = (Input))
+	float PlayRate = 1.0f;
 
-	UPROPERTY()
-	double PlayRate = 1.0;
+	/** The time at which we should start playing this sequence. */
+	UPROPERTY(meta = (Input))
+	float StartPosition = 0.0f;
+
+	/** Whether or not this sequence playback will loop. */
+	UPROPERTY(meta = (Input))
+	bool bLoop = 0.0f;
 };
 
 namespace UE::AnimNext
@@ -40,7 +48,10 @@ namespace UE::AnimNext
 
 		struct FInstanceData : FDecorator::FInstanceData
 		{
-			double CurrentTime = 0.0;
+			float InternalTimeAccumulator = 0.0f;
+			float PrevInternalTimeAccumulator = 0.0f;
+
+			void Construct(FExecutionContext& Context, FWeakDecoratorPtr DecoratorPtr, const FAnimNextSequencePlayerDecoratorSharedData& SharedData);
 		};
 
 		// IEvaluate impl
diff --git a/Engine/Plugins/Experimental/AnimNext/Source/AnimNext/Private/Graph/AnimGraph/AnimNode_AnimNextGraph.cpp b/Engine/Plugins/Experimental/AnimNext/Source/AnimNext/Private/Graph/AnimGraph/AnimNode_AnimNextGraph.cpp
--- a/Engine/Plugins/Experimental/AnimNext/Source/AnimNext/Private/Graph/AnimGraph/AnimNode_AnimNextGraph.cpp
+++ b/Engine/Plugins/Experimental/AnimNext/Source/AnimNext/Private/Graph/AnimGraph/AnimNode_AnimNextGraph.cpp
@@ -11,8 +11,7 @@
 #include "GenerationTools.h"
 #include "ReferencePose.h"
 #include "Graph/AnimNext_LODPose.h"
-#include "Animation/AnimSequence.h" // TEST
-#include "Graph/RigUnit_AnimNextAnimSequence.h" // TEST
+#include "Graph/AnimNextExecuteContext.h"
 #include "Engine/SkeletalMesh.h"
 #include "BoneContainer.h"
 #include "Param/ParamStack.h"
@@ -28,7 +27,14 @@ FAnimNode_AnimNextGraph::FAnimNode_AnimNextGraph()
 	, AnimNextGraph(nullptr)
 	, LODThreshold(INDEX_NONE)
 {
+}
 
+FAnimNode_AnimNextGraph::~FAnimNode_AnimNextGraph()
+{
+	if (GraphInstancePtr.IsValid())
+	{
+		AnimNextGraph->ReleaseInstance(GraphInstancePtr);
+	}
 }
 
 void FAnimNode_AnimNextGraph::OnInitializeAnimInstance(const FAnimInstanceProxy* InProxy, const UAnimInstance* InAnimInstance)
@@ -51,8 +57,6 @@ void FAnimNode_AnimNextGraph::Update_AnyThread(const FAnimationUpdateContext& Co
 {
 	DECLARE_SCOPE_HIERARCHICAL_COUNTER_FUNC()
 
-	GraphDeltaTime += Context.GetDeltaTime();
-
 	SourceLink.Update(Context);
 
 	if (IsLODEnabled(Context.AnimInstanceProxy))
@@ -60,6 +64,17 @@ void FAnimNode_AnimNextGraph::Update_AnyThread(const FAnimationUpdateContext& Co
 		GetEvaluateGraphExposedInputs().Execute(Context);
 
 		PropagateInputProperties(Context.AnimInstanceProxy->GetAnimInstanceObject());
+
+		// Populate our param stack since our instance data might need it during construction
+		const int32 LODLevel = Context.AnimInstanceProxy->GetLODLevel();
+
+		UE::AnimNext::FParamStack::Get().PushValues(
+			"GraphLODLevel", LODLevel
+		);
+
+		UE::AnimNext::FContext AnimNextContext(Context.GetDeltaTime());
+
+		AnimNextGraph->Run(AnimNextContext, GraphInstancePtr, EAnimNextGraphSimulationSteps::Update);
 	}
 
 	FAnimNode_CustomProperty::Update_AnyThread(Context);
@@ -73,6 +88,20 @@ void FAnimNode_AnimNextGraph::Initialize_AnyThread(const FAnimationInitializeCon
 
 	SourceLink.Initialize(Context);
 
+	if (!GraphInstancePtr.IsValid() && AnimNextGraph)
+	{
+		// If we don't have an instance yet, create one
+
+		// Populate our param stack since our instance data might need it during construction
+		const int32 LODLevel = Context.AnimInstanceProxy->GetLODLevel();
+
+		UE::AnimNext::FParamStack::Get().PushValues(
+			"GraphLODLevel", LODLevel
+		);
+
+		GraphInstancePtr = AnimNextGraph->AllocateInstance();
+	}
+
 	FAnimNode_CustomProperty::Initialize_AnyThread(Context);
 }
 
@@ -110,29 +139,23 @@ void FAnimNode_AnimNextGraph::Evaluate_AnyThread(FPoseContext & Output)
 
 	const int32 LODLevel = Output.AnimInstanceProxy->GetLODLevel();
 	
-	// TODO : Using AnimInstanceProxy->GetDeltaSeconds() makes the preview to advance  multiple times when debug options are activated (i.e. ShowUncompressedAnim)
-	// See where to get / how to calculate the correct delta value
-	UE::AnimNext::FContext Context(GraphDeltaTime);
-	GraphDeltaTime = 0.f; // Reset for the case that we receive multiple calls to Evaluate (debug options)
+	UE::AnimNext::FContext Context(0.0f);
 
 	FAnimNextGraphReferencePose GraphReferencePose(&RefPose);
-	FAnimNextGraph_AnimSequence GraphTestSequence(TestSequence);  // TEST
 
 	FAnimNextGraphLODPose GraphSourceLODPose(FLODPose(RefPose, LODLevel, false, Output.ExpectsAdditivePose()));
 	FAnimNextGraphLODPose ResultPose(FLODPose(RefPose, LODLevel, true, Output.ExpectsAdditivePose()));
 	FGenerationTools::RemapPose(LODLevel, SourcePose, RefPose, GraphSourceLODPose.LODPose);
 
 	FParamStack::FPushedLayerHandle LayerHandle = Context.GetMutableParamStack().PushValues(
-		"AnimSequencePlayerState", SequencePlayerState,
 		"GraphReferencePose", GraphReferencePose,
 		"ResultPose", ResultPose,
 		"GraphLODLevel", LODLevel,
 		"GraphExpectsAdditive", Output.ExpectsAdditivePose(),
-		"SourcePose", GraphSourceLODPose,						// TODO : Pass this as a external variable maybe ? When we have support for variables in the rigvm graph
-		"TestSequence", GraphTestSequence						// TEST anim decompression - Anim Sequence to decompress);
+		"SourcePose", GraphSourceLODPose						// TODO : Pass this as a external variable maybe ? When we have support for variables in the rigvm graph
 	);
 
-	AnimNextGraph->Run(Context);
+	AnimNextGraph->Run(Context, GraphInstancePtr, EAnimNextGraphSimulationSteps::Evaluate);
 
 	FGenerationTools::RemapPose(LODLevel, RefPose, ResultPose.LODPose, Output);
 
diff --git a/Engine/Plugins/Experimental/AnimNext/Source/AnimNext/Private/Graph/AnimNextGraph.cpp b/Engine/Plugins/Experimental/AnimNext/Source/AnimNext/Private/Graph/AnimNextGraph.cpp
--- a/Engine/Plugins/Experimental/AnimNext/Source/AnimNext/Private/Graph/AnimNextGraph.cpp
+++ b/Engine/Plugins/Experimental/AnimNext/Source/AnimNext/Private/Graph/AnimNextGraph.cpp
@@ -3,10 +3,13 @@
 #include "Graph/AnimNextGraph.h"
 #include "RigVMCore/RigVMMemoryStorage.h"
 #include "RigVMCore/RigVMExecuteContext.h"
-#include "Graph/RigUnit_AnimNextBeginExecution.h"
+#include "Graph/RigUnit_AnimNextGraphRoot.h"
 #include "UObject/ObjectSaveContext.h"
 #include "UObject/Package.h"
 #include "Param/ParamTypeHandle.h"
+#include "DecoratorBase/DecoratorReader.h"
+#include "DecoratorBase/ExecutionContext.h"
+#include "Serialization/MemoryReader.h"
 
 #include UE_INLINE_GENERATED_CPP_BY_NAME(AnimNextGraph)
 
@@ -16,15 +19,42 @@ const FName EntryPointName("GetData");
 const FName ResultName("Result");
 }
 
-void UAnimNextGraph::Run(const UE::AnimNext::FContext& Context) const
+UE::AnimNext::FDecoratorPtr UAnimNextGraph::AllocateInstance() const
 {
-	if(RigVM)
+	if (!ResolvedRootDecoratorHandle.IsValid())
+	{
+		return UE::AnimNext::FDecoratorPtr();
+	}
+
+	UE::AnimNext::FExecutionContext Context(SharedDataBuffer);
+
+	return Context.AllocateNodeInstance(UE::AnimNext::FWeakDecoratorPtr(), ResolvedRootDecoratorHandle);
+}
+
+void UAnimNextGraph::ReleaseInstance(UE::AnimNext::FDecoratorPtr& GraphInstancePtr) const
+{
+	if (!GraphInstancePtr.IsValid())
+	{
+		return;
+	}
+
+	// We need an execution context for this graph to be active when we delete the graph instance
+	UE::AnimNext::FExecutionContext Context(SharedDataBuffer);
+
+	GraphInstancePtr.Reset();
+}
+
+void UAnimNextGraph::Run(const UE::AnimNext::FContext& Context, UE::AnimNext::FWeakDecoratorPtr GraphInstancePtr, EAnimNextGraphSimulationSteps SimulationSteps) const
+{
+	if (RigVM && GraphInstancePtr.IsValid())
 	{
 		FRigVMExtendedExecuteContext RigVMExtendedExecuteContext;
 		FAnimNextExecuteContext& AnimNextContext = RigVMExtendedExecuteContext.GetPublicDataSafe<FAnimNextExecuteContext>();
 		AnimNextContext.SetContextData(Context);
+		AnimNextContext.InitializeWithGraph(SharedDataBuffer, GraphInstancePtr);
+		AnimNextContext.SetSimulationSteps(SimulationSteps);
 
-		RigVM->Execute(RigVMExtendedExecuteContext, TArray<URigVMMemoryStorage*>(), FRigUnit_AnimNextBeginExecution::EventName);
+		RigVM->Execute(RigVMExtendedExecuteContext, TArray<URigVMMemoryStorage*>(), FRigUnit_AnimNextGraphRoot::EventName);
 	}
 }
 
@@ -85,3 +115,77 @@ void UAnimNextGraph::GetPreloadDependencies(TArray<UObject*>& OutDeps)
 		}
 	}
 }
+
+void UAnimNextGraph::Serialize(FArchive& Ar)
+{
+	Super::Serialize(Ar);
+
+	if (Ar.IsLoading())
+	{
+		int32 SharedDataArchiveBufferSize = 0;
+		Ar << SharedDataArchiveBufferSize;
+
+#if !WITH_EDITORONLY_DATA
+		// When editor data isn't present, we don't persist the archive buffer as it is only needed on load
+		// to populate the graph shared data
+		TArray<uint8> SharedDataArchiveBuffer;
+#endif
+
+		SharedDataArchiveBuffer.SetNumUninitialized(SharedDataArchiveBufferSize);
+		Ar.Serialize(SharedDataArchiveBuffer.GetData(), SharedDataArchiveBufferSize);
+
+		LoadFromArchiveBuffer(SharedDataArchiveBuffer);
+	}
+	else if (Ar.IsSaving())
+	{
+#if WITH_EDITORONLY_DATA
+		// We only save the archive buffer, if code changes we'll be able to de-serialize from it when
+		// building the runtime buffer
+		// This allows us to have editor only/non-shipping only properties that are stripped out on load
+		int32 SharedDataArchiveBufferSize = SharedDataArchiveBuffer.Num();
+		Ar << SharedDataArchiveBufferSize;
+		Ar.Serialize(SharedDataArchiveBuffer.GetData(), SharedDataArchiveBufferSize);
+#endif
+	}
+	else
+	{
+		// Counting, etc
+		Ar << SharedDataBuffer;
+		Ar << TrackedObjectsForGC;
+
+#if WITH_EDITORONLY_DATA
+		Ar << SharedDataArchiveBuffer;
+#endif
+	}
+}
+
+void UAnimNextGraph::AddReferencedObjects(UObject* InThis, FReferenceCollector& Collector)
+{
+	UAnimNextGraph* This = CastChecked<UAnimNextGraph>(InThis);
+	Collector.AddStableReferenceArray(&This->TrackedObjectsForGC);
+
+	Super::AddReferencedObjects(InThis, Collector);
+}
+
+bool UAnimNextGraph::LoadFromArchiveBuffer(const TArray<uint8>& InSharedDataArchiveBuffer)
+{
+	using namespace UE::AnimNext;
+
+	// Reconstruct our graph shared data
+	FMemoryReader GraphSharedDataArchive(InSharedDataArchiveBuffer);
+	FDecoratorReader DecoratorReader(GraphSharedDataArchive);
+
+	const FDecoratorReader::EErrorState ErrorState = DecoratorReader.ReadGraph(SharedDataBuffer, TrackedObjectsForGC);
+	if (ErrorState == FDecoratorReader::EErrorState::None)
+	{
+		ResolvedRootDecoratorHandle = DecoratorReader.ResolveEntryPointHandle(RootDecoratorHandle);
+		return true;
+	}
+	else
+	{
+		SharedDataBuffer.Empty(0);
+		TrackedObjectsForGC.Empty(0);
+		ResolvedRootDecoratorHandle = FAnimNextDecoratorHandle();
+		return false;
+	}
+}
diff --git a/Engine/Plugins/Experimental/AnimNext/Source/AnimNext/Private/Graph/RigDecorator_AnimNextCppDecorator.cpp b/Engine/Plugins/Experimental/AnimNext/Source/AnimNext/Private/Graph/RigDecorator_AnimNextCppDecorator.cpp
new file mode 100644
--- /dev/null
+++ b/Engine/Plugins/Experimental/AnimNext/Source/AnimNext/Private/Graph/RigDecorator_AnimNextCppDecorator.cpp
@@ -0,0 +1,29 @@
+// Copyright Epic Games, Inc. All Rights Reserved.
+
+#include "Graph/RigDecorator_AnimNextCppDecorator.h"
+
+#if WITH_EDITOR
+#include "RigVMModel/RigVMController.h"
+#include "DecoratorBase/DecoratorRegistry.h"
+#endif
+
+#include UE_INLINE_GENERATED_CPP_BY_NAME(RigDecorator_AnimNextCppDecorator)
+
+#if WITH_EDITOR
+void FRigDecorator_AnimNextCppDecorator::GetProgrammaticPins(URigVMController* InController, int32 InParentPinIndex, FRigVMPinInfoArray& OutPinArray) const
+{
+	if (DecoratorSharedDataStruct == nullptr)
+	{
+		return;
+	}
+
+	FStructOnScope DefaultValueMemoryScope(DecoratorSharedDataStruct);
+
+	OutPinArray.AddPins(DecoratorSharedDataStruct, InController, ERigVMPinDirection::Invalid, InParentPinIndex, DefaultValueMemoryScope.GetStructMemory(), true);
+}
+
+const UE::AnimNext::FDecorator* FRigDecorator_AnimNextCppDecorator::GetDecorator() const
+{
+	return UE::AnimNext::FDecoratorRegistry::Get().Find(DecoratorSharedDataStruct);
+}
+#endif
diff --git a/Engine/Plugins/Experimental/AnimNext/Source/AnimNext/Private/Graph/RigUnit_AnimNextAnimSequence.cpp b/Engine/Plugins/Experimental/AnimNext/Source/AnimNext/Private/Graph/RigUnit_AnimNextAnimSequence.cpp
deleted file mode 100644
diff --git a/Engine/Plugins/Experimental/AnimNext/Source/AnimNext/Private/Graph/RigUnit_AnimNextAnimSequence.h b/Engine/Plugins/Experimental/AnimNext/Source/AnimNext/Private/Graph/RigUnit_AnimNextAnimSequence.h
deleted file mode 100644
diff --git a/Engine/Plugins/Experimental/AnimNext/Source/AnimNext/Private/Graph/RigUnit_AnimNextDecoratorStack.cpp b/Engine/Plugins/Experimental/AnimNext/Source/AnimNext/Private/Graph/RigUnit_AnimNextDecoratorStack.cpp
new file mode 100644
--- /dev/null
+++ b/Engine/Plugins/Experimental/AnimNext/Source/AnimNext/Private/Graph/RigUnit_AnimNextDecoratorStack.cpp
@@ -0,0 +1,11 @@
+﻿// Copyright Epic Games, Inc. All Rights Reserved.
+
+#include "Graph/RigUnit_AnimNextDecoratorStack.h"
+
+#include UE_INLINE_GENERATED_CPP_BY_NAME(RigUnit_AnimNextDecoratorStack)
+
+FRigUnit_AnimNextDecoratorStack_Execute()
+{
+	// This RigUnit should never execute, it is only used in the editor as a container
+	ensure(false);
+}
diff --git a/Engine/Plugins/Experimental/AnimNext/Source/AnimNext/Private/Graph/RigUnit_AnimNextGraphRoot.cpp b/Engine/Plugins/Experimental/AnimNext/Source/AnimNext/Private/Graph/RigUnit_AnimNextGraphRoot.cpp
new file mode 100644
--- /dev/null
+++ b/Engine/Plugins/Experimental/AnimNext/Source/AnimNext/Private/Graph/RigUnit_AnimNextGraphRoot.cpp
@@ -0,0 +1,41 @@
+﻿// Copyright Epic Games, Inc. All Rights Reserved.
+
+#include "Graph/RigUnit_AnimNextGraphRoot.h"
+#include "DecoratorBase/ExecutionContext.h"
+#include "DecoratorInterfaces/IUpdate.h"
+#include "DecoratorInterfaces/IEvaluate.h"
+#include "Context.h"
+
+#include UE_INLINE_GENERATED_CPP_BY_NAME(RigUnit_AnimNextGraphRoot)
+
+FName FRigUnit_AnimNextGraphRoot::EventName = TEXT("Execute");
+
+FRigUnit_AnimNextGraphRoot_Execute()
+{
+	using namespace UE::AnimNext;
+
+	const FAnimNextExecuteContext& VMExecuteContext = static_cast<const FAnimNextExecuteContext&>(ExecuteContext);
+
+	// Setup what we need to execute
+	FExecutionContext Context(VMExecuteContext.GetSharedDataBuffer());
+	FWeakDecoratorPtr GraphInstancePtr = VMExecuteContext.GetGraphInstancePtr();
+	const EAnimNextGraphSimulationSteps SimulationSteps = VMExecuteContext.GetSimulationSteps();
+
+	if (EnumHasAnyFlags(SimulationSteps, EAnimNextGraphSimulationSteps::Update))
+	{
+		const FContext& InterfaceContext = ExecuteContext.GetContext();
+
+		// Call pre/post update on our graph
+		UpdateGraph(Context, GraphInstancePtr, InterfaceContext.GetDeltaTime());
+	}
+
+	if (EnumHasAnyFlags(SimulationSteps, EAnimNextGraphSimulationSteps::Evaluate))
+	{
+		// Call pre/post evaluate on our graph
+		FPoseContainer PoseContainer;
+		EvaluateGraph(Context, GraphInstancePtr, EEvaluationFlags::All, PoseContainer);
+
+		// TODO: Write out our output pose/curves/attributes/etc
+		//ExecuteContext_.GetContext().SetResult<FAnimationGraphNodeResult>(Result);
+	}
+}
diff --git a/Engine/Plugins/Experimental/AnimNext/Source/AnimNext/Private/Module.cpp b/Engine/Plugins/Experimental/AnimNext/Source/AnimNext/Private/Module.cpp
--- a/Engine/Plugins/Experimental/AnimNext/Source/AnimNext/Private/Module.cpp
+++ b/Engine/Plugins/Experimental/AnimNext/Source/AnimNext/Private/Module.cpp
@@ -20,7 +20,8 @@ public:
 	{
 		static UClass* const AllowedAssetTypes[] =
 		{
-			UAnimSequence::StaticClass()
+			UAnimSequence::StaticClass(),
+			UScriptStruct::StaticClass(),
 		};
 
 		FRigVMRegistry::Get().RegisterObjectTypes(AllowedAssetTypes);
@@ -28,21 +29,14 @@ public:
 		FDataRegistry::Init();
 		FDecoratorRegistry::Init();
 		FNodeTemplateRegistry::Init();
-
-		EnginePreExitHandle = FCoreDelegates::OnEnginePreExit.AddLambda([]()
-		{
-			FDataRegistry::Destroy();
-			FDecoratorRegistry::Destroy();
-			FNodeTemplateRegistry::Destroy();
-		});
 	}
 
 	virtual void ShutdownModule() override
 	{
-		FCoreDelegates::OnEnginePreExit.Remove(EnginePreExitHandle);
+		FNodeTemplateRegistry::Destroy();
+		FDecoratorRegistry::Destroy();
+		FDataRegistry::Destroy();
 	}
-
-	FDelegateHandle EnginePreExitHandle;
 };
 
 }
diff --git a/Engine/Plugins/Experimental/AnimNext/Source/AnimNext/Public/Context.h b/Engine/Plugins/Experimental/AnimNext/Source/AnimNext/Public/Context.h
--- a/Engine/Plugins/Experimental/AnimNext/Source/AnimNext/Public/Context.h
+++ b/Engine/Plugins/Experimental/AnimNext/Source/AnimNext/Public/Context.h
@@ -20,10 +20,6 @@ struct ANIMNEXT_API FContext
 public:
 	FContext();
 	explicit FContext(float InDeltaTime);
-	FContext(const FContext& Other) = delete;
-	FContext& operator=(const FContext&) = delete;
-	FContext(FContext&& Other) = delete;
-	FContext& operator=(FContext&& Other) = delete;
 
 	// Access the parameter stack, cached on construction
 	FParamStack& GetMutableParamStack() const { return Stack; }
@@ -32,6 +28,9 @@ public:
 	float GetDeltaTime() const { return DeltaTime; }
 
 private:
+	FContext(const FContext& Other) = delete;
+	FContext& operator=(const FContext&) = delete;
+
 	FParamStack& Stack;
 	float DeltaTime = 0.0f;
 };
diff --git a/Engine/Plugins/Experimental/AnimNext/Source/AnimNext/Public/DecoratorBase/Decorator.h b/Engine/Plugins/Experimental/AnimNext/Source/AnimNext/Public/DecoratorBase/Decorator.h
--- a/Engine/Plugins/Experimental/AnimNext/Source/AnimNext/Public/DecoratorBase/Decorator.h
+++ b/Engine/Plugins/Experimental/AnimNext/Source/AnimNext/Public/DecoratorBase/Decorator.h
@@ -19,7 +19,7 @@ class FArchive;
 #define DECLARE_ANIM_DECORATOR(DecoratorName, DecoratorNameHash, SuperDecoratorName) \
 	using DecoratorSuper = SuperDecoratorName; \
 	/* FDecorator impl */ \
-	static constexpr UE::AnimNext::FDecoratorUID DecoratorUID = UE::AnimNext::FDecoratorUID(TEXT(#DecoratorName), DecoratorNameHash); \
+	static constexpr UE::AnimNext::FDecoratorUID DecoratorUID = UE::AnimNext::FDecoratorUID(DecoratorNameHash, TEXT(#DecoratorName)); \
 	virtual UE::AnimNext::FDecoratorUID GetDecoratorUID() const override { return DecoratorUID; } \
 	static const UE::AnimNext::FDecoratorMemoryLayout DecoratorMemoryDescription; \
 	virtual UE::AnimNext::FDecoratorMemoryLayout GetDecoratorMemoryDescription() const override { return DecoratorMemoryDescription; } \
@@ -33,7 +33,7 @@ class FArchive;
 #define DECLARE_ABSTRACT_ANIM_DECORATOR(DecoratorName, DecoratorNameHash, SuperDecoratorName) \
 	using DecoratorSuper = SuperDecoratorName; \
 	/* FDecorator impl */ \
-	static constexpr UE::AnimNext::FDecoratorUID DecoratorUID = UE::AnimNext::FDecoratorUID(TEXT(#DecoratorName), DecoratorNameHash); \
+	static constexpr UE::AnimNext::FDecoratorUID DecoratorUID = UE::AnimNext::FDecoratorUID(DecoratorNameHash, TEXT(#DecoratorName)); \
 	virtual UE::AnimNext::FDecoratorUID GetDecoratorUID() const override { return DecoratorUID; }
 
 // In the decorator cpp, these three macros implement the base functionality
@@ -152,7 +152,7 @@ namespace UE::AnimNext
 		// The globally unique UID for this decorator
 		// Derived types will have their own DecoratorUID member that hides/aliases/shadows this one
 		// @see DECLARE_ANIM_DECORATOR
-		static constexpr FDecoratorUID DecoratorUID = FDecoratorUID(TEXT("FDecorator"), 0x4fc735a2);
+		static constexpr FDecoratorUID DecoratorUID = FDecoratorUID(0x4fc735a2, TEXT("FDecorator"));
 
 		// Returns the globally unique UID for this decorator
 		virtual FDecoratorUID GetDecoratorUID() const { return DecoratorUID; };
@@ -205,7 +205,7 @@ namespace UE::AnimNext
 		// derived type using UE reflection.
 		// Decorators can override this function to control how editor only properties are coerced into the runtime shared data
 		// instance.
-		virtual void SaveDecoratorSharedData(FDecoratorWriter& Writer, const TMap<FString, FString>& Properties, FAnimNextDecoratorSharedData& OutSharedData) const;
+		virtual void SaveDecoratorSharedData(FDecoratorWriter& Writer, const TFunction<FString(const FString& PropertyName)>& GetDecoratorProperty, FAnimNextDecoratorSharedData& OutSharedData) const;
 #endif
 	};
 
diff --git a/Engine/Plugins/Experimental/AnimNext/Source/AnimNext/Public/DecoratorBase/DecoratorHandle.h b/Engine/Plugins/Experimental/AnimNext/Source/AnimNext/Public/DecoratorBase/DecoratorHandle.h
--- a/Engine/Plugins/Experimental/AnimNext/Source/AnimNext/Public/DecoratorBase/DecoratorHandle.h
+++ b/Engine/Plugins/Experimental/AnimNext/Source/AnimNext/Public/DecoratorBase/DecoratorHandle.h
@@ -12,43 +12,59 @@
  * Decorator Handle
  * A decorator handle represents a reference to a specific decorator instance in the shared/read-only portion
  * of a sub-graph. It points to a FNodeDescription when resolved.
- * @see FNodeDescription
+ * 
+ * Internally, it contains a node handle and a decorator index.
+ * Decorator handles can only be used within an AnimNext graph serialized with FDecoratorWriter.
+ * 
+ * @see FNodeDescription, FNodeHandle
  */
-USTRUCT()
-struct FAnimNextDecoratorHandle
+USTRUCT(BlueprintType)
+struct FAnimNextDecoratorHandle final
 {
 	GENERATED_BODY()
 
 	// Creates an invalid decorator handle
-	FAnimNextDecoratorHandle()
-		: PackedDecoratorIndexAndNodeSharedOffset(INVALID_SHARED_OFFSET_VALUE)
+	constexpr FAnimNextDecoratorHandle() noexcept
+		: PackedDecoratorIndexAndNodeHandle(UE::AnimNext::FNodeHandle::INVALID_NODE_HANDLE_RAW_VALUE)
 	{}
 
 	// Creates a decorator handle pointing to the first decorator of the specified node
-	explicit FAnimNextDecoratorHandle(UE::AnimNext::FNodeHandle NodeHandle)
-		: PackedDecoratorIndexAndNodeSharedOffset(NodeHandle.GetSharedOffset())
+	explicit constexpr FAnimNextDecoratorHandle(UE::AnimNext::FNodeHandle NodeHandle) noexcept
+		: PackedDecoratorIndexAndNodeHandle(NodeHandle.GetPackedValue() & NODE_HANDLE_MASK)
 	{}
 
 	// Creates a decorator handle pointing to the specified decorator on the specified node
-	FAnimNextDecoratorHandle(UE::AnimNext::FNodeHandle NodeHandle, uint32 DecoratorIndex_)
-		: PackedDecoratorIndexAndNodeSharedOffset(NodeHandle.GetSharedOffset() | (DecoratorIndex_ << 24))
-	{}
+	FAnimNextDecoratorHandle(UE::AnimNext::FNodeHandle NodeHandle, uint32 DecoratorIndex)
+		: PackedDecoratorIndexAndNodeHandle((NodeHandle.GetPackedValue() & NODE_HANDLE_MASK) | (DecoratorIndex << 24))
+	{
+		check(DecoratorIndex <= MAX_uint8);	// Make sure we don't truncate
+	}
 
 	// Returns true if this decorator handle is valid, false otherwise
-	bool IsValid() const { return (PackedDecoratorIndexAndNodeSharedOffset & SHARED_OFFSET_MASK) != INVALID_SHARED_OFFSET_VALUE; }
+	constexpr bool IsValid() const noexcept { return GetNodeHandle().IsValid(); }
 
 	// Returns the decorator index
-	uint32 GetDecoratorIndex() const { return PackedDecoratorIndexAndNodeSharedOffset >> DECORATOR_SHIFT; }
+	constexpr uint32 GetDecoratorIndex() const noexcept { return PackedDecoratorIndexAndNodeHandle >> DECORATOR_SHIFT; }
+
+	// Returns a handle to the node referenced
+	constexpr UE::AnimNext::FNodeHandle GetNodeHandle() const noexcept { return UE::AnimNext::FNodeHandle::FromPackedValue(PackedDecoratorIndexAndNodeHandle); }
 
-	// Returns a handle to the node in the shared data segment
-	UE::AnimNext::FNodeHandle GetNodeHandle() const { return UE::AnimNext::FNodeHandle(PackedDecoratorIndexAndNodeSharedOffset & SHARED_OFFSET_MASK); }
+	ANIMNEXT_API bool Serialize(class FArchive& Ar);
 
 private:
-	// Bottom 24 bits are used by the node shared offset while the top 8 bits by the decorator index
+	// Bottom 24 bits are used by the node handle while the top 8 bits by the decorator index
 	static constexpr uint32 DECORATOR_SHIFT = 24;
-	static constexpr uint32 SHARED_OFFSET_MASK = ~0u >> (32 - DECORATOR_SHIFT);
-	static constexpr uint32 INVALID_SHARED_OFFSET_VALUE = SHARED_OFFSET_MASK;
+	static constexpr uint32 NODE_HANDLE_MASK = ~0u >> (32 - DECORATOR_SHIFT);
 
 	UPROPERTY()
-	uint32		PackedDecoratorIndexAndNodeSharedOffset;
+	uint32		PackedDecoratorIndexAndNodeHandle;
+};
+
+template<>
+struct TStructOpsTypeTraits<FAnimNextDecoratorHandle> : public TStructOpsTypeTraitsBase2<FAnimNextDecoratorHandle>
+{
+	enum
+	{
+		WithSerializer = true,
+	};
 };
diff --git a/Engine/Plugins/Experimental/AnimNext/Source/AnimNext/Public/DecoratorBase/DecoratorInterfaceUID.h b/Engine/Plugins/Experimental/AnimNext/Source/AnimNext/Public/DecoratorBase/DecoratorInterfaceUID.h
--- a/Engine/Plugins/Experimental/AnimNext/Source/AnimNext/Public/DecoratorBase/DecoratorInterfaceUID.h
+++ b/Engine/Plugins/Experimental/AnimNext/Source/AnimNext/Public/DecoratorBase/DecoratorInterfaceUID.h
@@ -6,6 +6,9 @@
 
 namespace UE::AnimNext
 {
+	// Type alias for a raw decorator UID, not typesafe
+	using FDecoratorInterfaceUIDRaw = uint32;
+
 	/**
 	 * FDecoratorInterfaceUID
 	 *
@@ -19,17 +22,17 @@ namespace UE::AnimNext
 	{
 		// Constructs an invalid UID
 		constexpr FDecoratorInterfaceUID()
-			: UID(0)
+			: UID(INVALID_UID)
 #if !UE_BUILD_SHIPPING
 			, InterfaceName(TEXT("<Invalid Interface UID>"))
 #endif
 		{}
 
 		// Constructs an interface UID
-		constexpr FDecoratorInterfaceUID(const TCHAR* InterfaceName_, uint32 UID_)
-			: UID(UID_)
+		explicit constexpr FDecoratorInterfaceUID(FDecoratorInterfaceUIDRaw InUID, const TCHAR* InInterfaceName = TEXT("<Unknown Decorator Interface Name>"))
+			: UID(InUID)
 #if !UE_BUILD_SHIPPING
-			, InterfaceName(InterfaceName_)
+			, InterfaceName(InInterfaceName)
 #endif
 		{
 		}
@@ -40,20 +43,22 @@ namespace UE::AnimNext
 #endif
 
 		// Returns the interface global UID
-		constexpr uint32 GetUID() const { return UID; }
+		constexpr FDecoratorInterfaceUIDRaw GetUID() const { return UID; }
 
 		// Returns whether this UID is valid or not
-		constexpr bool IsValid() const { return UID != 0; }
-
-		// Compares for equality and inequality
-		constexpr bool operator==(FDecoratorInterfaceUID RHS) const { return UID == RHS.UID; }
-		constexpr bool operator!=(FDecoratorInterfaceUID RHS) const { return UID != RHS.UID; }
+		constexpr bool IsValid() const { return UID != INVALID_UID; }
 
 	private:
-		uint32	UID;
+		static constexpr FDecoratorInterfaceUIDRaw INVALID_UID = 0;
+
+		FDecoratorInterfaceUIDRaw	UID;
 
 #if !UE_BUILD_SHIPPING
-		const TCHAR* InterfaceName;
+		const TCHAR*				InterfaceName;
 #endif
 	};
+
+	// Compares for equality and inequality
+	constexpr bool operator==(FDecoratorInterfaceUID LHS, FDecoratorInterfaceUID RHS) { return LHS.GetUID() == RHS.GetUID(); }
+	constexpr bool operator!=(FDecoratorInterfaceUID LHS, FDecoratorInterfaceUID RHS) { return LHS.GetUID() != RHS.GetUID(); }
 }
diff --git a/Engine/Plugins/Experimental/AnimNext/Source/AnimNext/Public/DecoratorBase/DecoratorPtr.h b/Engine/Plugins/Experimental/AnimNext/Source/AnimNext/Public/DecoratorBase/DecoratorPtr.h
--- a/Engine/Plugins/Experimental/AnimNext/Source/AnimNext/Public/DecoratorBase/DecoratorPtr.h
+++ b/Engine/Plugins/Experimental/AnimNext/Source/AnimNext/Public/DecoratorBase/DecoratorPtr.h
@@ -100,12 +100,8 @@ namespace UE::AnimNext
 			, DecoratorIndex(DecoratorPtr.DecoratorIndex)
 		{}
 
-		FWeakDecoratorPtr(FNodeInstance* NodeInstance_, uint32 DecoratorIndex_)
-			: NodeInstance(NodeInstance_)
-			, DecoratorIndex(static_cast<uint8>(DecoratorIndex_))
-		{
-			check(DecoratorIndex == DecoratorIndex_);	// Make sure we didn't truncate
-		}
+		// Constructs a weak pointer handle to the provided instance
+		FWeakDecoratorPtr(FNodeInstance* NodeInstance_, uint32 DecoratorIndex_);
 
 		// Returns a pointer to the node instance
 		constexpr FNodeInstance* GetNodeInstance() const noexcept { return NodeInstance; }
@@ -127,7 +123,7 @@ namespace UE::AnimNext
 
 	private:
 		FNodeInstance*	NodeInstance = nullptr;
-		uint8			DecoratorIndex = 0;
+		uint32			DecoratorIndex = 0;			// Only need 8 bits, but use 32 since we have ample padding, avoids truncation
 	};
 
 	//////////////////////////////////////////////////////////////////////////
diff --git a/Engine/Plugins/Experimental/AnimNext/Source/AnimNext/Public/DecoratorBase/DecoratorUID.h b/Engine/Plugins/Experimental/AnimNext/Source/AnimNext/Public/DecoratorBase/DecoratorUID.h
--- a/Engine/Plugins/Experimental/AnimNext/Source/AnimNext/Public/DecoratorBase/DecoratorUID.h
+++ b/Engine/Plugins/Experimental/AnimNext/Source/AnimNext/Public/DecoratorBase/DecoratorUID.h
@@ -6,6 +6,9 @@
 
 namespace UE::AnimNext
 {
+	// Type alias for a raw decorator UID, not typesafe
+	using FDecoratorUIDRaw = uint32;
+
 	/**
 	 * FDecoratorUID
 	 *
@@ -19,17 +22,18 @@ namespace UE::AnimNext
 	{
 		// Constructs an invalid UID
 		constexpr FDecoratorUID()
-			: UID(0)
+			: UID(INVALID_UID)
 #if !UE_BUILD_SHIPPING
 			, DecoratorName(TEXT("<Invalid decorator UID>"))
 #endif
-		{}
+		{
+		}
 
 		// Constructs a decorator UID
-		constexpr FDecoratorUID(const TCHAR* DecoratorName_, uint32 UID_)
-			: UID(UID_)
+		explicit constexpr FDecoratorUID(FDecoratorUIDRaw InUID, const TCHAR* InDecoratorName = TEXT("<Unknown Decorator Name>"))
+			: UID(InUID)
 #if !UE_BUILD_SHIPPING
-			, DecoratorName(DecoratorName_)
+			, DecoratorName(InDecoratorName)
 #endif
 		{
 		}
@@ -40,24 +44,26 @@ namespace UE::AnimNext
 #endif
 
 		// Returns the decorator global UID
-		constexpr uint32 GetUID() const { return UID; }
+		constexpr FDecoratorUIDRaw GetUID() const { return UID; }
 
 		// Returns whether this UID is valid or not
-		constexpr bool IsValid() const { return UID != 0; }
+		constexpr bool IsValid() const { return UID != INVALID_UID; }
 
 	private:
-		uint32	UID;
+		static constexpr FDecoratorUIDRaw INVALID_UID = 0;
+
+		FDecoratorUIDRaw	UID;
 
 #if !UE_BUILD_SHIPPING
-		const TCHAR* DecoratorName;
+		const TCHAR*		DecoratorName;
 #endif
 	};
 
 	// Compares for equality and inequality
 	constexpr bool operator==(FDecoratorUID LHS, FDecoratorUID RHS) { return LHS.GetUID() == RHS.GetUID(); }
 	constexpr bool operator!=(FDecoratorUID LHS, FDecoratorUID RHS) { return LHS.GetUID() != RHS.GetUID(); }
-	constexpr bool operator==(FDecoratorUID LHS, uint32 RHS) { return LHS.GetUID() == RHS; }
-	constexpr bool operator!=(FDecoratorUID LHS, uint32 RHS) { return LHS.GetUID() != RHS; }
-	constexpr bool operator==(uint32 LHS, FDecoratorUID RHS) { return LHS == RHS.GetUID(); }
-	constexpr bool operator!=(uint32 LHS, FDecoratorUID RHS) { return LHS != RHS.GetUID(); }
+	constexpr bool operator==(FDecoratorUID LHS, FDecoratorUIDRaw RHS) { return LHS.GetUID() == RHS; }
+	constexpr bool operator!=(FDecoratorUID LHS, FDecoratorUIDRaw RHS) { return LHS.GetUID() != RHS; }
+	constexpr bool operator==(FDecoratorUIDRaw LHS, FDecoratorUID RHS) { return LHS == RHS.GetUID(); }
+	constexpr bool operator!=(FDecoratorUIDRaw LHS, FDecoratorUID RHS) { return LHS != RHS.GetUID(); }
 }
diff --git a/Engine/Plugins/Experimental/AnimNext/Source/AnimNext/Public/DecoratorBase/EntryPointHandle.h b/Engine/Plugins/Experimental/AnimNext/Source/AnimNext/Public/DecoratorBase/EntryPointHandle.h
new file mode 100644
--- /dev/null
+++ b/Engine/Plugins/Experimental/AnimNext/Source/AnimNext/Public/DecoratorBase/EntryPointHandle.h
@@ -0,0 +1,62 @@
+// Copyright Epic Games, Inc. All Rights Reserved.
+
+#pragma once
+
+#include "CoreMinimal.h"
+
+#include "DecoratorBase/NodeHandle.h"
+
+#include "EntryPointHandle.generated.h"
+
+/**
+ * Entry Point Handle
+ * An entry point handle is equivalent to a decorator handle but it will not resolve automatically on load.
+ * As such, it is safe to use outside of an AnimNext graph.
+ * They must be manually resolved through FDecoratorReader.
+ *
+ * Internally, it contains a node handle (as a node ID) and a decorator index.
+ *
+ * @see FAnimNextDecoratorHandle, FDecoratorReader
+ */
+USTRUCT(BlueprintType)
+struct FAnimNextEntryPointHandle final
+{
+	GENERATED_BODY()
+
+	// Creates an invalid entry point handle
+	constexpr FAnimNextEntryPointHandle() noexcept
+		: PackedDecoratorIndexAndNodeHandle(UE::AnimNext::FNodeHandle::INVALID_NODE_HANDLE_RAW_VALUE)
+	{}
+
+	// Creates an entry point handle pointing to the first decorator of the specified node
+	explicit FAnimNextEntryPointHandle(UE::AnimNext::FNodeHandle NodeHandle)
+		: PackedDecoratorIndexAndNodeHandle(NodeHandle.GetPackedValue() & NODE_HANDLE_MASK)
+	{
+		check(!NodeHandle.IsValid() || NodeHandle.IsNodeID());
+	}
+
+	// Creates an entry point handle pointing to the specified decorator on the specified node
+	FAnimNextEntryPointHandle(UE::AnimNext::FNodeHandle NodeHandle, uint32 DecoratorIndex)
+		: PackedDecoratorIndexAndNodeHandle((NodeHandle.GetPackedValue() & NODE_HANDLE_MASK) | (DecoratorIndex << 24))
+	{
+		check(!NodeHandle.IsValid() || NodeHandle.IsNodeID());
+		check(DecoratorIndex <= MAX_uint8);	// Make sure we don't truncate
+	}
+
+	// Returns true if this entry point handle is valid, false otherwise
+	constexpr bool IsValid() const noexcept { return GetNodeHandle().IsValid(); }
+
+	// Returns the decorator index
+	constexpr uint32 GetDecoratorIndex() const noexcept { return PackedDecoratorIndexAndNodeHandle >> DECORATOR_SHIFT; }
+
+	// Returns a handle to the node referenced (its node ID)
+	constexpr UE::AnimNext::FNodeHandle GetNodeHandle() const noexcept { return UE::AnimNext::FNodeHandle::FromPackedValue(PackedDecoratorIndexAndNodeHandle); }
+
+private:
+	// Bottom 24 bits are used by the node handle while the top 8 bits by the decorator index
+	static constexpr uint32 DECORATOR_SHIFT = 24;
+	static constexpr uint32 NODE_HANDLE_MASK = ~0u >> (32 - DECORATOR_SHIFT);
+
+	UPROPERTY()
+	uint32		PackedDecoratorIndexAndNodeHandle;
+};
diff --git a/Engine/Plugins/Experimental/AnimNext/Source/AnimNext/Public/DecoratorBase/IDecoratorInterface.h b/Engine/Plugins/Experimental/AnimNext/Source/AnimNext/Public/DecoratorBase/IDecoratorInterface.h
--- a/Engine/Plugins/Experimental/AnimNext/Source/AnimNext/Public/DecoratorBase/IDecoratorInterface.h
+++ b/Engine/Plugins/Experimental/AnimNext/Source/AnimNext/Public/DecoratorBase/IDecoratorInterface.h
@@ -9,7 +9,7 @@
 // Helper macros
 #define DECLARE_ANIM_DECORATOR_INTERFACE(InterfaceName, InterfaceNameHash) \
 	/* Globally unique UID for this interface */ \
-	static constexpr UE::AnimNext::FDecoratorInterfaceUID InterfaceUID = UE::AnimNext::FDecoratorInterfaceUID(TEXT(#InterfaceName), InterfaceNameHash); \
+	static constexpr UE::AnimNext::FDecoratorInterfaceUID InterfaceUID = UE::AnimNext::FDecoratorInterfaceUID(InterfaceNameHash, TEXT(#InterfaceName)); \
 	virtual UE::AnimNext::FDecoratorInterfaceUID GetInterfaceUID() const override { return InterfaceUID; }
 
 namespace UE::AnimNext
@@ -28,7 +28,7 @@ namespace UE::AnimNext
 		// The globally unique UID for this interface
 		// Derived types will have their own InterfaceUID member that hides/aliases/shadows this one
 		// @see DECLARE_ANIM_DECORATOR_INTERFACE
-		static constexpr FDecoratorInterfaceUID InterfaceUID = FDecoratorInterfaceUID(TEXT("IDecoratorInterface"), 0x402f6df4);
+		static constexpr FDecoratorInterfaceUID InterfaceUID = FDecoratorInterfaceUID(0x402f6df4, TEXT("IDecoratorInterface"));
 
 		// Returns the globally unique UID for this interface
 		virtual FDecoratorInterfaceUID GetInterfaceUID() const { return InterfaceUID; };
diff --git a/Engine/Plugins/Experimental/AnimNext/Source/AnimNext/Public/DecoratorBase/NodeDescription.h b/Engine/Plugins/Experimental/AnimNext/Source/AnimNext/Public/DecoratorBase/NodeDescription.h
--- a/Engine/Plugins/Experimental/AnimNext/Source/AnimNext/Public/DecoratorBase/NodeDescription.h
+++ b/Engine/Plugins/Experimental/AnimNext/Source/AnimNext/Public/DecoratorBase/NodeDescription.h
@@ -3,10 +3,9 @@
 #pragma once
 
 #include "CoreMinimal.h"
+#include "DecoratorBase/NodeID.h"
 #include "DecoratorBase/NodeTemplateRegistryHandle.h"
 
-#include <limits>
-
 class FArchive;
 
 namespace UE::AnimNext
@@ -26,22 +25,17 @@ namespace UE::AnimNext
 	 */
 	struct alignas(alignof(uint32)) FNodeDescription
 	{
-		// Largest allowed size for a node description
-		// Graphs are currently limited to 64 KB because we use 16 bit offsets within our node/decorator handles
-		// We have no particular limit on the node size besides the graph limit (for now)
-		static constexpr uint32 MAXIMUM_SIZE = 64 * 1024;
-
-		// The maximum number of nodes allowed within a single graph
-		// We use 16 bits to represent node UIDs
-		static constexpr uint32 MAXIMUM_COUNT = std::numeric_limits<uint16>::max();
-
-		FNodeDescription(uint16 UID_, FNodeTemplateRegistryHandle TemplateHandle_)
-			: UID(UID_)
-			, TemplateHandle(TemplateHandle_)
+		// Largest allowed size for a node description and the shared data of its decorators
+		// We use unsigned 16 bit offsets within the NodeTemplate/DecoratorTemplate
+		static constexpr uint32 MAXIMUM_NODE_SHARED_DATA_SIZE = 64 * 1024;
+
+		FNodeDescription(FNodeID InNodeID, FNodeTemplateRegistryHandle InTemplateHandle)
+			: NodeID(InNodeID)
+			, TemplateHandle(InTemplateHandle)
 		{}
 
 		// Returns the node UID, unique to the owning sub-graph
-		uint32 GetUID() const { return UID; }
+		FNodeID GetUID() const { return NodeID; }
 
 		// Returns the handle of the node's template in the node template registry
 		FNodeTemplateRegistryHandle GetTemplateHandle() const { return TemplateHandle; }
@@ -50,7 +44,7 @@ namespace UE::AnimNext
 		ANIMNEXT_API void Serialize(FArchive& Ar);
 
 	private:
-		uint16							UID;				// assigned during export/cook, unique to current sub-graph
+		FNodeID							NodeID;				// assigned during export/cook, unique to current sub-graph
 		FNodeTemplateRegistryHandle		TemplateHandle;		// offset of the node template within the global list
 
 		// Followed by a list of [FAnimNextDecoratorSharedData] instances and optional padding
diff --git a/Engine/Plugins/Experimental/AnimNext/Source/AnimNext/Public/DecoratorBase/NodeHandle.h b/Engine/Plugins/Experimental/AnimNext/Source/AnimNext/Public/DecoratorBase/NodeHandle.h
--- a/Engine/Plugins/Experimental/AnimNext/Source/AnimNext/Public/DecoratorBase/NodeHandle.h
+++ b/Engine/Plugins/Experimental/AnimNext/Source/AnimNext/Public/DecoratorBase/NodeHandle.h
@@ -3,8 +3,11 @@
 #pragma once
 
 #include "CoreMinimal.h"
+#include "DecoratorBase/NodeID.h"
+#include "Serialization/Archive.h"
 
 struct FAnimNextDecoratorHandle;
+struct FAnimNextEntryPointHandle;
 
 namespace UE::AnimNext
 {
@@ -13,37 +16,81 @@ namespace UE::AnimNext
 	 * A node handle represents a reference to a specific node instance in the shared/read-only portion
 	 * of a sub-graph. It points to a FNodeDescription when resolved.
 	 * 
-	 * @see FNodeDescription
+	 * Internally, it contains either a shared data offset into the parent graph (at runtime) or a FNodeID (during compilation and on disk).
+	 * Shared offsets must be computed at runtime and as such we use the FNodeID to resolve into our final offset on load.
+	 * We only use the bottom 24 bits. If the LSB is 0, the packed value represents a shared offset (due to alignment) but if the
+	 * LSB is 1, the packed value represents a node ID shifted left by 1.
+	 * 
+	 * @see FNodeDescription, FNodeID
 	 */
-	struct FNodeHandle
+	struct FNodeHandle final
 	{
 		// Creates an invalid node handle
-		FNodeHandle()
-			: SharedOffset(INVALID_SHARED_OFFSET_VALUE)
+		constexpr FNodeHandle() noexcept
+			: PackedSharedOffsetOrID(INVALID_NODE_HANDLE_RAW_VALUE)
 		{}
 
 		// Creates a node handle for the specified shared offset
-		explicit FNodeHandle(uint32 SharedOffset_)
-			: SharedOffset(SharedOffset_)
-		{}
+		static constexpr FNodeHandle FromSharedOffset(uint32 SharedOffset) noexcept
+		{
+			if (SharedOffset > NODE_HANDLE_MASK)
+			{
+				return FNodeHandle();	// The shared offset is too large, the handle isn't valid
+			}
+
+			return FNodeHandle(SharedOffset);
+		}
+
+		// Creates a node handle from a node ID
+		static constexpr FNodeHandle FromNodeID(FNodeID NodeID) noexcept
+		{
+			if (!NodeID.IsValid())
+			{
+				return FNodeHandle();	// The NodeID isn't valid, the handle isn't valid
+			}
+
+			return FNodeHandle((NodeID.GetRawValue() << NODE_ID_SHIFT) | NODE_ID_MARKER_BIT_MASK);
+		}
 
 		// Returns true if this node handle is valid, false otherwise
-		bool IsValid() const { return SharedOffset != INVALID_SHARED_OFFSET_VALUE; }
+		constexpr bool IsValid() const noexcept { return (PackedSharedOffsetOrID & NODE_HANDLE_MASK) != INVALID_NODE_HANDLE_RAW_VALUE; }
+
+		// Returns whether this node handle encodes a shared offset or not
+		constexpr bool IsSharedOffset() const noexcept { return IsValid() && (PackedSharedOffsetOrID & NODE_ID_MARKER_BIT_MASK) == 0; }
+
+		// Returns whether this node handle encodes a node ID or not
+		constexpr bool IsNodeID() const noexcept { return IsValid() && (PackedSharedOffsetOrID & NODE_ID_MARKER_BIT_MASK) != 0; }
 
 		// Returns the offset into the shared data relative to the root of the owning sub-graph
-		uint32 GetSharedOffset() const { return SharedOffset; }
+		uint32 GetSharedOffset() const { check(IsSharedOffset()); return PackedSharedOffsetOrID & NODE_HANDLE_MASK; }
+
+		// Returns the node ID encoded by the handle
+		FNodeID GetNodeID() const { check(IsNodeID()); return FNodeID::FromRawValue((PackedSharedOffsetOrID & NODE_HANDLE_MASK) >> NODE_ID_SHIFT); }
 
 		// Equality and inequality tests
-		bool operator==(FNodeHandle RHS) const { return SharedOffset == RHS.SharedOffset; }
-		bool operator!=(FNodeHandle RHS) const { return SharedOffset != RHS.SharedOffset; }
+		constexpr bool operator==(FNodeHandle RHS) const noexcept { return (PackedSharedOffsetOrID & NODE_HANDLE_MASK) == (RHS.PackedSharedOffsetOrID & NODE_HANDLE_MASK); }
+		constexpr bool operator!=(FNodeHandle RHS) const noexcept { return (PackedSharedOffsetOrID & NODE_HANDLE_MASK) != (RHS.PackedSharedOffsetOrID & NODE_HANDLE_MASK); }
 
 	private:
-		// We pack the shared offset on 24 bits, the top 8 bits are truncated
-		static constexpr uint32 INVALID_SHARED_OFFSET_VALUE = ~0u >> (32 - 24);
+		explicit constexpr FNodeHandle(uint32 PackedValue) noexcept
+			: PackedSharedOffsetOrID(PackedValue)
+		{}
+
+		// Returns the packed raw value
+		constexpr uint32 GetPackedValue() const { return PackedSharedOffsetOrID; }
+
+		// Creates a node handle from a raw packed value
+		static constexpr FNodeHandle FromPackedValue(uint32 PackedValue) noexcept { return FNodeHandle(PackedValue); }
+
+		// We pack the shared offset on 24 bits, the top 8 bits are unused and they can have any value (zero by default)
+		static constexpr uint32 INVALID_NODE_HANDLE_RAW_VALUE = ~0u >> (32 - 24);	// 0x00FFFFFF
+		static constexpr uint32 NODE_HANDLE_MASK = INVALID_NODE_HANDLE_RAW_VALUE;
+		static constexpr uint32 NODE_ID_MARKER_BIT_MASK = 0x01u;					// When valid, a LSB of 1 means we encode a node ID, 0 means a shared offset
+		static constexpr uint32 NODE_ID_SHIFT = 1;									// Our LSB is reserved as a marker, the node ID is shifted left by 1
 
-		uint32	Padding : 8;			// unused padding
-		uint32	SharedOffset : 24;		// relative to root of sub-graph
+		uint32	PackedSharedOffsetOrID;
 
 		friend FAnimNextDecoratorHandle;
+		friend FAnimNextEntryPointHandle;
 	};
 }
diff --git a/Engine/Plugins/Experimental/AnimNext/Source/AnimNext/Public/DecoratorBase/NodeID.h b/Engine/Plugins/Experimental/AnimNext/Source/AnimNext/Public/DecoratorBase/NodeID.h
new file mode 100644
--- /dev/null
+++ b/Engine/Plugins/Experimental/AnimNext/Source/AnimNext/Public/DecoratorBase/NodeID.h
@@ -0,0 +1,71 @@
+// Copyright Epic Games, Inc. All Rights Reserved.
+
+#pragma once
+
+#include "CoreMinimal.h"
+#include "Serialization/Archive.h"
+
+namespace UE::AnimNext
+{
+	/**
+	  * Node ID
+	  * 
+	  * A node ID fits on 16 bits and is used to encore a node index.
+	  * The zero ID is reserved as an invalid value, this means that the
+	  * node index is 1-based.
+	  * 
+	  * Node IDs are assigned during compilation of an AnimNext graph to each node.
+	  * The decorator reader is then used to convert them into a shared offset within
+	  * the graph.
+	  * 
+	  * @see FDecoratorReader
+	  */
+	struct FNodeID final
+	{
+		// Constructs an invalid FNodeID
+		constexpr FNodeID() noexcept : ID(INVALID_NODE_ID) {}
+
+		// Returns whether or not this ID is valid
+		constexpr bool IsValid() const noexcept { return ID != INVALID_NODE_ID; }
+
+		// Returns the node index represented by this node ID
+		// If the node ID is invalid, the index returned is 0xFFFFFFFF
+		constexpr uint32 GetNodeIndex() const noexcept { return static_cast<uint32>(ID) - 1; }
+
+		// The first valid node ID
+		static constexpr FNodeID GetFirstID() noexcept { return FNodeID(1); }
+
+		// Returns the next node ID
+		// If the current ID is invalid, the next is invalid as well
+		// Note that this function does not handle wrapping if more than 65536 IDs are used
+		// If wrapping occurs, the next ID returned will be invalid and all subsequent IDs will be invalid
+		constexpr FNodeID GetNextID() const noexcept { return IsValid() ? FNodeID(ID + 1) : FNodeID(); }
+
+		// Equality and inequality tests
+		constexpr bool operator==(const FNodeID RHS) const noexcept { return ID == RHS.ID; }
+		constexpr bool operator!=(const FNodeID RHS) const noexcept { return ID != RHS.ID; }
+
+		// Serialization operator
+		friend FArchive& operator<<(FArchive& Ar, FNodeID& Value)
+		{
+			Ar << Value.ID;
+			return Ar;
+		}
+
+	private:
+		// Constructs a FNodeID instance from a raw value
+		explicit constexpr FNodeID(uint16 InID) noexcept : ID(InID) {}
+
+		// Returns the packed raw value
+		constexpr uint32 GetRawValue() const { return ID; }
+
+		// Creates a node handle from a raw packed value
+		static constexpr FNodeID FromRawValue(uint32 InID) noexcept { return FNodeID(static_cast<uint16>(InID)); }
+
+		static constexpr uint16 INVALID_NODE_ID = 0;
+
+		uint16 ID;
+
+		friend struct FNodeHandle;
+	};
+}
diff --git a/Engine/Plugins/Experimental/AnimNext/Source/AnimNext/Public/DecoratorBase/NodeInstance.h b/Engine/Plugins/Experimental/AnimNext/Source/AnimNext/Public/DecoratorBase/NodeInstance.h
--- a/Engine/Plugins/Experimental/AnimNext/Source/AnimNext/Public/DecoratorBase/NodeInstance.h
+++ b/Engine/Plugins/Experimental/AnimNext/Source/AnimNext/Public/DecoratorBase/NodeInstance.h
@@ -18,6 +18,10 @@ namespace UE::AnimNext
 	 */
 	struct ANIMNEXT_API FNodeInstance
 	{
+		// Largest allowed size for a node instance and the instance data of its decorators
+		// We use unsigned 16 bit offsets within the NodeTemplate/DecoratorTemplate
+		static constexpr uint32 MAXIMUM_NODE_INSTANCE_DATA_SIZE = 64 * 1024;
+
 		FNodeInstance(const FNodeInstance&) = delete;
 		FNodeInstance& operator=(const FNodeInstance&) = delete;
 
@@ -37,7 +41,7 @@ namespace UE::AnimNext
 		void AddReference() { ReferenceCount++; }
 		void ReleaseReference();
 
-		uint32		ReferenceCount;		// how many non-weak FNodePtr and FDecoratorPtr handles point to us, not thread safe
+		uint32		ReferenceCount;		// how many non-weak FDecoratorPtr handles point to us, not thread safe
 		FNodeHandle	NodeHandle;			// relative to root of sub-graph, should this be a pointer?
 
 		// Followed by a list of [FDecoratorInstanceData] instances and optional padding
diff --git a/Engine/Plugins/Experimental/AnimNext/Source/AnimNext/Public/DecoratorInterfaces/IUpdate.h b/Engine/Plugins/Experimental/AnimNext/Source/AnimNext/Public/DecoratorInterfaces/IUpdate.h
--- a/Engine/Plugins/Experimental/AnimNext/Source/AnimNext/Public/DecoratorInterfaces/IUpdate.h
+++ b/Engine/Plugins/Experimental/AnimNext/Source/AnimNext/Public/DecoratorInterfaces/IUpdate.h
@@ -15,8 +15,12 @@ namespace UE::AnimNext
 	 */
 	struct FUpdateTraversalContext : ITraversalContext
 	{
-		double GetDeltaTime() const { return 1.0 / 30.0; }
-		double GetPlayRate() const { return 1.0; }
+		explicit FUpdateTraversalContext(float InDeltaTime);
+
+		float GetDeltaTime() const { return DeltaTime; }
+
+	private:
+		float DeltaTime;
 	};
 
 	/**
@@ -75,5 +79,5 @@ namespace UE::AnimNext
 	 *
 	 * @see IUpdate::PreUpdate, IUpdate::PostUpdate, IHierarchy::GetChildren
 	 */
-	ANIMNEXT_API void UpdateGraph(FExecutionContext& Context, FWeakDecoratorPtr GraphRootPtr);
+	ANIMNEXT_API void UpdateGraph(FExecutionContext& Context, FWeakDecoratorPtr GraphRootPtr, float DeltaTime);
 }
diff --git a/Engine/Plugins/Experimental/AnimNext/Source/AnimNextEditor/AnimNextEditor.Build.cs b/Engine/Plugins/Experimental/AnimNext/Source/AnimNextEditor/AnimNextEditor.Build.cs
--- a/Engine/Plugins/Experimental/AnimNext/Source/AnimNextEditor/AnimNextEditor.Build.cs
+++ b/Engine/Plugins/Experimental/AnimNext/Source/AnimNextEditor/AnimNextEditor.Build.cs
@@ -22,6 +22,7 @@ namespace UnrealBuildTool.Rules
 					"InputCore",
 					"PropertyEditor",
 					"RigVM",
+					"RigVMEditor",
 					"RigVMDeveloper",
 					"ControlRig",
 					"ControlRigDeveloper",
diff --git a/Engine/Plugins/Experimental/AnimNext/Source/AnimNextEditor/Private/AnimNextEditorModule.cpp b/Engine/Plugins/Experimental/AnimNext/Source/AnimNextEditor/Private/AnimNextEditorModule.cpp
--- a/Engine/Plugins/Experimental/AnimNext/Source/AnimNextEditor/Private/AnimNextEditorModule.cpp
+++ b/Engine/Plugins/Experimental/AnimNext/Source/AnimNextEditor/Private/AnimNextEditorModule.cpp
@@ -4,6 +4,7 @@
 #include "AssetToolsModule.h"
 #include "IAssetTools.h"
 #include "Graph/AssetTypeActions.h"
+#include "Graph/AnimNextGraphPanelNodeFactory.h"
 #include "Param/ParamTypePropertyCustomization.h"
 #include "Param/ParameterPickerArgs.h"
 #include "Param/SParameterPicker.h"
@@ -26,6 +27,9 @@ class FModule : public IModule
 		PropertyModule.RegisterCustomPropertyTypeLayout(
 			"AnimNextParamType",
 			FOnGetPropertyTypeCustomizationInstance::CreateLambda([] { return MakeShared<FParamTypePropertyTypeCustomization>(); }));
+
+		AnimNextGraphPanelNodeFactory = MakeShared<FAnimNextGraphPanelNodeFactory>();
+		FEdGraphUtilities::RegisterVisualNodeFactory(AnimNextGraphPanelNodeFactory);
 	}
 
 	virtual void ShutdownModule() override
@@ -41,6 +45,8 @@ class FModule : public IModule
 			FPropertyEditorModule& PropertyModule = FModuleManager::LoadModuleChecked<FPropertyEditorModule>("PropertyEditor");
 			PropertyModule.UnregisterCustomPropertyTypeLayout("AnimNextParamType");
 		}
+
+		FEdGraphUtilities::UnregisterVisualNodeFactory(AnimNextGraphPanelNodeFactory);
 	}
 
 	virtual TSharedRef<SWidget> CreateParameterPicker(const FParameterPickerArgs& InArgs) override
@@ -50,6 +56,9 @@ class FModule : public IModule
 	}
 
 	TSharedPtr<FAssetTypeActions_AnimNextGraph> AssetTypeActions_AnimNextGraph;
+
+	/** Node factory for the AnimNext graph */
+	TSharedPtr<FAnimNextGraphPanelNodeFactory> AnimNextGraphPanelNodeFactory;
 };
 
 }
diff --git a/Engine/Plugins/Experimental/AnimNext/Source/AnimNextEditor/Private/Graph/AnimNextGraphEditor.cpp b/Engine/Plugins/Experimental/AnimNext/Source/AnimNextEditor/Private/Graph/AnimNextGraphEditor.cpp
--- a/Engine/Plugins/Experimental/AnimNext/Source/AnimNextEditor/Private/Graph/AnimNextGraphEditor.cpp
+++ b/Engine/Plugins/Experimental/AnimNext/Source/AnimNextEditor/Private/Graph/AnimNextGraphEditor.cpp
@@ -3,6 +3,7 @@
 #include "AnimNextGraphEditor.h"
 #include "GraphEditorMode.h"
 #include "Graph/AnimNextGraph.h"
+#include "Graph/AnimNextGraph_EdGraphNode.h"
 #include "Graph/AnimNextGraph_EditorData.h"
 #include "EdGraphNode_Comment.h"
 #include "SActionMenu.h"
@@ -62,6 +63,9 @@ void FGraphEditor::InitEditor(const EToolkitMode::Type InMode, const TSharedPtr<
 
 	// Open initial document
 	DocumentManager->OpenDocument(FTabPayload_UObject::Make(AnimNextGraph_EditorData->RootGraph), FDocumentTracker::EOpenDocumentCause::OpenNewDocument);
+
+	// Make sure editor graph matches the model
+	AnimNextGraph_EditorData->RebuildEdGraphFromModel();
 }
 
 void FGraphEditor::BindCommands()
@@ -241,6 +245,8 @@ void FGraphEditor::DeleteSelectedNodes()
 		FocusedGraphEd->ClearSelectionSet();
 	}
 
+	TMap<UAnimNextGraph_EdGraphNode*, URigVMNode*> NodesToRemove;
+
 	// Some nodes have sub-objects that are represented as other tabs.
 	// Close them here as a pre-pass before we remove their nodes. If the documents are left open they
 	// may reference dangling data and function incorrectly in cases such as FindBlueprintforNodeChecked
@@ -268,26 +274,35 @@ void FGraphEditor::DeleteSelectedNodes()
 				{
 					CloseAllDocumentsTab(Node);
 				}
+
+				if (UAnimNextGraph_EdGraphNode* AnimNextNode = Cast<UAnimNextGraph_EdGraphNode>(Node))
+				{
+					if (URigVMNode* ModelNode = GetFocusedVMController()->GetGraph()->FindNodeByName(*AnimNextNode->GetModelNodePath()))
+					{
+						NodesToRemove.Add(AnimNextNode, ModelNode);
+					}
+				}
 			}
 		}
 	}
 
-	// Now remove the selected nodes
-	for (FGraphPanelSelectionSet::TConstIterator NodeIt( SelectedNodes ); NodeIt; ++NodeIt)
+	if (NodesToRemove.IsEmpty())
 	{
-		if (UEdGraphNode* Node = Cast<UEdGraphNode>(*NodeIt))
-		{
-			if (Node->CanUserDeleteNode())
-			{
-				if (Node->GetSubGraphs().Num() > 0)
-				{
-					DocumentManager->CleanInvalidTabs();
-				}
+		return;
+	}
 
-				FBlueprintEditorUtils::RemoveNode(nullptr, Node);
-			}
-		}
+	GetFocusedVMController()->OpenUndoBracket(TEXT("Delete selected nodes"));
+
+	TArray<URigVMNode*> ModelNodesToRemove;
+	for (auto& It : NodesToRemove)
+	{
+		URigVMNode* ModelNode = It.Value;
+
+		ModelNodesToRemove.Add(ModelNode);
 	}
+
+	GetFocusedVMController()->RemoveNodes(ModelNodesToRemove, true);
+	GetFocusedVMController()->CloseUndoBracket();
 }
 
 bool FGraphEditor::CanDeleteSelectedNodes()
diff --git a/Engine/Plugins/Experimental/AnimNext/Source/AnimNextEditor/Private/Graph/AnimNextGraphPanelNodeFactory.cpp b/Engine/Plugins/Experimental/AnimNext/Source/AnimNextEditor/Private/Graph/AnimNextGraphPanelNodeFactory.cpp
new file mode 100644
--- /dev/null
+++ b/Engine/Plugins/Experimental/AnimNext/Source/AnimNextEditor/Private/Graph/AnimNextGraphPanelNodeFactory.cpp
@@ -0,0 +1,21 @@
+// Copyright Epic Games, Inc. All Rights Reserved.
+
+#include "Graph/AnimNextGraphPanelNodeFactory.h"
+#include "Graph/AnimNextGraph_EdGraphNode.h"
+#include "Graph/SAnimNextGraphNode.h"
+
+TSharedPtr<SGraphNode> FAnimNextGraphPanelNodeFactory::CreateNode(UEdGraphNode* Node) const
+{
+	if (UAnimNextGraph_EdGraphNode* AnimNextGraphNode = Cast<UAnimNextGraph_EdGraphNode>(Node))
+	{
+		TSharedPtr<SGraphNode> GraphNode =
+			SNew(SAnimNextGraphNode)
+			.GraphNodeObj(AnimNextGraphNode);
+
+		GraphNode->SlatePrepass();
+		AnimNextGraphNode->SetDimensions(GraphNode->GetDesiredSize());
+		return GraphNode;
+	}
+
+	return nullptr;
+}
diff --git a/Engine/Plugins/Experimental/AnimNext/Source/AnimNextEditor/Private/Graph/AnimNextGraphPanelNodeFactory.h b/Engine/Plugins/Experimental/AnimNext/Source/AnimNextEditor/Private/Graph/AnimNextGraphPanelNodeFactory.h
new file mode 100644
--- /dev/null
+++ b/Engine/Plugins/Experimental/AnimNext/Source/AnimNextEditor/Private/Graph/AnimNextGraphPanelNodeFactory.h
@@ -0,0 +1,10 @@
+// Copyright Epic Games, Inc. All Rights Reserved.
+
+#pragma once
+
+#include "EdGraphUtilities.h"
+
+class FAnimNextGraphPanelNodeFactory : public FGraphPanelNodeFactory
+{
+	virtual TSharedPtr<class SGraphNode> CreateNode(UEdGraphNode* Node) const override;
+};
diff --git a/Engine/Plugins/Experimental/AnimNext/Source/AnimNextEditor/Private/Graph/GraphEditorSchemaActions.cpp b/Engine/Plugins/Experimental/AnimNext/Source/AnimNextEditor/Private/Graph/GraphEditorSchemaActions.cpp
--- a/Engine/Plugins/Experimental/AnimNext/Source/AnimNextEditor/Private/Graph/GraphEditorSchemaActions.cpp
+++ b/Engine/Plugins/Experimental/AnimNext/Source/AnimNextEditor/Private/Graph/GraphEditorSchemaActions.cpp
@@ -7,6 +7,7 @@
 #include "Graph/AnimNextGraph_EditorData.h"
 #include "RigVMModel/Nodes/RigVMUnitNode.h"
 #include "Settings/ControlRigSettings.h"
+#include "Units/RigUnit.h"
 
 UEdGraphNode* FAnimNextGraphSchemaAction_RigUnit::PerformAction(UEdGraph* ParentGraph, TArray<UEdGraphPin*>& FromPins, const FVector2D Location, bool bSelectNewNode)
 {
@@ -24,7 +25,7 @@ UEdGraphNode* FAnimNextGraphSchemaAction_RigUnit::PerformAction(UEdGraph* Parent
 		FRigVMUnitNodeCreatedContext& UnitNodeCreatedContext = Controller->GetUnitNodeCreatedContext();
 		FRigVMUnitNodeCreatedContext::FScope ReasonScope(UnitNodeCreatedContext, ERigVMNodeCreatedReason::NodeSpawner);
 
-		if (URigVMUnitNode* ModelNode = Controller->AddUnitNode(StructTemplate, FRigUnit::GetMethodName(), Location, Name.ToString(), true, false))
+		if (URigVMUnitNode* ModelNode = Controller->AddUnitNode(StructTemplate, FRigVMStruct::ExecuteName, Location, Name.ToString(), true, false))
 		{
 			NewNode = Cast<UAnimNextGraph_EdGraphNode>(EdGraph->FindNodeForModelNodeName(ModelNode->GetFName()));
 			check(NewNode);
diff --git a/Engine/Plugins/Experimental/AnimNext/Source/AnimNextEditor/Private/Graph/GraphFactory.cpp b/Engine/Plugins/Experimental/AnimNext/Source/AnimNextEditor/Private/Graph/GraphFactory.cpp
--- a/Engine/Plugins/Experimental/AnimNext/Source/AnimNextEditor/Private/Graph/GraphFactory.cpp
+++ b/Engine/Plugins/Experimental/AnimNext/Source/AnimNextEditor/Private/Graph/GraphFactory.cpp
@@ -5,8 +5,7 @@
 #include "Graph/AnimNextGraph_EditorData.h"
 #include "Graph/AnimNextGraph_EdGraph.h"
 #include "UncookedOnlyUtils.h"
-#include "Graph/RigUnit_AnimNextBeginExecution.h"
-#include "Graph/RigUnit_AnimNextEndExecution.h"
+#include "Graph/RigUnit_AnimNextGraphRoot.h"
 #include "RigVMModel/RigVMController.h"
 #include "Units/RigUnit.h"
 
@@ -35,10 +34,10 @@ UObject* UAnimNextGraphFactory::FactoryCreateNew(UClass* Class, UObject* InParen
 	// Add initial execution unit
 	URigVMController* Controller = EditorData->GetRigVMClient()->GetController(EditorData->RootGraph);
 
-	URigVMUnitNode* MainEntryPointNode = Controller->AddUnitNode(FRigUnit_AnimNextBeginExecution::StaticStruct(), FRigUnit::GetMethodName(), FVector2D(-400.0f, 0.0f), FString(), false);
-	URigVMPin* BeginExecutePin = MainEntryPointNode->FindPin(GET_MEMBER_NAME_STRING_CHECKED(FRigUnit_AnimNextBeginExecution, ExecuteContext));
+	URigVMUnitNode* MainEntryPointNode = Controller->AddUnitNode(FRigUnit_AnimNextGraphRoot::StaticStruct(), FRigVMStruct::ExecuteName, FVector2D(-400.0f, 0.0f), FString(), false);
+	URigVMPin* BeginExecutePin = MainEntryPointNode->FindPin(GET_MEMBER_NAME_STRING_CHECKED(FRigUnit_AnimNextGraphRoot, Result));
 	check(BeginExecutePin);
-	check(BeginExecutePin->GetDirection() == ERigVMPinDirection::Output);
+	check(BeginExecutePin->GetDirection() == ERigVMPinDirection::Input);
 
 	//// Add function to function lib
 	//{
@@ -74,7 +73,7 @@ UObject* UAnimNextGraphFactory::FactoryCreateNew(UClass* Class, UObject* InParen
 
 	//// Add end-execution unit of the correct type
 	//// TODO: using a float for now, but need to use a registry to determine correct type
-	//URigVMUnitNode* MainExitPointNode = Controller->AddUnitNode(FRigUnit_AnimNextEndExecution_Float::StaticStruct(), FRigUnit::GetMethodName(), FVector2D(400.0f, 0.0f), FString(), false);
+	//URigVMUnitNode* MainExitPointNode = Controller->AddUnitNode(FRigUnit_AnimNextEndExecution_Float::StaticStruct(), FRigVMStruct::ExecuteName, FVector2D(400.0f, 0.0f), FString(), false);
 	//URigVMPin* EndExecutePin = MainExitPointNode->FindPin(GET_MEMBER_NAME_STRING_CHECKED(FRigUnit_AnimNextEndExecution, ExecuteContext));
 	//check(EndExecutePin);
 	//check(EndExecutePin->GetDirection() == ERigVMPinDirection::Input);
diff --git a/Engine/Plugins/Experimental/AnimNext/Source/AnimNextEditor/Private/Graph/SAnimNextGraphNode.cpp b/Engine/Plugins/Experimental/AnimNext/Source/AnimNextEditor/Private/Graph/SAnimNextGraphNode.cpp
new file mode 100644
--- /dev/null
+++ b/Engine/Plugins/Experimental/AnimNext/Source/AnimNextEditor/Private/Graph/SAnimNextGraphNode.cpp
@@ -0,0 +1,655 @@
+// Copyright Epic Games, Inc. All Rights Reserved.
+
+#include "Graph/SAnimNextGraphNode.h"
+#include "Graph/AnimNextGraph_EdGraphNode.h"
+#include "SGraphPin.h"
+#include "RigVMModel/Nodes/RigVMDispatchNode.h"
+#include "DecoratorBase/DecoratorHandle.h"
+#include "Widgets/SBoxPanel.h"
+#include "Widgets/Input/SButton.h"
+#include "Widgets/Layout/SSpacer.h"
+
+#define LOCTEXT_NAMESPACE "SAnimNextGraphNode"
+
+void SAnimNextGraphNode::Construct( const FArguments& InArgs )
+{
+	SRigVMGraphNode::Construct(SRigVMGraphNode::FArguments().GraphNodeObj(InArgs._GraphNodeObj));
+}
+
+void SAnimNextGraphNode::UpdatePinTreeView()
+{
+	static const float PinWidgetSidePadding = 6.f;
+	static const float EmptySidePadding = 60.f;
+	static const float TopPadding = 2.f; 
+	static const float MaxHeight = 30.f;
+
+	check(GraphNode);
+
+	// remove all existing content in the Left
+	LeftNodeBox->ClearChildren();
+
+	URigVMEdGraphNode* RigGraphNode = Cast<URigVMEdGraphNode>(GraphNode);
+	URigVMNode* Node = RigGraphNode->GetModelNode();
+	const FRigVMTemplate* Template = nullptr;
+	if (URigVMTemplateNode* TemplateNode = Cast<URigVMTemplateNode>(Node))
+	{
+		Template = TemplateNode->GetTemplate();
+	}
+
+	TMap<UEdGraphPin*, int32> EdGraphPinToInputPin;
+	for(int32 InputPinIndex = 0; InputPinIndex < InputPins.Num(); InputPinIndex++)
+	{
+		EdGraphPinToInputPin.Add(InputPins[InputPinIndex]->GetPinObj(), InputPinIndex);
+	}
+	TMap<UEdGraphPin*, int32> EdGraphPinToOutputPin;
+	for(int32 OutputPinIndex = 0; OutputPinIndex < OutputPins.Num(); OutputPinIndex++)
+	{
+		EdGraphPinToOutputPin.Add(OutputPins[OutputPinIndex]->GetPinObj(), OutputPinIndex);
+	}
+
+	TArray<URigVMPin*> RootModelPins = ModelNode->GetPins();
+	// orphaned pins are appended to the end of pin list on each side of the node
+	RootModelPins.Append(ModelNode->GetOrphanedPins());
+	TArray<URigVMPin*> ModelPins;
+
+	const bool bSupportSubPins = !RigGraphNode->DrawAsCompactNode();
+
+	// sort model pins
+	// a) execute IOs, b) IO pins, c) input / visible pins, d) output pins
+	struct Local
+	{
+		static void VisitPinRecursively(URigVMPin* InPin, TArray<URigVMPin*>& OutPins, bool bSupportSubPins)
+		{
+			OutPins.Add(InPin);
+			if(!bSupportSubPins)
+			{
+				return;
+			}
+
+			if(InPin->ShouldHideSubPins())
+			{
+				return;
+			}
+
+			if (InPin->GetCPPType() == TEXT("FRotator"))
+			{
+				const TArray<URigVMPin*>& SubPins = InPin->GetSubPins();
+				if (SubPins.Num() == 3)
+				{
+					OutPins.Add(SubPins[2]);	
+					OutPins.Add(SubPins[0]);	
+					OutPins.Add(SubPins[1]);
+				}	
+			}
+			else
+			{				
+				for (URigVMPin* SubPin : InPin->GetSubPins())
+				{
+					VisitPinRecursively(SubPin, OutPins, bSupportSubPins);
+				}
+			}
+		}
+	};
+	
+	for(int32 SortPhase = 0; SortPhase < 4; SortPhase++)
+	{
+		for(URigVMPin* RootPin : RootModelPins)
+		{
+			switch (SortPhase)
+			{
+				case 0: // execute IO pins
+				{
+					if(RootPin->IsExecuteContext() && RootPin->GetDirection() == ERigVMPinDirection::IO)
+					{
+						Local::VisitPinRecursively(RootPin, ModelPins, bSupportSubPins);
+					}
+					break;
+				}
+				case 1: // IO pins
+				{
+					if(!RootPin->IsExecuteContext() && RootPin->GetDirection() == ERigVMPinDirection::IO)
+					{
+						Local::VisitPinRecursively(RootPin, ModelPins, bSupportSubPins);
+					}
+					break;
+				}
+				case 2: // input / visible pins
+				{
+					if(RootPin->GetDirection() == ERigVMPinDirection::Input || RootPin->GetDirection() == ERigVMPinDirection::Visible)
+					{
+						Local::VisitPinRecursively(RootPin, ModelPins, bSupportSubPins);
+					}
+					break;
+				}
+				case 3: // output pins
+				default:
+				{
+					if(RootPin->GetDirection() == ERigVMPinDirection::Output)
+					{
+						Local::VisitPinRecursively(RootPin, ModelPins, bSupportSubPins);
+					}
+					break;
+				}
+			}
+		}
+	}
+
+	const URigVMEdGraphSchema* RigSchema = Cast<URigVMEdGraphSchema>(RigGraphNode->GetSchema());
+
+	FRigVMDispatchContext DispatchContext;
+	if(const URigVMDispatchNode* DispatchNode = Cast<URigVMDispatchNode>(ModelNode))
+	{
+		DispatchContext = DispatchNode->GetDispatchContext();
+	}
+	
+	TMap<URigVMPin*, int32> ModelPinToInfoIndex;
+	for(URigVMPin* ModelPin : ModelPins)
+	{
+		FPinInfo PinInfo;
+		PinInfo.Index = PinInfos.Num();
+		PinInfo.ParentIndex = INDEX_NONE;
+		PinInfo.bHasChildren = (ModelPin->GetSubPins().Num() > 0);
+		PinInfo.bIsContainer = ModelPin->IsArray();
+		PinInfo.Depth = 0;
+		PinInfo.bExpanded = ModelPin->IsExpanded();
+		PinInfo.ModelPinPath = ModelPin->GetPinPath();
+		PinInfo.bAutoHeight = false;
+		PinInfo.bShowOnlySubPins = ModelPin->ShouldOnlyShowSubPins();
+
+		if(!bSupportSubPins)
+		{
+			PinInfo.bHasChildren = false;
+			PinInfo.bIsContainer = false;
+		}
+
+		// Decorator handle pins have special packing that are not user editable, hide their innards
+		if (ModelPin->GetCPPTypeObject() == FAnimNextDecoratorHandle::StaticStruct())
+		{
+			PinInfo.bHasChildren = false;
+			PinInfo.bIsContainer = false;
+		}
+		
+		const bool bAskSchemaForEdition = RigSchema && ModelPin->IsStruct() && !ModelPin->IsBoundToVariable();
+		PinInfo.bHideInputWidget = (!ModelPin->IsBoundToVariable()) && PinInfo.bIsContainer;
+		if (!PinInfo.bHideInputWidget)
+		{
+			if (bAskSchemaForEdition && !PinInfo.bHasChildren)
+			{
+				const bool bIsStructEditable = RigSchema->IsStructEditable(ModelPin->GetScriptStruct()); 
+				PinInfo.bHideInputWidget = !bIsStructEditable;
+				PinInfo.bAutoHeight = bIsStructEditable;
+			}
+			else if(PinInfo.bHasChildren && !ModelPin->IsBoundToVariable())
+			{
+				PinInfo.bHideInputWidget = true;
+			}
+		}
+		
+		if(URigVMPin* ParentPin = ModelPin->GetParentPin())
+		{
+			const int32* ParentIndexPtr = ModelPinToInfoIndex.Find(ParentPin);
+			if(ParentIndexPtr == nullptr)
+			{
+				continue;
+			}
+			PinInfo.ParentIndex = *ParentIndexPtr;
+			PinInfo.Depth = PinInfos[PinInfo.ParentIndex].Depth + 1;
+			if(PinInfos[PinInfo.ParentIndex].bShowOnlySubPins)
+			{
+				PinInfo.Depth--;
+			}
+		}
+
+		TAttribute<EVisibility> PinVisibilityAttribute = TAttribute<EVisibility>::CreateSP(this, &SAnimNextGraphNode::GetPinVisibility, PinInfo.Index, false);
+
+		bool bPinWidgetForExpanderLeft = false;
+		TSharedPtr<SGraphPin> PinWidgetForExpander;
+
+		bool bPinInfoIsValid = false;
+		if(UEdGraphPin* OutputEdGraphPin = RigGraphNode->FindPin(ModelPin->GetPinPath(), EEdGraphPinDirection::EGPD_Output))
+		{
+			if(const int32* PinIndexPtr = EdGraphPinToOutputPin.Find(OutputEdGraphPin))
+			{
+				PinInfo.OutputPinWidget = OutputPins[*PinIndexPtr];
+				PinInfo.OutputPinWidget->SetVisibility(PinVisibilityAttribute);
+				PinWidgetForExpander = PinInfo.OutputPinWidget;
+				bPinWidgetForExpanderLeft = false;
+				bPinInfoIsValid = true;
+				if (Template && !ModelPin->IsExecuteContext())
+				{
+					if (URigVMPin* RootPin = ModelPin->GetRootPin())
+					{
+						FLinearColor PinColorAndOpacity = PinInfo.OutputPinWidget->GetColorAndOpacity();
+						if (Template->FindArgument(RootPin->GetFName()) == nullptr && Template->FindExecuteArgument(RootPin->GetFName(), DispatchContext) == nullptr)
+						{
+							PinColorAndOpacity.A = 0.2f;
+						}
+						else
+						{
+							PinColorAndOpacity.A = 1.0f;
+						}
+						PinInfo.OutputPinWidget->SetColorAndOpacity(PinColorAndOpacity);
+					}
+				}
+			}
+		}
+		
+		if(UEdGraphPin* InputEdGraphPin = RigGraphNode->FindPin(ModelPin->GetPinPath(), EEdGraphPinDirection::EGPD_Input))
+		{
+			if(const int32* PinIndexPtr = EdGraphPinToInputPin.Find(InputEdGraphPin))
+			{
+				PinInfo.InputPinWidget = InputPins[*PinIndexPtr];
+				PinInfo.InputPinWidget->SetVisibility(PinVisibilityAttribute);
+				PinWidgetForExpander = PinInfo.InputPinWidget;
+				bPinWidgetForExpanderLeft = true;
+				bPinInfoIsValid = true;
+				if (Template && !ModelPin->IsExecuteContext())
+				{
+					if (URigVMPin* RootPin = ModelPin->GetRootPin())
+					{
+						FLinearColor PinColorAndOpacity = PinInfo.InputPinWidget->GetColorAndOpacity();
+						if (Template->FindArgument(RootPin->GetFName()) == nullptr && Template->FindExecuteArgument(RootPin->GetFName(), DispatchContext) == nullptr)
+						{
+							PinColorAndOpacity.A = 0.2f;
+						}
+						else
+						{
+							PinColorAndOpacity.A = 1.0f;
+						}
+						PinInfo.InputPinWidget->SetColorAndOpacity(PinColorAndOpacity);
+					}
+				}
+			}
+		}
+
+		if(!bPinInfoIsValid)
+		{
+			continue;
+		}
+
+		ModelPinToInfoIndex.Add(ModelPin, PinInfos.Add(PinInfo));
+		
+		// check if this pin has sub pins
+		TSharedPtr<SHorizontalBox> FullPinHorizontalRowWidget = PinWidgetForExpander->GetFullPinHorizontalRowWidget().Pin();
+		if(FullPinHorizontalRowWidget.IsValid())
+		{
+			// indent the pin by padding
+			const float DepthIndentation = 12.f * float(PinInfo.Depth + (PinInfo.bHasChildren ? 0 : 1));
+			const float LeftIndentation = bPinWidgetForExpanderLeft ? DepthIndentation : 0.f;
+			const float RightIndentation = bPinWidgetForExpanderLeft ? 0.f : DepthIndentation;
+			const FMargin LineIndentation = FMargin(RightIndentation, 0, LeftIndentation, 0); 
+
+			static const TSharedRef<FTagMetaData> ExpanderButtonMetadata = MakeShared<FTagMetaData>(TEXT("SRigVMGraphNode.ExpanderButton"));
+
+			// check if this pin widget may already have the expander button 
+			int32 ExpanderSlotIndex = INDEX_NONE;
+			int32 PaddedSlotIndex = INDEX_NONE;
+
+			for(int32 SlotIndex=0; SlotIndex<FullPinHorizontalRowWidget->NumSlots(); SlotIndex++)
+			{
+				const TSharedRef<SWidget> Widget = FullPinHorizontalRowWidget->GetSlot(SlotIndex).GetWidget();
+				const TSharedPtr<FTagMetaData> Metadata = Widget->GetMetaData<FTagMetaData>();
+				if(Metadata.IsValid())
+				{
+					if(Metadata->Tag == ExpanderButtonMetadata->Tag)
+					{
+						ExpanderSlotIndex = SlotIndex;
+						break;
+					}
+				}
+			}
+
+			// The expander needs to be recreated to adjust for PinInfo.Index changes
+			if (FullPinHorizontalRowWidget->IsValidSlotIndex(ExpanderSlotIndex))
+			{
+				SHorizontalBox::FSlot& Slot = FullPinHorizontalRowWidget->GetSlot(ExpanderSlotIndex);
+				const TSharedRef<SWidget> Widget = Slot.GetWidget();
+				Widget->RemoveMetaData(ExpanderButtonMetadata);
+				FullPinHorizontalRowWidget->RemoveSlot(Widget);
+				ExpanderSlotIndex = INDEX_NONE;
+			}
+	
+			if(PinInfo.bHasChildren && ExpanderSlotIndex == INDEX_NONE)
+			{
+				// only inject the expander arrow for inputs on input / IO
+				// or for output pins
+				if(
+					(
+						(
+							(ModelPin->GetDirection() == ERigVMPinDirection::Input) ||
+							(ModelPin->GetDirection() == ERigVMPinDirection::IO)
+						) &&
+						bPinWidgetForExpanderLeft
+					) ||
+					(
+						(ModelPin->GetDirection() == ERigVMPinDirection::Output) &&
+						(!bPinWidgetForExpanderLeft)
+					)
+				)
+				{
+					// Add the expander arrow
+					FullPinHorizontalRowWidget->InsertSlot(bPinWidgetForExpanderLeft ? 1 : FullPinHorizontalRowWidget->GetChildren()->Num() - 1)
+					.Padding(FMargin(0, 0, 0, 0))
+					.AutoWidth()
+					[
+						SNew(SButton)
+						.ButtonStyle(FAppStyle::Get(), TEXT("SimpleButton"))
+						.ContentPadding(0)
+						.VAlign(VAlign_Center)
+						.HAlign(HAlign_Center)
+						.ClickMethod( EButtonClickMethod::MouseDown )
+						.OnClicked(this, &SAnimNextGraphNode::OnExpanderArrowClicked, PinInfo.Index)
+						.ToolTipText(LOCTEXT("ExpandSubPin", "Expand Pin"))
+						[
+							SNew(SImage)
+							.Image(this, &SAnimNextGraphNode::GetExpanderImage, PinInfo.Index, bPinWidgetForExpanderLeft, false)
+							.ColorAndOpacity(FSlateColor::UseForeground())
+						]
+						.AddMetaData(ExpanderButtonMetadata)
+					];
+				}
+			}
+
+			// adjust the padding
+			{
+				const int32 SlotToAdjustIndex = bPinWidgetForExpanderLeft ? 0 : FullPinHorizontalRowWidget->NumSlots() - 1;
+				SHorizontalBox::FSlot& Slot = FullPinHorizontalRowWidget->GetSlot(SlotToAdjustIndex);
+				Slot.SetPadding(LineIndentation);
+			}
+		}
+	}
+
+	// add spacer widget at the start
+	LeftNodeBox->AddSlot()
+	.HAlign(HAlign_Fill)
+	.VAlign(VAlign_Center)
+	.AutoHeight()
+	[
+		SNew(SSpacer)
+		.Size(FVector2D(1.f, 2.f))
+	];
+
+	auto AddArrayPlusButtonLambda = [this](URigVMPin* InModelPin, TSharedPtr<SHorizontalBox> InSlotLayout, const float InEmptySidePadding)
+	{
+		// add array plus button
+		InSlotLayout->AddSlot()
+		.AutoWidth()
+		.HAlign(HAlign_Left)
+		.Padding(PinWidgetSidePadding, TopPadding, InEmptySidePadding, 0.f)
+		[
+			SNew(SButton)
+			.ContentPadding(0.0f)
+			.ButtonStyle(FAppStyle::Get(), "NoBorder")
+			.OnClicked(this, &SAnimNextGraphNode::HandleAddArrayElement, InModelPin->GetPinPath())
+			.IsEnabled(this, &SGraphNode::IsNodeEditable)
+			.Cursor(EMouseCursor::Default)
+			.Visibility(this, &SAnimNextGraphNode::GetArrayPlusButtonVisibility, InModelPin)
+			.ToolTipText(LOCTEXT("AddArrayElement", "Add Array Element"))
+			[
+				SNew(SHorizontalBox)
+				+SHorizontalBox::Slot()
+				.AutoWidth()
+				.VAlign(VAlign_Center)
+				[
+					SNew(SImage)
+					.Image(FAppStyle::GetBrush(TEXT("Icons.PlusCircle")))
+				]
+			]
+		];
+	};
+
+	struct FPinRowInfo
+	{
+		int32 InputPinInfo;
+		int32 OutputPinInfo;
+		
+		FPinRowInfo()
+			: InputPinInfo(INDEX_NONE)
+			, OutputPinInfo(INDEX_NONE)
+		{}
+	};
+
+	TArray<FPinRowInfo> Rows;
+	Rows.Reserve(PinInfos.Num());
+
+	// build a tree with the pin infos in order as before
+	for(const FPinInfo& PinInfo : PinInfos)
+	{
+		FPinRowInfo Row;
+		if(PinInfo.InputPinWidget.IsValid())
+		{
+			Row.InputPinInfo = PinInfo.Index;
+		}
+		if(PinInfo.OutputPinWidget.IsValid())
+		{
+			Row.OutputPinInfo = PinInfo.Index;
+		}
+		Rows.Add(Row);
+	}
+
+	// compact the rows - re-use input rows for outputs
+	// if their sub-pins count matches
+	if(RigGraphNode->DrawAsCompactNode())
+	{
+		for(int32 RowIndex = 0; RowIndex < Rows.Num();)
+		{
+			if(Rows[RowIndex].OutputPinInfo == INDEX_NONE)
+			{
+				RowIndex++;
+				continue;
+			}
+
+			bool bFoundMatch = false;
+			for(int32 InputRowIndex = 0; InputRowIndex < RowIndex; InputRowIndex++)
+			{
+				if(Rows[InputRowIndex].OutputPinInfo == INDEX_NONE)
+				{
+					const FPinInfo& OutputPinInfo = PinInfos[Rows[RowIndex].OutputPinInfo];
+					const FPinInfo& InputPinInfo = PinInfos[Rows[InputRowIndex].InputPinInfo];
+					if(OutputPinInfo.bHasChildren == false &&
+						InputPinInfo.bHasChildren == false)
+					{
+						Rows[InputRowIndex].OutputPinInfo = Rows[RowIndex].OutputPinInfo;
+						bFoundMatch = true;
+						break;
+					}
+				}
+			}
+
+			if(bFoundMatch)
+			{
+				Rows.RemoveAt(RowIndex);
+			}
+			else
+			{
+				RowIndex++;
+			}
+		}
+	}
+	
+	for(const FPinRowInfo& Row : Rows)
+	{
+		const FPinInfo& InputPinInfo = PinInfos[Row.InputPinInfo == INDEX_NONE ? Row.OutputPinInfo : Row.InputPinInfo];
+		const FPinInfo& OutputPinInfo = PinInfos[Row.OutputPinInfo == INDEX_NONE ? Row.InputPinInfo : Row.OutputPinInfo];
+
+		if(InputPinInfo.InputPinWidget.IsValid())
+		{
+			if(InputPinInfo.bHideInputWidget)
+			{
+				if(InputPinInfo.InputPinWidget->GetValueWidget() != SNullWidget::NullWidget)
+				{
+					InputPinInfo.InputPinWidget->GetValueWidget()->SetVisibility(EVisibility::Collapsed);
+				}
+			}
+				
+			// input pins
+			if(!OutputPinInfo.OutputPinWidget.IsValid())
+			{
+				TSharedPtr<SHorizontalBox> SlotLayout;
+				SHorizontalBox::FSlot* FirstSlot = nullptr;
+
+				const float MyEmptySidePadding = InputPinInfo.bHideInputWidget ? EmptySidePadding : 0.f; 
+				
+				LeftNodeBox->AddSlot()
+                .HAlign(HAlign_Fill)
+				.VAlign(VAlign_Center)
+				.AutoHeight()
+				.MaxHeight(InputPinInfo.bAutoHeight ? TAttribute<float>() : MaxHeight)
+                [
+                    SAssignNew(SlotLayout, SHorizontalBox)
+                    .Visibility(this, &SAnimNextGraphNode::GetPinVisibility, InputPinInfo.Index, false)
+		            
+                    +SHorizontalBox::Slot()
+                    .Expose(FirstSlot)
+                    .FillWidth(1.f)
+                    .HAlign(HAlign_Left)
+                    .Padding(PinWidgetSidePadding, TopPadding, InputPinInfo.bIsContainer ? 0.f : MyEmptySidePadding, 0.f)
+                    [
+                        InputPinInfo.InputPinWidget.ToSharedRef()
+                    ]
+                ];
+
+				if(InputPinInfo.bIsContainer)
+				{
+					URigVMPin* ModelPin = ModelNode->GetGraph()->FindPin(InputPinInfo.ModelPinPath);
+					if(ModelPin)
+					{
+						// make sure to minimize the width of the label
+						FirstSlot->SetAutoWidth();
+						AddArrayPlusButtonLambda(ModelPin, SlotLayout, MyEmptySidePadding);
+					}
+				}
+			}
+			// io pins
+			else
+			{
+				TSharedPtr<SHorizontalBox> SlotLayout;
+				SHorizontalBox::FSlot* FirstSlot = nullptr;
+
+				OutputPinInfo.OutputPinWidget->SetShowLabel(false);
+			
+				LeftNodeBox->AddSlot()
+                .HAlign(HAlign_Fill)
+				.VAlign(VAlign_Center)
+                .AutoHeight()
+				.MaxHeight(InputPinInfo.bAutoHeight ? TAttribute<float>() : MaxHeight)
+                [
+                    SAssignNew(SlotLayout, SHorizontalBox)
+                    .Visibility(this, &SAnimNextGraphNode::GetPinVisibility, OutputPinInfo.Index, false)
+
+                    +SHorizontalBox::Slot()
+					.Expose(FirstSlot)
+                    .FillWidth(1.f)
+                    .HAlign(HAlign_Left)
+                    .VAlign(VAlign_Center)
+                    .Padding(PinWidgetSidePadding, TopPadding, 0.f, 0.f)
+                    [
+                        InputPinInfo.InputPinWidget.ToSharedRef()
+                    ]
+                ];
+
+				if(InputPinInfo.bIsContainer)
+				{
+					URigVMPin* ModelPin = ModelNode->GetGraph()->FindPin(InputPinInfo.ModelPinPath);
+					if(ModelPin)
+					{
+						// make sure to minimize the width of the label
+						FirstSlot->SetAutoWidth();
+						AddArrayPlusButtonLambda(ModelPin, SlotLayout, EmptySidePadding);
+					}
+				}
+
+				SlotLayout->AddSlot()
+				.FillWidth(1.f)
+				.HAlign(HAlign_Right)
+				.VAlign(VAlign_Center)
+				.Padding(0.f, TopPadding, PinWidgetSidePadding, 0.f)
+				[
+					OutputPinInfo.OutputPinWidget.ToSharedRef()
+				];
+			}
+		}
+		// output pins
+		else if(OutputPinInfo.OutputPinWidget.IsValid())
+		{
+			LeftNodeBox->AddSlot()
+            .HAlign(HAlign_Fill)
+			.VAlign(VAlign_Center)
+            .AutoHeight()
+			.MaxHeight(OutputPinInfo.bAutoHeight ? TAttribute<float>() : MaxHeight)
+            [
+            	SNew(SHorizontalBox)
+	            .Visibility(this, &SAnimNextGraphNode::GetPinVisibility, OutputPinInfo.Index, false)
+	            
+				+SHorizontalBox::Slot()
+				.FillWidth(1.f)
+				.HAlign(HAlign_Right)
+				.VAlign(VAlign_Center)
+	            .Padding(EmptySidePadding, TopPadding, PinWidgetSidePadding, 0.f)
+				[
+	                OutputPinInfo.OutputPinWidget.ToSharedRef()
+				]
+            ];
+		}
+	}
+
+#if 0
+	if(URigVMFunctionReferenceNode* FunctionReferenceNode = Cast<URigVMFunctionReferenceNode>(ModelNode))
+	{
+		TWeakObjectPtr<URigVMFunctionReferenceNode> WeakFunctionReferenceNode = FunctionReferenceNode;
+		TWeakObjectPtr<URigVMBlueprint> WeakRigVMBlueprint = Blueprint;
+
+		// add the entries for the variable remapping
+		for(const TSharedPtr<FRigVMExternalVariable>& ExternalVariable : RigGraphNode->ExternalVariables)	// TODO: FIX THIS, PRIVATE?
+		{
+			LeftNodeBox->AddSlot()
+			.HAlign(HAlign_Fill)
+			.VAlign(VAlign_Center)
+			.AutoHeight()
+			.MaxHeight(MaxHeight)
+			[
+				SNew(SHorizontalBox)
+	            
+				+SHorizontalBox::Slot()
+				.AutoWidth()
+				.HAlign(HAlign_Left)
+				.VAlign(VAlign_Center)
+				.Padding(PinWidgetSidePadding, TopPadding, PinWidgetSidePadding, 0.f)
+				[
+					SNew(STextBlock)
+					.Text(FText::FromName(ExternalVariable->Name))
+					.TextStyle(FAppStyle::Get(), NAME_DefaultPinLabelStyle)
+					.ColorAndOpacity(this, &SAnimNextGraphNode::GetVariableLabelTextColor, WeakFunctionReferenceNode, ExternalVariable->Name)
+					.ToolTipText(this, &SAnimNextGraphNode::GetVariableLabelTooltipText, WeakControlRigBlueprint, ExternalVariable->Name)
+				]
+
+				+SHorizontalBox::Slot()
+				.AutoWidth()
+				.HAlign(HAlign_Left)
+				.VAlign(VAlign_Center)
+				.Padding(PinWidgetSidePadding, TopPadding, PinWidgetSidePadding, 0.f)
+				[
+					SNew(SRigVMGraphVariableBinding)
+					.Blueprint(Blueprint.Get())
+					.FunctionReferenceNode(FunctionReferenceNode)
+					.InnerVariableName(ExternalVariable->Name)
+				]
+			];
+		}
+	}
+#endif
+
+	CreateAddPinButton();
+	
+	// add spacer widget at the end
+	LeftNodeBox->AddSlot()
+	.HAlign(HAlign_Fill)
+	.VAlign(VAlign_Center)
+	.AutoHeight()
+	[
+		SNew(SSpacer)
+		.Size(FVector2D(1.f, 4.f))
+	];
+}
+
+#undef LOCTEXT_NAMESPACE
diff --git a/Engine/Plugins/Experimental/AnimNext/Source/AnimNextEditor/Private/Graph/SAnimNextGraphNode.h b/Engine/Plugins/Experimental/AnimNext/Source/AnimNextEditor/Private/Graph/SAnimNextGraphNode.h
new file mode 100644
--- /dev/null
+++ b/Engine/Plugins/Experimental/AnimNext/Source/AnimNextEditor/Private/Graph/SAnimNextGraphNode.h
@@ -0,0 +1,25 @@
+// Copyright Epic Games, Inc. All Rights Reserved.
+
+#pragma once
+
+#include "CoreMinimal.h"
+#include "Widgets/SRigVMGraphNode.h"
+
+class UAnimNextGraph_EdGraphNode;
+
+class SAnimNextGraphNode : public SRigVMGraphNode
+{
+public:
+	SLATE_BEGIN_ARGS(SAnimNextGraphNode)
+		: _GraphNodeObj(nullptr)
+		{}
+
+	SLATE_ARGUMENT(UAnimNextGraph_EdGraphNode*, GraphNodeObj)
+
+	SLATE_END_ARGS()
+
+	void Construct(const FArguments& InArgs);
+
+protected:
+	virtual void UpdatePinTreeView() override;
+};
diff --git a/Engine/Plugins/Experimental/AnimNext/Source/AnimNextTestSuite/Private/AnimNextDecoratorBaseTest.cpp b/Engine/Plugins/Experimental/AnimNext/Source/AnimNextTestSuite/Private/AnimNextDecoratorBaseTest.cpp
--- a/Engine/Plugins/Experimental/AnimNext/Source/AnimNextTestSuite/Private/AnimNextDecoratorBaseTest.cpp
+++ b/Engine/Plugins/Experimental/AnimNext/Source/AnimNextTestSuite/Private/AnimNextDecoratorBaseTest.cpp
@@ -16,6 +16,7 @@
 #include "DecoratorBase/ExecutionContext.h"
 #include "DecoratorBase/IDecoratorInterface.h"
 #include "DecoratorBase/NodeInstance.h"
+#include "DecoratorBase/NodeTemplateBuilder.h"
 #include "DecoratorBase/NodeTemplateRegistry.h"
 
 //****************************************************************************
@@ -296,7 +297,9 @@ bool FAnimationAnimNextRuntimeTest_DecoratorRegistry::RunTest(const FString& InP
 
 	FDecoratorRegistry& Registry = FDecoratorRegistry::Get();
 
-	AddErrorIfFalse(Registry.GetNum() == 0, "FAnimationAnimNextRuntimeTest_DecoratorRegistry -> Registry should be empty");
+	// Some decorators already exist in the engine, keep track of them
+	const uint32 NumAutoRegisteredDecorators = Registry.GetNum();
+
 	AddErrorIfFalse(!Registry.FindHandle(FDecoratorA_Base::DecoratorUID).IsValid(), "FAnimationAnimNextRuntimeTest_DecoratorRegistry -> Registry should not contain our decorator");
 	AddErrorIfFalse(!Registry.FindHandle(FDecoratorAB_Add::DecoratorUID).IsValid(), "FAnimationAnimNextRuntimeTest_DecoratorRegistry -> Registry should not contain our decorator");
 	AddErrorIfFalse(!Registry.FindHandle(FDecoratorAC_Add::DecoratorUID).IsValid(), "FAnimationAnimNextRuntimeTest_DecoratorRegistry -> Registry should not contain our decorator");
@@ -305,7 +308,7 @@ bool FAnimationAnimNextRuntimeTest_DecoratorRegistry::RunTest(const FString& InP
 		// Auto register a decorator
 		AUTO_REGISTER_ANIM_DECORATOR(FDecoratorA_Base)
 
-		AddErrorIfFalse(Registry.GetNum() == 1, "FAnimationAnimNextRuntimeTest_DecoratorRegistry -> Registry should contain 1 decorator");
+		AddErrorIfFalse(Registry.GetNum() == NumAutoRegisteredDecorators + 1, "FAnimationAnimNextRuntimeTest_DecoratorRegistry -> Registry should contain 1 new decorator");
 
 		FDecoratorRegistryHandle HandleA = Registry.FindHandle(FDecoratorA_Base::DecoratorUID);
 		AddErrorIfFalse(HandleA.IsValid(), "FAnimationAnimNextRuntimeTest_DecoratorRegistry -> Decorator should have registered automatically");
@@ -321,7 +324,7 @@ bool FAnimationAnimNextRuntimeTest_DecoratorRegistry::RunTest(const FString& InP
 				// Auto register another decorator
 				AUTO_REGISTER_ANIM_DECORATOR(FDecoratorAB_Add)
 
-				AddErrorIfFalse(Registry.GetNum() == 2, "FAnimationAnimNextRuntimeTest_DecoratorRegistry -> Registry should contain 2 decorators");
+				AddErrorIfFalse(Registry.GetNum() == NumAutoRegisteredDecorators + 2, "FAnimationAnimNextRuntimeTest_DecoratorRegistry -> Registry should contain 2 new decorators");
 
 				FDecoratorRegistryHandle HandleAB = Registry.FindHandle(FDecoratorAB_Add::DecoratorUID);
 				AddErrorIfFalse(HandleAB.IsValid(), "FAnimationAnimNextRuntimeTest_DecoratorRegistry -> Decorator should have registered automatically");
@@ -340,7 +343,7 @@ bool FAnimationAnimNextRuntimeTest_DecoratorRegistry::RunTest(const FString& InP
 						FDecoratorAC_Add DecoratorAC_0;
 						Registry.Register(&DecoratorAC_0);
 
-						AddErrorIfFalse(Registry.GetNum() == 3, "FAnimationAnimNextRuntimeTest_DecoratorRegistry -> Registry should contain 3 decorators");
+						AddErrorIfFalse(Registry.GetNum() == NumAutoRegisteredDecorators + 3, "FAnimationAnimNextRuntimeTest_DecoratorRegistry -> Registry should contain 3 new decorators");
 
 						HandleAC_0 = Registry.FindHandle(FDecoratorAC_Add::DecoratorUID);
 						AddErrorIfFalse(HandleAC_0.IsValid(), "FAnimationAnimNextRuntimeTest_DecoratorRegistry -> Decorator should have registered automatically");
@@ -357,7 +360,7 @@ bool FAnimationAnimNextRuntimeTest_DecoratorRegistry::RunTest(const FString& InP
 							// Unregister our instances
 							Registry.Unregister(&DecoratorAC_0);
 
-							AddErrorIfFalse(Registry.GetNum() == 2, "FAnimationAnimNextRuntimeTest_DecoratorRegistry -> Registry should contain 2 decorators");
+							AddErrorIfFalse(Registry.GetNum() == NumAutoRegisteredDecorators + 2, "FAnimationAnimNextRuntimeTest_DecoratorRegistry -> Registry should contain 2 extra decorators");
 							AddErrorIfFalse(!Registry.FindHandle(FDecoratorAC_Add::DecoratorUID).IsValid(), "FAnimationAnimNextRuntimeTest_DecoratorRegistry -> Decorator should have unregistered");
 						}
 					}
@@ -367,7 +370,7 @@ bool FAnimationAnimNextRuntimeTest_DecoratorRegistry::RunTest(const FString& InP
 						FDecoratorAC_Add DecoratorAC_1;
 						Registry.Register(&DecoratorAC_1);
 
-						AddErrorIfFalse(Registry.GetNum() == 3, "FAnimationAnimNextRuntimeTest_DecoratorRegistry -> Registry should contain 3 decorators");
+						AddErrorIfFalse(Registry.GetNum() == NumAutoRegisteredDecorators + 3, "FAnimationAnimNextRuntimeTest_DecoratorRegistry -> Registry should contain 3 new decorators");
 
 						FDecoratorRegistryHandle HandleAC_1 = Registry.FindHandle(FDecoratorAC_Add::DecoratorUID);
 						AddErrorIfFalse(HandleAC_1.IsValid(), "FAnimationAnimNextRuntimeTest_DecoratorRegistry -> Decorator should have registered automatically");
@@ -385,20 +388,20 @@ bool FAnimationAnimNextRuntimeTest_DecoratorRegistry::RunTest(const FString& InP
 							// Unregister our instances
 							Registry.Unregister(&DecoratorAC_1);
 
-							AddErrorIfFalse(Registry.GetNum() == 2, "FAnimationAnimNextRuntimeTest_DecoratorRegistry -> Registry should contain 2 decorators");
+							AddErrorIfFalse(Registry.GetNum() == NumAutoRegisteredDecorators + 2, "FAnimationAnimNextRuntimeTest_DecoratorRegistry -> Registry should contain 2 extra decorators");
 							AddErrorIfFalse(!Registry.FindHandle(FDecoratorAC_Add::DecoratorUID).IsValid(), "FAnimationAnimNextRuntimeTest_DecoratorRegistry -> Decorator should have unregistered");
 						}
 					}
 				}
 			}
 
-			AddErrorIfFalse(Registry.GetNum() == 1, "FAnimationAnimNextRuntimeTest_DecoratorRegistry -> Registry should contain 1 decorator");
+			AddErrorIfFalse(Registry.GetNum() == NumAutoRegisteredDecorators + 1, "FAnimationAnimNextRuntimeTest_DecoratorRegistry -> Registry should contain 1 extra decorator");
 			AddErrorIfFalse(!Registry.FindHandle(FDecoratorAB_Add::DecoratorUID).IsValid(), "FAnimationAnimNextRuntimeTest_DecoratorRegistry -> Decorator should have unregistered automatically");
 			AddErrorIfFalse(HandleA == Registry.FindHandle(FDecoratorA_Base::DecoratorUID), "FAnimationAnimNextRuntimeTest_DecoratorRegistry -> Decorator handle should not have changed");
 		}
 	}
 
-	AddErrorIfFalse(Registry.GetNum() == 0, "FAnimationAnimNextRuntimeTest_DecoratorRegistry -> All decorators should have unregistered");
+	AddErrorIfFalse(Registry.GetNum() == NumAutoRegisteredDecorators, "FAnimationAnimNextRuntimeTest_DecoratorRegistry -> All decorators should have unregistered");
 	AddErrorIfFalse(!Registry.FindHandle(FDecoratorA_Base::DecoratorUID).IsValid(), "FAnimationAnimNextRuntimeTest_DecoratorRegistry -> Decorator should have unregistered automatically");
 	AddErrorIfFalse(!Registry.FindHandle(FDecoratorAB_Add::DecoratorUID).IsValid(), "FAnimationAnimNextRuntimeTest_DecoratorRegistry -> Decorator should have unregistered automatically");
 	AddErrorIfFalse(!Registry.FindHandle(FDecoratorAC_Add::DecoratorUID).IsValid(), "FAnimationAnimNextRuntimeTest_DecoratorRegistry -> Decorator should have unregistered automatically");
@@ -416,6 +419,7 @@ bool FAnimationAnimNextRuntimeTest_NodeTemplateRegistry::RunTest(const FString&
 	AUTO_REGISTER_ANIM_DECORATOR(FDecoratorAB_Add)
 	AUTO_REGISTER_ANIM_DECORATOR(FDecoratorAC_Add)
 
+	FScopedClearNodeTemplateRegistry ScopedClearNodeTemplateRegistry;
 	FNodeTemplateRegistry& Registry = FNodeTemplateRegistry::Get();
 
 	TArray<FDecoratorUID> NodeTemplateDecoratorList;
@@ -427,7 +431,7 @@ bool FAnimationAnimNextRuntimeTest_NodeTemplateRegistry::RunTest(const FString&
 
 	// Populate our node template registry
 	TArray<uint8> NodeTemplateBuffer0;
-	const FNodeTemplate* NodeTemplate0 = BuildNodeTemplate(NodeTemplateDecoratorList, NodeTemplateBuffer0);
+	const FNodeTemplate* NodeTemplate0 = FNodeTemplateBuilder::BuildNodeTemplate(NodeTemplateDecoratorList, NodeTemplateBuffer0);
 
 	AddErrorIfFalse(Registry.GetNum() == 0, "FAnimationAnimNextRuntimeTest_NodeTemplateRegistry -> Registry should contain any templates");
 
@@ -445,7 +449,7 @@ bool FAnimationAnimNextRuntimeTest_NodeTemplateRegistry::RunTest(const FString&
 
 		// Try and register a duplicate template
 		TArray<uint8> NodeTemplateBuffer1;
-		const FNodeTemplate* NodeTemplate1 = BuildNodeTemplate(NodeTemplateDecoratorList, NodeTemplateBuffer1);
+		const FNodeTemplate* NodeTemplate1 = FNodeTemplateBuilder::BuildNodeTemplate(NodeTemplateDecoratorList, NodeTemplateBuffer1);
 		if (NodeTemplate1 != nullptr)
 		{
 			AddErrorIfFalse(NodeTemplate0 != NodeTemplate1, "FAnimationAnimNextRuntimeTest_NodeTemplateRegistry -> Node template pointers should be different");
@@ -463,7 +467,7 @@ bool FAnimationAnimNextRuntimeTest_NodeTemplateRegistry::RunTest(const FString&
 			NodeTemplateDecoratorList2.Add(FDecoratorAC_Add::DecoratorUID);
 
 			TArray<uint8> NodeTemplateBuffer2;
-			const FNodeTemplate* NodeTemplate2 = BuildNodeTemplate(NodeTemplateDecoratorList2, NodeTemplateBuffer2);
+			const FNodeTemplate* NodeTemplate2 = FNodeTemplateBuilder::BuildNodeTemplate(NodeTemplateDecoratorList2, NodeTemplateBuffer2);
 			if (NodeTemplate2 != nullptr)
 			{
 				AddErrorIfFalse(NodeTemplate0->GetUID() != NodeTemplate2->GetUID(), "FAnimationAnimNextRuntimeTest_NodeTemplateRegistry -> Node template UIDs should be different");
@@ -496,6 +500,7 @@ bool FAnimationAnimNextRuntimeTest_NodeLifetime::RunTest(const FString& InParame
 	AUTO_REGISTER_ANIM_DECORATOR(FDecoratorAB_Add)
 	AUTO_REGISTER_ANIM_DECORATOR(FDecoratorAC_Add)
 
+	FScopedClearNodeTemplateRegistry ScopedClearNodeTemplateRegistry;
 	FNodeTemplateRegistry& Registry = FNodeTemplateRegistry::Get();
 
 	TArray<FDecoratorUID> NodeTemplateDecoratorList;
@@ -507,10 +512,9 @@ bool FAnimationAnimNextRuntimeTest_NodeLifetime::RunTest(const FString& InParame
 
 	// Populate our node template registry
 	TArray<uint8> NodeTemplateBuffer0;
-	const FNodeTemplate* NodeTemplate0 = BuildNodeTemplate(NodeTemplateDecoratorList, NodeTemplateBuffer0);
+	const FNodeTemplate* NodeTemplate0 = FNodeTemplateBuilder::BuildNodeTemplate(NodeTemplateDecoratorList, NodeTemplateBuffer0);
 
 	FNodeTemplateRegistryHandle TemplateHandle0 = Registry.FindOrAdd(NodeTemplate0);
-	AddErrorIfFalse(Registry.GetNum() == 1, "FAnimationAnimNextRuntimeTest_NodeLifetime -> Registry should contain 1 template");
 	AddErrorIfFalse(TemplateHandle0.IsValid(), "FAnimationAnimNextRuntimeTest_NodeLifetime -> Registry should contain our template");
 
 	FNodeHandle Node0;
@@ -525,19 +529,17 @@ bool FAnimationAnimNextRuntimeTest_NodeLifetime::RunTest(const FString& InParame
 		Node1 = DecoratorWriter.RegisterNode(*NodeTemplate0);
 
 		// We don't have decorator properties
-		TArray<TMap<FString, FString>> DecoratorProperties;
-		DecoratorProperties.AddDefaulted(NodeTemplateDecoratorList.Num());
 
 		DecoratorWriter.BeginNodeWriting();
 		DecoratorWriter.WriteNode(Node0,
-			[&DecoratorProperties](uint32 DecoratorIndex) -> const TMap<FString, FString>&
+			[](uint32 DecoratorIndex, const FString& PropertyName)
 			{
-				return DecoratorProperties[DecoratorIndex];
+				return FString();
 			});
 		DecoratorWriter.WriteNode(Node1,
-			[&DecoratorProperties](uint32 DecoratorIndex) -> const TMap<FString, FString>&
+			[](uint32 DecoratorIndex, const FString& PropertyName)
 			{
-				return DecoratorProperties[DecoratorIndex];
+				return FString();
 			});
 		DecoratorWriter.EndNodeWriting();
 
@@ -547,11 +549,16 @@ bool FAnimationAnimNextRuntimeTest_NodeLifetime::RunTest(const FString& InParame
 
 	// Read our graph
 	TArray<uint8> GraphSharedDataBuffer;
+	TArray<TObjectPtr<UObject>> TrackedObjectsForGC;
 	{
 		FMemoryReader GraphSharedDataArchive(GraphSharedDataArchiveBuffer);
 		FDecoratorReader DecoratorReader(GraphSharedDataArchive);
 
-		DecoratorReader.ReadGraphSharedData(GraphSharedDataBuffer);
+		FDecoratorReader::EErrorState ErrorState = DecoratorReader.ReadGraph(GraphSharedDataBuffer, TrackedObjectsForGC);
+		AddErrorIfFalse(ErrorState == FDecoratorReader::EErrorState::None, "FAnimationAnimNextRuntimeTest_NodeLifetime -> Failed to read graph shared data");
+
+		Node0 = DecoratorReader.ResolveNodeHandle(Node0);
+		Node1 = DecoratorReader.ResolveNodeHandle(Node1);
 	}
 
 	FExecutionContext Context(GraphSharedDataBuffer);
@@ -697,6 +704,7 @@ bool FAnimationAnimNextRuntimeTest_GetDecoratorInterface::RunTest(const FString&
 	AUTO_REGISTER_ANIM_DECORATOR(FDecoratorAB_Add)
 	AUTO_REGISTER_ANIM_DECORATOR(FDecoratorAC_Add)
 
+	FScopedClearNodeTemplateRegistry ScopedClearNodeTemplateRegistry;
 	FNodeTemplateRegistry& Registry = FNodeTemplateRegistry::Get();
 
 	TArray<FDecoratorUID> NodeTemplateDecoratorList;
@@ -708,10 +716,9 @@ bool FAnimationAnimNextRuntimeTest_GetDecoratorInterface::RunTest(const FString&
 
 	// Populate our node template registry
 	TArray<uint8> NodeTemplateBuffer0;
-	const FNodeTemplate* NodeTemplate0 = BuildNodeTemplate(NodeTemplateDecoratorList, NodeTemplateBuffer0);
+	const FNodeTemplate* NodeTemplate0 = FNodeTemplateBuilder::BuildNodeTemplate(NodeTemplateDecoratorList, NodeTemplateBuffer0);
 
 	FNodeTemplateRegistryHandle TemplateHandle0 = Registry.FindOrAdd(NodeTemplate0);
-	AddErrorIfFalse(Registry.GetNum() == 1, "FAnimationAnimNextRuntimeTest_GetDecoratorInterface -> Registry should contain 1 template");
 	AddErrorIfFalse(TemplateHandle0.IsValid(), "FAnimationAnimNextRuntimeTest_GetDecoratorInterface -> Registry should contain our template");
 
 	FNodeHandle Node0;
@@ -724,14 +731,12 @@ bool FAnimationAnimNextRuntimeTest_GetDecoratorInterface::RunTest(const FString&
 		Node0 = DecoratorWriter.RegisterNode(*NodeTemplate0);
 
 		// We don't have decorator properties
-		TArray<TMap<FString, FString>> DecoratorProperties;
-		DecoratorProperties.AddDefaulted(NodeTemplateDecoratorList.Num());
 
 		DecoratorWriter.BeginNodeWriting();
 		DecoratorWriter.WriteNode(Node0,
-			[&DecoratorProperties](uint32 DecoratorIndex) -> const TMap<FString, FString>&
+			[](uint32 DecoratorIndex, const FString& PropertyName)
 			{
-				return DecoratorProperties[DecoratorIndex];
+				return FString();
 			});
 		DecoratorWriter.EndNodeWriting();
 
@@ -741,11 +746,15 @@ bool FAnimationAnimNextRuntimeTest_GetDecoratorInterface::RunTest(const FString&
 
 	// Read our graph
 	TArray<uint8> GraphSharedDataBuffer;
+	TArray<TObjectPtr<UObject>> TrackedObjectsForGC;
 	{
 		FMemoryReader GraphSharedDataArchive(GraphSharedDataArchiveBuffer);
 		FDecoratorReader DecoratorReader(GraphSharedDataArchive);
 
-		DecoratorReader.ReadGraphSharedData(GraphSharedDataBuffer);
+		FDecoratorReader::EErrorState ErrorState = DecoratorReader.ReadGraph(GraphSharedDataBuffer, TrackedObjectsForGC);
+		AddErrorIfFalse(ErrorState == FDecoratorReader::EErrorState::None, "FAnimationAnimNextRuntimeTest_GetDecoratorInterface -> Failed to read graph shared data");
+
+		Node0 = DecoratorReader.ResolveNodeHandle(Node0);
 	}
 
 	FExecutionContext Context(GraphSharedDataBuffer);
@@ -960,6 +969,7 @@ bool FAnimationAnimNextRuntimeTest_GetDecoratorInterfaceSuper::RunTest(const FSt
 	AUTO_REGISTER_ANIM_DECORATOR(FDecoratorAB_Add)
 	AUTO_REGISTER_ANIM_DECORATOR(FDecoratorAC_Add)
 
+	FScopedClearNodeTemplateRegistry ScopedClearNodeTemplateRegistry;
 	FNodeTemplateRegistry& Registry = FNodeTemplateRegistry::Get();
 
 	TArray<FDecoratorUID> NodeTemplateDecoratorList;
@@ -971,7 +981,7 @@ bool FAnimationAnimNextRuntimeTest_GetDecoratorInterfaceSuper::RunTest(const FSt
 
 	// Populate our node template registry
 	TArray<uint8> NodeTemplateBuffer0;
-	const FNodeTemplate* NodeTemplate0 = BuildNodeTemplate(NodeTemplateDecoratorList, NodeTemplateBuffer0);
+	const FNodeTemplate* NodeTemplate0 = FNodeTemplateBuilder::BuildNodeTemplate(NodeTemplateDecoratorList, NodeTemplateBuffer0);
 
 	FNodeTemplateRegistryHandle TemplateHandle0 = Registry.FindOrAdd(NodeTemplate0);
 	AddErrorIfFalse(Registry.GetNum() == 1, "FAnimationAnimNextRuntimeTest_GetDecoratorInterfaceSuper -> Registry should contain 1 template");
@@ -987,14 +997,12 @@ bool FAnimationAnimNextRuntimeTest_GetDecoratorInterfaceSuper::RunTest(const FSt
 		Node0 = DecoratorWriter.RegisterNode(*NodeTemplate0);
 
 		// We don't have decorator properties
-		TArray<TMap<FString, FString>> DecoratorProperties;
-		DecoratorProperties.AddDefaulted(NodeTemplateDecoratorList.Num());
 
 		DecoratorWriter.BeginNodeWriting();
 		DecoratorWriter.WriteNode(Node0,
-			[&DecoratorProperties](uint32 DecoratorIndex) -> const TMap<FString, FString>&
+			[](uint32 DecoratorIndex, const FString& PropertyName)
 			{
-				return DecoratorProperties[DecoratorIndex];
+				return FString();
 			});
 		DecoratorWriter.EndNodeWriting();
 
@@ -1004,11 +1012,15 @@ bool FAnimationAnimNextRuntimeTest_GetDecoratorInterfaceSuper::RunTest(const FSt
 
 	// Read our graph
 	TArray<uint8> GraphSharedDataBuffer;
+	TArray<TObjectPtr<UObject>> TrackedObjectsForGC;
 	{
 		FMemoryReader GraphSharedDataArchive(GraphSharedDataArchiveBuffer);
 		FDecoratorReader DecoratorReader(GraphSharedDataArchive);
 
-		DecoratorReader.ReadGraphSharedData(GraphSharedDataBuffer);
+		FDecoratorReader::EErrorState ErrorState = DecoratorReader.ReadGraph(GraphSharedDataBuffer, TrackedObjectsForGC);
+		AddErrorIfFalse(ErrorState == FDecoratorReader::EErrorState::None, "FAnimationAnimNextRuntimeTest_GetDecoratorInterfaceSuper -> Failed to read graph shared data");
+
+		Node0 = DecoratorReader.ResolveNodeHandle(Node0);
 	}
 
 	FExecutionContext Context(GraphSharedDataBuffer);
@@ -1153,6 +1165,7 @@ bool FAnimationAnimNextRuntimeTest_DecoratorSerialization::RunTest(const FString
 	AUTO_REGISTER_ANIM_DECORATOR(FDecoratorSerialization_Add)
 	AUTO_REGISTER_ANIM_DECORATOR(FDecoratorNativeSerialization_Add)
 
+	FScopedClearNodeTemplateRegistry ScopedClearNodeTemplateRegistry;
 	FNodeTemplateRegistry& Registry = FNodeTemplateRegistry::Get();
 
 	TArray<FDecoratorUID> NodeTemplateDecoratorList;
@@ -1162,10 +1175,9 @@ bool FAnimationAnimNextRuntimeTest_DecoratorSerialization::RunTest(const FString
 
 	// Populate our node template registry
 	TArray<uint8> NodeTemplateBuffer0;
-	const FNodeTemplate* NodeTemplate0 = BuildNodeTemplate(NodeTemplateDecoratorList, NodeTemplateBuffer0);
+	const FNodeTemplate* NodeTemplate0 = FNodeTemplateBuilder::BuildNodeTemplate(NodeTemplateDecoratorList, NodeTemplateBuffer0);
 
 	FNodeTemplateRegistryHandle TemplateHandle0 = Registry.FindOrAdd(NodeTemplate0);
-	AddErrorIfFalse(Registry.GetNum() == 1, "FAnimationAnimNextRuntimeTest_DecoratorSerialization -> Registry should contain 1 template");
 	AddErrorIfFalse(TemplateHandle0.IsValid(), "FAnimationAnimNextRuntimeTest_DecoratorSerialization -> Registry should contain our template");
 
 	FDecoratorSerialization_Base::FSharedData DecoratorBaseRef0;
@@ -1326,14 +1338,14 @@ bool FAnimationAnimNextRuntimeTest_DecoratorSerialization::RunTest(const FString
 
 		DecoratorWriter.BeginNodeWriting();
 		DecoratorWriter.WriteNode(Node0,
-			[&DecoratorProperties0](uint32 DecoratorIndex) -> const TMap<FString, FString>&
+			[&DecoratorProperties0](uint32 DecoratorIndex, const FString& PropertyName)
 			{
-				return DecoratorProperties0[DecoratorIndex];
+				return DecoratorProperties0[DecoratorIndex][PropertyName];
 			});
 		DecoratorWriter.WriteNode(Node1,
-			[&DecoratorProperties1](uint32 DecoratorIndex) -> const TMap<FString, FString>&
+			[&DecoratorProperties1](uint32 DecoratorIndex, const FString& PropertyName)
 			{
-				return DecoratorProperties1[DecoratorIndex];
+				return DecoratorProperties1[DecoratorIndex][PropertyName];
 			});
 		DecoratorWriter.EndNodeWriting();
 
@@ -1342,143 +1354,146 @@ bool FAnimationAnimNextRuntimeTest_DecoratorSerialization::RunTest(const FString
 	}
 
 	// Clear out the node template registry to test registration on load
-	Registry.Clear();
-
-	AddErrorIfFalse(Registry.GetNum() == 0, "FAnimationAnimNextRuntimeTest_DecoratorSerialization -> Registry should contain 0 templates");
-
-	// Read our graph
-	TArray<uint8> GraphSharedDataBuffer;
 	{
-		FMemoryReader GraphSharedDataArchive(GraphSharedDataArchiveBuffer);
-		FDecoratorReader DecoratorReader(GraphSharedDataArchive);
+		FScopedClearNodeTemplateRegistry ScopedClearNodeTemplateRegistryForLoad;
 
-		DecoratorReader.ReadGraphSharedData(GraphSharedDataBuffer);
+		AddErrorIfFalse(Registry.GetNum() == 0, "FAnimationAnimNextRuntimeTest_DecoratorSerialization -> Registry should contain 0 templates");
 
-		AddErrorIfFalse(Registry.GetNum() == 1, "FAnimationAnimNextRuntimeTest_DecoratorSerialization -> Registry should contain 1 template");
-	}
+		// Read our graph
+		TArray<uint8> GraphSharedDataBuffer;
+		TArray<TObjectPtr<UObject>> TrackedObjectsForGC;
+		{
+			FMemoryReader GraphSharedDataArchive(GraphSharedDataArchiveBuffer);
+			FDecoratorReader DecoratorReader(GraphSharedDataArchive);
 
-	FExecutionContext Context(GraphSharedDataBuffer);
+			FDecoratorReader::EErrorState ErrorState = DecoratorReader.ReadGraph(GraphSharedDataBuffer, TrackedObjectsForGC);
+			AddErrorIfFalse(ErrorState == FDecoratorReader::EErrorState::None, "FAnimationAnimNextRuntimeTest_DecoratorSerialization -> Failed to read graph shared data");
 
-	// Validate decorator serialization
-	{
-		FDecoratorBinding ParentBinding;				// Empty, no parent
-		FAnimNextDecoratorHandle DecoratorHandle0(Node0, 0);	// Point to first node, first base decorator
-		FAnimNextDecoratorHandle DecoratorHandle1(Node1, 0);	// Point to second node, first base decorator
+			Node0 = DecoratorReader.ResolveNodeHandle(Node0);
+			Node1 = DecoratorReader.ResolveNodeHandle(Node1);
 
-		FDecoratorPtr DecoratorPtr0 = Context.AllocateNodeInstance(ParentBinding, DecoratorHandle0);
-		AddErrorIfFalse(DecoratorPtr0.IsValid(), "FAnimationAnimNextRuntimeTest_DecoratorSerialization -> Failed to allocate a node instance");
+			AddErrorIfFalse(Registry.GetNum() == 1, "FAnimationAnimNextRuntimeTest_DecoratorSerialization -> Registry should contain 1 template");
+		}
 
-		FDecoratorPtr DecoratorPtr1 = Context.AllocateNodeInstance(ParentBinding, DecoratorHandle1);
-		AddErrorIfFalse(DecoratorPtr1.IsValid(), "FAnimationAnimNextRuntimeTest_DecoratorSerialization -> Failed to allocate a node instance");
+		FExecutionContext Context(GraphSharedDataBuffer);
 
-		// Validate shared data for base decorator on node 0
+		// Validate decorator serialization
 		{
-			TDecoratorBinding<IInterfaceA> BindingA0;
-			AddErrorIfFalse(Context.GetInterface(DecoratorPtr0, BindingA0), "FAnimationAnimNextRuntimeTest_DecoratorSerialization -> InterfaceA not found");
-
-			const auto* SharedDataA0 = BindingA0.GetSharedData<FDecoratorSerialization_Base::FSharedData>();
-
-			AddErrorIfFalse(SharedDataA0->Integer == DecoratorBaseRef0.Integer, "FAnimationAnimNextRuntimeTest_DecoratorSerialization -> Unexpected serialized value");
-			AddErrorIfFalse(FMemory::Memcmp(SharedDataA0->IntegerArray, DecoratorBaseRef0.IntegerArray, sizeof(DecoratorBaseRef0.IntegerArray)) == 0, "FAnimationAnimNextRuntimeTest_DecoratorSerialization -> Unexpected serialized value");
-			AddErrorIfFalse(SharedDataA0->IntegerTArray == DecoratorBaseRef0.IntegerTArray, "FAnimationAnimNextRuntimeTest_DecoratorSerialization -> Unexpected serialized value");
-			AddErrorIfFalse(SharedDataA0->Vector == DecoratorBaseRef0.Vector, "FAnimationAnimNextRuntimeTest_DecoratorSerialization -> Unexpected serialized value");
-			AddErrorIfFalse(FMemory::Memcmp(SharedDataA0->VectorArray, DecoratorBaseRef0.VectorArray, sizeof(DecoratorBaseRef0.VectorArray)) == 0, "FAnimationAnimNextRuntimeTest_DecoratorSerialization -> Unexpected serialized value");
-			AddErrorIfFalse(SharedDataA0->VectorTArray == DecoratorBaseRef0.VectorTArray, "FAnimationAnimNextRuntimeTest_DecoratorSerialization -> Unexpected serialized value");
-			AddErrorIfFalse(SharedDataA0->String == DecoratorBaseRef0.String, "FAnimationAnimNextRuntimeTest_DecoratorSerialization -> Unexpected serialized value");
-			AddErrorIfFalse(SharedDataA0->Name == DecoratorBaseRef0.Name, "FAnimationAnimNextRuntimeTest_DecoratorSerialization -> Unexpected serialized value");
-		}
+			FDecoratorBinding ParentBinding;				// Empty, no parent
+			FAnimNextDecoratorHandle DecoratorHandle0(Node0, 0);	// Point to first node, first base decorator
+			FAnimNextDecoratorHandle DecoratorHandle1(Node1, 0);	// Point to second node, first base decorator
 
-		// Validate shared data for additive decorator on node 0
-		{
-			TDecoratorBinding<IInterfaceB> BindingB0;
-			AddErrorIfFalse(Context.GetInterface(DecoratorPtr0, BindingB0), "FAnimationAnimNextRuntimeTest_DecoratorSerialization -> InterfaceB not found");
-
-			const auto* SharedDataB0 = BindingB0.GetSharedData<FDecoratorSerialization_Add::FSharedData>();
-
-			AddErrorIfFalse(SharedDataB0->Integer == DecoratorAddRef0.Integer, "FAnimationAnimNextRuntimeTest_DecoratorSerialization -> Unexpected serialized value");
-			AddErrorIfFalse(FMemory::Memcmp(SharedDataB0->IntegerArray, DecoratorAddRef0.IntegerArray, sizeof(DecoratorAddRef0.IntegerArray)) == 0, "FAnimationAnimNextRuntimeTest_DecoratorSerialization -> Unexpected serialized value");
-			AddErrorIfFalse(SharedDataB0->IntegerTArray == DecoratorAddRef0.IntegerTArray, "FAnimationAnimNextRuntimeTest_DecoratorSerialization -> Unexpected serialized value");
-			AddErrorIfFalse(SharedDataB0->Vector == DecoratorAddRef0.Vector, "FAnimationAnimNextRuntimeTest_DecoratorSerialization -> Unexpected serialized value");
-			AddErrorIfFalse(FMemory::Memcmp(SharedDataB0->VectorArray, DecoratorAddRef0.VectorArray, sizeof(DecoratorAddRef0.VectorArray)) == 0, "FAnimationAnimNextRuntimeTest_DecoratorSerialization -> Unexpected serialized value");
-			AddErrorIfFalse(SharedDataB0->VectorTArray == DecoratorAddRef0.VectorTArray, "FAnimationAnimNextRuntimeTest_DecoratorSerialization -> Unexpected serialized value");
-			AddErrorIfFalse(SharedDataB0->String == DecoratorAddRef0.String, "FAnimationAnimNextRuntimeTest_DecoratorSerialization -> Unexpected serialized value");
-			AddErrorIfFalse(SharedDataB0->Name == DecoratorAddRef0.Name, "FAnimationAnimNextRuntimeTest_DecoratorSerialization -> Unexpected serialized value");
-		}
+			FDecoratorPtr DecoratorPtr0 = Context.AllocateNodeInstance(ParentBinding, DecoratorHandle0);
+			AddErrorIfFalse(DecoratorPtr0.IsValid(), "FAnimationAnimNextRuntimeTest_DecoratorSerialization -> Failed to allocate a node instance");
 
-		// Validate shared data for native decorator on node 0
-		{
-			TDecoratorBinding<IInterfaceC> BindingC0;
-			AddErrorIfFalse(Context.GetInterface(DecoratorPtr0, BindingC0), "FAnimationAnimNextRuntimeTest_DecoratorSerialization -> InterfaceC not found");
-
-			const auto* SharedDataC0 = BindingC0.GetSharedData<FDecoratorNativeSerialization_Add::FSharedData>();
-
-			AddErrorIfFalse(SharedDataC0->Integer == DecoratorNativeRef0.Integer, "FAnimationAnimNextRuntimeTest_DecoratorSerialization -> Unexpected serialized value");
-			AddErrorIfFalse(FMemory::Memcmp(SharedDataC0->IntegerArray, DecoratorNativeRef0.IntegerArray, sizeof(DecoratorNativeRef0.IntegerArray)) == 0, "FAnimationAnimNextRuntimeTest_DecoratorSerialization -> Unexpected serialized value");
-			AddErrorIfFalse(SharedDataC0->IntegerTArray == DecoratorNativeRef0.IntegerTArray, "FAnimationAnimNextRuntimeTest_DecoratorSerialization -> Unexpected serialized value");
-			AddErrorIfFalse(SharedDataC0->Vector == DecoratorNativeRef0.Vector, "FAnimationAnimNextRuntimeTest_DecoratorSerialization -> Unexpected serialized value");
-			AddErrorIfFalse(FMemory::Memcmp(SharedDataC0->VectorArray, DecoratorNativeRef0.VectorArray, sizeof(DecoratorNativeRef0.VectorArray)) == 0, "FAnimationAnimNextRuntimeTest_DecoratorSerialization -> Unexpected serialized value");
-			AddErrorIfFalse(SharedDataC0->VectorTArray == DecoratorNativeRef0.VectorTArray, "FAnimationAnimNextRuntimeTest_DecoratorSerialization -> Unexpected serialized value");
-			AddErrorIfFalse(SharedDataC0->String == DecoratorNativeRef0.String, "FAnimationAnimNextRuntimeTest_DecoratorSerialization -> Unexpected serialized value");
-			AddErrorIfFalse(SharedDataC0->Name == DecoratorNativeRef0.Name, "FAnimationAnimNextRuntimeTest_DecoratorSerialization -> Unexpected serialized value");
-			AddErrorIfFalse(SharedDataC0->bSerializeCalled, "FAnimationAnimNextRuntimeTest_DecoratorSerialization -> Unexpected serialized value");
-		}
+			FDecoratorPtr DecoratorPtr1 = Context.AllocateNodeInstance(ParentBinding, DecoratorHandle1);
+			AddErrorIfFalse(DecoratorPtr1.IsValid(), "FAnimationAnimNextRuntimeTest_DecoratorSerialization -> Failed to allocate a node instance");
 
-		// Validate shared data for base decorator on node 1
-		{
-			TDecoratorBinding<IInterfaceA> BindingA1;
-			AddErrorIfFalse(Context.GetInterface(DecoratorPtr1, BindingA1), "FAnimationAnimNextRuntimeTest_DecoratorSerialization -> InterfaceA not found");
-
-			const auto* SharedDataA1 = BindingA1.GetSharedData<FDecoratorSerialization_Base::FSharedData>();
-
-			AddErrorIfFalse(SharedDataA1->Integer == DecoratorBaseRef1.Integer, "FAnimationAnimNextRuntimeTest_DecoratorSerialization -> Unexpected serialized value");
-			AddErrorIfFalse(FMemory::Memcmp(SharedDataA1->IntegerArray, DecoratorBaseRef1.IntegerArray, sizeof(DecoratorBaseRef1.IntegerArray)) == 0, "FAnimationAnimNextRuntimeTest_DecoratorSerialization -> Unexpected serialized value");
-			AddErrorIfFalse(SharedDataA1->IntegerTArray == DecoratorBaseRef1.IntegerTArray, "FAnimationAnimNextRuntimeTest_DecoratorSerialization -> Unexpected serialized value");
-			AddErrorIfFalse(SharedDataA1->Vector == DecoratorBaseRef1.Vector, "FAnimationAnimNextRuntimeTest_DecoratorSerialization -> Unexpected serialized value");
-			AddErrorIfFalse(FMemory::Memcmp(SharedDataA1->VectorArray, DecoratorBaseRef1.VectorArray, sizeof(DecoratorBaseRef1.VectorArray)) == 0, "FAnimationAnimNextRuntimeTest_DecoratorSerialization -> Unexpected serialized value");
-			AddErrorIfFalse(SharedDataA1->VectorTArray == DecoratorBaseRef1.VectorTArray, "FAnimationAnimNextRuntimeTest_DecoratorSerialization -> Unexpected serialized value");
-			AddErrorIfFalse(SharedDataA1->String == DecoratorBaseRef1.String, "FAnimationAnimNextRuntimeTest_DecoratorSerialization -> Unexpected serialized value");
-			AddErrorIfFalse(SharedDataA1->Name == DecoratorBaseRef1.Name, "FAnimationAnimNextRuntimeTest_DecoratorSerialization -> Unexpected serialized value");
-		}
+			// Validate shared data for base decorator on node 0
+			{
+				TDecoratorBinding<IInterfaceA> BindingA0;
+				AddErrorIfFalse(Context.GetInterface(DecoratorPtr0, BindingA0), "FAnimationAnimNextRuntimeTest_DecoratorSerialization -> InterfaceA not found");
+
+				const auto* SharedDataA0 = BindingA0.GetSharedData<FDecoratorSerialization_Base::FSharedData>();
+
+				AddErrorIfFalse(SharedDataA0->Integer == DecoratorBaseRef0.Integer, "FAnimationAnimNextRuntimeTest_DecoratorSerialization -> Unexpected serialized value");
+				AddErrorIfFalse(FMemory::Memcmp(SharedDataA0->IntegerArray, DecoratorBaseRef0.IntegerArray, sizeof(DecoratorBaseRef0.IntegerArray)) == 0, "FAnimationAnimNextRuntimeTest_DecoratorSerialization -> Unexpected serialized value");
+				AddErrorIfFalse(SharedDataA0->IntegerTArray == DecoratorBaseRef0.IntegerTArray, "FAnimationAnimNextRuntimeTest_DecoratorSerialization -> Unexpected serialized value");
+				AddErrorIfFalse(SharedDataA0->Vector == DecoratorBaseRef0.Vector, "FAnimationAnimNextRuntimeTest_DecoratorSerialization -> Unexpected serialized value");
+				AddErrorIfFalse(FMemory::Memcmp(SharedDataA0->VectorArray, DecoratorBaseRef0.VectorArray, sizeof(DecoratorBaseRef0.VectorArray)) == 0, "FAnimationAnimNextRuntimeTest_DecoratorSerialization -> Unexpected serialized value");
+				AddErrorIfFalse(SharedDataA0->VectorTArray == DecoratorBaseRef0.VectorTArray, "FAnimationAnimNextRuntimeTest_DecoratorSerialization -> Unexpected serialized value");
+				AddErrorIfFalse(SharedDataA0->String == DecoratorBaseRef0.String, "FAnimationAnimNextRuntimeTest_DecoratorSerialization -> Unexpected serialized value");
+				AddErrorIfFalse(SharedDataA0->Name == DecoratorBaseRef0.Name, "FAnimationAnimNextRuntimeTest_DecoratorSerialization -> Unexpected serialized value");
+			}
 
-		// Validate shared data for additive decorator on node 1
-		{
-			TDecoratorBinding<IInterfaceB> BindingB1;
-			AddErrorIfFalse(Context.GetInterface(DecoratorPtr1, BindingB1), "FAnimationAnimNextRuntimeTest_DecoratorSerialization -> InterfaceB not found");
-
-			const auto* SharedDataB1 = BindingB1.GetSharedData<FDecoratorSerialization_Add::FSharedData>();
-
-			AddErrorIfFalse(SharedDataB1->Integer == DecoratorAddRef1.Integer, "FAnimationAnimNextRuntimeTest_DecoratorSerialization -> Unexpected serialized value");
-			AddErrorIfFalse(FMemory::Memcmp(SharedDataB1->IntegerArray, DecoratorAddRef1.IntegerArray, sizeof(DecoratorAddRef1.IntegerArray)) == 0, "FAnimationAnimNextRuntimeTest_DecoratorSerialization -> Unexpected serialized value");
-			AddErrorIfFalse(SharedDataB1->IntegerTArray == DecoratorAddRef1.IntegerTArray, "FAnimationAnimNextRuntimeTest_DecoratorSerialization -> Unexpected serialized value");
-			AddErrorIfFalse(SharedDataB1->Vector == DecoratorAddRef1.Vector, "FAnimationAnimNextRuntimeTest_DecoratorSerialization -> Unexpected serialized value");
-			AddErrorIfFalse(FMemory::Memcmp(SharedDataB1->VectorArray, DecoratorAddRef1.VectorArray, sizeof(DecoratorAddRef1.VectorArray)) == 0, "FAnimationAnimNextRuntimeTest_DecoratorSerialization -> Unexpected serialized value");
-			AddErrorIfFalse(SharedDataB1->VectorTArray == DecoratorAddRef1.VectorTArray, "FAnimationAnimNextRuntimeTest_DecoratorSerialization -> Unexpected serialized value");
-			AddErrorIfFalse(SharedDataB1->String == DecoratorAddRef1.String, "FAnimationAnimNextRuntimeTest_DecoratorSerialization -> Unexpected serialized value");
-			AddErrorIfFalse(SharedDataB1->Name == DecoratorAddRef1.Name, "FAnimationAnimNextRuntimeTest_DecoratorSerialization -> Unexpected serialized value");
-		}
+			// Validate shared data for additive decorator on node 0
+			{
+				TDecoratorBinding<IInterfaceB> BindingB0;
+				AddErrorIfFalse(Context.GetInterface(DecoratorPtr0, BindingB0), "FAnimationAnimNextRuntimeTest_DecoratorSerialization -> InterfaceB not found");
+
+				const auto* SharedDataB0 = BindingB0.GetSharedData<FDecoratorSerialization_Add::FSharedData>();
+
+				AddErrorIfFalse(SharedDataB0->Integer == DecoratorAddRef0.Integer, "FAnimationAnimNextRuntimeTest_DecoratorSerialization -> Unexpected serialized value");
+				AddErrorIfFalse(FMemory::Memcmp(SharedDataB0->IntegerArray, DecoratorAddRef0.IntegerArray, sizeof(DecoratorAddRef0.IntegerArray)) == 0, "FAnimationAnimNextRuntimeTest_DecoratorSerialization -> Unexpected serialized value");
+				AddErrorIfFalse(SharedDataB0->IntegerTArray == DecoratorAddRef0.IntegerTArray, "FAnimationAnimNextRuntimeTest_DecoratorSerialization -> Unexpected serialized value");
+				AddErrorIfFalse(SharedDataB0->Vector == DecoratorAddRef0.Vector, "FAnimationAnimNextRuntimeTest_DecoratorSerialization -> Unexpected serialized value");
+				AddErrorIfFalse(FMemory::Memcmp(SharedDataB0->VectorArray, DecoratorAddRef0.VectorArray, sizeof(DecoratorAddRef0.VectorArray)) == 0, "FAnimationAnimNextRuntimeTest_DecoratorSerialization -> Unexpected serialized value");
+				AddErrorIfFalse(SharedDataB0->VectorTArray == DecoratorAddRef0.VectorTArray, "FAnimationAnimNextRuntimeTest_DecoratorSerialization -> Unexpected serialized value");
+				AddErrorIfFalse(SharedDataB0->String == DecoratorAddRef0.String, "FAnimationAnimNextRuntimeTest_DecoratorSerialization -> Unexpected serialized value");
+				AddErrorIfFalse(SharedDataB0->Name == DecoratorAddRef0.Name, "FAnimationAnimNextRuntimeTest_DecoratorSerialization -> Unexpected serialized value");
+			}
 
-		// Validate shared data for native decorator on node 1
-		{
-			TDecoratorBinding<IInterfaceC> BindingC1;
-			AddErrorIfFalse(Context.GetInterface(DecoratorPtr1, BindingC1), "FAnimationAnimNextRuntimeTest_DecoratorSerialization -> InterfaceC not found");
-
-			const auto* SharedDataC1 = BindingC1.GetSharedData<FDecoratorNativeSerialization_Add::FSharedData>();
-
-			AddErrorIfFalse(SharedDataC1->Integer == DecoratorNativeRef1.Integer, "FAnimationAnimNextRuntimeTest_DecoratorSerialization -> Unexpected serialized value");
-			AddErrorIfFalse(FMemory::Memcmp(SharedDataC1->IntegerArray, DecoratorNativeRef1.IntegerArray, sizeof(DecoratorNativeRef1.IntegerArray)) == 0, "FAnimationAnimNextRuntimeTest_DecoratorSerialization -> Unexpected serialized value");
-			AddErrorIfFalse(SharedDataC1->IntegerTArray == DecoratorNativeRef1.IntegerTArray, "FAnimationAnimNextRuntimeTest_DecoratorSerialization -> Unexpected serialized value");
-			AddErrorIfFalse(SharedDataC1->Vector == DecoratorNativeRef1.Vector, "FAnimationAnimNextRuntimeTest_DecoratorSerialization -> Unexpected serialized value");
-			AddErrorIfFalse(FMemory::Memcmp(SharedDataC1->VectorArray, DecoratorNativeRef1.VectorArray, sizeof(DecoratorNativeRef1.VectorArray)) == 0, "FAnimationAnimNextRuntimeTest_DecoratorSerialization -> Unexpected serialized value");
-			AddErrorIfFalse(SharedDataC1->VectorTArray == DecoratorNativeRef1.VectorTArray, "FAnimationAnimNextRuntimeTest_DecoratorSerialization -> Unexpected serialized value");
-			AddErrorIfFalse(SharedDataC1->String == DecoratorNativeRef1.String, "FAnimationAnimNextRuntimeTest_DecoratorSerialization -> Unexpected serialized value");
-			AddErrorIfFalse(SharedDataC1->Name == DecoratorNativeRef1.Name, "FAnimationAnimNextRuntimeTest_DecoratorSerialization -> Unexpected serialized value");
-			AddErrorIfFalse(SharedDataC1->bSerializeCalled, "FAnimationAnimNextRuntimeTest_DecoratorSerialization -> Unexpected serialized value");
-		}
-	}
+			// Validate shared data for native decorator on node 0
+			{
+				TDecoratorBinding<IInterfaceC> BindingC0;
+				AddErrorIfFalse(Context.GetInterface(DecoratorPtr0, BindingC0), "FAnimationAnimNextRuntimeTest_DecoratorSerialization -> InterfaceC not found");
+
+				const auto* SharedDataC0 = BindingC0.GetSharedData<FDecoratorNativeSerialization_Add::FSharedData>();
+
+				AddErrorIfFalse(SharedDataC0->Integer == DecoratorNativeRef0.Integer, "FAnimationAnimNextRuntimeTest_DecoratorSerialization -> Unexpected serialized value");
+				AddErrorIfFalse(FMemory::Memcmp(SharedDataC0->IntegerArray, DecoratorNativeRef0.IntegerArray, sizeof(DecoratorNativeRef0.IntegerArray)) == 0, "FAnimationAnimNextRuntimeTest_DecoratorSerialization -> Unexpected serialized value");
+				AddErrorIfFalse(SharedDataC0->IntegerTArray == DecoratorNativeRef0.IntegerTArray, "FAnimationAnimNextRuntimeTest_DecoratorSerialization -> Unexpected serialized value");
+				AddErrorIfFalse(SharedDataC0->Vector == DecoratorNativeRef0.Vector, "FAnimationAnimNextRuntimeTest_DecoratorSerialization -> Unexpected serialized value");
+				AddErrorIfFalse(FMemory::Memcmp(SharedDataC0->VectorArray, DecoratorNativeRef0.VectorArray, sizeof(DecoratorNativeRef0.VectorArray)) == 0, "FAnimationAnimNextRuntimeTest_DecoratorSerialization -> Unexpected serialized value");
+				AddErrorIfFalse(SharedDataC0->VectorTArray == DecoratorNativeRef0.VectorTArray, "FAnimationAnimNextRuntimeTest_DecoratorSerialization -> Unexpected serialized value");
+				AddErrorIfFalse(SharedDataC0->String == DecoratorNativeRef0.String, "FAnimationAnimNextRuntimeTest_DecoratorSerialization -> Unexpected serialized value");
+				AddErrorIfFalse(SharedDataC0->Name == DecoratorNativeRef0.Name, "FAnimationAnimNextRuntimeTest_DecoratorSerialization -> Unexpected serialized value");
+				AddErrorIfFalse(SharedDataC0->bSerializeCalled, "FAnimationAnimNextRuntimeTest_DecoratorSerialization -> Unexpected serialized value");
+			}
 
-	Registry.Clear();
+			// Validate shared data for base decorator on node 1
+			{
+				TDecoratorBinding<IInterfaceA> BindingA1;
+				AddErrorIfFalse(Context.GetInterface(DecoratorPtr1, BindingA1), "FAnimationAnimNextRuntimeTest_DecoratorSerialization -> InterfaceA not found");
+
+				const auto* SharedDataA1 = BindingA1.GetSharedData<FDecoratorSerialization_Base::FSharedData>();
+
+				AddErrorIfFalse(SharedDataA1->Integer == DecoratorBaseRef1.Integer, "FAnimationAnimNextRuntimeTest_DecoratorSerialization -> Unexpected serialized value");
+				AddErrorIfFalse(FMemory::Memcmp(SharedDataA1->IntegerArray, DecoratorBaseRef1.IntegerArray, sizeof(DecoratorBaseRef1.IntegerArray)) == 0, "FAnimationAnimNextRuntimeTest_DecoratorSerialization -> Unexpected serialized value");
+				AddErrorIfFalse(SharedDataA1->IntegerTArray == DecoratorBaseRef1.IntegerTArray, "FAnimationAnimNextRuntimeTest_DecoratorSerialization -> Unexpected serialized value");
+				AddErrorIfFalse(SharedDataA1->Vector == DecoratorBaseRef1.Vector, "FAnimationAnimNextRuntimeTest_DecoratorSerialization -> Unexpected serialized value");
+				AddErrorIfFalse(FMemory::Memcmp(SharedDataA1->VectorArray, DecoratorBaseRef1.VectorArray, sizeof(DecoratorBaseRef1.VectorArray)) == 0, "FAnimationAnimNextRuntimeTest_DecoratorSerialization -> Unexpected serialized value");
+				AddErrorIfFalse(SharedDataA1->VectorTArray == DecoratorBaseRef1.VectorTArray, "FAnimationAnimNextRuntimeTest_DecoratorSerialization -> Unexpected serialized value");
+				AddErrorIfFalse(SharedDataA1->String == DecoratorBaseRef1.String, "FAnimationAnimNextRuntimeTest_DecoratorSerialization -> Unexpected serialized value");
+				AddErrorIfFalse(SharedDataA1->Name == DecoratorBaseRef1.Name, "FAnimationAnimNextRuntimeTest_DecoratorSerialization -> Unexpected serialized value");
+			}
+
+			// Validate shared data for additive decorator on node 1
+			{
+				TDecoratorBinding<IInterfaceB> BindingB1;
+				AddErrorIfFalse(Context.GetInterface(DecoratorPtr1, BindingB1), "FAnimationAnimNextRuntimeTest_DecoratorSerialization -> InterfaceB not found");
+
+				const auto* SharedDataB1 = BindingB1.GetSharedData<FDecoratorSerialization_Add::FSharedData>();
+
+				AddErrorIfFalse(SharedDataB1->Integer == DecoratorAddRef1.Integer, "FAnimationAnimNextRuntimeTest_DecoratorSerialization -> Unexpected serialized value");
+				AddErrorIfFalse(FMemory::Memcmp(SharedDataB1->IntegerArray, DecoratorAddRef1.IntegerArray, sizeof(DecoratorAddRef1.IntegerArray)) == 0, "FAnimationAnimNextRuntimeTest_DecoratorSerialization -> Unexpected serialized value");
+				AddErrorIfFalse(SharedDataB1->IntegerTArray == DecoratorAddRef1.IntegerTArray, "FAnimationAnimNextRuntimeTest_DecoratorSerialization -> Unexpected serialized value");
+				AddErrorIfFalse(SharedDataB1->Vector == DecoratorAddRef1.Vector, "FAnimationAnimNextRuntimeTest_DecoratorSerialization -> Unexpected serialized value");
+				AddErrorIfFalse(FMemory::Memcmp(SharedDataB1->VectorArray, DecoratorAddRef1.VectorArray, sizeof(DecoratorAddRef1.VectorArray)) == 0, "FAnimationAnimNextRuntimeTest_DecoratorSerialization -> Unexpected serialized value");
+				AddErrorIfFalse(SharedDataB1->VectorTArray == DecoratorAddRef1.VectorTArray, "FAnimationAnimNextRuntimeTest_DecoratorSerialization -> Unexpected serialized value");
+				AddErrorIfFalse(SharedDataB1->String == DecoratorAddRef1.String, "FAnimationAnimNextRuntimeTest_DecoratorSerialization -> Unexpected serialized value");
+				AddErrorIfFalse(SharedDataB1->Name == DecoratorAddRef1.Name, "FAnimationAnimNextRuntimeTest_DecoratorSerialization -> Unexpected serialized value");
+			}
 
-	AddErrorIfFalse(Registry.GetNum() == 0, "FAnimationAnimNextRuntimeTest_DecoratorSerialization -> Registry should contain 0 templates");
+			// Validate shared data for native decorator on node 1
+			{
+				TDecoratorBinding<IInterfaceC> BindingC1;
+				AddErrorIfFalse(Context.GetInterface(DecoratorPtr1, BindingC1), "FAnimationAnimNextRuntimeTest_DecoratorSerialization -> InterfaceC not found");
+
+				const auto* SharedDataC1 = BindingC1.GetSharedData<FDecoratorNativeSerialization_Add::FSharedData>();
+
+				AddErrorIfFalse(SharedDataC1->Integer == DecoratorNativeRef1.Integer, "FAnimationAnimNextRuntimeTest_DecoratorSerialization -> Unexpected serialized value");
+				AddErrorIfFalse(FMemory::Memcmp(SharedDataC1->IntegerArray, DecoratorNativeRef1.IntegerArray, sizeof(DecoratorNativeRef1.IntegerArray)) == 0, "FAnimationAnimNextRuntimeTest_DecoratorSerialization -> Unexpected serialized value");
+				AddErrorIfFalse(SharedDataC1->IntegerTArray == DecoratorNativeRef1.IntegerTArray, "FAnimationAnimNextRuntimeTest_DecoratorSerialization -> Unexpected serialized value");
+				AddErrorIfFalse(SharedDataC1->Vector == DecoratorNativeRef1.Vector, "FAnimationAnimNextRuntimeTest_DecoratorSerialization -> Unexpected serialized value");
+				AddErrorIfFalse(FMemory::Memcmp(SharedDataC1->VectorArray, DecoratorNativeRef1.VectorArray, sizeof(DecoratorNativeRef1.VectorArray)) == 0, "FAnimationAnimNextRuntimeTest_DecoratorSerialization -> Unexpected serialized value");
+				AddErrorIfFalse(SharedDataC1->VectorTArray == DecoratorNativeRef1.VectorTArray, "FAnimationAnimNextRuntimeTest_DecoratorSerialization -> Unexpected serialized value");
+				AddErrorIfFalse(SharedDataC1->String == DecoratorNativeRef1.String, "FAnimationAnimNextRuntimeTest_DecoratorSerialization -> Unexpected serialized value");
+				AddErrorIfFalse(SharedDataC1->Name == DecoratorNativeRef1.Name, "FAnimationAnimNextRuntimeTest_DecoratorSerialization -> Unexpected serialized value");
+				AddErrorIfFalse(SharedDataC1->bSerializeCalled, "FAnimationAnimNextRuntimeTest_DecoratorSerialization -> Unexpected serialized value");
+			}
+		}
+	}
 
 	return true;
 }
diff --git a/Engine/Plugins/Experimental/AnimNext/Source/AnimNextTestSuite/Private/AnimNextDecoratorInterfacesTest.cpp b/Engine/Plugins/Experimental/AnimNext/Source/AnimNextTestSuite/Private/AnimNextDecoratorInterfacesTest.cpp
--- a/Engine/Plugins/Experimental/AnimNext/Source/AnimNextTestSuite/Private/AnimNextDecoratorInterfacesTest.cpp
+++ b/Engine/Plugins/Experimental/AnimNext/Source/AnimNextTestSuite/Private/AnimNextDecoratorInterfacesTest.cpp
@@ -16,6 +16,7 @@
 #include "DecoratorBase/ExecutionContext.h"
 #include "DecoratorBase/IDecoratorInterface.h"
 #include "DecoratorBase/NodeInstance.h"
+#include "DecoratorBase/NodeTemplateBuilder.h"
 #include "DecoratorBase/NodeTemplateRegistry.h"
 #include "DecoratorInterfaces/IEvaluate.h"
 #include "DecoratorInterfaces/IHierarchy.h"
@@ -212,6 +213,7 @@ bool FAnimationAnimNextRuntimeTest_IHierarchy::RunTest(const FString& InParamete
 	AUTO_REGISTER_ANIM_DECORATOR(FDecoratorWithOneChild)
 	AUTO_REGISTER_ANIM_DECORATOR(FDecoratorWithChildren)
 
+	FScopedClearNodeTemplateRegistry ScopedClearNodeTemplateRegistry;
 	FNodeTemplateRegistry& Registry = FNodeTemplateRegistry::Get();
 
 	// We create a few node templates
@@ -234,10 +236,10 @@ bool FAnimationAnimNextRuntimeTest_IHierarchy::RunTest(const FString& InParamete
 
 	// Populate our node template registry
 	TArray<uint8> NodeTemplateBufferA, NodeTemplateBufferB, NodeTemplateBufferC, NodeTemplateBufferD;
-	const FNodeTemplate* NodeTemplateA = BuildNodeTemplate(NodeTemplateDecoratorListA, NodeTemplateBufferA);
-	const FNodeTemplate* NodeTemplateB = BuildNodeTemplate(NodeTemplateDecoratorListB, NodeTemplateBufferB);
-	const FNodeTemplate* NodeTemplateC = BuildNodeTemplate(NodeTemplateDecoratorListC, NodeTemplateBufferC);
-	const FNodeTemplate* NodeTemplateD = BuildNodeTemplate(NodeTemplateDecoratorListD, NodeTemplateBufferD);
+	const FNodeTemplate* NodeTemplateA = FNodeTemplateBuilder::BuildNodeTemplate(NodeTemplateDecoratorListA, NodeTemplateBufferA);
+	const FNodeTemplate* NodeTemplateB = FNodeTemplateBuilder::BuildNodeTemplate(NodeTemplateDecoratorListB, NodeTemplateBufferB);
+	const FNodeTemplate* NodeTemplateC = FNodeTemplateBuilder::BuildNodeTemplate(NodeTemplateDecoratorListC, NodeTemplateBufferC);
+	const FNodeTemplate* NodeTemplateD = FNodeTemplateBuilder::BuildNodeTemplate(NodeTemplateDecoratorListD, NodeTemplateBufferD);
 
 	// Build our graph, it as follow (each node template has a single node instance):
 	// NodeA has no children
@@ -280,24 +282,24 @@ bool FAnimationAnimNextRuntimeTest_IHierarchy::RunTest(const FString& InParamete
 
 		DecoratorWriter.BeginNodeWriting();
 		DecoratorWriter.WriteNode(NodeA,
-			[&DecoratorPropertiesA](uint32 DecoratorIndex) -> const TMap<FString, FString>&
+			[&DecoratorPropertiesA](uint32 DecoratorIndex, const FString& PropertyName)
 			{
-				return DecoratorPropertiesA[DecoratorIndex];
+				return DecoratorPropertiesA[DecoratorIndex][PropertyName];
 			});
 		DecoratorWriter.WriteNode(NodeB,
-			[&DecoratorPropertiesB](uint32 DecoratorIndex) -> const TMap<FString, FString>&
+			[&DecoratorPropertiesB](uint32 DecoratorIndex, const FString& PropertyName)
 			{
-				return DecoratorPropertiesB[DecoratorIndex];
+				return DecoratorPropertiesB[DecoratorIndex][PropertyName];
 			});
 		DecoratorWriter.WriteNode(NodeC,
-			[&DecoratorPropertiesC](uint32 DecoratorIndex) -> const TMap<FString, FString>&
+			[&DecoratorPropertiesC](uint32 DecoratorIndex, const FString& PropertyName)
 			{
-				return DecoratorPropertiesC[DecoratorIndex];
+				return DecoratorPropertiesC[DecoratorIndex][PropertyName];
 			});
 		DecoratorWriter.WriteNode(NodeD,
-			[&DecoratorPropertiesD](uint32 DecoratorIndex) -> const TMap<FString, FString>&
+			[&DecoratorPropertiesD](uint32 DecoratorIndex, const FString& PropertyName)
 			{
-				return DecoratorPropertiesD[DecoratorIndex];
+				return DecoratorPropertiesD[DecoratorIndex][PropertyName];
 			});
 		DecoratorWriter.EndNodeWriting();
 
@@ -307,11 +309,18 @@ bool FAnimationAnimNextRuntimeTest_IHierarchy::RunTest(const FString& InParamete
 
 	// Read our graph
 	TArray<uint8> GraphSharedDataBuffer;
+	TArray<TObjectPtr<UObject>> TrackedObjectsForGC;
 	{
 		FMemoryReader GraphSharedDataArchive(GraphSharedDataArchiveBuffer);
 		FDecoratorReader DecoratorReader(GraphSharedDataArchive);
 
-		DecoratorReader.ReadGraphSharedData(GraphSharedDataBuffer);
+		FDecoratorReader::EErrorState ErrorState = DecoratorReader.ReadGraph(GraphSharedDataBuffer, TrackedObjectsForGC);
+		AddErrorIfFalse(ErrorState == FDecoratorReader::EErrorState::None, "FAnimationAnimNextRuntimeTest_IHierarchy -> Failed to read graph shared data");
+
+		NodeA = DecoratorReader.ResolveNodeHandle(NodeA);
+		NodeB = DecoratorReader.ResolveNodeHandle(NodeB);
+		NodeC = DecoratorReader.ResolveNodeHandle(NodeC);
+		NodeD = DecoratorReader.ResolveNodeHandle(NodeD);
 	}
 
 	FExecutionContext Context(GraphSharedDataBuffer);
@@ -398,6 +407,7 @@ bool FAnimationAnimNextRuntimeTest_IUpdate::RunTest(const FString& InParameters)
 	AUTO_REGISTER_ANIM_DECORATOR(FDecoratorWithOneChild)
 	AUTO_REGISTER_ANIM_DECORATOR(FDecoratorWithChildren)
 
+	FScopedClearNodeTemplateRegistry ScopedClearNodeTemplateRegistry;
 	FNodeTemplateRegistry& Registry = FNodeTemplateRegistry::Get();
 
 	// We create a few node templates
@@ -415,9 +425,9 @@ bool FAnimationAnimNextRuntimeTest_IUpdate::RunTest(const FString& InParameters)
 
 	// Populate our node template registry
 	TArray<uint8> NodeTemplateBufferA, NodeTemplateBufferB, NodeTemplateBufferC;
-	const FNodeTemplate* NodeTemplateA = BuildNodeTemplate(NodeTemplateDecoratorListA, NodeTemplateBufferA);
-	const FNodeTemplate* NodeTemplateB = BuildNodeTemplate(NodeTemplateDecoratorListB, NodeTemplateBufferB);
-	const FNodeTemplate* NodeTemplateC = BuildNodeTemplate(NodeTemplateDecoratorListC, NodeTemplateBufferC);
+	const FNodeTemplate* NodeTemplateA = FNodeTemplateBuilder::BuildNodeTemplate(NodeTemplateDecoratorListA, NodeTemplateBufferA);
+	const FNodeTemplate* NodeTemplateB = FNodeTemplateBuilder::BuildNodeTemplate(NodeTemplateDecoratorListB, NodeTemplateBufferB);
+	const FNodeTemplate* NodeTemplateC = FNodeTemplateBuilder::BuildNodeTemplate(NodeTemplateDecoratorListC, NodeTemplateBufferC);
 
 	// Build our graph, it as follow (each node template has a single node instance):
 	// NodeA has no children
@@ -452,19 +462,19 @@ bool FAnimationAnimNextRuntimeTest_IUpdate::RunTest(const FString& InParameters)
 
 		DecoratorWriter.BeginNodeWriting();
 		DecoratorWriter.WriteNode(NodeA,
-			[&DecoratorPropertiesA](uint32 DecoratorIndex) -> const TMap<FString, FString>&
+			[&DecoratorPropertiesA](uint32 DecoratorIndex, const FString& PropertyName)
 			{
-				return DecoratorPropertiesA[DecoratorIndex];
+				return DecoratorPropertiesA[DecoratorIndex][PropertyName];
 			});
 		DecoratorWriter.WriteNode(NodeB,
-			[&DecoratorPropertiesB](uint32 DecoratorIndex) -> const TMap<FString, FString>&
+			[&DecoratorPropertiesB](uint32 DecoratorIndex, const FString& PropertyName)
 			{
-				return DecoratorPropertiesB[DecoratorIndex];
+				return DecoratorPropertiesB[DecoratorIndex][PropertyName];
 			});
 		DecoratorWriter.WriteNode(NodeC,
-			[&DecoratorPropertiesC](uint32 DecoratorIndex) -> const TMap<FString, FString>&
+			[&DecoratorPropertiesC](uint32 DecoratorIndex, const FString& PropertyName)
 			{
-				return DecoratorPropertiesC[DecoratorIndex];
+				return DecoratorPropertiesC[DecoratorIndex][PropertyName];
 			});
 		DecoratorWriter.EndNodeWriting();
 
@@ -474,11 +484,17 @@ bool FAnimationAnimNextRuntimeTest_IUpdate::RunTest(const FString& InParameters)
 
 	// Read our graph
 	TArray<uint8> GraphSharedDataBuffer;
+	TArray<TObjectPtr<UObject>> TrackedObjectsForGC;
 	{
 		FMemoryReader GraphSharedDataArchive(GraphSharedDataArchiveBuffer);
 		FDecoratorReader DecoratorReader(GraphSharedDataArchive);
 
-		DecoratorReader.ReadGraphSharedData(GraphSharedDataBuffer);
+		FDecoratorReader::EErrorState ErrorState = DecoratorReader.ReadGraph(GraphSharedDataBuffer, TrackedObjectsForGC);
+		AddErrorIfFalse(ErrorState == FDecoratorReader::EErrorState::None, "FAnimationAnimNextRuntimeTest_IUpdate -> Failed to read graph shared data");
+
+		NodeA = DecoratorReader.ResolveNodeHandle(NodeA);
+		NodeB = DecoratorReader.ResolveNodeHandle(NodeB);
+		NodeC = DecoratorReader.ResolveNodeHandle(NodeC);
 	}
 
 	FExecutionContext Context(GraphSharedDataBuffer);
@@ -495,7 +511,7 @@ bool FAnimationAnimNextRuntimeTest_IUpdate::RunTest(const FString& InParameters)
 		AddErrorIfFalse(NodeCPtr.IsValid(), "FAnimationAnimNextRuntimeTest_IUpdate -> Failed to allocate root node instance");
 
 		// Call pre/post update on our graph
-		UpdateGraph(Context, NodeCPtr);
+		UpdateGraph(Context, NodeCPtr, 0.0333f);
 
 		AddErrorIfFalse(UpdatedDecorators.Num() == 6, "FAnimationAnimNextRuntimeTest_IUpdate -> Expected 6 nodes to have been visited during the update traversal");
 		AddErrorIfFalse(UpdatedDecorators[0] == FDecoratorWithChildren::DecoratorUID, "FAnimationAnimNextRuntimeTest_IUpdate -> Unexpected update order");		// NodeC
@@ -527,6 +543,7 @@ bool FAnimationAnimNextRuntimeTest_IEvaluate::RunTest(const FString& InParameter
 	AUTO_REGISTER_ANIM_DECORATOR(FDecoratorWithOneChild)
 	AUTO_REGISTER_ANIM_DECORATOR(FDecoratorWithChildren)
 
+	FScopedClearNodeTemplateRegistry ScopedClearNodeTemplateRegistry;
 	FNodeTemplateRegistry& Registry = FNodeTemplateRegistry::Get();
 
 	// We create a few node templates
@@ -544,9 +561,9 @@ bool FAnimationAnimNextRuntimeTest_IEvaluate::RunTest(const FString& InParameter
 
 	// Populate our node template registry
 	TArray<uint8> NodeTemplateBufferA, NodeTemplateBufferB, NodeTemplateBufferC;
-	const FNodeTemplate* NodeTemplateA = BuildNodeTemplate(NodeTemplateDecoratorListA, NodeTemplateBufferA);
-	const FNodeTemplate* NodeTemplateB = BuildNodeTemplate(NodeTemplateDecoratorListB, NodeTemplateBufferB);
-	const FNodeTemplate* NodeTemplateC = BuildNodeTemplate(NodeTemplateDecoratorListC, NodeTemplateBufferC);
+	const FNodeTemplate* NodeTemplateA = FNodeTemplateBuilder::BuildNodeTemplate(NodeTemplateDecoratorListA, NodeTemplateBufferA);
+	const FNodeTemplate* NodeTemplateB = FNodeTemplateBuilder::BuildNodeTemplate(NodeTemplateDecoratorListB, NodeTemplateBufferB);
+	const FNodeTemplate* NodeTemplateC = FNodeTemplateBuilder::BuildNodeTemplate(NodeTemplateDecoratorListC, NodeTemplateBufferC);
 
 	// Build our graph, it as follow (each node template has a single node instance):
 	// NodeA has no children
@@ -582,19 +599,19 @@ bool FAnimationAnimNextRuntimeTest_IEvaluate::RunTest(const FString& InParameter
 
 		DecoratorWriter.BeginNodeWriting();
 		DecoratorWriter.WriteNode(NodeA,
-			[&DecoratorPropertiesA](uint32 DecoratorIndex) -> const TMap<FString, FString>&
+			[&DecoratorPropertiesA](uint32 DecoratorIndex, const FString& PropertyName)
 			{
-				return DecoratorPropertiesA[DecoratorIndex];
+				return DecoratorPropertiesA[DecoratorIndex][PropertyName];
 			});
 		DecoratorWriter.WriteNode(NodeB,
-			[&DecoratorPropertiesB](uint32 DecoratorIndex) -> const TMap<FString, FString>&
+			[&DecoratorPropertiesB](uint32 DecoratorIndex, const FString& PropertyName)
 			{
-				return DecoratorPropertiesB[DecoratorIndex];
+				return DecoratorPropertiesB[DecoratorIndex][PropertyName];
 			});
 		DecoratorWriter.WriteNode(NodeC,
-			[&DecoratorPropertiesC](uint32 DecoratorIndex) -> const TMap<FString, FString>&
+			[&DecoratorPropertiesC](uint32 DecoratorIndex, const FString& PropertyName)
 			{
-				return DecoratorPropertiesC[DecoratorIndex];
+				return DecoratorPropertiesC[DecoratorIndex][PropertyName];
 			});
 		DecoratorWriter.EndNodeWriting();
 
@@ -604,11 +621,17 @@ bool FAnimationAnimNextRuntimeTest_IEvaluate::RunTest(const FString& InParameter
 
 	// Read our graph
 	TArray<uint8> GraphSharedDataBuffer;
+	TArray<TObjectPtr<UObject>> TrackedObjectsForGC;
 	{
 		FMemoryReader GraphSharedDataArchive(GraphSharedDataArchiveBuffer);
 		FDecoratorReader DecoratorReader(GraphSharedDataArchive);
 
-		DecoratorReader.ReadGraphSharedData(GraphSharedDataBuffer);
+		FDecoratorReader::EErrorState ErrorState = DecoratorReader.ReadGraph(GraphSharedDataBuffer, TrackedObjectsForGC);
+		AddErrorIfFalse(ErrorState == FDecoratorReader::EErrorState::None, "FAnimationAnimNextRuntimeTest_IEvaluate -> Failed to read graph shared data");
+
+		NodeA = DecoratorReader.ResolveNodeHandle(NodeA);
+		NodeB = DecoratorReader.ResolveNodeHandle(NodeB);
+		NodeC = DecoratorReader.ResolveNodeHandle(NodeC);
 	}
 
 	FExecutionContext Context(GraphSharedDataBuffer);
diff --git a/Engine/Plugins/Experimental/AnimNext/Source/AnimNextTestSuite/Private/AnimNextRuntimeTest.cpp b/Engine/Plugins/Experimental/AnimNext/Source/AnimNextTestSuite/Private/AnimNextRuntimeTest.cpp
--- a/Engine/Plugins/Experimental/AnimNext/Source/AnimNextTestSuite/Private/AnimNextRuntimeTest.cpp
+++ b/Engine/Plugins/Experimental/AnimNext/Source/AnimNextTestSuite/Private/AnimNextRuntimeTest.cpp
@@ -7,112 +7,18 @@
 
 #if WITH_DEV_AUTOMATION_TESTS
 
-#include "DecoratorBase/Decorator.h"
-#include "DecoratorBase/DecoratorRegistry.h"
-#include "DecoratorBase/NodeDescription.h"
-#include "DecoratorBase/NodeInstance.h"
-
 namespace UE::AnimNext
 {
-	uint32 GetNodeTemplateUID(const TArray<FDecoratorUID>& NodeTemplateDecoratorList)
-	{
-		uint32 NodeTemplateUID = 0;
-		for (const FDecoratorUID DecoratorUID : NodeTemplateDecoratorList)
-		{
-			NodeTemplateUID = HashCombineFast(NodeTemplateUID, DecoratorUID.GetUID());
-		}
-
-		return NodeTemplateUID;
-	}
-
-	uint32 GetNodeInstanceSize(const TArray<FDecoratorUID>& NodeTemplateDecoratorList)
-	{
-		const FDecoratorRegistry& Registry = FDecoratorRegistry::Get();
-
-		uint32 NodeInstanceSize = sizeof(FNodeInstance);
-		for (const FDecoratorUID DecoratorUID : NodeTemplateDecoratorList)
-		{
-			const FDecorator* Decorator = Registry.Find(DecoratorUID);
-
-			const FDecoratorMemoryLayout MemoryLayout = Decorator->GetDecoratorMemoryDescription();
-			NodeInstanceSize = Align(NodeInstanceSize, MemoryLayout.InstanceDataAlignment);
-			NodeInstanceSize += MemoryLayout.InstanceDataSize;
-		}
-
-		return NodeInstanceSize;
-	}
-
-	void AppendTemplateDecorator(const TArray<FDecoratorUID>& NodeTemplateDecoratorList, int32 DecoratorIndex, TArray<uint8>& NodeTemplateBuffer, uint32& SharedDataOffset, uint32& InstanceDataOffset)
+	FScopedClearNodeTemplateRegistry::FScopedClearNodeTemplateRegistry()
 	{
-		const FDecoratorRegistry& Registry = FDecoratorRegistry::Get();
-
-		const FDecoratorUID DecoratorUID = NodeTemplateDecoratorList[DecoratorIndex];
-		const FDecoratorRegistryHandle DecoratorHandle = Registry.FindHandle(DecoratorUID);
-		const FDecorator* Decorator = Registry.Find(DecoratorHandle);
-		const EDecoratorMode DecoratorMode = Decorator->GetDecoratorMode();
-
-		uint32 AdditiveIndexOrNumAdditive;
-		if (DecoratorMode == EDecoratorMode::Base)
-		{
-			// Find out how many additive decorators we have
-			AdditiveIndexOrNumAdditive = 0;
-			for (int32 Index = DecoratorIndex + 1; Index < NodeTemplateDecoratorList.Num(); ++Index)	// Skip ourself
-			{
-				const FDecorator* ChildDecorator = Registry.Find(NodeTemplateDecoratorList[Index]);
-				if (ChildDecorator->GetDecoratorMode() == EDecoratorMode::Base)
-				{
-					break;	// Found another base decorator, we are done
-				}
-
-				// We are additive
-				AdditiveIndexOrNumAdditive++;
-			}
-		}
-		else
-		{
-			// Find out our additive index
-			AdditiveIndexOrNumAdditive = 1;	// Skip ourself
-			for (int32 Index = DecoratorIndex - 1; Index >= 0; --Index)
-			{
-				const FDecorator* ParentDecorator = Registry.Find(NodeTemplateDecoratorList[Index]);
-				if (ParentDecorator->GetDecoratorMode() == EDecoratorMode::Base)
-				{
-					break;	// Found our base decorator, we are done
-				}
-
-				// We are additive
-				AdditiveIndexOrNumAdditive++;
-			}
-		}
-
-		const FDecoratorMemoryLayout MemoryLayout = Decorator->GetDecoratorMemoryDescription();
-
-		// Align our data
-		SharedDataOffset = Align(SharedDataOffset, MemoryLayout.SharedDataAlignment);
-		InstanceDataOffset = Align(InstanceDataOffset, MemoryLayout.InstanceDataAlignment);
-
-		// Append and update our offsets
-		const int32 BufferIndex = NodeTemplateBuffer.AddUninitialized(sizeof(FDecoratorTemplate));
-		new(&NodeTemplateBuffer[BufferIndex]) FDecoratorTemplate(DecoratorUID, DecoratorHandle, DecoratorMode, AdditiveIndexOrNumAdditive, SharedDataOffset, InstanceDataOffset);
-
-		SharedDataOffset += MemoryLayout.SharedDataSize;
-		InstanceDataOffset += MemoryLayout.InstanceDataSize;
+		FNodeTemplateRegistry& Registry = FNodeTemplateRegistry::Get();
+		Swap(Registry, TmpRegistry);
 	}
 
-	const FNodeTemplate* BuildNodeTemplate(const TArray<FDecoratorUID>& NodeTemplateDecoratorList, TArray<uint8>& NodeTemplateBuffer)
+	FScopedClearNodeTemplateRegistry::~FScopedClearNodeTemplateRegistry()
 	{
-		const int32 BufferIndex = NodeTemplateBuffer.AddUninitialized(sizeof(FNodeTemplate));
-		new(&NodeTemplateBuffer[BufferIndex]) FNodeTemplate(GetNodeTemplateUID(NodeTemplateDecoratorList), GetNodeInstanceSize(NodeTemplateDecoratorList), NodeTemplateDecoratorList.Num());
-
-		uint32 SharedDataOffset = sizeof(FNodeDescription);
-		uint32 InstanceDataOffset = sizeof(FNodeInstance);
-
-		for (int32 DecoratorIndex = 0; DecoratorIndex < NodeTemplateDecoratorList.Num(); ++DecoratorIndex)
-		{
-			AppendTemplateDecorator(NodeTemplateDecoratorList, DecoratorIndex, NodeTemplateBuffer, SharedDataOffset, InstanceDataOffset);
-		}
-
-		return reinterpret_cast<const FNodeTemplate*>(&NodeTemplateBuffer[0]);
+		FNodeTemplateRegistry& Registry = FNodeTemplateRegistry::Get();
+		Swap(Registry, TmpRegistry);
 	}
 }
 #endif
diff --git a/Engine/Plugins/Experimental/AnimNext/Source/AnimNextTestSuite/Private/AnimNextRuntimeTest.h b/Engine/Plugins/Experimental/AnimNext/Source/AnimNextTestSuite/Private/AnimNextRuntimeTest.h
--- a/Engine/Plugins/Experimental/AnimNext/Source/AnimNextTestSuite/Private/AnimNextRuntimeTest.h
+++ b/Engine/Plugins/Experimental/AnimNext/Source/AnimNextTestSuite/Private/AnimNextRuntimeTest.h
@@ -11,64 +11,70 @@
 
 #include "DecoratorBase/DecoratorUID.h"
 #include "DecoratorBase/NodeTemplate.h"
+#include "DecoratorBase/NodeTemplateRegistry.h"
 
 namespace UE::AnimNext
 {
-	// Builds a node template from a list of decorator UID
-	const FNodeTemplate* BuildNodeTemplate(const TArray<FDecoratorUID>& NodeTemplateDecoratorList, TArray<uint8>& NodeTemplateBuffer);
-
 	// Converts a property value into its string representation using UE reflection
 	template<class DecoratorSharedDataType, typename PropertyType>
 	static FString ToString(const FString& PropertyName, PropertyType PropertyValue)
 	{
 		const UScriptStruct* SharedDataStruct = DecoratorSharedDataType::StaticStruct();
-
-		for (const FProperty* Property = SharedDataStruct->PropertyLink; Property != nullptr; Property = Property->PropertyLinkNext)
+		if (const FProperty* Property = SharedDataStruct->FindPropertyByName(*PropertyName))
 		{
-			if (Property->GetName() == PropertyName)
-			{
-				void* PropertyDefaults = Property->AllocateAndInitializeValue();
+			void* PropertyDefaults = Property->AllocateAndInitializeValue();
 
-				FString Result;
+			FString Result;
 
-				if constexpr (std::is_pointer_v<PropertyType>)
-				{
-					// C-style array properties aren't handled by ExportText, we need to handle it manually
-
-					const bool bIsCArray = Property->ArrayDim > 1;
-					if (bIsCArray)
-					{
-						Result += TEXT("(");
-					}
+			if constexpr (std::is_pointer_v<PropertyType>)
+			{
+				// C-style array properties aren't handled by ExportText, we need to handle it manually
 
-					for (int32 Index = 0; Index < Property->ArrayDim; ++Index)
-					{
-						Property->ExportText_Direct(Result, PropertyValue + Index, (PropertyType)PropertyDefaults + Index, nullptr, PPF_None, nullptr);
+				const bool bIsCArray = Property->ArrayDim > 1;
+				if (bIsCArray)
+				{
+					Result += TEXT("(");
+				}
 
-						if (Index + 1 < Property->ArrayDim)
-						{
-							Result += TEXT(",");
-						}
-					}
+				for (int32 Index = 0; Index < Property->ArrayDim; ++Index)
+				{
+					Property->ExportText_Direct(Result, PropertyValue + Index, (PropertyType)PropertyDefaults + Index, nullptr, PPF_None, nullptr);
 
-					if (bIsCArray)
+					if (Index + 1 < Property->ArrayDim)
 					{
-						Result += TEXT(")");
+						Result += TEXT(",");
 					}
 				}
-				else
+
+				if (bIsCArray)
 				{
-					Property->ExportText_Direct(Result, &PropertyValue, PropertyDefaults, nullptr, PPF_None);
+					Result += TEXT(")");
 				}
+			}
+			else
+			{
+				Property->ExportText_Direct(Result, &PropertyValue, PropertyDefaults, nullptr, PPF_None);
+			}
 
-				Property->DestroyAndFreeValue(PropertyDefaults);
+			Property->DestroyAndFreeValue(PropertyDefaults);
 
-				return Result;
-			}
+			return Result;
 		}
 
 		return FString();
 	}
+
+	// Creates a temporary empty node template registry instance and swaps it for the current one
+	struct FScopedClearNodeTemplateRegistry final
+	{
+		FScopedClearNodeTemplateRegistry(const FScopedClearNodeTemplateRegistry&) = delete;
+		FScopedClearNodeTemplateRegistry& operator=(const FScopedClearNodeTemplateRegistry&) = delete;
+
+		FScopedClearNodeTemplateRegistry();
+		~FScopedClearNodeTemplateRegistry();
+
+		FNodeTemplateRegistry TmpRegistry;
+	};
 }
 
 #endif
diff --git a/Engine/Plugins/Experimental/AnimNext/Source/AnimNextUncookedOnly/AnimNextUncookedOnly.Build.cs b/Engine/Plugins/Experimental/AnimNext/Source/AnimNextUncookedOnly/AnimNextUncookedOnly.Build.cs
--- a/Engine/Plugins/Experimental/AnimNext/Source/AnimNextUncookedOnly/AnimNextUncookedOnly.Build.cs
+++ b/Engine/Plugins/Experimental/AnimNext/Source/AnimNextUncookedOnly/AnimNextUncookedOnly.Build.cs
@@ -29,7 +29,8 @@ namespace UnrealBuildTool.Rules
 					"Kismet",
 					"Slate",
 					"SlateCore",
-					"StructUtils"
+					"StructUtils",
+					"ToolMenus",
 				}
 			);
 
diff --git a/Engine/Plugins/Experimental/AnimNext/Source/AnimNextUncookedOnly/Internal/Graph/AnimNextGraph_EdGraph.h b/Engine/Plugins/Experimental/AnimNext/Source/AnimNextUncookedOnly/Internal/Graph/AnimNextGraph_EdGraph.h
--- a/Engine/Plugins/Experimental/AnimNext/Source/AnimNextUncookedOnly/Internal/Graph/AnimNextGraph_EdGraph.h
+++ b/Engine/Plugins/Experimental/AnimNext/Source/AnimNextUncookedOnly/Internal/Graph/AnimNextGraph_EdGraph.h
@@ -3,7 +3,7 @@
 #pragma once
 
 #include "CoreMinimal.h"
-#include "Graph/ControlRigGraph.h"
+#include "EdGraph/RigVMEdGraph.h"
 #include "AnimNextGraph_EdGraph.generated.h"
 
 class UAnimNextGraph_EditorData;
@@ -13,14 +13,17 @@ namespace UE::AnimNext::UncookedOnly
 	struct FUtils;
 }
 
+/**
+  * Wraps UEdGraph which represents the node graph
+  */
 UCLASS(MinimalAPI)
-class UAnimNextGraph_EdGraph : public UControlRigGraph
+class UAnimNextGraph_EdGraph : public URigVMEdGraph
 {
 	GENERATED_BODY()
 
 	friend class UAnimNextGraph_EditorData;
 
-	// UControlRigGraph interface
+	// URigVMEdGraph interface
 	virtual FRigVMClient* GetRigVMClient() const override;
 	
 	void Initialize(UAnimNextGraph_EditorData* InEditorData);
diff --git a/Engine/Plugins/Experimental/AnimNext/Source/AnimNextUncookedOnly/Internal/Graph/AnimNextGraph_EdGraphNode.h b/Engine/Plugins/Experimental/AnimNext/Source/AnimNextUncookedOnly/Internal/Graph/AnimNextGraph_EdGraphNode.h
--- a/Engine/Plugins/Experimental/AnimNext/Source/AnimNextUncookedOnly/Internal/Graph/AnimNextGraph_EdGraphNode.h
+++ b/Engine/Plugins/Experimental/AnimNext/Source/AnimNextUncookedOnly/Internal/Graph/AnimNextGraph_EdGraphNode.h
@@ -3,11 +3,29 @@
 #pragma once
 
 #include "CoreMinimal.h"
-#include "Graph/ControlRigGraphNode.h"
+#include "EdGraph/RigVMEdGraphNode.h"
+
 #include "AnimNextGraph_EdGraphNode.generated.h"
 
+// EdGraphNode representation for AnimNext nodes
+// A node can hold a decorator stack or a decorator entry
 UCLASS(MinimalAPI)
-class UAnimNextGraph_EdGraphNode : public UControlRigGraphNode
+class UAnimNextGraph_EdGraphNode : public URigVMEdGraphNode
 {
 	GENERATED_BODY()
-};
\ No newline at end of file
+
+public:
+	//////////////////////////////////////////////////////////////////////////
+	// UEdGraphNode implementation
+	virtual void GetNodeContextMenuActions(class UToolMenu* Menu, class UGraphNodeContextMenuContext* Context) const override;
+
+	//////////////////////////////////////////////////////////////////////////
+	// Our implementation
+
+	// Returns whether this node is a decorator stack or not
+	ANIMNEXTUNCOOKEDONLY_API bool IsDecoratorStack() const;
+
+private:
+	// Populates the SubMenu with entries for each decorator that can be added through the context menu
+	void BuildAddDecoratorContextMenu(class UToolMenu* SubMenu);
+};
diff --git a/Engine/Plugins/Experimental/AnimNext/Source/AnimNextUncookedOnly/Internal/Graph/AnimNextGraph_EdGraphSchema.h b/Engine/Plugins/Experimental/AnimNext/Source/AnimNextUncookedOnly/Internal/Graph/AnimNextGraph_EdGraphSchema.h
--- a/Engine/Plugins/Experimental/AnimNext/Source/AnimNextUncookedOnly/Internal/Graph/AnimNextGraph_EdGraphSchema.h
+++ b/Engine/Plugins/Experimental/AnimNext/Source/AnimNextUncookedOnly/Internal/Graph/AnimNextGraph_EdGraphSchema.h
@@ -3,16 +3,16 @@
 #pragma once
 
 #include "CoreMinimal.h"
-#include "Graph/ControlRigGraphSchema.h"
+#include "EdGraph/RigVMEdGraphSchema.h"
 #include "AnimNextGraph_EdGraphNode.h"
 #include "AnimNextGraph_EdGraphSchema.generated.h"
 
 UCLASS(MinimalAPI)
-class UAnimNextGraph_EdGraphSchema : public UControlRigGraphSchema
+class UAnimNextGraph_EdGraphSchema : public URigVMEdGraphSchema
 {
 	GENERATED_BODY()
 
-	// UControlRigGraphSchema interface
+	// URigVMEdGraphSchema interface
 	virtual TSubclassOf<URigVMEdGraphNode> GetGraphNodeClass(const URigVMEdGraph* InGraph) const override { return UAnimNextGraph_EdGraphNode::StaticClass(); }
 
 	// UEdGraphSchema interface
diff --git a/Engine/Plugins/Experimental/AnimNext/Source/AnimNextUncookedOnly/Internal/Graph/AnimNextGraph_EditorData.h b/Engine/Plugins/Experimental/AnimNext/Source/AnimNextUncookedOnly/Internal/Graph/AnimNextGraph_EditorData.h
--- a/Engine/Plugins/Experimental/AnimNext/Source/AnimNextUncookedOnly/Internal/Graph/AnimNextGraph_EditorData.h
+++ b/Engine/Plugins/Experimental/AnimNext/Source/AnimNextUncookedOnly/Internal/Graph/AnimNextGraph_EditorData.h
@@ -27,17 +27,31 @@ enum class EAnimNextGraphLoadType : uint8
 	CheckUserDefinedStructs
 };
 
+/**
+ * The Schema is used to determine which actions are allowed
+ * on a graph. This includes any topological change.
+ */
 UCLASS()
 class UAnimNextGraph_Schema : public URigVMSchema
 {
 	GENERATED_BODY()
 };
 
+/**
+  * Implements a RigVM client host and RigVM graph function store host
+  * 
+  * A RigVM client holds a graph schema, multiple graph models and controllers, as well as an action stack, undo/redo
+  * information, and a function library.
+  * 
+  * A RigVM function store holds various RigVM functions that can be called and linked against. Functions are created
+  * in the editor UI under the 'My Blueprint' tab.
+  */
 UCLASS(MinimalAPI)
 class UAnimNextGraph_EditorData : public UObject, public IRigVMClientHost, public IRigVMGraphFunctionHost
 {
 	GENERATED_BODY()
 
+public:
 	UAnimNextGraph_EditorData(const FObjectInitializer& ObjectInitializer);
 	
 	friend class UAnimNextGraphFactory;
@@ -72,14 +86,24 @@ class UAnimNextGraph_EditorData : public UObject, public IRigVMClientHost, publi
 	virtual const FRigVMGraphFunctionStore* GetRigVMGraphFunctionStore() const override;
 
 	ANIMNEXTUNCOOKEDONLY_API void Initialize(bool bRecompileVM);
+	ANIMNEXTUNCOOKEDONLY_API UAnimNextGraph_EdGraph* GetRootGraph() const;
 
+protected:
+#if WITH_EDITOR
 	void RefreshAllModels(EAnimNextGraphLoadType InLoadType);
+	ANIMNEXTUNCOOKEDONLY_API void RebuildEdGraphFromModel();
+	void GetAllGraphs(TArray<UEdGraph*>& Graphs) const;
+#endif
 
 	void RecompileVM();
 	
 	void RecompileVMIfRequired();
 
 	void RequestAutoVMRecompilation();
+
+	void IncrementVMRecompileBracket();
+
+	void DecrementVMRecompileBracket();
 	
 	void HandleModifiedEvent(ERigVMGraphNotifType InNotifType, URigVMGraph* InGraph, UObject* InSubject);
 
@@ -122,6 +146,9 @@ class UAnimNextGraph_EditorData : public UObject, public IRigVMClientHost, publi
 	UPROPERTY(transient, DuplicateTransient)
 	TMap<FString, FRigVMOperand> PinToOperandMap;
 
+	UPROPERTY(transient, DuplicateTransient)
+	int32 VMRecompilationBracket = 0;
+
 	UPROPERTY(transient, DuplicateTransient)
 	bool bVMRecompilationRequired = false;
 
diff --git a/Engine/Plugins/Experimental/AnimNext/Source/AnimNextUncookedOnly/Private/AnimNextUncookedOnlyModule.cpp b/Engine/Plugins/Experimental/AnimNext/Source/AnimNextUncookedOnly/Private/AnimNextUncookedOnlyModule.cpp
--- a/Engine/Plugins/Experimental/AnimNext/Source/AnimNextUncookedOnly/Private/AnimNextUncookedOnlyModule.cpp
+++ b/Engine/Plugins/Experimental/AnimNext/Source/AnimNextUncookedOnly/Private/AnimNextUncookedOnlyModule.cpp
@@ -1,5 +1,21 @@
 // Copyright Epic Games, Inc. All Rights Reserved.
 
+#include "CoreMinimal.h"
 #include "Modules/ModuleManager.h"
 
-IMPLEMENT_MODULE(FDefaultModuleImpl, AnimNextUncookedOnly);
\ No newline at end of file
+namespace UE::AnimNext::UncookedOnly
+{
+	class FModule : public IModuleInterface
+	{
+	public:
+
+	private:
+		virtual void StartupModule() override;
+	};
+
+	void FModule::StartupModule()
+	{
+	}
+}
+
+IMPLEMENT_MODULE(UE::AnimNext::UncookedOnly::FModule, AnimNextUncookedOnly);
diff --git a/Engine/Plugins/Experimental/AnimNext/Source/AnimNextUncookedOnly/Private/Graph/AnimNextGraph_EdGraphNode.cpp b/Engine/Plugins/Experimental/AnimNext/Source/AnimNextUncookedOnly/Private/Graph/AnimNextGraph_EdGraphNode.cpp
new file mode 100644
--- /dev/null
+++ b/Engine/Plugins/Experimental/AnimNext/Source/AnimNextUncookedOnly/Private/Graph/AnimNextGraph_EdGraphNode.cpp
@@ -0,0 +1,100 @@
+﻿// Copyright Epic Games, Inc. All Rights Reserved.
+
+#include "Graph/AnimNextGraph_EdGraphNode.h"
+
+#include "Graph/RigUnit_AnimNextDecoratorStack.h"
+#include "Graph/RigDecorator_AnimNextCppDecorator.h"
+#include "RigVMModel/RigVMController.h"
+#include "DecoratorBase/DecoratorRegistry.h"
+#include "ToolMenu.h"
+
+#define LOCTEXT_NAMESPACE "AnimNextGraph_EdGraphNode"
+
+void UAnimNextGraph_EdGraphNode::GetNodeContextMenuActions(UToolMenu* Menu, class UGraphNodeContextMenuContext* Context) const
+{
+	URigVMEdGraphNode::GetNodeContextMenuActions(Menu, Context);
+
+	if (IsDecoratorStack())
+	{
+		FToolMenuSection& Section = Menu->AddSection("AnimNextDecoratorNodeActions", LOCTEXT("AnimNextDecoratorNodeActionsMenuHeader", "Anim Next Decorator Actions"));
+
+		UAnimNextGraph_EdGraphNode* NonConstThis = const_cast<UAnimNextGraph_EdGraphNode*>(this);
+
+		Section.AddSubMenu(
+			TEXT("AddDecoratorMenu"),
+			LOCTEXT("AddDecoratorMenu", "Add Decorator"),
+			LOCTEXT("AddDecoratorMenuTooltip", "Add the chosen decorator to currently selected node"),
+			FNewToolMenuDelegate::CreateUObject(NonConstThis, &UAnimNextGraph_EdGraphNode::BuildAddDecoratorContextMenu));
+	}
+}
+
+bool UAnimNextGraph_EdGraphNode::IsDecoratorStack() const
+{
+	if (const URigVMUnitNode* VMNode = Cast<URigVMUnitNode>(GetModelNode()))
+	{
+		const UScriptStruct* ScriptStruct = VMNode->GetScriptStruct();
+		return ScriptStruct == FRigUnit_AnimNextDecoratorStack::StaticStruct();
+	}
+
+	return false;
+}
+
+void UAnimNextGraph_EdGraphNode::BuildAddDecoratorContextMenu(UToolMenu* SubMenu)
+{
+	using namespace UE::AnimNext;
+
+	const FDecoratorRegistry& DecoratorRegistry = FDecoratorRegistry::Get();
+	TArray<const FDecorator*> Decorators = DecoratorRegistry.GetDecorators();
+
+	URigVMController* VMController = GetController();
+	URigVMNode* VMNode = GetModelNode();
+
+	const UScriptStruct* CppDecoratorStruct = FRigDecorator_AnimNextCppDecorator::StaticStruct();
+
+	for (const FDecorator* Decorator : Decorators)
+	{
+		UScriptStruct* ScriptStruct = Decorator->GetDecoratorSharedDataStruct();
+
+		FString DefaultValue;
+		{
+			const FRigDecorator_AnimNextCppDecorator DefaultCppDecoratorStructInstance;
+			FRigDecorator_AnimNextCppDecorator CppDecoratorStructInstance;
+			CppDecoratorStructInstance.DecoratorSharedDataStruct = ScriptStruct;
+
+			if (!CppDecoratorStructInstance.CanBeAddedToNode(VMNode, nullptr))
+			{
+				continue;	// This decorator isn't supported on this node
+			}
+
+			const FProperty* Prop = FAnimNextCppDecoratorWrapper::StaticStruct()->FindPropertyByName(GET_MEMBER_NAME_STRING_CHECKED(FAnimNextCppDecoratorWrapper, CppDecorator));
+			Prop->ExportText_Direct(DefaultValue, &CppDecoratorStructInstance, &DefaultCppDecoratorStructInstance, nullptr, PPF_None);
+		}
+
+		FString DisplayNameMetadata;
+		ScriptStruct->GetStringMetaDataHierarchical(FRigVMStruct::DisplayNameMetaName, &DisplayNameMetadata);
+		const FString DisplayName = DisplayNameMetadata.IsEmpty() ? Decorator->GetDecoratorUID().GetDecoratorName() : DisplayNameMetadata;
+
+		const FText ToolTip = ScriptStruct->GetToolTipText();
+
+		FToolMenuEntry DecoratorEntry = FToolMenuEntry::InitMenuEntry(
+			Decorator->GetDecoratorUID().GetDecoratorName(),
+			FText::FromString(DisplayName),
+			ToolTip,
+			FSlateIcon(),
+			FUIAction(FExecuteAction::CreateLambda(
+				[this, Decorator, VMController, VMNode, CppDecoratorStruct, DefaultValue, DisplayName]()
+				{
+					VMController->AddDecorator(
+						VMNode->GetFName(),
+						*CppDecoratorStruct->GetPathName(),
+						*DisplayName,
+						DefaultValue, INDEX_NONE, true, true);
+				})
+			)
+		);
+
+		SubMenu->AddMenuEntry(NAME_None, DecoratorEntry);
+	}
+}
+
+#undef LOCTEXT_NAMESPACE
diff --git a/Engine/Plugins/Experimental/AnimNext/Source/AnimNextUncookedOnly/Private/Graph/AnimNextGraph_EditorData.cpp b/Engine/Plugins/Experimental/AnimNext/Source/AnimNextUncookedOnly/Private/Graph/AnimNextGraph_EditorData.cpp
--- a/Engine/Plugins/Experimental/AnimNext/Source/AnimNextUncookedOnly/Private/Graph/AnimNextGraph_EditorData.cpp
+++ b/Engine/Plugins/Experimental/AnimNext/Source/AnimNextUncookedOnly/Private/Graph/AnimNextGraph_EditorData.cpp
@@ -23,6 +23,7 @@ UAnimNextGraph_EditorData::UAnimNextGraph_EditorData(const FObjectInitializer& O
 	{
 		TGuardValue<bool> DisableClientNotifs(RigVMClient.bSuspendNotifications, true);
 		RigVMClient.AddModel(TEXT("RigVMGraph"), false, &ObjectInitializer);
+		RigVMClient.AddModel(TEXT("TempRigVMGraph"), false, &ObjectInitializer);	// Used during compilation, otherwise empty
 		RigVMClient.GetOrCreateFunctionLibrary(false, &ObjectInitializer);
 	}
 	RigVMClient.SetExecuteContextStruct(FAnimNextExecuteContext::StaticStruct());
@@ -62,6 +63,7 @@ void UAnimNextGraph_EditorData::Initialize(bool bRecompileVM)
 		
 		RigVMClient.GetOrCreateController(RigVMClient.GetDefaultModel());
 		RigVMClient.GetOrCreateController(RigVMClient.GetFunctionLibrary());
+		RigVMClient.GetOrCreateController(TEXT("TempRigVMGraph"));
 
 		// Init function library controllers
 		for(URigVMLibraryNode* LibraryNode : RigVMClient.GetFunctionLibrary()->GetFunctions())
@@ -83,6 +85,11 @@ void UAnimNextGraph_EditorData::Initialize(bool bRecompileVM)
 	}
 }
 
+UAnimNextGraph_EdGraph* UAnimNextGraph_EditorData::GetRootGraph() const
+{
+	return RootGraph;
+}
+
 void UAnimNextGraph_EditorData::PostLoad()
 {
 	Super::PostLoad();
@@ -233,6 +240,57 @@ void UAnimNextGraph_EditorData::RefreshAllModels(EAnimNextGraphLoadType InLoadTy
 	}
 }
 
+void UAnimNextGraph_EditorData::RebuildEdGraphFromModel()
+{
+	DECLARE_SCOPE_HIERARCHICAL_COUNTER_FUNC()
+
+	TGuardValue<bool> SelfGuard(bSuspendModelNotificationsForSelf, true);
+	TGuardValue<bool> ClientIgnoreModificationsGuard(RigVMClient.bIgnoreModelNotifications, true);
+
+	verify(RigVMClient.GetOrCreateController((URigVMGraph*)nullptr));
+
+	TArray<UEdGraph*> EdGraphs;
+	GetAllGraphs(EdGraphs);
+
+	for (UEdGraph* Graph : EdGraphs)
+	{
+		TArray<UEdGraphNode*> Nodes = Graph->Nodes;
+		for (UEdGraphNode* Node : Nodes)
+		{
+			Graph->RemoveNode(Node);
+		}
+	}
+
+	if (FunctionLibraryEdGraph && RigVMClient.GetFunctionLibrary())
+	{
+		FunctionLibraryEdGraph->ModelNodePath = RigVMClient.GetFunctionLibrary()->GetNodePath();
+	}
+
+	TArray<URigVMGraph*> RigGraphs = RigVMClient.GetAllModels(true, true);
+
+	for (int32 RigGraphIndex = 0; RigGraphIndex < RigGraphs.Num(); RigGraphIndex++)
+	{
+		RigVMClient.GetOrCreateController(RigGraphs[RigGraphIndex])->ResendAllNotifications();
+	}
+}
+
+void UAnimNextGraph_EditorData::GetAllGraphs(TArray<UEdGraph*>& Graphs) const
+{
+	Graphs.Reset();
+
+	Graphs.Add(RootGraph);
+
+	if (EntryPointGraph)
+	{
+		Graphs.Add(EntryPointGraph);
+	}
+
+	if (FunctionLibraryEdGraph)
+	{
+		Graphs.Add(FunctionLibraryEdGraph);
+	}
+}
+
 #endif // WITH_EDITOR
 
 FRigVMClient* UAnimNextGraph_EditorData::GetRigVMClient()
@@ -363,18 +421,50 @@ void UAnimNextGraph_EditorData::RecompileVMIfRequired()
 void UAnimNextGraph_EditorData::RequestAutoVMRecompilation()
 {
 	bVMRecompilationRequired = true;
-	if (bAutoRecompileVM)
+	if (bAutoRecompileVM && VMRecompilationBracket == 0)
 	{
 		RecompileVMIfRequired();
 	}
 }
 
+void UAnimNextGraph_EditorData::IncrementVMRecompileBracket()
+{
+	VMRecompilationBracket++;
+}
+
+void UAnimNextGraph_EditorData::DecrementVMRecompileBracket()
+{
+	if (VMRecompilationBracket == 1)
+	{
+		if (bAutoRecompileVM)
+		{
+			RecompileVMIfRequired();
+		}
+		VMRecompilationBracket = 0;
+	}
+	else if (VMRecompilationBracket > 0)
+	{
+		VMRecompilationBracket--;
+	}
+}
+
 void UAnimNextGraph_EditorData::HandleModifiedEvent(ERigVMGraphNotifType InNotifType, URigVMGraph* InGraph, UObject* InSubject)
 {
 	bool bNotifForOthersPending = true;
 
 	switch(InNotifType)
 	{
+	case ERigVMGraphNotifType::InteractionBracketOpened:
+		{
+			IncrementVMRecompileBracket();
+			break;
+		}
+	case ERigVMGraphNotifType::InteractionBracketClosed:
+	case ERigVMGraphNotifType::InteractionBracketCanceled:
+		{
+			DecrementVMRecompileBracket();
+			break;
+		}
 	case ERigVMGraphNotifType::NodeAdded:
 		{
 			if (URigVMCollapseNode* CollapseNode = Cast<URigVMCollapseNode>(InSubject))
@@ -389,10 +479,20 @@ void UAnimNextGraph_EditorData::HandleModifiedEvent(ERigVMGraphNotifType InNotif
 	case ERigVMGraphNotifType::PinArraySizeChanged:
 	case ERigVMGraphNotifType::PinDirectionChanged:
 		{
-			RecompileVM();
+			RequestAutoVMRecompilation();
+			break;
+		}
+	case ERigVMGraphNotifType::PinAdded:
+		{
+			if (URigVMPin* Pin = Cast<URigVMPin>(InSubject))
+			{
+				if (Pin->IsDecoratorPin())
+				{
+					RequestAutoVMRecompilation();
+				}
+			}
 			break;
 		}
-
 	case ERigVMGraphNotifType::PinDefaultValueChanged:
 		{
 			if (InGraph->GetRuntimeAST().IsValid())
@@ -403,14 +503,14 @@ void UAnimNextGraph_EditorData::HandleModifiedEvent(ERigVMGraphNotifType InNotif
 				if (Expression == nullptr)
 				{
 					InGraph->ClearAST();
-					break;
 				}
 				else if (Expression->NumParents() > 1)
 				{
 					InGraph->ClearAST();
-					break;
 				}
 			}
+
+			RequestAutoVMRecompilation();	// We need to rebuild our metadata when a default value changes
 			break;
 		}
 	}
diff --git a/Engine/Plugins/Experimental/AnimNext/Source/AnimNextUncookedOnly/Private/UncookedOnlyUtils.cpp b/Engine/Plugins/Experimental/AnimNext/Source/AnimNextUncookedOnly/Private/UncookedOnlyUtils.cpp
--- a/Engine/Plugins/Experimental/AnimNext/Source/AnimNextUncookedOnly/Private/UncookedOnlyUtils.cpp
+++ b/Engine/Plugins/Experimental/AnimNext/Source/AnimNextUncookedOnly/Private/UncookedOnlyUtils.cpp
@@ -3,20 +3,240 @@
 #include "UncookedOnlyUtils.h"
 #include "Graph/AnimNextGraph.h"
 #include "Graph/AnimNextGraph_EditorData.h"
-#include "RigVMCompiler/RigVMCompiler.h"
-#include "RigVMCore/RigVM.h"
+#include "Graph/RigUnit_AnimNextGraphRoot.h"
+#include "Graph/RigUnit_AnimNextDecoratorStack.h"
 #include "Graph/AnimNextExecuteContext.h"
 #include "Param/AnimNextParameter.h"
 #include "Param/AnimNextParameterBlock.h"
 #include "Param/AnimNextParameterBlock_EditorData.h"
 #include "Param/AnimNextParameterBlock_EdGraph.h"
+#include "Graph/RigDecorator_AnimNextCppDecorator.h"
 #include "Graph/RigUnit_AnimNextBeginExecution.h"
 #include "Param/RigVMDispatch_SetParameter.h"
 #include "Param/AnimNextParameterBlockEntry.h"
 #include "Param/IAnimNextParameterBlockBindingInterface.h"
+#include "DecoratorBase/DecoratorReader.h"
+#include "DecoratorBase/DecoratorWriter.h"
+#include "DecoratorBase/NodeTemplate.h"
+#include "DecoratorBase/NodeTemplateBuilder.h"
+#include "DecoratorBase/NodeTemplateRegistry.h"
+#include "DecoratorBase/NodeDescription.h"
+#include "DecoratorBase/NodeInstance.h"
+#include "DecoratorBase/DecoratorRegistry.h"
+#include "DecoratorBase/Decorator.h"
+#include "Serialization/MemoryReader.h"
+
+#include "RigVMCompiler/RigVMCompiler.h"
+#include "RigVMCore/RigVM.h"
 
 namespace UE::AnimNext::UncookedOnly
 {
+namespace Private
+{
+	struct FDecoratorEntryMapping
+	{
+		const URigVMNode* DecoratorStackNode;
+		const URigVMPin* DecoratorEntryPin;
+		const FDecorator* Decorator;
+
+		FDecoratorEntryMapping(const URigVMNode* InDecoratorStackNode, const URigVMPin* InDecoratorEntryPin, const FDecorator* InDecorator)
+			: DecoratorStackNode(InDecoratorStackNode)
+			, DecoratorEntryPin(InDecoratorEntryPin)
+			, Decorator(InDecorator)
+		{}
+	};
+
+	struct FDecoratorStackMapping
+	{
+		const URigVMNode* DecoratorStackNode;
+		TArray<FDecoratorEntryMapping> DecoratorEntries;
+
+		FNodeHandle DecoratorStackNodeHandle;
+
+		explicit FDecoratorStackMapping(const URigVMNode* InDecoratorStackNode)
+			: DecoratorStackNode(InDecoratorStackNode)
+		{}
+	};
+
+	template<typename DecoratorAction>
+	void ForEachDecoratorInStack(const URigVMNode* DecoratorStackNode, const DecoratorAction& Action)
+	{
+		const TArray<URigVMPin*>& Pins = DecoratorStackNode->GetPins();
+		for (URigVMPin* Pin : Pins)
+		{
+			if (!Pin->IsDecoratorPin())
+			{
+				continue;	// Not a decorator pin
+			}
+
+			if (Pin->GetScriptStruct() == FRigDecorator_AnimNextCppDecorator::StaticStruct())
+			{
+				TSharedPtr<FStructOnScope> DecoratorScope = Pin->GetDecoratorInstance();
+				FRigDecorator_AnimNextCppDecorator* VMDecorator = (FRigDecorator_AnimNextCppDecorator*)DecoratorScope->GetStructMemory();
+
+				if (const FDecorator* Decorator = VMDecorator->GetDecorator())
+				{
+					Action(DecoratorStackNode, Pin, Decorator);
+				}
+			}
+		}
+	}
+
+	TArray<FDecoratorUID> GetDecoratorUIDs(const URigVMNode* DecoratorStackNode)
+	{
+		TArray<FDecoratorUID> Decorators;
+
+		ForEachDecoratorInStack(DecoratorStackNode,
+			[&Decorators](const URigVMNode* DecoratorStackNode, const URigVMPin* DecoratorPin, const FDecorator* Decorator)
+			{
+				Decorators.Add(Decorator->GetDecoratorUID());
+			});
+
+		return Decorators;
+	}
+
+	FNodeHandle RegisterDecoratorNodeTemplate(FDecoratorWriter& DecoratorWriter, const URigVMNode* DecoratorStackNode)
+	{
+		const TArray<FDecoratorUID> DecoratorUIDs = GetDecoratorUIDs(DecoratorStackNode);
+
+		TArray<uint8> NodeTemplateBuffer;
+		const FNodeTemplate* NodeTemplate = FNodeTemplateBuilder::BuildNodeTemplate(DecoratorUIDs, NodeTemplateBuffer);
+
+		return DecoratorWriter.RegisterNode(*NodeTemplate);
+	}
+
+	FString GetDecoratorProperty(const FDecoratorStackMapping& DecoratorStack, uint32 DecoratorIndex, const FString& PropertyName, const TArray<FDecoratorStackMapping>& DecoratorStackNodes)
+	{
+		const TArray<URigVMPin*>& Pins = DecoratorStack.DecoratorEntries[DecoratorIndex].DecoratorEntryPin->GetSubPins();
+		for (URigVMPin* Pin : Pins)
+		{
+			if (Pin->GetDirection() != ERigVMPinDirection::Input)
+			{
+				continue;	// We only look for input pins
+			}
+
+			if (Pin->GetName() == PropertyName)
+			{
+				if (Pin->GetCPPTypeObject() == FAnimNextDecoratorHandle::StaticStruct())
+				{
+					// Decorator handle pins don't have a value, just an optional link
+					const TArray<URigVMLink*>& PinLinks = Pin->GetLinks();
+					if (PinLinks.Num() != 0)
+					{
+						// Something is connected to us, find the corresponding node handle so that we can encode it as our property value
+						check(PinLinks.Num() == 1);
+
+						const URigVMNode* SourceNode = PinLinks[0]->GetSourceNode();
+
+						FNodeHandle SourceNodeHandle;
+						int32 SourceDecoratorIndex = INDEX_NONE;
+
+						const FDecoratorStackMapping* SourceDecoratorStack = DecoratorStackNodes.FindByPredicate([SourceNode](const FDecoratorStackMapping& Mapping) { return Mapping.DecoratorStackNode == SourceNode; });
+						if (SourceDecoratorStack != nullptr)
+						{
+							SourceNodeHandle = SourceDecoratorStack->DecoratorStackNodeHandle;
+							SourceDecoratorIndex = 0;	// We always bind to the first decorator index since we only allow one base decorator per stack for now
+						}
+
+						if (SourceNodeHandle.IsValid())
+						{
+							check(SourceDecoratorIndex != INDEX_NONE);
+
+							const FAnimNextDecoratorHandle DecoratorHandle(SourceNodeHandle, SourceDecoratorIndex);
+							const FAnimNextDecoratorHandle DefaultDecoratorHandle;
+
+							// We need an instance of a decorator handle property to be able to serialize it into text, grab it from the root
+							const FProperty* Property = FRigUnit_AnimNextGraphRoot::StaticStruct()->FindPropertyByName(GET_MEMBER_NAME_STRING_CHECKED(FRigUnit_AnimNextGraphRoot, Result));
+
+							FString PropertyValue;
+							Property->ExportText_Direct(PropertyValue, &DecoratorHandle, &DefaultDecoratorHandle, nullptr, PPF_None);
+
+							return PropertyValue;
+						}
+					}
+
+					// This handle pin isn't connected
+					return FString();
+				}
+
+				// A regular property pin
+				return Pin->GetDefaultValue();
+			}
+		}
+
+		// Unknown property
+		return FString();
+	}
+
+	void WriteDecoratorProperties(FDecoratorWriter& DecoratorWriter, const FDecoratorStackMapping& Mapping, const TArray<FDecoratorStackMapping>& DecoratorStackNodes)
+	{
+		DecoratorWriter.WriteNode(Mapping.DecoratorStackNodeHandle,
+			[&Mapping, &DecoratorStackNodes](uint32 DecoratorIndex, const FString& PropertyName)
+			{
+				return GetDecoratorProperty(Mapping, DecoratorIndex, PropertyName, DecoratorStackNodes);
+			});
+	}
+
+	const URigVMUnitNode* FindRootNode(const TArray<URigVMNode*>& VMNodes)
+	{
+		for (const URigVMNode* VMNode : VMNodes)
+		{
+			if (const URigVMUnitNode* VMUnitNode = Cast<URigVMUnitNode>(VMNode))
+			{
+				const UScriptStruct* ScriptStruct = VMUnitNode->GetScriptStruct();
+				if (ScriptStruct == FRigUnit_AnimNextGraphRoot::StaticStruct())
+				{
+					return VMUnitNode;
+				}
+			}
+		}
+
+		return nullptr;
+	}
+
+	TArray<FDecoratorStackMapping> CollectDecoratorStacks(const URigVMGraph* VMGraph)
+	{
+		const TArray<URigVMNode*>& VMNodes = VMGraph->GetNodes();
+		const URigVMUnitNode* VMRootNode = FindRootNode(VMNodes);
+
+		TArray<FDecoratorStackMapping> DecoratorStackNodes;
+
+		if (VMRootNode == nullptr)
+		{
+			return DecoratorStackNodes;
+		}
+
+		TArray<const URigVMNode*> NodesToVisit;
+		NodesToVisit.Add(VMRootNode);
+
+		while (NodesToVisit.Num() != 0)
+		{
+			const URigVMNode* VMNode = NodesToVisit[0];
+			NodesToVisit.RemoveAt(0);
+
+			if (const URigVMUnitNode* VMUnitNode = Cast<URigVMUnitNode>(VMNode))
+			{
+				const UScriptStruct* ScriptStruct = VMUnitNode->GetScriptStruct();
+				if (ScriptStruct == FRigUnit_AnimNextDecoratorStack::StaticStruct())
+				{
+					FDecoratorStackMapping Mapping(VMNode);
+					ForEachDecoratorInStack(VMNode,
+						[&Mapping](const URigVMNode* DecoratorStackNode, const URigVMPin* DecoratorPin, const FDecorator* Decorator)
+						{
+							Mapping.DecoratorEntries.Add(FDecoratorEntryMapping(DecoratorStackNode, DecoratorPin, Decorator));
+						});
+
+					DecoratorStackNodes.Add(MoveTemp(Mapping));
+				}
+			}
+
+			const TArray<URigVMNode*> SourceNodes = VMNode->GetLinkedSourceNodes();
+			NodesToVisit.Append(SourceNodes);
+		}
+
+		return DecoratorStackNodes;
+	}
+}
 
 void FUtils::Compile(UAnimNextGraph* InGraph)
 {
@@ -46,11 +266,76 @@ void FUtils::Compile(UAnimNextGraph* InGraph)
 	EditorData->CompileLog.Messages.Reset();
 	EditorData->CompileLog.NumErrors = EditorData->CompileLog.NumWarnings = 0;
 
+	// We use a temporary graph model to build our final graph that we'll compile
+	URigVMGraph* VMTempGraph = EditorData->GetRigVMClient()->GetModel(TEXT("TempRigVMGraph"));
+	URigVMController* TempController = nullptr;
+
+	TArray<URigVMGraph*> GraphModelsToCompile;
+
+	if (VMTempGraph != nullptr)
+	{
+		TempController = EditorData->GetRigVMClient()->GetOrCreateController(VMTempGraph);
+
+		// Clear our temporary model
+		TempController->RemoveNodes(VMTempGraph->GetNodes());
+
+		URigVMGraph* VMRootGraph = EditorData->GetRigVMClient()->GetDefaultModel();
+
+		// Add our root node, it's always there
+		URigVMUnitNode* TempRootNode = TempController->AddUnitNode(FRigUnit_AnimNextGraphRoot::StaticStruct(), FRigVMStruct::ExecuteName, FVector2D(0.0f, 0.0f), FString(), false);
+
+		// Gather our decorator stacks
+		TArray<Private::FDecoratorStackMapping> DecoratorStackNodes = Private::CollectDecoratorStacks(VMRootGraph);
+
+		FDecoratorWriter DecoratorWriter;
+
+		// Iterate over every decorator stack and register our node templates
+		for (Private::FDecoratorStackMapping& NodeMapping : DecoratorStackNodes)
+		{
+			NodeMapping.DecoratorStackNodeHandle = Private::RegisterDecoratorNodeTemplate(DecoratorWriter, NodeMapping.DecoratorStackNode);
+		}
+
+		// Write our node shared data
+		DecoratorWriter.BeginNodeWriting();
+
+		for (const Private::FDecoratorStackMapping& NodeMapping : DecoratorStackNodes)
+		{
+			Private::WriteDecoratorProperties(DecoratorWriter, NodeMapping, DecoratorStackNodes);
+		}
+
+		DecoratorWriter.EndNodeWriting();
+
+		// Find our root node handle, if we have any stack nodes, the first one is our root stack
+		FAnimNextEntryPointHandle RootDecoratorHandle;
+		if (DecoratorStackNodes.Num() != 0)
+		{
+			RootDecoratorHandle = FAnimNextEntryPointHandle(DecoratorStackNodes[0].DecoratorStackNodeHandle);
+		}
+
+		// Cache our compiled metadata
+		InGraph->SharedDataArchiveBuffer = DecoratorWriter.GetGraphSharedData();
+		InGraph->RootDecoratorHandle = RootDecoratorHandle;
+
+		// Populate our runtime metadata
+		InGraph->LoadFromArchiveBuffer(InGraph->SharedDataArchiveBuffer);
+
+		// Compile our new graph
+		GraphModelsToCompile.Add(VMTempGraph);
+	}
+
+	if (VMTempGraph == nullptr)
+	{
+		// This is an old graph that was saved before the introduction of the temporary graph, use the root graph instead
+		URigVMGraph* VMRootGraph = EditorData->GetRigVMClient()->GetDefaultModel();
+
+		GraphModelsToCompile.Add(VMRootGraph);
+		TempController = EditorData->GetRigVMClient()->GetOrCreateController(VMRootGraph);
+	}
+
 	URigVMCompiler* Compiler = URigVMCompiler::StaticClass()->GetDefaultObject<URigVMCompiler>();
 	EditorData->VMCompileSettings.SetExecuteContextStruct(EditorData->RigVMClient.GetExecuteContextStruct());
 	Compiler->Settings = (EditorData->bCompileInDebugMode) ? FRigVMCompileSettings::Fast(EditorData->VMCompileSettings.GetExecuteContextStruct()) : EditorData->VMCompileSettings;
-	URigVMController* RootController = EditorData->GetRigVMClient()->GetOrCreateController(EditorData->GetRigVMClient()->GetDefaultModel());
-	Compiler->Compile(EditorData->GetRigVMClient()->GetAllModels(false, false), RootController, InGraph->RigVM, InGraph->ExtendedExecuteContext, InGraph->GetRigVMExternalVariables(), &EditorData->PinToOperandMap);
+	Compiler->Compile(GraphModelsToCompile, TempController, InGraph->RigVM, InGraph->ExtendedExecuteContext, InGraph->GetRigVMExternalVariables(), &EditorData->PinToOperandMap);
 
 	if (EditorData->bErrorsDuringCompilation)
 	{
@@ -615,7 +900,7 @@ void FUtils::SetupBindingGraphForLiteral(URigVMController* InController, FName I
 	InController->RemoveNodes(InController->GetGraph()->GetNodes());
 
 	// Add new nodes for a simple literal binding
-	URigVMUnitNode* EntryPointNode = InController->AddUnitNode(FRigUnit_AnimNextBeginExecution::StaticStruct(), FRigUnit::GetMethodName(), FVector2D(-200.0f, 0.0f), FString(), false);
+	URigVMUnitNode* EntryPointNode = InController->AddUnitNode(FRigUnit_AnimNextBeginExecution::StaticStruct(), FRigVMStruct::ExecuteName, FVector2D(-200.0f, 0.0f), FString(), false);
 
 	const FName FactoryName = FRigVMDispatch_SetParameter().GetFactoryName();
 	FRigVMDispatch_SetParameter* Factory = static_cast<FRigVMDispatch_SetParameter*>(FRigVMRegistry::Get().FindDispatchFactory(FactoryName));
diff --git a/Engine/Plugins/Runtime/RigVM/Source/RigVM/Public/RigVMCore/RigVMDecorator.h b/Engine/Plugins/Runtime/RigVM/Source/RigVM/Public/RigVMCore/RigVMDecorator.h
--- a/Engine/Plugins/Runtime/RigVM/Source/RigVM/Public/RigVMCore/RigVMDecorator.h
+++ b/Engine/Plugins/Runtime/RigVM/Source/RigVM/Public/RigVMCore/RigVMDecorator.h
@@ -62,6 +62,9 @@ public:
 	// allows the decorator to react when added to a node
 	virtual void OnDecoratorAdded(URigVMController* InController, URigVMNode* InNode) {}
 
+	// allows the decorator to return dynamic pins (parent pin index must be INDEX_NONE or point to a valid index of the parent pin in the OutPinArray)
+	virtual void GetProgrammaticPins(URigVMController* InController, int32 InParentPinIndex, struct FRigVMPinInfoArray& OutPinArray) const {}
+
 #endif
 
 private:
diff --git a/Engine/Plugins/Runtime/RigVM/Source/RigVMDeveloper/Private/EdGraph/RigVMEdGraphNode.cpp b/Engine/Plugins/Runtime/RigVM/Source/RigVMDeveloper/Private/EdGraph/RigVMEdGraphNode.cpp
--- a/Engine/Plugins/Runtime/RigVM/Source/RigVMDeveloper/Private/EdGraph/RigVMEdGraphNode.cpp
+++ b/Engine/Plugins/Runtime/RigVM/Source/RigVMDeveloper/Private/EdGraph/RigVMEdGraphNode.cpp
@@ -1547,6 +1547,11 @@ FName URigVMEdGraphNode::GetModelNodeName() const
 	return NAME_None;
 }
 
+const FString& URigVMEdGraphNode::GetModelNodePath() const
+{
+	return ModelNodePath;
+}
+
 URigVMPin* URigVMEdGraphNode::GetModelPinFromPinPath(const FString& InPinPath) const
 {
 	if (TWeakObjectPtr<URigVMPin> const* CachedModelPinPtr = PinPathToModelPin.Find(InPinPath))
diff --git a/Engine/Plugins/Runtime/RigVM/Source/RigVMDeveloper/Private/RigVMCompiler/RigVMAST.cpp b/Engine/Plugins/Runtime/RigVM/Source/RigVMDeveloper/Private/RigVMCompiler/RigVMAST.cpp
--- a/Engine/Plugins/Runtime/RigVM/Source/RigVMDeveloper/Private/RigVMCompiler/RigVMAST.cpp
+++ b/Engine/Plugins/Runtime/RigVM/Source/RigVMDeveloper/Private/RigVMCompiler/RigVMAST.cpp
@@ -2519,7 +2519,7 @@ const FRigVMASTLinkDescription& FRigVMParserAST::GetLink(int32 InLinkIndex) cons
 	return Links[InLinkIndex];
 }
 
-void FRigVMParserAST::Inline(TArray<URigVMGraph*> InGraphs)
+void FRigVMParserAST::Inline(const TArray<URigVMGraph*>& InGraphs)
 {
 	TArray<FRigVMASTProxy> LocalNodeProxies;
 	for(URigVMGraph* Graph : InGraphs)
@@ -2532,7 +2532,7 @@ void FRigVMParserAST::Inline(TArray<URigVMGraph*> InGraphs)
 	Inline(InGraphs, LocalNodeProxies);
 }
 
-void FRigVMParserAST::Inline(TArray<URigVMGraph*> InGraphs, const TArray<FRigVMASTProxy>& InNodeProxies)
+void FRigVMParserAST::Inline(const TArray<URigVMGraph*>& InGraphs, const TArray<FRigVMASTProxy>& InNodeProxies)
 {
 	DECLARE_SCOPE_HIERARCHICAL_COUNTER_FUNC()
 
diff --git a/Engine/Plugins/Runtime/RigVM/Source/RigVMDeveloper/Private/RigVMModel/RigVMController.cpp b/Engine/Plugins/Runtime/RigVM/Source/RigVMDeveloper/Private/RigVMModel/RigVMController.cpp
--- a/Engine/Plugins/Runtime/RigVM/Source/RigVMDeveloper/Private/RigVMModel/RigVMController.cpp
+++ b/Engine/Plugins/Runtime/RigVM/Source/RigVMDeveloper/Private/RigVMModel/RigVMController.cpp
@@ -351,7 +351,7 @@ void URigVMController::SetIsRunningUnitTest(bool bIsRunning)
 	}	
 }
 
-URigVMController::FPinInfo::FPinInfo()
+FRigVMPinInfo::FRigVMPinInfo()
 	: ParentIndex(INDEX_NONE)
 	, Name(NAME_None)
 	, Direction(ERigVMPinDirection::Invalid)
@@ -365,7 +365,7 @@ URigVMController::FPinInfo::FPinInfo()
 {
 }
 
-URigVMController::FPinInfo::FPinInfo(const URigVMPin* InPin, int32 InParentIndex, ERigVMPinDirection InDirection)
+FRigVMPinInfo::FRigVMPinInfo(const URigVMPin* InPin, int32 InParentIndex, ERigVMPinDirection InDirection)
 	: ParentIndex(InParentIndex)
 	, Name(InPin->GetName())
 	, Direction(InDirection == ERigVMPinDirection::Invalid ? InPin->GetDirection() : InDirection)
@@ -389,7 +389,7 @@ URigVMController::FPinInfo::FPinInfo(const URigVMPin* InPin, int32 InParentIndex
 	}
 }
 
-URigVMController::FPinInfo::FPinInfo(FProperty* InProperty, ERigVMPinDirection InDirection, int32 InParentIndex, const uint8* InDefaultValueMemory)
+FRigVMPinInfo::FRigVMPinInfo(FProperty* InProperty, ERigVMPinDirection InDirection, int32 InParentIndex, const uint8* InDefaultValueMemory)
 	: ParentIndex(InParentIndex)
 	, Name(InProperty->GetFName())
 	, Direction(InDirection)
@@ -496,7 +496,7 @@ URigVMController::FPinInfo::FPinInfo(FProperty* InProperty, ERigVMPinDirection I
 	CorrectExecuteTypeIndex();
 }
 
-void URigVMController::FPinInfo::CorrectExecuteTypeIndex()
+void FRigVMPinInfo::CorrectExecuteTypeIndex()
 {
 	const FRigVMRegistry& Registry = FRigVMRegistry::Get();
 	if(Registry.IsExecuteType(TypeIndex))
@@ -510,7 +510,7 @@ void URigVMController::FPinInfo::CorrectExecuteTypeIndex()
 	}
 }
 
-uint32 GetTypeHash(const URigVMController::FPinInfo& InPin)
+uint32 GetTypeHash(const FRigVMPinInfo& InPin)
 {
 	uint32 Hash = 0; //GetTypeHash(InPin.ParentIndex);
 	Hash = HashCombine(Hash, GetTypeHash(InPin.Name));
@@ -523,7 +523,7 @@ uint32 GetTypeHash(const URigVMController::FPinInfo& InPin)
 	return Hash;
 }
 
-URigVMController::FPinInfoArray::FPinInfoArray(const URigVMNode* InNode)
+FRigVMPinInfoArray::FRigVMPinInfoArray(const URigVMNode* InNode)
 {
 	// this method adds all pins as currently represented in the model.
 	for(const URigVMPin* Pin : InNode->GetPins())
@@ -532,7 +532,7 @@ URigVMController::FPinInfoArray::FPinInfoArray(const URigVMNode* InNode)
 	}
 }
 
-int32 URigVMController::FPinInfoArray::AddPin(const URigVMPin* InPin, int32 InParentIndex, ERigVMPinDirection InDirection)
+int32 FRigVMPinInfoArray::AddPin(const URigVMPin* InPin, int32 InParentIndex, ERigVMPinDirection InDirection)
 {
 	// this method adds all pins as currently represented in the model.
 	const int32 Index = Pins.Emplace(InPin, InParentIndex, InDirection);
@@ -544,7 +544,7 @@ int32 URigVMController::FPinInfoArray::AddPin(const URigVMPin* InPin, int32 InPa
 	return Index;
 }
 
-URigVMController::FPinInfoArray::FPinInfoArray(const URigVMNode* InNode, URigVMController* InController, const FPinInfoArray* InPreviousPinInfos)
+FRigVMPinInfoArray::FRigVMPinInfoArray(const URigVMNode* InNode, URigVMController* InController, const FRigVMPinInfoArray* InPreviousPinInfos)
 {
 	const bool bAddSubPins = !InNode->IsA<URigVMRerouteNode>();
 	
@@ -557,8 +557,8 @@ URigVMController::FPinInfoArray::FPinInfoArray(const URigVMNode* InNode, URigVMC
 	}
 }
 
-URigVMController::FPinInfoArray::FPinInfoArray(const FRigVMGraphFunctionHeader& FunctionHeader,
-	URigVMController* InController, const FPinInfoArray* InPreviousPinInfos)
+FRigVMPinInfoArray::FRigVMPinInfoArray(const FRigVMGraphFunctionHeader& FunctionHeader,
+	URigVMController* InController, const FRigVMPinInfoArray* InPreviousPinInfos)
 {
 	// this method adds pins as needed based on the property structure
 	const FRigVMRegistry& Registry = FRigVMRegistry::Get();
@@ -570,7 +570,7 @@ URigVMController::FPinInfoArray::FPinInfoArray(const FRigVMGraphFunctionHeader&
 	}
 }
 
-int32 URigVMController::FPinInfoArray::AddPin(FProperty* InProperty, URigVMController* InController,
+int32 FRigVMPinInfoArray::AddPin(FProperty* InProperty, URigVMController* InController,
                                               ERigVMPinDirection InDirection, int32 InParentIndex, const uint8* InDefaultValueMemory, bool bAddSubPins)
 {
 	// this method adds pins as needed based on the property structure
@@ -603,13 +603,13 @@ int32 URigVMController::FPinInfoArray::AddPin(FProperty* InProperty, URigVMContr
 	return Index;
 }
 
-int32 URigVMController::FPinInfoArray::AddPin(URigVMController* InController, int32 InParentIndex, const FName& InName, ERigVMPinDirection InDirection,
-	TRigVMTypeIndex InTypeIndex, const FString& InDefaultValue, const uint8* InDefaultValueMemory, const FPinInfoArray* InPreviousPinInfos, bool bAddSubPins)
+int32 FRigVMPinInfoArray::AddPin(URigVMController* InController, int32 InParentIndex, const FName& InName, ERigVMPinDirection InDirection,
+	TRigVMTypeIndex InTypeIndex, const FString& InDefaultValue, const uint8* InDefaultValueMemory, const FRigVMPinInfoArray* InPreviousPinInfos, bool bAddSubPins)
 {
 	ensureMsgf(InTypeIndex != INDEX_NONE, TEXT("Invalid pin type for pin %s in %s"), *InName.ToString(), *InController->GetPackage()->GetPathName());
 	const FRigVMRegistry& Registry = FRigVMRegistry::Get();
 		
-	FPinInfo Info;
+	FRigVMPinInfo Info;
 	Info.ParentIndex = InParentIndex;
 	Info.Name = InName;
 	Info.Direction = InDirection;
@@ -626,7 +626,7 @@ int32 URigVMController::FPinInfoArray::AddPin(URigVMController* InController, in
 		const int32 PreviousIndex = InPreviousPinInfos->GetIndexFromPinPath(PinPath);
 		if(PreviousIndex != INDEX_NONE)
 		{
-			const FPinInfo& PreviousPin = (*InPreviousPinInfos)[PreviousIndex];
+			const FRigVMPinInfo& PreviousPin = (*InPreviousPinInfos)[PreviousIndex];
 			if(PreviousPin.TypeIndex == InTypeIndex)
 			{
 				Pins[Index].DefaultValue = PreviousPin.DefaultValue;
@@ -684,7 +684,7 @@ int32 URigVMController::FPinInfoArray::AddPin(URigVMController* InController, in
 	return Index;
 }
 
-void URigVMController::FPinInfoArray::AddPins(UScriptStruct* InScriptStruct, URigVMController* InController,
+void FRigVMPinInfoArray::AddPins(UScriptStruct* InScriptStruct, URigVMController* InController,
 	ERigVMPinDirection InDirection, int32 InParentIndex, const uint8* InDefaultValueMemory, bool bAddSubPins)
 {
 	if (InController->GetSchema()->ShouldUnfoldStruct(InController, InScriptStruct))
@@ -718,7 +718,7 @@ void URigVMController::FPinInfoArray::AddPins(UScriptStruct* InScriptStruct, URi
 	}
 }
 
-const FString& URigVMController::FPinInfoArray::GetPinPath(const int32 InIndex) const
+const FString& FRigVMPinInfoArray::GetPinPath(const int32 InIndex) const
 {
 	if(!Pins.IsValidIndex(InIndex))
 	{
@@ -741,7 +741,7 @@ const FString& URigVMController::FPinInfoArray::GetPinPath(const int32 InIndex)
 	return Pins[InIndex].PinPath;
 }
 
-int32 URigVMController::FPinInfoArray::GetIndexFromPinPath(const FString& InPinPath) const
+int32 FRigVMPinInfoArray::GetIndexFromPinPath(const FString& InPinPath) const
 {
 	if(PinPathLookup.Num() != Num())
 	{
@@ -759,7 +759,7 @@ int32 URigVMController::FPinInfoArray::GetIndexFromPinPath(const FString& InPinP
 	return INDEX_NONE;
 }
 
-const URigVMController::FPinInfo* URigVMController::FPinInfoArray::GetPinFromPinPath(const FString& InPinPath) const
+const FRigVMPinInfo* FRigVMPinInfoArray::GetPinFromPinPath(const FString& InPinPath) const
 {
 	const int32 Index = GetIndexFromPinPath(InPinPath);
 	if(Pins.IsValidIndex(Index))
@@ -769,7 +769,7 @@ const URigVMController::FPinInfo* URigVMController::FPinInfoArray::GetPinFromPin
 	return nullptr;
 }
 
-int32 URigVMController::FPinInfoArray::GetRootIndex(const int32 InIndex) const
+int32 FRigVMPinInfoArray::GetRootIndex(const int32 InIndex) const
 {
 	if(Pins.IsValidIndex(InIndex))
 	{
@@ -782,13 +782,13 @@ int32 URigVMController::FPinInfoArray::GetRootIndex(const int32 InIndex) const
 	return INDEX_NONE;
 }
 
-uint32 GetTypeHash(const URigVMController::FPinInfoArray& InPins)
+uint32 GetTypeHash(const FRigVMPinInfoArray& InPins)
 {
 	TArray<uint32> Hashes;
 	Hashes.Reserve(InPins.Num());
 	
 	uint32 OverAllHash = GetTypeHash(InPins.Num());
-	for(const URigVMController::FPinInfo& Info : InPins)
+	for(const FRigVMPinInfo& Info : InPins)
 	{
 		uint32 PinHash = GetTypeHash(Info);
 		if(Info.ParentIndex != INDEX_NONE)
@@ -13750,6 +13750,73 @@ FName URigVMController::AddDecorator(URigVMNode* InNode, UScriptStruct* InDecora
 	
 	AddPinsForStruct(InDecoratorScriptStruct, InNode, DecoratorPin, DecoratorPin->GetDirection(), InDefaultValue, true);
 
+	FRigVMPinInfoArray ProgrammaticPins;
+	Decorator->GetProgrammaticPins(this, INDEX_NONE, ProgrammaticPins);
+
+	const FRigVMRegistry& Registry = FRigVMRegistry::Get();
+	auto CreatePinFromPinInfo = [this, &Registry](const FRigVMPinInfo& InPinInfo, const FString& InPinPath, UObject* InOuter) -> URigVMPin*
+	{
+		check(InOuter);
+		URigVMPin* Pin = NewObject<URigVMPin>(InOuter, InPinInfo.Name);
+		if (InPinInfo.Property)
+		{
+			ConfigurePinFromProperty(InPinInfo.Property, Pin, InPinInfo.Direction);
+		}
+		else
+		{
+			const FRigVMTemplateArgumentType& Type = Registry.GetType(InPinInfo.TypeIndex);
+			Pin->CPPType = Type.CPPType.ToString();
+			Pin->CPPTypeObject = Type.CPPTypeObject;
+			if (Pin->CPPTypeObject)
+			{
+				Pin->CPPTypeObjectPath = *Pin->CPPTypeObject->GetPathName();
+			}
+			if (Registry.IsExecuteType(InPinInfo.TypeIndex))
+			{
+				MakeExecutePin(Pin);
+			}
+
+			Pin->Direction = InPinInfo.Direction;
+			Pin->DisplayName = InPinInfo.DisplayName.IsEmpty() ? NAME_None : FName(*InPinInfo.DisplayName);
+			Pin->bIsConstant = InPinInfo.bIsConstant;
+			Pin->bIsDynamicArray = InPinInfo.bIsDynamicArray;
+			Pin->CustomWidgetName = InPinInfo.CustomWidgetName.IsEmpty() ? NAME_None : FName(*InPinInfo.CustomWidgetName);
+		}
+
+		Pin->bIsExpanded = InPinInfo.bIsExpanded;
+		Pin->DefaultValue = InPinInfo.DefaultValue;
+
+		if (URigVMPin* ParentPin = Cast<URigVMPin>(InOuter))
+		{
+			AddSubPin(ParentPin, Pin);
+		}
+		else if (URigVMNode* OwnerNode = Cast<URigVMNode>(InOuter))
+		{
+			AddNodePin(OwnerNode, Pin);
+		}
+		else
+		{
+			ensureMsgf(false, TEXT("Outer %s of pin info %s is not a pin or a node"), *InOuter->GetPathName(), *InPinPath);
+		}
+
+		Notify(ERigVMGraphNotifType::PinAdded, Pin);
+
+		return Pin;
+	};
+
+	for (int32 PinIndex = 0; PinIndex < ProgrammaticPins.Num(); ++PinIndex)
+	{
+		const FString& PinPath = ProgrammaticPins.GetPinPath(PinIndex);
+		FString ParentPinPath, PinName;
+		UObject* OuterForPin = DecoratorPin;
+		if (URigVMPin::SplitPinPathAtEnd(PinPath, ParentPinPath, PinName))
+		{
+			OuterForPin = DecoratorPin->FindSubPin(ParentPinPath);
+		}
+
+		CreatePinFromPinInfo(ProgrammaticPins[PinIndex], PinPath, OuterForPin);
+	}
+
 	// move the the pin to the right index as required
 	if(DecoratorPin->GetPinIndex() != InPinIndex &&
 		InPinIndex >=0 && InPinIndex < InNode->GetPins().Num())
@@ -14036,7 +14103,7 @@ bool URigVMController::IsValidLinkForGraph(const URigVMLink* InLink)
 	return true;
 }
 
-void URigVMController::AddPinsForStruct(UStruct* InStruct, URigVMNode* InNode, URigVMPin* InParentPin, ERigVMPinDirection InPinDirection, const FString& InDefaultValue, bool bAutoExpandArrays, const FPinInfoArray* PreviousPins)
+void URigVMController::AddPinsForStruct(UStruct* InStruct, URigVMNode* InNode, URigVMPin* InParentPin, ERigVMPinDirection InPinDirection, const FString& InDefaultValue, bool bAutoExpandArrays, const FRigVMPinInfoArray* PreviousPins)
 {
 	if(!InStruct->IsChildOf(FRigVMStruct::StaticStruct()))
 	{
@@ -14349,13 +14416,15 @@ void URigVMController::AddPinsForTemplate(const FRigVMTemplate* InTemplate, cons
 
 void URigVMController::ConfigurePinFromProperty(FProperty* InProperty, URigVMPin* InOutPin, ERigVMPinDirection InPinDirection)
 {
+	const ERigVMPinDirection PropertyPinDirection = FRigVMStruct::GetPinDirectionFromProperty(InProperty);
 	if (InPinDirection == ERigVMPinDirection::Invalid)
 	{
-		InOutPin->Direction = FRigVMStruct::GetPinDirectionFromProperty(InProperty);
+		InOutPin->Direction = PropertyPinDirection;
 	}
 	else
 	{
-		InOutPin->Direction = InPinDirection;
+		// Keep the pin direction specified unless we wish to be hidden
+		InOutPin->Direction = PropertyPinDirection == ERigVMPinDirection::Hidden ? ERigVMPinDirection::Hidden : InPinDirection;
 	}
 
 #if WITH_EDITOR
@@ -14810,9 +14879,9 @@ void URigVMController::RepopulatePinsOnNode(URigVMNode* InNode, bool bFollowCore
 		}
 	}
 
-	const FPinInfoArray PreviousPinInfos(InNode);
+	const FRigVMPinInfoArray PreviousPinInfos(InNode);
 	const uint32 PreviousPinHash = GetTypeHash(PreviousPinInfos);
-	FPinInfoArray NewPinInfos;
+	FRigVMPinInfoArray NewPinInfos;
 
 	// step 2/3: clear pins on the node and repopulate the node with new pins
 	if (UnitNode != nullptr)
@@ -15026,7 +15095,7 @@ void URigVMController::RepopulatePinsOnNode(URigVMNode* InNode, bool bFollowCore
 			}
 		}
 
-		NewPinInfos = FPinInfoArray(InNode, this, &PreviousPinInfos);
+		NewPinInfos = FRigVMPinInfoArray(InNode, this, &PreviousPinInfos);
 	}
 	else if (EntryNode || ReturnNode)
 	{
@@ -15085,20 +15154,20 @@ void URigVMController::RepopulatePinsOnNode(URigVMNode* InNode, bool bFollowCore
 	}
 	else if (CollapseNode)
 	{
-		NewPinInfos = FPinInfoArray(InNode, this, &PreviousPinInfos);
+		NewPinInfos = FRigVMPinInfoArray(InNode, this, &PreviousPinInfos);
 	}
 	else if (FunctionRefNode)
 	{
 		const FRigVMGraphFunctionHeader& FunctionHeader = FunctionRefNode->GetReferencedFunctionHeader();
 		if(FunctionHeader.IsValid())
 		{
-			NewPinInfos = FPinInfoArray(FunctionHeader, this, &PreviousPinInfos);
+			NewPinInfos = FRigVMPinInfoArray(FunctionHeader, this, &PreviousPinInfos);
 		}
 		// if we can't find referenced node anymore
 		// let's keep all pins
 		else
 		{
-			NewPinInfos = FPinInfoArray(FunctionRefNode, this, &PreviousPinInfos);
+			NewPinInfos = FRigVMPinInfoArray(FunctionRefNode, this, &PreviousPinInfos);
 		}
 	}
 	else
@@ -15128,11 +15197,19 @@ void URigVMController::RepopulatePinsOnNode(URigVMNode* InNode, bool bFollowCore
 	// make sure the new pin infos contains the decorator pins from the last run
 	for(int32 Index = 0; Index < PreviousPinInfos.Num(); Index++)
 	{
-		const FPinInfo& PreviousPin = PreviousPinInfos[Index]; 
+		const FRigVMPinInfo& PreviousPin = PreviousPinInfos[Index]; 
 		if(PreviousPin.bIsDecorator)
 		{
 			const int32 NewPinIndex = NewPinInfos.AddPin(this, INDEX_NONE, PreviousPin.Name, PreviousPin.Direction, PreviousPin.TypeIndex, PreviousPin.DefaultValue, nullptr, &PreviousPinInfos, true);
 			NewPinInfos[NewPinIndex].bIsDecorator = true;
+
+			if (URigVMPin* Pin = InNode->FindPin(PreviousPin.PinPath))
+			{
+				TSharedPtr<FStructOnScope> DecoratorScope = Pin->GetDecoratorInstance();
+				FRigVMDecorator* VMDecorator = (FRigVMDecorator*)DecoratorScope->GetStructMemory();
+
+				VMDecorator->GetProgrammaticPins(this, NewPinIndex, NewPinInfos);
+			}
 		}
 	}
 
@@ -15268,7 +15345,7 @@ void URigVMController::RepopulatePinsOnNode(URigVMNode* InNode, bool bFollowCore
 		}
 	}
 
-	auto CreatePinFromPinInfo = [this, &Registry, &PreviousPinInfos](const FPinInfo& InPinInfo, const FString& InPinPath, UObject* InOuter) -> URigVMPin*
+	auto CreatePinFromPinInfo = [this, &Registry, &PreviousPinInfos](const FRigVMPinInfo& InPinInfo, const FString& InPinPath, UObject* InOuter) -> URigVMPin*
 	{
 		check(InOuter);
 		URigVMPin* Pin = NewObject<URigVMPin>(InOuter, InPinInfo.Name);
@@ -15301,7 +15378,7 @@ void URigVMController::RepopulatePinsOnNode(URigVMNode* InNode, bool bFollowCore
 		Pin->DefaultValue = InPinInfo.DefaultValue;
 
 		// reuse expansion state and default value
-		if(const FPinInfo* PreviousPin = PreviousPinInfos.GetPinFromPinPath(InPinPath))
+		if(const FRigVMPinInfo* PreviousPin = PreviousPinInfos.GetPinFromPinPath(InPinPath))
 		{
 			if(PreviousPin->TypeIndex == InPinInfo.TypeIndex)
 			{
@@ -15428,7 +15505,7 @@ void URigVMController::RepopulatePinsOnNode(URigVMNode* InNode, bool bFollowCore
 	for(int32 Index = 0; Index < PreviousPinsToUpdate.Num(); Index++)
 	{
 		const FString& PinPath = PreviousPinInfos.GetPinPath(PreviousPinsToUpdate[Index]);
-		const FPinInfo* NewPinInfo = NewPinInfos.GetPinFromPinPath(PinPath);
+		const FRigVMPinInfo* NewPinInfo = NewPinInfos.GetPinFromPinPath(PinPath);
 		check(NewPinInfo);
 		
 		if(URigVMPin* Pin = InNode->FindPin(PinPath))
@@ -15469,7 +15546,7 @@ void URigVMController::RepopulatePinsOnNode(URigVMNode* InNode, bool bFollowCore
 	TMap<FString, TArray<FName>> PinOrder;
 	for(int32 Index = 0; Index < NewPinInfos.Num(); Index++)
 	{
-		const FPinInfo& NewPin = NewPinInfos[Index];
+		const FRigVMPinInfo& NewPin = NewPinInfos[Index];
 		FString ParentPinPath;
 		if(NewPin.ParentIndex != INDEX_NONE)
 		{
diff --git a/Engine/Plugins/Runtime/RigVM/Source/RigVMDeveloper/Private/RigVMModel/RigVMNode.cpp b/Engine/Plugins/Runtime/RigVM/Source/RigVMDeveloper/Private/RigVMModel/RigVMNode.cpp
--- a/Engine/Plugins/Runtime/RigVM/Source/RigVMDeveloper/Private/RigVMModel/RigVMNode.cpp
+++ b/Engine/Plugins/Runtime/RigVM/Source/RigVMDeveloper/Private/RigVMModel/RigVMNode.cpp
@@ -282,7 +282,10 @@ bool URigVMNode::IsMutable() const
 	{
 		if(const UScriptStruct* ScriptStruct = Pin->GetScriptStruct())
 		{
-			return ScriptStruct->IsChildOf(FRigVMExecuteContext::StaticStruct());
+			if (ScriptStruct->IsChildOf(FRigVMExecuteContext::StaticStruct()))
+			{
+				return true;
+			}
 		}
 	}
 	return false;
diff --git a/Engine/Plugins/Runtime/RigVM/Source/RigVMDeveloper/Public/EdGraph/RigVMEdGraphNode.h b/Engine/Plugins/Runtime/RigVM/Source/RigVMDeveloper/Public/EdGraph/RigVMEdGraphNode.h
--- a/Engine/Plugins/Runtime/RigVM/Source/RigVMDeveloper/Public/EdGraph/RigVMEdGraphNode.h
+++ b/Engine/Plugins/Runtime/RigVM/Source/RigVMDeveloper/Public/EdGraph/RigVMEdGraphNode.h
@@ -147,6 +147,9 @@ public:
 	/** Get the VM node name this node is wrapping */
 	FName GetModelNodeName() const;
 
+	/** Get the VM node path this node is wrapping */
+	const FString& GetModelNodePath() const;
+
 	URigVMPin* GetModelPinFromPinPath(const FString& InPinPath) const;
 
 	/** Add a new element to the aggregate node referred to by the property path */
diff --git a/Engine/Plugins/Runtime/RigVM/Source/RigVMDeveloper/Public/RigVMCompiler/RigVMAST.h b/Engine/Plugins/Runtime/RigVM/Source/RigVMDeveloper/Public/RigVMCompiler/RigVMAST.h
--- a/Engine/Plugins/Runtime/RigVM/Source/RigVMDeveloper/Public/RigVMCompiler/RigVMAST.h
+++ b/Engine/Plugins/Runtime/RigVM/Source/RigVMDeveloper/Public/RigVMCompiler/RigVMAST.h
@@ -1317,8 +1317,8 @@ private:
 	void FoldAssignments();
 
 	// helper function to inline all contributing nodes of the graph
-	void Inline(TArray<URigVMGraph*> InGraphs);
-	void Inline(TArray<URigVMGraph*> InGraphs, const TArray<FRigVMASTProxy>& InNodeProxies);
+	void Inline(const TArray<URigVMGraph*>& InGraphs);
+	void Inline(const TArray<URigVMGraph*>& InGraphs, const TArray<FRigVMASTProxy>& InNodeProxies);
 
 	// helper functions to retrieve links for a given pin
 	TArray<int32> GetSourceLinkIndices(const FRigVMASTProxy& InPinProxy, bool bRecursive = false) const;
diff --git a/Engine/Plugins/Runtime/RigVM/Source/RigVMDeveloper/Public/RigVMModel/RigVMController.h b/Engine/Plugins/Runtime/RigVM/Source/RigVMDeveloper/Public/RigVMModel/RigVMController.h
--- a/Engine/Plugins/Runtime/RigVM/Source/RigVMDeveloper/Public/RigVMModel/RigVMController.h
+++ b/Engine/Plugins/Runtime/RigVM/Source/RigVMDeveloper/Public/RigVMModel/RigVMController.h
@@ -166,6 +166,62 @@ private:
 	friend class URigVMController;
 };
 
+struct RIGVMDEVELOPER_API FRigVMPinInfo
+{
+	FRigVMPinInfo();
+	FRigVMPinInfo(const URigVMPin* InPin, int32 InParentIndex, ERigVMPinDirection InDirection = ERigVMPinDirection::Invalid);
+	FRigVMPinInfo(FProperty* InProperty, ERigVMPinDirection InDirection, int32 InParentIndex, const uint8* InDefaultValueMemory);
+
+	void CorrectExecuteTypeIndex();
+
+	int32 ParentIndex;
+	FName Name;
+	ERigVMPinDirection Direction;
+	TRigVMTypeIndex TypeIndex;
+	bool bIsArray;
+	FProperty* Property;
+	FString PinPath;
+	FString DefaultValue;
+	FString DisplayName;
+	FString CustomWidgetName;
+	bool bIsExpanded;
+	bool bIsConstant;
+	bool bIsDynamicArray;
+	bool bIsDecorator;
+	TArray<int32> SubPins;
+
+	friend uint32 GetTypeHash(const FRigVMPinInfo& InPin);
+};
+
+struct RIGVMDEVELOPER_API FRigVMPinInfoArray
+{
+	FRigVMPinInfoArray() {}
+	explicit FRigVMPinInfoArray(const URigVMNode* InNode);
+	FRigVMPinInfoArray(const URigVMNode* InNode, URigVMController* InController, const FRigVMPinInfoArray* InPreviousPinInfos = nullptr);
+	FRigVMPinInfoArray(const FRigVMGraphFunctionHeader& FunctionHeader, URigVMController* InController, const FRigVMPinInfoArray* InPreviousPinInfos = nullptr);
+
+	int32 Num() const { return Pins.Num(); }
+	const FRigVMPinInfo& operator[](int32 InIndex) const { return Pins[InIndex]; }
+	FRigVMPinInfo& operator[](int32 InIndex) { return Pins[InIndex]; }
+	TArray<FRigVMPinInfo>::RangedForIteratorType begin() const { return Pins.begin(); }
+	TArray<FRigVMPinInfo>::RangedForIteratorType end() const { return Pins.end(); }
+
+	int32 AddPin(const URigVMPin* InPin, int32 InParentIndex, ERigVMPinDirection InDirection = ERigVMPinDirection::Invalid);
+	int32 AddPin(FProperty* InProperty, URigVMController* InController, ERigVMPinDirection InDirection, int32 InParentIndex, const uint8* InDefaultValueMemory, bool bAddSubPins);
+	int32 AddPin(URigVMController* InController, int32 InParentIndex, const FName& InName, ERigVMPinDirection InDirection, TRigVMTypeIndex InTypeIndex, const FString& InDefaultValue, const uint8* InDefaultValueMemory, const FRigVMPinInfoArray* InPreviousPinInfos, bool bAddSubPins);
+	void AddPins(UScriptStruct* InScriptStruct, URigVMController* InController, ERigVMPinDirection InDirection, int32 InParentIndex, const uint8* InDefaultValueMemory, bool bAddSubPins);
+
+	const FString& GetPinPath(const int32 InIndex) const;
+	int32 GetIndexFromPinPath(const FString& InPinPath) const;
+	const FRigVMPinInfo* GetPinFromPinPath(const FString& InPinPath) const;
+	int32 GetRootIndex(const int32 InIndex) const;
+
+	friend uint32 GetTypeHash(const FRigVMPinInfoArray& InPins);
+
+	mutable TArray<FRigVMPinInfo> Pins;
+	mutable TMap<FString, int32> PinPathLookup;;
+};
+
 /**
  * The Controller is the sole authority to perform changes
  * on the Graph. The Controller itself is stateless.
@@ -1034,64 +1090,6 @@ public:
 	UFUNCTION(BlueprintCallable, Category = RigVMController)
 	void SetIsRunningUnitTest(bool bIsRunning);
 
-public:
-	
-	struct FPinInfo
-	{
-		FPinInfo();
-		FPinInfo(const URigVMPin* InPin, int32 InParentIndex, ERigVMPinDirection InDirection = ERigVMPinDirection::Invalid);
-		FPinInfo(FProperty* InProperty, ERigVMPinDirection InDirection, int32 InParentIndex, const uint8* InDefaultValueMemory);
-
-		void CorrectExecuteTypeIndex();
-		
-		int32 ParentIndex;
-		FName Name;
-		ERigVMPinDirection Direction;
-		TRigVMTypeIndex TypeIndex;
-		bool bIsArray;
-		FProperty* Property;
-		FString PinPath;
-		FString DefaultValue;
-		FString DisplayName;
-		FString CustomWidgetName;
-		bool bIsExpanded;
-		bool bIsConstant;
-		bool bIsDynamicArray;
-		bool bIsDecorator;
-		TArray<int32> SubPins;
-
-		friend uint32 GetTypeHash(const FPinInfo& InPin);
-	};
-
-	struct FPinInfoArray
-	{
-		FPinInfoArray() {}
-		FPinInfoArray(const URigVMNode* InNode);
-		FPinInfoArray(const URigVMNode* InNode, URigVMController* InController, const FPinInfoArray* InPreviousPinInfos = nullptr);
-		FPinInfoArray(const FRigVMGraphFunctionHeader& FunctionHeader, URigVMController* InController, const FPinInfoArray* InPreviousPinInfos = nullptr);
-
-		int32 Num() const { return Pins.Num(); }
-		const FPinInfo& operator[](int32 InIndex) const { return Pins[InIndex]; }
-		FPinInfo& operator[](int32 InIndex) { return Pins[InIndex]; }
-		TArray<FPinInfo>::RangedForIteratorType begin() const { return Pins.begin(); }
-		TArray<FPinInfo>::RangedForIteratorType end() const { return Pins.end(); }
-
-		int32 AddPin(const URigVMPin* InPin, int32 InParentIndex, ERigVMPinDirection InDirection = ERigVMPinDirection::Invalid);
-		int32 AddPin(FProperty* InProperty, URigVMController* InController, ERigVMPinDirection InDirection, int32 InParentIndex, const uint8* InDefaultValueMemory, bool bAddSubPins);
-		int32 AddPin(URigVMController* InController, int32 InParentIndex, const FName& InName, ERigVMPinDirection InDirection, TRigVMTypeIndex InTypeIndex, const FString& InDefaultValue, const uint8* InDefaultValueMemory, const FPinInfoArray* InPreviousPinInfos, bool bAddSubPins);
-		void AddPins(UScriptStruct* InScriptStruct, URigVMController* InController, ERigVMPinDirection InDirection, int32 InParentIndex, const uint8* InDefaultValueMemory, bool bAddSubPins);
-
-		const FString& GetPinPath(const int32 InIndex) const;
-		int32 GetIndexFromPinPath(const FString& InPinPath) const;
-		const FPinInfo* GetPinFromPinPath(const FString& InPinPath) const;
-		int32 GetRootIndex(const int32 InIndex) const;
-
-		friend uint32 GetTypeHash(const FPinInfoArray& InPins);
-		
-		mutable TArray<FPinInfo> Pins;
-		mutable TMap<FString, int32> PinPathLookup;;
-	};
-
 private:
 
 	UPROPERTY(BlueprintReadOnly, Category = RigVMController, meta = (ScriptName = "ModifiedEvent", AllowPrivateAccess = "true"))
@@ -1106,7 +1104,7 @@ private:
 	bool IsValidNodeForGraph(const URigVMNode* InNode);
 	bool IsValidPinForGraph(const URigVMPin* InPin);
 	bool IsValidLinkForGraph(const URigVMLink* InLink);
-	void AddPinsForStruct(UStruct* InStruct, URigVMNode* InNode, URigVMPin* InParentPin, ERigVMPinDirection InPinDirection, const FString& InDefaultValue, bool bAutoExpandArrays, const FPinInfoArray* PreviousPins = nullptr);
+	void AddPinsForStruct(UStruct* InStruct, URigVMNode* InNode, URigVMPin* InParentPin, ERigVMPinDirection InPinDirection, const FString& InDefaultValue, bool bAutoExpandArrays, const FRigVMPinInfoArray* PreviousPins = nullptr);
 	void AddPinsForArray(FArrayProperty* InArrayProperty, URigVMNode* InNode, URigVMPin* InParentPin, ERigVMPinDirection InPinDirection, const TArray<FString>& InDefaultValues, bool bAutoExpandArrays);
 	void AddPinsForTemplate(const FRigVMTemplate* InTemplate, const FRigVMTemplateTypeMap& InPinTypeMap, URigVMNode* InNode);
 	void ConfigurePinFromProperty(FProperty* InProperty, URigVMPin* InOutPin, ERigVMPinDirection InPinDirection = ERigVMPinDirection::Invalid);
@@ -1451,7 +1449,7 @@ private:
 	friend struct FRigVMEjectNodeFromPinAction;
 	friend class FRigVMParserAST;
 	friend class FRigVMControllerCompileBracketScope;
-	friend struct FPinInfoArray;
+	friend struct FRigVMPinInfoArray;
 	friend class FRigVMControllerNotifGuard;
 	friend struct FRigVMClient;
 	friend struct FRigVMActionWrapper;
diff --git a/Engine/Plugins/Runtime/RigVM/Source/RigVMEditor/Public/Widgets/SRigVMGraphNode.h b/Engine/Plugins/Runtime/RigVM/Source/RigVMEditor/Public/Widgets/SRigVMGraphNode.h
--- a/Engine/Plugins/Runtime/RigVM/Source/RigVMEditor/Public/Widgets/SRigVMGraphNode.h
+++ b/Engine/Plugins/Runtime/RigVM/Source/RigVMEditor/Public/Widgets/SRigVMGraphNode.h
@@ -62,7 +62,7 @@ public:
 	/** Callback function executed when Add pin button is clicked */
 	virtual FReply OnAddPin() override;
 
-private:
+protected:
 
 	bool UseLowDetailNodeContent() const;
 
@@ -86,14 +86,14 @@ private:
 	FText GetInstructionCountText() const;
 	FText GetInstructionDurationText() const;
 
-private:
+protected:
 
 	int32 GetNodeTopologyVersion() const;
 	EVisibility GetPinVisibility(int32 InPinInfoIndex, bool bAskingForSubPin) const;
 	const FSlateBrush * GetExpanderImage(int32 InPinInfoIndex, bool bLeft, bool bHovered) const;
 	FReply OnExpanderArrowClicked(int32 InPinInfoIndex);
 	void HandleModifiedEvent(ERigVMGraphNotifType InNotifType, URigVMGraph* InGraph, UObject* InSubject);
-	void UpdatePinTreeView();
+	virtual void UpdatePinTreeView();
 
 	/** Cached widget title area */
 	TSharedPtr<SOverlay> TitleAreaWidget;
