[Backout] - CL26927079
Original CL Desc
-----------------------------------------------------------------
Implement animnext graph decorator authoring
Moved decorator offset computation to runtime to support #ifdef and adding/removing members
Added a node template builder to ensure node templates are always fully formed when used
Add support for RigVM decorator dynamic pins
Other minor imrovements to decorator runtime

diff --git a/Engine/Plugins/Experimental/AnimNext/Source/AnimNext/AnimNext.Build.cs b/Engine/Plugins/Experimental/AnimNext/Source/AnimNext/AnimNext.Build.cs
--- a/Engine/Plugins/Experimental/AnimNext/Source/AnimNext/AnimNext.Build.cs
+++ b/Engine/Plugins/Experimental/AnimNext/Source/AnimNext/AnimNext.Build.cs
@@ -5,7 +5,7 @@ namespace UnrealBuildTool.Rules
 	public class AnimNext : ModuleRules
 	{
 		public AnimNext(ReadOnlyTargetRules Target) : base(Target)
-		{
+		{	
 			PublicDependencyModuleNames.AddRange(
 				new string[]
 				{
@@ -17,23 +17,12 @@ namespace UnrealBuildTool.Rules
 					"Engine"
 				}
 			);
-
 			PrivateDependencyModuleNames.AddRange(
 				new string[]
 				{
 					"AssetRegistry"
 				}
 			);
-
-			if (Target.bBuildEditor == true)
-			{
-				PrivateDependencyModuleNames.AddRange(
-					new string[]
-					{
-						"RigVMDeveloper"
-					}
-				);
-			}
 		}
 	}
 }
\ No newline at end of file
diff --git a/Engine/Plugins/Experimental/AnimNext/Source/AnimNext/Internal/DecoratorBase/DecoratorReader.h b/Engine/Plugins/Experimental/AnimNext/Source/AnimNext/Internal/DecoratorBase/DecoratorReader.h
--- a/Engine/Plugins/Experimental/AnimNext/Source/AnimNext/Internal/DecoratorBase/DecoratorReader.h
+++ b/Engine/Plugins/Experimental/AnimNext/Source/AnimNext/Internal/DecoratorBase/DecoratorReader.h
@@ -3,9 +3,6 @@
 #pragma once
 
 #include "CoreMinimal.h"
-#include "DecoratorBase/DecoratorHandle.h"
-#include "DecoratorBase/EntryPointHandle.h"
-#include "DecoratorBase/NodeHandle.h"
 #include "Serialization/ArchiveProxy.h"
 
 namespace UE::AnimNext
@@ -18,53 +15,11 @@ namespace UE::AnimNext
 	  *     - A list of FNodeTemplates that the nodes use
 	  *     - The graph shared data (FNodeDescription for every node)
 	  */
-	class ANIMNEXT_API FDecoratorReader final : public FArchiveProxy
+	class ANIMNEXT_API FDecoratorReader : public FArchiveProxy
 	{
 	public:
-		// The largest size allowed for the shared data of a single graph
-		// Node/decorator handles use an unsigned 24 bits value to represent graph offsets
-		// which limits us to 16 MB
-		// We could extend this range by leveraging the fact that nodes have a minimum required alignment
-		static constexpr uint32 MAXIMUM_GRAPH_SHARED_DATA_SIZE = (1 << 24) - 1;
+		FDecoratorReader(FArchive& Ar);
 
-		enum class EErrorState
-		{
-			None,						// All good, no error
-			GraphTooLarge,				// Exceeded the maximum graph size, @see FDecoratorReader::MAXIMUM_GRAPH_SHARED_DATA_SIZE
-			NodeSharedDataTooLarge,		// Exceeded the maximum node shared data size, @see FNodeDescription::MAXIMUM_NODE_SHARED_DATA_SIZE
-			NodeInstanceDataTooLarge,	// Exceeded the maximum node instance data size, @see FNodeInstance::MAXIMUM_NODE_INSTANCE_DATA_SIZE
-		};
-
-		explicit FDecoratorReader(FArchive& Ar);
-
-		[[nodiscard]] EErrorState ReadGraph(TArray<uint8>& GraphSharedData, TArray<TObjectPtr<UObject>>& TrackedObjectsForGC);
-
-		// Takes a node handle representing a node index and resolves it into a node handle representing a shared data offset
-		// Must be called after ReadGraphSharedData as it populates the necessary data
-		[[nodiscard]] FNodeHandle ResolveNodeHandle(FNodeHandle NodeHandle) const;
-
-		// Takes a decorator handle representing a node index and resolves it into a node handle representing a shared data offset
-		// Must be called after ReadGraphSharedData as it populates the necessary data
-		[[nodiscard]] FAnimNextDecoratorHandle ResolveDecoratorHandle(FAnimNextDecoratorHandle DecoratorHandle) const;
-
-		// Takes an entry point handle representing a node index and resolves it into a node handle representing a shared data offset
-		// Must be called after ReadGraphSharedData as it populates the necessary data
-		[[nodiscard]] FAnimNextDecoratorHandle ResolveEntryPointHandle(FAnimNextEntryPointHandle EntryPointHandle) const;
-
-		// FArchive implementation
-		virtual FArchive& operator<<(UObject*& Obj) override;
-		virtual FArchive& operator<<(FObjectPtr& Obj) override;
-
-	private:
-		// Call first to read the graph shared data
-		[[nodiscard]] EErrorState ReadGraphSharedData(TArray<uint8>& GraphSharedData);
-
-		// Call second to read the tracked UObjects for GC
-		[[nodiscard]] EErrorState ReadTrackedObjectsForGC(TArray<TObjectPtr<UObject>>& TrackedObjectsForGC);
-
-		// A list of node handles for each node within the archive
-		TArray<FNodeHandle> NodeHandles;
-
-		friend FAnimNextDecoratorHandle;
+		void ReadGraphSharedData(TArray<uint8>& GraphSharedData);
 	};
 }
diff --git a/Engine/Plugins/Experimental/AnimNext/Source/AnimNext/Internal/DecoratorBase/DecoratorRegistry.h b/Engine/Plugins/Experimental/AnimNext/Source/AnimNext/Internal/DecoratorBase/DecoratorRegistry.h
--- a/Engine/Plugins/Experimental/AnimNext/Source/AnimNext/Internal/DecoratorBase/DecoratorRegistry.h
+++ b/Engine/Plugins/Experimental/AnimNext/Source/AnimNext/Internal/DecoratorBase/DecoratorRegistry.h
@@ -21,40 +21,33 @@ namespace UE::AnimNext
 		// Access the global registry
 		static FDecoratorRegistry& Get();
 
+		// Finds and returns the decorator handle for the provided decorator UID or an invalid
+		// handle if that decorator hasn't been registered yet.
+		FDecoratorRegistryHandle FindHandle(uint32 DecoratorUID) const;
+
 		// Finds and returns the decorator handle for the provided decorator UID or an invalid
 		// handle if that decorator hasn't been registered yet.
 		FDecoratorRegistryHandle FindHandle(FDecoratorUID DecoratorUID) const;
 
-		// Finds and returns the decorator associated with the provided handle.
-		// If the handle is not valid, nullptr is returned.
+		// Finds and returns the decorator associated with the provided handle. If the handle
+		// is invalid, nullptr is returned.
 		const FDecorator* Find(FDecoratorRegistryHandle DecoratorHandle) const;
 
-		// Finds and returns the decorator associated with the provided decorator UID.
-		// If the decorator is not registered, nullptr is returned.
+		// Finds and returns the decorator associated with the provided decorator UID. If the decorator
+		// is registered, nullptr is returned.
 		const FDecorator* Find(FDecoratorUID DecoratorUID) const;
 
-		// Finds and returns the decorator associated with the provided decorator shared data UScriptStruct.
-		// If the matching decorator is not registered, nullptr is returned.
-		const FDecorator* Find(const UScriptStruct* DecoratorSharedDataStruct) const;
-
 		// Registers a decorator dynamically
 		void Register(FDecorator* Decorator);
 
 		// Unregisters a decorator dynamically
 		void Unregister(FDecorator* Decorator);
 
-		// Returns a list of all registered decorators
-		TArray<const FDecorator*> GetDecorators() const;
-
 		// Returns the number of registered decorators
 		uint32 GetNum() const;
 
 	private:
 		FDecoratorRegistry() = default;
-		FDecoratorRegistry(const FDecoratorRegistry&) = delete;
-		FDecoratorRegistry(FDecoratorRegistry&&) = default;
-		FDecoratorRegistry& operator=(const FDecoratorRegistry&) = delete;
-		FDecoratorRegistry& operator=(FDecoratorRegistry&&) = default;
 
 		// Static init lifetime functions
 		static void StaticRegister(DecoratorConstructorFunc DecoratorConstructor);
diff --git a/Engine/Plugins/Experimental/AnimNext/Source/AnimNext/Internal/DecoratorBase/DecoratorTemplate.h b/Engine/Plugins/Experimental/AnimNext/Source/AnimNext/Internal/DecoratorBase/DecoratorTemplate.h
--- a/Engine/Plugins/Experimental/AnimNext/Source/AnimNext/Internal/DecoratorBase/DecoratorTemplate.h
+++ b/Engine/Plugins/Experimental/AnimNext/Source/AnimNext/Internal/DecoratorBase/DecoratorTemplate.h
@@ -26,8 +26,17 @@ namespace UE::AnimNext
 	 */
 	struct FDecoratorTemplate
 	{
+		FDecoratorTemplate(FDecoratorUID UID_, FDecoratorRegistryHandle RegistryHandle_, EDecoratorMode Mode_, uint32 AdditiveIndexOrNumAdditive_, uint32 NodeSharedOffset_, uint32 NodeInstanceOffset_)
+			: UID(UID_.GetUID())
+			, RegistryHandle(RegistryHandle_)
+			, Mode(static_cast<uint8>(Mode_))
+			, AdditiveIndexOrNumAdditive(AdditiveIndexOrNumAdditive_)
+			, NodeSharedOffset(NodeSharedOffset_)
+			, NodeInstanceOffset(NodeInstanceOffset_)
+		{}
+
 		// Returns the globally unique identifier for the decorator
-		FDecoratorUID GetUID() const { return FDecoratorUID(UID); }
+		uint32 GetUID() const { return UID; }
 
 		// Returns the decorator registry handle
 		FDecoratorRegistryHandle GetRegistryHandle() const { return RegistryHandle; }
@@ -83,19 +92,7 @@ namespace UE::AnimNext
 		ANIMNEXT_API void Serialize(FArchive& Ar);
 
 	private:
-		friend struct FNodeTemplate;
-		friend struct FNodeTemplateBuilder;
-
-		FDecoratorTemplate(FDecoratorUID InUID, FDecoratorRegistryHandle InRegistryHandle, EDecoratorMode InMode, uint32 InAdditiveIndexOrNumAdditive)
-			: UID(InUID.GetUID())
-			, RegistryHandle(InRegistryHandle)
-			, Mode(static_cast<uint8>(InMode))
-			, AdditiveIndexOrNumAdditive(InAdditiveIndexOrNumAdditive)
-			, NodeSharedOffset(0)
-			, NodeInstanceOffset(0)				// For instance data, 0 is an invalid offset since the data follows an instance of FNodeInstance
-		{}
-
-		FDecoratorUIDRaw			UID;			// decorator globally unique identifier
+		uint32	UID;							// decorator globally unique identifier
 
 		FDecoratorRegistryHandle	RegistryHandle;	// decorator registry handle
 		uint8						Mode;			// the decorator mode (we only need 1 bit)
@@ -106,8 +103,8 @@ namespace UE::AnimNext
 		uint8	AdditiveIndexOrNumAdditive;
 
 		// Offsets into the shared read-only and instance data portions of a node
-		uint16	NodeSharedOffset;				// relative to root of node description data (max 64 KB per node) (not serialized, @see FNodeTemplate::Finalize)
-		uint16	NodeInstanceOffset;				// relative to root of node instance data (max 64 KB per node) (not serialized, @see FNodeTemplate::Finalize)
+		uint16	NodeSharedOffset;				// relative to root of node description data (max 64 KB per node)
+		uint16	NodeInstanceOffset;				// relative to root of node instance data (max 64 KB per node)
 
 		// TODO: We could cache which parent/base decorator handles which interface
 		// This would avoid the need to iterate on every decorator to look up a 'Super'. Perhaps only common interfaces could be cached.
diff --git a/Engine/Plugins/Experimental/AnimNext/Source/AnimNext/Internal/DecoratorBase/DecoratorWriter.h b/Engine/Plugins/Experimental/AnimNext/Source/AnimNext/Internal/DecoratorBase/DecoratorWriter.h
--- a/Engine/Plugins/Experimental/AnimNext/Source/AnimNext/Internal/DecoratorBase/DecoratorWriter.h
+++ b/Engine/Plugins/Experimental/AnimNext/Source/AnimNext/Internal/DecoratorBase/DecoratorWriter.h
@@ -22,22 +22,34 @@ namespace UE::AnimNext
 	  *     - A list of FNodeTemplates that the nodes use
 	  *     - The graph shared data (FNodeDescription for every node)
 	  */
-	class ANIMNEXT_API FDecoratorWriter final : public FMemoryWriter
+	class ANIMNEXT_API FDecoratorWriter : public FMemoryWriter
 	{
 	public:
+		// The largest size allowed for the shared data of a single graph
+		// We use 16 bit offsets to index into the shared data
+		static constexpr uint32 MAXIMUM_GRAPH_SHARED_DATA_SIZE = 64 * 1024;
+
 		enum class EErrorState
 		{
 			None,					// All good, no error
 			TooManyNodes,			// Exceeded the maximum number of nodes in a graph, @see FNodeDescription::MAXIMUM_COUNT
+			GraphTooLarge,			// Exceeded the maximum graph size, @see FDecoratorWriter::MAXIMUM_GRAPH_SHARED_DATA_SIZE
 			NodeTemplateNotFound,	// Failed to find a necessary node template
 			NodeTemplateTooLarge,	// Exceeded the maximum node template size, @see FNodeTemplate::MAXIMUM_SIZE
+			NodeSharedDataTooLarge,	// Exceeded the maximum node shared data size, @see FNodeDescription::MAXIMUM_SIZE
 			NodeHandleNotFound,		// Failed to find the mapping for a node handle, it was likely not registered
 		};
 
 		FDecoratorWriter();
 
 		// Registers an instance of the provided node template and assigns a node handle and node UID to it
-		[[nodiscard]] FNodeHandle RegisterNode(const FNodeTemplate& NodeTemplate);
+		FNodeHandle RegisterNode(const FNodeTemplate& NodeTemplate);
+
+		// Returns the node UID associated with a node handle
+		uint16 GetNodeUID(const FNodeHandle NodeHandle) const;
+
+		// Returns the node handle associated with a node UID
+		FNodeHandle GetNodeHandle(uint16 NodeUID) const;
 
 		// Called before node writing can begin
 		void BeginNodeWriting();
@@ -47,39 +59,31 @@ namespace UE::AnimNext
 
 		// Writes out the provided node using the decorator properties
 		// Nodes must be written in the same order they were registered in
-		void WriteNode(const FNodeHandle NodeHandle, const TFunction<FString(uint32 DecoratorIndex, const FString& PropertyName)>& GetDecoratorProperty);
+		void WriteNode(const FNodeHandle NodeHandle, const TFunction<const TMap<FString, FString>& (uint32 DecoratorIndex)>& GetDecoratorProperties);
 
 		// Returns the error state
-		[[nodiscard]] EErrorState GetErrorState() const;
+		EErrorState GetErrorState() const;
 
 		// Returns the populated raw graph shared data buffer
-		[[nodiscard]] const TArray<uint8>& GetGraphSharedData() const;
-
-		// FArchive implementation
-		virtual FArchive& operator<<(UObject*& Obj) override;
-		virtual FArchive& operator<<(FObjectPtr& Obj) override;
+		const TArray<uint8>& GetGraphSharedData() const;
 
 	private:
 		struct FNodeMapping
 		{
-			// The node handle for this entry (encoded as a node ID)
 			FNodeHandle NodeHandle;
-
-			// The node template handle the node uses
+			uint32 NodeSize;
+			uint16 NodeUID;
 			FNodeTemplateRegistryHandle NodeTemplateHandle;
-
-			// The unique node template index that we'll serialize
-			uint32 NodeTemplateIndex;
 		};
 
 		TArray<uint8> GraphSharedDataArchiveBuffer;
 		TArray<FNodeMapping> NodeMappings;
 
 		// To track the node registration process
-		FNodeID NextNodeID;
+		uint32 GraphSharedDataSize;
+		uint16 NextNodeUID;
 
 		// To track node writing
-		TArray<UObject*> TrackedObjectsForGC;
 		uint32 NumNodesWritten;
 		bool bIsNodeWriting;
 
diff --git a/Engine/Plugins/Experimental/AnimNext/Source/AnimNext/Internal/DecoratorBase/NodeTemplate.h b/Engine/Plugins/Experimental/AnimNext/Source/AnimNext/Internal/DecoratorBase/NodeTemplate.h
--- a/Engine/Plugins/Experimental/AnimNext/Source/AnimNext/Internal/DecoratorBase/NodeTemplate.h
+++ b/Engine/Plugins/Experimental/AnimNext/Source/AnimNext/Internal/DecoratorBase/NodeTemplate.h
@@ -21,9 +21,7 @@ namespace UE::AnimNext
 	 * Templates can exist in multiple sub-graphs on disk. On load, a single unique copy is retained in
 	 * memory and shared between ALL graphs.
 	 * 
-	 * Node templates are created through the FNodeTemplateBuilder.
-	 * 
-	 * @see FNodeDescription, FNodeTemplateBuilder
+	 * @see FNodeDescription
 	 */
 	struct FNodeTemplate
 	{
@@ -32,24 +30,24 @@ namespace UE::AnimNext
 		// We artificially specify a conservative upper bound
 		static constexpr uint32 MAXIMUM_SIZE = 64 * 1024;
 
+		FNodeTemplate(uint32 UID_, uint32 InstanceSize_, uint32 NumDecorators_)
+			: UID(UID_)
+			, InstanceSize(InstanceSize_)
+			, NumDecorators(NumDecorators_)
+		{}
+
 		// Returns the globally unique template identifier (a hash of all the decorator UIDs present in this node)
 		uint32 GetUID() const { return UID; }
 
 		// Returns the size in bytes of this node template
-		uint32 GetNodeTemplateSize() const { return sizeof(FNodeTemplate) + (NumDecorators * sizeof(FDecoratorTemplate)); }
+		uint32 GetNodeTemplateSize() const { return sizeof(FNodeTemplate) + NumDecorators * sizeof(FDecoratorTemplate); }
 
-		// Returns the size in bytes of a node shared data
-		uint32 GetNodeSharedDataSize() const { return NodeSharedDataSize; }
-
-		// Returns the size in bytes of a node instance (how much space to allocate to hold a node instance of this template)
-		uint32 GetNodeInstanceDataSize() const { return NodeInstanceDataSize; }
+		// Returns the size in bytes of a node instance
+		uint32 GetInstanceSize() const { return InstanceSize; }
 
 		// Returns the number of decorators present in the node template
 		uint32 GetNumDecorators() const { return NumDecorators; }
 
-		// Returns whether the node template is valid or not
-		bool IsValid() const { return NodeSharedDataSize != 0 && NodeInstanceDataSize != 0; }
-
 		// Returns a pointer to the list of decorator template descriptions
 		FDecoratorTemplate* GetDecorators() { return reinterpret_cast<FDecoratorTemplate*>(reinterpret_cast<uint8*>(this) + sizeof(FNodeTemplate)); }
 
@@ -60,31 +58,17 @@ namespace UE::AnimNext
 		ANIMNEXT_API void Serialize(FArchive& Ar);
 
 	private:
-		friend struct FNodeTemplateBuilder;
-
-		FNodeTemplate(uint32 InUID, uint32 InNumDecorators)
-			: UID(InUID)
-			, NodeSharedDataSize(0)
-			, NodeInstanceDataSize(0)
-			, NumDecorators(static_cast<uint8>(InNumDecorators))
-		{}
-
-		// This function performs the necessary fix-ups once all decorator templates have been populated
-		// Call this function whenever a new node template has been constructed before using it
-		ANIMNEXT_API void Finalize();
-
-		uint32	UID;					// globally unique template identifier or hash
+		uint32	UID;				// globally unique template identifier or hash
 
-		uint16	NodeSharedDataSize;		// size in bytes of a node shared data (not serialized, @see FNodeTemplate::Finalize)
-		uint16	NodeInstanceDataSize;	// size in bytes of a node instance (not serialized, @see FNodeTemplate::Finalize)
+		uint16	InstanceSize;		// size in bytes of a node instance
 
 		uint8	NumDecorators;
-		uint8	Padding[3];
+		uint8	Padding[1];
 
 		// TODO: We could use the padding (and extra space) here to cache which decorator handles which interface
 		// This would avoid the need to iterate on every decorator to look it up. Perhaps only common interfaces could be cached.
 
-		// Followed in memory by a list of [FDecoratorTemplate] instances
+		// Followed by a list of [FDecoratorTemplate] instances
 	};
 
 	static_assert(std::is_trivially_copyable_v<FNodeTemplate>, "FNodeTemplate needs to be trivially copyable");
diff --git a/Engine/Plugins/Experimental/AnimNext/Source/AnimNext/Internal/DecoratorBase/NodeTemplateBuilder.h b/Engine/Plugins/Experimental/AnimNext/Source/AnimNext/Internal/DecoratorBase/NodeTemplateBuilder.h
deleted file mode 100644
diff --git a/Engine/Plugins/Experimental/AnimNext/Source/AnimNext/Internal/DecoratorBase/NodeTemplateRegistry.h b/Engine/Plugins/Experimental/AnimNext/Source/AnimNext/Internal/DecoratorBase/NodeTemplateRegistry.h
--- a/Engine/Plugins/Experimental/AnimNext/Source/AnimNext/Internal/DecoratorBase/NodeTemplateRegistry.h
+++ b/Engine/Plugins/Experimental/AnimNext/Source/AnimNext/Internal/DecoratorBase/NodeTemplateRegistry.h
@@ -27,6 +27,9 @@ namespace UE::AnimNext
 		// Finds or adds the specified node template and returns its offset
 		FNodeTemplateRegistryHandle FindOrAdd(const FNodeTemplate* NodeTemplate);
 
+		// Finds and returns a node template based on its handle or nullptr if the handle is invalid
+		FNodeTemplate* Find(FNodeTemplateRegistryHandle TemplateHandle);
+
 		// Finds and returns a node template based on its handle or nullptr if the handle is invalid
 		const FNodeTemplate* Find(FNodeTemplateRegistryHandle TemplateHandle) const;
 
@@ -36,15 +39,11 @@ namespace UE::AnimNext
 		// Removes the specified node template from the registry
 		void Unregister(const FNodeTemplate* NodeTemplate);
 
+		// Clears the registry by removing every node template
+		void Clear();
+
 	private:
 		FNodeTemplateRegistry() = default;
-		FNodeTemplateRegistry(const FNodeTemplateRegistry&) = delete;
-		FNodeTemplateRegistry(FNodeTemplateRegistry&&) = default;
-		FNodeTemplateRegistry& operator=(const FNodeTemplateRegistry&) = delete;
-		FNodeTemplateRegistry& operator=(FNodeTemplateRegistry&&) = default;
-
-		// Finds and returns a node template based on its handle or nullptr if the handle is invalid
-		FNodeTemplate* FindMutable(FNodeTemplateRegistryHandle TemplateHandle);
 
 		// Module lifetime functions
 		static void Init();
@@ -61,7 +60,5 @@ namespace UE::AnimNext
 		TMap<uint32, FNodeTemplateRegistryHandle>	TemplateUIDToHandleMap;
 
 		friend class FModule;
-		friend class FDecoratorWriter;
-		friend struct FScopedClearNodeTemplateRegistry;
 	};
 }
diff --git a/Engine/Plugins/Experimental/AnimNext/Source/AnimNext/Internal/Graph/AnimGraph/AnimNode_AnimNextGraph.h b/Engine/Plugins/Experimental/AnimNext/Source/AnimNext/Internal/Graph/AnimGraph/AnimNode_AnimNextGraph.h
--- a/Engine/Plugins/Experimental/AnimNext/Source/AnimNext/Internal/Graph/AnimGraph/AnimNode_AnimNextGraph.h
+++ b/Engine/Plugins/Experimental/AnimNext/Source/AnimNext/Internal/Graph/AnimGraph/AnimNode_AnimNextGraph.h
@@ -4,10 +4,24 @@
 
 #include "Animation/AnimNode_CustomProperty.h"
 #include "Graph/AnimNextGraph.h"
-#include "DecoratorBase/DecoratorPtr.h"
 #include "Context.h"
 #include "AnimNode_AnimNextGraph.generated.h"
 
+class UNodeMappingContainer;
+
+// TEST - until we can allocate per-node state again
+USTRUCT()
+struct FAnimSequencePlayerState
+{
+	GENERATED_BODY()
+
+	UPROPERTY()
+	float InternalTimeAccumulator = 0.0f;
+
+	UPROPERTY()
+	float PrevInternalTimeAccumulator = 0.0f;
+};
+
 /**
  * Animation node that allows a AnimNextGraph output to be used in an animation graph
  */
@@ -17,7 +31,6 @@ struct ANIMNEXT_API FAnimNode_AnimNextGraph : public FAnimNode_CustomProperty
 	GENERATED_BODY()
 
 	FAnimNode_AnimNextGraph();
-	~FAnimNode_AnimNextGraph();
 
 	// FAnimNode_Base interface
 	virtual void OnInitializeAnimInstance(const FAnimInstanceProxy* InProxy, const UAnimInstance* InAnimInstance) override;
@@ -44,8 +57,10 @@ private:
 	UPROPERTY(EditAnywhere, Category = Settings, meta = (PinHiddenByDefault))
 	TObjectPtr<UAnimNextGraph> AnimNextGraph;
 
-	// Shared pointer to our graph instance, we own it
-	UE::AnimNext::FDecoratorPtr GraphInstancePtr;
+	UPROPERTY(EditAnywhere, Category = Settings, meta = (PinHiddenByDefault, DisallowedClasses = "/Script/Engine.AnimMontage"))
+	TObjectPtr<UAnimSequence> TestSequence = nullptr;
+
+	FAnimSequencePlayerState SequencePlayerState;	// TEST - until we can allocate per-node state again
 
 	/*
 	 * Max LOD that this node is allowed to run
@@ -56,6 +71,9 @@ private:
 	UPROPERTY(EditAnywhere, Category = Performance, meta = (DisplayName = "LOD Threshold"))
 	int32 LODThreshold;
 
+	// Delta time received accumulated in update and used at Evaluate (so we can receive multiple calls to Evaluate)
+	float GraphDeltaTime = 0.f;
+
 protected:
 	virtual UClass* GetTargetClass() const override { return AnimNextGraph ? AnimNextGraph->StaticClass() : nullptr; }
 	
diff --git a/Engine/Plugins/Experimental/AnimNext/Source/AnimNext/Internal/Graph/AnimNextExecuteContext.h b/Engine/Plugins/Experimental/AnimNext/Source/AnimNext/Internal/Graph/AnimNextExecuteContext.h
--- a/Engine/Plugins/Experimental/AnimNext/Source/AnimNext/Internal/Graph/AnimNextExecuteContext.h
+++ b/Engine/Plugins/Experimental/AnimNext/Source/AnimNext/Internal/Graph/AnimNextExecuteContext.h
@@ -5,8 +5,6 @@
 #include "CoreMinimal.h"
 #include "RigVMCore/RigVMExecuteContext.h"
 #include "Units/RigUnit.h"
-#include "DecoratorBase/DecoratorPtr.h"
-
 #include "AnimNextExecuteContext.generated.h"
 
 namespace UE::AnimNext
@@ -14,26 +12,6 @@ namespace UE::AnimNext
 	struct FContext;
 }
 
-/**
-  * An enum to control which simulation steps should be performed when the animation graph is processed.
-  */
-enum class EAnimNextGraphSimulationSteps
-{
-	// No step is performed
-	None = 0x00,
-
-	// Performs the update step
-	Update = 0x01,
-
-	// Performs the evaluate step
-	Evaluate = 0x02,
-
-	// Perform all simulation steps
-	All = 0xff,
-};
-
-ENUM_CLASS_FLAGS(EAnimNextGraphSimulationSteps)
-
 USTRUCT(BlueprintType)
 struct FAnimNextExecuteContext : public FRigVMExecuteContext
 {
@@ -51,11 +29,10 @@ struct FAnimNextExecuteContext : public FRigVMExecuteContext
 		return *Context;
 	}
 
-	const TArrayView<const uint8>& GetSharedDataBuffer() const { return SharedDataBuffer; }
-
-	UE::AnimNext::FWeakDecoratorPtr GetGraphInstancePtr() const { return GraphInstancePtr; }
-
-	EAnimNextGraphSimulationSteps GetSimulationSteps() const { return SimulationSteps; }
+	void SetContextData(const UE::AnimNext::FContext& InContext)
+	{
+		Context = &InContext;
+	}
 
 	virtual void Copy(const FRigVMExecuteContext* InOtherContext) override
 	{
@@ -63,35 +40,10 @@ struct FAnimNextExecuteContext : public FRigVMExecuteContext
 
 		const FAnimNextExecuteContext* OtherContext = (const FAnimNextExecuteContext*)InOtherContext;
 		Context = OtherContext->Context;
-		SharedDataBuffer = OtherContext->SharedDataBuffer;
-		GraphInstancePtr = OtherContext->GraphInstancePtr;
 	}
 
 private:
-	void SetContextData(const UE::AnimNext::FContext& InContext)
-	{
-		Context = &InContext;
-	}
-
-	void InitializeWithGraph(TArrayView<const uint8> InSharedDataBuffer, UE::AnimNext::FWeakDecoratorPtr InGraphInstancePtr)
-	{
-		SharedDataBuffer = InSharedDataBuffer;
-		GraphInstancePtr = InGraphInstancePtr;
-	}
-
-	void SetSimulationSteps(EAnimNextGraphSimulationSteps InSimulationSteps)
-	{
-		SimulationSteps = InSimulationSteps;
-	}
-
 	const UE::AnimNext::FContext* Context;
-
-	TArrayView<const uint8> SharedDataBuffer;
-	UE::AnimNext::FWeakDecoratorPtr GraphInstancePtr;
-	EAnimNextGraphSimulationSteps SimulationSteps;
-
-	friend class UAnimNextGraph;
-	friend class UAnimNextParameterBlock;
 };
 
 USTRUCT(meta=(ExecuteContext="FAnimNextExecuteContext"))
diff --git a/Engine/Plugins/Experimental/AnimNext/Source/AnimNext/Internal/Graph/AnimNextGraph.h b/Engine/Plugins/Experimental/AnimNext/Source/AnimNext/Internal/Graph/AnimNextGraph.h
--- a/Engine/Plugins/Experimental/AnimNext/Source/AnimNext/Internal/Graph/AnimNextGraph.h
+++ b/Engine/Plugins/Experimental/AnimNext/Source/AnimNext/Internal/Graph/AnimNextGraph.h
@@ -4,9 +4,6 @@
 
 #include "CoreMinimal.h"
 #include "RigVMCore/RigVM.h"
-#include "DecoratorBase/DecoratorPtr.h"
-#include "DecoratorBase/DecoratorHandle.h"
-#include "DecoratorBase/EntryPointHandle.h"
 
 #include "AnimNextGraph.generated.h"
 
@@ -14,7 +11,6 @@ class UEdGraph;
 class UAnimNextGraph;
 class UAnimGraphNode_AnimNextGraph;
 struct FAnimNode_AnimNextGraph;
-enum class EAnimNextGraphSimulationSteps;
 
 namespace UE::AnimNext
 {
@@ -43,28 +39,15 @@ class ANIMNEXT_API UAnimNextGraph : public UObject
 {
 	GENERATED_BODY()
 
-public:
 	// UObject interface
 	virtual void PostRename(UObject* OldOuter, const FName OldName) override;
 	virtual void GetPreloadDependencies(TArray<UObject*>& OutDeps) override;
-	virtual void Serialize(FArchive& Ar) override;
-	static void AddReferencedObjects(UObject* InThis, FReferenceCollector& Collector);
 
-	// Allocates an instance of the graph, retain the handle and use it with the Run() function to evaluate it
-	UE::AnimNext::FDecoratorPtr AllocateInstance() const;
+	// Run this graph with the given context
+	void Run(const UE::AnimNext::FContext& Context) const;
 
-	// Releases an instance of the graph and clears the provided handle
-	void ReleaseInstance(UE::AnimNext::FDecoratorPtr& GraphInstancePtr) const;
-
-	// Run the specified simulation steps on the provided graph with the given context
-	void Run(const UE::AnimNext::FContext& Context, UE::AnimNext::FWeakDecoratorPtr GraphInstancePtr, EAnimNextGraphSimulationSteps SimulationSteps) const;
-
-protected:
 	// Support rig VM execution
 	TArray<FRigVMExternalVariable> GetRigVMExternalVariables();
-
-	// Loads the graph data from the provided archive buffer and returns true on success, false otherwise
-	bool LoadFromArchiveBuffer(const TArray<uint8>& SharedDataArchiveBuffer);
 	
 	friend class UAnimNextGraphFactory;
 	friend class UAnimNextGraph_EditorData;
@@ -72,26 +55,11 @@ protected:
 	friend class UE::AnimNext::Editor::FGraphEditor;
 	friend class UAnimNextGraph;
 	friend class UAnimGraphNode_AnimNextGraph;
+	friend struct FAnimNode_AnimNextGraph;
 
 	UPROPERTY()
 	TObjectPtr<URigVM> RigVM;
 
-	// This is a handle to the root decorator in our graph
-	UPROPERTY()
-	FAnimNextEntryPointHandle RootDecoratorHandle;
-
-	// This is a resolved handle to the root decorator in our graph
-	FAnimNextDecoratorHandle ResolvedRootDecoratorHandle;
-
-	// This is the graph shared data used by the decorator system, the output of FDecoratorReader
-	// We de-serialize manually into this buffer from the archive buffer, this is never saved on disk
-	TArray<uint8> SharedDataBuffer;
-
-	// This is the list of all referenced UObjects in the graph shared data
-	// We collect all the references here to make it quick and easy for the GC to query them
-	// It means that object references in the graph shared data are not visited at runtime by the GC (they are immutable)
-	TArray<TObjectPtr<UObject>> TrackedObjectsForGC;
-
 	UPROPERTY(transient)
 	mutable FRigVMExtendedExecuteContext ExtendedExecuteContext;
 
@@ -101,9 +69,5 @@ protected:
 #if WITH_EDITORONLY_DATA
 	UPROPERTY(VisibleAnywhere, Instanced, Category = "Graph", meta = (ShowInnerProperties))
 	TObjectPtr<UObject> EditorData;
-
-	// This buffer holds the output of the FDecoratorWriter post compilation
-	// We serialize it manually and it is discarded at runtime
-	TArray<uint8> SharedDataArchiveBuffer;
 #endif
 };
diff --git a/Engine/Plugins/Experimental/AnimNext/Source/AnimNext/Internal/Graph/RigDecorator_AnimNextCppDecorator.h b/Engine/Plugins/Experimental/AnimNext/Source/AnimNext/Internal/Graph/RigDecorator_AnimNextCppDecorator.h
deleted file mode 100644
diff --git a/Engine/Plugins/Experimental/AnimNext/Source/AnimNext/Internal/Graph/RigUnit_AnimNextDecoratorStack.h b/Engine/Plugins/Experimental/AnimNext/Source/AnimNext/Internal/Graph/RigUnit_AnimNextDecoratorStack.h
deleted file mode 100644
diff --git a/Engine/Plugins/Experimental/AnimNext/Source/AnimNext/Internal/Graph/RigUnit_AnimNextGraphRoot.h b/Engine/Plugins/Experimental/AnimNext/Source/AnimNext/Internal/Graph/RigUnit_AnimNextGraphRoot.h
deleted file mode 100644
diff --git a/Engine/Plugins/Experimental/AnimNext/Source/AnimNext/Private/DecoratorBase/Decorator.cpp b/Engine/Plugins/Experimental/AnimNext/Source/AnimNext/Private/DecoratorBase/Decorator.cpp
--- a/Engine/Plugins/Experimental/AnimNext/Source/AnimNext/Private/DecoratorBase/Decorator.cpp
+++ b/Engine/Plugins/Experimental/AnimNext/Source/AnimNext/Private/DecoratorBase/Decorator.cpp
@@ -18,7 +18,7 @@ namespace UE::AnimNext
 	}
 
 #if WITH_EDITOR
-	void FDecorator::SaveDecoratorSharedData(FDecoratorWriter& Writer, const TFunction<FString(const FString& PropertyName)>& GetDecoratorProperty, FAnimNextDecoratorSharedData& OutSharedData) const
+	void FDecorator::SaveDecoratorSharedData(FDecoratorWriter& Writer, const TMap<FString, FString>& Properties, FAnimNextDecoratorSharedData& OutSharedData) const
 	{
 		const UScriptStruct* SharedDataStruct = GetDecoratorSharedDataStruct();
 
@@ -31,10 +31,9 @@ namespace UE::AnimNext
 		// We convert every property from its string representation into its binary form
 		for (const FProperty* Property = SharedDataStruct->PropertyLink; Property != nullptr; Property = Property->PropertyLinkNext)
 		{
-			const FString PropertyValue = GetDecoratorProperty(Property->GetName());
-			if (PropertyValue.Len() != 0)
+			if (const FString* PropertyValue = Properties.Find(Property->GetName()))
 			{
-				const TCHAR* PropertyValuePtr = *PropertyValue;
+				const TCHAR* PropertyValuePtr = **PropertyValue;
 
 				// C-style array properties aren't handled by ExportText, we need to handle it manually
 				const bool bIsCArray = Property->ArrayDim > 1;
diff --git a/Engine/Plugins/Experimental/AnimNext/Source/AnimNext/Private/DecoratorBase/DecoratorHandle.cpp b/Engine/Plugins/Experimental/AnimNext/Source/AnimNext/Private/DecoratorBase/DecoratorHandle.cpp
deleted file mode 100644
diff --git a/Engine/Plugins/Experimental/AnimNext/Source/AnimNext/Private/DecoratorBase/DecoratorPtr.cpp b/Engine/Plugins/Experimental/AnimNext/Source/AnimNext/Private/DecoratorBase/DecoratorPtr.cpp
--- a/Engine/Plugins/Experimental/AnimNext/Source/AnimNext/Private/DecoratorBase/DecoratorPtr.cpp
+++ b/Engine/Plugins/Experimental/AnimNext/Source/AnimNext/Private/DecoratorBase/DecoratorPtr.cpp
@@ -7,10 +7,10 @@ namespace UE::AnimNext
 {
 	FDecoratorPtr::FDecoratorPtr(FNodeInstance* NodeInstance, uint32 DecoratorIndex_)
 		: PackedPointerAndFlags(NodeInstance != nullptr ? reinterpret_cast<uintptr_t>(NodeInstance) : 0)
-		, DecoratorIndex(DecoratorIndex_)
+		, DecoratorIndex(static_cast<uint8>(DecoratorIndex_))
 	{
 		check((reinterpret_cast<uintptr_t>(NodeInstance) & FLAGS_MASK) == 0);	// Make sure we have enough alignment
-		check(DecoratorIndex <= MAX_uint8);	// Make sure we don't truncate
+		check(DecoratorIndex == DecoratorIndex_);	// Make sure we didn't truncate
 
 		if (NodeInstance != nullptr)
 		{
@@ -43,10 +43,10 @@ namespace UE::AnimNext
 
 	FDecoratorPtr::FDecoratorPtr(FNodeInstance* NodeInstance, EFlags Flags, uint32 DecoratorIndex_)
 		: PackedPointerAndFlags(NodeInstance != nullptr ? (reinterpret_cast<uintptr_t>(NodeInstance) | Flags) : 0)
-		, DecoratorIndex(DecoratorIndex_)
+		, DecoratorIndex(static_cast<uint8>(DecoratorIndex_))
 	{
 		check((reinterpret_cast<uintptr_t>(NodeInstance) & FLAGS_MASK) == 0);	// Make sure we have enough alignment
-		check(DecoratorIndex <= MAX_uint8);	// Make sure we don't truncate
+		check(DecoratorIndex == DecoratorIndex_);	// Make sure we didn't truncate
 
 		// Only increment the reference count if we aren't a weak handle
 		if (NodeInstance != nullptr && (Flags & IS_WEAK_BIT) == 0)
@@ -112,13 +112,6 @@ namespace UE::AnimNext
 		DecoratorIndex = 0;
 	}
 
-	FWeakDecoratorPtr::FWeakDecoratorPtr(FNodeInstance* NodeInstance_, uint32 DecoratorIndex_)
-		: NodeInstance(NodeInstance_)
-		, DecoratorIndex(DecoratorIndex_)
-	{
-		check(DecoratorIndex <= MAX_uint8);	// Make sure we don't truncate
-	}
-
 	void FWeakDecoratorPtr::Reset()
 	{
 		NodeInstance = nullptr;
diff --git a/Engine/Plugins/Experimental/AnimNext/Source/AnimNext/Private/DecoratorBase/DecoratorReader.cpp b/Engine/Plugins/Experimental/AnimNext/Source/AnimNext/Private/DecoratorBase/DecoratorReader.cpp
--- a/Engine/Plugins/Experimental/AnimNext/Source/AnimNext/Private/DecoratorBase/DecoratorReader.cpp
+++ b/Engine/Plugins/Experimental/AnimNext/Source/AnimNext/Private/DecoratorBase/DecoratorReader.cpp
@@ -3,11 +3,8 @@
 #include "DecoratorBase/DecoratorReader.h"
 
 #include "DecoratorBase/NodeDescription.h"
-#include "DecoratorBase/NodeInstance.h"
 #include "DecoratorBase/NodeTemplate.h"
 #include "DecoratorBase/NodeTemplateRegistry.h"
-#include "Serialization/ArchiveUObject.h"
-#include "UObject/UObjectGlobals.h"
 
 namespace UE::AnimNext
 {
@@ -16,191 +13,48 @@ namespace UE::AnimNext
 	{
 	}
 
-	FDecoratorReader::EErrorState FDecoratorReader::ReadGraph(TArray<uint8>& GraphSharedData, TArray<TObjectPtr<UObject>>& TrackedObjectsForGC)
+	void FDecoratorReader::ReadGraphSharedData(TArray<uint8>& GraphSharedData)
 	{
 		GraphSharedData.Empty(0);
-		TrackedObjectsForGC.Empty(0);
 
-		EErrorState ErrorState = ReadGraphSharedData(GraphSharedData);
-		if (ErrorState != EErrorState::None)
-		{
-			return ErrorState;
-		}
-
-		ErrorState = ReadTrackedObjectsForGC(TrackedObjectsForGC);
-
-		return ErrorState;
-	}
-
-	FDecoratorReader::EErrorState FDecoratorReader::ReadGraphSharedData(TArray<uint8>& GraphSharedData)
-	{
 		// Read the node templates and register them as needed
 		uint32 NumNodeTemplates = 0;
 		*this << NumNodeTemplates;
 
-		TArray<FNodeTemplateRegistryHandle> NodeTemplateHandles;
-		NodeTemplateHandles.Reserve(NumNodeTemplates);
+		// We serialize each node template in the same buffer, they get copied into the registry and we don't need to retain them
+		alignas(16) uint8 NodeTemplateBuffer[64 * 1024];
+		FNodeTemplate* NodeTemplate = reinterpret_cast<FNodeTemplate*>(&NodeTemplateBuffer[0]);
 
 		FNodeTemplateRegistry& NodeTemplateRegistry = FNodeTemplateRegistry::Get();
-
+		for (uint32 NodeTemplateIndex = 0; NodeTemplateIndex < NumNodeTemplates; ++NodeTemplateIndex)
 		{
-			// We serialize each node template in the same buffer, they get copied into the registry and we don't need to retain them
-			alignas(16) uint8 NodeTemplateBuffer[64 * 1024];
-			FNodeTemplate* NodeTemplate = reinterpret_cast<FNodeTemplate*>(&NodeTemplateBuffer[0]);
+			NodeTemplate->Serialize(*this);
 
-			for (uint32 NodeTemplateIndex = 0; NodeTemplateIndex < NumNodeTemplates; ++NodeTemplateIndex)
-			{
-				NodeTemplate->Serialize(*this);
-
-				if (!NodeTemplate->IsValid())
-				{
-					if (NodeTemplate->GetNodeSharedDataSize() == 0)
-					{
-						// This node shared data is too large
-						return EErrorState::NodeSharedDataTooLarge;
-					}
-
-					if (NodeTemplate->GetNodeInstanceDataSize() == 0)
-					{
-						// This node instance data is too large
-						return EErrorState::NodeInstanceDataTooLarge;
-					}
-				}
-
-				// Register our node template
-				NodeTemplateHandles.Add(NodeTemplateRegistry.FindOrAdd(NodeTemplate));
-			}
+			// Register our node template
+			NodeTemplateRegistry.FindOrAdd(NodeTemplate);
 		}
 
 		// Read our graph shared data
 		uint32 NumNodes = 0;
 		*this << NumNodes;
 
-		NodeHandles.Empty(0);
-		NodeHandles.Reserve(NumNodes);
-
-		// Calculate our shared data size and all node offsets
 		uint32 SharedDataSize = 0;
-		for (uint32 NodeIndex = 0; NodeIndex < NumNodes; ++NodeIndex)
-		{
-			SharedDataSize = Align(SharedDataSize, alignof(FNodeDescription));	// Make sure we respect our alignment constraints
-
-			if (SharedDataSize > MAXIMUM_GRAPH_SHARED_DATA_SIZE)
-			{
-				// This graph shared data is too large, we won't be able to create handles to this node
-				return EErrorState::GraphTooLarge;
-			}
-
-			NodeHandles.Add(FNodeHandle::FromSharedOffset(SharedDataSize));	// This node starts here
-			check(NodeHandles.Last().IsSharedOffset());
+		*this << SharedDataSize;
 
-			uint32 NodeTemplateIndex = 0;
-			*this << NodeTemplateIndex;
-
-			const FNodeTemplate* NodeTemplate = NodeTemplateRegistry.Find(NodeTemplateHandles[NodeTemplateIndex]);
-			check(NodeTemplate != nullptr);
-
-			SharedDataSize += NodeTemplate->GetNodeSharedDataSize();
-		}
-
-		// The shared data size might exceed MAXIMUM_GRAPH_SHARED_DATA_SIZE a little bit
-		// The only requirement is that the node begins before that threshold so that we can
-		// create handles to it
-
-		GraphSharedData.Empty(0);
 		GraphSharedData.AddZeroed(SharedDataSize);
 
-		// Serialize our graph
+		uint32 SharedDataOffset = 0;
 		for (uint32 NodeIndex = 0; NodeIndex < NumNodes; ++NodeIndex)
 		{
-			const uint32 SharedDataOffset = NodeHandles[NodeIndex].GetSharedOffset();
+			uint32 NodeSize = 0;
+			*this << NodeSize;
+
 			FNodeDescription* NodeDesc = reinterpret_cast<FNodeDescription*>(&GraphSharedData[SharedDataOffset]);
 			NodeDesc->Serialize(*this);
-		}
-
-		return EErrorState::None;
-	}
-
-	FDecoratorReader::EErrorState FDecoratorReader::ReadTrackedObjectsForGC(TArray<TObjectPtr<UObject>>& TrackedObjectsForGC)
-	{
-		int32 NumTrackedObjectsForGC = 0;
-		*this << NumTrackedObjectsForGC;
-
-		TrackedObjectsForGC.SetNumZeroed(NumTrackedObjectsForGC);
-
-		for (int32 TrackedIndex = 0; TrackedIndex < NumTrackedObjectsForGC; ++TrackedIndex)
-		{
-			// Load the path name to the object
-			FString LoadedString;
-			*this << LoadedString;
 
-			// Look up the object by fully qualified pathname
-			UObject* Obj = FindObject<UObject>(nullptr, *LoadedString, false);
-			TrackedObjectsForGC[TrackedIndex] = Obj;
-		}
-
-		return EErrorState::None;
-	}
-
-	FArchive& FDecoratorReader::operator<<(UObject*& Obj)
-	{
-		// Load the path name to the object
-		FString LoadedString;
-		*this << LoadedString;
-
-		// Look up the object by fully qualified pathname
-		Obj = FindObject<UObject>(nullptr, *LoadedString, false);
-
-		// If we couldn't find it, load it
-		if (Obj == nullptr)
-		{
-			Obj = LoadObject<UObject>(nullptr, *LoadedString);
-		}
-
-		return *this;
-	}
-
-	FArchive& FDecoratorReader::operator<<(FObjectPtr& Obj)
-	{
-		return FArchiveUObject::SerializeObjectPtr(*this, Obj);
-	}
-
-	FNodeHandle FDecoratorReader::ResolveNodeHandle(FNodeHandle NodeHandle) const
-	{
-		if (!NodeHandle.IsValid())
-		{
-			// The node handle is invalid, return it unchanged
-			return NodeHandle;
-		}
-
-		check(NodeHandle.IsNodeID());
-		const FNodeID NodeID = NodeHandle.GetNodeID();
-		check(NodeID.IsValid());
-
-		return NodeHandles[NodeID.GetNodeIndex()];
-	}
-
-	FAnimNextDecoratorHandle FDecoratorReader::ResolveDecoratorHandle(FAnimNextDecoratorHandle DecoratorHandle) const
-	{
-		if (!DecoratorHandle.IsValid())
-		{
-			// The decorator handle is invalid, return it unchanged
-			return DecoratorHandle;
-		}
-
-		const FNodeHandle NodeHandle = ResolveNodeHandle(DecoratorHandle.GetNodeHandle());
-		return FAnimNextDecoratorHandle(NodeHandle, DecoratorHandle.GetDecoratorIndex());
-	}
-
-	FAnimNextDecoratorHandle FDecoratorReader::ResolveEntryPointHandle(FAnimNextEntryPointHandle EntryPointHandle) const
-	{
-		if (!EntryPointHandle.IsValid())
-		{
-			// The decorator handle is invalid, return an invalid handle
-			return FAnimNextDecoratorHandle();
+			SharedDataOffset += NodeSize;
 		}
 
-		const FNodeHandle NodeHandle = ResolveNodeHandle(EntryPointHandle.GetNodeHandle());
-		return FAnimNextDecoratorHandle(NodeHandle, EntryPointHandle.GetDecoratorIndex());
+		ensure(SharedDataOffset == SharedDataSize);
 	}
 }
diff --git a/Engine/Plugins/Experimental/AnimNext/Source/AnimNext/Private/DecoratorBase/DecoratorRegistry.cpp b/Engine/Plugins/Experimental/AnimNext/Source/AnimNext/Private/DecoratorBase/DecoratorRegistry.cpp
--- a/Engine/Plugins/Experimental/AnimNext/Source/AnimNext/Private/DecoratorBase/DecoratorRegistry.cpp
+++ b/Engine/Plugins/Experimental/AnimNext/Source/AnimNext/Private/DecoratorBase/DecoratorRegistry.cpp
@@ -206,14 +206,9 @@ namespace UE::AnimNext
 		}
 	}
 
-	FDecoratorRegistryHandle FDecoratorRegistry::FindHandle(FDecoratorUID DecoratorUID) const
+	FDecoratorRegistryHandle FDecoratorRegistry::FindHandle(uint32 DecoratorUID) const
 	{
-		if (!DecoratorUID.IsValid())
-		{
-			return FDecoratorRegistryHandle();
-		}
-
-		if (const FRegistryEntry* Entry = DecoratorUIDToEntryMap.Find(DecoratorUID.GetUID()))
+		if (const FRegistryEntry* Entry = DecoratorUIDToEntryMap.Find(DecoratorUID))
 		{
 			return Entry->DecoratorHandle;
 		}
@@ -222,6 +217,11 @@ namespace UE::AnimNext
 		return FDecoratorRegistryHandle();
 	}
 
+	FDecoratorRegistryHandle FDecoratorRegistry::FindHandle(FDecoratorUID DecoratorUID) const
+	{
+		return FindHandle(DecoratorUID.GetUID());
+	}
+
 	const FDecorator* FDecoratorRegistry::Find(FDecoratorRegistryHandle DecoratorHandle) const
 	{
 		if (!DecoratorHandle.IsValid())
@@ -247,25 +247,6 @@ namespace UE::AnimNext
 		return Find(DecoratorHandle);
 	}
 
-	const FDecorator* FDecoratorRegistry::Find(const UScriptStruct* DecoratorSharedDataStruct) const
-	{
-		if (DecoratorSharedDataStruct == nullptr)
-		{
-			return nullptr;
-		}
-
-		for (const auto& KeyValuePair : DecoratorUIDToEntryMap)
-		{
-			const FDecorator* Decorator = KeyValuePair.Value.Decorator;
-			if (Decorator->GetDecoratorSharedDataStruct() == DecoratorSharedDataStruct)
-			{
-				return Decorator;
-			}
-		}
-
-		return nullptr;
-	}
-
 	void FDecoratorRegistry::Register(FDecorator* Decorator)
 	{
 		if (Decorator == nullptr)
@@ -341,19 +322,6 @@ namespace UE::AnimNext
 		}
 	}
 
-	TArray<const FDecorator*> FDecoratorRegistry::GetDecorators() const
-	{
-		TArray<const FDecorator*> Decorators;
-		Decorators.Reserve(DecoratorUIDToEntryMap.Num());
-		
-		for (const auto& It : DecoratorUIDToEntryMap)
-		{
-			Decorators.Add(It.Value.Decorator);
-		}
-
-		return Decorators;
-	}
-
 	uint32 FDecoratorRegistry::GetNum() const
 	{
 		return DecoratorUIDToEntryMap.Num();
diff --git a/Engine/Plugins/Experimental/AnimNext/Source/AnimNext/Private/DecoratorBase/DecoratorTemplate.cpp b/Engine/Plugins/Experimental/AnimNext/Source/AnimNext/Private/DecoratorBase/DecoratorTemplate.cpp
--- a/Engine/Plugins/Experimental/AnimNext/Source/AnimNext/Private/DecoratorBase/DecoratorTemplate.cpp
+++ b/Engine/Plugins/Experimental/AnimNext/Source/AnimNext/Private/DecoratorBase/DecoratorTemplate.cpp
@@ -10,13 +10,11 @@ namespace UE::AnimNext
 {
 	void FDecoratorTemplate::Serialize(FArchive& Ar)
 	{
-		const FDecoratorRegistry& DecoratorRegistry = FDecoratorRegistry::Get();
-
 		Ar << UID;
 
 		if (Ar.IsSaving())
 		{
-			const FDecorator* Decorator = DecoratorRegistry.Find(RegistryHandle);
+			const FDecorator* Decorator = FDecoratorRegistry::Get().Find(RegistryHandle);
 
 			uint32 DecoratorUID = Decorator->GetDecoratorUID().GetUID();
 			Ar << DecoratorUID;
@@ -26,11 +24,7 @@ namespace UE::AnimNext
 			uint32 DecoratorUID = 0;
 			Ar << DecoratorUID;
 
-			// It is possible that we fail to find the decorator that we need
-			// This can happen if the decorator hasn't been loaded or registered
-			// When this happens, the decorator is a no-op and the runtime behavior
-			// may not be what is expected
-			RegistryHandle = DecoratorRegistry.FindHandle(FDecoratorUID(DecoratorUID));
+			RegistryHandle = FDecoratorRegistry::Get().FindHandle(DecoratorUID);
 		}
 		else
 		{
@@ -49,5 +43,7 @@ namespace UE::AnimNext
 
 		Ar << Mode;
 		Ar << AdditiveIndexOrNumAdditive;
+		Ar << NodeSharedOffset;
+		Ar << NodeInstanceOffset;
 	}
 }
diff --git a/Engine/Plugins/Experimental/AnimNext/Source/AnimNext/Private/DecoratorBase/DecoratorWriter.cpp b/Engine/Plugins/Experimental/AnimNext/Source/AnimNext/Private/DecoratorBase/DecoratorWriter.cpp
--- a/Engine/Plugins/Experimental/AnimNext/Source/AnimNext/Private/DecoratorBase/DecoratorWriter.cpp
+++ b/Engine/Plugins/Experimental/AnimNext/Source/AnimNext/Private/DecoratorBase/DecoratorWriter.cpp
@@ -8,15 +8,37 @@
 #include "DecoratorBase/NodeDescription.h"
 #include "DecoratorBase/NodeTemplate.h"
 #include "DecoratorBase/NodeTemplateRegistry.h"
-#include "Serialization/ArchiveUObject.h"
-
-#include "DecoratorBase/DecoratorReader.h"
 
 namespace UE::AnimNext
 {
+	namespace Private
+	{
+		// Note: this does not account for inline/pooled/pinned properties
+		static uint32 GetNodeSharedSize(const FNodeTemplate& NodeTemplate)
+		{
+			const FDecoratorRegistry& Registry = FDecoratorRegistry::Get();
+
+			const uint32 NumDecorators = NodeTemplate.GetNumDecorators();
+			const FDecoratorTemplate* DecoratorTemplates = NodeTemplate.GetDecorators();
+
+			uint32 NodeSharedSize = sizeof(FNodeDescription);
+			for (uint32_t DecoratorIndex = 0; DecoratorIndex < NumDecorators; ++DecoratorIndex)
+			{
+				const FDecorator* Decorator = Registry.Find(DecoratorTemplates[DecoratorIndex].GetRegistryHandle());
+
+				const FDecoratorMemoryLayout MemoryLayout = Decorator->GetDecoratorMemoryDescription();
+				NodeSharedSize = Align(NodeSharedSize, MemoryLayout.SharedDataAlignment);
+				NodeSharedSize += MemoryLayout.SharedDataSize;
+			}
+
+			return NodeSharedSize;
+		}
+	}
+
 	FDecoratorWriter::FDecoratorWriter()
 		: FMemoryWriter(GraphSharedDataArchiveBuffer)
-		, NextNodeID(FNodeID::GetFirstID())
+		, GraphSharedDataSize(0)
+		, NextNodeUID(1)	// Node UID 0 is reserved for the invalid value (null)
 		, NumNodesWritten(0)
 		, bIsNodeWriting(false)
 		, ErrorState(EErrorState::None)
@@ -40,28 +62,54 @@ namespace UE::AnimNext
 			return FNodeHandle();
 		}
 
-		if (!NextNodeID.IsValid())
+		if (NextNodeUID == FNodeDescription::MAXIMUM_COUNT)
 		{
 			// We have too many nodes in the graph, we need to be able to represent them with 16 bits
-			// The node ID must have wrapped around
 			ErrorState = EErrorState::TooManyNodes;
 			return FNodeHandle();
 		}
 
-		const FNodeHandle NodeHandle = FNodeHandle::FromNodeID(NextNodeID);
-		check(NodeHandle.IsValid() && NodeHandle.IsNodeID());
-		check(NodeMappings.Num() == NodeHandle.GetNodeID().GetNodeIndex());
+		const uint32 NodeSharedDataSize = Private::GetNodeSharedSize(NodeTemplate);
+		if (NodeSharedDataSize > FNodeDescription::MAXIMUM_SIZE)
+		{
+			// This node shared data is too large
+			ErrorState = EErrorState::NodeSharedDataTooLarge;
+			return FNodeHandle();
+		}
+
+		const uint32 NodeSharedDataOffset = GraphSharedDataSize;
+		if (GraphSharedDataSize + NodeSharedDataSize > MAXIMUM_GRAPH_SHARED_DATA_SIZE)
+		{
+			// The graph shared data size is too large
+			ErrorState = EErrorState::GraphTooLarge;
+			return FNodeHandle();
+		}
 
-		NextNodeID = NextNodeID.GetNextID();
+		GraphSharedDataSize += NodeSharedDataSize;
+
+		const FNodeHandle NodeHandle(NodeSharedDataOffset);
+		const uint16 NodeUID = NextNodeUID++;
 
 		FNodeTemplateRegistry& NodeTemplateRegistry = FNodeTemplateRegistry::Get();
 		const FNodeTemplateRegistryHandle NodeTemplateHandle = NodeTemplateRegistry.FindOrAdd(&NodeTemplate);
 
-		NodeMappings.Add({ NodeHandle, NodeTemplateHandle, 0 });
+		NodeMappings.Add({ NodeHandle, NodeSharedDataSize, NodeUID, NodeTemplateHandle });
 
 		return NodeHandle;
 	}
 
+	uint16 FDecoratorWriter::GetNodeUID(const FNodeHandle NodeHandle) const
+	{
+		const FNodeMapping* NodeMapping = NodeMappings.FindByPredicate([NodeHandle](const FNodeMapping& It) { return It.NodeHandle == NodeHandle; });
+		return NodeMapping != nullptr ? NodeMapping->NodeUID : 0;
+	}
+
+	FNodeHandle FDecoratorWriter::GetNodeHandle(uint16 NodeUID) const
+	{
+		const FNodeMapping* NodeMapping = NodeMappings.FindByPredicate([NodeUID](const FNodeMapping& It) { return It.NodeUID == NodeUID; });
+		return NodeMapping != nullptr ? NodeMapping->NodeHandle : FNodeHandle();
+	}
+
 	void FDecoratorWriter::BeginNodeWriting()
 	{
 		ensure(!bIsNodeWriting);
@@ -74,15 +122,14 @@ namespace UE::AnimNext
 		}
 
 		bIsNodeWriting = true;
-		TrackedObjectsForGC.Reset();
 
 		// Serialize the node templates
 		TArray<FNodeTemplateRegistryHandle> NodeTemplateHandles;
 		NodeTemplateHandles.Reserve(NodeMappings.Num());
 
-		for (FNodeMapping& NodeMapping : NodeMappings)
+		for (const FNodeMapping& NodeMapping : NodeMappings)
 		{
-			NodeMapping.NodeTemplateIndex = NodeTemplateHandles.AddUnique(NodeMapping.NodeTemplateHandle);
+			NodeTemplateHandles.AddUnique(NodeMapping.NodeTemplateHandle);
 		}
 
 		uint32 NumNodeTemplates = NodeTemplateHandles.Num();
@@ -91,7 +138,7 @@ namespace UE::AnimNext
 		FNodeTemplateRegistry& NodeTemplateRegistry = FNodeTemplateRegistry::Get();
 		for (FNodeTemplateRegistryHandle NodeTemplateHandle : NodeTemplateHandles)
 		{
-			FNodeTemplate* NodeTemplate = NodeTemplateRegistry.FindMutable(NodeTemplateHandle);
+			FNodeTemplate* NodeTemplate = NodeTemplateRegistry.Find(NodeTemplateHandle);
 			NodeTemplate->Serialize(*this);
 		}
 
@@ -99,12 +146,8 @@ namespace UE::AnimNext
 		uint32 NumNodes = NodeMappings.Num();
 		*this << NumNodes;
 
-		// Serialize the node template indices that we'll use for each node
-		for (const FNodeMapping& NodeMapping : NodeMappings)
-		{
-			uint32 NodeTemplateIndex = NodeMapping.NodeTemplateIndex;
-			*this << NodeTemplateIndex;
-		}
+		uint32 SharedDataSize = GraphSharedDataSize;
+		*this << SharedDataSize;
 	}
 
 	void FDecoratorWriter::EndNodeWriting()
@@ -119,20 +162,9 @@ namespace UE::AnimNext
 		}
 
 		ensure(NumNodesWritten == NodeMappings.Num());
-
-		// Now write all the object reference we found, we'll use them for the GC callbacks at runtime to avoid travering the graph for the references
-		int32 NumTrackedObjectsForGC = TrackedObjectsForGC.Num();
-		*this << NumTrackedObjectsForGC;
-
-		for (UObject* Obj : TrackedObjectsForGC)
-		{
-			// Save out the fully qualified object name
-			FString SavedString(Obj->GetPathName());
-			*this << SavedString;
-		}
 	}
 
-	void FDecoratorWriter::WriteNode(const FNodeHandle NodeHandle, const TFunction<FString (uint32 DecoratorIndex, const FString& PropertyName)>& GetDecoratorProperty)
+	void FDecoratorWriter::WriteNode(const FNodeHandle NodeHandle, const TFunction<const TMap<FString, FString>& (uint32 DecoratorIndex)>& GetDecoratorProperties)
 	{
 		ensure(bIsNodeWriting);
 
@@ -162,7 +194,7 @@ namespace UE::AnimNext
 		// Populate our node description into a temporary buffer
 		alignas(16) uint8 Buffer[64 * 1024];	// Max node size
 
-		FNodeDescription* NodeDesc = new(Buffer) FNodeDescription(NodeMapping->NodeHandle.GetNodeID(), NodeMapping->NodeTemplateHandle);
+		FNodeDescription* NodeDesc = new(Buffer) FNodeDescription(NodeMapping->NodeUID, NodeMapping->NodeTemplateHandle);
 
 		// Populate our decorator properties
 		const uint32 NumDecorators = NodeTemplate->GetNumDecorators();
@@ -174,15 +206,14 @@ namespace UE::AnimNext
 
 			FAnimNextDecoratorSharedData* SharedData = DecoratorTemplates[DecoratorIndex].GetDecoratorDescription(*NodeDesc);
 
-			// Curry our lambda with the decorator index
-			const auto GetDecoratorPropertyAt = [&GetDecoratorProperty, DecoratorIndex](const FString& PropertyName)
-			{
-				return GetDecoratorProperty(DecoratorIndex, PropertyName);
-			};
-
-			Decorator->SaveDecoratorSharedData(*this, GetDecoratorPropertyAt, *SharedData);
+			const TMap<FString, FString>& Properties = GetDecoratorProperties(DecoratorIndex);
+			Decorator->SaveDecoratorSharedData(*this, Properties, *SharedData);
 		}
 
+		// Write our node size first so the reader knows how much to skip ahead to the next node
+		uint32 NodeSize = NodeMapping->NodeSize;
+		*this << NodeSize;
+
 		// Append it to our archive
 		NodeDesc->Serialize(*this);
 
@@ -198,27 +229,5 @@ namespace UE::AnimNext
 	{
 		return GraphSharedDataArchiveBuffer;
 	}
-
-	FArchive& FDecoratorWriter::operator<<(UObject*& Obj)
-	{
-		// Add our object for GC tracking
-		TrackedObjectsForGC.AddUnique(Obj);
-
-		// TODO: Instead of saving our object path, we could save the object index
-		//       Indices can be re-used if an object is referenced more than once (is it common?)
-		//       Allows all objects to be batch found/loaded before nodes are de-serialized
-		//       Requires seeking to read the objects first or an intermediary archive to write objects first
-
-		// Save out the fully qualified object name
-		FString SavedString(Obj->GetPathName());
-		*this << SavedString;
-
-		return *this;
-	}
-
-	FArchive& FDecoratorWriter::operator<<(FObjectPtr& Obj)
-	{
-		return FArchiveUObject::SerializeObjectPtr(*this, Obj);
-	}
 }
 #endif
diff --git a/Engine/Plugins/Experimental/AnimNext/Source/AnimNext/Private/DecoratorBase/ExecutionContext.cpp b/Engine/Plugins/Experimental/AnimNext/Source/AnimNext/Private/DecoratorBase/ExecutionContext.cpp
--- a/Engine/Plugins/Experimental/AnimNext/Source/AnimNext/Private/DecoratorBase/ExecutionContext.cpp
+++ b/Engine/Plugins/Experimental/AnimNext/Source/AnimNext/Private/DecoratorBase/ExecutionContext.cpp
@@ -72,7 +72,7 @@ namespace UE::AnimNext
 		// We need to allocate a new node instance
 		const FDecoratorTemplate* DecoratorDescs = NodeTemplate->GetDecorators();
 
-		const uint32 InstanceSize = NodeTemplate->GetNodeInstanceDataSize();
+		const uint32 InstanceSize = NodeTemplate->GetInstanceSize();
 		uint8* NodeInstanceBuffer = reinterpret_cast<uint8*>(FMemory::Malloc(InstanceSize, 16));
 		FNodeInstance* NodeInstance = new(NodeInstanceBuffer) FNodeInstance(ChildNodeHandle);
 
@@ -80,12 +80,16 @@ namespace UE::AnimNext
 		const FDecoratorTemplate* StartDesc = DecoratorDescs;
 		const FDecoratorTemplate* EndDesc = DecoratorDescs + NodeTemplate->GetNumDecorators();
 
+		const FDecoratorTemplate* FailedDecoratorDesc = nullptr;
 		for (const FDecoratorTemplate* DecoratorDesc = StartDesc; DecoratorDesc != EndDesc; ++DecoratorDesc)
 		{
 			const FDecorator* Decorator = GetDecorator(*DecoratorDesc);
+			ensure(Decorator != nullptr);
 			if (Decorator == nullptr)
 			{
-				continue;	// Decorator hasn't been loaded or registered, skip it
+				// Failed to find the matching decorator, did it get unregistered or is the decorator descriptor corrupted?
+				FailedDecoratorDesc = DecoratorDesc;
+				break;
 			}
 
 			const uint32 DecoratorIndex = DecoratorDesc - DecoratorDescs;
@@ -98,6 +102,30 @@ namespace UE::AnimNext
 			Decorator->ConstructDecoratorInstance(*this, DecoratorPtr, *SharedData, *InstanceData);
 		}
 
+		if (FailedDecoratorDesc != nullptr)
+		{
+			// We failed to construct our node instance, destroy it
+			// Start destruction with the decorator prior to the one that failed
+			StartDesc = FailedDecoratorDesc - 1;
+			EndDesc = DecoratorDescs - 1;
+			for (const FDecoratorTemplate* DecoratorDesc = StartDesc; DecoratorDesc != EndDesc; --DecoratorDesc)
+			{
+				const FDecorator* Decorator = GetDecorator(*DecoratorDesc);
+				const uint32 DecoratorIndex = DecoratorDesc - DecoratorDescs;
+
+				FWeakDecoratorPtr DecoratorPtr(NodeInstance, DecoratorIndex);
+
+				const FAnimNextDecoratorSharedData* SharedData = DecoratorDesc->GetDecoratorDescription(NodeDesc);
+				FDecoratorInstanceData* InstanceData = DecoratorDesc->GetDecoratorInstance(*NodeInstance);
+
+				Decorator->DestructDecoratorInstance(*this, DecoratorPtr, *SharedData, *InstanceData);
+			}
+
+			FMemory::Free(NodeInstance);
+
+			return FDecoratorPtr();
+		}
+
 		return FDecoratorPtr(NodeInstance, ChildDecoratorIndex);
 	}
 
@@ -131,7 +159,8 @@ namespace UE::AnimNext
 		const FDecoratorTemplate* EndDesc = DecoratorDescs - 1;
 		for (const FDecoratorTemplate* DecoratorDesc = StartDesc; DecoratorDesc != EndDesc; --DecoratorDesc)
 		{
-			if (const FDecorator* Decorator = GetDecorator(*DecoratorDesc))
+			const FDecorator* Decorator = GetDecorator(*DecoratorDesc);
+			if (ensure(Decorator != nullptr))
 			{
 				const uint32 DecoratorIndex = DecoratorDesc - DecoratorDescs;
 
@@ -178,9 +207,10 @@ namespace UE::AnimNext
 		for (const FDecoratorTemplate* DecoratorDesc = StartDesc; DecoratorDesc != EndDesc; --DecoratorDesc)
 		{
 			const FDecorator* Decorator = GetDecorator(*DecoratorDesc);
+			ensure(Decorator != nullptr);
 			if (Decorator == nullptr)
 			{
-				continue;	// Decorator hasn't been loaded or registered, skip it
+				return false;	// Failed to find the matching decorator, did it get unregistered or is the decorator descriptor corrupted?
 			}
 
 			if (const IDecoratorInterface* Interface = Decorator->GetDecoratorInterface(InterfaceUID))
@@ -235,9 +265,10 @@ namespace UE::AnimNext
 		for (const FDecoratorTemplate* DecoratorDesc = StartDesc; DecoratorDesc != EndDesc; --DecoratorDesc)
 		{
 			const FDecorator* Decorator = GetDecorator(*DecoratorDesc);
+			ensure(Decorator != nullptr);
 			if (Decorator == nullptr)
 			{
-				continue;	// Decorator hasn't been loaded or registered, skip it
+				return false;	// Failed to find the matching decorator, did it get unregistered or is the decorator descriptor corrupted?
 			}
 
 			if (const IDecoratorInterface* Interface = Decorator->GetDecoratorInterface(InterfaceUID))
@@ -266,17 +297,13 @@ namespace UE::AnimNext
 	const FNodeTemplate* FExecutionContext::GetNodeTemplate(const FNodeDescription& NodeDesc) const
 	{
 		check(NodeDesc.GetTemplateHandle().IsValid());
-
-		const FNodeTemplateRegistry& NodeTemplateRegistry = FNodeTemplateRegistry::Get();
-		return NodeTemplateRegistry.Find(NodeDesc.GetTemplateHandle());
+		return FNodeTemplateRegistry::Get().Find(NodeDesc.GetTemplateHandle());
 	}
 
 	const FDecorator* FExecutionContext::GetDecorator(const FDecoratorTemplate& Template) const
 	{
 		check(Template.GetRegistryHandle().IsValid());
-
-		const FDecoratorRegistry& DecoratorRegistry = FDecoratorRegistry::Get();
-		return DecoratorRegistry.Find(Template.GetRegistryHandle());
+		return FDecoratorRegistry::Get().Find(Template.GetRegistryHandle());
 	}
 
 	FExecutionContext* GetThreadExecutionContext()
diff --git a/Engine/Plugins/Experimental/AnimNext/Source/AnimNext/Private/DecoratorBase/NodeDescription.cpp b/Engine/Plugins/Experimental/AnimNext/Source/AnimNext/Private/DecoratorBase/NodeDescription.cpp
--- a/Engine/Plugins/Experimental/AnimNext/Source/AnimNext/Private/DecoratorBase/NodeDescription.cpp
+++ b/Engine/Plugins/Experimental/AnimNext/Source/AnimNext/Private/DecoratorBase/NodeDescription.cpp
@@ -11,13 +11,11 @@ namespace UE::AnimNext
 {
 	void FNodeDescription::Serialize(FArchive& Ar)
 	{
-		const FNodeTemplateRegistry& NodeTemplateRegistry = FNodeTemplateRegistry::Get();
-
-		Ar << NodeID;
+		Ar << UID;
 
 		if (Ar.IsSaving())
 		{
-			const FNodeTemplate* NodeTemplate = NodeTemplateRegistry.Find(TemplateHandle);
+			const FNodeTemplate* NodeTemplate = FNodeTemplateRegistry::Get().Find(TemplateHandle);
 
 			uint32 TemplateUID = NodeTemplate->GetUID();
 			Ar << TemplateUID;
@@ -27,7 +25,7 @@ namespace UE::AnimNext
 			uint32 TemplateUID = 0;
 			Ar << TemplateUID;
 
-			TemplateHandle = NodeTemplateRegistry.Find(TemplateUID);
+			TemplateHandle = FNodeTemplateRegistry::Get().Find(TemplateUID);
 		}
 		else
 		{
@@ -37,7 +35,7 @@ namespace UE::AnimNext
 		}
 
 		// Use our template to serialize our decorators
-		const FNodeTemplate* NodeTemplate = NodeTemplateRegistry.Find(TemplateHandle);
+		const FNodeTemplate* NodeTemplate = FNodeTemplateRegistry::Get().Find(TemplateHandle);
 
 		const uint32 NumDecorators = NodeTemplate->GetNumDecorators();
 		const FDecoratorTemplate* DecoratorTemplates = NodeTemplate->GetDecorators();
diff --git a/Engine/Plugins/Experimental/AnimNext/Source/AnimNext/Private/DecoratorBase/NodeTemplate.cpp b/Engine/Plugins/Experimental/AnimNext/Source/AnimNext/Private/DecoratorBase/NodeTemplate.cpp
--- a/Engine/Plugins/Experimental/AnimNext/Source/AnimNext/Private/DecoratorBase/NodeTemplate.cpp
+++ b/Engine/Plugins/Experimental/AnimNext/Source/AnimNext/Private/DecoratorBase/NodeTemplate.cpp
@@ -3,82 +3,19 @@
 #include "DecoratorBase/NodeTemplate.h"
 
 #include "Serialization/Archive.h"
-#include "DecoratorBase/Decorator.h"
-#include "DecoratorBase/DecoratorRegistry.h"
-#include "DecoratorBase/NodeDescription.h"
-#include "DecoratorBase/NodeInstance.h"
 
 namespace UE::AnimNext
 {
 	void FNodeTemplate::Serialize(FArchive& Ar)
 	{
 		Ar << UID;
+		Ar << InstanceSize;
 		Ar << NumDecorators;
 
-		const uint32 NumDecorators_ = GetNumDecorators();
 		FDecoratorTemplate* DecoratorTemplates = GetDecorators();
-
-		for (uint32 DecoratorIndex = 0; DecoratorIndex < NumDecorators_; ++DecoratorIndex)
+		for (uint32 DecoratorIndex = 0; DecoratorIndex < GetNumDecorators(); ++DecoratorIndex)
 		{
 			DecoratorTemplates[DecoratorIndex].Serialize(Ar);
 		}
-
-		if (Ar.IsLoading())
-		{
-			// When loading, make sure to recompute all runtime dependent values (e.g. sizes and offsets)
-			Finalize();
-		}
-	}
-
-	void FNodeTemplate::Finalize()
-	{
-		const FDecoratorRegistry& DecoratorRegistry = FDecoratorRegistry::Get();
-
-		const uint32 NumDecorators_ = GetNumDecorators();
-		FDecoratorTemplate* DecoratorTemplates = GetDecorators();
-
-		uint32 SharedDataOffset = sizeof(FNodeDescription);
-		uint32 InstanceDataOffset = sizeof(FNodeInstance);
-
-		for (uint32 DecoratorIndex = 0; DecoratorIndex < NumDecorators_; ++DecoratorIndex)
-		{
-			FDecoratorTemplate& DecoratorTemplate = DecoratorTemplates[DecoratorIndex];
-			const FDecoratorUID DecoratorUID = DecoratorTemplate.GetUID();
-
-			uint32 DecoratorSharedDataOffset = 0;
-			uint32 DecoratorInstanceDataOffset = 0;	// For instance data, 0 is an invalid offset since the data follows an instance of FNodeInstance
-
-			// Skip decorators that we can't find
-			// If a decorator isn't loaded and we attempt to run the graph, it will be a no-op entry
-			if (const FDecorator* Decorator = DecoratorRegistry.Find(DecoratorUID))
-			{
-				const FDecoratorMemoryLayout MemoryLayout = Decorator->GetDecoratorMemoryDescription();
-
-				// Align our data
-				SharedDataOffset = Align(SharedDataOffset, MemoryLayout.SharedDataAlignment);
-				InstanceDataOffset = Align(InstanceDataOffset, MemoryLayout.InstanceDataAlignment);
-
-				// Save our decorator offsets
-				DecoratorSharedDataOffset = SharedDataOffset;
-				DecoratorInstanceDataOffset = InstanceDataOffset;
-
-				// Include our decorator
-				SharedDataOffset += MemoryLayout.SharedDataSize;
-				InstanceDataOffset += MemoryLayout.InstanceDataSize;
-			}
-
-			check(DecoratorSharedDataOffset <= MAX_uint16);
-			check(DecoratorInstanceDataOffset <= MAX_uint16);
-
-			// Update our decorator offsets
-			DecoratorTemplate.NodeSharedOffset = static_cast<uint16>(DecoratorSharedDataOffset);
-			DecoratorTemplate.NodeInstanceOffset = static_cast<uint16>(DecoratorInstanceDataOffset);
-		}
-
-		// Our size is the offset of the decorator that would follow afterwards
-		// If the size is too large, we'll end up truncating the offsets/size
-		// Set a value of 0 to be able to detect it later
-		NodeSharedDataSize = SharedDataOffset > MAX_uint16 ? 0 : static_cast<uint16>(SharedDataOffset);
-		NodeInstanceDataSize = InstanceDataOffset > MAX_uint16 ? 0 : static_cast<uint16>(InstanceDataOffset);
 	}
 }
diff --git a/Engine/Plugins/Experimental/AnimNext/Source/AnimNext/Private/DecoratorBase/NodeTemplateBuilder.cpp b/Engine/Plugins/Experimental/AnimNext/Source/AnimNext/Private/DecoratorBase/NodeTemplateBuilder.cpp
deleted file mode 100644
diff --git a/Engine/Plugins/Experimental/AnimNext/Source/AnimNext/Private/DecoratorBase/NodeTemplateRegistry.cpp b/Engine/Plugins/Experimental/AnimNext/Source/AnimNext/Private/DecoratorBase/NodeTemplateRegistry.cpp
--- a/Engine/Plugins/Experimental/AnimNext/Source/AnimNext/Private/DecoratorBase/NodeTemplateRegistry.cpp
+++ b/Engine/Plugins/Experimental/AnimNext/Source/AnimNext/Private/DecoratorBase/NodeTemplateRegistry.cpp
@@ -81,7 +81,7 @@ namespace UE::AnimNext
 		TemplateUIDToHandleMap.Remove(TemplateUID);
 	}
 
-	FNodeTemplate* FNodeTemplateRegistry::FindMutable(FNodeTemplateRegistryHandle TemplateHandle)
+	FNodeTemplate* FNodeTemplateRegistry::Find(FNodeTemplateRegistryHandle TemplateHandle)
 	{
 		if (!TemplateHandle.IsValid())
 		{
@@ -105,4 +105,10 @@ namespace UE::AnimNext
 	{
 		return TemplateUIDToHandleMap.Num();
 	}
+
+	void FNodeTemplateRegistry::Clear()
+	{
+		TemplateBuffer.Empty(0);
+		TemplateUIDToHandleMap.Empty(0);
+	}
 }
diff --git a/Engine/Plugins/Experimental/AnimNext/Source/AnimNext/Private/DecoratorInterfaces/IUpdate.cpp b/Engine/Plugins/Experimental/AnimNext/Source/AnimNext/Private/DecoratorInterfaces/IUpdate.cpp
--- a/Engine/Plugins/Experimental/AnimNext/Source/AnimNext/Private/DecoratorInterfaces/IUpdate.cpp
+++ b/Engine/Plugins/Experimental/AnimNext/Source/AnimNext/Private/DecoratorInterfaces/IUpdate.cpp
@@ -7,11 +7,6 @@
 
 namespace UE::AnimNext
 {
-	FUpdateTraversalContext::FUpdateTraversalContext(float InDeltaTime)
-		: DeltaTime(InDeltaTime)
-	{
-	}
-
 	void IUpdate::PreUpdate(FExecutionContext& Context, const TDecoratorBinding<IUpdate>& Binding) const
 	{
 		TDecoratorBinding<IUpdate> SuperBinding;
@@ -42,7 +37,7 @@ namespace UE::AnimNext
 		EUpdateStep			DesiredStep = EUpdateStep::PreUpdate;
 	};
 
-	void UpdateGraph(FExecutionContext& Context, FWeakDecoratorPtr GraphRootPtr, float DeltaTime)
+	void UpdateGraph(FExecutionContext& Context, FWeakDecoratorPtr GraphRootPtr)
 	{
 		if (!GraphRootPtr.IsValid())
 		{
@@ -57,7 +52,7 @@ namespace UE::AnimNext
 		FChildrenArray Children;
 		Children.Reserve(64);
 
-		FUpdateTraversalContext TraversalContext(DeltaTime);
+		FUpdateTraversalContext TraversalContext;
 
 		FScopedTraversalContext ScopedTraversalContext(Context, TraversalContext);
 
diff --git a/Engine/Plugins/Experimental/AnimNext/Source/AnimNext/Private/Decorators/BlendTwoWay.cpp b/Engine/Plugins/Experimental/AnimNext/Source/AnimNext/Private/Decorators/BlendTwoWay.cpp
--- a/Engine/Plugins/Experimental/AnimNext/Source/AnimNext/Private/Decorators/BlendTwoWay.cpp
+++ b/Engine/Plugins/Experimental/AnimNext/Source/AnimNext/Private/Decorators/BlendTwoWay.cpp
@@ -6,12 +6,9 @@
 
 namespace UE::AnimNext
 {
-	AUTO_REGISTER_ANIM_DECORATOR(FBlendTwoWayDecorator)
-
 	DEFINE_ANIM_DECORATOR_BEGIN(FBlendTwoWayDecorator)
 		DEFINE_ANIM_DECORATOR_IMPLEMENTS_INTERFACE(IEvaluate)
 		DEFINE_ANIM_DECORATOR_IMPLEMENTS_INTERFACE(IUpdate)
-		DEFINE_ANIM_DECORATOR_IMPLEMENTS_INTERFACE(IHierarchy)
 	DEFINE_ANIM_DECORATOR_END(FBlendTwoWayDecorator)
 
 	void FBlendTwoWayDecorator::PostEvaluate(FExecutionContext& Context, const TDecoratorBinding<IEvaluate>& Binding) const
@@ -35,31 +32,31 @@ namespace UE::AnimNext
 
 		if (SharedData->BlendWeight < 1.0)
 		{
-			if (!InstanceData->ChildA.IsValid())
+			if (!InstanceData->Children[0].IsValid())
 			{
 				// We need to blend a child that isn't instanced yet, allocate it
-				InstanceData->ChildA = Context.AllocateNodeInstance(Binding, SharedData->ChildA);
+				InstanceData->Children[0] = Context.AllocateNodeInstance(Binding, SharedData->Children[0]);
 			}
 
 			if (SharedData->BlendWeight == 0.0)
 			{
 				// We no longer need this child, release it
-				InstanceData->ChildB.Reset();
+				InstanceData->Children[1].Reset();
 			}
 		}
 
 		if (SharedData->BlendWeight > 0.0)
 		{
-			if (!InstanceData->ChildB.IsValid())
+			if (!InstanceData->Children[1].IsValid())
 			{
 				// We need to blend a child that isn't instanced yet, allocate it
-				InstanceData->ChildB = Context.AllocateNodeInstance(Binding, SharedData->ChildB);
+				InstanceData->Children[1] = Context.AllocateNodeInstance(Binding, SharedData->Children[1]);
 			}
 
 			if (SharedData->BlendWeight == 1.0)
 			{
 				// We no longer need this child, release it
-				InstanceData->ChildA.Reset();
+				InstanceData->Children[0].Reset();
 			}
 		}
 	}
@@ -69,7 +66,7 @@ namespace UE::AnimNext
 		const FInstanceData* InstanceData = Binding.GetInstanceData<FInstanceData>();
 
 		// Add the two child handles, even if they are empty
-		Children.Add(InstanceData->ChildA);
-		Children.Add(InstanceData->ChildB);
+		Children.Add(InstanceData->Children[0]);
+		Children.Add(InstanceData->Children[1]);
 	}
 }
diff --git a/Engine/Plugins/Experimental/AnimNext/Source/AnimNext/Private/Decorators/BlendTwoWay.h b/Engine/Plugins/Experimental/AnimNext/Source/AnimNext/Private/Decorators/BlendTwoWay.h
--- a/Engine/Plugins/Experimental/AnimNext/Source/AnimNext/Private/Decorators/BlendTwoWay.h
+++ b/Engine/Plugins/Experimental/AnimNext/Source/AnimNext/Private/Decorators/BlendTwoWay.h
@@ -16,16 +16,10 @@ struct FAnimNextBlendTwoWayDecoratorSharedData : public FAnimNextDecoratorShared
 {
 	GENERATED_BODY()
 
-	/** First output to be blended (full weight is 0.0). */
-	UPROPERTY(meta = (Input))
-	FAnimNextDecoratorHandle ChildA;
+	UPROPERTY()
+	FAnimNextDecoratorHandle Children[2];
 
-	/** Second output to be blended (full weight is 1.0). */
-	UPROPERTY(meta = (Input))
-	FAnimNextDecoratorHandle ChildB;
-
-	/** How much to blend our two children: 0.0 is fully child A while 1.0 is fully child B. */
-	UPROPERTY(meta = (Input))
+	UPROPERTY()
 	double BlendWeight = 0.0;
 };
 
@@ -44,8 +38,7 @@ namespace UE::AnimNext
 
 		struct FInstanceData : FDecorator::FInstanceData
 		{
-			FDecoratorPtr ChildA;
-			FDecoratorPtr ChildB;
+			FDecoratorPtr Children[2];
 		};
 
 		// IEvaluate impl
diff --git a/Engine/Plugins/Experimental/AnimNext/Source/AnimNext/Private/Decorators/SequencePlayer.cpp b/Engine/Plugins/Experimental/AnimNext/Source/AnimNext/Private/Decorators/SequencePlayer.cpp
--- a/Engine/Plugins/Experimental/AnimNext/Source/AnimNext/Private/Decorators/SequencePlayer.cpp
+++ b/Engine/Plugins/Experimental/AnimNext/Source/AnimNext/Private/Decorators/SequencePlayer.cpp
@@ -3,64 +3,20 @@
 #include "Decorators/SequencePlayer.h"
 
 #include "DecoratorBase/ExecutionContext.h"
-#include "Param/ParamStack.h"
-#include "Graph/AnimNext_LODPose.h"
-#include "DecompressionTools.h"
 
 namespace UE::AnimNext
 {
-	AUTO_REGISTER_ANIM_DECORATOR(FSequencePlayerDecorator)
-
 	DEFINE_ANIM_DECORATOR_BEGIN(FSequencePlayerDecorator)
 		DEFINE_ANIM_DECORATOR_IMPLEMENTS_INTERFACE(IEvaluate)
 		DEFINE_ANIM_DECORATOR_IMPLEMENTS_INTERFACE(IUpdate)
 		DEFINE_ANIM_DECORATOR_IMPLEMENTS_INTERFACE(ITimeline)
 	DEFINE_ANIM_DECORATOR_END(FSequencePlayerDecorator)
 
-	void FSequencePlayerDecorator::FInstanceData::Construct(FExecutionContext& Context, FWeakDecoratorPtr DecoratorPtr, const FAnimNextSequencePlayerDecoratorSharedData& SharedData)
-	{
-		if (SharedData.AnimSequence != nullptr)
-		{
-			const float SequenceLength = SharedData.AnimSequence->GetPlayLength();
-			InternalTimeAccumulator = FMath::Clamp(SharedData.StartPosition, 0.f, SequenceLength);
-			PrevInternalTimeAccumulator = InternalTimeAccumulator;
-		}
-	}
-
 	void FSequencePlayerDecorator::PreEvaluate(FExecutionContext& Context, const TDecoratorBinding<IEvaluate>& Binding) const
 	{
-		const FSharedData* SharedData = Binding.GetSharedData<FSharedData>();
-		const FInstanceData* InstanceData = Binding.GetInstanceData<FInstanceData>();
-
 		// TODO: Sample pose
 		// FPose foo = ...
 		// Context.PushPose(foo);
-
-		FParamStack& ParamStack = FParamStack::Get();
-
-		FDeltaTimeRecord DeltaTimeRecord;
-		DeltaTimeRecord.Set(InstanceData->PrevInternalTimeAccumulator, InstanceData->InternalTimeAccumulator);
-
-		const FAnimExtractContext ExtractionContext(static_cast<double>(InstanceData->InternalTimeAccumulator)
-			, false /*Output.AnimInstanceProxy->ShouldExtractRootMotion()*/
-			, DeltaTimeRecord
-			, SharedData->bLoop);
-
-		const FAnimNextGraphReferencePose& GraphReferencePose = ParamStack.GetParam<FAnimNextGraphReferencePose>("GraphReferencePose");
-		const int32 GraphLODLevel = ParamStack.GetParam<int32>("GraphLODLevel");
-		const bool bGraphExpectsAdditive = ParamStack.GetParam<bool>("GraphExpectsAdditive");
-
-		// HACK: Write to our output directly, only works if we have a single node in the graph that outputs a pose
-		FAnimNextGraphLODPose& ResultPose = ParamStack.GetMutableParam<FAnimNextGraphLODPose>("ResultPose");
-
-		ResultPose.LODPose.PrepareForLOD(*GraphReferencePose.ReferencePose, GraphLODLevel, true, bGraphExpectsAdditive);
-
-		// Note : calling FDecompressionTools instead of UAnimSequence / UAnimSequenceBase, as I can not modify the engine code (so I extracted the function)
-		// TODO : Revisit this, so we can plug other sequence types
-		if (SharedData->AnimSequence != nullptr)
-		{
-			FDecompressionTools::GetAnimationPose(SharedData->AnimSequence, ResultPose.LODPose, ExtractionContext);
-		}
 	}
 
 	double FSequencePlayerDecorator::GetPlayRate(FExecutionContext& Context, const TDecoratorBinding<ITimeline>& Binding) const
@@ -74,30 +30,17 @@ namespace UE::AnimNext
 		const FSharedData* SharedData = Binding.GetSharedData<FSharedData>();
 		FInstanceData* InstanceData = Binding.GetInstanceData<FInstanceData>();
 
-		if (SharedData->AnimSequence != nullptr)
-		{
-			FUpdateTraversalContext& TraversalContext = Context.GetTraversalContext<FUpdateTraversalContext>();
-
-			TDecoratorBinding<ITimeline> TimelineDecorator;
-			Context.GetInterface(Binding, TimelineDecorator);
-
-			const float DeltaTime = TraversalContext.GetDeltaTime();
-			const float PlayRate = (float)TimelineDecorator.GetPlayRate(Context);
+		TDecoratorBinding<ITimeline> TimelineDecorator;
+		Context.GetInterface(Binding, TimelineDecorator);
 
-			const float EffectiveDelta = FMath::IsNearlyZero(DeltaTime) || FMath::IsNearlyZero(PlayRate) ? 0.f : DeltaTime * PlayRate;
+		FUpdateTraversalContext& TraversalContext = Context.GetTraversalContext<FUpdateTraversalContext>();
 
-			const float SequenceLength = SharedData->AnimSequence->GetPlayLength();
-			float CurrentTime = SharedData->bLoop
-				? FMath::Fmod(InstanceData->InternalTimeAccumulator + EffectiveDelta, SequenceLength)
-				: FMath::Clamp(InstanceData->InternalTimeAccumulator + EffectiveDelta, 0.f, SequenceLength);
+		double DeltaTime = TraversalContext.GetDeltaTime();
 
-			if (SharedData->bLoop && CurrentTime < 0.f)
-			{
-				CurrentTime += SequenceLength;
-			}
+		// Combine the execution play rate along with the timeline play rate, possibly overriden by another decorator
+		double PlayRate = TraversalContext.GetPlayRate() * TimelineDecorator.GetPlayRate(Context);
 
-			InstanceData->PrevInternalTimeAccumulator = InstanceData->InternalTimeAccumulator;
-			InstanceData->InternalTimeAccumulator = CurrentTime;
-		}
+		// TODO: need an actual impl, for illustrative purposes
+		InstanceData->CurrentTime += DeltaTime * PlayRate;
 	}
 }
diff --git a/Engine/Plugins/Experimental/AnimNext/Source/AnimNext/Private/Decorators/SequencePlayer.h b/Engine/Plugins/Experimental/AnimNext/Source/AnimNext/Private/Decorators/SequencePlayer.h
--- a/Engine/Plugins/Experimental/AnimNext/Source/AnimNext/Private/Decorators/SequencePlayer.h
+++ b/Engine/Plugins/Experimental/AnimNext/Source/AnimNext/Private/Decorators/SequencePlayer.h
@@ -3,34 +3,26 @@
 #pragma once
 
 #include "CoreMinimal.h"
+#include "UObject/ObjectPtr.h"
 #include "DecoratorBase/Decorator.h"
 #include "DecoratorInterfaces/IEvaluate.h"
 #include "DecoratorInterfaces/ITimeline.h"
 #include "DecoratorInterfaces/IUpdate.h"
-#include "Animation/AnimSequence.h"
 
 #include "SequencePlayer.generated.h"
 
+class UAnimSequence;
+
 USTRUCT()
 struct FAnimNextSequencePlayerDecoratorSharedData : public FAnimNextDecoratorSharedData
 {
 	GENERATED_BODY()
 
-	/** The sequence to play. */
-	UPROPERTY(meta = (Input))
-	TObjectPtr<UAnimSequence> AnimSequence;
-
-	/** The play rate multiplier at which this sequence plays. */
-	UPROPERTY(meta = (Input))
-	float PlayRate = 1.0f;
+	UPROPERTY()
+	TObjectPtr<UAnimSequence> AnimSeq;
 
-	/** The time at which we should start playing this sequence. */
-	UPROPERTY(meta = (Input))
-	float StartPosition = 0.0f;
-
-	/** Whether or not this sequence playback will loop. */
-	UPROPERTY(meta = (Input))
-	bool bLoop = 0.0f;
+	UPROPERTY()
+	double PlayRate = 1.0;
 };
 
 namespace UE::AnimNext
@@ -48,10 +40,7 @@ namespace UE::AnimNext
 
 		struct FInstanceData : FDecorator::FInstanceData
 		{
-			float InternalTimeAccumulator = 0.0f;
-			float PrevInternalTimeAccumulator = 0.0f;
-
-			void Construct(FExecutionContext& Context, FWeakDecoratorPtr DecoratorPtr, const FAnimNextSequencePlayerDecoratorSharedData& SharedData);
+			double CurrentTime = 0.0;
 		};
 
 		// IEvaluate impl
diff --git a/Engine/Plugins/Experimental/AnimNext/Source/AnimNext/Private/Graph/AnimGraph/AnimNode_AnimNextGraph.cpp b/Engine/Plugins/Experimental/AnimNext/Source/AnimNext/Private/Graph/AnimGraph/AnimNode_AnimNextGraph.cpp
--- a/Engine/Plugins/Experimental/AnimNext/Source/AnimNext/Private/Graph/AnimGraph/AnimNode_AnimNextGraph.cpp
+++ b/Engine/Plugins/Experimental/AnimNext/Source/AnimNext/Private/Graph/AnimGraph/AnimNode_AnimNextGraph.cpp
@@ -11,7 +11,8 @@
 #include "GenerationTools.h"
 #include "ReferencePose.h"
 #include "Graph/AnimNext_LODPose.h"
-#include "Graph/AnimNextExecuteContext.h"
+#include "Animation/AnimSequence.h" // TEST
+#include "Graph/RigUnit_AnimNextAnimSequence.h" // TEST
 #include "Engine/SkeletalMesh.h"
 #include "BoneContainer.h"
 #include "Param/ParamStack.h"
@@ -27,14 +28,7 @@ FAnimNode_AnimNextGraph::FAnimNode_AnimNextGraph()
 	, AnimNextGraph(nullptr)
 	, LODThreshold(INDEX_NONE)
 {
-}
 
-FAnimNode_AnimNextGraph::~FAnimNode_AnimNextGraph()
-{
-	if (GraphInstancePtr.IsValid())
-	{
-		AnimNextGraph->ReleaseInstance(GraphInstancePtr);
-	}
 }
 
 void FAnimNode_AnimNextGraph::OnInitializeAnimInstance(const FAnimInstanceProxy* InProxy, const UAnimInstance* InAnimInstance)
@@ -57,6 +51,8 @@ void FAnimNode_AnimNextGraph::Update_AnyThread(const FAnimationUpdateContext& Co
 {
 	DECLARE_SCOPE_HIERARCHICAL_COUNTER_FUNC()
 
+	GraphDeltaTime += Context.GetDeltaTime();
+
 	SourceLink.Update(Context);
 
 	if (IsLODEnabled(Context.AnimInstanceProxy))
@@ -64,17 +60,6 @@ void FAnimNode_AnimNextGraph::Update_AnyThread(const FAnimationUpdateContext& Co
 		GetEvaluateGraphExposedInputs().Execute(Context);
 
 		PropagateInputProperties(Context.AnimInstanceProxy->GetAnimInstanceObject());
-
-		// Populate our param stack since our instance data might need it during construction
-		const int32 LODLevel = Context.AnimInstanceProxy->GetLODLevel();
-
-		UE::AnimNext::FParamStack::Get().PushValues(
-			"GraphLODLevel", LODLevel
-		);
-
-		UE::AnimNext::FContext AnimNextContext(Context.GetDeltaTime());
-
-		AnimNextGraph->Run(AnimNextContext, GraphInstancePtr, EAnimNextGraphSimulationSteps::Update);
 	}
 
 	FAnimNode_CustomProperty::Update_AnyThread(Context);
@@ -88,20 +73,6 @@ void FAnimNode_AnimNextGraph::Initialize_AnyThread(const FAnimationInitializeCon
 
 	SourceLink.Initialize(Context);
 
-	if (!GraphInstancePtr.IsValid() && AnimNextGraph)
-	{
-		// If we don't have an instance yet, create one
-
-		// Populate our param stack since our instance data might need it during construction
-		const int32 LODLevel = Context.AnimInstanceProxy->GetLODLevel();
-
-		UE::AnimNext::FParamStack::Get().PushValues(
-			"GraphLODLevel", LODLevel
-		);
-
-		GraphInstancePtr = AnimNextGraph->AllocateInstance();
-	}
-
 	FAnimNode_CustomProperty::Initialize_AnyThread(Context);
 }
 
@@ -139,23 +110,29 @@ void FAnimNode_AnimNextGraph::Evaluate_AnyThread(FPoseContext & Output)
 
 	const int32 LODLevel = Output.AnimInstanceProxy->GetLODLevel();
 	
-	UE::AnimNext::FContext Context(0.0f);
+	// TODO : Using AnimInstanceProxy->GetDeltaSeconds() makes the preview to advance  multiple times when debug options are activated (i.e. ShowUncompressedAnim)
+	// See where to get / how to calculate the correct delta value
+	UE::AnimNext::FContext Context(GraphDeltaTime);
+	GraphDeltaTime = 0.f; // Reset for the case that we receive multiple calls to Evaluate (debug options)
 
 	FAnimNextGraphReferencePose GraphReferencePose(&RefPose);
+	FAnimNextGraph_AnimSequence GraphTestSequence(TestSequence);  // TEST
 
 	FAnimNextGraphLODPose GraphSourceLODPose(FLODPose(RefPose, LODLevel, false, Output.ExpectsAdditivePose()));
 	FAnimNextGraphLODPose ResultPose(FLODPose(RefPose, LODLevel, true, Output.ExpectsAdditivePose()));
 	FGenerationTools::RemapPose(LODLevel, SourcePose, RefPose, GraphSourceLODPose.LODPose);
 
 	FParamStack::FPushedLayerHandle LayerHandle = Context.GetMutableParamStack().PushValues(
+		"AnimSequencePlayerState", SequencePlayerState,
 		"GraphReferencePose", GraphReferencePose,
 		"ResultPose", ResultPose,
 		"GraphLODLevel", LODLevel,
 		"GraphExpectsAdditive", Output.ExpectsAdditivePose(),
-		"SourcePose", GraphSourceLODPose						// TODO : Pass this as a external variable maybe ? When we have support for variables in the rigvm graph
+		"SourcePose", GraphSourceLODPose,						// TODO : Pass this as a external variable maybe ? When we have support for variables in the rigvm graph
+		"TestSequence", GraphTestSequence						// TEST anim decompression - Anim Sequence to decompress);
 	);
 
-	AnimNextGraph->Run(Context, GraphInstancePtr, EAnimNextGraphSimulationSteps::Evaluate);
+	AnimNextGraph->Run(Context);
 
 	FGenerationTools::RemapPose(LODLevel, RefPose, ResultPose.LODPose, Output);
 
diff --git a/Engine/Plugins/Experimental/AnimNext/Source/AnimNext/Private/Graph/AnimNextGraph.cpp b/Engine/Plugins/Experimental/AnimNext/Source/AnimNext/Private/Graph/AnimNextGraph.cpp
--- a/Engine/Plugins/Experimental/AnimNext/Source/AnimNext/Private/Graph/AnimNextGraph.cpp
+++ b/Engine/Plugins/Experimental/AnimNext/Source/AnimNext/Private/Graph/AnimNextGraph.cpp
@@ -3,13 +3,10 @@
 #include "Graph/AnimNextGraph.h"
 #include "RigVMCore/RigVMMemoryStorage.h"
 #include "RigVMCore/RigVMExecuteContext.h"
-#include "Graph/RigUnit_AnimNextGraphRoot.h"
+#include "Graph/RigUnit_AnimNextBeginExecution.h"
 #include "UObject/ObjectSaveContext.h"
 #include "UObject/Package.h"
 #include "Param/ParamTypeHandle.h"
-#include "DecoratorBase/DecoratorReader.h"
-#include "DecoratorBase/ExecutionContext.h"
-#include "Serialization/MemoryReader.h"
 
 #include UE_INLINE_GENERATED_CPP_BY_NAME(AnimNextGraph)
 
@@ -19,42 +16,15 @@ const FName EntryPointName("GetData");
 const FName ResultName("Result");
 }
 
-UE::AnimNext::FDecoratorPtr UAnimNextGraph::AllocateInstance() const
+void UAnimNextGraph::Run(const UE::AnimNext::FContext& Context) const
 {
-	if (!ResolvedRootDecoratorHandle.IsValid())
-	{
-		return UE::AnimNext::FDecoratorPtr();
-	}
-
-	UE::AnimNext::FExecutionContext Context(SharedDataBuffer);
-
-	return Context.AllocateNodeInstance(UE::AnimNext::FWeakDecoratorPtr(), ResolvedRootDecoratorHandle);
-}
-
-void UAnimNextGraph::ReleaseInstance(UE::AnimNext::FDecoratorPtr& GraphInstancePtr) const
-{
-	if (!GraphInstancePtr.IsValid())
-	{
-		return;
-	}
-
-	// We need an execution context for this graph to be active when we delete the graph instance
-	UE::AnimNext::FExecutionContext Context(SharedDataBuffer);
-
-	GraphInstancePtr.Reset();
-}
-
-void UAnimNextGraph::Run(const UE::AnimNext::FContext& Context, UE::AnimNext::FWeakDecoratorPtr GraphInstancePtr, EAnimNextGraphSimulationSteps SimulationSteps) const
-{
-	if (RigVM && GraphInstancePtr.IsValid())
+	if(RigVM)
 	{
 		FRigVMExtendedExecuteContext RigVMExtendedExecuteContext;
 		FAnimNextExecuteContext& AnimNextContext = RigVMExtendedExecuteContext.GetPublicDataSafe<FAnimNextExecuteContext>();
 		AnimNextContext.SetContextData(Context);
-		AnimNextContext.InitializeWithGraph(SharedDataBuffer, GraphInstancePtr);
-		AnimNextContext.SetSimulationSteps(SimulationSteps);
 
-		RigVM->Execute(RigVMExtendedExecuteContext, TArray<URigVMMemoryStorage*>(), FRigUnit_AnimNextGraphRoot::EventName);
+		RigVM->Execute(RigVMExtendedExecuteContext, TArray<URigVMMemoryStorage*>(), FRigUnit_AnimNextBeginExecution::EventName);
 	}
 }
 
@@ -115,77 +85,3 @@ void UAnimNextGraph::GetPreloadDependencies(TArray<UObject*>& OutDeps)
 		}
 	}
 }
-
-void UAnimNextGraph::Serialize(FArchive& Ar)
-{
-	Super::Serialize(Ar);
-
-	if (Ar.IsLoading())
-	{
-		int32 SharedDataArchiveBufferSize = 0;
-		Ar << SharedDataArchiveBufferSize;
-
-#if !WITH_EDITORONLY_DATA
-		// When editor data isn't present, we don't persist the archive buffer as it is only needed on load
-		// to populate the graph shared data
-		TArray<uint8> SharedDataArchiveBuffer;
-#endif
-
-		SharedDataArchiveBuffer.SetNumUninitialized(SharedDataArchiveBufferSize);
-		Ar.Serialize(SharedDataArchiveBuffer.GetData(), SharedDataArchiveBufferSize);
-
-		LoadFromArchiveBuffer(SharedDataArchiveBuffer);
-	}
-	else if (Ar.IsSaving())
-	{
-#if WITH_EDITORONLY_DATA
-		// We only save the archive buffer, if code changes we'll be able to de-serialize from it when
-		// building the runtime buffer
-		// This allows us to have editor only/non-shipping only properties that are stripped out on load
-		int32 SharedDataArchiveBufferSize = SharedDataArchiveBuffer.Num();
-		Ar << SharedDataArchiveBufferSize;
-		Ar.Serialize(SharedDataArchiveBuffer.GetData(), SharedDataArchiveBufferSize);
-#endif
-	}
-	else
-	{
-		// Counting, etc
-		Ar << SharedDataBuffer;
-		Ar << TrackedObjectsForGC;
-
-#if WITH_EDITORONLY_DATA
-		Ar << SharedDataArchiveBuffer;
-#endif
-	}
-}
-
-void UAnimNextGraph::AddReferencedObjects(UObject* InThis, FReferenceCollector& Collector)
-{
-	UAnimNextGraph* This = CastChecked<UAnimNextGraph>(InThis);
-	Collector.AddStableReferenceArray(&This->TrackedObjectsForGC);
-
-	Super::AddReferencedObjects(InThis, Collector);
-}
-
-bool UAnimNextGraph::LoadFromArchiveBuffer(const TArray<uint8>& InSharedDataArchiveBuffer)
-{
-	using namespace UE::AnimNext;
-
-	// Reconstruct our graph shared data
-	FMemoryReader GraphSharedDataArchive(InSharedDataArchiveBuffer);
-	FDecoratorReader DecoratorReader(GraphSharedDataArchive);
-
-	const FDecoratorReader::EErrorState ErrorState = DecoratorReader.ReadGraph(SharedDataBuffer, TrackedObjectsForGC);
-	if (ErrorState == FDecoratorReader::EErrorState::None)
-	{
-		ResolvedRootDecoratorHandle = DecoratorReader.ResolveEntryPointHandle(RootDecoratorHandle);
-		return true;
-	}
-	else
-	{
-		SharedDataBuffer.Empty(0);
-		TrackedObjectsForGC.Empty(0);
-		ResolvedRootDecoratorHandle = FAnimNextDecoratorHandle();
-		return false;
-	}
-}
diff --git a/Engine/Plugins/Experimental/AnimNext/Source/AnimNext/Private/Graph/RigDecorator_AnimNextCppDecorator.cpp b/Engine/Plugins/Experimental/AnimNext/Source/AnimNext/Private/Graph/RigDecorator_AnimNextCppDecorator.cpp
deleted file mode 100644
diff --git a/Engine/Plugins/Experimental/AnimNext/Source/AnimNext/Private/Graph/RigUnit_AnimNextAnimSequence.cpp b/Engine/Plugins/Experimental/AnimNext/Source/AnimNext/Private/Graph/RigUnit_AnimNextAnimSequence.cpp
new file mode 100644
--- /dev/null
+++ b/Engine/Plugins/Experimental/AnimNext/Source/AnimNext/Private/Graph/RigUnit_AnimNextAnimSequence.cpp
@@ -0,0 +1,92 @@
+// Copyright Epic Games, Inc. All Rights Reserved.
+
+#include "RigUnit_AnimNextAnimSequence.h"
+#include "Graph/AnimNextExecuteContext.h"
+#include "DecompressionTools.h"
+#include "Animation/AnimSequence.h"
+#include "Graph/AnimNext_LODPose.h"
+#include "Context.h"
+#include "Param/ParamStack.h"
+
+#include UE_INLINE_GENERATED_CPP_BY_NAME(RigUnit_AnimNextAnimSequence)
+
+// --- FRigUnit_AnimNext_AnimSequenceAsset --- 
+
+FRigUnit_AnimNext_AnimSequenceAsset_Execute()
+{
+	using namespace UE::AnimNext;
+
+	//Sequence.AnimSequence = AnimSequence; // TODO : Enable once we have RigVM object support enabled
+}
+
+// --- FRigUnit_AnimNext_AnimSequencePlayer --- 
+
+FRigUnit_AnimNext_AnimSequencePlayer_Initialize()
+{
+	using namespace UE::AnimNext;
+
+	const FContext& InterfaceContext = ExecuteContext.GetContext();
+
+	const UAnimSequenceBase* AnimSequenceBase = Sequence.AnimSequence;
+	if (AnimSequenceBase != nullptr)
+	{
+		// Get internal state (as a param)
+		FAnimSequencePlayerState& AnimSequencePlayerState = InterfaceContext.GetMutableParamStack().GetMutableParam<FAnimSequencePlayerState>("AnimSequencePlayerState");
+
+		const float SequenceLength = AnimSequenceBase->GetPlayLength();
+		AnimSequencePlayerState.InternalTimeAccumulator = FMath::Clamp(Parameters.StartPosition, 0.f, SequenceLength);
+		AnimSequencePlayerState.PrevInternalTimeAccumulator = AnimSequencePlayerState.InternalTimeAccumulator;
+	}
+}
+
+FRigUnit_AnimNext_AnimSequencePlayer_Execute()
+{
+	using namespace UE::AnimNext;
+
+	const FContext& InterfaceContext = ExecuteContext.GetContext();
+	
+	const UAnimSequenceBase* AnimSequenceBase = Sequence.AnimSequence;
+	if (AnimSequenceBase != nullptr)
+	{
+		// Get internal state (as a param)
+		FAnimSequencePlayerState& AnimSequencePlayerState = InterfaceContext.GetMutableParamStack().GetMutableParam<FAnimSequencePlayerState>("AnimSequencePlayerState");
+
+		const float DeltaTime = InterfaceContext.GetDeltaTime();
+		const float EffectiveDelta = FMath::IsNearlyZero(DeltaTime) || FMath::IsNearlyZero(Parameters.PlayRate) ? 0.f : DeltaTime * Parameters.PlayRate;
+
+		const float SequenceLength = AnimSequenceBase->GetPlayLength();
+		float CurrentTime = Parameters.bLoop
+			? FMath::Fmod(AnimSequencePlayerState.InternalTimeAccumulator + EffectiveDelta, SequenceLength)
+			: FMath::Clamp(AnimSequencePlayerState.InternalTimeAccumulator + EffectiveDelta, 0.f, SequenceLength);
+
+		if (Parameters.bLoop && CurrentTime < 0.f)
+		{
+			CurrentTime += SequenceLength;
+		}
+
+		AnimSequencePlayerState.PrevInternalTimeAccumulator = AnimSequencePlayerState.InternalTimeAccumulator;
+		AnimSequencePlayerState.InternalTimeAccumulator = CurrentTime;
+
+		FDeltaTimeRecord DeltaTimeRecord;
+		DeltaTimeRecord.Set(AnimSequencePlayerState.PrevInternalTimeAccumulator, AnimSequencePlayerState.InternalTimeAccumulator);
+
+		const FAnimExtractContext ExtractionContext(static_cast<double>(AnimSequencePlayerState.InternalTimeAccumulator)
+			, false /*Output.AnimInstanceProxy->ShouldExtractRootMotion()*/
+			, DeltaTimeRecord
+			, Parameters.bLoop);
+
+		const FAnimNextGraphReferencePose& GraphReferencePose = InterfaceContext.GetParamStack().GetParam<FAnimNextGraphReferencePose>("GraphReferencePose");
+		const int32 GraphLODLevel = InterfaceContext.GetParamStack().GetParam<int32>("GraphLODLevel");
+		const bool bGraphExpectsAdditive = InterfaceContext.GetParamStack().GetParam<bool>("GraphExpectsAdditive");
+
+		LODPose.LODPose.PrepareForLOD(*GraphReferencePose.ReferencePose, GraphLODLevel, true, bGraphExpectsAdditive);
+		
+		// Note : calling FDecompressionTools instead of UAnimSequence / UAnimSequenceBase, as I can not modify the engine code (so I extracted the function)
+		// TODO : Revisit this, so we can plug other sequence types
+		const UAnimSequence* AnimSequence = Cast<UAnimSequence>(AnimSequenceBase);
+		if (AnimSequence != nullptr)
+		{
+			FDecompressionTools::GetAnimationPose(AnimSequence, LODPose.LODPose, ExtractionContext);
+		}
+	}
+}
diff --git a/Engine/Plugins/Experimental/AnimNext/Source/AnimNext/Private/Graph/RigUnit_AnimNextAnimSequence.h b/Engine/Plugins/Experimental/AnimNext/Source/AnimNext/Private/Graph/RigUnit_AnimNextAnimSequence.h
new file mode 100644
--- /dev/null
+++ b/Engine/Plugins/Experimental/AnimNext/Source/AnimNext/Private/Graph/RigUnit_AnimNextAnimSequence.h
@@ -0,0 +1,92 @@
+// Copyright Epic Games, Inc. All Rights Reserved.
+
+#pragma once
+
+#include "Graph/AnimNextExecuteContext.h"
+#include "Graph/AnimNext_LODPose.h"
+#include "Graph/AnimGraph/AnimNode_AnimNextGraph.h"	// TEST - until we can allocate per-node state again
+#include "RigUnit_AnimNextAnimSequence.generated.h"
+
+class UAnimSequenceBase;
+
+USTRUCT(BlueprintType, meta = (DisplayName = "AnimSequence"))
+struct FAnimNextGraph_AnimSequence
+{
+	GENERATED_BODY()
+
+	FAnimNextGraph_AnimSequence() = default;
+
+	explicit FAnimNextGraph_AnimSequence(const UAnimSequenceBase* InAnimSequence)
+		: AnimSequence(InAnimSequence)
+	{
+	}
+
+	const UAnimSequenceBase* AnimSequence = nullptr;
+};
+
+struct FAnimNextUnitContext;
+class UAnimNextGraph;
+class UAnimSequenceBase;
+
+
+// --- FRigUnit_AnimNext_AnimSequenceAsset ---
+
+/** Unit for getting an asset reference we can pass to a player or data extractor */
+USTRUCT(meta = (DisplayName = "Anim Sequence Asset", Category = "Animation", TitleColor = "1 0 0", NodeColor = "1 1 1"))
+struct FRigUnit_AnimNext_AnimSequenceAsset : public FRigUnit_AnimNextBase
+{
+	GENERATED_BODY()
+
+	RIGVM_METHOD()
+	void Execute();
+
+protected:
+	// TODO : Enable once we have RigVM object support enabled
+	//UPROPERTY(EditAnywhere, Category = "Asset", meta = (Input, EditInline, PinHiddenByDefault, DisallowedClasses = "/Script/Engine.AnimMontage"))
+	//TObjectPtr<UAnimSequenceBase> AnimSequence = nullptr;
+
+	UPROPERTY(EditAnywhere, Category = "Animation", DisplayName = "AnimSequence", meta = (Output))
+	FAnimNextGraph_AnimSequence Sequence;
+};
+
+// --- FRigUnit_AnimNext_AnimSequencePlayer --- 
+
+USTRUCT(BlueprintType)
+struct FAnimSequenceParameters
+{
+	GENERATED_BODY()
+
+	UPROPERTY(EditAnywhere, Category = "Anim Sequence")
+	float PlayRate = 1.0f;
+
+	UPROPERTY(EditAnywhere, Category = "Anim Sequence")
+	float StartPosition = 0.0f;
+
+	UPROPERTY(EditAnywhere, Category = "Anim Sequence")
+	bool bLoop = 0.0f;
+};
+
+/** Unit for getting a pose from an animation sequence */
+USTRUCT(meta = (DisplayName = "AnimSequencePlayer", Category = "Animation", TitleColor = "1 0 0", NodeColor = "1 1 1"))
+struct FRigUnit_AnimNext_AnimSequencePlayer : public FRigUnit_AnimNextBase
+{
+	GENERATED_BODY()
+
+	RIGVM_METHOD()
+	void Initialize();
+
+	RIGVM_METHOD()
+	void Execute();
+
+protected:
+	UPROPERTY(EditAnywhere, Category = "Anim Interface", meta = (Input))
+	FAnimSequenceParameters Parameters;
+
+	// The animation sequence asset to play
+	UPROPERTY(EditAnywhere, Category = "Animation", DisplayName = "AnimSequence", meta = (Input))
+	FAnimNextGraph_AnimSequence Sequence;
+
+	// The output pose generated
+	UPROPERTY(EditAnywhere, Transient, Category = "Animation", DisplayName = "Result", meta = (Output))
+	FAnimNextGraphLODPose LODPose;
+};
diff --git a/Engine/Plugins/Experimental/AnimNext/Source/AnimNext/Private/Graph/RigUnit_AnimNextDecoratorStack.cpp b/Engine/Plugins/Experimental/AnimNext/Source/AnimNext/Private/Graph/RigUnit_AnimNextDecoratorStack.cpp
deleted file mode 100644
diff --git a/Engine/Plugins/Experimental/AnimNext/Source/AnimNext/Private/Graph/RigUnit_AnimNextGraphRoot.cpp b/Engine/Plugins/Experimental/AnimNext/Source/AnimNext/Private/Graph/RigUnit_AnimNextGraphRoot.cpp
deleted file mode 100644
diff --git a/Engine/Plugins/Experimental/AnimNext/Source/AnimNext/Private/Module.cpp b/Engine/Plugins/Experimental/AnimNext/Source/AnimNext/Private/Module.cpp
--- a/Engine/Plugins/Experimental/AnimNext/Source/AnimNext/Private/Module.cpp
+++ b/Engine/Plugins/Experimental/AnimNext/Source/AnimNext/Private/Module.cpp
@@ -20,8 +20,7 @@ public:
 	{
 		static UClass* const AllowedAssetTypes[] =
 		{
-			UAnimSequence::StaticClass(),
-			UScriptStruct::StaticClass(),
+			UAnimSequence::StaticClass()
 		};
 
 		FRigVMRegistry::Get().RegisterObjectTypes(AllowedAssetTypes);
@@ -29,14 +28,21 @@ public:
 		FDataRegistry::Init();
 		FDecoratorRegistry::Init();
 		FNodeTemplateRegistry::Init();
+
+		EnginePreExitHandle = FCoreDelegates::OnEnginePreExit.AddLambda([]()
+		{
+			FDataRegistry::Destroy();
+			FDecoratorRegistry::Destroy();
+			FNodeTemplateRegistry::Destroy();
+		});
 	}
 
 	virtual void ShutdownModule() override
 	{
-		FNodeTemplateRegistry::Destroy();
-		FDecoratorRegistry::Destroy();
-		FDataRegistry::Destroy();
+		FCoreDelegates::OnEnginePreExit.Remove(EnginePreExitHandle);
 	}
+
+	FDelegateHandle EnginePreExitHandle;
 };
 
 }
diff --git a/Engine/Plugins/Experimental/AnimNext/Source/AnimNext/Public/Context.h b/Engine/Plugins/Experimental/AnimNext/Source/AnimNext/Public/Context.h
--- a/Engine/Plugins/Experimental/AnimNext/Source/AnimNext/Public/Context.h
+++ b/Engine/Plugins/Experimental/AnimNext/Source/AnimNext/Public/Context.h
@@ -20,6 +20,10 @@ struct ANIMNEXT_API FContext
 public:
 	FContext();
 	explicit FContext(float InDeltaTime);
+	FContext(const FContext& Other) = delete;
+	FContext& operator=(const FContext&) = delete;
+	FContext(FContext&& Other) = delete;
+	FContext& operator=(FContext&& Other) = delete;
 
 	// Access the parameter stack, cached on construction
 	FParamStack& GetMutableParamStack() const { return Stack; }
@@ -28,9 +32,6 @@ public:
 	float GetDeltaTime() const { return DeltaTime; }
 
 private:
-	FContext(const FContext& Other) = delete;
-	FContext& operator=(const FContext&) = delete;
-
 	FParamStack& Stack;
 	float DeltaTime = 0.0f;
 };
diff --git a/Engine/Plugins/Experimental/AnimNext/Source/AnimNext/Public/DecoratorBase/Decorator.h b/Engine/Plugins/Experimental/AnimNext/Source/AnimNext/Public/DecoratorBase/Decorator.h
--- a/Engine/Plugins/Experimental/AnimNext/Source/AnimNext/Public/DecoratorBase/Decorator.h
+++ b/Engine/Plugins/Experimental/AnimNext/Source/AnimNext/Public/DecoratorBase/Decorator.h
@@ -19,7 +19,7 @@ class FArchive;
 #define DECLARE_ANIM_DECORATOR(DecoratorName, DecoratorNameHash, SuperDecoratorName) \
 	using DecoratorSuper = SuperDecoratorName; \
 	/* FDecorator impl */ \
-	static constexpr UE::AnimNext::FDecoratorUID DecoratorUID = UE::AnimNext::FDecoratorUID(DecoratorNameHash, TEXT(#DecoratorName)); \
+	static constexpr UE::AnimNext::FDecoratorUID DecoratorUID = UE::AnimNext::FDecoratorUID(TEXT(#DecoratorName), DecoratorNameHash); \
 	virtual UE::AnimNext::FDecoratorUID GetDecoratorUID() const override { return DecoratorUID; } \
 	static const UE::AnimNext::FDecoratorMemoryLayout DecoratorMemoryDescription; \
 	virtual UE::AnimNext::FDecoratorMemoryLayout GetDecoratorMemoryDescription() const override { return DecoratorMemoryDescription; } \
@@ -33,7 +33,7 @@ class FArchive;
 #define DECLARE_ABSTRACT_ANIM_DECORATOR(DecoratorName, DecoratorNameHash, SuperDecoratorName) \
 	using DecoratorSuper = SuperDecoratorName; \
 	/* FDecorator impl */ \
-	static constexpr UE::AnimNext::FDecoratorUID DecoratorUID = UE::AnimNext::FDecoratorUID(DecoratorNameHash, TEXT(#DecoratorName)); \
+	static constexpr UE::AnimNext::FDecoratorUID DecoratorUID = UE::AnimNext::FDecoratorUID(TEXT(#DecoratorName), DecoratorNameHash); \
 	virtual UE::AnimNext::FDecoratorUID GetDecoratorUID() const override { return DecoratorUID; }
 
 // In the decorator cpp, these three macros implement the base functionality
@@ -152,7 +152,7 @@ namespace UE::AnimNext
 		// The globally unique UID for this decorator
 		// Derived types will have their own DecoratorUID member that hides/aliases/shadows this one
 		// @see DECLARE_ANIM_DECORATOR
-		static constexpr FDecoratorUID DecoratorUID = FDecoratorUID(0x4fc735a2, TEXT("FDecorator"));
+		static constexpr FDecoratorUID DecoratorUID = FDecoratorUID(TEXT("FDecorator"), 0x4fc735a2);
 
 		// Returns the globally unique UID for this decorator
 		virtual FDecoratorUID GetDecoratorUID() const { return DecoratorUID; };
@@ -205,7 +205,7 @@ namespace UE::AnimNext
 		// derived type using UE reflection.
 		// Decorators can override this function to control how editor only properties are coerced into the runtime shared data
 		// instance.
-		virtual void SaveDecoratorSharedData(FDecoratorWriter& Writer, const TFunction<FString(const FString& PropertyName)>& GetDecoratorProperty, FAnimNextDecoratorSharedData& OutSharedData) const;
+		virtual void SaveDecoratorSharedData(FDecoratorWriter& Writer, const TMap<FString, FString>& Properties, FAnimNextDecoratorSharedData& OutSharedData) const;
 #endif
 	};
 
diff --git a/Engine/Plugins/Experimental/AnimNext/Source/AnimNext/Public/DecoratorBase/DecoratorHandle.h b/Engine/Plugins/Experimental/AnimNext/Source/AnimNext/Public/DecoratorBase/DecoratorHandle.h
--- a/Engine/Plugins/Experimental/AnimNext/Source/AnimNext/Public/DecoratorBase/DecoratorHandle.h
+++ b/Engine/Plugins/Experimental/AnimNext/Source/AnimNext/Public/DecoratorBase/DecoratorHandle.h
@@ -12,59 +12,43 @@
  * Decorator Handle
  * A decorator handle represents a reference to a specific decorator instance in the shared/read-only portion
  * of a sub-graph. It points to a FNodeDescription when resolved.
- * 
- * Internally, it contains a node handle and a decorator index.
- * Decorator handles can only be used within an AnimNext graph serialized with FDecoratorWriter.
- * 
- * @see FNodeDescription, FNodeHandle
+ * @see FNodeDescription
  */
-USTRUCT(BlueprintType)
-struct FAnimNextDecoratorHandle final
+USTRUCT()
+struct FAnimNextDecoratorHandle
 {
 	GENERATED_BODY()
 
 	// Creates an invalid decorator handle
-	constexpr FAnimNextDecoratorHandle() noexcept
-		: PackedDecoratorIndexAndNodeHandle(UE::AnimNext::FNodeHandle::INVALID_NODE_HANDLE_RAW_VALUE)
+	FAnimNextDecoratorHandle()
+		: PackedDecoratorIndexAndNodeSharedOffset(INVALID_SHARED_OFFSET_VALUE)
 	{}
 
 	// Creates a decorator handle pointing to the first decorator of the specified node
-	explicit constexpr FAnimNextDecoratorHandle(UE::AnimNext::FNodeHandle NodeHandle) noexcept
-		: PackedDecoratorIndexAndNodeHandle(NodeHandle.GetPackedValue() & NODE_HANDLE_MASK)
+	explicit FAnimNextDecoratorHandle(UE::AnimNext::FNodeHandle NodeHandle)
+		: PackedDecoratorIndexAndNodeSharedOffset(NodeHandle.GetSharedOffset())
 	{}
 
 	// Creates a decorator handle pointing to the specified decorator on the specified node
-	FAnimNextDecoratorHandle(UE::AnimNext::FNodeHandle NodeHandle, uint32 DecoratorIndex)
-		: PackedDecoratorIndexAndNodeHandle((NodeHandle.GetPackedValue() & NODE_HANDLE_MASK) | (DecoratorIndex << 24))
-	{
-		check(DecoratorIndex <= MAX_uint8);	// Make sure we don't truncate
-	}
+	FAnimNextDecoratorHandle(UE::AnimNext::FNodeHandle NodeHandle, uint32 DecoratorIndex_)
+		: PackedDecoratorIndexAndNodeSharedOffset(NodeHandle.GetSharedOffset() | (DecoratorIndex_ << 24))
+	{}
 
 	// Returns true if this decorator handle is valid, false otherwise
-	constexpr bool IsValid() const noexcept { return GetNodeHandle().IsValid(); }
+	bool IsValid() const { return (PackedDecoratorIndexAndNodeSharedOffset & SHARED_OFFSET_MASK) != INVALID_SHARED_OFFSET_VALUE; }
 
 	// Returns the decorator index
-	constexpr uint32 GetDecoratorIndex() const noexcept { return PackedDecoratorIndexAndNodeHandle >> DECORATOR_SHIFT; }
-
-	// Returns a handle to the node referenced
-	constexpr UE::AnimNext::FNodeHandle GetNodeHandle() const noexcept { return UE::AnimNext::FNodeHandle::FromPackedValue(PackedDecoratorIndexAndNodeHandle); }
+	uint32 GetDecoratorIndex() const { return PackedDecoratorIndexAndNodeSharedOffset >> DECORATOR_SHIFT; }
 
-	ANIMNEXT_API bool Serialize(class FArchive& Ar);
+	// Returns a handle to the node in the shared data segment
+	UE::AnimNext::FNodeHandle GetNodeHandle() const { return UE::AnimNext::FNodeHandle(PackedDecoratorIndexAndNodeSharedOffset & SHARED_OFFSET_MASK); }
 
 private:
-	// Bottom 24 bits are used by the node handle while the top 8 bits by the decorator index
+	// Bottom 24 bits are used by the node shared offset while the top 8 bits by the decorator index
 	static constexpr uint32 DECORATOR_SHIFT = 24;
-	static constexpr uint32 NODE_HANDLE_MASK = ~0u >> (32 - DECORATOR_SHIFT);
+	static constexpr uint32 SHARED_OFFSET_MASK = ~0u >> (32 - DECORATOR_SHIFT);
+	static constexpr uint32 INVALID_SHARED_OFFSET_VALUE = SHARED_OFFSET_MASK;
 
 	UPROPERTY()
-	uint32		PackedDecoratorIndexAndNodeHandle;
-};
-
-template<>
-struct TStructOpsTypeTraits<FAnimNextDecoratorHandle> : public TStructOpsTypeTraitsBase2<FAnimNextDecoratorHandle>
-{
-	enum
-	{
-		WithSerializer = true,
-	};
+	uint32		PackedDecoratorIndexAndNodeSharedOffset;
 };
diff --git a/Engine/Plugins/Experimental/AnimNext/Source/AnimNext/Public/DecoratorBase/DecoratorInterfaceUID.h b/Engine/Plugins/Experimental/AnimNext/Source/AnimNext/Public/DecoratorBase/DecoratorInterfaceUID.h
--- a/Engine/Plugins/Experimental/AnimNext/Source/AnimNext/Public/DecoratorBase/DecoratorInterfaceUID.h
+++ b/Engine/Plugins/Experimental/AnimNext/Source/AnimNext/Public/DecoratorBase/DecoratorInterfaceUID.h
@@ -6,9 +6,6 @@
 
 namespace UE::AnimNext
 {
-	// Type alias for a raw decorator UID, not typesafe
-	using FDecoratorInterfaceUIDRaw = uint32;
-
 	/**
 	 * FDecoratorInterfaceUID
 	 *
@@ -22,17 +19,17 @@ namespace UE::AnimNext
 	{
 		// Constructs an invalid UID
 		constexpr FDecoratorInterfaceUID()
-			: UID(INVALID_UID)
+			: UID(0)
 #if !UE_BUILD_SHIPPING
 			, InterfaceName(TEXT("<Invalid Interface UID>"))
 #endif
 		{}
 
 		// Constructs an interface UID
-		explicit constexpr FDecoratorInterfaceUID(FDecoratorInterfaceUIDRaw InUID, const TCHAR* InInterfaceName = TEXT("<Unknown Decorator Interface Name>"))
-			: UID(InUID)
+		constexpr FDecoratorInterfaceUID(const TCHAR* InterfaceName_, uint32 UID_)
+			: UID(UID_)
 #if !UE_BUILD_SHIPPING
-			, InterfaceName(InInterfaceName)
+			, InterfaceName(InterfaceName_)
 #endif
 		{
 		}
@@ -43,22 +40,20 @@ namespace UE::AnimNext
 #endif
 
 		// Returns the interface global UID
-		constexpr FDecoratorInterfaceUIDRaw GetUID() const { return UID; }
+		constexpr uint32 GetUID() const { return UID; }
 
 		// Returns whether this UID is valid or not
-		constexpr bool IsValid() const { return UID != INVALID_UID; }
+		constexpr bool IsValid() const { return UID != 0; }
 
-	private:
-		static constexpr FDecoratorInterfaceUIDRaw INVALID_UID = 0;
+		// Compares for equality and inequality
+		constexpr bool operator==(FDecoratorInterfaceUID RHS) const { return UID == RHS.UID; }
+		constexpr bool operator!=(FDecoratorInterfaceUID RHS) const { return UID != RHS.UID; }
 
-		FDecoratorInterfaceUIDRaw	UID;
+	private:
+		uint32	UID;
 
 #if !UE_BUILD_SHIPPING
-		const TCHAR*				InterfaceName;
+		const TCHAR* InterfaceName;
 #endif
 	};
-
-	// Compares for equality and inequality
-	constexpr bool operator==(FDecoratorInterfaceUID LHS, FDecoratorInterfaceUID RHS) { return LHS.GetUID() == RHS.GetUID(); }
-	constexpr bool operator!=(FDecoratorInterfaceUID LHS, FDecoratorInterfaceUID RHS) { return LHS.GetUID() != RHS.GetUID(); }
 }
diff --git a/Engine/Plugins/Experimental/AnimNext/Source/AnimNext/Public/DecoratorBase/DecoratorPtr.h b/Engine/Plugins/Experimental/AnimNext/Source/AnimNext/Public/DecoratorBase/DecoratorPtr.h
--- a/Engine/Plugins/Experimental/AnimNext/Source/AnimNext/Public/DecoratorBase/DecoratorPtr.h
+++ b/Engine/Plugins/Experimental/AnimNext/Source/AnimNext/Public/DecoratorBase/DecoratorPtr.h
@@ -100,8 +100,12 @@ namespace UE::AnimNext
 			, DecoratorIndex(DecoratorPtr.DecoratorIndex)
 		{}
 
-		// Constructs a weak pointer handle to the provided instance
-		FWeakDecoratorPtr(FNodeInstance* NodeInstance_, uint32 DecoratorIndex_);
+		FWeakDecoratorPtr(FNodeInstance* NodeInstance_, uint32 DecoratorIndex_)
+			: NodeInstance(NodeInstance_)
+			, DecoratorIndex(static_cast<uint8>(DecoratorIndex_))
+		{
+			check(DecoratorIndex == DecoratorIndex_);	// Make sure we didn't truncate
+		}
 
 		// Returns a pointer to the node instance
 		constexpr FNodeInstance* GetNodeInstance() const noexcept { return NodeInstance; }
@@ -123,7 +127,7 @@ namespace UE::AnimNext
 
 	private:
 		FNodeInstance*	NodeInstance = nullptr;
-		uint32			DecoratorIndex = 0;			// Only need 8 bits, but use 32 since we have ample padding, avoids truncation
+		uint8			DecoratorIndex = 0;
 	};
 
 	//////////////////////////////////////////////////////////////////////////
diff --git a/Engine/Plugins/Experimental/AnimNext/Source/AnimNext/Public/DecoratorBase/DecoratorUID.h b/Engine/Plugins/Experimental/AnimNext/Source/AnimNext/Public/DecoratorBase/DecoratorUID.h
--- a/Engine/Plugins/Experimental/AnimNext/Source/AnimNext/Public/DecoratorBase/DecoratorUID.h
+++ b/Engine/Plugins/Experimental/AnimNext/Source/AnimNext/Public/DecoratorBase/DecoratorUID.h
@@ -6,9 +6,6 @@
 
 namespace UE::AnimNext
 {
-	// Type alias for a raw decorator UID, not typesafe
-	using FDecoratorUIDRaw = uint32;
-
 	/**
 	 * FDecoratorUID
 	 *
@@ -22,18 +19,17 @@ namespace UE::AnimNext
 	{
 		// Constructs an invalid UID
 		constexpr FDecoratorUID()
-			: UID(INVALID_UID)
+			: UID(0)
 #if !UE_BUILD_SHIPPING
 			, DecoratorName(TEXT("<Invalid decorator UID>"))
 #endif
-		{
-		}
+		{}
 
 		// Constructs a decorator UID
-		explicit constexpr FDecoratorUID(FDecoratorUIDRaw InUID, const TCHAR* InDecoratorName = TEXT("<Unknown Decorator Name>"))
-			: UID(InUID)
+		constexpr FDecoratorUID(const TCHAR* DecoratorName_, uint32 UID_)
+			: UID(UID_)
 #if !UE_BUILD_SHIPPING
-			, DecoratorName(InDecoratorName)
+			, DecoratorName(DecoratorName_)
 #endif
 		{
 		}
@@ -44,26 +40,24 @@ namespace UE::AnimNext
 #endif
 
 		// Returns the decorator global UID
-		constexpr FDecoratorUIDRaw GetUID() const { return UID; }
+		constexpr uint32 GetUID() const { return UID; }
 
 		// Returns whether this UID is valid or not
-		constexpr bool IsValid() const { return UID != INVALID_UID; }
+		constexpr bool IsValid() const { return UID != 0; }
 
 	private:
-		static constexpr FDecoratorUIDRaw INVALID_UID = 0;
-
-		FDecoratorUIDRaw	UID;
+		uint32	UID;
 
 #if !UE_BUILD_SHIPPING
-		const TCHAR*		DecoratorName;
+		const TCHAR* DecoratorName;
 #endif
 	};
 
 	// Compares for equality and inequality
 	constexpr bool operator==(FDecoratorUID LHS, FDecoratorUID RHS) { return LHS.GetUID() == RHS.GetUID(); }
 	constexpr bool operator!=(FDecoratorUID LHS, FDecoratorUID RHS) { return LHS.GetUID() != RHS.GetUID(); }
-	constexpr bool operator==(FDecoratorUID LHS, FDecoratorUIDRaw RHS) { return LHS.GetUID() == RHS; }
-	constexpr bool operator!=(FDecoratorUID LHS, FDecoratorUIDRaw RHS) { return LHS.GetUID() != RHS; }
-	constexpr bool operator==(FDecoratorUIDRaw LHS, FDecoratorUID RHS) { return LHS == RHS.GetUID(); }
-	constexpr bool operator!=(FDecoratorUIDRaw LHS, FDecoratorUID RHS) { return LHS != RHS.GetUID(); }
+	constexpr bool operator==(FDecoratorUID LHS, uint32 RHS) { return LHS.GetUID() == RHS; }
+	constexpr bool operator!=(FDecoratorUID LHS, uint32 RHS) { return LHS.GetUID() != RHS; }
+	constexpr bool operator==(uint32 LHS, FDecoratorUID RHS) { return LHS == RHS.GetUID(); }
+	constexpr bool operator!=(uint32 LHS, FDecoratorUID RHS) { return LHS != RHS.GetUID(); }
 }
diff --git a/Engine/Plugins/Experimental/AnimNext/Source/AnimNext/Public/DecoratorBase/EntryPointHandle.h b/Engine/Plugins/Experimental/AnimNext/Source/AnimNext/Public/DecoratorBase/EntryPointHandle.h
deleted file mode 100644
diff --git a/Engine/Plugins/Experimental/AnimNext/Source/AnimNext/Public/DecoratorBase/IDecoratorInterface.h b/Engine/Plugins/Experimental/AnimNext/Source/AnimNext/Public/DecoratorBase/IDecoratorInterface.h
--- a/Engine/Plugins/Experimental/AnimNext/Source/AnimNext/Public/DecoratorBase/IDecoratorInterface.h
+++ b/Engine/Plugins/Experimental/AnimNext/Source/AnimNext/Public/DecoratorBase/IDecoratorInterface.h
@@ -9,7 +9,7 @@
 // Helper macros
 #define DECLARE_ANIM_DECORATOR_INTERFACE(InterfaceName, InterfaceNameHash) \
 	/* Globally unique UID for this interface */ \
-	static constexpr UE::AnimNext::FDecoratorInterfaceUID InterfaceUID = UE::AnimNext::FDecoratorInterfaceUID(InterfaceNameHash, TEXT(#InterfaceName)); \
+	static constexpr UE::AnimNext::FDecoratorInterfaceUID InterfaceUID = UE::AnimNext::FDecoratorInterfaceUID(TEXT(#InterfaceName), InterfaceNameHash); \
 	virtual UE::AnimNext::FDecoratorInterfaceUID GetInterfaceUID() const override { return InterfaceUID; }
 
 namespace UE::AnimNext
@@ -28,7 +28,7 @@ namespace UE::AnimNext
 		// The globally unique UID for this interface
 		// Derived types will have their own InterfaceUID member that hides/aliases/shadows this one
 		// @see DECLARE_ANIM_DECORATOR_INTERFACE
-		static constexpr FDecoratorInterfaceUID InterfaceUID = FDecoratorInterfaceUID(0x402f6df4, TEXT("IDecoratorInterface"));
+		static constexpr FDecoratorInterfaceUID InterfaceUID = FDecoratorInterfaceUID(TEXT("IDecoratorInterface"), 0x402f6df4);
 
 		// Returns the globally unique UID for this interface
 		virtual FDecoratorInterfaceUID GetInterfaceUID() const { return InterfaceUID; };
diff --git a/Engine/Plugins/Experimental/AnimNext/Source/AnimNext/Public/DecoratorBase/NodeDescription.h b/Engine/Plugins/Experimental/AnimNext/Source/AnimNext/Public/DecoratorBase/NodeDescription.h
--- a/Engine/Plugins/Experimental/AnimNext/Source/AnimNext/Public/DecoratorBase/NodeDescription.h
+++ b/Engine/Plugins/Experimental/AnimNext/Source/AnimNext/Public/DecoratorBase/NodeDescription.h
@@ -3,9 +3,10 @@
 #pragma once
 
 #include "CoreMinimal.h"
-#include "DecoratorBase/NodeID.h"
 #include "DecoratorBase/NodeTemplateRegistryHandle.h"
 
+#include <limits>
+
 class FArchive;
 
 namespace UE::AnimNext
@@ -25,17 +26,22 @@ namespace UE::AnimNext
 	 */
 	struct alignas(alignof(uint32)) FNodeDescription
 	{
-		// Largest allowed size for a node description and the shared data of its decorators
-		// We use unsigned 16 bit offsets within the NodeTemplate/DecoratorTemplate
-		static constexpr uint32 MAXIMUM_NODE_SHARED_DATA_SIZE = 64 * 1024;
-
-		FNodeDescription(FNodeID InNodeID, FNodeTemplateRegistryHandle InTemplateHandle)
-			: NodeID(InNodeID)
-			, TemplateHandle(InTemplateHandle)
+		// Largest allowed size for a node description
+		// Graphs are currently limited to 64 KB because we use 16 bit offsets within our node/decorator handles
+		// We have no particular limit on the node size besides the graph limit (for now)
+		static constexpr uint32 MAXIMUM_SIZE = 64 * 1024;
+
+		// The maximum number of nodes allowed within a single graph
+		// We use 16 bits to represent node UIDs
+		static constexpr uint32 MAXIMUM_COUNT = std::numeric_limits<uint16>::max();
+
+		FNodeDescription(uint16 UID_, FNodeTemplateRegistryHandle TemplateHandle_)
+			: UID(UID_)
+			, TemplateHandle(TemplateHandle_)
 		{}
 
 		// Returns the node UID, unique to the owning sub-graph
-		FNodeID GetUID() const { return NodeID; }
+		uint32 GetUID() const { return UID; }
 
 		// Returns the handle of the node's template in the node template registry
 		FNodeTemplateRegistryHandle GetTemplateHandle() const { return TemplateHandle; }
@@ -44,7 +50,7 @@ namespace UE::AnimNext
 		ANIMNEXT_API void Serialize(FArchive& Ar);
 
 	private:
-		FNodeID							NodeID;				// assigned during export/cook, unique to current sub-graph
+		uint16							UID;				// assigned during export/cook, unique to current sub-graph
 		FNodeTemplateRegistryHandle		TemplateHandle;		// offset of the node template within the global list
 
 		// Followed by a list of [FAnimNextDecoratorSharedData] instances and optional padding
diff --git a/Engine/Plugins/Experimental/AnimNext/Source/AnimNext/Public/DecoratorBase/NodeHandle.h b/Engine/Plugins/Experimental/AnimNext/Source/AnimNext/Public/DecoratorBase/NodeHandle.h
--- a/Engine/Plugins/Experimental/AnimNext/Source/AnimNext/Public/DecoratorBase/NodeHandle.h
+++ b/Engine/Plugins/Experimental/AnimNext/Source/AnimNext/Public/DecoratorBase/NodeHandle.h
@@ -3,11 +3,8 @@
 #pragma once
 
 #include "CoreMinimal.h"
-#include "DecoratorBase/NodeID.h"
-#include "Serialization/Archive.h"
 
 struct FAnimNextDecoratorHandle;
-struct FAnimNextEntryPointHandle;
 
 namespace UE::AnimNext
 {
@@ -16,81 +13,37 @@ namespace UE::AnimNext
 	 * A node handle represents a reference to a specific node instance in the shared/read-only portion
 	 * of a sub-graph. It points to a FNodeDescription when resolved.
 	 * 
-	 * Internally, it contains either a shared data offset into the parent graph (at runtime) or a FNodeID (during compilation and on disk).
-	 * Shared offsets must be computed at runtime and as such we use the FNodeID to resolve into our final offset on load.
-	 * We only use the bottom 24 bits. If the LSB is 0, the packed value represents a shared offset (due to alignment) but if the
-	 * LSB is 1, the packed value represents a node ID shifted left by 1.
-	 * 
-	 * @see FNodeDescription, FNodeID
+	 * @see FNodeDescription
 	 */
-	struct FNodeHandle final
+	struct FNodeHandle
 	{
 		// Creates an invalid node handle
-		constexpr FNodeHandle() noexcept
-			: PackedSharedOffsetOrID(INVALID_NODE_HANDLE_RAW_VALUE)
+		FNodeHandle()
+			: SharedOffset(INVALID_SHARED_OFFSET_VALUE)
 		{}
 
 		// Creates a node handle for the specified shared offset
-		static constexpr FNodeHandle FromSharedOffset(uint32 SharedOffset) noexcept
-		{
-			if (SharedOffset > NODE_HANDLE_MASK)
-			{
-				return FNodeHandle();	// The shared offset is too large, the handle isn't valid
-			}
-
-			return FNodeHandle(SharedOffset);
-		}
-
-		// Creates a node handle from a node ID
-		static constexpr FNodeHandle FromNodeID(FNodeID NodeID) noexcept
-		{
-			if (!NodeID.IsValid())
-			{
-				return FNodeHandle();	// The NodeID isn't valid, the handle isn't valid
-			}
-
-			return FNodeHandle((NodeID.GetRawValue() << NODE_ID_SHIFT) | NODE_ID_MARKER_BIT_MASK);
-		}
+		explicit FNodeHandle(uint32 SharedOffset_)
+			: SharedOffset(SharedOffset_)
+		{}
 
 		// Returns true if this node handle is valid, false otherwise
-		constexpr bool IsValid() const noexcept { return (PackedSharedOffsetOrID & NODE_HANDLE_MASK) != INVALID_NODE_HANDLE_RAW_VALUE; }
-
-		// Returns whether this node handle encodes a shared offset or not
-		constexpr bool IsSharedOffset() const noexcept { return IsValid() && (PackedSharedOffsetOrID & NODE_ID_MARKER_BIT_MASK) == 0; }
-
-		// Returns whether this node handle encodes a node ID or not
-		constexpr bool IsNodeID() const noexcept { return IsValid() && (PackedSharedOffsetOrID & NODE_ID_MARKER_BIT_MASK) != 0; }
+		bool IsValid() const { return SharedOffset != INVALID_SHARED_OFFSET_VALUE; }
 
 		// Returns the offset into the shared data relative to the root of the owning sub-graph
-		uint32 GetSharedOffset() const { check(IsSharedOffset()); return PackedSharedOffsetOrID & NODE_HANDLE_MASK; }
-
-		// Returns the node ID encoded by the handle
-		FNodeID GetNodeID() const { check(IsNodeID()); return FNodeID::FromRawValue((PackedSharedOffsetOrID & NODE_HANDLE_MASK) >> NODE_ID_SHIFT); }
+		uint32 GetSharedOffset() const { return SharedOffset; }
 
 		// Equality and inequality tests
-		constexpr bool operator==(FNodeHandle RHS) const noexcept { return (PackedSharedOffsetOrID & NODE_HANDLE_MASK) == (RHS.PackedSharedOffsetOrID & NODE_HANDLE_MASK); }
-		constexpr bool operator!=(FNodeHandle RHS) const noexcept { return (PackedSharedOffsetOrID & NODE_HANDLE_MASK) != (RHS.PackedSharedOffsetOrID & NODE_HANDLE_MASK); }
+		bool operator==(FNodeHandle RHS) const { return SharedOffset == RHS.SharedOffset; }
+		bool operator!=(FNodeHandle RHS) const { return SharedOffset != RHS.SharedOffset; }
 
 	private:
-		explicit constexpr FNodeHandle(uint32 PackedValue) noexcept
-			: PackedSharedOffsetOrID(PackedValue)
-		{}
-
-		// Returns the packed raw value
-		constexpr uint32 GetPackedValue() const { return PackedSharedOffsetOrID; }
-
-		// Creates a node handle from a raw packed value
-		static constexpr FNodeHandle FromPackedValue(uint32 PackedValue) noexcept { return FNodeHandle(PackedValue); }
-
-		// We pack the shared offset on 24 bits, the top 8 bits are unused and they can have any value (zero by default)
-		static constexpr uint32 INVALID_NODE_HANDLE_RAW_VALUE = ~0u >> (32 - 24);	// 0x00FFFFFF
-		static constexpr uint32 NODE_HANDLE_MASK = INVALID_NODE_HANDLE_RAW_VALUE;
-		static constexpr uint32 NODE_ID_MARKER_BIT_MASK = 0x01u;					// When valid, a LSB of 1 means we encode a node ID, 0 means a shared offset
-		static constexpr uint32 NODE_ID_SHIFT = 1;									// Our LSB is reserved as a marker, the node ID is shifted left by 1
+		// We pack the shared offset on 24 bits, the top 8 bits are truncated
+		static constexpr uint32 INVALID_SHARED_OFFSET_VALUE = ~0u >> (32 - 24);
 
-		uint32	PackedSharedOffsetOrID;
+		uint32	Padding : 8;			// unused padding
+		uint32	SharedOffset : 24;		// relative to root of sub-graph
 
 		friend FAnimNextDecoratorHandle;
-		friend FAnimNextEntryPointHandle;
 	};
 }
diff --git a/Engine/Plugins/Experimental/AnimNext/Source/AnimNext/Public/DecoratorBase/NodeID.h b/Engine/Plugins/Experimental/AnimNext/Source/AnimNext/Public/DecoratorBase/NodeID.h
deleted file mode 100644
diff --git a/Engine/Plugins/Experimental/AnimNext/Source/AnimNext/Public/DecoratorBase/NodeInstance.h b/Engine/Plugins/Experimental/AnimNext/Source/AnimNext/Public/DecoratorBase/NodeInstance.h
--- a/Engine/Plugins/Experimental/AnimNext/Source/AnimNext/Public/DecoratorBase/NodeInstance.h
+++ b/Engine/Plugins/Experimental/AnimNext/Source/AnimNext/Public/DecoratorBase/NodeInstance.h
@@ -18,10 +18,6 @@ namespace UE::AnimNext
 	 */
 	struct ANIMNEXT_API FNodeInstance
 	{
-		// Largest allowed size for a node instance and the instance data of its decorators
-		// We use unsigned 16 bit offsets within the NodeTemplate/DecoratorTemplate
-		static constexpr uint32 MAXIMUM_NODE_INSTANCE_DATA_SIZE = 64 * 1024;
-
 		FNodeInstance(const FNodeInstance&) = delete;
 		FNodeInstance& operator=(const FNodeInstance&) = delete;
 
@@ -41,7 +37,7 @@ namespace UE::AnimNext
 		void AddReference() { ReferenceCount++; }
 		void ReleaseReference();
 
-		uint32		ReferenceCount;		// how many non-weak FDecoratorPtr handles point to us, not thread safe
+		uint32		ReferenceCount;		// how many non-weak FNodePtr and FDecoratorPtr handles point to us, not thread safe
 		FNodeHandle	NodeHandle;			// relative to root of sub-graph, should this be a pointer?
 
 		// Followed by a list of [FDecoratorInstanceData] instances and optional padding
diff --git a/Engine/Plugins/Experimental/AnimNext/Source/AnimNext/Public/DecoratorInterfaces/IUpdate.h b/Engine/Plugins/Experimental/AnimNext/Source/AnimNext/Public/DecoratorInterfaces/IUpdate.h
--- a/Engine/Plugins/Experimental/AnimNext/Source/AnimNext/Public/DecoratorInterfaces/IUpdate.h
+++ b/Engine/Plugins/Experimental/AnimNext/Source/AnimNext/Public/DecoratorInterfaces/IUpdate.h
@@ -15,12 +15,8 @@ namespace UE::AnimNext
 	 */
 	struct FUpdateTraversalContext : ITraversalContext
 	{
-		explicit FUpdateTraversalContext(float InDeltaTime);
-
-		float GetDeltaTime() const { return DeltaTime; }
-
-	private:
-		float DeltaTime;
+		double GetDeltaTime() const { return 1.0 / 30.0; }
+		double GetPlayRate() const { return 1.0; }
 	};
 
 	/**
@@ -79,5 +75,5 @@ namespace UE::AnimNext
 	 *
 	 * @see IUpdate::PreUpdate, IUpdate::PostUpdate, IHierarchy::GetChildren
 	 */
-	ANIMNEXT_API void UpdateGraph(FExecutionContext& Context, FWeakDecoratorPtr GraphRootPtr, float DeltaTime);
+	ANIMNEXT_API void UpdateGraph(FExecutionContext& Context, FWeakDecoratorPtr GraphRootPtr);
 }
diff --git a/Engine/Plugins/Experimental/AnimNext/Source/AnimNextEditor/AnimNextEditor.Build.cs b/Engine/Plugins/Experimental/AnimNext/Source/AnimNextEditor/AnimNextEditor.Build.cs
--- a/Engine/Plugins/Experimental/AnimNext/Source/AnimNextEditor/AnimNextEditor.Build.cs
+++ b/Engine/Plugins/Experimental/AnimNext/Source/AnimNextEditor/AnimNextEditor.Build.cs
@@ -22,7 +22,6 @@ namespace UnrealBuildTool.Rules
 					"InputCore",
 					"PropertyEditor",
 					"RigVM",
-					"RigVMEditor",
 					"RigVMDeveloper",
 					"ControlRig",
 					"ControlRigDeveloper",
diff --git a/Engine/Plugins/Experimental/AnimNext/Source/AnimNextEditor/Private/AnimNextEditorModule.cpp b/Engine/Plugins/Experimental/AnimNext/Source/AnimNextEditor/Private/AnimNextEditorModule.cpp
--- a/Engine/Plugins/Experimental/AnimNext/Source/AnimNextEditor/Private/AnimNextEditorModule.cpp
+++ b/Engine/Plugins/Experimental/AnimNext/Source/AnimNextEditor/Private/AnimNextEditorModule.cpp
@@ -4,7 +4,6 @@
 #include "AssetToolsModule.h"
 #include "IAssetTools.h"
 #include "Graph/AssetTypeActions.h"
-#include "Graph/AnimNextGraphPanelNodeFactory.h"
 #include "Param/ParamTypePropertyCustomization.h"
 #include "Param/ParameterPickerArgs.h"
 #include "Param/SParameterPicker.h"
@@ -27,9 +26,6 @@ class FModule : public IModule
 		PropertyModule.RegisterCustomPropertyTypeLayout(
 			"AnimNextParamType",
 			FOnGetPropertyTypeCustomizationInstance::CreateLambda([] { return MakeShared<FParamTypePropertyTypeCustomization>(); }));
-
-		AnimNextGraphPanelNodeFactory = MakeShared<FAnimNextGraphPanelNodeFactory>();
-		FEdGraphUtilities::RegisterVisualNodeFactory(AnimNextGraphPanelNodeFactory);
 	}
 
 	virtual void ShutdownModule() override
@@ -45,8 +41,6 @@ class FModule : public IModule
 			FPropertyEditorModule& PropertyModule = FModuleManager::LoadModuleChecked<FPropertyEditorModule>("PropertyEditor");
 			PropertyModule.UnregisterCustomPropertyTypeLayout("AnimNextParamType");
 		}
-
-		FEdGraphUtilities::UnregisterVisualNodeFactory(AnimNextGraphPanelNodeFactory);
 	}
 
 	virtual TSharedRef<SWidget> CreateParameterPicker(const FParameterPickerArgs& InArgs) override
@@ -56,9 +50,6 @@ class FModule : public IModule
 	}
 
 	TSharedPtr<FAssetTypeActions_AnimNextGraph> AssetTypeActions_AnimNextGraph;
-
-	/** Node factory for the AnimNext graph */
-	TSharedPtr<FAnimNextGraphPanelNodeFactory> AnimNextGraphPanelNodeFactory;
 };
 
 }
diff --git a/Engine/Plugins/Experimental/AnimNext/Source/AnimNextEditor/Private/Graph/AnimNextGraphEditor.cpp b/Engine/Plugins/Experimental/AnimNext/Source/AnimNextEditor/Private/Graph/AnimNextGraphEditor.cpp
--- a/Engine/Plugins/Experimental/AnimNext/Source/AnimNextEditor/Private/Graph/AnimNextGraphEditor.cpp
+++ b/Engine/Plugins/Experimental/AnimNext/Source/AnimNextEditor/Private/Graph/AnimNextGraphEditor.cpp
@@ -3,7 +3,6 @@
 #include "AnimNextGraphEditor.h"
 #include "GraphEditorMode.h"
 #include "Graph/AnimNextGraph.h"
-#include "Graph/AnimNextGraph_EdGraphNode.h"
 #include "Graph/AnimNextGraph_EditorData.h"
 #include "EdGraphNode_Comment.h"
 #include "SActionMenu.h"
@@ -63,9 +62,6 @@ void FGraphEditor::InitEditor(const EToolkitMode::Type InMode, const TSharedPtr<
 
 	// Open initial document
 	DocumentManager->OpenDocument(FTabPayload_UObject::Make(AnimNextGraph_EditorData->RootGraph), FDocumentTracker::EOpenDocumentCause::OpenNewDocument);
-
-	// Make sure editor graph matches the model
-	AnimNextGraph_EditorData->RebuildEdGraphFromModel();
 }
 
 void FGraphEditor::BindCommands()
@@ -245,8 +241,6 @@ void FGraphEditor::DeleteSelectedNodes()
 		FocusedGraphEd->ClearSelectionSet();
 	}
 
-	TMap<UAnimNextGraph_EdGraphNode*, URigVMNode*> NodesToRemove;
-
 	// Some nodes have sub-objects that are represented as other tabs.
 	// Close them here as a pre-pass before we remove their nodes. If the documents are left open they
 	// may reference dangling data and function incorrectly in cases such as FindBlueprintforNodeChecked
@@ -274,35 +268,26 @@ void FGraphEditor::DeleteSelectedNodes()
 				{
 					CloseAllDocumentsTab(Node);
 				}
-
-				if (UAnimNextGraph_EdGraphNode* AnimNextNode = Cast<UAnimNextGraph_EdGraphNode>(Node))
-				{
-					if (URigVMNode* ModelNode = GetFocusedVMController()->GetGraph()->FindNodeByName(*AnimNextNode->GetModelNodePath()))
-					{
-						NodesToRemove.Add(AnimNextNode, ModelNode);
-					}
-				}
 			}
 		}
 	}
 
-	if (NodesToRemove.IsEmpty())
-	{
-		return;
-	}
-
-	GetFocusedVMController()->OpenUndoBracket(TEXT("Delete selected nodes"));
-
-	TArray<URigVMNode*> ModelNodesToRemove;
-	for (auto& It : NodesToRemove)
+	// Now remove the selected nodes
+	for (FGraphPanelSelectionSet::TConstIterator NodeIt( SelectedNodes ); NodeIt; ++NodeIt)
 	{
-		URigVMNode* ModelNode = It.Value;
+		if (UEdGraphNode* Node = Cast<UEdGraphNode>(*NodeIt))
+		{
+			if (Node->CanUserDeleteNode())
+			{
+				if (Node->GetSubGraphs().Num() > 0)
+				{
+					DocumentManager->CleanInvalidTabs();
+				}
 
-		ModelNodesToRemove.Add(ModelNode);
+				FBlueprintEditorUtils::RemoveNode(nullptr, Node);
+			}
+		}
 	}
-
-	GetFocusedVMController()->RemoveNodes(ModelNodesToRemove, true);
-	GetFocusedVMController()->CloseUndoBracket();
 }
 
 bool FGraphEditor::CanDeleteSelectedNodes()
diff --git a/Engine/Plugins/Experimental/AnimNext/Source/AnimNextEditor/Private/Graph/AnimNextGraphPanelNodeFactory.cpp b/Engine/Plugins/Experimental/AnimNext/Source/AnimNextEditor/Private/Graph/AnimNextGraphPanelNodeFactory.cpp
deleted file mode 100644
diff --git a/Engine/Plugins/Experimental/AnimNext/Source/AnimNextEditor/Private/Graph/AnimNextGraphPanelNodeFactory.h b/Engine/Plugins/Experimental/AnimNext/Source/AnimNextEditor/Private/Graph/AnimNextGraphPanelNodeFactory.h
deleted file mode 100644
diff --git a/Engine/Plugins/Experimental/AnimNext/Source/AnimNextEditor/Private/Graph/GraphEditorSchemaActions.cpp b/Engine/Plugins/Experimental/AnimNext/Source/AnimNextEditor/Private/Graph/GraphEditorSchemaActions.cpp
--- a/Engine/Plugins/Experimental/AnimNext/Source/AnimNextEditor/Private/Graph/GraphEditorSchemaActions.cpp
+++ b/Engine/Plugins/Experimental/AnimNext/Source/AnimNextEditor/Private/Graph/GraphEditorSchemaActions.cpp
@@ -7,7 +7,6 @@
 #include "Graph/AnimNextGraph_EditorData.h"
 #include "RigVMModel/Nodes/RigVMUnitNode.h"
 #include "Settings/ControlRigSettings.h"
-#include "Units/RigUnit.h"
 
 UEdGraphNode* FAnimNextGraphSchemaAction_RigUnit::PerformAction(UEdGraph* ParentGraph, TArray<UEdGraphPin*>& FromPins, const FVector2D Location, bool bSelectNewNode)
 {
@@ -25,7 +24,7 @@ UEdGraphNode* FAnimNextGraphSchemaAction_RigUnit::PerformAction(UEdGraph* Parent
 		FRigVMUnitNodeCreatedContext& UnitNodeCreatedContext = Controller->GetUnitNodeCreatedContext();
 		FRigVMUnitNodeCreatedContext::FScope ReasonScope(UnitNodeCreatedContext, ERigVMNodeCreatedReason::NodeSpawner);
 
-		if (URigVMUnitNode* ModelNode = Controller->AddUnitNode(StructTemplate, FRigVMStruct::ExecuteName, Location, Name.ToString(), true, false))
+		if (URigVMUnitNode* ModelNode = Controller->AddUnitNode(StructTemplate, FRigUnit::GetMethodName(), Location, Name.ToString(), true, false))
 		{
 			NewNode = Cast<UAnimNextGraph_EdGraphNode>(EdGraph->FindNodeForModelNodeName(ModelNode->GetFName()));
 			check(NewNode);
diff --git a/Engine/Plugins/Experimental/AnimNext/Source/AnimNextEditor/Private/Graph/GraphFactory.cpp b/Engine/Plugins/Experimental/AnimNext/Source/AnimNextEditor/Private/Graph/GraphFactory.cpp
--- a/Engine/Plugins/Experimental/AnimNext/Source/AnimNextEditor/Private/Graph/GraphFactory.cpp
+++ b/Engine/Plugins/Experimental/AnimNext/Source/AnimNextEditor/Private/Graph/GraphFactory.cpp
@@ -5,7 +5,8 @@
 #include "Graph/AnimNextGraph_EditorData.h"
 #include "Graph/AnimNextGraph_EdGraph.h"
 #include "UncookedOnlyUtils.h"
-#include "Graph/RigUnit_AnimNextGraphRoot.h"
+#include "Graph/RigUnit_AnimNextBeginExecution.h"
+#include "Graph/RigUnit_AnimNextEndExecution.h"
 #include "RigVMModel/RigVMController.h"
 #include "Units/RigUnit.h"
 
@@ -34,10 +35,10 @@ UObject* UAnimNextGraphFactory::FactoryCreateNew(UClass* Class, UObject* InParen
 	// Add initial execution unit
 	URigVMController* Controller = EditorData->GetRigVMClient()->GetController(EditorData->RootGraph);
 
-	URigVMUnitNode* MainEntryPointNode = Controller->AddUnitNode(FRigUnit_AnimNextGraphRoot::StaticStruct(), FRigVMStruct::ExecuteName, FVector2D(-400.0f, 0.0f), FString(), false);
-	URigVMPin* BeginExecutePin = MainEntryPointNode->FindPin(GET_MEMBER_NAME_STRING_CHECKED(FRigUnit_AnimNextGraphRoot, Result));
+	URigVMUnitNode* MainEntryPointNode = Controller->AddUnitNode(FRigUnit_AnimNextBeginExecution::StaticStruct(), FRigUnit::GetMethodName(), FVector2D(-400.0f, 0.0f), FString(), false);
+	URigVMPin* BeginExecutePin = MainEntryPointNode->FindPin(GET_MEMBER_NAME_STRING_CHECKED(FRigUnit_AnimNextBeginExecution, ExecuteContext));
 	check(BeginExecutePin);
-	check(BeginExecutePin->GetDirection() == ERigVMPinDirection::Input);
+	check(BeginExecutePin->GetDirection() == ERigVMPinDirection::Output);
 
 	//// Add function to function lib
 	//{
@@ -73,7 +74,7 @@ UObject* UAnimNextGraphFactory::FactoryCreateNew(UClass* Class, UObject* InParen
 
 	//// Add end-execution unit of the correct type
 	//// TODO: using a float for now, but need to use a registry to determine correct type
-	//URigVMUnitNode* MainExitPointNode = Controller->AddUnitNode(FRigUnit_AnimNextEndExecution_Float::StaticStruct(), FRigVMStruct::ExecuteName, FVector2D(400.0f, 0.0f), FString(), false);
+	//URigVMUnitNode* MainExitPointNode = Controller->AddUnitNode(FRigUnit_AnimNextEndExecution_Float::StaticStruct(), FRigUnit::GetMethodName(), FVector2D(400.0f, 0.0f), FString(), false);
 	//URigVMPin* EndExecutePin = MainExitPointNode->FindPin(GET_MEMBER_NAME_STRING_CHECKED(FRigUnit_AnimNextEndExecution, ExecuteContext));
 	//check(EndExecutePin);
 	//check(EndExecutePin->GetDirection() == ERigVMPinDirection::Input);
diff --git a/Engine/Plugins/Experimental/AnimNext/Source/AnimNextEditor/Private/Graph/SAnimNextGraphNode.cpp b/Engine/Plugins/Experimental/AnimNext/Source/AnimNextEditor/Private/Graph/SAnimNextGraphNode.cpp
deleted file mode 100644
diff --git a/Engine/Plugins/Experimental/AnimNext/Source/AnimNextEditor/Private/Graph/SAnimNextGraphNode.h b/Engine/Plugins/Experimental/AnimNext/Source/AnimNextEditor/Private/Graph/SAnimNextGraphNode.h
deleted file mode 100644
diff --git a/Engine/Plugins/Experimental/AnimNext/Source/AnimNextTestSuite/Private/AnimNextDecoratorBaseTest.cpp b/Engine/Plugins/Experimental/AnimNext/Source/AnimNextTestSuite/Private/AnimNextDecoratorBaseTest.cpp
--- a/Engine/Plugins/Experimental/AnimNext/Source/AnimNextTestSuite/Private/AnimNextDecoratorBaseTest.cpp
+++ b/Engine/Plugins/Experimental/AnimNext/Source/AnimNextTestSuite/Private/AnimNextDecoratorBaseTest.cpp
@@ -16,7 +16,6 @@
 #include "DecoratorBase/ExecutionContext.h"
 #include "DecoratorBase/IDecoratorInterface.h"
 #include "DecoratorBase/NodeInstance.h"
-#include "DecoratorBase/NodeTemplateBuilder.h"
 #include "DecoratorBase/NodeTemplateRegistry.h"
 
 //****************************************************************************
@@ -297,9 +296,7 @@ bool FAnimationAnimNextRuntimeTest_DecoratorRegistry::RunTest(const FString& InP
 
 	FDecoratorRegistry& Registry = FDecoratorRegistry::Get();
 
-	// Some decorators already exist in the engine, keep track of them
-	const uint32 NumAutoRegisteredDecorators = Registry.GetNum();
-
+	AddErrorIfFalse(Registry.GetNum() == 0, "FAnimationAnimNextRuntimeTest_DecoratorRegistry -> Registry should be empty");
 	AddErrorIfFalse(!Registry.FindHandle(FDecoratorA_Base::DecoratorUID).IsValid(), "FAnimationAnimNextRuntimeTest_DecoratorRegistry -> Registry should not contain our decorator");
 	AddErrorIfFalse(!Registry.FindHandle(FDecoratorAB_Add::DecoratorUID).IsValid(), "FAnimationAnimNextRuntimeTest_DecoratorRegistry -> Registry should not contain our decorator");
 	AddErrorIfFalse(!Registry.FindHandle(FDecoratorAC_Add::DecoratorUID).IsValid(), "FAnimationAnimNextRuntimeTest_DecoratorRegistry -> Registry should not contain our decorator");
@@ -308,7 +305,7 @@ bool FAnimationAnimNextRuntimeTest_DecoratorRegistry::RunTest(const FString& InP
 		// Auto register a decorator
 		AUTO_REGISTER_ANIM_DECORATOR(FDecoratorA_Base)
 
-		AddErrorIfFalse(Registry.GetNum() == NumAutoRegisteredDecorators + 1, "FAnimationAnimNextRuntimeTest_DecoratorRegistry -> Registry should contain 1 new decorator");
+		AddErrorIfFalse(Registry.GetNum() == 1, "FAnimationAnimNextRuntimeTest_DecoratorRegistry -> Registry should contain 1 decorator");
 
 		FDecoratorRegistryHandle HandleA = Registry.FindHandle(FDecoratorA_Base::DecoratorUID);
 		AddErrorIfFalse(HandleA.IsValid(), "FAnimationAnimNextRuntimeTest_DecoratorRegistry -> Decorator should have registered automatically");
@@ -324,7 +321,7 @@ bool FAnimationAnimNextRuntimeTest_DecoratorRegistry::RunTest(const FString& InP
 				// Auto register another decorator
 				AUTO_REGISTER_ANIM_DECORATOR(FDecoratorAB_Add)
 
-				AddErrorIfFalse(Registry.GetNum() == NumAutoRegisteredDecorators + 2, "FAnimationAnimNextRuntimeTest_DecoratorRegistry -> Registry should contain 2 new decorators");
+				AddErrorIfFalse(Registry.GetNum() == 2, "FAnimationAnimNextRuntimeTest_DecoratorRegistry -> Registry should contain 2 decorators");
 
 				FDecoratorRegistryHandle HandleAB = Registry.FindHandle(FDecoratorAB_Add::DecoratorUID);
 				AddErrorIfFalse(HandleAB.IsValid(), "FAnimationAnimNextRuntimeTest_DecoratorRegistry -> Decorator should have registered automatically");
@@ -343,7 +340,7 @@ bool FAnimationAnimNextRuntimeTest_DecoratorRegistry::RunTest(const FString& InP
 						FDecoratorAC_Add DecoratorAC_0;
 						Registry.Register(&DecoratorAC_0);
 
-						AddErrorIfFalse(Registry.GetNum() == NumAutoRegisteredDecorators + 3, "FAnimationAnimNextRuntimeTest_DecoratorRegistry -> Registry should contain 3 new decorators");
+						AddErrorIfFalse(Registry.GetNum() == 3, "FAnimationAnimNextRuntimeTest_DecoratorRegistry -> Registry should contain 3 decorators");
 
 						HandleAC_0 = Registry.FindHandle(FDecoratorAC_Add::DecoratorUID);
 						AddErrorIfFalse(HandleAC_0.IsValid(), "FAnimationAnimNextRuntimeTest_DecoratorRegistry -> Decorator should have registered automatically");
@@ -360,7 +357,7 @@ bool FAnimationAnimNextRuntimeTest_DecoratorRegistry::RunTest(const FString& InP
 							// Unregister our instances
 							Registry.Unregister(&DecoratorAC_0);
 
-							AddErrorIfFalse(Registry.GetNum() == NumAutoRegisteredDecorators + 2, "FAnimationAnimNextRuntimeTest_DecoratorRegistry -> Registry should contain 2 extra decorators");
+							AddErrorIfFalse(Registry.GetNum() == 2, "FAnimationAnimNextRuntimeTest_DecoratorRegistry -> Registry should contain 2 decorators");
 							AddErrorIfFalse(!Registry.FindHandle(FDecoratorAC_Add::DecoratorUID).IsValid(), "FAnimationAnimNextRuntimeTest_DecoratorRegistry -> Decorator should have unregistered");
 						}
 					}
@@ -370,7 +367,7 @@ bool FAnimationAnimNextRuntimeTest_DecoratorRegistry::RunTest(const FString& InP
 						FDecoratorAC_Add DecoratorAC_1;
 						Registry.Register(&DecoratorAC_1);
 
-						AddErrorIfFalse(Registry.GetNum() == NumAutoRegisteredDecorators + 3, "FAnimationAnimNextRuntimeTest_DecoratorRegistry -> Registry should contain 3 new decorators");
+						AddErrorIfFalse(Registry.GetNum() == 3, "FAnimationAnimNextRuntimeTest_DecoratorRegistry -> Registry should contain 3 decorators");
 
 						FDecoratorRegistryHandle HandleAC_1 = Registry.FindHandle(FDecoratorAC_Add::DecoratorUID);
 						AddErrorIfFalse(HandleAC_1.IsValid(), "FAnimationAnimNextRuntimeTest_DecoratorRegistry -> Decorator should have registered automatically");
@@ -388,20 +385,20 @@ bool FAnimationAnimNextRuntimeTest_DecoratorRegistry::RunTest(const FString& InP
 							// Unregister our instances
 							Registry.Unregister(&DecoratorAC_1);
 
-							AddErrorIfFalse(Registry.GetNum() == NumAutoRegisteredDecorators + 2, "FAnimationAnimNextRuntimeTest_DecoratorRegistry -> Registry should contain 2 extra decorators");
+							AddErrorIfFalse(Registry.GetNum() == 2, "FAnimationAnimNextRuntimeTest_DecoratorRegistry -> Registry should contain 2 decorators");
 							AddErrorIfFalse(!Registry.FindHandle(FDecoratorAC_Add::DecoratorUID).IsValid(), "FAnimationAnimNextRuntimeTest_DecoratorRegistry -> Decorator should have unregistered");
 						}
 					}
 				}
 			}
 
-			AddErrorIfFalse(Registry.GetNum() == NumAutoRegisteredDecorators + 1, "FAnimationAnimNextRuntimeTest_DecoratorRegistry -> Registry should contain 1 extra decorator");
+			AddErrorIfFalse(Registry.GetNum() == 1, "FAnimationAnimNextRuntimeTest_DecoratorRegistry -> Registry should contain 1 decorator");
 			AddErrorIfFalse(!Registry.FindHandle(FDecoratorAB_Add::DecoratorUID).IsValid(), "FAnimationAnimNextRuntimeTest_DecoratorRegistry -> Decorator should have unregistered automatically");
 			AddErrorIfFalse(HandleA == Registry.FindHandle(FDecoratorA_Base::DecoratorUID), "FAnimationAnimNextRuntimeTest_DecoratorRegistry -> Decorator handle should not have changed");
 		}
 	}
 
-	AddErrorIfFalse(Registry.GetNum() == NumAutoRegisteredDecorators, "FAnimationAnimNextRuntimeTest_DecoratorRegistry -> All decorators should have unregistered");
+	AddErrorIfFalse(Registry.GetNum() == 0, "FAnimationAnimNextRuntimeTest_DecoratorRegistry -> All decorators should have unregistered");
 	AddErrorIfFalse(!Registry.FindHandle(FDecoratorA_Base::DecoratorUID).IsValid(), "FAnimationAnimNextRuntimeTest_DecoratorRegistry -> Decorator should have unregistered automatically");
 	AddErrorIfFalse(!Registry.FindHandle(FDecoratorAB_Add::DecoratorUID).IsValid(), "FAnimationAnimNextRuntimeTest_DecoratorRegistry -> Decorator should have unregistered automatically");
 	AddErrorIfFalse(!Registry.FindHandle(FDecoratorAC_Add::DecoratorUID).IsValid(), "FAnimationAnimNextRuntimeTest_DecoratorRegistry -> Decorator should have unregistered automatically");
@@ -419,7 +416,6 @@ bool FAnimationAnimNextRuntimeTest_NodeTemplateRegistry::RunTest(const FString&
 	AUTO_REGISTER_ANIM_DECORATOR(FDecoratorAB_Add)
 	AUTO_REGISTER_ANIM_DECORATOR(FDecoratorAC_Add)
 
-	FScopedClearNodeTemplateRegistry ScopedClearNodeTemplateRegistry;
 	FNodeTemplateRegistry& Registry = FNodeTemplateRegistry::Get();
 
 	TArray<FDecoratorUID> NodeTemplateDecoratorList;
@@ -431,7 +427,7 @@ bool FAnimationAnimNextRuntimeTest_NodeTemplateRegistry::RunTest(const FString&
 
 	// Populate our node template registry
 	TArray<uint8> NodeTemplateBuffer0;
-	const FNodeTemplate* NodeTemplate0 = FNodeTemplateBuilder::BuildNodeTemplate(NodeTemplateDecoratorList, NodeTemplateBuffer0);
+	const FNodeTemplate* NodeTemplate0 = BuildNodeTemplate(NodeTemplateDecoratorList, NodeTemplateBuffer0);
 
 	AddErrorIfFalse(Registry.GetNum() == 0, "FAnimationAnimNextRuntimeTest_NodeTemplateRegistry -> Registry should contain any templates");
 
@@ -449,7 +445,7 @@ bool FAnimationAnimNextRuntimeTest_NodeTemplateRegistry::RunTest(const FString&
 
 		// Try and register a duplicate template
 		TArray<uint8> NodeTemplateBuffer1;
-		const FNodeTemplate* NodeTemplate1 = FNodeTemplateBuilder::BuildNodeTemplate(NodeTemplateDecoratorList, NodeTemplateBuffer1);
+		const FNodeTemplate* NodeTemplate1 = BuildNodeTemplate(NodeTemplateDecoratorList, NodeTemplateBuffer1);
 		if (NodeTemplate1 != nullptr)
 		{
 			AddErrorIfFalse(NodeTemplate0 != NodeTemplate1, "FAnimationAnimNextRuntimeTest_NodeTemplateRegistry -> Node template pointers should be different");
@@ -467,7 +463,7 @@ bool FAnimationAnimNextRuntimeTest_NodeTemplateRegistry::RunTest(const FString&
 			NodeTemplateDecoratorList2.Add(FDecoratorAC_Add::DecoratorUID);
 
 			TArray<uint8> NodeTemplateBuffer2;
-			const FNodeTemplate* NodeTemplate2 = FNodeTemplateBuilder::BuildNodeTemplate(NodeTemplateDecoratorList2, NodeTemplateBuffer2);
+			const FNodeTemplate* NodeTemplate2 = BuildNodeTemplate(NodeTemplateDecoratorList2, NodeTemplateBuffer2);
 			if (NodeTemplate2 != nullptr)
 			{
 				AddErrorIfFalse(NodeTemplate0->GetUID() != NodeTemplate2->GetUID(), "FAnimationAnimNextRuntimeTest_NodeTemplateRegistry -> Node template UIDs should be different");
@@ -500,7 +496,6 @@ bool FAnimationAnimNextRuntimeTest_NodeLifetime::RunTest(const FString& InParame
 	AUTO_REGISTER_ANIM_DECORATOR(FDecoratorAB_Add)
 	AUTO_REGISTER_ANIM_DECORATOR(FDecoratorAC_Add)
 
-	FScopedClearNodeTemplateRegistry ScopedClearNodeTemplateRegistry;
 	FNodeTemplateRegistry& Registry = FNodeTemplateRegistry::Get();
 
 	TArray<FDecoratorUID> NodeTemplateDecoratorList;
@@ -512,9 +507,10 @@ bool FAnimationAnimNextRuntimeTest_NodeLifetime::RunTest(const FString& InParame
 
 	// Populate our node template registry
 	TArray<uint8> NodeTemplateBuffer0;
-	const FNodeTemplate* NodeTemplate0 = FNodeTemplateBuilder::BuildNodeTemplate(NodeTemplateDecoratorList, NodeTemplateBuffer0);
+	const FNodeTemplate* NodeTemplate0 = BuildNodeTemplate(NodeTemplateDecoratorList, NodeTemplateBuffer0);
 
 	FNodeTemplateRegistryHandle TemplateHandle0 = Registry.FindOrAdd(NodeTemplate0);
+	AddErrorIfFalse(Registry.GetNum() == 1, "FAnimationAnimNextRuntimeTest_NodeLifetime -> Registry should contain 1 template");
 	AddErrorIfFalse(TemplateHandle0.IsValid(), "FAnimationAnimNextRuntimeTest_NodeLifetime -> Registry should contain our template");
 
 	FNodeHandle Node0;
@@ -529,17 +525,19 @@ bool FAnimationAnimNextRuntimeTest_NodeLifetime::RunTest(const FString& InParame
 		Node1 = DecoratorWriter.RegisterNode(*NodeTemplate0);
 
 		// We don't have decorator properties
+		TArray<TMap<FString, FString>> DecoratorProperties;
+		DecoratorProperties.AddDefaulted(NodeTemplateDecoratorList.Num());
 
 		DecoratorWriter.BeginNodeWriting();
 		DecoratorWriter.WriteNode(Node0,
-			[](uint32 DecoratorIndex, const FString& PropertyName)
+			[&DecoratorProperties](uint32 DecoratorIndex) -> const TMap<FString, FString>&
 			{
-				return FString();
+				return DecoratorProperties[DecoratorIndex];
 			});
 		DecoratorWriter.WriteNode(Node1,
-			[](uint32 DecoratorIndex, const FString& PropertyName)
+			[&DecoratorProperties](uint32 DecoratorIndex) -> const TMap<FString, FString>&
 			{
-				return FString();
+				return DecoratorProperties[DecoratorIndex];
 			});
 		DecoratorWriter.EndNodeWriting();
 
@@ -549,16 +547,11 @@ bool FAnimationAnimNextRuntimeTest_NodeLifetime::RunTest(const FString& InParame
 
 	// Read our graph
 	TArray<uint8> GraphSharedDataBuffer;
-	TArray<TObjectPtr<UObject>> TrackedObjectsForGC;
 	{
 		FMemoryReader GraphSharedDataArchive(GraphSharedDataArchiveBuffer);
 		FDecoratorReader DecoratorReader(GraphSharedDataArchive);
 
-		FDecoratorReader::EErrorState ErrorState = DecoratorReader.ReadGraph(GraphSharedDataBuffer, TrackedObjectsForGC);
-		AddErrorIfFalse(ErrorState == FDecoratorReader::EErrorState::None, "FAnimationAnimNextRuntimeTest_NodeLifetime -> Failed to read graph shared data");
-
-		Node0 = DecoratorReader.ResolveNodeHandle(Node0);
-		Node1 = DecoratorReader.ResolveNodeHandle(Node1);
+		DecoratorReader.ReadGraphSharedData(GraphSharedDataBuffer);
 	}
 
 	FExecutionContext Context(GraphSharedDataBuffer);
@@ -704,7 +697,6 @@ bool FAnimationAnimNextRuntimeTest_GetDecoratorInterface::RunTest(const FString&
 	AUTO_REGISTER_ANIM_DECORATOR(FDecoratorAB_Add)
 	AUTO_REGISTER_ANIM_DECORATOR(FDecoratorAC_Add)
 
-	FScopedClearNodeTemplateRegistry ScopedClearNodeTemplateRegistry;
 	FNodeTemplateRegistry& Registry = FNodeTemplateRegistry::Get();
 
 	TArray<FDecoratorUID> NodeTemplateDecoratorList;
@@ -716,9 +708,10 @@ bool FAnimationAnimNextRuntimeTest_GetDecoratorInterface::RunTest(const FString&
 
 	// Populate our node template registry
 	TArray<uint8> NodeTemplateBuffer0;
-	const FNodeTemplate* NodeTemplate0 = FNodeTemplateBuilder::BuildNodeTemplate(NodeTemplateDecoratorList, NodeTemplateBuffer0);
+	const FNodeTemplate* NodeTemplate0 = BuildNodeTemplate(NodeTemplateDecoratorList, NodeTemplateBuffer0);
 
 	FNodeTemplateRegistryHandle TemplateHandle0 = Registry.FindOrAdd(NodeTemplate0);
+	AddErrorIfFalse(Registry.GetNum() == 1, "FAnimationAnimNextRuntimeTest_GetDecoratorInterface -> Registry should contain 1 template");
 	AddErrorIfFalse(TemplateHandle0.IsValid(), "FAnimationAnimNextRuntimeTest_GetDecoratorInterface -> Registry should contain our template");
 
 	FNodeHandle Node0;
@@ -731,12 +724,14 @@ bool FAnimationAnimNextRuntimeTest_GetDecoratorInterface::RunTest(const FString&
 		Node0 = DecoratorWriter.RegisterNode(*NodeTemplate0);
 
 		// We don't have decorator properties
+		TArray<TMap<FString, FString>> DecoratorProperties;
+		DecoratorProperties.AddDefaulted(NodeTemplateDecoratorList.Num());
 
 		DecoratorWriter.BeginNodeWriting();
 		DecoratorWriter.WriteNode(Node0,
-			[](uint32 DecoratorIndex, const FString& PropertyName)
+			[&DecoratorProperties](uint32 DecoratorIndex) -> const TMap<FString, FString>&
 			{
-				return FString();
+				return DecoratorProperties[DecoratorIndex];
 			});
 		DecoratorWriter.EndNodeWriting();
 
@@ -746,15 +741,11 @@ bool FAnimationAnimNextRuntimeTest_GetDecoratorInterface::RunTest(const FString&
 
 	// Read our graph
 	TArray<uint8> GraphSharedDataBuffer;
-	TArray<TObjectPtr<UObject>> TrackedObjectsForGC;
 	{
 		FMemoryReader GraphSharedDataArchive(GraphSharedDataArchiveBuffer);
 		FDecoratorReader DecoratorReader(GraphSharedDataArchive);
 
-		FDecoratorReader::EErrorState ErrorState = DecoratorReader.ReadGraph(GraphSharedDataBuffer, TrackedObjectsForGC);
-		AddErrorIfFalse(ErrorState == FDecoratorReader::EErrorState::None, "FAnimationAnimNextRuntimeTest_GetDecoratorInterface -> Failed to read graph shared data");
-
-		Node0 = DecoratorReader.ResolveNodeHandle(Node0);
+		DecoratorReader.ReadGraphSharedData(GraphSharedDataBuffer);
 	}
 
 	FExecutionContext Context(GraphSharedDataBuffer);
@@ -969,7 +960,6 @@ bool FAnimationAnimNextRuntimeTest_GetDecoratorInterfaceSuper::RunTest(const FSt
 	AUTO_REGISTER_ANIM_DECORATOR(FDecoratorAB_Add)
 	AUTO_REGISTER_ANIM_DECORATOR(FDecoratorAC_Add)
 
-	FScopedClearNodeTemplateRegistry ScopedClearNodeTemplateRegistry;
 	FNodeTemplateRegistry& Registry = FNodeTemplateRegistry::Get();
 
 	TArray<FDecoratorUID> NodeTemplateDecoratorList;
@@ -981,7 +971,7 @@ bool FAnimationAnimNextRuntimeTest_GetDecoratorInterfaceSuper::RunTest(const FSt
 
 	// Populate our node template registry
 	TArray<uint8> NodeTemplateBuffer0;
-	const FNodeTemplate* NodeTemplate0 = FNodeTemplateBuilder::BuildNodeTemplate(NodeTemplateDecoratorList, NodeTemplateBuffer0);
+	const FNodeTemplate* NodeTemplate0 = BuildNodeTemplate(NodeTemplateDecoratorList, NodeTemplateBuffer0);
 
 	FNodeTemplateRegistryHandle TemplateHandle0 = Registry.FindOrAdd(NodeTemplate0);
 	AddErrorIfFalse(Registry.GetNum() == 1, "FAnimationAnimNextRuntimeTest_GetDecoratorInterfaceSuper -> Registry should contain 1 template");
@@ -997,12 +987,14 @@ bool FAnimationAnimNextRuntimeTest_GetDecoratorInterfaceSuper::RunTest(const FSt
 		Node0 = DecoratorWriter.RegisterNode(*NodeTemplate0);
 
 		// We don't have decorator properties
+		TArray<TMap<FString, FString>> DecoratorProperties;
+		DecoratorProperties.AddDefaulted(NodeTemplateDecoratorList.Num());
 
 		DecoratorWriter.BeginNodeWriting();
 		DecoratorWriter.WriteNode(Node0,
-			[](uint32 DecoratorIndex, const FString& PropertyName)
+			[&DecoratorProperties](uint32 DecoratorIndex) -> const TMap<FString, FString>&
 			{
-				return FString();
+				return DecoratorProperties[DecoratorIndex];
 			});
 		DecoratorWriter.EndNodeWriting();
 
@@ -1012,15 +1004,11 @@ bool FAnimationAnimNextRuntimeTest_GetDecoratorInterfaceSuper::RunTest(const FSt
 
 	// Read our graph
 	TArray<uint8> GraphSharedDataBuffer;
-	TArray<TObjectPtr<UObject>> TrackedObjectsForGC;
 	{
 		FMemoryReader GraphSharedDataArchive(GraphSharedDataArchiveBuffer);
 		FDecoratorReader DecoratorReader(GraphSharedDataArchive);
 
-		FDecoratorReader::EErrorState ErrorState = DecoratorReader.ReadGraph(GraphSharedDataBuffer, TrackedObjectsForGC);
-		AddErrorIfFalse(ErrorState == FDecoratorReader::EErrorState::None, "FAnimationAnimNextRuntimeTest_GetDecoratorInterfaceSuper -> Failed to read graph shared data");
-
-		Node0 = DecoratorReader.ResolveNodeHandle(Node0);
+		DecoratorReader.ReadGraphSharedData(GraphSharedDataBuffer);
 	}
 
 	FExecutionContext Context(GraphSharedDataBuffer);
@@ -1165,7 +1153,6 @@ bool FAnimationAnimNextRuntimeTest_DecoratorSerialization::RunTest(const FString
 	AUTO_REGISTER_ANIM_DECORATOR(FDecoratorSerialization_Add)
 	AUTO_REGISTER_ANIM_DECORATOR(FDecoratorNativeSerialization_Add)
 
-	FScopedClearNodeTemplateRegistry ScopedClearNodeTemplateRegistry;
 	FNodeTemplateRegistry& Registry = FNodeTemplateRegistry::Get();
 
 	TArray<FDecoratorUID> NodeTemplateDecoratorList;
@@ -1175,9 +1162,10 @@ bool FAnimationAnimNextRuntimeTest_DecoratorSerialization::RunTest(const FString
 
 	// Populate our node template registry
 	TArray<uint8> NodeTemplateBuffer0;
-	const FNodeTemplate* NodeTemplate0 = FNodeTemplateBuilder::BuildNodeTemplate(NodeTemplateDecoratorList, NodeTemplateBuffer0);
+	const FNodeTemplate* NodeTemplate0 = BuildNodeTemplate(NodeTemplateDecoratorList, NodeTemplateBuffer0);
 
 	FNodeTemplateRegistryHandle TemplateHandle0 = Registry.FindOrAdd(NodeTemplate0);
+	AddErrorIfFalse(Registry.GetNum() == 1, "FAnimationAnimNextRuntimeTest_DecoratorSerialization -> Registry should contain 1 template");
 	AddErrorIfFalse(TemplateHandle0.IsValid(), "FAnimationAnimNextRuntimeTest_DecoratorSerialization -> Registry should contain our template");
 
 	FDecoratorSerialization_Base::FSharedData DecoratorBaseRef0;
@@ -1338,14 +1326,14 @@ bool FAnimationAnimNextRuntimeTest_DecoratorSerialization::RunTest(const FString
 
 		DecoratorWriter.BeginNodeWriting();
 		DecoratorWriter.WriteNode(Node0,
-			[&DecoratorProperties0](uint32 DecoratorIndex, const FString& PropertyName)
+			[&DecoratorProperties0](uint32 DecoratorIndex) -> const TMap<FString, FString>&
 			{
-				return DecoratorProperties0[DecoratorIndex][PropertyName];
+				return DecoratorProperties0[DecoratorIndex];
 			});
 		DecoratorWriter.WriteNode(Node1,
-			[&DecoratorProperties1](uint32 DecoratorIndex, const FString& PropertyName)
+			[&DecoratorProperties1](uint32 DecoratorIndex) -> const TMap<FString, FString>&
 			{
-				return DecoratorProperties1[DecoratorIndex][PropertyName];
+				return DecoratorProperties1[DecoratorIndex];
 			});
 		DecoratorWriter.EndNodeWriting();
 
@@ -1354,147 +1342,144 @@ bool FAnimationAnimNextRuntimeTest_DecoratorSerialization::RunTest(const FString
 	}
 
 	// Clear out the node template registry to test registration on load
-	{
-		FScopedClearNodeTemplateRegistry ScopedClearNodeTemplateRegistryForLoad;
+	Registry.Clear();
 
-		AddErrorIfFalse(Registry.GetNum() == 0, "FAnimationAnimNextRuntimeTest_DecoratorSerialization -> Registry should contain 0 templates");
+	AddErrorIfFalse(Registry.GetNum() == 0, "FAnimationAnimNextRuntimeTest_DecoratorSerialization -> Registry should contain 0 templates");
 
-		// Read our graph
-		TArray<uint8> GraphSharedDataBuffer;
-		TArray<TObjectPtr<UObject>> TrackedObjectsForGC;
-		{
-			FMemoryReader GraphSharedDataArchive(GraphSharedDataArchiveBuffer);
-			FDecoratorReader DecoratorReader(GraphSharedDataArchive);
-
-			FDecoratorReader::EErrorState ErrorState = DecoratorReader.ReadGraph(GraphSharedDataBuffer, TrackedObjectsForGC);
-			AddErrorIfFalse(ErrorState == FDecoratorReader::EErrorState::None, "FAnimationAnimNextRuntimeTest_DecoratorSerialization -> Failed to read graph shared data");
+	// Read our graph
+	TArray<uint8> GraphSharedDataBuffer;
+	{
+		FMemoryReader GraphSharedDataArchive(GraphSharedDataArchiveBuffer);
+		FDecoratorReader DecoratorReader(GraphSharedDataArchive);
 
-			Node0 = DecoratorReader.ResolveNodeHandle(Node0);
-			Node1 = DecoratorReader.ResolveNodeHandle(Node1);
+		DecoratorReader.ReadGraphSharedData(GraphSharedDataBuffer);
 
-			AddErrorIfFalse(Registry.GetNum() == 1, "FAnimationAnimNextRuntimeTest_DecoratorSerialization -> Registry should contain 1 template");
-		}
+		AddErrorIfFalse(Registry.GetNum() == 1, "FAnimationAnimNextRuntimeTest_DecoratorSerialization -> Registry should contain 1 template");
+	}
 
-		FExecutionContext Context(GraphSharedDataBuffer);
+	FExecutionContext Context(GraphSharedDataBuffer);
 
-		// Validate decorator serialization
-		{
-			FDecoratorBinding ParentBinding;				// Empty, no parent
-			FAnimNextDecoratorHandle DecoratorHandle0(Node0, 0);	// Point to first node, first base decorator
-			FAnimNextDecoratorHandle DecoratorHandle1(Node1, 0);	// Point to second node, first base decorator
+	// Validate decorator serialization
+	{
+		FDecoratorBinding ParentBinding;				// Empty, no parent
+		FAnimNextDecoratorHandle DecoratorHandle0(Node0, 0);	// Point to first node, first base decorator
+		FAnimNextDecoratorHandle DecoratorHandle1(Node1, 0);	// Point to second node, first base decorator
 
-			FDecoratorPtr DecoratorPtr0 = Context.AllocateNodeInstance(ParentBinding, DecoratorHandle0);
-			AddErrorIfFalse(DecoratorPtr0.IsValid(), "FAnimationAnimNextRuntimeTest_DecoratorSerialization -> Failed to allocate a node instance");
+		FDecoratorPtr DecoratorPtr0 = Context.AllocateNodeInstance(ParentBinding, DecoratorHandle0);
+		AddErrorIfFalse(DecoratorPtr0.IsValid(), "FAnimationAnimNextRuntimeTest_DecoratorSerialization -> Failed to allocate a node instance");
 
-			FDecoratorPtr DecoratorPtr1 = Context.AllocateNodeInstance(ParentBinding, DecoratorHandle1);
-			AddErrorIfFalse(DecoratorPtr1.IsValid(), "FAnimationAnimNextRuntimeTest_DecoratorSerialization -> Failed to allocate a node instance");
+		FDecoratorPtr DecoratorPtr1 = Context.AllocateNodeInstance(ParentBinding, DecoratorHandle1);
+		AddErrorIfFalse(DecoratorPtr1.IsValid(), "FAnimationAnimNextRuntimeTest_DecoratorSerialization -> Failed to allocate a node instance");
 
-			// Validate shared data for base decorator on node 0
-			{
-				TDecoratorBinding<IInterfaceA> BindingA0;
-				AddErrorIfFalse(Context.GetInterface(DecoratorPtr0, BindingA0), "FAnimationAnimNextRuntimeTest_DecoratorSerialization -> InterfaceA not found");
-
-				const auto* SharedDataA0 = BindingA0.GetSharedData<FDecoratorSerialization_Base::FSharedData>();
-
-				AddErrorIfFalse(SharedDataA0->Integer == DecoratorBaseRef0.Integer, "FAnimationAnimNextRuntimeTest_DecoratorSerialization -> Unexpected serialized value");
-				AddErrorIfFalse(FMemory::Memcmp(SharedDataA0->IntegerArray, DecoratorBaseRef0.IntegerArray, sizeof(DecoratorBaseRef0.IntegerArray)) == 0, "FAnimationAnimNextRuntimeTest_DecoratorSerialization -> Unexpected serialized value");
-				AddErrorIfFalse(SharedDataA0->IntegerTArray == DecoratorBaseRef0.IntegerTArray, "FAnimationAnimNextRuntimeTest_DecoratorSerialization -> Unexpected serialized value");
-				AddErrorIfFalse(SharedDataA0->Vector == DecoratorBaseRef0.Vector, "FAnimationAnimNextRuntimeTest_DecoratorSerialization -> Unexpected serialized value");
-				AddErrorIfFalse(FMemory::Memcmp(SharedDataA0->VectorArray, DecoratorBaseRef0.VectorArray, sizeof(DecoratorBaseRef0.VectorArray)) == 0, "FAnimationAnimNextRuntimeTest_DecoratorSerialization -> Unexpected serialized value");
-				AddErrorIfFalse(SharedDataA0->VectorTArray == DecoratorBaseRef0.VectorTArray, "FAnimationAnimNextRuntimeTest_DecoratorSerialization -> Unexpected serialized value");
-				AddErrorIfFalse(SharedDataA0->String == DecoratorBaseRef0.String, "FAnimationAnimNextRuntimeTest_DecoratorSerialization -> Unexpected serialized value");
-				AddErrorIfFalse(SharedDataA0->Name == DecoratorBaseRef0.Name, "FAnimationAnimNextRuntimeTest_DecoratorSerialization -> Unexpected serialized value");
-			}
+		// Validate shared data for base decorator on node 0
+		{
+			TDecoratorBinding<IInterfaceA> BindingA0;
+			AddErrorIfFalse(Context.GetInterface(DecoratorPtr0, BindingA0), "FAnimationAnimNextRuntimeTest_DecoratorSerialization -> InterfaceA not found");
+
+			const auto* SharedDataA0 = BindingA0.GetSharedData<FDecoratorSerialization_Base::FSharedData>();
+
+			AddErrorIfFalse(SharedDataA0->Integer == DecoratorBaseRef0.Integer, "FAnimationAnimNextRuntimeTest_DecoratorSerialization -> Unexpected serialized value");
+			AddErrorIfFalse(FMemory::Memcmp(SharedDataA0->IntegerArray, DecoratorBaseRef0.IntegerArray, sizeof(DecoratorBaseRef0.IntegerArray)) == 0, "FAnimationAnimNextRuntimeTest_DecoratorSerialization -> Unexpected serialized value");
+			AddErrorIfFalse(SharedDataA0->IntegerTArray == DecoratorBaseRef0.IntegerTArray, "FAnimationAnimNextRuntimeTest_DecoratorSerialization -> Unexpected serialized value");
+			AddErrorIfFalse(SharedDataA0->Vector == DecoratorBaseRef0.Vector, "FAnimationAnimNextRuntimeTest_DecoratorSerialization -> Unexpected serialized value");
+			AddErrorIfFalse(FMemory::Memcmp(SharedDataA0->VectorArray, DecoratorBaseRef0.VectorArray, sizeof(DecoratorBaseRef0.VectorArray)) == 0, "FAnimationAnimNextRuntimeTest_DecoratorSerialization -> Unexpected serialized value");
+			AddErrorIfFalse(SharedDataA0->VectorTArray == DecoratorBaseRef0.VectorTArray, "FAnimationAnimNextRuntimeTest_DecoratorSerialization -> Unexpected serialized value");
+			AddErrorIfFalse(SharedDataA0->String == DecoratorBaseRef0.String, "FAnimationAnimNextRuntimeTest_DecoratorSerialization -> Unexpected serialized value");
+			AddErrorIfFalse(SharedDataA0->Name == DecoratorBaseRef0.Name, "FAnimationAnimNextRuntimeTest_DecoratorSerialization -> Unexpected serialized value");
+		}
 
-			// Validate shared data for additive decorator on node 0
-			{
-				TDecoratorBinding<IInterfaceB> BindingB0;
-				AddErrorIfFalse(Context.GetInterface(DecoratorPtr0, BindingB0), "FAnimationAnimNextRuntimeTest_DecoratorSerialization -> InterfaceB not found");
-
-				const auto* SharedDataB0 = BindingB0.GetSharedData<FDecoratorSerialization_Add::FSharedData>();
-
-				AddErrorIfFalse(SharedDataB0->Integer == DecoratorAddRef0.Integer, "FAnimationAnimNextRuntimeTest_DecoratorSerialization -> Unexpected serialized value");
-				AddErrorIfFalse(FMemory::Memcmp(SharedDataB0->IntegerArray, DecoratorAddRef0.IntegerArray, sizeof(DecoratorAddRef0.IntegerArray)) == 0, "FAnimationAnimNextRuntimeTest_DecoratorSerialization -> Unexpected serialized value");
-				AddErrorIfFalse(SharedDataB0->IntegerTArray == DecoratorAddRef0.IntegerTArray, "FAnimationAnimNextRuntimeTest_DecoratorSerialization -> Unexpected serialized value");
-				AddErrorIfFalse(SharedDataB0->Vector == DecoratorAddRef0.Vector, "FAnimationAnimNextRuntimeTest_DecoratorSerialization -> Unexpected serialized value");
-				AddErrorIfFalse(FMemory::Memcmp(SharedDataB0->VectorArray, DecoratorAddRef0.VectorArray, sizeof(DecoratorAddRef0.VectorArray)) == 0, "FAnimationAnimNextRuntimeTest_DecoratorSerialization -> Unexpected serialized value");
-				AddErrorIfFalse(SharedDataB0->VectorTArray == DecoratorAddRef0.VectorTArray, "FAnimationAnimNextRuntimeTest_DecoratorSerialization -> Unexpected serialized value");
-				AddErrorIfFalse(SharedDataB0->String == DecoratorAddRef0.String, "FAnimationAnimNextRuntimeTest_DecoratorSerialization -> Unexpected serialized value");
-				AddErrorIfFalse(SharedDataB0->Name == DecoratorAddRef0.Name, "FAnimationAnimNextRuntimeTest_DecoratorSerialization -> Unexpected serialized value");
-			}
+		// Validate shared data for additive decorator on node 0
+		{
+			TDecoratorBinding<IInterfaceB> BindingB0;
+			AddErrorIfFalse(Context.GetInterface(DecoratorPtr0, BindingB0), "FAnimationAnimNextRuntimeTest_DecoratorSerialization -> InterfaceB not found");
+
+			const auto* SharedDataB0 = BindingB0.GetSharedData<FDecoratorSerialization_Add::FSharedData>();
+
+			AddErrorIfFalse(SharedDataB0->Integer == DecoratorAddRef0.Integer, "FAnimationAnimNextRuntimeTest_DecoratorSerialization -> Unexpected serialized value");
+			AddErrorIfFalse(FMemory::Memcmp(SharedDataB0->IntegerArray, DecoratorAddRef0.IntegerArray, sizeof(DecoratorAddRef0.IntegerArray)) == 0, "FAnimationAnimNextRuntimeTest_DecoratorSerialization -> Unexpected serialized value");
+			AddErrorIfFalse(SharedDataB0->IntegerTArray == DecoratorAddRef0.IntegerTArray, "FAnimationAnimNextRuntimeTest_DecoratorSerialization -> Unexpected serialized value");
+			AddErrorIfFalse(SharedDataB0->Vector == DecoratorAddRef0.Vector, "FAnimationAnimNextRuntimeTest_DecoratorSerialization -> Unexpected serialized value");
+			AddErrorIfFalse(FMemory::Memcmp(SharedDataB0->VectorArray, DecoratorAddRef0.VectorArray, sizeof(DecoratorAddRef0.VectorArray)) == 0, "FAnimationAnimNextRuntimeTest_DecoratorSerialization -> Unexpected serialized value");
+			AddErrorIfFalse(SharedDataB0->VectorTArray == DecoratorAddRef0.VectorTArray, "FAnimationAnimNextRuntimeTest_DecoratorSerialization -> Unexpected serialized value");
+			AddErrorIfFalse(SharedDataB0->String == DecoratorAddRef0.String, "FAnimationAnimNextRuntimeTest_DecoratorSerialization -> Unexpected serialized value");
+			AddErrorIfFalse(SharedDataB0->Name == DecoratorAddRef0.Name, "FAnimationAnimNextRuntimeTest_DecoratorSerialization -> Unexpected serialized value");
+		}
 
-			// Validate shared data for native decorator on node 0
-			{
-				TDecoratorBinding<IInterfaceC> BindingC0;
-				AddErrorIfFalse(Context.GetInterface(DecoratorPtr0, BindingC0), "FAnimationAnimNextRuntimeTest_DecoratorSerialization -> InterfaceC not found");
-
-				const auto* SharedDataC0 = BindingC0.GetSharedData<FDecoratorNativeSerialization_Add::FSharedData>();
-
-				AddErrorIfFalse(SharedDataC0->Integer == DecoratorNativeRef0.Integer, "FAnimationAnimNextRuntimeTest_DecoratorSerialization -> Unexpected serialized value");
-				AddErrorIfFalse(FMemory::Memcmp(SharedDataC0->IntegerArray, DecoratorNativeRef0.IntegerArray, sizeof(DecoratorNativeRef0.IntegerArray)) == 0, "FAnimationAnimNextRuntimeTest_DecoratorSerialization -> Unexpected serialized value");
-				AddErrorIfFalse(SharedDataC0->IntegerTArray == DecoratorNativeRef0.IntegerTArray, "FAnimationAnimNextRuntimeTest_DecoratorSerialization -> Unexpected serialized value");
-				AddErrorIfFalse(SharedDataC0->Vector == DecoratorNativeRef0.Vector, "FAnimationAnimNextRuntimeTest_DecoratorSerialization -> Unexpected serialized value");
-				AddErrorIfFalse(FMemory::Memcmp(SharedDataC0->VectorArray, DecoratorNativeRef0.VectorArray, sizeof(DecoratorNativeRef0.VectorArray)) == 0, "FAnimationAnimNextRuntimeTest_DecoratorSerialization -> Unexpected serialized value");
-				AddErrorIfFalse(SharedDataC0->VectorTArray == DecoratorNativeRef0.VectorTArray, "FAnimationAnimNextRuntimeTest_DecoratorSerialization -> Unexpected serialized value");
-				AddErrorIfFalse(SharedDataC0->String == DecoratorNativeRef0.String, "FAnimationAnimNextRuntimeTest_DecoratorSerialization -> Unexpected serialized value");
-				AddErrorIfFalse(SharedDataC0->Name == DecoratorNativeRef0.Name, "FAnimationAnimNextRuntimeTest_DecoratorSerialization -> Unexpected serialized value");
-				AddErrorIfFalse(SharedDataC0->bSerializeCalled, "FAnimationAnimNextRuntimeTest_DecoratorSerialization -> Unexpected serialized value");
-			}
+		// Validate shared data for native decorator on node 0
+		{
+			TDecoratorBinding<IInterfaceC> BindingC0;
+			AddErrorIfFalse(Context.GetInterface(DecoratorPtr0, BindingC0), "FAnimationAnimNextRuntimeTest_DecoratorSerialization -> InterfaceC not found");
+
+			const auto* SharedDataC0 = BindingC0.GetSharedData<FDecoratorNativeSerialization_Add::FSharedData>();
+
+			AddErrorIfFalse(SharedDataC0->Integer == DecoratorNativeRef0.Integer, "FAnimationAnimNextRuntimeTest_DecoratorSerialization -> Unexpected serialized value");
+			AddErrorIfFalse(FMemory::Memcmp(SharedDataC0->IntegerArray, DecoratorNativeRef0.IntegerArray, sizeof(DecoratorNativeRef0.IntegerArray)) == 0, "FAnimationAnimNextRuntimeTest_DecoratorSerialization -> Unexpected serialized value");
+			AddErrorIfFalse(SharedDataC0->IntegerTArray == DecoratorNativeRef0.IntegerTArray, "FAnimationAnimNextRuntimeTest_DecoratorSerialization -> Unexpected serialized value");
+			AddErrorIfFalse(SharedDataC0->Vector == DecoratorNativeRef0.Vector, "FAnimationAnimNextRuntimeTest_DecoratorSerialization -> Unexpected serialized value");
+			AddErrorIfFalse(FMemory::Memcmp(SharedDataC0->VectorArray, DecoratorNativeRef0.VectorArray, sizeof(DecoratorNativeRef0.VectorArray)) == 0, "FAnimationAnimNextRuntimeTest_DecoratorSerialization -> Unexpected serialized value");
+			AddErrorIfFalse(SharedDataC0->VectorTArray == DecoratorNativeRef0.VectorTArray, "FAnimationAnimNextRuntimeTest_DecoratorSerialization -> Unexpected serialized value");
+			AddErrorIfFalse(SharedDataC0->String == DecoratorNativeRef0.String, "FAnimationAnimNextRuntimeTest_DecoratorSerialization -> Unexpected serialized value");
+			AddErrorIfFalse(SharedDataC0->Name == DecoratorNativeRef0.Name, "FAnimationAnimNextRuntimeTest_DecoratorSerialization -> Unexpected serialized value");
+			AddErrorIfFalse(SharedDataC0->bSerializeCalled, "FAnimationAnimNextRuntimeTest_DecoratorSerialization -> Unexpected serialized value");
+		}
 
-			// Validate shared data for base decorator on node 1
-			{
-				TDecoratorBinding<IInterfaceA> BindingA1;
-				AddErrorIfFalse(Context.GetInterface(DecoratorPtr1, BindingA1), "FAnimationAnimNextRuntimeTest_DecoratorSerialization -> InterfaceA not found");
-
-				const auto* SharedDataA1 = BindingA1.GetSharedData<FDecoratorSerialization_Base::FSharedData>();
-
-				AddErrorIfFalse(SharedDataA1->Integer == DecoratorBaseRef1.Integer, "FAnimationAnimNextRuntimeTest_DecoratorSerialization -> Unexpected serialized value");
-				AddErrorIfFalse(FMemory::Memcmp(SharedDataA1->IntegerArray, DecoratorBaseRef1.IntegerArray, sizeof(DecoratorBaseRef1.IntegerArray)) == 0, "FAnimationAnimNextRuntimeTest_DecoratorSerialization -> Unexpected serialized value");
-				AddErrorIfFalse(SharedDataA1->IntegerTArray == DecoratorBaseRef1.IntegerTArray, "FAnimationAnimNextRuntimeTest_DecoratorSerialization -> Unexpected serialized value");
-				AddErrorIfFalse(SharedDataA1->Vector == DecoratorBaseRef1.Vector, "FAnimationAnimNextRuntimeTest_DecoratorSerialization -> Unexpected serialized value");
-				AddErrorIfFalse(FMemory::Memcmp(SharedDataA1->VectorArray, DecoratorBaseRef1.VectorArray, sizeof(DecoratorBaseRef1.VectorArray)) == 0, "FAnimationAnimNextRuntimeTest_DecoratorSerialization -> Unexpected serialized value");
-				AddErrorIfFalse(SharedDataA1->VectorTArray == DecoratorBaseRef1.VectorTArray, "FAnimationAnimNextRuntimeTest_DecoratorSerialization -> Unexpected serialized value");
-				AddErrorIfFalse(SharedDataA1->String == DecoratorBaseRef1.String, "FAnimationAnimNextRuntimeTest_DecoratorSerialization -> Unexpected serialized value");
-				AddErrorIfFalse(SharedDataA1->Name == DecoratorBaseRef1.Name, "FAnimationAnimNextRuntimeTest_DecoratorSerialization -> Unexpected serialized value");
-			}
+		// Validate shared data for base decorator on node 1
+		{
+			TDecoratorBinding<IInterfaceA> BindingA1;
+			AddErrorIfFalse(Context.GetInterface(DecoratorPtr1, BindingA1), "FAnimationAnimNextRuntimeTest_DecoratorSerialization -> InterfaceA not found");
+
+			const auto* SharedDataA1 = BindingA1.GetSharedData<FDecoratorSerialization_Base::FSharedData>();
+
+			AddErrorIfFalse(SharedDataA1->Integer == DecoratorBaseRef1.Integer, "FAnimationAnimNextRuntimeTest_DecoratorSerialization -> Unexpected serialized value");
+			AddErrorIfFalse(FMemory::Memcmp(SharedDataA1->IntegerArray, DecoratorBaseRef1.IntegerArray, sizeof(DecoratorBaseRef1.IntegerArray)) == 0, "FAnimationAnimNextRuntimeTest_DecoratorSerialization -> Unexpected serialized value");
+			AddErrorIfFalse(SharedDataA1->IntegerTArray == DecoratorBaseRef1.IntegerTArray, "FAnimationAnimNextRuntimeTest_DecoratorSerialization -> Unexpected serialized value");
+			AddErrorIfFalse(SharedDataA1->Vector == DecoratorBaseRef1.Vector, "FAnimationAnimNextRuntimeTest_DecoratorSerialization -> Unexpected serialized value");
+			AddErrorIfFalse(FMemory::Memcmp(SharedDataA1->VectorArray, DecoratorBaseRef1.VectorArray, sizeof(DecoratorBaseRef1.VectorArray)) == 0, "FAnimationAnimNextRuntimeTest_DecoratorSerialization -> Unexpected serialized value");
+			AddErrorIfFalse(SharedDataA1->VectorTArray == DecoratorBaseRef1.VectorTArray, "FAnimationAnimNextRuntimeTest_DecoratorSerialization -> Unexpected serialized value");
+			AddErrorIfFalse(SharedDataA1->String == DecoratorBaseRef1.String, "FAnimationAnimNextRuntimeTest_DecoratorSerialization -> Unexpected serialized value");
+			AddErrorIfFalse(SharedDataA1->Name == DecoratorBaseRef1.Name, "FAnimationAnimNextRuntimeTest_DecoratorSerialization -> Unexpected serialized value");
+		}
 
-			// Validate shared data for additive decorator on node 1
-			{
-				TDecoratorBinding<IInterfaceB> BindingB1;
-				AddErrorIfFalse(Context.GetInterface(DecoratorPtr1, BindingB1), "FAnimationAnimNextRuntimeTest_DecoratorSerialization -> InterfaceB not found");
-
-				const auto* SharedDataB1 = BindingB1.GetSharedData<FDecoratorSerialization_Add::FSharedData>();
-
-				AddErrorIfFalse(SharedDataB1->Integer == DecoratorAddRef1.Integer, "FAnimationAnimNextRuntimeTest_DecoratorSerialization -> Unexpected serialized value");
-				AddErrorIfFalse(FMemory::Memcmp(SharedDataB1->IntegerArray, DecoratorAddRef1.IntegerArray, sizeof(DecoratorAddRef1.IntegerArray)) == 0, "FAnimationAnimNextRuntimeTest_DecoratorSerialization -> Unexpected serialized value");
-				AddErrorIfFalse(SharedDataB1->IntegerTArray == DecoratorAddRef1.IntegerTArray, "FAnimationAnimNextRuntimeTest_DecoratorSerialization -> Unexpected serialized value");
-				AddErrorIfFalse(SharedDataB1->Vector == DecoratorAddRef1.Vector, "FAnimationAnimNextRuntimeTest_DecoratorSerialization -> Unexpected serialized value");
-				AddErrorIfFalse(FMemory::Memcmp(SharedDataB1->VectorArray, DecoratorAddRef1.VectorArray, sizeof(DecoratorAddRef1.VectorArray)) == 0, "FAnimationAnimNextRuntimeTest_DecoratorSerialization -> Unexpected serialized value");
-				AddErrorIfFalse(SharedDataB1->VectorTArray == DecoratorAddRef1.VectorTArray, "FAnimationAnimNextRuntimeTest_DecoratorSerialization -> Unexpected serialized value");
-				AddErrorIfFalse(SharedDataB1->String == DecoratorAddRef1.String, "FAnimationAnimNextRuntimeTest_DecoratorSerialization -> Unexpected serialized value");
-				AddErrorIfFalse(SharedDataB1->Name == DecoratorAddRef1.Name, "FAnimationAnimNextRuntimeTest_DecoratorSerialization -> Unexpected serialized value");
-			}
+		// Validate shared data for additive decorator on node 1
+		{
+			TDecoratorBinding<IInterfaceB> BindingB1;
+			AddErrorIfFalse(Context.GetInterface(DecoratorPtr1, BindingB1), "FAnimationAnimNextRuntimeTest_DecoratorSerialization -> InterfaceB not found");
+
+			const auto* SharedDataB1 = BindingB1.GetSharedData<FDecoratorSerialization_Add::FSharedData>();
+
+			AddErrorIfFalse(SharedDataB1->Integer == DecoratorAddRef1.Integer, "FAnimationAnimNextRuntimeTest_DecoratorSerialization -> Unexpected serialized value");
+			AddErrorIfFalse(FMemory::Memcmp(SharedDataB1->IntegerArray, DecoratorAddRef1.IntegerArray, sizeof(DecoratorAddRef1.IntegerArray)) == 0, "FAnimationAnimNextRuntimeTest_DecoratorSerialization -> Unexpected serialized value");
+			AddErrorIfFalse(SharedDataB1->IntegerTArray == DecoratorAddRef1.IntegerTArray, "FAnimationAnimNextRuntimeTest_DecoratorSerialization -> Unexpected serialized value");
+			AddErrorIfFalse(SharedDataB1->Vector == DecoratorAddRef1.Vector, "FAnimationAnimNextRuntimeTest_DecoratorSerialization -> Unexpected serialized value");
+			AddErrorIfFalse(FMemory::Memcmp(SharedDataB1->VectorArray, DecoratorAddRef1.VectorArray, sizeof(DecoratorAddRef1.VectorArray)) == 0, "FAnimationAnimNextRuntimeTest_DecoratorSerialization -> Unexpected serialized value");
+			AddErrorIfFalse(SharedDataB1->VectorTArray == DecoratorAddRef1.VectorTArray, "FAnimationAnimNextRuntimeTest_DecoratorSerialization -> Unexpected serialized value");
+			AddErrorIfFalse(SharedDataB1->String == DecoratorAddRef1.String, "FAnimationAnimNextRuntimeTest_DecoratorSerialization -> Unexpected serialized value");
+			AddErrorIfFalse(SharedDataB1->Name == DecoratorAddRef1.Name, "FAnimationAnimNextRuntimeTest_DecoratorSerialization -> Unexpected serialized value");
+		}
 
-			// Validate shared data for native decorator on node 1
-			{
-				TDecoratorBinding<IInterfaceC> BindingC1;
-				AddErrorIfFalse(Context.GetInterface(DecoratorPtr1, BindingC1), "FAnimationAnimNextRuntimeTest_DecoratorSerialization -> InterfaceC not found");
-
-				const auto* SharedDataC1 = BindingC1.GetSharedData<FDecoratorNativeSerialization_Add::FSharedData>();
-
-				AddErrorIfFalse(SharedDataC1->Integer == DecoratorNativeRef1.Integer, "FAnimationAnimNextRuntimeTest_DecoratorSerialization -> Unexpected serialized value");
-				AddErrorIfFalse(FMemory::Memcmp(SharedDataC1->IntegerArray, DecoratorNativeRef1.IntegerArray, sizeof(DecoratorNativeRef1.IntegerArray)) == 0, "FAnimationAnimNextRuntimeTest_DecoratorSerialization -> Unexpected serialized value");
-				AddErrorIfFalse(SharedDataC1->IntegerTArray == DecoratorNativeRef1.IntegerTArray, "FAnimationAnimNextRuntimeTest_DecoratorSerialization -> Unexpected serialized value");
-				AddErrorIfFalse(SharedDataC1->Vector == DecoratorNativeRef1.Vector, "FAnimationAnimNextRuntimeTest_DecoratorSerialization -> Unexpected serialized value");
-				AddErrorIfFalse(FMemory::Memcmp(SharedDataC1->VectorArray, DecoratorNativeRef1.VectorArray, sizeof(DecoratorNativeRef1.VectorArray)) == 0, "FAnimationAnimNextRuntimeTest_DecoratorSerialization -> Unexpected serialized value");
-				AddErrorIfFalse(SharedDataC1->VectorTArray == DecoratorNativeRef1.VectorTArray, "FAnimationAnimNextRuntimeTest_DecoratorSerialization -> Unexpected serialized value");
-				AddErrorIfFalse(SharedDataC1->String == DecoratorNativeRef1.String, "FAnimationAnimNextRuntimeTest_DecoratorSerialization -> Unexpected serialized value");
-				AddErrorIfFalse(SharedDataC1->Name == DecoratorNativeRef1.Name, "FAnimationAnimNextRuntimeTest_DecoratorSerialization -> Unexpected serialized value");
-				AddErrorIfFalse(SharedDataC1->bSerializeCalled, "FAnimationAnimNextRuntimeTest_DecoratorSerialization -> Unexpected serialized value");
-			}
+		// Validate shared data for native decorator on node 1
+		{
+			TDecoratorBinding<IInterfaceC> BindingC1;
+			AddErrorIfFalse(Context.GetInterface(DecoratorPtr1, BindingC1), "FAnimationAnimNextRuntimeTest_DecoratorSerialization -> InterfaceC not found");
+
+			const auto* SharedDataC1 = BindingC1.GetSharedData<FDecoratorNativeSerialization_Add::FSharedData>();
+
+			AddErrorIfFalse(SharedDataC1->Integer == DecoratorNativeRef1.Integer, "FAnimationAnimNextRuntimeTest_DecoratorSerialization -> Unexpected serialized value");
+			AddErrorIfFalse(FMemory::Memcmp(SharedDataC1->IntegerArray, DecoratorNativeRef1.IntegerArray, sizeof(DecoratorNativeRef1.IntegerArray)) == 0, "FAnimationAnimNextRuntimeTest_DecoratorSerialization -> Unexpected serialized value");
+			AddErrorIfFalse(SharedDataC1->IntegerTArray == DecoratorNativeRef1.IntegerTArray, "FAnimationAnimNextRuntimeTest_DecoratorSerialization -> Unexpected serialized value");
+			AddErrorIfFalse(SharedDataC1->Vector == DecoratorNativeRef1.Vector, "FAnimationAnimNextRuntimeTest_DecoratorSerialization -> Unexpected serialized value");
+			AddErrorIfFalse(FMemory::Memcmp(SharedDataC1->VectorArray, DecoratorNativeRef1.VectorArray, sizeof(DecoratorNativeRef1.VectorArray)) == 0, "FAnimationAnimNextRuntimeTest_DecoratorSerialization -> Unexpected serialized value");
+			AddErrorIfFalse(SharedDataC1->VectorTArray == DecoratorNativeRef1.VectorTArray, "FAnimationAnimNextRuntimeTest_DecoratorSerialization -> Unexpected serialized value");
+			AddErrorIfFalse(SharedDataC1->String == DecoratorNativeRef1.String, "FAnimationAnimNextRuntimeTest_DecoratorSerialization -> Unexpected serialized value");
+			AddErrorIfFalse(SharedDataC1->Name == DecoratorNativeRef1.Name, "FAnimationAnimNextRuntimeTest_DecoratorSerialization -> Unexpected serialized value");
+			AddErrorIfFalse(SharedDataC1->bSerializeCalled, "FAnimationAnimNextRuntimeTest_DecoratorSerialization -> Unexpected serialized value");
 		}
 	}
 
+	Registry.Clear();
+
+	AddErrorIfFalse(Registry.GetNum() == 0, "FAnimationAnimNextRuntimeTest_DecoratorSerialization -> Registry should contain 0 templates");
+
 	return true;
 }
 
diff --git a/Engine/Plugins/Experimental/AnimNext/Source/AnimNextTestSuite/Private/AnimNextDecoratorInterfacesTest.cpp b/Engine/Plugins/Experimental/AnimNext/Source/AnimNextTestSuite/Private/AnimNextDecoratorInterfacesTest.cpp
--- a/Engine/Plugins/Experimental/AnimNext/Source/AnimNextTestSuite/Private/AnimNextDecoratorInterfacesTest.cpp
+++ b/Engine/Plugins/Experimental/AnimNext/Source/AnimNextTestSuite/Private/AnimNextDecoratorInterfacesTest.cpp
@@ -16,7 +16,6 @@
 #include "DecoratorBase/ExecutionContext.h"
 #include "DecoratorBase/IDecoratorInterface.h"
 #include "DecoratorBase/NodeInstance.h"
-#include "DecoratorBase/NodeTemplateBuilder.h"
 #include "DecoratorBase/NodeTemplateRegistry.h"
 #include "DecoratorInterfaces/IEvaluate.h"
 #include "DecoratorInterfaces/IHierarchy.h"
@@ -213,7 +212,6 @@ bool FAnimationAnimNextRuntimeTest_IHierarchy::RunTest(const FString& InParamete
 	AUTO_REGISTER_ANIM_DECORATOR(FDecoratorWithOneChild)
 	AUTO_REGISTER_ANIM_DECORATOR(FDecoratorWithChildren)
 
-	FScopedClearNodeTemplateRegistry ScopedClearNodeTemplateRegistry;
 	FNodeTemplateRegistry& Registry = FNodeTemplateRegistry::Get();
 
 	// We create a few node templates
@@ -236,10 +234,10 @@ bool FAnimationAnimNextRuntimeTest_IHierarchy::RunTest(const FString& InParamete
 
 	// Populate our node template registry
 	TArray<uint8> NodeTemplateBufferA, NodeTemplateBufferB, NodeTemplateBufferC, NodeTemplateBufferD;
-	const FNodeTemplate* NodeTemplateA = FNodeTemplateBuilder::BuildNodeTemplate(NodeTemplateDecoratorListA, NodeTemplateBufferA);
-	const FNodeTemplate* NodeTemplateB = FNodeTemplateBuilder::BuildNodeTemplate(NodeTemplateDecoratorListB, NodeTemplateBufferB);
-	const FNodeTemplate* NodeTemplateC = FNodeTemplateBuilder::BuildNodeTemplate(NodeTemplateDecoratorListC, NodeTemplateBufferC);
-	const FNodeTemplate* NodeTemplateD = FNodeTemplateBuilder::BuildNodeTemplate(NodeTemplateDecoratorListD, NodeTemplateBufferD);
+	const FNodeTemplate* NodeTemplateA = BuildNodeTemplate(NodeTemplateDecoratorListA, NodeTemplateBufferA);
+	const FNodeTemplate* NodeTemplateB = BuildNodeTemplate(NodeTemplateDecoratorListB, NodeTemplateBufferB);
+	const FNodeTemplate* NodeTemplateC = BuildNodeTemplate(NodeTemplateDecoratorListC, NodeTemplateBufferC);
+	const FNodeTemplate* NodeTemplateD = BuildNodeTemplate(NodeTemplateDecoratorListD, NodeTemplateBufferD);
 
 	// Build our graph, it as follow (each node template has a single node instance):
 	// NodeA has no children
@@ -282,24 +280,24 @@ bool FAnimationAnimNextRuntimeTest_IHierarchy::RunTest(const FString& InParamete
 
 		DecoratorWriter.BeginNodeWriting();
 		DecoratorWriter.WriteNode(NodeA,
-			[&DecoratorPropertiesA](uint32 DecoratorIndex, const FString& PropertyName)
+			[&DecoratorPropertiesA](uint32 DecoratorIndex) -> const TMap<FString, FString>&
 			{
-				return DecoratorPropertiesA[DecoratorIndex][PropertyName];
+				return DecoratorPropertiesA[DecoratorIndex];
 			});
 		DecoratorWriter.WriteNode(NodeB,
-			[&DecoratorPropertiesB](uint32 DecoratorIndex, const FString& PropertyName)
+			[&DecoratorPropertiesB](uint32 DecoratorIndex) -> const TMap<FString, FString>&
 			{
-				return DecoratorPropertiesB[DecoratorIndex][PropertyName];
+				return DecoratorPropertiesB[DecoratorIndex];
 			});
 		DecoratorWriter.WriteNode(NodeC,
-			[&DecoratorPropertiesC](uint32 DecoratorIndex, const FString& PropertyName)
+			[&DecoratorPropertiesC](uint32 DecoratorIndex) -> const TMap<FString, FString>&
 			{
-				return DecoratorPropertiesC[DecoratorIndex][PropertyName];
+				return DecoratorPropertiesC[DecoratorIndex];
 			});
 		DecoratorWriter.WriteNode(NodeD,
-			[&DecoratorPropertiesD](uint32 DecoratorIndex, const FString& PropertyName)
+			[&DecoratorPropertiesD](uint32 DecoratorIndex) -> const TMap<FString, FString>&
 			{
-				return DecoratorPropertiesD[DecoratorIndex][PropertyName];
+				return DecoratorPropertiesD[DecoratorIndex];
 			});
 		DecoratorWriter.EndNodeWriting();
 
@@ -309,18 +307,11 @@ bool FAnimationAnimNextRuntimeTest_IHierarchy::RunTest(const FString& InParamete
 
 	// Read our graph
 	TArray<uint8> GraphSharedDataBuffer;
-	TArray<TObjectPtr<UObject>> TrackedObjectsForGC;
 	{
 		FMemoryReader GraphSharedDataArchive(GraphSharedDataArchiveBuffer);
 		FDecoratorReader DecoratorReader(GraphSharedDataArchive);
 
-		FDecoratorReader::EErrorState ErrorState = DecoratorReader.ReadGraph(GraphSharedDataBuffer, TrackedObjectsForGC);
-		AddErrorIfFalse(ErrorState == FDecoratorReader::EErrorState::None, "FAnimationAnimNextRuntimeTest_IHierarchy -> Failed to read graph shared data");
-
-		NodeA = DecoratorReader.ResolveNodeHandle(NodeA);
-		NodeB = DecoratorReader.ResolveNodeHandle(NodeB);
-		NodeC = DecoratorReader.ResolveNodeHandle(NodeC);
-		NodeD = DecoratorReader.ResolveNodeHandle(NodeD);
+		DecoratorReader.ReadGraphSharedData(GraphSharedDataBuffer);
 	}
 
 	FExecutionContext Context(GraphSharedDataBuffer);
@@ -407,7 +398,6 @@ bool FAnimationAnimNextRuntimeTest_IUpdate::RunTest(const FString& InParameters)
 	AUTO_REGISTER_ANIM_DECORATOR(FDecoratorWithOneChild)
 	AUTO_REGISTER_ANIM_DECORATOR(FDecoratorWithChildren)
 
-	FScopedClearNodeTemplateRegistry ScopedClearNodeTemplateRegistry;
 	FNodeTemplateRegistry& Registry = FNodeTemplateRegistry::Get();
 
 	// We create a few node templates
@@ -425,9 +415,9 @@ bool FAnimationAnimNextRuntimeTest_IUpdate::RunTest(const FString& InParameters)
 
 	// Populate our node template registry
 	TArray<uint8> NodeTemplateBufferA, NodeTemplateBufferB, NodeTemplateBufferC;
-	const FNodeTemplate* NodeTemplateA = FNodeTemplateBuilder::BuildNodeTemplate(NodeTemplateDecoratorListA, NodeTemplateBufferA);
-	const FNodeTemplate* NodeTemplateB = FNodeTemplateBuilder::BuildNodeTemplate(NodeTemplateDecoratorListB, NodeTemplateBufferB);
-	const FNodeTemplate* NodeTemplateC = FNodeTemplateBuilder::BuildNodeTemplate(NodeTemplateDecoratorListC, NodeTemplateBufferC);
+	const FNodeTemplate* NodeTemplateA = BuildNodeTemplate(NodeTemplateDecoratorListA, NodeTemplateBufferA);
+	const FNodeTemplate* NodeTemplateB = BuildNodeTemplate(NodeTemplateDecoratorListB, NodeTemplateBufferB);
+	const FNodeTemplate* NodeTemplateC = BuildNodeTemplate(NodeTemplateDecoratorListC, NodeTemplateBufferC);
 
 	// Build our graph, it as follow (each node template has a single node instance):
 	// NodeA has no children
@@ -462,19 +452,19 @@ bool FAnimationAnimNextRuntimeTest_IUpdate::RunTest(const FString& InParameters)
 
 		DecoratorWriter.BeginNodeWriting();
 		DecoratorWriter.WriteNode(NodeA,
-			[&DecoratorPropertiesA](uint32 DecoratorIndex, const FString& PropertyName)
+			[&DecoratorPropertiesA](uint32 DecoratorIndex) -> const TMap<FString, FString>&
 			{
-				return DecoratorPropertiesA[DecoratorIndex][PropertyName];
+				return DecoratorPropertiesA[DecoratorIndex];
 			});
 		DecoratorWriter.WriteNode(NodeB,
-			[&DecoratorPropertiesB](uint32 DecoratorIndex, const FString& PropertyName)
+			[&DecoratorPropertiesB](uint32 DecoratorIndex) -> const TMap<FString, FString>&
 			{
-				return DecoratorPropertiesB[DecoratorIndex][PropertyName];
+				return DecoratorPropertiesB[DecoratorIndex];
 			});
 		DecoratorWriter.WriteNode(NodeC,
-			[&DecoratorPropertiesC](uint32 DecoratorIndex, const FString& PropertyName)
+			[&DecoratorPropertiesC](uint32 DecoratorIndex) -> const TMap<FString, FString>&
 			{
-				return DecoratorPropertiesC[DecoratorIndex][PropertyName];
+				return DecoratorPropertiesC[DecoratorIndex];
 			});
 		DecoratorWriter.EndNodeWriting();
 
@@ -484,17 +474,11 @@ bool FAnimationAnimNextRuntimeTest_IUpdate::RunTest(const FString& InParameters)
 
 	// Read our graph
 	TArray<uint8> GraphSharedDataBuffer;
-	TArray<TObjectPtr<UObject>> TrackedObjectsForGC;
 	{
 		FMemoryReader GraphSharedDataArchive(GraphSharedDataArchiveBuffer);
 		FDecoratorReader DecoratorReader(GraphSharedDataArchive);
 
-		FDecoratorReader::EErrorState ErrorState = DecoratorReader.ReadGraph(GraphSharedDataBuffer, TrackedObjectsForGC);
-		AddErrorIfFalse(ErrorState == FDecoratorReader::EErrorState::None, "FAnimationAnimNextRuntimeTest_IUpdate -> Failed to read graph shared data");
-
-		NodeA = DecoratorReader.ResolveNodeHandle(NodeA);
-		NodeB = DecoratorReader.ResolveNodeHandle(NodeB);
-		NodeC = DecoratorReader.ResolveNodeHandle(NodeC);
+		DecoratorReader.ReadGraphSharedData(GraphSharedDataBuffer);
 	}
 
 	FExecutionContext Context(GraphSharedDataBuffer);
@@ -511,7 +495,7 @@ bool FAnimationAnimNextRuntimeTest_IUpdate::RunTest(const FString& InParameters)
 		AddErrorIfFalse(NodeCPtr.IsValid(), "FAnimationAnimNextRuntimeTest_IUpdate -> Failed to allocate root node instance");
 
 		// Call pre/post update on our graph
-		UpdateGraph(Context, NodeCPtr, 0.0333f);
+		UpdateGraph(Context, NodeCPtr);
 
 		AddErrorIfFalse(UpdatedDecorators.Num() == 6, "FAnimationAnimNextRuntimeTest_IUpdate -> Expected 6 nodes to have been visited during the update traversal");
 		AddErrorIfFalse(UpdatedDecorators[0] == FDecoratorWithChildren::DecoratorUID, "FAnimationAnimNextRuntimeTest_IUpdate -> Unexpected update order");		// NodeC
@@ -543,7 +527,6 @@ bool FAnimationAnimNextRuntimeTest_IEvaluate::RunTest(const FString& InParameter
 	AUTO_REGISTER_ANIM_DECORATOR(FDecoratorWithOneChild)
 	AUTO_REGISTER_ANIM_DECORATOR(FDecoratorWithChildren)
 
-	FScopedClearNodeTemplateRegistry ScopedClearNodeTemplateRegistry;
 	FNodeTemplateRegistry& Registry = FNodeTemplateRegistry::Get();
 
 	// We create a few node templates
@@ -561,9 +544,9 @@ bool FAnimationAnimNextRuntimeTest_IEvaluate::RunTest(const FString& InParameter
 
 	// Populate our node template registry
 	TArray<uint8> NodeTemplateBufferA, NodeTemplateBufferB, NodeTemplateBufferC;
-	const FNodeTemplate* NodeTemplateA = FNodeTemplateBuilder::BuildNodeTemplate(NodeTemplateDecoratorListA, NodeTemplateBufferA);
-	const FNodeTemplate* NodeTemplateB = FNodeTemplateBuilder::BuildNodeTemplate(NodeTemplateDecoratorListB, NodeTemplateBufferB);
-	const FNodeTemplate* NodeTemplateC = FNodeTemplateBuilder::BuildNodeTemplate(NodeTemplateDecoratorListC, NodeTemplateBufferC);
+	const FNodeTemplate* NodeTemplateA = BuildNodeTemplate(NodeTemplateDecoratorListA, NodeTemplateBufferA);
+	const FNodeTemplate* NodeTemplateB = BuildNodeTemplate(NodeTemplateDecoratorListB, NodeTemplateBufferB);
+	const FNodeTemplate* NodeTemplateC = BuildNodeTemplate(NodeTemplateDecoratorListC, NodeTemplateBufferC);
 
 	// Build our graph, it as follow (each node template has a single node instance):
 	// NodeA has no children
@@ -599,19 +582,19 @@ bool FAnimationAnimNextRuntimeTest_IEvaluate::RunTest(const FString& InParameter
 
 		DecoratorWriter.BeginNodeWriting();
 		DecoratorWriter.WriteNode(NodeA,
-			[&DecoratorPropertiesA](uint32 DecoratorIndex, const FString& PropertyName)
+			[&DecoratorPropertiesA](uint32 DecoratorIndex) -> const TMap<FString, FString>&
 			{
-				return DecoratorPropertiesA[DecoratorIndex][PropertyName];
+				return DecoratorPropertiesA[DecoratorIndex];
 			});
 		DecoratorWriter.WriteNode(NodeB,
-			[&DecoratorPropertiesB](uint32 DecoratorIndex, const FString& PropertyName)
+			[&DecoratorPropertiesB](uint32 DecoratorIndex) -> const TMap<FString, FString>&
 			{
-				return DecoratorPropertiesB[DecoratorIndex][PropertyName];
+				return DecoratorPropertiesB[DecoratorIndex];
 			});
 		DecoratorWriter.WriteNode(NodeC,
-			[&DecoratorPropertiesC](uint32 DecoratorIndex, const FString& PropertyName)
+			[&DecoratorPropertiesC](uint32 DecoratorIndex) -> const TMap<FString, FString>&
 			{
-				return DecoratorPropertiesC[DecoratorIndex][PropertyName];
+				return DecoratorPropertiesC[DecoratorIndex];
 			});
 		DecoratorWriter.EndNodeWriting();
 
@@ -621,17 +604,11 @@ bool FAnimationAnimNextRuntimeTest_IEvaluate::RunTest(const FString& InParameter
 
 	// Read our graph
 	TArray<uint8> GraphSharedDataBuffer;
-	TArray<TObjectPtr<UObject>> TrackedObjectsForGC;
 	{
 		FMemoryReader GraphSharedDataArchive(GraphSharedDataArchiveBuffer);
 		FDecoratorReader DecoratorReader(GraphSharedDataArchive);
 
-		FDecoratorReader::EErrorState ErrorState = DecoratorReader.ReadGraph(GraphSharedDataBuffer, TrackedObjectsForGC);
-		AddErrorIfFalse(ErrorState == FDecoratorReader::EErrorState::None, "FAnimationAnimNextRuntimeTest_IEvaluate -> Failed to read graph shared data");
-
-		NodeA = DecoratorReader.ResolveNodeHandle(NodeA);
-		NodeB = DecoratorReader.ResolveNodeHandle(NodeB);
-		NodeC = DecoratorReader.ResolveNodeHandle(NodeC);
+		DecoratorReader.ReadGraphSharedData(GraphSharedDataBuffer);
 	}
 
 	FExecutionContext Context(GraphSharedDataBuffer);
diff --git a/Engine/Plugins/Experimental/AnimNext/Source/AnimNextTestSuite/Private/AnimNextRuntimeTest.cpp b/Engine/Plugins/Experimental/AnimNext/Source/AnimNextTestSuite/Private/AnimNextRuntimeTest.cpp
--- a/Engine/Plugins/Experimental/AnimNext/Source/AnimNextTestSuite/Private/AnimNextRuntimeTest.cpp
+++ b/Engine/Plugins/Experimental/AnimNext/Source/AnimNextTestSuite/Private/AnimNextRuntimeTest.cpp
@@ -7,18 +7,112 @@
 
 #if WITH_DEV_AUTOMATION_TESTS
 
+#include "DecoratorBase/Decorator.h"
+#include "DecoratorBase/DecoratorRegistry.h"
+#include "DecoratorBase/NodeDescription.h"
+#include "DecoratorBase/NodeInstance.h"
+
 namespace UE::AnimNext
 {
-	FScopedClearNodeTemplateRegistry::FScopedClearNodeTemplateRegistry()
+	uint32 GetNodeTemplateUID(const TArray<FDecoratorUID>& NodeTemplateDecoratorList)
+	{
+		uint32 NodeTemplateUID = 0;
+		for (const FDecoratorUID DecoratorUID : NodeTemplateDecoratorList)
+		{
+			NodeTemplateUID = HashCombineFast(NodeTemplateUID, DecoratorUID.GetUID());
+		}
+
+		return NodeTemplateUID;
+	}
+
+	uint32 GetNodeInstanceSize(const TArray<FDecoratorUID>& NodeTemplateDecoratorList)
+	{
+		const FDecoratorRegistry& Registry = FDecoratorRegistry::Get();
+
+		uint32 NodeInstanceSize = sizeof(FNodeInstance);
+		for (const FDecoratorUID DecoratorUID : NodeTemplateDecoratorList)
+		{
+			const FDecorator* Decorator = Registry.Find(DecoratorUID);
+
+			const FDecoratorMemoryLayout MemoryLayout = Decorator->GetDecoratorMemoryDescription();
+			NodeInstanceSize = Align(NodeInstanceSize, MemoryLayout.InstanceDataAlignment);
+			NodeInstanceSize += MemoryLayout.InstanceDataSize;
+		}
+
+		return NodeInstanceSize;
+	}
+
+	void AppendTemplateDecorator(const TArray<FDecoratorUID>& NodeTemplateDecoratorList, int32 DecoratorIndex, TArray<uint8>& NodeTemplateBuffer, uint32& SharedDataOffset, uint32& InstanceDataOffset)
 	{
-		FNodeTemplateRegistry& Registry = FNodeTemplateRegistry::Get();
-		Swap(Registry, TmpRegistry);
+		const FDecoratorRegistry& Registry = FDecoratorRegistry::Get();
+
+		const FDecoratorUID DecoratorUID = NodeTemplateDecoratorList[DecoratorIndex];
+		const FDecoratorRegistryHandle DecoratorHandle = Registry.FindHandle(DecoratorUID);
+		const FDecorator* Decorator = Registry.Find(DecoratorHandle);
+		const EDecoratorMode DecoratorMode = Decorator->GetDecoratorMode();
+
+		uint32 AdditiveIndexOrNumAdditive;
+		if (DecoratorMode == EDecoratorMode::Base)
+		{
+			// Find out how many additive decorators we have
+			AdditiveIndexOrNumAdditive = 0;
+			for (int32 Index = DecoratorIndex + 1; Index < NodeTemplateDecoratorList.Num(); ++Index)	// Skip ourself
+			{
+				const FDecorator* ChildDecorator = Registry.Find(NodeTemplateDecoratorList[Index]);
+				if (ChildDecorator->GetDecoratorMode() == EDecoratorMode::Base)
+				{
+					break;	// Found another base decorator, we are done
+				}
+
+				// We are additive
+				AdditiveIndexOrNumAdditive++;
+			}
+		}
+		else
+		{
+			// Find out our additive index
+			AdditiveIndexOrNumAdditive = 1;	// Skip ourself
+			for (int32 Index = DecoratorIndex - 1; Index >= 0; --Index)
+			{
+				const FDecorator* ParentDecorator = Registry.Find(NodeTemplateDecoratorList[Index]);
+				if (ParentDecorator->GetDecoratorMode() == EDecoratorMode::Base)
+				{
+					break;	// Found our base decorator, we are done
+				}
+
+				// We are additive
+				AdditiveIndexOrNumAdditive++;
+			}
+		}
+
+		const FDecoratorMemoryLayout MemoryLayout = Decorator->GetDecoratorMemoryDescription();
+
+		// Align our data
+		SharedDataOffset = Align(SharedDataOffset, MemoryLayout.SharedDataAlignment);
+		InstanceDataOffset = Align(InstanceDataOffset, MemoryLayout.InstanceDataAlignment);
+
+		// Append and update our offsets
+		const int32 BufferIndex = NodeTemplateBuffer.AddUninitialized(sizeof(FDecoratorTemplate));
+		new(&NodeTemplateBuffer[BufferIndex]) FDecoratorTemplate(DecoratorUID, DecoratorHandle, DecoratorMode, AdditiveIndexOrNumAdditive, SharedDataOffset, InstanceDataOffset);
+
+		SharedDataOffset += MemoryLayout.SharedDataSize;
+		InstanceDataOffset += MemoryLayout.InstanceDataSize;
 	}
 
-	FScopedClearNodeTemplateRegistry::~FScopedClearNodeTemplateRegistry()
+	const FNodeTemplate* BuildNodeTemplate(const TArray<FDecoratorUID>& NodeTemplateDecoratorList, TArray<uint8>& NodeTemplateBuffer)
 	{
-		FNodeTemplateRegistry& Registry = FNodeTemplateRegistry::Get();
-		Swap(Registry, TmpRegistry);
+		const int32 BufferIndex = NodeTemplateBuffer.AddUninitialized(sizeof(FNodeTemplate));
+		new(&NodeTemplateBuffer[BufferIndex]) FNodeTemplate(GetNodeTemplateUID(NodeTemplateDecoratorList), GetNodeInstanceSize(NodeTemplateDecoratorList), NodeTemplateDecoratorList.Num());
+
+		uint32 SharedDataOffset = sizeof(FNodeDescription);
+		uint32 InstanceDataOffset = sizeof(FNodeInstance);
+
+		for (int32 DecoratorIndex = 0; DecoratorIndex < NodeTemplateDecoratorList.Num(); ++DecoratorIndex)
+		{
+			AppendTemplateDecorator(NodeTemplateDecoratorList, DecoratorIndex, NodeTemplateBuffer, SharedDataOffset, InstanceDataOffset);
+		}
+
+		return reinterpret_cast<const FNodeTemplate*>(&NodeTemplateBuffer[0]);
 	}
 }
 #endif
diff --git a/Engine/Plugins/Experimental/AnimNext/Source/AnimNextTestSuite/Private/AnimNextRuntimeTest.h b/Engine/Plugins/Experimental/AnimNext/Source/AnimNextTestSuite/Private/AnimNextRuntimeTest.h
--- a/Engine/Plugins/Experimental/AnimNext/Source/AnimNextTestSuite/Private/AnimNextRuntimeTest.h
+++ b/Engine/Plugins/Experimental/AnimNext/Source/AnimNextTestSuite/Private/AnimNextRuntimeTest.h
@@ -11,70 +11,64 @@
 
 #include "DecoratorBase/DecoratorUID.h"
 #include "DecoratorBase/NodeTemplate.h"
-#include "DecoratorBase/NodeTemplateRegistry.h"
 
 namespace UE::AnimNext
 {
+	// Builds a node template from a list of decorator UID
+	const FNodeTemplate* BuildNodeTemplate(const TArray<FDecoratorUID>& NodeTemplateDecoratorList, TArray<uint8>& NodeTemplateBuffer);
+
 	// Converts a property value into its string representation using UE reflection
 	template<class DecoratorSharedDataType, typename PropertyType>
 	static FString ToString(const FString& PropertyName, PropertyType PropertyValue)
 	{
 		const UScriptStruct* SharedDataStruct = DecoratorSharedDataType::StaticStruct();
-		if (const FProperty* Property = SharedDataStruct->FindPropertyByName(*PropertyName))
-		{
-			void* PropertyDefaults = Property->AllocateAndInitializeValue();
-
-			FString Result;
 
-			if constexpr (std::is_pointer_v<PropertyType>)
+		for (const FProperty* Property = SharedDataStruct->PropertyLink; Property != nullptr; Property = Property->PropertyLinkNext)
+		{
+			if (Property->GetName() == PropertyName)
 			{
-				// C-style array properties aren't handled by ExportText, we need to handle it manually
+				void* PropertyDefaults = Property->AllocateAndInitializeValue();
 
-				const bool bIsCArray = Property->ArrayDim > 1;
-				if (bIsCArray)
-				{
-					Result += TEXT("(");
-				}
+				FString Result;
 
-				for (int32 Index = 0; Index < Property->ArrayDim; ++Index)
+				if constexpr (std::is_pointer_v<PropertyType>)
 				{
-					Property->ExportText_Direct(Result, PropertyValue + Index, (PropertyType)PropertyDefaults + Index, nullptr, PPF_None, nullptr);
+					// C-style array properties aren't handled by ExportText, we need to handle it manually
 
-					if (Index + 1 < Property->ArrayDim)
+					const bool bIsCArray = Property->ArrayDim > 1;
+					if (bIsCArray)
 					{
-						Result += TEXT(",");
+						Result += TEXT("(");
+					}
+
+					for (int32 Index = 0; Index < Property->ArrayDim; ++Index)
+					{
+						Property->ExportText_Direct(Result, PropertyValue + Index, (PropertyType)PropertyDefaults + Index, nullptr, PPF_None, nullptr);
+
+						if (Index + 1 < Property->ArrayDim)
+						{
+							Result += TEXT(",");
+						}
 					}
-				}
 
-				if (bIsCArray)
+					if (bIsCArray)
+					{
+						Result += TEXT(")");
+					}
+				}
+				else
 				{
-					Result += TEXT(")");
+					Property->ExportText_Direct(Result, &PropertyValue, PropertyDefaults, nullptr, PPF_None);
 				}
-			}
-			else
-			{
-				Property->ExportText_Direct(Result, &PropertyValue, PropertyDefaults, nullptr, PPF_None);
-			}
 
-			Property->DestroyAndFreeValue(PropertyDefaults);
+				Property->DestroyAndFreeValue(PropertyDefaults);
 
-			return Result;
+				return Result;
+			}
 		}
 
 		return FString();
 	}
-
-	// Creates a temporary empty node template registry instance and swaps it for the current one
-	struct FScopedClearNodeTemplateRegistry final
-	{
-		FScopedClearNodeTemplateRegistry(const FScopedClearNodeTemplateRegistry&) = delete;
-		FScopedClearNodeTemplateRegistry& operator=(const FScopedClearNodeTemplateRegistry&) = delete;
-
-		FScopedClearNodeTemplateRegistry();
-		~FScopedClearNodeTemplateRegistry();
-
-		FNodeTemplateRegistry TmpRegistry;
-	};
 }
 
 #endif
diff --git a/Engine/Plugins/Experimental/AnimNext/Source/AnimNextUncookedOnly/AnimNextUncookedOnly.Build.cs b/Engine/Plugins/Experimental/AnimNext/Source/AnimNextUncookedOnly/AnimNextUncookedOnly.Build.cs
--- a/Engine/Plugins/Experimental/AnimNext/Source/AnimNextUncookedOnly/AnimNextUncookedOnly.Build.cs
+++ b/Engine/Plugins/Experimental/AnimNext/Source/AnimNextUncookedOnly/AnimNextUncookedOnly.Build.cs
@@ -29,8 +29,7 @@ namespace UnrealBuildTool.Rules
 					"Kismet",
 					"Slate",
 					"SlateCore",
-					"StructUtils",
-					"ToolMenus",
+					"StructUtils"
 				}
 			);
 
diff --git a/Engine/Plugins/Experimental/AnimNext/Source/AnimNextUncookedOnly/Internal/Graph/AnimNextGraph_EdGraph.h b/Engine/Plugins/Experimental/AnimNext/Source/AnimNextUncookedOnly/Internal/Graph/AnimNextGraph_EdGraph.h
--- a/Engine/Plugins/Experimental/AnimNext/Source/AnimNextUncookedOnly/Internal/Graph/AnimNextGraph_EdGraph.h
+++ b/Engine/Plugins/Experimental/AnimNext/Source/AnimNextUncookedOnly/Internal/Graph/AnimNextGraph_EdGraph.h
@@ -3,7 +3,7 @@
 #pragma once
 
 #include "CoreMinimal.h"
-#include "EdGraph/RigVMEdGraph.h"
+#include "Graph/ControlRigGraph.h"
 #include "AnimNextGraph_EdGraph.generated.h"
 
 class UAnimNextGraph_EditorData;
@@ -13,17 +13,14 @@ namespace UE::AnimNext::UncookedOnly
 	struct FUtils;
 }
 
-/**
-  * Wraps UEdGraph which represents the node graph
-  */
 UCLASS(MinimalAPI)
-class UAnimNextGraph_EdGraph : public URigVMEdGraph
+class UAnimNextGraph_EdGraph : public UControlRigGraph
 {
 	GENERATED_BODY()
 
 	friend class UAnimNextGraph_EditorData;
 
-	// URigVMEdGraph interface
+	// UControlRigGraph interface
 	virtual FRigVMClient* GetRigVMClient() const override;
 	
 	void Initialize(UAnimNextGraph_EditorData* InEditorData);
diff --git a/Engine/Plugins/Experimental/AnimNext/Source/AnimNextUncookedOnly/Internal/Graph/AnimNextGraph_EdGraphNode.h b/Engine/Plugins/Experimental/AnimNext/Source/AnimNextUncookedOnly/Internal/Graph/AnimNextGraph_EdGraphNode.h
--- a/Engine/Plugins/Experimental/AnimNext/Source/AnimNextUncookedOnly/Internal/Graph/AnimNextGraph_EdGraphNode.h
+++ b/Engine/Plugins/Experimental/AnimNext/Source/AnimNextUncookedOnly/Internal/Graph/AnimNextGraph_EdGraphNode.h
@@ -3,29 +3,11 @@
 #pragma once
 
 #include "CoreMinimal.h"
-#include "EdGraph/RigVMEdGraphNode.h"
-
+#include "Graph/ControlRigGraphNode.h"
 #include "AnimNextGraph_EdGraphNode.generated.h"
 
-// EdGraphNode representation for AnimNext nodes
-// A node can hold a decorator stack or a decorator entry
 UCLASS(MinimalAPI)
-class UAnimNextGraph_EdGraphNode : public URigVMEdGraphNode
+class UAnimNextGraph_EdGraphNode : public UControlRigGraphNode
 {
 	GENERATED_BODY()
-
-public:
-	//////////////////////////////////////////////////////////////////////////
-	// UEdGraphNode implementation
-	virtual void GetNodeContextMenuActions(class UToolMenu* Menu, class UGraphNodeContextMenuContext* Context) const override;
-
-	//////////////////////////////////////////////////////////////////////////
-	// Our implementation
-
-	// Returns whether this node is a decorator stack or not
-	ANIMNEXTUNCOOKEDONLY_API bool IsDecoratorStack() const;
-
-private:
-	// Populates the SubMenu with entries for each decorator that can be added through the context menu
-	void BuildAddDecoratorContextMenu(class UToolMenu* SubMenu);
-};
+};
\ No newline at end of file
diff --git a/Engine/Plugins/Experimental/AnimNext/Source/AnimNextUncookedOnly/Internal/Graph/AnimNextGraph_EdGraphSchema.h b/Engine/Plugins/Experimental/AnimNext/Source/AnimNextUncookedOnly/Internal/Graph/AnimNextGraph_EdGraphSchema.h
--- a/Engine/Plugins/Experimental/AnimNext/Source/AnimNextUncookedOnly/Internal/Graph/AnimNextGraph_EdGraphSchema.h
+++ b/Engine/Plugins/Experimental/AnimNext/Source/AnimNextUncookedOnly/Internal/Graph/AnimNextGraph_EdGraphSchema.h
@@ -3,16 +3,16 @@
 #pragma once
 
 #include "CoreMinimal.h"
-#include "EdGraph/RigVMEdGraphSchema.h"
+#include "Graph/ControlRigGraphSchema.h"
 #include "AnimNextGraph_EdGraphNode.h"
 #include "AnimNextGraph_EdGraphSchema.generated.h"
 
 UCLASS(MinimalAPI)
-class UAnimNextGraph_EdGraphSchema : public URigVMEdGraphSchema
+class UAnimNextGraph_EdGraphSchema : public UControlRigGraphSchema
 {
 	GENERATED_BODY()
 
-	// URigVMEdGraphSchema interface
+	// UControlRigGraphSchema interface
 	virtual TSubclassOf<URigVMEdGraphNode> GetGraphNodeClass(const URigVMEdGraph* InGraph) const override { return UAnimNextGraph_EdGraphNode::StaticClass(); }
 
 	// UEdGraphSchema interface
diff --git a/Engine/Plugins/Experimental/AnimNext/Source/AnimNextUncookedOnly/Internal/Graph/AnimNextGraph_EditorData.h b/Engine/Plugins/Experimental/AnimNext/Source/AnimNextUncookedOnly/Internal/Graph/AnimNextGraph_EditorData.h
--- a/Engine/Plugins/Experimental/AnimNext/Source/AnimNextUncookedOnly/Internal/Graph/AnimNextGraph_EditorData.h
+++ b/Engine/Plugins/Experimental/AnimNext/Source/AnimNextUncookedOnly/Internal/Graph/AnimNextGraph_EditorData.h
@@ -27,31 +27,17 @@ enum class EAnimNextGraphLoadType : uint8
 	CheckUserDefinedStructs
 };
 
-/**
- * The Schema is used to determine which actions are allowed
- * on a graph. This includes any topological change.
- */
 UCLASS()
 class UAnimNextGraph_Schema : public URigVMSchema
 {
 	GENERATED_BODY()
 };
 
-/**
-  * Implements a RigVM client host and RigVM graph function store host
-  * 
-  * A RigVM client holds a graph schema, multiple graph models and controllers, as well as an action stack, undo/redo
-  * information, and a function library.
-  * 
-  * A RigVM function store holds various RigVM functions that can be called and linked against. Functions are created
-  * in the editor UI under the 'My Blueprint' tab.
-  */
 UCLASS(MinimalAPI)
 class UAnimNextGraph_EditorData : public UObject, public IRigVMClientHost, public IRigVMGraphFunctionHost
 {
 	GENERATED_BODY()
 
-public:
 	UAnimNextGraph_EditorData(const FObjectInitializer& ObjectInitializer);
 	
 	friend class UAnimNextGraphFactory;
@@ -86,24 +72,14 @@ public:
 	virtual const FRigVMGraphFunctionStore* GetRigVMGraphFunctionStore() const override;
 
 	ANIMNEXTUNCOOKEDONLY_API void Initialize(bool bRecompileVM);
-	ANIMNEXTUNCOOKEDONLY_API UAnimNextGraph_EdGraph* GetRootGraph() const;
 
-protected:
-#if WITH_EDITOR
 	void RefreshAllModels(EAnimNextGraphLoadType InLoadType);
-	ANIMNEXTUNCOOKEDONLY_API void RebuildEdGraphFromModel();
-	void GetAllGraphs(TArray<UEdGraph*>& Graphs) const;
-#endif
 
 	void RecompileVM();
 	
 	void RecompileVMIfRequired();
 
 	void RequestAutoVMRecompilation();
-
-	void IncrementVMRecompileBracket();
-
-	void DecrementVMRecompileBracket();
 	
 	void HandleModifiedEvent(ERigVMGraphNotifType InNotifType, URigVMGraph* InGraph, UObject* InSubject);
 
@@ -146,9 +122,6 @@ protected:
 	UPROPERTY(transient, DuplicateTransient)
 	TMap<FString, FRigVMOperand> PinToOperandMap;
 
-	UPROPERTY(transient, DuplicateTransient)
-	int32 VMRecompilationBracket = 0;
-
 	UPROPERTY(transient, DuplicateTransient)
 	bool bVMRecompilationRequired = false;
 
diff --git a/Engine/Plugins/Experimental/AnimNext/Source/AnimNextUncookedOnly/Private/AnimNextUncookedOnlyModule.cpp b/Engine/Plugins/Experimental/AnimNext/Source/AnimNextUncookedOnly/Private/AnimNextUncookedOnlyModule.cpp
--- a/Engine/Plugins/Experimental/AnimNext/Source/AnimNextUncookedOnly/Private/AnimNextUncookedOnlyModule.cpp
+++ b/Engine/Plugins/Experimental/AnimNext/Source/AnimNextUncookedOnly/Private/AnimNextUncookedOnlyModule.cpp
@@ -1,21 +1,5 @@
 // Copyright Epic Games, Inc. All Rights Reserved.
 
-#include "CoreMinimal.h"
 #include "Modules/ModuleManager.h"
 
-namespace UE::AnimNext::UncookedOnly
-{
-	class FModule : public IModuleInterface
-	{
-	public:
-
-	private:
-		virtual void StartupModule() override;
-	};
-
-	void FModule::StartupModule()
-	{
-	}
-}
-
-IMPLEMENT_MODULE(UE::AnimNext::UncookedOnly::FModule, AnimNextUncookedOnly);
+IMPLEMENT_MODULE(FDefaultModuleImpl, AnimNextUncookedOnly);
\ No newline at end of file
diff --git a/Engine/Plugins/Experimental/AnimNext/Source/AnimNextUncookedOnly/Private/Graph/AnimNextGraph_EdGraphNode.cpp b/Engine/Plugins/Experimental/AnimNext/Source/AnimNextUncookedOnly/Private/Graph/AnimNextGraph_EdGraphNode.cpp
deleted file mode 100644
diff --git a/Engine/Plugins/Experimental/AnimNext/Source/AnimNextUncookedOnly/Private/Graph/AnimNextGraph_EditorData.cpp b/Engine/Plugins/Experimental/AnimNext/Source/AnimNextUncookedOnly/Private/Graph/AnimNextGraph_EditorData.cpp
--- a/Engine/Plugins/Experimental/AnimNext/Source/AnimNextUncookedOnly/Private/Graph/AnimNextGraph_EditorData.cpp
+++ b/Engine/Plugins/Experimental/AnimNext/Source/AnimNextUncookedOnly/Private/Graph/AnimNextGraph_EditorData.cpp
@@ -23,7 +23,6 @@ UAnimNextGraph_EditorData::UAnimNextGraph_EditorData(const FObjectInitializer& O
 	{
 		TGuardValue<bool> DisableClientNotifs(RigVMClient.bSuspendNotifications, true);
 		RigVMClient.AddModel(TEXT("RigVMGraph"), false, &ObjectInitializer);
-		RigVMClient.AddModel(TEXT("TempRigVMGraph"), false, &ObjectInitializer);	// Used during compilation, otherwise empty
 		RigVMClient.GetOrCreateFunctionLibrary(false, &ObjectInitializer);
 	}
 	RigVMClient.SetExecuteContextStruct(FAnimNextExecuteContext::StaticStruct());
@@ -63,7 +62,6 @@ void UAnimNextGraph_EditorData::Initialize(bool bRecompileVM)
 		
 		RigVMClient.GetOrCreateController(RigVMClient.GetDefaultModel());
 		RigVMClient.GetOrCreateController(RigVMClient.GetFunctionLibrary());
-		RigVMClient.GetOrCreateController(TEXT("TempRigVMGraph"));
 
 		// Init function library controllers
 		for(URigVMLibraryNode* LibraryNode : RigVMClient.GetFunctionLibrary()->GetFunctions())
@@ -85,11 +83,6 @@ void UAnimNextGraph_EditorData::Initialize(bool bRecompileVM)
 	}
 }
 
-UAnimNextGraph_EdGraph* UAnimNextGraph_EditorData::GetRootGraph() const
-{
-	return RootGraph;
-}
-
 void UAnimNextGraph_EditorData::PostLoad()
 {
 	Super::PostLoad();
@@ -240,57 +233,6 @@ void UAnimNextGraph_EditorData::RefreshAllModels(EAnimNextGraphLoadType InLoadTy
 	}
 }
 
-void UAnimNextGraph_EditorData::RebuildEdGraphFromModel()
-{
-	DECLARE_SCOPE_HIERARCHICAL_COUNTER_FUNC()
-
-	TGuardValue<bool> SelfGuard(bSuspendModelNotificationsForSelf, true);
-	TGuardValue<bool> ClientIgnoreModificationsGuard(RigVMClient.bIgnoreModelNotifications, true);
-
-	verify(RigVMClient.GetOrCreateController((URigVMGraph*)nullptr));
-
-	TArray<UEdGraph*> EdGraphs;
-	GetAllGraphs(EdGraphs);
-
-	for (UEdGraph* Graph : EdGraphs)
-	{
-		TArray<UEdGraphNode*> Nodes = Graph->Nodes;
-		for (UEdGraphNode* Node : Nodes)
-		{
-			Graph->RemoveNode(Node);
-		}
-	}
-
-	if (FunctionLibraryEdGraph && RigVMClient.GetFunctionLibrary())
-	{
-		FunctionLibraryEdGraph->ModelNodePath = RigVMClient.GetFunctionLibrary()->GetNodePath();
-	}
-
-	TArray<URigVMGraph*> RigGraphs = RigVMClient.GetAllModels(true, true);
-
-	for (int32 RigGraphIndex = 0; RigGraphIndex < RigGraphs.Num(); RigGraphIndex++)
-	{
-		RigVMClient.GetOrCreateController(RigGraphs[RigGraphIndex])->ResendAllNotifications();
-	}
-}
-
-void UAnimNextGraph_EditorData::GetAllGraphs(TArray<UEdGraph*>& Graphs) const
-{
-	Graphs.Reset();
-
-	Graphs.Add(RootGraph);
-
-	if (EntryPointGraph)
-	{
-		Graphs.Add(EntryPointGraph);
-	}
-
-	if (FunctionLibraryEdGraph)
-	{
-		Graphs.Add(FunctionLibraryEdGraph);
-	}
-}
-
 #endif // WITH_EDITOR
 
 FRigVMClient* UAnimNextGraph_EditorData::GetRigVMClient()
@@ -421,50 +363,18 @@ void UAnimNextGraph_EditorData::RecompileVMIfRequired()
 void UAnimNextGraph_EditorData::RequestAutoVMRecompilation()
 {
 	bVMRecompilationRequired = true;
-	if (bAutoRecompileVM && VMRecompilationBracket == 0)
+	if (bAutoRecompileVM)
 	{
 		RecompileVMIfRequired();
 	}
 }
 
-void UAnimNextGraph_EditorData::IncrementVMRecompileBracket()
-{
-	VMRecompilationBracket++;
-}
-
-void UAnimNextGraph_EditorData::DecrementVMRecompileBracket()
-{
-	if (VMRecompilationBracket == 1)
-	{
-		if (bAutoRecompileVM)
-		{
-			RecompileVMIfRequired();
-		}
-		VMRecompilationBracket = 0;
-	}
-	else if (VMRecompilationBracket > 0)
-	{
-		VMRecompilationBracket--;
-	}
-}
-
 void UAnimNextGraph_EditorData::HandleModifiedEvent(ERigVMGraphNotifType InNotifType, URigVMGraph* InGraph, UObject* InSubject)
 {
 	bool bNotifForOthersPending = true;
 
 	switch(InNotifType)
 	{
-	case ERigVMGraphNotifType::InteractionBracketOpened:
-		{
-			IncrementVMRecompileBracket();
-			break;
-		}
-	case ERigVMGraphNotifType::InteractionBracketClosed:
-	case ERigVMGraphNotifType::InteractionBracketCanceled:
-		{
-			DecrementVMRecompileBracket();
-			break;
-		}
 	case ERigVMGraphNotifType::NodeAdded:
 		{
 			if (URigVMCollapseNode* CollapseNode = Cast<URigVMCollapseNode>(InSubject))
@@ -479,20 +389,10 @@ void UAnimNextGraph_EditorData::HandleModifiedEvent(ERigVMGraphNotifType InNotif
 	case ERigVMGraphNotifType::PinArraySizeChanged:
 	case ERigVMGraphNotifType::PinDirectionChanged:
 		{
-			RequestAutoVMRecompilation();
-			break;
-		}
-	case ERigVMGraphNotifType::PinAdded:
-		{
-			if (URigVMPin* Pin = Cast<URigVMPin>(InSubject))
-			{
-				if (Pin->IsDecoratorPin())
-				{
-					RequestAutoVMRecompilation();
-				}
-			}
+			RecompileVM();
 			break;
 		}
+
 	case ERigVMGraphNotifType::PinDefaultValueChanged:
 		{
 			if (InGraph->GetRuntimeAST().IsValid())
@@ -503,14 +403,14 @@ void UAnimNextGraph_EditorData::HandleModifiedEvent(ERigVMGraphNotifType InNotif
 				if (Expression == nullptr)
 				{
 					InGraph->ClearAST();
+					break;
 				}
 				else if (Expression->NumParents() > 1)
 				{
 					InGraph->ClearAST();
+					break;
 				}
 			}
-
-			RequestAutoVMRecompilation();	// We need to rebuild our metadata when a default value changes
 			break;
 		}
 	}
diff --git a/Engine/Plugins/Experimental/AnimNext/Source/AnimNextUncookedOnly/Private/UncookedOnlyUtils.cpp b/Engine/Plugins/Experimental/AnimNext/Source/AnimNextUncookedOnly/Private/UncookedOnlyUtils.cpp
--- a/Engine/Plugins/Experimental/AnimNext/Source/AnimNextUncookedOnly/Private/UncookedOnlyUtils.cpp
+++ b/Engine/Plugins/Experimental/AnimNext/Source/AnimNextUncookedOnly/Private/UncookedOnlyUtils.cpp
@@ -3,240 +3,20 @@
 #include "UncookedOnlyUtils.h"
 #include "Graph/AnimNextGraph.h"
 #include "Graph/AnimNextGraph_EditorData.h"
-#include "Graph/RigUnit_AnimNextGraphRoot.h"
-#include "Graph/RigUnit_AnimNextDecoratorStack.h"
+#include "RigVMCompiler/RigVMCompiler.h"
+#include "RigVMCore/RigVM.h"
 #include "Graph/AnimNextExecuteContext.h"
 #include "Param/AnimNextParameter.h"
 #include "Param/AnimNextParameterBlock.h"
 #include "Param/AnimNextParameterBlock_EditorData.h"
 #include "Param/AnimNextParameterBlock_EdGraph.h"
-#include "Graph/RigDecorator_AnimNextCppDecorator.h"
 #include "Graph/RigUnit_AnimNextBeginExecution.h"
 #include "Param/RigVMDispatch_SetParameter.h"
 #include "Param/AnimNextParameterBlockEntry.h"
 #include "Param/IAnimNextParameterBlockBindingInterface.h"
-#include "DecoratorBase/DecoratorReader.h"
-#include "DecoratorBase/DecoratorWriter.h"
-#include "DecoratorBase/NodeTemplate.h"
-#include "DecoratorBase/NodeTemplateBuilder.h"
-#include "DecoratorBase/NodeTemplateRegistry.h"
-#include "DecoratorBase/NodeDescription.h"
-#include "DecoratorBase/NodeInstance.h"
-#include "DecoratorBase/DecoratorRegistry.h"
-#include "DecoratorBase/Decorator.h"
-#include "Serialization/MemoryReader.h"
-
-#include "RigVMCompiler/RigVMCompiler.h"
-#include "RigVMCore/RigVM.h"
 
 namespace UE::AnimNext::UncookedOnly
 {
-namespace Private
-{
-	struct FDecoratorEntryMapping
-	{
-		const URigVMNode* DecoratorStackNode;
-		const URigVMPin* DecoratorEntryPin;
-		const FDecorator* Decorator;
-
-		FDecoratorEntryMapping(const URigVMNode* InDecoratorStackNode, const URigVMPin* InDecoratorEntryPin, const FDecorator* InDecorator)
-			: DecoratorStackNode(InDecoratorStackNode)
-			, DecoratorEntryPin(InDecoratorEntryPin)
-			, Decorator(InDecorator)
-		{}
-	};
-
-	struct FDecoratorStackMapping
-	{
-		const URigVMNode* DecoratorStackNode;
-		TArray<FDecoratorEntryMapping> DecoratorEntries;
-
-		FNodeHandle DecoratorStackNodeHandle;
-
-		explicit FDecoratorStackMapping(const URigVMNode* InDecoratorStackNode)
-			: DecoratorStackNode(InDecoratorStackNode)
-		{}
-	};
-
-	template<typename DecoratorAction>
-	void ForEachDecoratorInStack(const URigVMNode* DecoratorStackNode, const DecoratorAction& Action)
-	{
-		const TArray<URigVMPin*>& Pins = DecoratorStackNode->GetPins();
-		for (URigVMPin* Pin : Pins)
-		{
-			if (!Pin->IsDecoratorPin())
-			{
-				continue;	// Not a decorator pin
-			}
-
-			if (Pin->GetScriptStruct() == FRigDecorator_AnimNextCppDecorator::StaticStruct())
-			{
-				TSharedPtr<FStructOnScope> DecoratorScope = Pin->GetDecoratorInstance();
-				FRigDecorator_AnimNextCppDecorator* VMDecorator = (FRigDecorator_AnimNextCppDecorator*)DecoratorScope->GetStructMemory();
-
-				if (const FDecorator* Decorator = VMDecorator->GetDecorator())
-				{
-					Action(DecoratorStackNode, Pin, Decorator);
-				}
-			}
-		}
-	}
-
-	TArray<FDecoratorUID> GetDecoratorUIDs(const URigVMNode* DecoratorStackNode)
-	{
-		TArray<FDecoratorUID> Decorators;
-
-		ForEachDecoratorInStack(DecoratorStackNode,
-			[&Decorators](const URigVMNode* DecoratorStackNode, const URigVMPin* DecoratorPin, const FDecorator* Decorator)
-			{
-				Decorators.Add(Decorator->GetDecoratorUID());
-			});
-
-		return Decorators;
-	}
-
-	FNodeHandle RegisterDecoratorNodeTemplate(FDecoratorWriter& DecoratorWriter, const URigVMNode* DecoratorStackNode)
-	{
-		const TArray<FDecoratorUID> DecoratorUIDs = GetDecoratorUIDs(DecoratorStackNode);
-
-		TArray<uint8> NodeTemplateBuffer;
-		const FNodeTemplate* NodeTemplate = FNodeTemplateBuilder::BuildNodeTemplate(DecoratorUIDs, NodeTemplateBuffer);
-
-		return DecoratorWriter.RegisterNode(*NodeTemplate);
-	}
-
-	FString GetDecoratorProperty(const FDecoratorStackMapping& DecoratorStack, uint32 DecoratorIndex, const FString& PropertyName, const TArray<FDecoratorStackMapping>& DecoratorStackNodes)
-	{
-		const TArray<URigVMPin*>& Pins = DecoratorStack.DecoratorEntries[DecoratorIndex].DecoratorEntryPin->GetSubPins();
-		for (URigVMPin* Pin : Pins)
-		{
-			if (Pin->GetDirection() != ERigVMPinDirection::Input)
-			{
-				continue;	// We only look for input pins
-			}
-
-			if (Pin->GetName() == PropertyName)
-			{
-				if (Pin->GetCPPTypeObject() == FAnimNextDecoratorHandle::StaticStruct())
-				{
-					// Decorator handle pins don't have a value, just an optional link
-					const TArray<URigVMLink*>& PinLinks = Pin->GetLinks();
-					if (PinLinks.Num() != 0)
-					{
-						// Something is connected to us, find the corresponding node handle so that we can encode it as our property value
-						check(PinLinks.Num() == 1);
-
-						const URigVMNode* SourceNode = PinLinks[0]->GetSourceNode();
-
-						FNodeHandle SourceNodeHandle;
-						int32 SourceDecoratorIndex = INDEX_NONE;
-
-						const FDecoratorStackMapping* SourceDecoratorStack = DecoratorStackNodes.FindByPredicate([SourceNode](const FDecoratorStackMapping& Mapping) { return Mapping.DecoratorStackNode == SourceNode; });
-						if (SourceDecoratorStack != nullptr)
-						{
-							SourceNodeHandle = SourceDecoratorStack->DecoratorStackNodeHandle;
-							SourceDecoratorIndex = 0;	// We always bind to the first decorator index since we only allow one base decorator per stack for now
-						}
-
-						if (SourceNodeHandle.IsValid())
-						{
-							check(SourceDecoratorIndex != INDEX_NONE);
-
-							const FAnimNextDecoratorHandle DecoratorHandle(SourceNodeHandle, SourceDecoratorIndex);
-							const FAnimNextDecoratorHandle DefaultDecoratorHandle;
-
-							// We need an instance of a decorator handle property to be able to serialize it into text, grab it from the root
-							const FProperty* Property = FRigUnit_AnimNextGraphRoot::StaticStruct()->FindPropertyByName(GET_MEMBER_NAME_STRING_CHECKED(FRigUnit_AnimNextGraphRoot, Result));
-
-							FString PropertyValue;
-							Property->ExportText_Direct(PropertyValue, &DecoratorHandle, &DefaultDecoratorHandle, nullptr, PPF_None);
-
-							return PropertyValue;
-						}
-					}
-
-					// This handle pin isn't connected
-					return FString();
-				}
-
-				// A regular property pin
-				return Pin->GetDefaultValue();
-			}
-		}
-
-		// Unknown property
-		return FString();
-	}
-
-	void WriteDecoratorProperties(FDecoratorWriter& DecoratorWriter, const FDecoratorStackMapping& Mapping, const TArray<FDecoratorStackMapping>& DecoratorStackNodes)
-	{
-		DecoratorWriter.WriteNode(Mapping.DecoratorStackNodeHandle,
-			[&Mapping, &DecoratorStackNodes](uint32 DecoratorIndex, const FString& PropertyName)
-			{
-				return GetDecoratorProperty(Mapping, DecoratorIndex, PropertyName, DecoratorStackNodes);
-			});
-	}
-
-	const URigVMUnitNode* FindRootNode(const TArray<URigVMNode*>& VMNodes)
-	{
-		for (const URigVMNode* VMNode : VMNodes)
-		{
-			if (const URigVMUnitNode* VMUnitNode = Cast<URigVMUnitNode>(VMNode))
-			{
-				const UScriptStruct* ScriptStruct = VMUnitNode->GetScriptStruct();
-				if (ScriptStruct == FRigUnit_AnimNextGraphRoot::StaticStruct())
-				{
-					return VMUnitNode;
-				}
-			}
-		}
-
-		return nullptr;
-	}
-
-	TArray<FDecoratorStackMapping> CollectDecoratorStacks(const URigVMGraph* VMGraph)
-	{
-		const TArray<URigVMNode*>& VMNodes = VMGraph->GetNodes();
-		const URigVMUnitNode* VMRootNode = FindRootNode(VMNodes);
-
-		TArray<FDecoratorStackMapping> DecoratorStackNodes;
-
-		if (VMRootNode == nullptr)
-		{
-			return DecoratorStackNodes;
-		}
-
-		TArray<const URigVMNode*> NodesToVisit;
-		NodesToVisit.Add(VMRootNode);
-
-		while (NodesToVisit.Num() != 0)
-		{
-			const URigVMNode* VMNode = NodesToVisit[0];
-			NodesToVisit.RemoveAt(0);
-
-			if (const URigVMUnitNode* VMUnitNode = Cast<URigVMUnitNode>(VMNode))
-			{
-				const UScriptStruct* ScriptStruct = VMUnitNode->GetScriptStruct();
-				if (ScriptStruct == FRigUnit_AnimNextDecoratorStack::StaticStruct())
-				{
-					FDecoratorStackMapping Mapping(VMNode);
-					ForEachDecoratorInStack(VMNode,
-						[&Mapping](const URigVMNode* DecoratorStackNode, const URigVMPin* DecoratorPin, const FDecorator* Decorator)
-						{
-							Mapping.DecoratorEntries.Add(FDecoratorEntryMapping(DecoratorStackNode, DecoratorPin, Decorator));
-						});
-
-					DecoratorStackNodes.Add(MoveTemp(Mapping));
-				}
-			}
-
-			const TArray<URigVMNode*> SourceNodes = VMNode->GetLinkedSourceNodes();
-			NodesToVisit.Append(SourceNodes);
-		}
-
-		return DecoratorStackNodes;
-	}
-}
 
 void FUtils::Compile(UAnimNextGraph* InGraph)
 {
@@ -266,76 +46,11 @@ void FUtils::Compile(UAnimNextGraph* InGraph)
 	EditorData->CompileLog.Messages.Reset();
 	EditorData->CompileLog.NumErrors = EditorData->CompileLog.NumWarnings = 0;
 
-	// We use a temporary graph model to build our final graph that we'll compile
-	URigVMGraph* VMTempGraph = EditorData->GetRigVMClient()->GetModel(TEXT("TempRigVMGraph"));
-	URigVMController* TempController = nullptr;
-
-	TArray<URigVMGraph*> GraphModelsToCompile;
-
-	if (VMTempGraph != nullptr)
-	{
-		TempController = EditorData->GetRigVMClient()->GetOrCreateController(VMTempGraph);
-
-		// Clear our temporary model
-		TempController->RemoveNodes(VMTempGraph->GetNodes());
-
-		URigVMGraph* VMRootGraph = EditorData->GetRigVMClient()->GetDefaultModel();
-
-		// Add our root node, it's always there
-		URigVMUnitNode* TempRootNode = TempController->AddUnitNode(FRigUnit_AnimNextGraphRoot::StaticStruct(), FRigVMStruct::ExecuteName, FVector2D(0.0f, 0.0f), FString(), false);
-
-		// Gather our decorator stacks
-		TArray<Private::FDecoratorStackMapping> DecoratorStackNodes = Private::CollectDecoratorStacks(VMRootGraph);
-
-		FDecoratorWriter DecoratorWriter;
-
-		// Iterate over every decorator stack and register our node templates
-		for (Private::FDecoratorStackMapping& NodeMapping : DecoratorStackNodes)
-		{
-			NodeMapping.DecoratorStackNodeHandle = Private::RegisterDecoratorNodeTemplate(DecoratorWriter, NodeMapping.DecoratorStackNode);
-		}
-
-		// Write our node shared data
-		DecoratorWriter.BeginNodeWriting();
-
-		for (const Private::FDecoratorStackMapping& NodeMapping : DecoratorStackNodes)
-		{
-			Private::WriteDecoratorProperties(DecoratorWriter, NodeMapping, DecoratorStackNodes);
-		}
-
-		DecoratorWriter.EndNodeWriting();
-
-		// Find our root node handle, if we have any stack nodes, the first one is our root stack
-		FAnimNextEntryPointHandle RootDecoratorHandle;
-		if (DecoratorStackNodes.Num() != 0)
-		{
-			RootDecoratorHandle = FAnimNextEntryPointHandle(DecoratorStackNodes[0].DecoratorStackNodeHandle);
-		}
-
-		// Cache our compiled metadata
-		InGraph->SharedDataArchiveBuffer = DecoratorWriter.GetGraphSharedData();
-		InGraph->RootDecoratorHandle = RootDecoratorHandle;
-
-		// Populate our runtime metadata
-		InGraph->LoadFromArchiveBuffer(InGraph->SharedDataArchiveBuffer);
-
-		// Compile our new graph
-		GraphModelsToCompile.Add(VMTempGraph);
-	}
-
-	if (VMTempGraph == nullptr)
-	{
-		// This is an old graph that was saved before the introduction of the temporary graph, use the root graph instead
-		URigVMGraph* VMRootGraph = EditorData->GetRigVMClient()->GetDefaultModel();
-
-		GraphModelsToCompile.Add(VMRootGraph);
-		TempController = EditorData->GetRigVMClient()->GetOrCreateController(VMRootGraph);
-	}
-
 	URigVMCompiler* Compiler = URigVMCompiler::StaticClass()->GetDefaultObject<URigVMCompiler>();
 	EditorData->VMCompileSettings.SetExecuteContextStruct(EditorData->RigVMClient.GetExecuteContextStruct());
 	Compiler->Settings = (EditorData->bCompileInDebugMode) ? FRigVMCompileSettings::Fast(EditorData->VMCompileSettings.GetExecuteContextStruct()) : EditorData->VMCompileSettings;
-	Compiler->Compile(GraphModelsToCompile, TempController, InGraph->RigVM, InGraph->ExtendedExecuteContext, InGraph->GetRigVMExternalVariables(), &EditorData->PinToOperandMap);
+	URigVMController* RootController = EditorData->GetRigVMClient()->GetOrCreateController(EditorData->GetRigVMClient()->GetDefaultModel());
+	Compiler->Compile(EditorData->GetRigVMClient()->GetAllModels(false, false), RootController, InGraph->RigVM, InGraph->ExtendedExecuteContext, InGraph->GetRigVMExternalVariables(), &EditorData->PinToOperandMap);
 
 	if (EditorData->bErrorsDuringCompilation)
 	{
@@ -900,7 +615,7 @@ void FUtils::SetupBindingGraphForLiteral(URigVMController* InController, FName I
 	InController->RemoveNodes(InController->GetGraph()->GetNodes());
 
 	// Add new nodes for a simple literal binding
-	URigVMUnitNode* EntryPointNode = InController->AddUnitNode(FRigUnit_AnimNextBeginExecution::StaticStruct(), FRigVMStruct::ExecuteName, FVector2D(-200.0f, 0.0f), FString(), false);
+	URigVMUnitNode* EntryPointNode = InController->AddUnitNode(FRigUnit_AnimNextBeginExecution::StaticStruct(), FRigUnit::GetMethodName(), FVector2D(-200.0f, 0.0f), FString(), false);
 
 	const FName FactoryName = FRigVMDispatch_SetParameter().GetFactoryName();
 	FRigVMDispatch_SetParameter* Factory = static_cast<FRigVMDispatch_SetParameter*>(FRigVMRegistry::Get().FindDispatchFactory(FactoryName));
diff --git a/Engine/Plugins/Runtime/RigVM/Source/RigVM/Public/RigVMCore/RigVMDecorator.h b/Engine/Plugins/Runtime/RigVM/Source/RigVM/Public/RigVMCore/RigVMDecorator.h
--- a/Engine/Plugins/Runtime/RigVM/Source/RigVM/Public/RigVMCore/RigVMDecorator.h
+++ b/Engine/Plugins/Runtime/RigVM/Source/RigVM/Public/RigVMCore/RigVMDecorator.h
@@ -62,9 +62,6 @@ public:
 	// allows the decorator to react when added to a node
 	virtual void OnDecoratorAdded(URigVMController* InController, URigVMNode* InNode) {}
 
-	// allows the decorator to return dynamic pins (parent pin index must be INDEX_NONE or point to a valid index of the parent pin in the OutPinArray)
-	virtual void GetProgrammaticPins(URigVMController* InController, int32 InParentPinIndex, struct FRigVMPinInfoArray& OutPinArray) const {}
-
 #endif
 
 private:
diff --git a/Engine/Plugins/Runtime/RigVM/Source/RigVMDeveloper/Private/EdGraph/RigVMEdGraphNode.cpp b/Engine/Plugins/Runtime/RigVM/Source/RigVMDeveloper/Private/EdGraph/RigVMEdGraphNode.cpp
--- a/Engine/Plugins/Runtime/RigVM/Source/RigVMDeveloper/Private/EdGraph/RigVMEdGraphNode.cpp
+++ b/Engine/Plugins/Runtime/RigVM/Source/RigVMDeveloper/Private/EdGraph/RigVMEdGraphNode.cpp
@@ -1547,11 +1547,6 @@ FName URigVMEdGraphNode::GetModelNodeName() const
 	return NAME_None;
 }
 
-const FString& URigVMEdGraphNode::GetModelNodePath() const
-{
-	return ModelNodePath;
-}
-
 URigVMPin* URigVMEdGraphNode::GetModelPinFromPinPath(const FString& InPinPath) const
 {
 	if (TWeakObjectPtr<URigVMPin> const* CachedModelPinPtr = PinPathToModelPin.Find(InPinPath))
diff --git a/Engine/Plugins/Runtime/RigVM/Source/RigVMDeveloper/Private/RigVMCompiler/RigVMAST.cpp b/Engine/Plugins/Runtime/RigVM/Source/RigVMDeveloper/Private/RigVMCompiler/RigVMAST.cpp
--- a/Engine/Plugins/Runtime/RigVM/Source/RigVMDeveloper/Private/RigVMCompiler/RigVMAST.cpp
+++ b/Engine/Plugins/Runtime/RigVM/Source/RigVMDeveloper/Private/RigVMCompiler/RigVMAST.cpp
@@ -2519,7 +2519,7 @@ const FRigVMASTLinkDescription& FRigVMParserAST::GetLink(int32 InLinkIndex) cons
 	return Links[InLinkIndex];
 }
 
-void FRigVMParserAST::Inline(const TArray<URigVMGraph*>& InGraphs)
+void FRigVMParserAST::Inline(TArray<URigVMGraph*> InGraphs)
 {
 	TArray<FRigVMASTProxy> LocalNodeProxies;
 	for(URigVMGraph* Graph : InGraphs)
@@ -2532,7 +2532,7 @@ void FRigVMParserAST::Inline(const TArray<URigVMGraph*>& InGraphs)
 	Inline(InGraphs, LocalNodeProxies);
 }
 
-void FRigVMParserAST::Inline(const TArray<URigVMGraph*>& InGraphs, const TArray<FRigVMASTProxy>& InNodeProxies)
+void FRigVMParserAST::Inline(TArray<URigVMGraph*> InGraphs, const TArray<FRigVMASTProxy>& InNodeProxies)
 {
 	DECLARE_SCOPE_HIERARCHICAL_COUNTER_FUNC()
 
diff --git a/Engine/Plugins/Runtime/RigVM/Source/RigVMDeveloper/Private/RigVMModel/RigVMController.cpp b/Engine/Plugins/Runtime/RigVM/Source/RigVMDeveloper/Private/RigVMModel/RigVMController.cpp
--- a/Engine/Plugins/Runtime/RigVM/Source/RigVMDeveloper/Private/RigVMModel/RigVMController.cpp
+++ b/Engine/Plugins/Runtime/RigVM/Source/RigVMDeveloper/Private/RigVMModel/RigVMController.cpp
@@ -351,7 +351,7 @@ void URigVMController::SetIsRunningUnitTest(bool bIsRunning)
 	}	
 }
 
-FRigVMPinInfo::FRigVMPinInfo()
+URigVMController::FPinInfo::FPinInfo()
 	: ParentIndex(INDEX_NONE)
 	, Name(NAME_None)
 	, Direction(ERigVMPinDirection::Invalid)
@@ -365,7 +365,7 @@ FRigVMPinInfo::FRigVMPinInfo()
 {
 }
 
-FRigVMPinInfo::FRigVMPinInfo(const URigVMPin* InPin, int32 InParentIndex, ERigVMPinDirection InDirection)
+URigVMController::FPinInfo::FPinInfo(const URigVMPin* InPin, int32 InParentIndex, ERigVMPinDirection InDirection)
 	: ParentIndex(InParentIndex)
 	, Name(InPin->GetName())
 	, Direction(InDirection == ERigVMPinDirection::Invalid ? InPin->GetDirection() : InDirection)
@@ -389,7 +389,7 @@ FRigVMPinInfo::FRigVMPinInfo(const URigVMPin* InPin, int32 InParentIndex, ERigVM
 	}
 }
 
-FRigVMPinInfo::FRigVMPinInfo(FProperty* InProperty, ERigVMPinDirection InDirection, int32 InParentIndex, const uint8* InDefaultValueMemory)
+URigVMController::FPinInfo::FPinInfo(FProperty* InProperty, ERigVMPinDirection InDirection, int32 InParentIndex, const uint8* InDefaultValueMemory)
 	: ParentIndex(InParentIndex)
 	, Name(InProperty->GetFName())
 	, Direction(InDirection)
@@ -496,7 +496,7 @@ FRigVMPinInfo::FRigVMPinInfo(FProperty* InProperty, ERigVMPinDirection InDirecti
 	CorrectExecuteTypeIndex();
 }
 
-void FRigVMPinInfo::CorrectExecuteTypeIndex()
+void URigVMController::FPinInfo::CorrectExecuteTypeIndex()
 {
 	const FRigVMRegistry& Registry = FRigVMRegistry::Get();
 	if(Registry.IsExecuteType(TypeIndex))
@@ -510,7 +510,7 @@ void FRigVMPinInfo::CorrectExecuteTypeIndex()
 	}
 }
 
-uint32 GetTypeHash(const FRigVMPinInfo& InPin)
+uint32 GetTypeHash(const URigVMController::FPinInfo& InPin)
 {
 	uint32 Hash = 0; //GetTypeHash(InPin.ParentIndex);
 	Hash = HashCombine(Hash, GetTypeHash(InPin.Name));
@@ -523,7 +523,7 @@ uint32 GetTypeHash(const FRigVMPinInfo& InPin)
 	return Hash;
 }
 
-FRigVMPinInfoArray::FRigVMPinInfoArray(const URigVMNode* InNode)
+URigVMController::FPinInfoArray::FPinInfoArray(const URigVMNode* InNode)
 {
 	// this method adds all pins as currently represented in the model.
 	for(const URigVMPin* Pin : InNode->GetPins())
@@ -532,7 +532,7 @@ FRigVMPinInfoArray::FRigVMPinInfoArray(const URigVMNode* InNode)
 	}
 }
 
-int32 FRigVMPinInfoArray::AddPin(const URigVMPin* InPin, int32 InParentIndex, ERigVMPinDirection InDirection)
+int32 URigVMController::FPinInfoArray::AddPin(const URigVMPin* InPin, int32 InParentIndex, ERigVMPinDirection InDirection)
 {
 	// this method adds all pins as currently represented in the model.
 	const int32 Index = Pins.Emplace(InPin, InParentIndex, InDirection);
@@ -544,7 +544,7 @@ int32 FRigVMPinInfoArray::AddPin(const URigVMPin* InPin, int32 InParentIndex, ER
 	return Index;
 }
 
-FRigVMPinInfoArray::FRigVMPinInfoArray(const URigVMNode* InNode, URigVMController* InController, const FRigVMPinInfoArray* InPreviousPinInfos)
+URigVMController::FPinInfoArray::FPinInfoArray(const URigVMNode* InNode, URigVMController* InController, const FPinInfoArray* InPreviousPinInfos)
 {
 	const bool bAddSubPins = !InNode->IsA<URigVMRerouteNode>();
 	
@@ -557,8 +557,8 @@ FRigVMPinInfoArray::FRigVMPinInfoArray(const URigVMNode* InNode, URigVMControlle
 	}
 }
 
-FRigVMPinInfoArray::FRigVMPinInfoArray(const FRigVMGraphFunctionHeader& FunctionHeader,
-	URigVMController* InController, const FRigVMPinInfoArray* InPreviousPinInfos)
+URigVMController::FPinInfoArray::FPinInfoArray(const FRigVMGraphFunctionHeader& FunctionHeader,
+	URigVMController* InController, const FPinInfoArray* InPreviousPinInfos)
 {
 	// this method adds pins as needed based on the property structure
 	const FRigVMRegistry& Registry = FRigVMRegistry::Get();
@@ -570,7 +570,7 @@ FRigVMPinInfoArray::FRigVMPinInfoArray(const FRigVMGraphFunctionHeader& Function
 	}
 }
 
-int32 FRigVMPinInfoArray::AddPin(FProperty* InProperty, URigVMController* InController,
+int32 URigVMController::FPinInfoArray::AddPin(FProperty* InProperty, URigVMController* InController,
                                               ERigVMPinDirection InDirection, int32 InParentIndex, const uint8* InDefaultValueMemory, bool bAddSubPins)
 {
 	// this method adds pins as needed based on the property structure
@@ -603,13 +603,13 @@ int32 FRigVMPinInfoArray::AddPin(FProperty* InProperty, URigVMController* InCont
 	return Index;
 }
 
-int32 FRigVMPinInfoArray::AddPin(URigVMController* InController, int32 InParentIndex, const FName& InName, ERigVMPinDirection InDirection,
-	TRigVMTypeIndex InTypeIndex, const FString& InDefaultValue, const uint8* InDefaultValueMemory, const FRigVMPinInfoArray* InPreviousPinInfos, bool bAddSubPins)
+int32 URigVMController::FPinInfoArray::AddPin(URigVMController* InController, int32 InParentIndex, const FName& InName, ERigVMPinDirection InDirection,
+	TRigVMTypeIndex InTypeIndex, const FString& InDefaultValue, const uint8* InDefaultValueMemory, const FPinInfoArray* InPreviousPinInfos, bool bAddSubPins)
 {
 	ensureMsgf(InTypeIndex != INDEX_NONE, TEXT("Invalid pin type for pin %s in %s"), *InName.ToString(), *InController->GetPackage()->GetPathName());
 	const FRigVMRegistry& Registry = FRigVMRegistry::Get();
 		
-	FRigVMPinInfo Info;
+	FPinInfo Info;
 	Info.ParentIndex = InParentIndex;
 	Info.Name = InName;
 	Info.Direction = InDirection;
@@ -626,7 +626,7 @@ int32 FRigVMPinInfoArray::AddPin(URigVMController* InController, int32 InParentI
 		const int32 PreviousIndex = InPreviousPinInfos->GetIndexFromPinPath(PinPath);
 		if(PreviousIndex != INDEX_NONE)
 		{
-			const FRigVMPinInfo& PreviousPin = (*InPreviousPinInfos)[PreviousIndex];
+			const FPinInfo& PreviousPin = (*InPreviousPinInfos)[PreviousIndex];
 			if(PreviousPin.TypeIndex == InTypeIndex)
 			{
 				Pins[Index].DefaultValue = PreviousPin.DefaultValue;
@@ -684,7 +684,7 @@ int32 FRigVMPinInfoArray::AddPin(URigVMController* InController, int32 InParentI
 	return Index;
 }
 
-void FRigVMPinInfoArray::AddPins(UScriptStruct* InScriptStruct, URigVMController* InController,
+void URigVMController::FPinInfoArray::AddPins(UScriptStruct* InScriptStruct, URigVMController* InController,
 	ERigVMPinDirection InDirection, int32 InParentIndex, const uint8* InDefaultValueMemory, bool bAddSubPins)
 {
 	if (InController->GetSchema()->ShouldUnfoldStruct(InController, InScriptStruct))
@@ -718,7 +718,7 @@ void FRigVMPinInfoArray::AddPins(UScriptStruct* InScriptStruct, URigVMController
 	}
 }
 
-const FString& FRigVMPinInfoArray::GetPinPath(const int32 InIndex) const
+const FString& URigVMController::FPinInfoArray::GetPinPath(const int32 InIndex) const
 {
 	if(!Pins.IsValidIndex(InIndex))
 	{
@@ -741,7 +741,7 @@ const FString& FRigVMPinInfoArray::GetPinPath(const int32 InIndex) const
 	return Pins[InIndex].PinPath;
 }
 
-int32 FRigVMPinInfoArray::GetIndexFromPinPath(const FString& InPinPath) const
+int32 URigVMController::FPinInfoArray::GetIndexFromPinPath(const FString& InPinPath) const
 {
 	if(PinPathLookup.Num() != Num())
 	{
@@ -759,7 +759,7 @@ int32 FRigVMPinInfoArray::GetIndexFromPinPath(const FString& InPinPath) const
 	return INDEX_NONE;
 }
 
-const FRigVMPinInfo* FRigVMPinInfoArray::GetPinFromPinPath(const FString& InPinPath) const
+const URigVMController::FPinInfo* URigVMController::FPinInfoArray::GetPinFromPinPath(const FString& InPinPath) const
 {
 	const int32 Index = GetIndexFromPinPath(InPinPath);
 	if(Pins.IsValidIndex(Index))
@@ -769,7 +769,7 @@ const FRigVMPinInfo* FRigVMPinInfoArray::GetPinFromPinPath(const FString& InPinP
 	return nullptr;
 }
 
-int32 FRigVMPinInfoArray::GetRootIndex(const int32 InIndex) const
+int32 URigVMController::FPinInfoArray::GetRootIndex(const int32 InIndex) const
 {
 	if(Pins.IsValidIndex(InIndex))
 	{
@@ -782,13 +782,13 @@ int32 FRigVMPinInfoArray::GetRootIndex(const int32 InIndex) const
 	return INDEX_NONE;
 }
 
-uint32 GetTypeHash(const FRigVMPinInfoArray& InPins)
+uint32 GetTypeHash(const URigVMController::FPinInfoArray& InPins)
 {
 	TArray<uint32> Hashes;
 	Hashes.Reserve(InPins.Num());
 	
 	uint32 OverAllHash = GetTypeHash(InPins.Num());
-	for(const FRigVMPinInfo& Info : InPins)
+	for(const URigVMController::FPinInfo& Info : InPins)
 	{
 		uint32 PinHash = GetTypeHash(Info);
 		if(Info.ParentIndex != INDEX_NONE)
@@ -13758,73 +13758,6 @@ FName URigVMController::AddDecorator(URigVMNode* InNode, UScriptStruct* InDecora
 	
 	AddPinsForStruct(InDecoratorScriptStruct, InNode, DecoratorPin, DecoratorPin->GetDirection(), InDefaultValue, true);
 
-	FRigVMPinInfoArray ProgrammaticPins;
-	Decorator->GetProgrammaticPins(this, INDEX_NONE, ProgrammaticPins);
-
-	const FRigVMRegistry& Registry = FRigVMRegistry::Get();
-	auto CreatePinFromPinInfo = [this, &Registry](const FRigVMPinInfo& InPinInfo, const FString& InPinPath, UObject* InOuter) -> URigVMPin*
-	{
-		check(InOuter);
-		URigVMPin* Pin = NewObject<URigVMPin>(InOuter, InPinInfo.Name);
-		if (InPinInfo.Property)
-		{
-			ConfigurePinFromProperty(InPinInfo.Property, Pin, InPinInfo.Direction);
-		}
-		else
-		{
-			const FRigVMTemplateArgumentType& Type = Registry.GetType(InPinInfo.TypeIndex);
-			Pin->CPPType = Type.CPPType.ToString();
-			Pin->CPPTypeObject = Type.CPPTypeObject;
-			if (Pin->CPPTypeObject)
-			{
-				Pin->CPPTypeObjectPath = *Pin->CPPTypeObject->GetPathName();
-			}
-			if (Registry.IsExecuteType(InPinInfo.TypeIndex))
-			{
-				MakeExecutePin(Pin);
-			}
-
-			Pin->Direction = InPinInfo.Direction;
-			Pin->DisplayName = InPinInfo.DisplayName.IsEmpty() ? NAME_None : FName(*InPinInfo.DisplayName);
-			Pin->bIsConstant = InPinInfo.bIsConstant;
-			Pin->bIsDynamicArray = InPinInfo.bIsDynamicArray;
-			Pin->CustomWidgetName = InPinInfo.CustomWidgetName.IsEmpty() ? NAME_None : FName(*InPinInfo.CustomWidgetName);
-		}
-
-		Pin->bIsExpanded = InPinInfo.bIsExpanded;
-		Pin->DefaultValue = InPinInfo.DefaultValue;
-
-		if (URigVMPin* ParentPin = Cast<URigVMPin>(InOuter))
-		{
-			AddSubPin(ParentPin, Pin);
-		}
-		else if (URigVMNode* OwnerNode = Cast<URigVMNode>(InOuter))
-		{
-			AddNodePin(OwnerNode, Pin);
-		}
-		else
-		{
-			ensureMsgf(false, TEXT("Outer %s of pin info %s is not a pin or a node"), *InOuter->GetPathName(), *InPinPath);
-		}
-
-		Notify(ERigVMGraphNotifType::PinAdded, Pin);
-
-		return Pin;
-	};
-
-	for (int32 PinIndex = 0; PinIndex < ProgrammaticPins.Num(); ++PinIndex)
-	{
-		const FString& PinPath = ProgrammaticPins.GetPinPath(PinIndex);
-		FString ParentPinPath, PinName;
-		UObject* OuterForPin = DecoratorPin;
-		if (URigVMPin::SplitPinPathAtEnd(PinPath, ParentPinPath, PinName))
-		{
-			OuterForPin = DecoratorPin->FindSubPin(ParentPinPath);
-		}
-
-		CreatePinFromPinInfo(ProgrammaticPins[PinIndex], PinPath, OuterForPin);
-	}
-
 	// move the the pin to the right index as required
 	if(DecoratorPin->GetPinIndex() != InPinIndex &&
 		InPinIndex >=0 && InPinIndex < InNode->GetPins().Num())
@@ -14111,7 +14044,7 @@ bool URigVMController::IsValidLinkForGraph(const URigVMLink* InLink)
 	return true;
 }
 
-void URigVMController::AddPinsForStruct(UStruct* InStruct, URigVMNode* InNode, URigVMPin* InParentPin, ERigVMPinDirection InPinDirection, const FString& InDefaultValue, bool bAutoExpandArrays, const FRigVMPinInfoArray* PreviousPins)
+void URigVMController::AddPinsForStruct(UStruct* InStruct, URigVMNode* InNode, URigVMPin* InParentPin, ERigVMPinDirection InPinDirection, const FString& InDefaultValue, bool bAutoExpandArrays, const FPinInfoArray* PreviousPins)
 {
 	if(!InStruct->IsChildOf(FRigVMStruct::StaticStruct()))
 	{
@@ -14424,15 +14357,13 @@ void URigVMController::AddPinsForTemplate(const FRigVMTemplate* InTemplate, cons
 
 void URigVMController::ConfigurePinFromProperty(FProperty* InProperty, URigVMPin* InOutPin, ERigVMPinDirection InPinDirection)
 {
-	const ERigVMPinDirection PropertyPinDirection = FRigVMStruct::GetPinDirectionFromProperty(InProperty);
 	if (InPinDirection == ERigVMPinDirection::Invalid)
 	{
-		InOutPin->Direction = PropertyPinDirection;
+		InOutPin->Direction = FRigVMStruct::GetPinDirectionFromProperty(InProperty);
 	}
 	else
 	{
-		// Keep the pin direction specified unless we wish to be hidden
-		InOutPin->Direction = PropertyPinDirection == ERigVMPinDirection::Hidden ? ERigVMPinDirection::Hidden : InPinDirection;
+		InOutPin->Direction = InPinDirection;
 	}
 
 #if WITH_EDITOR
@@ -14887,9 +14818,9 @@ void URigVMController::RepopulatePinsOnNode(URigVMNode* InNode, bool bFollowCore
 		}
 	}
 
-	const FRigVMPinInfoArray PreviousPinInfos(InNode);
+	const FPinInfoArray PreviousPinInfos(InNode);
 	const uint32 PreviousPinHash = GetTypeHash(PreviousPinInfos);
-	FRigVMPinInfoArray NewPinInfos;
+	FPinInfoArray NewPinInfos;
 
 	// step 2/3: clear pins on the node and repopulate the node with new pins
 	if (UnitNode != nullptr)
@@ -15103,7 +15034,7 @@ void URigVMController::RepopulatePinsOnNode(URigVMNode* InNode, bool bFollowCore
 			}
 		}
 
-		NewPinInfos = FRigVMPinInfoArray(InNode, this, &PreviousPinInfos);
+		NewPinInfos = FPinInfoArray(InNode, this, &PreviousPinInfos);
 	}
 	else if (EntryNode || ReturnNode)
 	{
@@ -15162,20 +15093,20 @@ void URigVMController::RepopulatePinsOnNode(URigVMNode* InNode, bool bFollowCore
 	}
 	else if (CollapseNode)
 	{
-		NewPinInfos = FRigVMPinInfoArray(InNode, this, &PreviousPinInfos);
+		NewPinInfos = FPinInfoArray(InNode, this, &PreviousPinInfos);
 	}
 	else if (FunctionRefNode)
 	{
 		const FRigVMGraphFunctionHeader& FunctionHeader = FunctionRefNode->GetReferencedFunctionHeader();
 		if(FunctionHeader.IsValid())
 		{
-			NewPinInfos = FRigVMPinInfoArray(FunctionHeader, this, &PreviousPinInfos);
+			NewPinInfos = FPinInfoArray(FunctionHeader, this, &PreviousPinInfos);
 		}
 		// if we can't find referenced node anymore
 		// let's keep all pins
 		else
 		{
-			NewPinInfos = FRigVMPinInfoArray(FunctionRefNode, this, &PreviousPinInfos);
+			NewPinInfos = FPinInfoArray(FunctionRefNode, this, &PreviousPinInfos);
 		}
 	}
 	else
@@ -15205,19 +15136,11 @@ void URigVMController::RepopulatePinsOnNode(URigVMNode* InNode, bool bFollowCore
 	// make sure the new pin infos contains the decorator pins from the last run
 	for(int32 Index = 0; Index < PreviousPinInfos.Num(); Index++)
 	{
-		const FRigVMPinInfo& PreviousPin = PreviousPinInfos[Index]; 
+		const FPinInfo& PreviousPin = PreviousPinInfos[Index]; 
 		if(PreviousPin.bIsDecorator)
 		{
 			const int32 NewPinIndex = NewPinInfos.AddPin(this, INDEX_NONE, PreviousPin.Name, PreviousPin.Direction, PreviousPin.TypeIndex, PreviousPin.DefaultValue, nullptr, &PreviousPinInfos, true);
 			NewPinInfos[NewPinIndex].bIsDecorator = true;
-
-			if (URigVMPin* Pin = InNode->FindPin(PreviousPin.PinPath))
-			{
-				TSharedPtr<FStructOnScope> DecoratorScope = Pin->GetDecoratorInstance();
-				FRigVMDecorator* VMDecorator = (FRigVMDecorator*)DecoratorScope->GetStructMemory();
-
-				VMDecorator->GetProgrammaticPins(this, NewPinIndex, NewPinInfos);
-			}
 		}
 	}
 
@@ -15353,7 +15276,7 @@ void URigVMController::RepopulatePinsOnNode(URigVMNode* InNode, bool bFollowCore
 		}
 	}
 
-	auto CreatePinFromPinInfo = [this, &Registry, &PreviousPinInfos](const FRigVMPinInfo& InPinInfo, const FString& InPinPath, UObject* InOuter) -> URigVMPin*
+	auto CreatePinFromPinInfo = [this, &Registry, &PreviousPinInfos](const FPinInfo& InPinInfo, const FString& InPinPath, UObject* InOuter) -> URigVMPin*
 	{
 		check(InOuter);
 		URigVMPin* Pin = NewObject<URigVMPin>(InOuter, InPinInfo.Name);
@@ -15386,7 +15309,7 @@ void URigVMController::RepopulatePinsOnNode(URigVMNode* InNode, bool bFollowCore
 		Pin->DefaultValue = InPinInfo.DefaultValue;
 
 		// reuse expansion state and default value
-		if(const FRigVMPinInfo* PreviousPin = PreviousPinInfos.GetPinFromPinPath(InPinPath))
+		if(const FPinInfo* PreviousPin = PreviousPinInfos.GetPinFromPinPath(InPinPath))
 		{
 			if(PreviousPin->TypeIndex == InPinInfo.TypeIndex)
 			{
@@ -15513,7 +15436,7 @@ void URigVMController::RepopulatePinsOnNode(URigVMNode* InNode, bool bFollowCore
 	for(int32 Index = 0; Index < PreviousPinsToUpdate.Num(); Index++)
 	{
 		const FString& PinPath = PreviousPinInfos.GetPinPath(PreviousPinsToUpdate[Index]);
-		const FRigVMPinInfo* NewPinInfo = NewPinInfos.GetPinFromPinPath(PinPath);
+		const FPinInfo* NewPinInfo = NewPinInfos.GetPinFromPinPath(PinPath);
 		check(NewPinInfo);
 		
 		if(URigVMPin* Pin = InNode->FindPin(PinPath))
@@ -15554,7 +15477,7 @@ void URigVMController::RepopulatePinsOnNode(URigVMNode* InNode, bool bFollowCore
 	TMap<FString, TArray<FName>> PinOrder;
 	for(int32 Index = 0; Index < NewPinInfos.Num(); Index++)
 	{
-		const FRigVMPinInfo& NewPin = NewPinInfos[Index];
+		const FPinInfo& NewPin = NewPinInfos[Index];
 		FString ParentPinPath;
 		if(NewPin.ParentIndex != INDEX_NONE)
 		{
diff --git a/Engine/Plugins/Runtime/RigVM/Source/RigVMDeveloper/Private/RigVMModel/RigVMNode.cpp b/Engine/Plugins/Runtime/RigVM/Source/RigVMDeveloper/Private/RigVMModel/RigVMNode.cpp
--- a/Engine/Plugins/Runtime/RigVM/Source/RigVMDeveloper/Private/RigVMModel/RigVMNode.cpp
+++ b/Engine/Plugins/Runtime/RigVM/Source/RigVMDeveloper/Private/RigVMModel/RigVMNode.cpp
@@ -282,10 +282,7 @@ bool URigVMNode::IsMutable() const
 	{
 		if(const UScriptStruct* ScriptStruct = Pin->GetScriptStruct())
 		{
-			if (ScriptStruct->IsChildOf(FRigVMExecuteContext::StaticStruct()))
-			{
-				return true;
-			}
+			return ScriptStruct->IsChildOf(FRigVMExecuteContext::StaticStruct());
 		}
 	}
 	return false;
diff --git a/Engine/Plugins/Runtime/RigVM/Source/RigVMDeveloper/Public/EdGraph/RigVMEdGraphNode.h b/Engine/Plugins/Runtime/RigVM/Source/RigVMDeveloper/Public/EdGraph/RigVMEdGraphNode.h
--- a/Engine/Plugins/Runtime/RigVM/Source/RigVMDeveloper/Public/EdGraph/RigVMEdGraphNode.h
+++ b/Engine/Plugins/Runtime/RigVM/Source/RigVMDeveloper/Public/EdGraph/RigVMEdGraphNode.h
@@ -147,9 +147,6 @@ public:
 	/** Get the VM node name this node is wrapping */
 	FName GetModelNodeName() const;
 
-	/** Get the VM node path this node is wrapping */
-	const FString& GetModelNodePath() const;
-
 	URigVMPin* GetModelPinFromPinPath(const FString& InPinPath) const;
 
 	/** Add a new element to the aggregate node referred to by the property path */
diff --git a/Engine/Plugins/Runtime/RigVM/Source/RigVMDeveloper/Public/RigVMCompiler/RigVMAST.h b/Engine/Plugins/Runtime/RigVM/Source/RigVMDeveloper/Public/RigVMCompiler/RigVMAST.h
--- a/Engine/Plugins/Runtime/RigVM/Source/RigVMDeveloper/Public/RigVMCompiler/RigVMAST.h
+++ b/Engine/Plugins/Runtime/RigVM/Source/RigVMDeveloper/Public/RigVMCompiler/RigVMAST.h
@@ -1317,8 +1317,8 @@ private:
 	void FoldAssignments();
 
 	// helper function to inline all contributing nodes of the graph
-	void Inline(const TArray<URigVMGraph*>& InGraphs);
-	void Inline(const TArray<URigVMGraph*>& InGraphs, const TArray<FRigVMASTProxy>& InNodeProxies);
+	void Inline(TArray<URigVMGraph*> InGraphs);
+	void Inline(TArray<URigVMGraph*> InGraphs, const TArray<FRigVMASTProxy>& InNodeProxies);
 
 	// helper functions to retrieve links for a given pin
 	TArray<int32> GetSourceLinkIndices(const FRigVMASTProxy& InPinProxy, bool bRecursive = false) const;
diff --git a/Engine/Plugins/Runtime/RigVM/Source/RigVMDeveloper/Public/RigVMModel/RigVMController.h b/Engine/Plugins/Runtime/RigVM/Source/RigVMDeveloper/Public/RigVMModel/RigVMController.h
--- a/Engine/Plugins/Runtime/RigVM/Source/RigVMDeveloper/Public/RigVMModel/RigVMController.h
+++ b/Engine/Plugins/Runtime/RigVM/Source/RigVMDeveloper/Public/RigVMModel/RigVMController.h
@@ -166,62 +166,6 @@ private:
 	friend class URigVMController;
 };
 
-struct RIGVMDEVELOPER_API FRigVMPinInfo
-{
-	FRigVMPinInfo();
-	FRigVMPinInfo(const URigVMPin* InPin, int32 InParentIndex, ERigVMPinDirection InDirection = ERigVMPinDirection::Invalid);
-	FRigVMPinInfo(FProperty* InProperty, ERigVMPinDirection InDirection, int32 InParentIndex, const uint8* InDefaultValueMemory);
-
-	void CorrectExecuteTypeIndex();
-
-	int32 ParentIndex;
-	FName Name;
-	ERigVMPinDirection Direction;
-	TRigVMTypeIndex TypeIndex;
-	bool bIsArray;
-	FProperty* Property;
-	FString PinPath;
-	FString DefaultValue;
-	FString DisplayName;
-	FString CustomWidgetName;
-	bool bIsExpanded;
-	bool bIsConstant;
-	bool bIsDynamicArray;
-	bool bIsDecorator;
-	TArray<int32> SubPins;
-
-	friend uint32 GetTypeHash(const FRigVMPinInfo& InPin);
-};
-
-struct RIGVMDEVELOPER_API FRigVMPinInfoArray
-{
-	FRigVMPinInfoArray() {}
-	explicit FRigVMPinInfoArray(const URigVMNode* InNode);
-	FRigVMPinInfoArray(const URigVMNode* InNode, URigVMController* InController, const FRigVMPinInfoArray* InPreviousPinInfos = nullptr);
-	FRigVMPinInfoArray(const FRigVMGraphFunctionHeader& FunctionHeader, URigVMController* InController, const FRigVMPinInfoArray* InPreviousPinInfos = nullptr);
-
-	int32 Num() const { return Pins.Num(); }
-	const FRigVMPinInfo& operator[](int32 InIndex) const { return Pins[InIndex]; }
-	FRigVMPinInfo& operator[](int32 InIndex) { return Pins[InIndex]; }
-	TArray<FRigVMPinInfo>::RangedForIteratorType begin() const { return Pins.begin(); }
-	TArray<FRigVMPinInfo>::RangedForIteratorType end() const { return Pins.end(); }
-
-	int32 AddPin(const URigVMPin* InPin, int32 InParentIndex, ERigVMPinDirection InDirection = ERigVMPinDirection::Invalid);
-	int32 AddPin(FProperty* InProperty, URigVMController* InController, ERigVMPinDirection InDirection, int32 InParentIndex, const uint8* InDefaultValueMemory, bool bAddSubPins);
-	int32 AddPin(URigVMController* InController, int32 InParentIndex, const FName& InName, ERigVMPinDirection InDirection, TRigVMTypeIndex InTypeIndex, const FString& InDefaultValue, const uint8* InDefaultValueMemory, const FRigVMPinInfoArray* InPreviousPinInfos, bool bAddSubPins);
-	void AddPins(UScriptStruct* InScriptStruct, URigVMController* InController, ERigVMPinDirection InDirection, int32 InParentIndex, const uint8* InDefaultValueMemory, bool bAddSubPins);
-
-	const FString& GetPinPath(const int32 InIndex) const;
-	int32 GetIndexFromPinPath(const FString& InPinPath) const;
-	const FRigVMPinInfo* GetPinFromPinPath(const FString& InPinPath) const;
-	int32 GetRootIndex(const int32 InIndex) const;
-
-	friend uint32 GetTypeHash(const FRigVMPinInfoArray& InPins);
-
-	mutable TArray<FRigVMPinInfo> Pins;
-	mutable TMap<FString, int32> PinPathLookup;;
-};
-
 /**
  * The Controller is the sole authority to perform changes
  * on the Graph. The Controller itself is stateless.
@@ -1090,6 +1034,64 @@ public:
 	UFUNCTION(BlueprintCallable, Category = RigVMController)
 	void SetIsRunningUnitTest(bool bIsRunning);
 
+public:
+	
+	struct FPinInfo
+	{
+		FPinInfo();
+		FPinInfo(const URigVMPin* InPin, int32 InParentIndex, ERigVMPinDirection InDirection = ERigVMPinDirection::Invalid);
+		FPinInfo(FProperty* InProperty, ERigVMPinDirection InDirection, int32 InParentIndex, const uint8* InDefaultValueMemory);
+
+		void CorrectExecuteTypeIndex();
+		
+		int32 ParentIndex;
+		FName Name;
+		ERigVMPinDirection Direction;
+		TRigVMTypeIndex TypeIndex;
+		bool bIsArray;
+		FProperty* Property;
+		FString PinPath;
+		FString DefaultValue;
+		FString DisplayName;
+		FString CustomWidgetName;
+		bool bIsExpanded;
+		bool bIsConstant;
+		bool bIsDynamicArray;
+		bool bIsDecorator;
+		TArray<int32> SubPins;
+
+		friend uint32 GetTypeHash(const FPinInfo& InPin);
+	};
+
+	struct FPinInfoArray
+	{
+		FPinInfoArray() {}
+		FPinInfoArray(const URigVMNode* InNode);
+		FPinInfoArray(const URigVMNode* InNode, URigVMController* InController, const FPinInfoArray* InPreviousPinInfos = nullptr);
+		FPinInfoArray(const FRigVMGraphFunctionHeader& FunctionHeader, URigVMController* InController, const FPinInfoArray* InPreviousPinInfos = nullptr);
+
+		int32 Num() const { return Pins.Num(); }
+		const FPinInfo& operator[](int32 InIndex) const { return Pins[InIndex]; }
+		FPinInfo& operator[](int32 InIndex) { return Pins[InIndex]; }
+		TArray<FPinInfo>::RangedForIteratorType begin() const { return Pins.begin(); }
+		TArray<FPinInfo>::RangedForIteratorType end() const { return Pins.end(); }
+
+		int32 AddPin(const URigVMPin* InPin, int32 InParentIndex, ERigVMPinDirection InDirection = ERigVMPinDirection::Invalid);
+		int32 AddPin(FProperty* InProperty, URigVMController* InController, ERigVMPinDirection InDirection, int32 InParentIndex, const uint8* InDefaultValueMemory, bool bAddSubPins);
+		int32 AddPin(URigVMController* InController, int32 InParentIndex, const FName& InName, ERigVMPinDirection InDirection, TRigVMTypeIndex InTypeIndex, const FString& InDefaultValue, const uint8* InDefaultValueMemory, const FPinInfoArray* InPreviousPinInfos, bool bAddSubPins);
+		void AddPins(UScriptStruct* InScriptStruct, URigVMController* InController, ERigVMPinDirection InDirection, int32 InParentIndex, const uint8* InDefaultValueMemory, bool bAddSubPins);
+
+		const FString& GetPinPath(const int32 InIndex) const;
+		int32 GetIndexFromPinPath(const FString& InPinPath) const;
+		const FPinInfo* GetPinFromPinPath(const FString& InPinPath) const;
+		int32 GetRootIndex(const int32 InIndex) const;
+
+		friend uint32 GetTypeHash(const FPinInfoArray& InPins);
+		
+		mutable TArray<FPinInfo> Pins;
+		mutable TMap<FString, int32> PinPathLookup;;
+	};
+
 private:
 
 	UPROPERTY(BlueprintReadOnly, Category = RigVMController, meta = (ScriptName = "ModifiedEvent", AllowPrivateAccess = "true"))
@@ -1104,7 +1106,7 @@ private:
 	bool IsValidNodeForGraph(const URigVMNode* InNode);
 	bool IsValidPinForGraph(const URigVMPin* InPin);
 	bool IsValidLinkForGraph(const URigVMLink* InLink);
-	void AddPinsForStruct(UStruct* InStruct, URigVMNode* InNode, URigVMPin* InParentPin, ERigVMPinDirection InPinDirection, const FString& InDefaultValue, bool bAutoExpandArrays, const FRigVMPinInfoArray* PreviousPins = nullptr);
+	void AddPinsForStruct(UStruct* InStruct, URigVMNode* InNode, URigVMPin* InParentPin, ERigVMPinDirection InPinDirection, const FString& InDefaultValue, bool bAutoExpandArrays, const FPinInfoArray* PreviousPins = nullptr);
 	void AddPinsForArray(FArrayProperty* InArrayProperty, URigVMNode* InNode, URigVMPin* InParentPin, ERigVMPinDirection InPinDirection, const TArray<FString>& InDefaultValues, bool bAutoExpandArrays);
 	void AddPinsForTemplate(const FRigVMTemplate* InTemplate, const FRigVMTemplateTypeMap& InPinTypeMap, URigVMNode* InNode);
 	void ConfigurePinFromProperty(FProperty* InProperty, URigVMPin* InOutPin, ERigVMPinDirection InPinDirection = ERigVMPinDirection::Invalid);
@@ -1449,7 +1451,7 @@ private:
 	friend struct FRigVMEjectNodeFromPinAction;
 	friend class FRigVMParserAST;
 	friend class FRigVMControllerCompileBracketScope;
-	friend struct FRigVMPinInfoArray;
+	friend struct FPinInfoArray;
 	friend class FRigVMControllerNotifGuard;
 	friend struct FRigVMClient;
 	friend struct FRigVMActionWrapper;
diff --git a/Engine/Plugins/Runtime/RigVM/Source/RigVMEditor/Public/Widgets/SRigVMGraphNode.h b/Engine/Plugins/Runtime/RigVM/Source/RigVMEditor/Public/Widgets/SRigVMGraphNode.h
--- a/Engine/Plugins/Runtime/RigVM/Source/RigVMEditor/Public/Widgets/SRigVMGraphNode.h
+++ b/Engine/Plugins/Runtime/RigVM/Source/RigVMEditor/Public/Widgets/SRigVMGraphNode.h
@@ -62,7 +62,7 @@ public:
 	/** Callback function executed when Add pin button is clicked */
 	virtual FReply OnAddPin() override;
 
-protected:
+private:
 
 	bool UseLowDetailNodeContent() const;
 
@@ -86,14 +86,14 @@ protected:
 	FText GetInstructionCountText() const;
 	FText GetInstructionDurationText() const;
 
-protected:
+private:
 
 	int32 GetNodeTopologyVersion() const;
 	EVisibility GetPinVisibility(int32 InPinInfoIndex, bool bAskingForSubPin) const;
 	const FSlateBrush * GetExpanderImage(int32 InPinInfoIndex, bool bLeft, bool bHovered) const;
 	FReply OnExpanderArrowClicked(int32 InPinInfoIndex);
 	void HandleModifiedEvent(ERigVMGraphNotifType InNotifType, URigVMGraph* InGraph, UObject* InSubject);
-	virtual void UpdatePinTreeView();
+	void UpdatePinTreeView();
 
 	/** Cached widget title area */
 	TSharedPtr<SOverlay> TitleAreaWidget;
