AnimNext: Give AnimNextTestSuite an independent module class name from AnimNext

CookedEditor was getting confused between the AnimNext and AnimNextTestSuite modules because the implementation class name was the same between the two. This led to AnimNextTestSuite invoking the AnimNext implementation, leading to an ensure in FDecoratorRegistry::Init.

#rnx

diff --git a/Engine/Plugins/Experimental/AnimNext/Source/AnimNextTestSuite/Private/Module.cpp b/Engine/Plugins/Experimental/AnimNext/Source/AnimNextTestSuite/Private/Module.cpp
--- a/Engine/Plugins/Experimental/AnimNext/Source/AnimNextTestSuite/Private/Module.cpp
+++ b/Engine/Plugins/Experimental/AnimNext/Source/AnimNextTestSuite/Private/Module.cpp
@@ -4,7 +4,7 @@
 #include "Modules/ModuleInterface.h"
 #include "Modules/ModuleManager.h"
 
-namespace UE::AnimNext
+namespace UE::AnimNext::Tests
 {
 
 class FModule : public IModuleInterface
@@ -14,4 +14,4 @@ public:
 
 }
 
-IMPLEMENT_MODULE(UE::AnimNext::FModule, AnimNextTestSuite)
+IMPLEMENT_MODULE(UE::AnimNext::Tests::FModule, AnimNextTestSuite)
