Add AnimNext graph unit tests

diff --git a/Engine/Plugins/Experimental/AnimNext/Source/AnimNextEditor/Private/Param/ParametersEditorSchemaActions.cpp b/Engine/Plugins/Experimental/AnimNext/Source/AnimNextEditor/Private/Param/ParametersEditorSchemaActions.cpp
--- a/Engine/Plugins/Experimental/AnimNext/Source/AnimNextEditor/Private/Param/ParametersEditorSchemaActions.cpp
+++ b/Engine/Plugins/Experimental/AnimNext/Source/AnimNextEditor/Private/Param/ParametersEditorSchemaActions.cpp
@@ -24,7 +24,7 @@ UEdGraphNode* FAnimNextParameterSchemaAction_RigUnit::PerformAction(UEdGraph* Pa
 		FRigVMUnitNodeCreatedContext& UnitNodeCreatedContext = Controller->GetUnitNodeCreatedContext();
 		FRigVMUnitNodeCreatedContext::FScope ReasonScope(UnitNodeCreatedContext, ERigVMNodeCreatedReason::NodeSpawner);
 
-		if (URigVMUnitNode* ModelNode = Controller->AddUnitNode(StructTemplate, FRigUnit::GetMethodName(), Location, Name.ToString(), true, false))
+		if (URigVMUnitNode* ModelNode = Controller->AddUnitNode(StructTemplate, FRigVMStruct::ExecuteName, Location, Name.ToString(), true, false))
 		{
 			NewNode = Cast<UAnimNextParameterBlock_EdGraphNode>(EdGraph->FindNodeForModelNodeName(ModelNode->GetFName()));
 			check(NewNode);
diff --git a/Engine/Plugins/Experimental/AnimNext/Source/AnimNextTestSuite/AnimNextTestSuite.Build.cs b/Engine/Plugins/Experimental/AnimNext/Source/AnimNextTestSuite/AnimNextTestSuite.Build.cs
--- a/Engine/Plugins/Experimental/AnimNext/Source/AnimNextTestSuite/AnimNextTestSuite.Build.cs
+++ b/Engine/Plugins/Experimental/AnimNext/Source/AnimNextTestSuite/AnimNextTestSuite.Build.cs
@@ -15,7 +15,8 @@ namespace UnrealBuildTool.Rules
 					"CoreUObject",
 					"Engine",
 					"AnimNext",
-					"StructUtils"
+					"StructUtils",
+					"RigVM",
 				}
 			);
 
@@ -27,7 +28,8 @@ namespace UnrealBuildTool.Rules
 						"UnrealEd",
 						"AnimNextUncookedOnly",
 						"AnimNextEditor",
-						"PythonScriptPlugin"
+						"PythonScriptPlugin",
+						"RigVMDeveloper",
 					}
 				);
 			}
diff --git a/Engine/Plugins/Experimental/AnimNext/Source/AnimNextTestSuite/Private/AnimNextDecoratorGraphTest.cpp b/Engine/Plugins/Experimental/AnimNext/Source/AnimNextTestSuite/Private/AnimNextDecoratorGraphTest.cpp
new file mode 100644
--- /dev/null
+++ b/Engine/Plugins/Experimental/AnimNext/Source/AnimNextTestSuite/Private/AnimNextDecoratorGraphTest.cpp
@@ -0,0 +1,217 @@
+// Copyright Epic Games, Inc. All Rights Reserved.
+
+#include "AnimNextDecoratorGraphTest.h"
+
+#include "AssetToolsModule.h"
+#include "Context.h"
+#include "UncookedOnlyUtils.h"
+#include "DecoratorBase/DecoratorRegistry.h"
+#include "DecoratorInterfaces/IEvaluate.h"
+#include "DecoratorInterfaces/IUpdate.h"
+#include "Graph/AnimNextExecuteContext.h"
+#include "Graph/AnimNextGraph.h"
+#include "Graph/AnimNextGraph_EditorData.h"
+#include "Graph/RigDecorator_AnimNextCppDecorator.h"
+#include "Graph/RigUnit_AnimNextGraphRoot.h"
+#include "Graph/RigUnit_AnimNextDecoratorStack.h"
+#include "Misc/AutomationTest.h"
+#include "Param/ParamStack.h"
+
+#if WITH_DEV_AUTOMATION_TESTS
+
+void UAnimNextGraphTest::SetEditorData(UAnimNextGraph_EditorData* InEditorData)
+{
+	EditorData = InEditorData;
+}
+
+//****************************************************************************
+// AnimNext Runtime Decorator Graph Tests
+//****************************************************************************
+
+namespace UE::AnimNext
+{
+	struct FTestDecorator final : FBaseDecorator, IEvaluate, IUpdate
+	{
+		DECLARE_ANIM_DECORATOR(FTestDecorator, 0x70df2a6a, FBaseDecorator)
+
+		using FSharedData = FTestDecoratorSharedData;
+
+		// IUpdate impl
+		virtual void PostUpdate(FExecutionContext& Context, const TDecoratorBinding<IUpdate>& Binding) const override
+		{
+			IUpdate::PostUpdate(Context, Binding);
+
+			const FSharedData* SharedData = Binding.GetSharedData<FSharedData>();
+			UE::AnimNext::FParamStack& ParamStack = UE::AnimNext::FParamStack::Get();
+
+			int32& UpdateCount = ParamStack.GetMutableParam<int32>("UpdateCount");
+			UpdateCount++;
+
+			ParamStack.GetMutableParam<int32>("SomeInt32") = SharedData->SomeInt32;
+			ParamStack.GetMutableParam<float>("SomeFloat") = SharedData->SomeFloat;
+		}
+
+		// IEvaluate impl
+		virtual void PostEvaluate(FExecutionContext& Context, const TDecoratorBinding<IEvaluate>& Binding) const override
+		{
+			IEvaluate::PostEvaluate(Context, Binding);
+
+			UE::AnimNext::FParamStack& ParamStack = UE::AnimNext::FParamStack::Get();
+
+			int32& EvaluateCount = ParamStack.GetMutableParam<int32>("EvaluateCount");
+			EvaluateCount++;
+		}
+	};
+
+	DEFINE_ANIM_DECORATOR_BEGIN(FTestDecorator)
+		DEFINE_ANIM_DECORATOR_IMPLEMENTS_INTERFACE(IEvaluate)
+		DEFINE_ANIM_DECORATOR_IMPLEMENTS_INTERFACE(IUpdate)
+	DEFINE_ANIM_DECORATOR_END(FTestDecorator)
+	AUTO_REGISTER_ANIM_DECORATOR(FTestDecorator)
+}
+
+IMPLEMENT_SIMPLE_AUTOMATION_TEST(FAnimationAnimNextRuntimeTest_GraphAddDecorator, "Animation.AnimNext.Runtime.Graph.AddDecorator", EAutomationTestFlags::EditorContext | EAutomationTestFlags::EngineFilter)
+
+bool FAnimationAnimNextRuntimeTest_GraphAddDecorator::RunTest(const FString& InParameters)
+{
+	UAnimNextGraphTest* AnimNextGraph = NewObject<UAnimNextGraphTest>();
+
+	UAnimNextGraph_EditorData* EditorData = NewObject<UAnimNextGraph_EditorData>(AnimNextGraph, TEXT("EditorData"));
+	AnimNextGraph->SetEditorData(EditorData);
+
+	EditorData->Initialize(/*bRecompileVM*/false);
+	EditorData->GetRigVMClient()->SetExecuteContextStruct(FAnimNextExecuteContext::StaticStruct());
+
+	URigVMController* Controller = EditorData->GetRigVMClient()->GetController(EditorData->GetRootGraph());
+
+	// Create an empty decorator stack node
+	URigVMUnitNode* DecoratorStackNode = Controller->AddUnitNode(FRigUnit_AnimNextDecoratorStack::StaticStruct(), FRigVMStruct::ExecuteName, FVector2D(0.0f, 0.0f), FString(), false);
+	UE_RETURN_ON_ERROR(DecoratorStackNode != nullptr, "FAnimationAnimNextRuntimeTest_GraphAddDecorator -> Failed to create decorator stack node");
+
+	// Add a decorator
+	const UScriptStruct* CppDecoratorStruct = FRigDecorator_AnimNextCppDecorator::StaticStruct();
+	const UE::AnimNext::FDecorator* Decorator = UE::AnimNext::FDecoratorRegistry::Get().Find(UE::AnimNext::FTestDecorator::DecoratorUID);
+	UScriptStruct* ScriptStruct = Decorator->GetDecoratorSharedDataStruct();
+
+	FString DefaultValue;
+	{
+		const FRigDecorator_AnimNextCppDecorator DefaultCppDecoratorStructInstance;
+		FRigDecorator_AnimNextCppDecorator CppDecoratorStructInstance;
+		CppDecoratorStructInstance.DecoratorSharedDataStruct = ScriptStruct;
+
+		UE_RETURN_ON_ERROR(CppDecoratorStructInstance.CanBeAddedToNode(DecoratorStackNode, nullptr), "FAnimationAnimNextRuntimeTest_GraphAddDecorator -> Decorator cannot be added to decorator stack node");
+
+		const FProperty* Prop = FAnimNextCppDecoratorWrapper::StaticStruct()->FindPropertyByName(GET_MEMBER_NAME_STRING_CHECKED(FAnimNextCppDecoratorWrapper, CppDecorator));
+		Prop->ExportText_Direct(DefaultValue, &CppDecoratorStructInstance, &DefaultCppDecoratorStructInstance, nullptr, PPF_None);
+	}
+
+	FString DisplayNameMetadata;
+	ScriptStruct->GetStringMetaDataHierarchical(FRigVMStruct::DisplayNameMetaName, &DisplayNameMetadata);
+	const FString DisplayName = DisplayNameMetadata.IsEmpty() ? Decorator->GetDecoratorUID().GetDecoratorName() : DisplayNameMetadata;
+
+	const FName DecoratorName = Controller->AddDecorator(
+		DecoratorStackNode->GetFName(),
+		*CppDecoratorStruct->GetPathName(),
+		*DisplayName,
+		DefaultValue, INDEX_NONE, true, true);
+	UE_RETURN_ON_ERROR(DecoratorName == DisplayName, TEXT("FAnimationAnimNextRuntimeTest_GraphAddDecorator -> Unexpected decorator name"));
+
+	// Our first pin is the hard coded output result, decorator pins follow
+	UE_RETURN_ON_ERROR(DecoratorStackNode->GetPins().Num() == 2, TEXT("FAnimationAnimNextRuntimeTest_GraphAddDecorator -> Unexpected number of pins"));
+	UE_RETURN_ON_ERROR(DecoratorStackNode->GetPins()[1]->IsDecoratorPin() == true, TEXT("FAnimationAnimNextRuntimeTest_GraphAddDecorator -> Unexpected pin type"));
+	UE_RETURN_ON_ERROR(DecoratorStackNode->GetPins()[1]->GetFName() == DecoratorName, TEXT("FAnimationAnimNextRuntimeTest_GraphAddDecorator -> Unexpected pin name"));
+	UE_RETURN_ON_ERROR(DecoratorStackNode->GetPins()[1]->GetCPPTypeObject() == FRigDecorator_AnimNextCppDecorator::StaticStruct(), TEXT("FAnimationAnimNextRuntimeTest_GraphAddDecorator -> Unexpected pin type"));
+
+	// Our first sub-pin is the hard coded script struct member that parametrizes the decorator, dynamic decorator sub-pins follow
+	UE_RETURN_ON_ERROR(DecoratorStackNode->GetPins()[1]->GetSubPins().Num() == 3, TEXT("FAnimationAnimNextRuntimeTest_GraphAddDecorator -> Unexpected decorator sub pins"));
+	UE_RETURN_ON_ERROR(DecoratorStackNode->GetPins()[1]->GetSubPins()[1]->GetCPPType() == TEXT("int32"), TEXT("FAnimationAnimNextRuntimeTest_GraphAddDecorator -> Unexpected decorator pin type"));
+	UE_RETURN_ON_ERROR(DecoratorStackNode->GetPins()[1]->GetSubPins()[1]->GetDefaultValue() == TEXT("3"), TEXT("FAnimationAnimNextRuntimeTest_GraphAddDecorator -> Unexpected decorator pin value"));
+	UE_RETURN_ON_ERROR(DecoratorStackNode->GetPins()[1]->GetSubPins()[2]->GetCPPType() == TEXT("float"), TEXT("FAnimationAnimNextRuntimeTest_GraphAddDecorator -> Unexpected decorator pin type"));
+	UE_RETURN_ON_ERROR(DecoratorStackNode->GetPins()[1]->GetSubPins()[2]->GetDefaultValue() == TEXT("34.000000"), TEXT("FAnimationAnimNextRuntimeTest_GraphAddDecorator -> Unexpected decorator pin value"));
+
+	return true;
+}
+
+IMPLEMENT_SIMPLE_AUTOMATION_TEST(FAnimationAnimNextRuntimeTest_GraphExecute, "Animation.AnimNext.Runtime.Graph.Execute", EAutomationTestFlags::EditorContext | EAutomationTestFlags::EngineFilter)
+
+bool FAnimationAnimNextRuntimeTest_GraphExecute::RunTest(const FString& InParameters)
+{
+	UAnimNextGraphTest* AnimNextGraph = NewObject<UAnimNextGraphTest>();
+
+	UAnimNextGraph_EditorData* EditorData = NewObject<UAnimNextGraph_EditorData>(AnimNextGraph, TEXT("EditorData"));
+	AnimNextGraph->SetEditorData(EditorData);
+
+	EditorData->Initialize(/*bRecompileVM*/false);
+	EditorData->GetRigVMClient()->SetExecuteContextStruct(FAnimNextExecuteContext::StaticStruct());
+
+	URigVMController* Controller = EditorData->GetRigVMClient()->GetController(EditorData->GetRootGraph());
+
+	// Add graph entry point
+	URigVMUnitNode* MainEntryPointNode = Controller->AddUnitNode(FRigUnit_AnimNextGraphRoot::StaticStruct(), FRigVMStruct::ExecuteName, FVector2D(0.0f, 0.0f), FString(), false);
+	URigVMPin* BeginExecutePin = MainEntryPointNode->FindPin(GET_MEMBER_NAME_STRING_CHECKED(FRigUnit_AnimNextGraphRoot, Result));
+	UE_RETURN_ON_ERROR(BeginExecutePin != nullptr && BeginExecutePin->GetDirection() == ERigVMPinDirection::Input, "FAnimationAnimNextRuntimeTest_GraphExecute -> Failed to create entry point");
+
+	// Create an empty decorator stack node
+	URigVMUnitNode* DecoratorStackNode = Controller->AddUnitNode(FRigUnit_AnimNextDecoratorStack::StaticStruct(), FRigVMStruct::ExecuteName, FVector2D(0.0f, 0.0f), FString(), false);
+	UE_RETURN_ON_ERROR(DecoratorStackNode != nullptr, "FAnimationAnimNextRuntimeTest_GraphExecute -> Failed to create decorator stack node");
+
+	// Link our stack result to our entry point
+	Controller->AddLink(DecoratorStackNode->GetPins()[0], MainEntryPointNode->GetPins()[0]);
+
+	// Add a decorator
+	const UScriptStruct* CppDecoratorStruct = FRigDecorator_AnimNextCppDecorator::StaticStruct();
+	const UE::AnimNext::FDecorator* Decorator = UE::AnimNext::FDecoratorRegistry::Get().Find(UE::AnimNext::FTestDecorator::DecoratorUID);
+	UScriptStruct* ScriptStruct = Decorator->GetDecoratorSharedDataStruct();
+
+	FString DefaultValue;
+	{
+		const FRigDecorator_AnimNextCppDecorator DefaultCppDecoratorStructInstance;
+		FRigDecorator_AnimNextCppDecorator CppDecoratorStructInstance;
+		CppDecoratorStructInstance.DecoratorSharedDataStruct = ScriptStruct;
+
+		UE_RETURN_ON_ERROR(CppDecoratorStructInstance.CanBeAddedToNode(DecoratorStackNode, nullptr), "FAnimationAnimNextRuntimeTest_GraphExecute -> Decorator cannot be added to decorator stack node");
+
+		const FProperty* Prop = FAnimNextCppDecoratorWrapper::StaticStruct()->FindPropertyByName(GET_MEMBER_NAME_STRING_CHECKED(FAnimNextCppDecoratorWrapper, CppDecorator));
+		Prop->ExportText_Direct(DefaultValue, &CppDecoratorStructInstance, &DefaultCppDecoratorStructInstance, nullptr, PPF_None);
+	}
+
+	FString DisplayNameMetadata;
+	ScriptStruct->GetStringMetaDataHierarchical(FRigVMStruct::DisplayNameMetaName, &DisplayNameMetadata);
+	const FString DisplayName = DisplayNameMetadata.IsEmpty() ? Decorator->GetDecoratorUID().GetDecoratorName() : DisplayNameMetadata;
+
+	const FName DecoratorName = Controller->AddDecorator(
+		DecoratorStackNode->GetFName(),
+		*CppDecoratorStruct->GetPathName(),
+		*DisplayName,
+		DefaultValue, INDEX_NONE, true, true);
+	UE_RETURN_ON_ERROR(DecoratorName == DisplayName, TEXT("FAnimationAnimNextRuntimeTest_GraphExecute -> Unexpected decorator name"));
+
+	// Set some values on our decorator
+	Controller->SetPinDefaultValue(DecoratorStackNode->GetPins()[1]->GetSubPins()[1]->GetPinPath(), TEXT("78"));
+	Controller->SetPinDefaultValue(DecoratorStackNode->GetPins()[1]->GetSubPins()[2]->GetPinPath(), TEXT("142.33"));
+
+	UE::AnimNext::FDecoratorPtr GraphInstance = AnimNextGraph->AllocateInstance();
+
+	UE::AnimNext::FContext Context(1.0f / 30.0f);
+
+	UE::AnimNext::FParamStack::FPushedLayerHandle LayerHandle = Context.GetMutableParamStack().PushValues(
+		"UpdateCount", (int32)0,
+		"EvaluateCount", (int32)0,
+		"SomeInt32", (int32)0,
+		"SomeFloat", 0.0f
+	);
+
+	AnimNextGraph->Run(Context, GraphInstance, EAnimNextGraphSimulationSteps::All);
+
+	AddErrorIfFalse(Context.GetParamStack().GetParam<int32>("UpdateCount") == 1, "FAnimationAnimNextRuntimeTest_GraphExecute -> Unexpected update count");
+	AddErrorIfFalse(Context.GetParamStack().GetParam<int32>("EvaluateCount") == 1, "FAnimationAnimNextRuntimeTest_GraphExecute -> Unexpected evaluate count");
+	AddErrorIfFalse(Context.GetParamStack().GetParam<int32>("SomeInt32") == 78, "FAnimationAnimNextRuntimeTest_GraphExecute -> Unexpected SomeInt32 value");
+	AddErrorIfFalse(Context.GetParamStack().GetParam<float>("SomeFloat") == 142.33f, "FAnimationAnimNextRuntimeTest_GraphExecute -> Unexpected SomeFloat value");
+
+	Context.GetMutableParamStack().PopLayer(LayerHandle);
+	AnimNextGraph->ReleaseInstance(GraphInstance);
+
+	return true;
+}
+
+#endif
diff --git a/Engine/Plugins/Experimental/AnimNext/Source/AnimNextTestSuite/Private/AnimNextDecoratorGraphTest.h b/Engine/Plugins/Experimental/AnimNext/Source/AnimNextTestSuite/Private/AnimNextDecoratorGraphTest.h
new file mode 100644
--- /dev/null
+++ b/Engine/Plugins/Experimental/AnimNext/Source/AnimNextTestSuite/Private/AnimNextDecoratorGraphTest.h
@@ -0,0 +1,31 @@
+// Copyright Epic Games, Inc. All Rights Reserved.
+
+#pragma once
+
+#include "CoreMinimal.h"
+
+#include "DecoratorBase/DecoratorSharedData.h"
+#include "Graph/AnimNextGraph.h"
+
+#include "AnimNextDecoratorGraphTest.generated.h"
+
+UCLASS(BlueprintType)
+class UAnimNextGraphTest : public UAnimNextGraph
+{
+	GENERATED_BODY()
+
+public:
+	void SetEditorData(class UAnimNextGraph_EditorData* InEditorData);
+};
+
+USTRUCT()
+struct FTestDecoratorSharedData : public FAnimNextDecoratorSharedData
+{
+	GENERATED_BODY()
+
+	UPROPERTY(meta = (Input))
+	int32 SomeInt32 = 3;
+
+	UPROPERTY(meta = (Input))
+	float SomeFloat = 34.0f;
+};
