Fix editor shipping configuration by adding decorator name function on FDecorator

diff --git a/Engine/Plugins/Experimental/AnimNext/Source/AnimNext/Public/DecoratorBase/Decorator.h b/Engine/Plugins/Experimental/AnimNext/Source/AnimNext/Public/DecoratorBase/Decorator.h
--- a/Engine/Plugins/Experimental/AnimNext/Source/AnimNext/Public/DecoratorBase/Decorator.h
+++ b/Engine/Plugins/Experimental/AnimNext/Source/AnimNext/Public/DecoratorBase/Decorator.h
@@ -21,6 +21,7 @@ class FArchive;
 	/* FDecorator impl */ \
 	static constexpr UE::AnimNext::FDecoratorUID DecoratorUID = UE::AnimNext::FDecoratorUID(DecoratorNameHash, TEXT(#DecoratorName)); \
 	virtual UE::AnimNext::FDecoratorUID GetDecoratorUID() const override { return DecoratorUID; } \
+	virtual FString GetDecoratorName() const { return TEXT(#DecoratorName); } \
 	static const UE::AnimNext::FDecoratorMemoryLayout DecoratorMemoryDescription; \
 	virtual UE::AnimNext::FDecoratorMemoryLayout GetDecoratorMemoryDescription() const override { return DecoratorMemoryDescription; } \
 	virtual UScriptStruct* GetDecoratorSharedDataStruct() const override { return FSharedData::StaticStruct(); } \
@@ -34,7 +35,8 @@ class FArchive;
 	using DecoratorSuper = SuperDecoratorName; \
 	/* FDecorator impl */ \
 	static constexpr UE::AnimNext::FDecoratorUID DecoratorUID = UE::AnimNext::FDecoratorUID(DecoratorNameHash, TEXT(#DecoratorName)); \
-	virtual UE::AnimNext::FDecoratorUID GetDecoratorUID() const override { return DecoratorUID; }
+	virtual UE::AnimNext::FDecoratorUID GetDecoratorUID() const override { return DecoratorUID; } \
+	virtual FString GetDecoratorName() const { return TEXT(#DecoratorName); }
 
 // In the decorator cpp, these three macros implement the base functionality
 // 
@@ -157,6 +159,9 @@ namespace UE::AnimNext
 		// Returns the globally unique UID for this decorator
 		virtual FDecoratorUID GetDecoratorUID() const { return DecoratorUID; };
 
+		// Returns the decorator name
+		virtual FString GetDecoratorName() const { return TEXT("FDecorator"); }
+
 		// Returns the memory requirements of the derived decorator instance
 		virtual FDecoratorMemoryLayout GetDecoratorMemoryDescription() const { return { sizeof(FDecorator), alignof(FDecorator) }; }
 
diff --git a/Engine/Plugins/Experimental/AnimNext/Source/AnimNextUncookedOnly/Private/Graph/AnimNextGraph_EdGraphNode.cpp b/Engine/Plugins/Experimental/AnimNext/Source/AnimNextUncookedOnly/Private/Graph/AnimNextGraph_EdGraphNode.cpp
--- a/Engine/Plugins/Experimental/AnimNext/Source/AnimNextUncookedOnly/Private/Graph/AnimNextGraph_EdGraphNode.cpp
+++ b/Engine/Plugins/Experimental/AnimNext/Source/AnimNextUncookedOnly/Private/Graph/AnimNextGraph_EdGraphNode.cpp
@@ -72,12 +72,12 @@ void UAnimNextGraph_EdGraphNode::BuildAddDecoratorContextMenu(UToolMenu* SubMenu
 
 		FString DisplayNameMetadata;
 		ScriptStruct->GetStringMetaDataHierarchical(FRigVMStruct::DisplayNameMetaName, &DisplayNameMetadata);
-		const FString DisplayName = DisplayNameMetadata.IsEmpty() ? Decorator->GetDecoratorUID().GetDecoratorName() : DisplayNameMetadata;
+		const FString DisplayName = DisplayNameMetadata.IsEmpty() ? Decorator->GetDecoratorName() : DisplayNameMetadata;
 
 		const FText ToolTip = ScriptStruct->GetToolTipText();
 
 		FToolMenuEntry DecoratorEntry = FToolMenuEntry::InitMenuEntry(
-			Decorator->GetDecoratorUID().GetDecoratorName(),
+			*Decorator->GetDecoratorName(),
 			FText::FromString(DisplayName),
 			ToolTip,
 			FSlateIcon(),
