Fix static analysis issues with anim next unit test

diff --git a/Engine/Plugins/Experimental/AnimNext/Source/AnimNextTestSuite/Private/AnimNextDecoratorGraphTest.cpp b/Engine/Plugins/Experimental/AnimNext/Source/AnimNextTestSuite/Private/AnimNextDecoratorGraphTest.cpp
--- a/Engine/Plugins/Experimental/AnimNext/Source/AnimNextTestSuite/Private/AnimNextDecoratorGraphTest.cpp
+++ b/Engine/Plugins/Experimental/AnimNext/Source/AnimNextTestSuite/Private/AnimNextDecoratorGraphTest.cpp
@@ -83,6 +83,7 @@ bool FAnimationAnimNextRuntimeTest_GraphAddDecorator::RunTest(const FString& InP
 	EditorData->GetRigVMClient()->SetExecuteContextStruct(FAnimNextExecuteContext::StaticStruct());
 
 	URigVMController* Controller = EditorData->GetRigVMClient()->GetController(EditorData->GetRootGraph());
+	UE_RETURN_ON_ERROR(Controller != nullptr, "FAnimationAnimNextRuntimeTest_GraphAddDecorator -> Failed to get RigVM controller");
 
 	// Create an empty decorator stack node
 	URigVMUnitNode* DecoratorStackNode = Controller->AddUnitNode(FRigUnit_AnimNextDecoratorStack::StaticStruct(), FRigVMStruct::ExecuteName, FVector2D(0.0f, 0.0f), FString(), false);
@@ -90,8 +91,13 @@ bool FAnimationAnimNextRuntimeTest_GraphAddDecorator::RunTest(const FString& InP
 
 	// Add a decorator
 	const UScriptStruct* CppDecoratorStruct = FRigDecorator_AnimNextCppDecorator::StaticStruct();
+	UE_RETURN_ON_ERROR(CppDecoratorStruct != nullptr, "FAnimationAnimNextRuntimeTest_GraphAddDecorator -> Failed to get find Cpp decorator static struct");
+
 	const UE::AnimNext::FDecorator* Decorator = UE::AnimNext::FDecoratorRegistry::Get().Find(UE::AnimNext::FTestDecorator::DecoratorUID);
+	UE_RETURN_ON_ERROR(Decorator != nullptr, "FAnimationAnimNextRuntimeTest_GraphAddDecorator -> Failed to find test decorator");
+
 	UScriptStruct* ScriptStruct = Decorator->GetDecoratorSharedDataStruct();
+	UE_RETURN_ON_ERROR(ScriptStruct != nullptr, "FAnimationAnimNextRuntimeTest_GraphAddDecorator -> Failed to find decorator shared data struct");
 
 	FString DefaultValue;
 	{
@@ -102,12 +108,14 @@ bool FAnimationAnimNextRuntimeTest_GraphAddDecorator::RunTest(const FString& InP
 		UE_RETURN_ON_ERROR(CppDecoratorStructInstance.CanBeAddedToNode(DecoratorStackNode, nullptr), "FAnimationAnimNextRuntimeTest_GraphAddDecorator -> Decorator cannot be added to decorator stack node");
 
 		const FProperty* Prop = FAnimNextCppDecoratorWrapper::StaticStruct()->FindPropertyByName(GET_MEMBER_NAME_STRING_CHECKED(FAnimNextCppDecoratorWrapper, CppDecorator));
+		UE_RETURN_ON_ERROR(Prop != nullptr, "FAnimationAnimNextRuntimeTest_GraphAddDecorator -> Failed to find wrapper property");
+
 		Prop->ExportText_Direct(DefaultValue, &CppDecoratorStructInstance, &DefaultCppDecoratorStructInstance, nullptr, PPF_None);
 	}
 
 	FString DisplayNameMetadata;
 	ScriptStruct->GetStringMetaDataHierarchical(FRigVMStruct::DisplayNameMetaName, &DisplayNameMetadata);
-	const FString DisplayName = DisplayNameMetadata.IsEmpty() ? Decorator->GetDecoratorUID().GetDecoratorName() : DisplayNameMetadata;
+	const FString DisplayName = DisplayNameMetadata.IsEmpty() ? Decorator->GetDecoratorName() : DisplayNameMetadata;
 
 	const FName DecoratorName = Controller->AddDecorator(
 		DecoratorStackNode->GetFName(),
@@ -145,9 +153,12 @@ bool FAnimationAnimNextRuntimeTest_GraphExecute::RunTest(const FString& InParame
 	EditorData->GetRigVMClient()->SetExecuteContextStruct(FAnimNextExecuteContext::StaticStruct());
 
 	URigVMController* Controller = EditorData->GetRigVMClient()->GetController(EditorData->GetRootGraph());
+	UE_RETURN_ON_ERROR(Controller != nullptr, "FAnimationAnimNextRuntimeTest_GraphExecute -> Failed to get RigVM controller");
 
 	// Add graph entry point
 	URigVMUnitNode* MainEntryPointNode = Controller->AddUnitNode(FRigUnit_AnimNextGraphRoot::StaticStruct(), FRigVMStruct::ExecuteName, FVector2D(0.0f, 0.0f), FString(), false);
+	UE_RETURN_ON_ERROR(MainEntryPointNode != nullptr, "FAnimationAnimNextRuntimeTest_GraphExecute -> Failed to find main entry point node");
+
 	URigVMPin* BeginExecutePin = MainEntryPointNode->FindPin(GET_MEMBER_NAME_STRING_CHECKED(FRigUnit_AnimNextGraphRoot, Result));
 	UE_RETURN_ON_ERROR(BeginExecutePin != nullptr && BeginExecutePin->GetDirection() == ERigVMPinDirection::Input, "FAnimationAnimNextRuntimeTest_GraphExecute -> Failed to create entry point");
 
@@ -160,8 +171,13 @@ bool FAnimationAnimNextRuntimeTest_GraphExecute::RunTest(const FString& InParame
 
 	// Add a decorator
 	const UScriptStruct* CppDecoratorStruct = FRigDecorator_AnimNextCppDecorator::StaticStruct();
+	UE_RETURN_ON_ERROR(CppDecoratorStruct != nullptr, "FAnimationAnimNextRuntimeTest_GraphExecute -> Failed to get find Cpp decorator static struct");
+
 	const UE::AnimNext::FDecorator* Decorator = UE::AnimNext::FDecoratorRegistry::Get().Find(UE::AnimNext::FTestDecorator::DecoratorUID);
+	UE_RETURN_ON_ERROR(Decorator != nullptr, "FAnimationAnimNextRuntimeTest_GraphExecute -> Failed to find test decorator");
+
 	UScriptStruct* ScriptStruct = Decorator->GetDecoratorSharedDataStruct();
+	UE_RETURN_ON_ERROR(ScriptStruct != nullptr, "FAnimationAnimNextRuntimeTest_GraphExecute -> Failed to find decorator shared data struct");
 
 	FString DefaultValue;
 	{
@@ -172,12 +188,14 @@ bool FAnimationAnimNextRuntimeTest_GraphExecute::RunTest(const FString& InParame
 		UE_RETURN_ON_ERROR(CppDecoratorStructInstance.CanBeAddedToNode(DecoratorStackNode, nullptr), "FAnimationAnimNextRuntimeTest_GraphExecute -> Decorator cannot be added to decorator stack node");
 
 		const FProperty* Prop = FAnimNextCppDecoratorWrapper::StaticStruct()->FindPropertyByName(GET_MEMBER_NAME_STRING_CHECKED(FAnimNextCppDecoratorWrapper, CppDecorator));
+		UE_RETURN_ON_ERROR(Prop != nullptr, "FAnimationAnimNextRuntimeTest_GraphExecute -> Failed to find wrapper property");
+
 		Prop->ExportText_Direct(DefaultValue, &CppDecoratorStructInstance, &DefaultCppDecoratorStructInstance, nullptr, PPF_None);
 	}
 
 	FString DisplayNameMetadata;
 	ScriptStruct->GetStringMetaDataHierarchical(FRigVMStruct::DisplayNameMetaName, &DisplayNameMetadata);
-	const FString DisplayName = DisplayNameMetadata.IsEmpty() ? Decorator->GetDecoratorUID().GetDecoratorName() : DisplayNameMetadata;
+	const FString DisplayName = DisplayNameMetadata.IsEmpty() ? Decorator->GetDecoratorName() : DisplayNameMetadata;
 
 	const FName DecoratorName = Controller->AddDecorator(
 		DecoratorStackNode->GetFName(),
