AnimNext: Fix cross-thread param type handle AnimNext tests

Sandboxed tests were only returning sandboxed maps on the main thread due to thread ID comparison, so ID generated on other threads were not correctly sandboxed.

diff --git a/Engine/Plugins/Experimental/AnimNext/Source/AnimNext/Private/Param/ParamId.cpp b/Engine/Plugins/Experimental/AnimNext/Source/AnimNext/Private/Param/ParamId.cpp
--- a/Engine/Plugins/Experimental/AnimNext/Source/AnimNext/Private/Param/ParamId.cpp
+++ b/Engine/Plugins/Experimental/AnimNext/Source/AnimNext/Private/Param/ParamId.cpp
@@ -19,14 +19,13 @@ static FParamIdGlobalData GParamIdGlobalData;
 
 #if WITH_DEV_AUTOMATION_TESTS
 static FParamIdGlobalData GSandboxedParamIdGlobalData;
-static uint32 GParamIdSandboxedThreadId = MAX_uint32;
-static bool bGParamIdSandboxed = false;
+static std::atomic<bool> bGParamIdSandboxed = false;
 #endif
 
 static FParamIdGlobalData& GetParamIdData()
 {
 #if WITH_DEV_AUTOMATION_TESTS		
-	if (bGParamIdSandboxed && GParamIdSandboxedThreadId == FPlatformTLS::GetCurrentThreadId())
+	if (bGParamIdSandboxed.load() == true)
 	{
 		return GSandboxedParamIdGlobalData;
 	}
@@ -68,9 +67,8 @@ FParamId FParamId::GetMaxParamId()
 void FParamId::BeginTestSandbox()
 {
 	FRWScopeLock Lock(GParamIdLock, SLT_Write);
-	check(bGParamIdSandboxed == false);
-	bGParamIdSandboxed = true;
-	GParamIdSandboxedThreadId = FPlatformTLS::GetCurrentThreadId();
+	check(bGParamIdSandboxed.load() == false);
+	bGParamIdSandboxed.exchange(true);
 	GSandboxedParamIdGlobalData.NameToParamId.Empty();
 	GSandboxedParamIdGlobalData.ParamIdToName.Empty();
 }
@@ -78,8 +76,8 @@ void FParamId::BeginTestSandbox()
 void FParamId::EndTestSandbox()
 {
 	FRWScopeLock Lock(GParamIdLock, SLT_Write);
-	check(bGParamIdSandboxed == true);
-	bGParamIdSandboxed = false;
+	check(bGParamIdSandboxed.load() == true);
+	bGParamIdSandboxed.exchange(false);
 	GSandboxedParamIdGlobalData.NameToParamId.Empty();
 	GSandboxedParamIdGlobalData.ParamIdToName.Empty();
 }
diff --git a/Engine/Plugins/Experimental/AnimNext/Source/AnimNext/Private/Param/ParamTypeHandle.cpp b/Engine/Plugins/Experimental/AnimNext/Source/AnimNext/Private/Param/ParamTypeHandle.cpp
--- a/Engine/Plugins/Experimental/AnimNext/Source/AnimNext/Private/Param/ParamTypeHandle.cpp
+++ b/Engine/Plugins/Experimental/AnimNext/Source/AnimNext/Private/Param/ParamTypeHandle.cpp
@@ -25,14 +25,13 @@ static FTypeHandleGlobalData GTypeHandleGlobalData;
 
 #if WITH_DEV_AUTOMATION_TESTS	
 static FTypeHandleGlobalData GSandboxedTypeHandleGlobalData;
-static uint32 GTypeHandleSandboxedThreadId = MAX_uint32;
-static bool bGTypeHandleSandboxed = false;
+static std::atomic<bool> bGTypeHandleSandboxed = false;
 #endif
 
 static FTypeHandleGlobalData& GetTypeHandleData()
 {
 #if WITH_DEV_AUTOMATION_TESTS		
-	if(bGTypeHandleSandboxed && GTypeHandleSandboxedThreadId == FPlatformTLS::GetCurrentThreadId())
+	if (bGTypeHandleSandboxed.load() == true)
 	{
 		return GSandboxedTypeHandleGlobalData;
 	}
@@ -47,9 +46,8 @@ static FTypeHandleGlobalData& GetTypeHandleData()
 void FParamTypeHandle::BeginTestSandbox()
 {
 	FRWScopeLock Lock(GParamTypeHandleLock, SLT_Write);
-	check(bGTypeHandleSandboxed == false);
-	bGTypeHandleSandboxed = true;
-	GTypeHandleSandboxedThreadId = FPlatformTLS::GetCurrentThreadId();
+	check(bGTypeHandleSandboxed.load() == false);
+	bGTypeHandleSandboxed.exchange(true);
 	GSandboxedTypeHandleGlobalData.CustomTypes.Empty();
 	GSandboxedTypeHandleGlobalData.TypeToIndexMap.Empty();
 }
@@ -57,9 +55,8 @@ void FParamTypeHandle::BeginTestSandbox()
 void FParamTypeHandle::EndTestSandbox()
 {
 	FRWScopeLock Lock(GParamTypeHandleLock, SLT_Write);
-	check(bGTypeHandleSandboxed == true);
-	bGTypeHandleSandboxed = false;
-	GTypeHandleSandboxedThreadId = MAX_uint32;
+	check(bGTypeHandleSandboxed.load() == true);
+	bGTypeHandleSandboxed.exchange(false);
 	GSandboxedTypeHandleGlobalData.CustomTypes.Empty();
 	GSandboxedTypeHandleGlobalData.TypeToIndexMap.Empty();
 }
