Fix uninitialized variables in a couple of structs

diff --git a/Engine/Plugins/Experimental/AnimNext/Source/AnimNext/Internal/Graph/AnimGraph/AnimNodeExposedValueHandler_AnimNextParameters.h b/Engine/Plugins/Experimental/AnimNext/Source/AnimNext/Internal/Graph/AnimGraph/AnimNodeExposedValueHandler_AnimNextParameters.h
--- a/Engine/Plugins/Experimental/AnimNext/Source/AnimNext/Internal/Graph/AnimGraph/AnimNodeExposedValueHandler_AnimNextParameters.h
+++ b/Engine/Plugins/Experimental/AnimNext/Source/AnimNext/Internal/Graph/AnimGraph/AnimNodeExposedValueHandler_AnimNextParameters.h
@@ -30,11 +30,11 @@ struct FAnimNodeExposedValueHandler_AnimNextParameters_Entry
 
 	// Property access index
 	UPROPERTY()
-	int32 AccessIndex;
+	int32 AccessIndex = INDEX_NONE;
 
 	// Access operation to perform
 	UPROPERTY()
-	EPropertyAccessCopyType AccessType;
+	EPropertyAccessCopyType AccessType = EPropertyAccessCopyType::None;
 };
 
 USTRUCT()
diff --git a/Engine/Plugins/Experimental/AnimNext/Source/AnimNextUncookedOnly/Internal/Graph/AnimGraph/AnimGraphNodeBinding_AnimNextParameters.h b/Engine/Plugins/Experimental/AnimNext/Source/AnimNextUncookedOnly/Internal/Graph/AnimGraph/AnimGraphNodeBinding_AnimNextParameters.h
--- a/Engine/Plugins/Experimental/AnimNext/Source/AnimNextUncookedOnly/Internal/Graph/AnimGraph/AnimGraphNodeBinding_AnimNextParameters.h
+++ b/Engine/Plugins/Experimental/AnimNext/Source/AnimNextUncookedOnly/Internal/Graph/AnimGraph/AnimGraphNodeBinding_AnimNextParameters.h
@@ -14,7 +14,7 @@ struct FAnimNextAnimGraphNodeParameterBinding
 	FName BindingName;
 
 	UPROPERTY()
-	int32 ArrayIndex;
+	int32 ArrayIndex = INDEX_NONE;
 
 	UPROPERTY()
 	FName ParameterName;
