Add support for latent pins in AnimNext decorators
Misc cleanup of decorator system
Refactored RigVM lazy pin state handling
Fix RigVM incorrect outer used to create storage class

diff --git a/Engine/Plugins/Experimental/AnimNext/Source/AnimNext/Internal/DecoratorBase/DecoratorTemplate.h b/Engine/Plugins/Experimental/AnimNext/Source/AnimNext/Internal/DecoratorBase/DecoratorTemplate.h
--- a/Engine/Plugins/Experimental/AnimNext/Source/AnimNext/Internal/DecoratorBase/DecoratorTemplate.h
+++ b/Engine/Plugins/Experimental/AnimNext/Source/AnimNext/Internal/DecoratorBase/DecoratorTemplate.h
@@ -6,6 +6,7 @@
 #include "DecoratorBase/DecoratorMode.h"
 #include "DecoratorBase/DecoratorUID.h"
 #include "DecoratorBase/DecoratorRegistryHandle.h"
+#include "DecoratorBase/LatentPropertyHandle.h"
 
 #include <type_traits>
 
@@ -27,13 +28,13 @@ namespace UE::AnimNext
 	struct FDecoratorTemplate
 	{
 		// Returns the globally unique identifier for the decorator
-		FDecoratorUID GetUID() const { return FDecoratorUID(UID); }
+		FDecoratorUID GetUID() const noexcept { return FDecoratorUID(UID); }
 
 		// Returns the decorator registry handle
-		FDecoratorRegistryHandle GetRegistryHandle() const { return RegistryHandle; }
+		FDecoratorRegistryHandle GetRegistryHandle() const noexcept { return RegistryHandle; }
 
 		// Returns the decorator mode
-		EDecoratorMode GetMode() const { return static_cast<EDecoratorMode>(Mode); }
+		EDecoratorMode GetMode() const noexcept { return static_cast<EDecoratorMode>(Mode); }
 
 		// For base decorators only, returns the number of additive decorators on top in the stack
 		uint32 GetNumAdditiveDecorators() const
@@ -50,31 +51,49 @@ namespace UE::AnimNext
 		}
 
 		// Returns the offset into the shared data where the descriptor begins, relative to the root of the node's shared data
-		uint32 GetNodeSharedOffset() const { return NodeSharedOffset; }
+		uint32 GetNodeSharedOffset() const noexcept { return NodeSharedOffset; }
+
+		// Returns the offset into the shared data where the latent property handles begin (or 0 if we have no latent properties), relative to the root of the node's shared data
+		uint32 GetNodeSharedLatentPropertyHandlesOffset() const noexcept { return NodeSharedLatentPropertyHandlesOffset; }
 
 		// Returns the offset into the instance data where the descriptor begins, relative to the root of the node's instance data
-		uint32 GetNodeInstanceOffset() const { return NodeInstanceOffset; }
+		uint32 GetNodeInstanceOffset() const noexcept { return NodeInstanceOffset; }
+
+		// Returns whether or not we have latent properties on this decorator
+		bool HasLatentProperties() const noexcept { return NodeSharedLatentPropertyHandlesOffset != 0; }
 
 		// Returns a pointer to the specified decorator description on the current node
-		FAnimNextDecoratorSharedData* GetDecoratorDescription(FNodeDescription& NodeDescription) const
+		FAnimNextDecoratorSharedData* GetDecoratorDescription(FNodeDescription& NodeDescription) const noexcept
 		{
 			return reinterpret_cast<FAnimNextDecoratorSharedData*>(reinterpret_cast<uint8*>(&NodeDescription) + GetNodeSharedOffset());
 		}
 
 		// Returns a pointer to the specified decorator description on the current node
-		const FAnimNextDecoratorSharedData* GetDecoratorDescription(const FNodeDescription& NodeDescription) const
+		const FAnimNextDecoratorSharedData* GetDecoratorDescription(const FNodeDescription& NodeDescription) const noexcept
 		{
 			return reinterpret_cast<const FAnimNextDecoratorSharedData*>(reinterpret_cast<const uint8*>(&NodeDescription) + GetNodeSharedOffset());
 		}
 
+		// Returns a pointer to the specified decorator latent property handles on the current node
+		FLatentPropertyHandle* GetDecoratorLatentPropertyHandles(FNodeDescription& NodeDescription) const noexcept
+		{
+			return reinterpret_cast<FLatentPropertyHandle*>(reinterpret_cast<uint8*>(&NodeDescription) + GetNodeSharedLatentPropertyHandlesOffset());
+		}
+
+		// Returns a pointer to the specified decorator latent property handles on the current node
+		const FLatentPropertyHandle* GetDecoratorLatentPropertyHandles(const FNodeDescription& NodeDescription) const noexcept
+		{
+			return reinterpret_cast<const FLatentPropertyHandle*>(reinterpret_cast<const uint8*>(&NodeDescription) + GetNodeSharedLatentPropertyHandlesOffset());
+		}
+
 		// Returns a pointer to the specified decorator instance on the current node
-		FDecoratorInstanceData* GetDecoratorInstance(FNodeInstance& NodeInstance) const
+		FDecoratorInstanceData* GetDecoratorInstance(FNodeInstance& NodeInstance) const noexcept
 		{
 			return reinterpret_cast<FDecoratorInstanceData*>(reinterpret_cast<uint8*>(&NodeInstance) + GetNodeInstanceOffset());
 		}
 
 		// Returns a pointer to the specified decorator instance on the current node
-		const FDecoratorInstanceData* GetDecoratorInstance(const FNodeInstance& NodeInstance) const
+		const FDecoratorInstanceData* GetDecoratorInstance(const FNodeInstance& NodeInstance) const noexcept
 		{
 			return reinterpret_cast<const FDecoratorInstanceData*>(reinterpret_cast<const uint8*>(&NodeInstance) + GetNodeInstanceOffset());
 		}
@@ -86,13 +105,15 @@ namespace UE::AnimNext
 		friend struct FNodeTemplate;
 		friend struct FNodeTemplateBuilder;
 
-		FDecoratorTemplate(FDecoratorUID InUID, FDecoratorRegistryHandle InRegistryHandle, EDecoratorMode InMode, uint32 InAdditiveIndexOrNumAdditive)
+		FDecoratorTemplate(FDecoratorUID InUID, FDecoratorRegistryHandle InRegistryHandle, EDecoratorMode InMode, uint32 InAdditiveIndexOrNumAdditive) noexcept
 			: UID(InUID.GetUID())
 			, RegistryHandle(InRegistryHandle)
 			, Mode(static_cast<uint8>(InMode))
 			, AdditiveIndexOrNumAdditive(InAdditiveIndexOrNumAdditive)
 			, NodeSharedOffset(0)
+			, NodeSharedLatentPropertyHandlesOffset(0)
 			, NodeInstanceOffset(0)				// For instance data, 0 is an invalid offset since the data follows an instance of FNodeInstance
+			, Padding0(0)
 		{}
 
 		FDecoratorUIDRaw			UID;			// decorator globally unique identifier
@@ -106,8 +127,13 @@ namespace UE::AnimNext
 		uint8	AdditiveIndexOrNumAdditive;
 
 		// Offsets into the shared read-only and instance data portions of a node
-		uint16	NodeSharedOffset;				// relative to root of node description data (max 64 KB per node) (not serialized, @see FNodeTemplate::Finalize)
-		uint16	NodeInstanceOffset;				// relative to root of node instance data (max 64 KB per node) (not serialized, @see FNodeTemplate::Finalize)
+		// These are not serialized, @see FNodeTemplate::Finalize
+		// These are relative to the root of the node description and instance data, respectively (max 64 KB per node each)
+		uint16	NodeSharedOffset;						// Start of shared data for this decorator
+		uint16	NodeSharedLatentPropertyHandlesOffset;	// Start of shared latent property handles for this decorator
+		uint16	NodeInstanceOffset;						// Start of instance data for this decorator
+
+		uint16	Padding0;						// Unused for now
 
 		// TODO: We could cache which parent/base decorator handles which interface
 		// This would avoid the need to iterate on every decorator to look up a 'Super'. Perhaps only common interfaces could be cached.
diff --git a/Engine/Plugins/Experimental/AnimNext/Source/AnimNext/Internal/DecoratorBase/DecoratorWriter.h b/Engine/Plugins/Experimental/AnimNext/Source/AnimNext/Internal/DecoratorBase/DecoratorWriter.h
--- a/Engine/Plugins/Experimental/AnimNext/Source/AnimNext/Internal/DecoratorBase/DecoratorWriter.h
+++ b/Engine/Plugins/Experimental/AnimNext/Source/AnimNext/Internal/DecoratorBase/DecoratorWriter.h
@@ -7,6 +7,7 @@
 #if WITH_EDITOR
 #include "Serialization/MemoryWriter.h"
 
+#include "DecoratorBase/LatentPropertyHandle.h"
 #include "DecoratorBase/NodeHandle.h"
 #include "DecoratorBase/NodeTemplateRegistryHandle.h"
 
@@ -47,7 +48,11 @@ namespace UE::AnimNext
 
 		// Writes out the provided node using the decorator properties
 		// Nodes must be written in the same order they were registered in
-		void WriteNode(const FNodeHandle NodeHandle, const TFunction<FString(uint32 DecoratorIndex, const FString& PropertyName)>& GetDecoratorProperty);
+		void WriteNode(
+			const FNodeHandle NodeHandle,
+			const TFunction<FString(uint32 DecoratorIndex, const FString& PropertyName)>& GetDecoratorProperty,
+			const TFunction<bool(uint32 DecoratorIndex, const FString& PropertyName)>& IsDecoratorPropertyLatent
+			);
 
 		// Returns the error state
 		[[nodiscard]] EErrorState GetErrorState() const;
@@ -81,6 +86,7 @@ namespace UE::AnimNext
 		// To track node writing
 		TArray<UObject*> TrackedObjectsForGC;
 		uint32 NumNodesWritten;
+		FLatentPropertyHandle CurrentLatentPropertyHandle;
 		bool bIsNodeWriting;
 
 		EErrorState ErrorState;
diff --git a/Engine/Plugins/Experimental/AnimNext/Source/AnimNext/Internal/Graph/AnimGraph/AnimNode_AnimNextGraph.h b/Engine/Plugins/Experimental/AnimNext/Source/AnimNext/Internal/Graph/AnimGraph/AnimNode_AnimNextGraph.h
--- a/Engine/Plugins/Experimental/AnimNext/Source/AnimNext/Internal/Graph/AnimGraph/AnimNode_AnimNextGraph.h
+++ b/Engine/Plugins/Experimental/AnimNext/Source/AnimNext/Internal/Graph/AnimGraph/AnimNode_AnimNextGraph.h
@@ -17,7 +17,6 @@ struct ANIMNEXT_API FAnimNode_AnimNextGraph : public FAnimNode_CustomProperty
 	GENERATED_BODY()
 
 	FAnimNode_AnimNextGraph();
-	~FAnimNode_AnimNextGraph();
 
 	// FAnimNode_Base interface
 	virtual void OnInitializeAnimInstance(const FAnimInstanceProxy* InProxy, const UAnimInstance* InAnimInstance) override;
@@ -44,8 +43,9 @@ private:
 	UPROPERTY(EditAnywhere, Category = Settings, meta = (PinHiddenByDefault))
 	TObjectPtr<UAnimNextGraph> AnimNextGraph;
 
-	// Shared pointer to our graph instance, we own it
-	UE::AnimNext::FDecoratorPtr GraphInstancePtr;
+	// Our graph instance, we own it
+	UPROPERTY()
+	FAnimNextGraphInstance GraphInstance;
 
 	/*
 	 * Max LOD that this node is allowed to run
diff --git a/Engine/Plugins/Experimental/AnimNext/Source/AnimNext/Internal/Graph/AnimNextExecuteContext.h b/Engine/Plugins/Experimental/AnimNext/Source/AnimNext/Internal/Graph/AnimNextExecuteContext.h
--- a/Engine/Plugins/Experimental/AnimNext/Source/AnimNext/Internal/Graph/AnimNextExecuteContext.h
+++ b/Engine/Plugins/Experimental/AnimNext/Source/AnimNext/Internal/Graph/AnimNextExecuteContext.h
@@ -42,6 +42,7 @@ struct FAnimNextExecuteContext : public FRigVMExecuteContext
 	FAnimNextExecuteContext()
 		: FRigVMExecuteContext()
 		, Context(nullptr)
+		, SimulationSteps(EAnimNextGraphSimulationSteps::None)
 	{
 	}
 
@@ -65,6 +66,7 @@ struct FAnimNextExecuteContext : public FRigVMExecuteContext
 		Context = OtherContext->Context;
 		SharedDataBuffer = OtherContext->SharedDataBuffer;
 		GraphInstancePtr = OtherContext->GraphInstancePtr;
+		SimulationSteps = OtherContext->SimulationSteps;
 	}
 
 private:
@@ -84,6 +86,15 @@ private:
 		SimulationSteps = InSimulationSteps;
 	}
 
+	// Call this to reset the context to its original state to detect stale usage, can't call it Reset due to virtual in base with that name
+	void DebugReset()
+	{
+		Context = nullptr;
+		SharedDataBuffer = TArrayView<const uint8>();
+		GraphInstancePtr.Reset();
+		SimulationSteps = EAnimNextGraphSimulationSteps::None;
+	}
+
 	const UE::AnimNext::FContext* Context;
 
 	TArrayView<const uint8> SharedDataBuffer;
diff --git a/Engine/Plugins/Experimental/AnimNext/Source/AnimNext/Internal/Graph/AnimNextGraph.h b/Engine/Plugins/Experimental/AnimNext/Source/AnimNext/Internal/Graph/AnimNextGraph.h
--- a/Engine/Plugins/Experimental/AnimNext/Source/AnimNext/Internal/Graph/AnimNextGraph.h
+++ b/Engine/Plugins/Experimental/AnimNext/Source/AnimNext/Internal/Graph/AnimNextGraph.h
@@ -37,6 +37,42 @@ namespace UE::AnimNext::Graph
 	extern ANIMNEXT_API const FName ResultName;
 }
 
+// Represents an instance of an AnimNext graph
+// This struct uses UE reflection because we wish for the GC to keep the graph
+// alive while we own a reference to it. It is not intended to be serialized on disk with a live instance.
+USTRUCT()
+struct ANIMNEXT_API FAnimNextGraphInstance
+{
+	GENERATED_BODY()
+
+	// Creates an empty graph instance that doesn't reference anything
+	FAnimNextGraphInstance() = default;
+
+	// If the graph instance is allocated, we release it during destruction
+	~FAnimNextGraphInstance();
+
+	// Releases the graph instance and frees all corresponding memory
+	void Release();
+
+	// Returns true if we have a live graph instance, false otherwise
+	bool IsValid() const;
+
+private:
+	// Hard reference to the graph used to create this instance to ensure we can release it safely
+	UPROPERTY()
+	TObjectPtr<const UAnimNextGraph> Graph;
+
+	// Hard reference to the graph instance data, we own it
+	UE::AnimNext::FDecoratorPtr GraphInstancePtr;
+
+	// Extended execute context instance for this graph instance, we own it
+	UPROPERTY()
+	FRigVMExtendedExecuteContext ExtendedExecuteContext;
+
+	// The graph is the one that allocates instances
+	friend UAnimNextGraph;
+};
+
 // A user-created graph of logic used to supply data
 UCLASS(BlueprintType)
 class ANIMNEXT_API UAnimNextGraph : public UObject
@@ -45,19 +81,17 @@ class ANIMNEXT_API UAnimNextGraph : public UObject
 
 public:
 	// UObject interface
+	virtual void PostLoad() override;
 	virtual void PostRename(UObject* OldOuter, const FName OldName) override;
 	virtual void GetPreloadDependencies(TArray<UObject*>& OutDeps) override;
 	virtual void Serialize(FArchive& Ar) override;
 	static void AddReferencedObjects(UObject* InThis, FReferenceCollector& Collector);
 
 	// Allocates an instance of the graph, retain the handle and use it with the Run() function to evaluate it
-	UE::AnimNext::FDecoratorPtr AllocateInstance() const;
-
-	// Releases an instance of the graph and clears the provided handle
-	void ReleaseInstance(UE::AnimNext::FDecoratorPtr& GraphInstancePtr) const;
+	void AllocateInstance(FAnimNextGraphInstance& Instance) const;
 
 	// Run the specified simulation steps on the provided graph with the given context
-	void Run(const UE::AnimNext::FContext& Context, UE::AnimNext::FWeakDecoratorPtr GraphInstancePtr, EAnimNextGraphSimulationSteps SimulationSteps) const;
+	void Run(const UE::AnimNext::FContext& Context, FAnimNextGraphInstance& GraphInstance, EAnimNextGraphSimulationSteps SimulationSteps) const;
 
 protected:
 	// Support rig VM execution
@@ -70,12 +104,9 @@ protected:
 	friend class UAnimNextGraph_EditorData;
 	friend struct UE::AnimNext::UncookedOnly::FUtils;
 	friend class UE::AnimNext::Editor::FGraphEditor;
-	friend class UAnimNextGraph;
+	friend struct FAnimNextGraphInstance;
 	friend class UAnimGraphNode_AnimNextGraph;
 
-	UPROPERTY()
-	TObjectPtr<URigVM> RigVM;
-
 	// This is a handle to the root decorator in our graph
 	UPROPERTY()
 	FAnimNextEntryPointHandle RootDecoratorHandle;
@@ -92,8 +123,15 @@ protected:
 	// It means that object references in the graph shared data are not visited at runtime by the GC (they are immutable)
 	TArray<TObjectPtr<UObject>> TrackedObjectsForGC;
 
-	UPROPERTY(transient)
-	mutable FRigVMExtendedExecuteContext ExtendedExecuteContext;
+	// The RigVM object holds the bytecode, literals, etc used by the RigVM internals, there is a single instance along with the UAnimNextGraph (1:1 mapping)
+	UPROPERTY()
+	TObjectPtr<URigVM> RigVM;
+
+	// The ExtendedExecuteContext object holds instance/work data used by the RigVM internals. It is populated during compilation same as the URigVM object above.
+	// We also have a 1:1 mapping with the UAnimNextGraph but each instance of the anim graph also needs its own exetended execute context object. This one is used
+	// as a reference we copy from.
+	UPROPERTY()
+	FRigVMExtendedExecuteContext ExtendedExecuteContext;
 
 	UPROPERTY()
 	FRigVMRuntimeSettings VMRuntimeSettings;
diff --git a/Engine/Plugins/Experimental/AnimNext/Source/AnimNext/Internal/Graph/RigDecorator_AnimNextCppDecorator.h b/Engine/Plugins/Experimental/AnimNext/Source/AnimNext/Internal/Graph/RigDecorator_AnimNextCppDecorator.h
--- a/Engine/Plugins/Experimental/AnimNext/Source/AnimNext/Internal/Graph/RigDecorator_AnimNextCppDecorator.h
+++ b/Engine/Plugins/Experimental/AnimNext/Source/AnimNext/Internal/Graph/RigDecorator_AnimNextCppDecorator.h
@@ -18,7 +18,7 @@ struct ANIMNEXT_API FRigDecorator_AnimNextCppDecorator : public FRigVMDecorator
 	GENERATED_BODY()
 
 	// The struct the decorator exposes with its shared data. Each one of its properties will be added as a pin.
-	UPROPERTY()
+	UPROPERTY(meta = (Hidden))
 	UScriptStruct* DecoratorSharedDataStruct = nullptr;
 
 #if WITH_EDITOR
diff --git a/Engine/Plugins/Experimental/AnimNext/Source/AnimNext/Internal/Graph/RigUnit_AnimNextGraphEvaluator.h b/Engine/Plugins/Experimental/AnimNext/Source/AnimNext/Internal/Graph/RigUnit_AnimNextGraphEvaluator.h
new file mode 100644
--- /dev/null
+++ b/Engine/Plugins/Experimental/AnimNext/Source/AnimNext/Internal/Graph/RigUnit_AnimNextGraphEvaluator.h
@@ -0,0 +1,24 @@
+// Copyright Epic Games, Inc. All Rights Reserved.
+
+#pragma once
+
+#include "AnimNextExecuteContext.h"
+
+#include "RigUnit_AnimNextGraphEvaluator.generated.h"
+
+/**
+ * Animation graph evaluator
+ * This node is only used at runtime.
+ * It performs the animation graph update and evaluation through the data provided in the execution context.
+ * It also holds all latent/lazy pins that the graph references in the editor.
+ */
+USTRUCT(meta=(DisplayName="Animation Runtime Output", Category="Events", NodeColor="1, 0, 0"))
+struct ANIMNEXT_API FRigUnit_AnimNextGraphEvaluator : public FRigUnit_AnimNextBase
+{
+	GENERATED_BODY()
+
+	UPROPERTY(meta = (Input, Output))
+	FAnimNextExecuteContext ExecuteContext;
+
+	static void StaticExecute(FRigVMExtendedExecuteContext& RigVMExecuteContext, FRigVMMemoryHandleArray RigVMMemoryHandles, FRigVMPredicateBranchArray RigVMBranches);
+};
diff --git a/Engine/Plugins/Experimental/AnimNext/Source/AnimNext/Internal/Graph/RigUnit_AnimNextGraphRoot.h b/Engine/Plugins/Experimental/AnimNext/Source/AnimNext/Internal/Graph/RigUnit_AnimNextGraphRoot.h
--- a/Engine/Plugins/Experimental/AnimNext/Source/AnimNext/Internal/Graph/RigUnit_AnimNextGraphRoot.h
+++ b/Engine/Plugins/Experimental/AnimNext/Source/AnimNext/Internal/Graph/RigUnit_AnimNextGraphRoot.h
@@ -9,6 +9,10 @@
 
 /**
  * Animation graph output
+ * This is a synthetic node that represents the entry point for an animation graph for RigVM.
+ * The graph editor will see this as the graph output in which to hook up the first animation node
+ * to evaluate.
+ * This node isn't used at runtime.
  */
 USTRUCT(meta=(DisplayName="Animation Output", Category="Events", NodeColor="1, 0, 0", Keywords="Root,Output"))
 struct ANIMNEXT_API FRigUnit_AnimNextGraphRoot : public FRigUnit_AnimNextBase
@@ -16,7 +20,7 @@ struct ANIMNEXT_API FRigUnit_AnimNextGraphRoot : public FRigUnit_AnimNextBase
 	GENERATED_BODY()
 
 	RIGVM_METHOD()
-	void Execute();
+	void DummyExecute();
 
 	virtual FName GetEventName() const override { return EventName; }
 	virtual bool CanOnlyExistOnce() const override { return true; }
diff --git a/Engine/Plugins/Experimental/AnimNext/Source/AnimNext/Internal/Graph/RigUnit_AnimNextShimRoot.h b/Engine/Plugins/Experimental/AnimNext/Source/AnimNext/Internal/Graph/RigUnit_AnimNextShimRoot.h
new file mode 100644
--- /dev/null
+++ b/Engine/Plugins/Experimental/AnimNext/Source/AnimNext/Internal/Graph/RigUnit_AnimNextShimRoot.h
@@ -0,0 +1,34 @@
+﻿// Copyright Epic Games, Inc. All Rights Reserved.
+
+#pragma once
+
+#include "AnimNextExecuteContext.h"
+
+#include "RigUnit_AnimNextShimRoot.generated.h"
+
+/**
+ * Animation graph output
+ * This is a synthetic node that represents the entry point for an animation graph for RigVM.
+ * The graph editor will not see this node. It is added during compilation as a shim to start the evaluation process.
+ * This node is only used at runtime.
+ */
+USTRUCT(meta=(DisplayName="Animation Output Shim", Category="Events", NodeColor="1, 0, 0", Keywords="Root,Output"))
+struct ANIMNEXT_API FRigUnit_AnimNextShimRoot : public FRigUnit_AnimNextBase
+{
+	GENERATED_BODY()
+
+	RIGVM_METHOD()
+	void Execute();
+
+	virtual FName GetEventName() const override { return EventName; }
+	virtual bool CanOnlyExistOnce() const override { return true; }
+
+	// In order for this node to be considered an executable RigUnit, it needs a pin to derive from FRigVMExecuteContext
+	// Although the anim graph flows from left to right with the graph output being on the right hand side with a single
+	// input as the final output, RigVM execution flows from outputs to inputs and as such an entry point has a single output.
+	UPROPERTY(meta = (Output))
+	FAnimNextExecuteContext ExecuteContext;
+
+	// This unit is our graph entry point, it needs an event so we can call it
+	static FName EventName;
+};
diff --git a/Engine/Plugins/Experimental/AnimNext/Source/AnimNext/Private/DecoratorBase/Decorator.cpp b/Engine/Plugins/Experimental/AnimNext/Source/AnimNext/Private/DecoratorBase/Decorator.cpp
--- a/Engine/Plugins/Experimental/AnimNext/Source/AnimNext/Private/DecoratorBase/Decorator.cpp
+++ b/Engine/Plugins/Experimental/AnimNext/Source/AnimNext/Private/DecoratorBase/Decorator.cpp
@@ -2,13 +2,8 @@
 
 #include "DecoratorBase/Decorator.h"
 
-#include "DecoratorBase/DecoratorReader.h"
 #include "DecoratorBase/DecoratorRegistry.h"
 
-#if WITH_EDITOR
-#include "DecoratorBase/DecoratorWriter.h"
-#endif
-
 namespace UE::AnimNext
 {
 	void FDecorator::SerializeDecoratorSharedData(FArchive& Ar, FAnimNextDecoratorSharedData& SharedData) const
@@ -18,7 +13,7 @@ namespace UE::AnimNext
 	}
 
 #if WITH_EDITOR
-	void FDecorator::SaveDecoratorSharedData(FDecoratorWriter& Writer, const TFunction<FString(const FString& PropertyName)>& GetDecoratorProperty, FAnimNextDecoratorSharedData& OutSharedData) const
+	void FDecorator::SaveDecoratorSharedData(const TFunction<FString(const FString& PropertyName)>& GetDecoratorProperty, FAnimNextDecoratorSharedData& OutSharedData) const
 	{
 		const UScriptStruct* SharedDataStruct = GetDecoratorSharedDataStruct();
 
@@ -31,6 +26,8 @@ namespace UE::AnimNext
 		// We convert every property from its string representation into its binary form
 		for (const FProperty* Property = SharedDataStruct->PropertyLink; Property != nullptr; Property = Property->PropertyLinkNext)
 		{
+			// No need to skip editor only properties since serialization will take care of that afterwards
+
 			const FString PropertyValue = GetDecoratorProperty(Property->GetName());
 			if (PropertyValue.Len() != 0)
 			{
@@ -64,12 +61,72 @@ namespace UE::AnimNext
 			}
 		}
 	}
+
+	TArray<FLatentPropertyHandle> FDecorator::GetLatentPropertyHandles(bool bFilterEditorOnly, const TFunction<bool(const FString& PropertyName)>& IsDecoratorPropertyLatent, FLatentPropertyHandle& CurrentLatentPropertyHandle) const
+	{
+		const UStruct* BaseStruct = GetDecoratorSharedDataStruct();
+
+		// The property linked list on UScriptStruct iterates over the properties starting in the derived type
+		// but with latent properties, the base type should be the first to be visited.
+		// Gather our struct hierarchy from most derived to base
+		TArray<const UStruct*> StructHierarchy;
+
+		do
+		{
+			StructHierarchy.Add(BaseStruct);
+			BaseStruct = BaseStruct->GetSuperStruct();
+		}
+		while (BaseStruct != nullptr);
+
+		TArray<FLatentPropertyHandle> LatentPropertyHandles;
+
+		// Gather our latent properties from base to most derived
+		for (auto It = StructHierarchy.rbegin(); It != StructHierarchy.rend(); ++It)
+		{
+			for (const FField* Field = (*It)->ChildProperties; Field; Field = Field->Next)
+			{
+				const FProperty* Property = CastField<FProperty>(Field);
+
+				if (bFilterEditorOnly && Property->IsEditorOnlyProperty())
+				{
+					continue;	// Skip editor only properties if we don't need them
+				}
+
+				// By default, properties are latent
+				// However, there are exceptions:
+				//     - Properties marked as hidden are not visible in the editor and cannot be hooked up manually
+				//     - Properties marked as inline are only visible in the details panel and cannot be hooked up to another node
+				//     - Properties of decorator handle type are never lazy since they just encode graph connectivity
+				const bool bIsPotentiallyLatent = !Property->HasMetaData(TEXT("Hidden")) &&
+					!Property->HasMetaData(TEXT("Inline")) &&
+					Property->GetCPPType() != TEXT("FAnimNextDecoratorHandle");
+
+				if (!bIsPotentiallyLatent)
+				{
+					continue;	// Skip non-latent properties
+				}
+
+				FLatentPropertyHandle LatentHandle;
+
+				if (IsDecoratorPropertyLatent(Property->GetName()))
+				{
+					// This property is marked latent and it has a non-inline value, grab our next latent handle and increment it
+					LatentHandle = CurrentLatentPropertyHandle;
+					CurrentLatentPropertyHandle = CurrentLatentPropertyHandle.GetNextHandle();
+				}
+
+				LatentPropertyHandles.Add(LatentHandle);
+			}
+		}
+
+		return LatentPropertyHandles;
+	}
 #endif
 
-	FDecoratorStaticInitHook::FDecoratorStaticInitHook(DecoratorConstructorFunc DecoratorConstructor_)
-		: DecoratorConstructor(DecoratorConstructor_)
+	FDecoratorStaticInitHook::FDecoratorStaticInitHook(DecoratorConstructorFunc InDecoratorConstructor)
+		: DecoratorConstructor(InDecoratorConstructor)
 	{
-		FDecoratorRegistry::StaticRegister(DecoratorConstructor_);
+		FDecoratorRegistry::StaticRegister(InDecoratorConstructor);
 	}
 
 	FDecoratorStaticInitHook::~FDecoratorStaticInitHook()
diff --git a/Engine/Plugins/Experimental/AnimNext/Source/AnimNext/Private/DecoratorBase/DecoratorReader.cpp b/Engine/Plugins/Experimental/AnimNext/Source/AnimNext/Private/DecoratorBase/DecoratorReader.cpp
--- a/Engine/Plugins/Experimental/AnimNext/Source/AnimNext/Private/DecoratorBase/DecoratorReader.cpp
+++ b/Engine/Plugins/Experimental/AnimNext/Source/AnimNext/Private/DecoratorBase/DecoratorReader.cpp
@@ -110,12 +110,31 @@ namespace UE::AnimNext
 		GraphSharedData.Empty(0);
 		GraphSharedData.AddZeroed(SharedDataSize);
 
-		// Serialize our graph
+		// Serialize our graph shared data
 		for (uint32 NodeIndex = 0; NodeIndex < NumNodes; ++NodeIndex)
 		{
+			// Serialize our node shared data
 			const uint32 SharedDataOffset = NodeHandles[NodeIndex].GetSharedOffset();
 			FNodeDescription* NodeDesc = reinterpret_cast<FNodeDescription*>(&GraphSharedData[SharedDataOffset]);
 			NodeDesc->Serialize(*this);
+
+			// Serialize our latent property handles
+			const FNodeTemplate* NodeTemplate = NodeTemplateRegistry.Find(NodeDesc->GetTemplateHandle());
+
+			const uint32 NumDecorators = NodeTemplate->GetNumDecorators();
+			const FDecoratorTemplate* DecoratorTemplates = NodeTemplate->GetDecorators();
+
+			for (uint32 DecoratorIndex = 0; DecoratorIndex < NumDecorators; ++DecoratorIndex)
+			{
+				int32 NumLatentHandles = 0;
+				*this << NumLatentHandles;
+
+				FLatentPropertyHandle* LatentHandles = DecoratorTemplates[DecoratorIndex].GetDecoratorLatentPropertyHandles(*NodeDesc);
+				for (int32 LatentHandleIndex = 0; LatentHandleIndex < NumLatentHandles; ++LatentHandleIndex)
+				{
+					*this << LatentHandles[LatentHandleIndex];
+				}
+			}
 		}
 
 		return EErrorState::None;
diff --git a/Engine/Plugins/Experimental/AnimNext/Source/AnimNext/Private/DecoratorBase/DecoratorWriter.cpp b/Engine/Plugins/Experimental/AnimNext/Source/AnimNext/Private/DecoratorBase/DecoratorWriter.cpp
--- a/Engine/Plugins/Experimental/AnimNext/Source/AnimNext/Private/DecoratorBase/DecoratorWriter.cpp
+++ b/Engine/Plugins/Experimental/AnimNext/Source/AnimNext/Private/DecoratorBase/DecoratorWriter.cpp
@@ -75,6 +75,7 @@ namespace UE::AnimNext
 
 		bIsNodeWriting = true;
 		TrackedObjectsForGC.Reset();
+		CurrentLatentPropertyHandle = FLatentPropertyHandle::GetFirstHandle();
 
 		// Serialize the node templates
 		TArray<FNodeTemplateRegistryHandle> NodeTemplateHandles;
@@ -132,7 +133,10 @@ namespace UE::AnimNext
 		}
 	}
 
-	void FDecoratorWriter::WriteNode(const FNodeHandle NodeHandle, const TFunction<FString (uint32 DecoratorIndex, const FString& PropertyName)>& GetDecoratorProperty)
+	void FDecoratorWriter::WriteNode(
+		const FNodeHandle NodeHandle,
+		const TFunction<FString (uint32 DecoratorIndex, const FString& PropertyName)>& GetDecoratorProperty,
+		const TFunction<bool(uint32 DecoratorIndex, const FString& PropertyName)>& IsDecoratorPropertyLatent)
 	{
 		ensure(bIsNodeWriting);
 
@@ -180,12 +184,35 @@ namespace UE::AnimNext
 				return GetDecoratorProperty(DecoratorIndex, PropertyName);
 			};
 
-			Decorator->SaveDecoratorSharedData(*this, GetDecoratorPropertyAt, *SharedData);
+			Decorator->SaveDecoratorSharedData(GetDecoratorPropertyAt, *SharedData);
 		}
 
-		// Append it to our archive
+		// Append our node and decorator shared data to our archive
 		NodeDesc->Serialize(*this);
 
+		// Append our decorator latent property handles to our archive
+		for (uint32 DecoratorIndex = 0; DecoratorIndex < NumDecorators; ++DecoratorIndex)
+		{
+			const FDecoratorRegistryHandle DecoratorHandle = DecoratorTemplates[DecoratorIndex].GetRegistryHandle();
+			const FDecorator* Decorator = DecoratorRegistry.Find(DecoratorHandle);
+
+			// Curry our lambda with the decorator index
+			const auto IsDecoratorPropertyLatentAt = [&IsDecoratorPropertyLatent, DecoratorIndex](const FString& PropertyName)
+			{
+				return IsDecoratorPropertyLatent(DecoratorIndex, PropertyName);
+			};
+
+			const TArray<FLatentPropertyHandle> LatentHandles = Decorator->GetLatentPropertyHandles(IsFilterEditorOnly(), IsDecoratorPropertyLatentAt, CurrentLatentPropertyHandle);
+
+			int32 NumLatentHandles = LatentHandles.Num();
+			*this << NumLatentHandles;
+
+			for (FLatentPropertyHandle Handle : LatentHandles)
+			{
+				*this << Handle;
+			}
+		}
+
 		NumNodesWritten++;
 	}
 
diff --git a/Engine/Plugins/Experimental/AnimNext/Source/AnimNext/Private/DecoratorBase/ExecutionContext.cpp b/Engine/Plugins/Experimental/AnimNext/Source/AnimNext/Private/DecoratorBase/ExecutionContext.cpp
--- a/Engine/Plugins/Experimental/AnimNext/Source/AnimNext/Private/DecoratorBase/ExecutionContext.cpp
+++ b/Engine/Plugins/Experimental/AnimNext/Source/AnimNext/Private/DecoratorBase/ExecutionContext.cpp
@@ -9,6 +9,7 @@
 #include "DecoratorBase/NodeInstance.h"
 #include "DecoratorBase/NodeTemplate.h"
 #include "DecoratorBase/NodeTemplateRegistry.h"
+#include "RigVMCore/RigVMExecuteContext.h"
 
 namespace UE::AnimNext
 {
@@ -17,8 +18,20 @@ namespace UE::AnimNext
 		static thread_local UE::AnimNext::FExecutionContext* GThreadLocalExecutionContext = nullptr;
 	}
 
-	FExecutionContext::FExecutionContext(TArrayView<const uint8> GraphSharedData_)
-		: GraphSharedData(GraphSharedData_)
+	FExecutionContext::FExecutionContext(TArrayView<const uint8> InGraphSharedData)
+		: GraphSharedData(InGraphSharedData)
+		, RigVMLatentMemoryHandles()
+		, RigVMExecuteContext(nullptr)
+	{
+		// There can be only one execution context alive per thread
+		ensure(Private::GThreadLocalExecutionContext == nullptr);
+		Private::GThreadLocalExecutionContext = this;
+	}
+
+	FExecutionContext::FExecutionContext(TArrayView<const uint8> InGraphSharedData, FRigVMExtendedExecuteContext& InRigVMExecuteContext, FRigVMMemoryHandleArray InRigVMLatentMemoryHandles)
+		: GraphSharedData(InGraphSharedData)
+		, RigVMLatentMemoryHandles(InRigVMLatentMemoryHandles)
+		, RigVMExecuteContext(&InRigVMExecuteContext)
 	{
 		// There can be only one execution context alive per thread
 		ensure(Private::GThreadLocalExecutionContext == nullptr);
@@ -95,7 +108,8 @@ namespace UE::AnimNext
 			const FAnimNextDecoratorSharedData* SharedData = DecoratorDesc->GetDecoratorDescription(NodeDesc);
 			FDecoratorInstanceData* InstanceData = DecoratorDesc->GetDecoratorInstance(*NodeInstance);
 
-			Decorator->ConstructDecoratorInstance(*this, DecoratorPtr, *SharedData, *InstanceData);
+			const FDecoratorBinding Binding(nullptr, DecoratorDesc, &NodeDesc, DecoratorPtr);
+			Decorator->ConstructDecoratorInstance(*this, Binding);
 		}
 
 		return FDecoratorPtr(NodeInstance, ChildDecoratorIndex);
@@ -137,10 +151,8 @@ namespace UE::AnimNext
 
 				FWeakDecoratorPtr DecoratorPtr(NodeInstance, DecoratorIndex);
 
-				const FAnimNextDecoratorSharedData* SharedData = DecoratorDesc->GetDecoratorDescription(NodeDesc);
-				FDecoratorInstanceData* InstanceData = DecoratorDesc->GetDecoratorInstance(*NodeInstance);
-
-				Decorator->DestructDecoratorInstance(*this, DecoratorPtr, *SharedData, *InstanceData);
+				const FDecoratorBinding Binding(nullptr, DecoratorDesc, &NodeDesc, DecoratorPtr);
+				Decorator->DestructDecoratorInstance(*this, Binding);
 			}
 		}
 
@@ -191,7 +203,7 @@ namespace UE::AnimNext
 				const uint32 DecoratorIndex = DecoratorDesc - DecoratorDescs;
 				FWeakDecoratorPtr InterfaceDecoratorPtr(NodeInstance, DecoratorIndex);
 
-				InterfaceBinding = FDecoratorBinding(Interface, SharedData, InstanceData, InterfaceDecoratorPtr);
+				InterfaceBinding = FDecoratorBinding(Interface, DecoratorDesc, &NodeDesc, InterfaceDecoratorPtr);
 				return true;
 			}
 		}
@@ -248,7 +260,7 @@ namespace UE::AnimNext
 				const uint32 DecoratorIndex = DecoratorDesc - DecoratorDescs;
 				FWeakDecoratorPtr SuperPtr(NodeInstance, DecoratorIndex);
 
-				SuperBinding = FDecoratorBinding(Interface, SharedData, InstanceData, SuperPtr);
+				SuperBinding = FDecoratorBinding(Interface, DecoratorDesc, &NodeDesc, SuperPtr);
 				return true;
 			}
 		}
@@ -257,6 +269,25 @@ namespace UE::AnimNext
 		return false;
 	}
 
+	const void* FExecutionContext::EvaluateLatentPinImpl(FLatentPropertyHandle LatentPropertyHandle) const
+	{
+		check(LatentPropertyHandle.IsValid());
+		check(RigVMLatentMemoryHandles.IsValidIndex(LatentPropertyHandle.GetLatentPropertyIndex()));
+
+		FRigVMMemoryHandle& MemoryHandle = RigVMLatentMemoryHandles[LatentPropertyHandle.GetLatentPropertyIndex()];
+
+		// This should be an assert. If this triggers, it means that we have a bug in how lazy memory handles
+		// are assigned during compilation. We keep it as an ensure because in this case, we can recover
+		// as even if the memory handle isn't lazy, it remains valid and we can use it. It won't have the
+		// value we expect but it'll work. The ensure will signal that we need to fix the bug.
+		if (ensure(MemoryHandle.IsLazy()))
+		{
+			MemoryHandle.ComputeLazyValueIfNecessary(*RigVMExecuteContext, RigVMExecuteContext->GetSlice().GetIndex());
+		}
+
+		return MemoryHandle.GetData();
+	}
+
 	const FNodeDescription& FExecutionContext::GetNodeDescription(FNodeHandle NodeHandle) const
 	{
 		check(NodeHandle.IsValid());
diff --git a/Engine/Plugins/Experimental/AnimNext/Source/AnimNext/Private/DecoratorBase/NodeTemplate.cpp b/Engine/Plugins/Experimental/AnimNext/Source/AnimNext/Private/DecoratorBase/NodeTemplate.cpp
--- a/Engine/Plugins/Experimental/AnimNext/Source/AnimNext/Private/DecoratorBase/NodeTemplate.cpp
+++ b/Engine/Plugins/Experimental/AnimNext/Source/AnimNext/Private/DecoratorBase/NodeTemplate.cpp
@@ -46,6 +46,7 @@ namespace UE::AnimNext
 			const FDecoratorUID DecoratorUID = DecoratorTemplate.GetUID();
 
 			uint32 DecoratorSharedDataOffset = 0;
+			uint32 DecoratorSharedLatentPropertyHandlesOffset = 0;
 			uint32 DecoratorInstanceDataOffset = 0;	// For instance data, 0 is an invalid offset since the data follows an instance of FNodeInstance
 
 			// Skip decorators that we can't find
@@ -65,13 +66,22 @@ namespace UE::AnimNext
 				// Include our decorator
 				SharedDataOffset += MemoryLayout.SharedDataSize;
 				InstanceDataOffset += MemoryLayout.InstanceDataSize;
+
+				// Save our latent pins offset (if we have any)
+				const uint32 NumLatentProperties = Decorator->GetNumLatentDecoratorProperties();
+				DecoratorSharedLatentPropertyHandlesOffset = NumLatentProperties != 0 ? Align(SharedDataOffset, alignof(FLatentPropertyHandle)) : 0;
+
+				// Include our latent pins
+				SharedDataOffset += NumLatentProperties * sizeof(FLatentPropertyHandle);
 			}
 
 			check(DecoratorSharedDataOffset <= MAX_uint16);
+			check(DecoratorSharedLatentPropertyHandlesOffset <= MAX_uint16);
 			check(DecoratorInstanceDataOffset <= MAX_uint16);
 
 			// Update our decorator offsets
 			DecoratorTemplate.NodeSharedOffset = static_cast<uint16>(DecoratorSharedDataOffset);
+			DecoratorTemplate.NodeSharedLatentPropertyHandlesOffset = static_cast<uint16>(DecoratorSharedLatentPropertyHandlesOffset);
 			DecoratorTemplate.NodeInstanceOffset = static_cast<uint16>(DecoratorInstanceDataOffset);
 		}
 
diff --git a/Engine/Plugins/Experimental/AnimNext/Source/AnimNext/Private/DecoratorInterfaces/IUpdate.cpp b/Engine/Plugins/Experimental/AnimNext/Source/AnimNext/Private/DecoratorInterfaces/IUpdate.cpp
--- a/Engine/Plugins/Experimental/AnimNext/Source/AnimNext/Private/DecoratorInterfaces/IUpdate.cpp
+++ b/Engine/Plugins/Experimental/AnimNext/Source/AnimNext/Private/DecoratorInterfaces/IUpdate.cpp
@@ -77,6 +77,46 @@ namespace UE::AnimNext
 				// Queue our node again so that post-update is called afterwards
 				NodesPendingUpdate.Push({ Entry.DecoratorPtr, EUpdateStep::PostUpdate });
 
+				// Performance note
+				// When we process an animation graph for a frame, typically we'll update first before we evaluate
+				// As a result of this, when we query for the update interface here, we will likely hit cold memory
+				// which will cache miss (by touching the graph instance for the first time).
+				// 
+				// The processor will cache miss and continue to process as many instructions as it can before the
+				// out-of-order execution window fills up. This is problematic here because a lot of the subsequent
+				// instructions depend on the node instance and the interface it returns. The processor will be unable
+				// to execute any of the instructions that follow in the current loop iteration. However, it might
+				// be able to get started on the next node entry which is likely to cache miss as well. Should the
+				// processor make it that far and it turns out that we have to push a child onto the stack, all
+				// of the work it tried to do ahead of time will have to be thrown away.
+				// 
+				// There are two things that we can do here to try and help performance: prefetch ahead and bulk
+				// query.
+				// 
+				// If we prefetch, we have to be careful because we do not know what the node will do in its PreUpdate.
+				// If it turns out that it does a lot of work, our prefetch might end up getting thrown out. This is
+				// because prefetched cache lines typically end up being the first evicted unless they are touched first.
+				// It is thus dangerous to use manual prefetching when the memory access pattern isn't fully known.
+				// In practice, it is likely viable as most nodes won't do too much work.
+				// 
+				// A better approach could be to instead bulk query for our interfaces. We could cache in the FUpdateEntry
+				// the decorator bindings for IUpdate and IHierarchy (and re-use the binding for IUpdate for PostUpdate).
+				// Every iteration we could check how many children are queued up on the stack. We could then grab N
+				// entries (2 to 4) and query their interfaces in bulk. The idea is to clump the cache miss instructions
+				// together and to interleave the interface queries. This will queue up as much work as possible in the
+				// out-of-order execution window that will not be thrown away because of a branch. Eventually the first
+				// interface query will complete and execution will resume here to call PreUpdate etc. This will be able
+				// to happen while the processor still waits on the cache misses and finishes the interface query of the
+				// other bulked children. The same effect could be achieved by querying the interfaces after the call
+				// to GetChildren by bulk querying all of them right then. This way, as soon as the execution window
+				// can clear the end of the loop, it can start working on the next entry which will be warm in the L1
+				// cache allowing the CPU to carry ahead before all child interfaces are fully resolved.
+				// 
+				// The above may seem like a stretch and an insignificant over optimization but it could very well be the
+				// key to unlocking large performance gains during traversal. The above optimization would allow us to
+				// perform as much useful work as possible while waiting for memory, hiding its slow latency by fully
+				// leveraging out-of-order CPU execution.
+
 				TDecoratorBinding<IUpdate> UpdateDecorator;
 				if (Context.GetInterface(Entry.DecoratorPtr, UpdateDecorator))
 				{
diff --git a/Engine/Plugins/Experimental/AnimNext/Source/AnimNext/Private/Decorators/BlendTwoWay.cpp b/Engine/Plugins/Experimental/AnimNext/Source/AnimNext/Private/Decorators/BlendTwoWay.cpp
--- a/Engine/Plugins/Experimental/AnimNext/Source/AnimNext/Private/Decorators/BlendTwoWay.cpp
+++ b/Engine/Plugins/Experimental/AnimNext/Source/AnimNext/Private/Decorators/BlendTwoWay.cpp
@@ -31,9 +31,8 @@ namespace UE::AnimNext
 		const FSharedData* SharedData = Binding.GetSharedData<FSharedData>();
 		FInstanceData* InstanceData = Binding.GetInstanceData<FInstanceData>();
 
-		// TODO: hook up blend weight to a pin, it could change every frame
-
-		if (SharedData->BlendWeight < 1.0)
+		const float BlendWeight = SharedData->GetBlendWeight(Context, Binding);
+		if (BlendWeight < 1.0f)
 		{
 			if (!InstanceData->ChildA.IsValid())
 			{
@@ -41,14 +40,14 @@ namespace UE::AnimNext
 				InstanceData->ChildA = Context.AllocateNodeInstance(Binding, SharedData->ChildA);
 			}
 
-			if (SharedData->BlendWeight == 0.0)
+			if (BlendWeight == 0.0f)
 			{
 				// We no longer need this child, release it
 				InstanceData->ChildB.Reset();
 			}
 		}
 
-		if (SharedData->BlendWeight > 0.0)
+		if (BlendWeight > 0.0f)
 		{
 			if (!InstanceData->ChildB.IsValid())
 			{
@@ -56,7 +55,7 @@ namespace UE::AnimNext
 				InstanceData->ChildB = Context.AllocateNodeInstance(Binding, SharedData->ChildB);
 			}
 
-			if (SharedData->BlendWeight == 1.0)
+			if (BlendWeight == 1.0f)
 			{
 				// We no longer need this child, release it
 				InstanceData->ChildA.Reset();
diff --git a/Engine/Plugins/Experimental/AnimNext/Source/AnimNext/Private/Decorators/BlendTwoWay.h b/Engine/Plugins/Experimental/AnimNext/Source/AnimNext/Private/Decorators/BlendTwoWay.h
--- a/Engine/Plugins/Experimental/AnimNext/Source/AnimNext/Private/Decorators/BlendTwoWay.h
+++ b/Engine/Plugins/Experimental/AnimNext/Source/AnimNext/Private/Decorators/BlendTwoWay.h
@@ -17,16 +17,23 @@ struct FAnimNextBlendTwoWayDecoratorSharedData : public FAnimNextDecoratorShared
 	GENERATED_BODY()
 
 	/** First output to be blended (full weight is 0.0). */
-	UPROPERTY(meta = (Input))
+	UPROPERTY()
 	FAnimNextDecoratorHandle ChildA;
 
 	/** Second output to be blended (full weight is 1.0). */
-	UPROPERTY(meta = (Input))
+	UPROPERTY()
 	FAnimNextDecoratorHandle ChildB;
 
 	/** How much to blend our two children: 0.0 is fully child A while 1.0 is fully child B. */
-	UPROPERTY(meta = (Input))
-	double BlendWeight = 0.0;
+	UPROPERTY()
+	float BlendWeight = 0.0f;
+
+	// Latent pin support boilerplate
+	#define DECORATOR_LATENT_PROPERTIES_ENUMERATOR(GeneratorMacro) \
+		GeneratorMacro(BlendWeight) \
+
+	GENERATE_DECORATOR_LATENT_PROPERTIES(FAnimNextBlendTwoWayDecoratorSharedData, DECORATOR_LATENT_PROPERTIES_ENUMERATOR)
+	#undef DECORATOR_LATENT_PROPERTIES_ENUMERATOR
 };
 
 namespace UE::AnimNext
diff --git a/Engine/Plugins/Experimental/AnimNext/Source/AnimNext/Private/Decorators/SequencePlayer.cpp b/Engine/Plugins/Experimental/AnimNext/Source/AnimNext/Private/Decorators/SequencePlayer.cpp
--- a/Engine/Plugins/Experimental/AnimNext/Source/AnimNext/Private/Decorators/SequencePlayer.cpp
+++ b/Engine/Plugins/Experimental/AnimNext/Source/AnimNext/Private/Decorators/SequencePlayer.cpp
@@ -17,12 +17,13 @@ namespace UE::AnimNext
 		DEFINE_ANIM_DECORATOR_IMPLEMENTS_INTERFACE(ITimeline)
 	DEFINE_ANIM_DECORATOR_END(FSequencePlayerDecorator)
 
-	void FSequencePlayerDecorator::FInstanceData::Construct(FExecutionContext& Context, FWeakDecoratorPtr DecoratorPtr, const FAnimNextSequencePlayerDecoratorSharedData& SharedData)
+	void FSequencePlayerDecorator::FInstanceData::Construct(FExecutionContext& Context, const FDecoratorBinding& Binding)
 	{
-		if (SharedData.AnimSequence != nullptr)
+		const FSharedData* SharedData = Binding.GetSharedData<FSharedData>();
+		if (SharedData->AnimSequence != nullptr)
 		{
-			const float SequenceLength = SharedData.AnimSequence->GetPlayLength();
-			InternalTimeAccumulator = FMath::Clamp(SharedData.StartPosition, 0.f, SequenceLength);
+			const float SequenceLength = SharedData->AnimSequence->GetPlayLength();
+			InternalTimeAccumulator = FMath::Clamp(SharedData->GetStartPosition(Context, Binding), 0.0f, SequenceLength);
 			PrevInternalTimeAccumulator = InternalTimeAccumulator;
 		}
 	}
@@ -44,7 +45,7 @@ namespace UE::AnimNext
 		const FAnimExtractContext ExtractionContext(static_cast<double>(InstanceData->InternalTimeAccumulator)
 			, false /*Output.AnimInstanceProxy->ShouldExtractRootMotion()*/
 			, DeltaTimeRecord
-			, SharedData->bLoop);
+			, SharedData->GetbLoop(Context, Binding));
 
 		const FAnimNextGraphReferencePose& GraphReferencePose = ParamStack.GetParam<FAnimNextGraphReferencePose>("GraphReferencePose");
 		const int32 GraphLODLevel = ParamStack.GetParam<int32>("GraphLODLevel");
@@ -66,7 +67,7 @@ namespace UE::AnimNext
 	double FSequencePlayerDecorator::GetPlayRate(FExecutionContext& Context, const TDecoratorBinding<ITimeline>& Binding) const
 	{
 		const FSharedData* SharedData = Binding.GetSharedData<FSharedData>();
-		return SharedData->PlayRate;
+		return SharedData->GetPlayRate(Context, Binding);
 	}
 
 	void FSequencePlayerDecorator::PreUpdate(FExecutionContext& Context, const TDecoratorBinding<IUpdate>& Binding) const
@@ -86,12 +87,14 @@ namespace UE::AnimNext
 
 			const float EffectiveDelta = FMath::IsNearlyZero(DeltaTime) || FMath::IsNearlyZero(PlayRate) ? 0.f : DeltaTime * PlayRate;
 
+			const bool bIsLooping = SharedData->GetbLoop(Context, Binding);
+
 			const float SequenceLength = SharedData->AnimSequence->GetPlayLength();
-			float CurrentTime = SharedData->bLoop
+			float CurrentTime = bIsLooping
 				? FMath::Fmod(InstanceData->InternalTimeAccumulator + EffectiveDelta, SequenceLength)
 				: FMath::Clamp(InstanceData->InternalTimeAccumulator + EffectiveDelta, 0.f, SequenceLength);
 
-			if (SharedData->bLoop && CurrentTime < 0.f)
+			if (bIsLooping && CurrentTime < 0.f)
 			{
 				CurrentTime += SequenceLength;
 			}
diff --git a/Engine/Plugins/Experimental/AnimNext/Source/AnimNext/Private/Decorators/SequencePlayer.h b/Engine/Plugins/Experimental/AnimNext/Source/AnimNext/Private/Decorators/SequencePlayer.h
--- a/Engine/Plugins/Experimental/AnimNext/Source/AnimNext/Private/Decorators/SequencePlayer.h
+++ b/Engine/Plugins/Experimental/AnimNext/Source/AnimNext/Private/Decorators/SequencePlayer.h
@@ -17,20 +17,29 @@ struct FAnimNextSequencePlayerDecoratorSharedData : public FAnimNextDecoratorSha
 	GENERATED_BODY()
 
 	/** The sequence to play. */
-	UPROPERTY(meta = (Input))
+	UPROPERTY(meta = (Inline))
 	TObjectPtr<UAnimSequence> AnimSequence;
 
 	/** The play rate multiplier at which this sequence plays. */
-	UPROPERTY(meta = (Input))
+	UPROPERTY()
 	float PlayRate = 1.0f;
 
 	/** The time at which we should start playing this sequence. */
-	UPROPERTY(meta = (Input))
+	UPROPERTY()
 	float StartPosition = 0.0f;
 
 	/** Whether or not this sequence playback will loop. */
-	UPROPERTY(meta = (Input))
+	UPROPERTY()
 	bool bLoop = 0.0f;
+
+	// Latent pin support boilerplate
+	#define DECORATOR_LATENT_PROPERTIES_ENUMERATOR(GeneratorMacro) \
+		GeneratorMacro(PlayRate) \
+		GeneratorMacro(StartPosition) \
+		GeneratorMacro(bLoop) \
+
+	GENERATE_DECORATOR_LATENT_PROPERTIES(FAnimNextSequencePlayerDecoratorSharedData, DECORATOR_LATENT_PROPERTIES_ENUMERATOR)
+	#undef DECORATOR_LATENT_PROPERTIES_ENUMERATOR
 };
 
 namespace UE::AnimNext
@@ -51,7 +60,7 @@ namespace UE::AnimNext
 			float InternalTimeAccumulator = 0.0f;
 			float PrevInternalTimeAccumulator = 0.0f;
 
-			void Construct(FExecutionContext& Context, FWeakDecoratorPtr DecoratorPtr, const FAnimNextSequencePlayerDecoratorSharedData& SharedData);
+			void Construct(FExecutionContext& Context, const FDecoratorBinding& Binding);
 		};
 
 		// IEvaluate impl
diff --git a/Engine/Plugins/Experimental/AnimNext/Source/AnimNext/Private/Graph/AnimGraph/AnimNode_AnimNextGraph.cpp b/Engine/Plugins/Experimental/AnimNext/Source/AnimNext/Private/Graph/AnimGraph/AnimNode_AnimNextGraph.cpp
--- a/Engine/Plugins/Experimental/AnimNext/Source/AnimNext/Private/Graph/AnimGraph/AnimNode_AnimNextGraph.cpp
+++ b/Engine/Plugins/Experimental/AnimNext/Source/AnimNext/Private/Graph/AnimGraph/AnimNode_AnimNextGraph.cpp
@@ -29,14 +29,6 @@ FAnimNode_AnimNextGraph::FAnimNode_AnimNextGraph()
 {
 }
 
-FAnimNode_AnimNextGraph::~FAnimNode_AnimNextGraph()
-{
-	if (GraphInstancePtr.IsValid())
-	{
-		AnimNextGraph->ReleaseInstance(GraphInstancePtr);
-	}
-}
-
 void FAnimNode_AnimNextGraph::OnInitializeAnimInstance(const FAnimInstanceProxy* InProxy, const UAnimInstance* InAnimInstance)
 {
 	DECLARE_SCOPE_HIERARCHICAL_COUNTER_FUNC()
@@ -57,6 +49,8 @@ void FAnimNode_AnimNextGraph::Update_AnyThread(const FAnimationUpdateContext& Co
 {
 	DECLARE_SCOPE_HIERARCHICAL_COUNTER_FUNC()
 
+	using namespace UE::AnimNext;
+
 	SourceLink.Update(Context);
 
 	if (IsLODEnabled(Context.AnimInstanceProxy))
@@ -68,13 +62,16 @@ void FAnimNode_AnimNextGraph::Update_AnyThread(const FAnimationUpdateContext& Co
 		// Populate our param stack since our instance data might need it during construction
 		const int32 LODLevel = Context.AnimInstanceProxy->GetLODLevel();
 
-		UE::AnimNext::FParamStack::Get().PushValues(
+		FParamStack& ParamStack = FParamStack::Get();
+		FParamStack::FPushedLayerHandle LayerHandle = ParamStack.PushValues(
 			"GraphLODLevel", LODLevel
 		);
 
-		UE::AnimNext::FContext AnimNextContext(Context.GetDeltaTime());
+		FContext AnimNextContext(Context.GetDeltaTime());
+
+		AnimNextGraph->Run(AnimNextContext, GraphInstance, EAnimNextGraphSimulationSteps::Update);
 
-		AnimNextGraph->Run(AnimNextContext, GraphInstancePtr, EAnimNextGraphSimulationSteps::Update);
+		ParamStack.PopLayer(LayerHandle);
 	}
 
 	FAnimNode_CustomProperty::Update_AnyThread(Context);
@@ -86,20 +83,25 @@ void FAnimNode_AnimNextGraph::Initialize_AnyThread(const FAnimationInitializeCon
 {
 	DECLARE_SCOPE_HIERARCHICAL_COUNTER_FUNC()
 
+	using namespace UE::AnimNext;
+
 	SourceLink.Initialize(Context);
 
-	if (!GraphInstancePtr.IsValid() && AnimNextGraph)
+	if (!GraphInstance.IsValid() && AnimNextGraph)
 	{
 		// If we don't have an instance yet, create one
 
 		// Populate our param stack since our instance data might need it during construction
 		const int32 LODLevel = Context.AnimInstanceProxy->GetLODLevel();
 
-		UE::AnimNext::FParamStack::Get().PushValues(
+		FParamStack& ParamStack = FParamStack::Get();
+		FParamStack::FPushedLayerHandle LayerHandle = ParamStack.PushValues(
 			"GraphLODLevel", LODLevel
 		);
 
-		GraphInstancePtr = AnimNextGraph->AllocateInstance();
+		AnimNextGraph->AllocateInstance(GraphInstance);
+
+		ParamStack.PopLayer(LayerHandle);
 	}
 
 	FAnimNode_CustomProperty::Initialize_AnyThread(Context);
@@ -147,7 +149,8 @@ void FAnimNode_AnimNextGraph::Evaluate_AnyThread(FPoseContext & Output)
 	FAnimNextGraphLODPose ResultPose(FLODPose(RefPose, LODLevel, true, Output.ExpectsAdditivePose()));
 	FGenerationTools::RemapPose(LODLevel, SourcePose, RefPose, GraphSourceLODPose.LODPose);
 
-	FParamStack::FPushedLayerHandle LayerHandle = Context.GetMutableParamStack().PushValues(
+	FParamStack& ParamStack = FParamStack::Get();
+	FParamStack::FPushedLayerHandle LayerHandle = ParamStack.PushValues(
 		"GraphReferencePose", GraphReferencePose,
 		"ResultPose", ResultPose,
 		"GraphLODLevel", LODLevel,
@@ -155,13 +158,13 @@ void FAnimNode_AnimNextGraph::Evaluate_AnyThread(FPoseContext & Output)
 		"SourcePose", GraphSourceLODPose						// TODO : Pass this as a external variable maybe ? When we have support for variables in the rigvm graph
 	);
 
-	AnimNextGraph->Run(Context, GraphInstancePtr, EAnimNextGraphSimulationSteps::Evaluate);
+	AnimNextGraph->Run(Context, GraphInstance, EAnimNextGraphSimulationSteps::Evaluate);
 
 	FGenerationTools::RemapPose(LODLevel, RefPose, ResultPose.LODPose, Output);
 
 	FAnimNode_CustomProperty::Evaluate_AnyThread(Output);
 
-	Context.GetMutableParamStack().PopLayer(LayerHandle);
+	ParamStack.PopLayer(LayerHandle);
 }
 
 void FAnimNode_AnimNextGraph::PostSerialize(const FArchive& Ar)
diff --git a/Engine/Plugins/Experimental/AnimNext/Source/AnimNext/Private/Graph/AnimNextGraph.cpp b/Engine/Plugins/Experimental/AnimNext/Source/AnimNext/Private/Graph/AnimNextGraph.cpp
--- a/Engine/Plugins/Experimental/AnimNext/Source/AnimNext/Private/Graph/AnimNextGraph.cpp
+++ b/Engine/Plugins/Experimental/AnimNext/Source/AnimNext/Private/Graph/AnimNextGraph.cpp
@@ -3,7 +3,7 @@
 #include "Graph/AnimNextGraph.h"
 #include "RigVMCore/RigVMMemoryStorage.h"
 #include "RigVMCore/RigVMExecuteContext.h"
-#include "Graph/RigUnit_AnimNextGraphRoot.h"
+#include "Graph/RigUnit_AnimNextShimRoot.h"
 #include "UObject/ObjectSaveContext.h"
 #include "UObject/Package.h"
 #include "Param/ParamTypeHandle.h"
@@ -19,42 +19,63 @@ const FName EntryPointName("GetData");
 const FName ResultName("Result");
 }
 
-UE::AnimNext::FDecoratorPtr UAnimNextGraph::AllocateInstance() const
+FAnimNextGraphInstance::~FAnimNextGraphInstance()
 {
-	if (!ResolvedRootDecoratorHandle.IsValid())
+	Release();
+}
+
+void FAnimNextGraphInstance::Release()
+{
+	if (!GraphInstancePtr.IsValid())
 	{
-		return UE::AnimNext::FDecoratorPtr();
+		return;
 	}
 
-	UE::AnimNext::FExecutionContext Context(SharedDataBuffer);
+	// We need an execution context for this graph to be active when we delete the graph instance
+	UE::AnimNext::FExecutionContext Context(Graph->SharedDataBuffer);
 
-	return Context.AllocateNodeInstance(UE::AnimNext::FWeakDecoratorPtr(), ResolvedRootDecoratorHandle);
+	GraphInstancePtr.Reset();
+	Graph = nullptr;
+	ExtendedExecuteContext.Reset();
 }
 
-void UAnimNextGraph::ReleaseInstance(UE::AnimNext::FDecoratorPtr& GraphInstancePtr) const
+bool FAnimNextGraphInstance::IsValid() const
 {
-	if (!GraphInstancePtr.IsValid())
+	return GraphInstancePtr.IsValid();
+}
+
+void UAnimNextGraph::AllocateInstance(FAnimNextGraphInstance& Instance) const
+{
+	Instance.Release();
+
+	if (!ResolvedRootDecoratorHandle.IsValid())
 	{
 		return;
 	}
 
-	// We need an execution context for this graph to be active when we delete the graph instance
+	Instance.Graph = this;
+
 	UE::AnimNext::FExecutionContext Context(SharedDataBuffer);
+	Instance.GraphInstancePtr = Context.AllocateNodeInstance(UE::AnimNext::FWeakDecoratorPtr(), ResolvedRootDecoratorHandle);
 
-	GraphInstancePtr.Reset();
+	Instance.ExtendedExecuteContext.CopyMemoryStorage(ExtendedExecuteContext, RigVM);
+
+	RigVM->InitializeInstance(Instance.ExtendedExecuteContext, RigVM->GetLocalMemoryArray(Instance.ExtendedExecuteContext));
 }
 
-void UAnimNextGraph::Run(const UE::AnimNext::FContext& Context, UE::AnimNext::FWeakDecoratorPtr GraphInstancePtr, EAnimNextGraphSimulationSteps SimulationSteps) const
+void UAnimNextGraph::Run(const UE::AnimNext::FContext& Context, FAnimNextGraphInstance& GraphInstance, EAnimNextGraphSimulationSteps SimulationSteps) const
 {
-	if (RigVM && GraphInstancePtr.IsValid())
+	if (RigVM && GraphInstance.IsValid())
 	{
-		FRigVMExtendedExecuteContext RigVMExtendedExecuteContext;
-		FAnimNextExecuteContext& AnimNextContext = RigVMExtendedExecuteContext.GetPublicDataSafe<FAnimNextExecuteContext>();
+		FAnimNextExecuteContext& AnimNextContext = GraphInstance.ExtendedExecuteContext.GetPublicDataSafe<FAnimNextExecuteContext>();
 		AnimNextContext.SetContextData(Context);
-		AnimNextContext.InitializeWithGraph(SharedDataBuffer, GraphInstancePtr);
+		AnimNextContext.InitializeWithGraph(SharedDataBuffer, GraphInstance.GraphInstancePtr);
 		AnimNextContext.SetSimulationSteps(SimulationSteps);
 
-		RigVM->Execute(RigVMExtendedExecuteContext, TArray<TRigVMMemoryStorage*>(), FRigUnit_AnimNextGraphRoot::EventName);
+		RigVM->Execute(GraphInstance.ExtendedExecuteContext, FRigUnit_AnimNextShimRoot::EventName);
+
+		// Reset the context to avoid issues if we forget to reset it the next time we use it
+		AnimNextContext.DebugReset();
 	}
 }
 
@@ -80,6 +101,13 @@ static TArray<UClass*> GetClassObjectsInPackage(UPackage* InPackage)
 	return ClassObjects;
 }
 
+void UAnimNextGraph::PostLoad()
+{
+	Super::PostLoad();
+
+	ExtendedExecuteContext.InvalidateCachedMemory();
+}
+
 void UAnimNextGraph::PostRename(UObject* OldOuter, const FName OldName)
 {
 	Super::PostRename(OldOuter, OldName);
diff --git a/Engine/Plugins/Experimental/AnimNext/Source/AnimNext/Private/Graph/RigDecorator_AnimNextCppDecorator.cpp b/Engine/Plugins/Experimental/AnimNext/Source/AnimNext/Private/Graph/RigDecorator_AnimNextCppDecorator.cpp
--- a/Engine/Plugins/Experimental/AnimNext/Source/AnimNext/Private/Graph/RigDecorator_AnimNextCppDecorator.cpp
+++ b/Engine/Plugins/Experimental/AnimNext/Source/AnimNext/Private/Graph/RigDecorator_AnimNextCppDecorator.cpp
@@ -19,7 +19,30 @@ void FRigDecorator_AnimNextCppDecorator::GetProgrammaticPins(URigVMController* I
 
 	FStructOnScope DefaultValueMemoryScope(DecoratorSharedDataStruct);
 
+	const int32 StartPinIndex = OutPinArray.Num();
 	OutPinArray.AddPins(DecoratorSharedDataStruct, InController, ERigVMPinDirection::Invalid, InParentPinIndex, DefaultValueMemoryScope.GetStructMemory(), true);
+
+	for (int32 PinIndex = StartPinIndex; PinIndex < OutPinArray.Num(); ++PinIndex)
+	{
+		FRigVMPinInfo& PinInfo = OutPinArray[PinIndex];
+
+		const bool bIsHidden = PinInfo.Property->HasMetaData(FRigVMStruct::HiddenMetaName);
+		const bool bIsInline = PinInfo.Property->HasMetaData("Inline");
+		const bool bIsDecoratorHandle = PinInfo.Property->GetCPPType() == TEXT("FAnimNextDecoratorHandle");
+
+		// Check if the metadata stipulates that we should explicitly hide this property, if not we mark it as an input
+		PinInfo.Direction = bIsHidden ? ERigVMPinDirection::Hidden : ERigVMPinDirection::Input;
+
+		// For top level properties of decorators, if we don't explicitly tag the property as inline or hidden, it is lazy
+		// Except for decorator handles which are never lazy since they just encode graph connectivity
+		if (InParentPinIndex == PinInfo.ParentIndex && !bIsHidden && !bIsInline && !bIsDecoratorHandle)
+		{
+			PinInfo.bIsLazy = true;
+		}
+
+		// Remove our property because we configure the pin explicitly
+		PinInfo.Property = nullptr;
+	}
 }
 
 const UE::AnimNext::FDecorator* FRigDecorator_AnimNextCppDecorator::GetDecorator() const
diff --git a/Engine/Plugins/Experimental/AnimNext/Source/AnimNext/Private/Graph/RigUnit_AnimNextGraphEvaluator.cpp b/Engine/Plugins/Experimental/AnimNext/Source/AnimNext/Private/Graph/RigUnit_AnimNextGraphEvaluator.cpp
new file mode 100644
--- /dev/null
+++ b/Engine/Plugins/Experimental/AnimNext/Source/AnimNext/Private/Graph/RigUnit_AnimNextGraphEvaluator.cpp
@@ -0,0 +1,39 @@
+// Copyright Epic Games, Inc. All Rights Reserved.
+
+#include "Graph/RigUnit_AnimNextGraphEvaluator.h"
+#include "DecoratorBase/ExecutionContext.h"
+#include "DecoratorInterfaces/IUpdate.h"
+#include "DecoratorInterfaces/IEvaluate.h"
+#include "Context.h"
+
+#include UE_INLINE_GENERATED_CPP_BY_NAME(RigUnit_AnimNextGraphEvaluator)
+
+void FRigUnit_AnimNextGraphEvaluator::StaticExecute(FRigVMExtendedExecuteContext& RigVMExecuteContext, FRigVMMemoryHandleArray RigVMMemoryHandles, FRigVMPredicateBranchArray RigVMBranches)
+{
+	using namespace UE::AnimNext;
+
+	const FAnimNextExecuteContext& VMExecuteContext = RigVMExecuteContext.GetPublicData<FAnimNextExecuteContext>();
+
+	// Setup what we need to execute
+	FExecutionContext Context(VMExecuteContext.GetSharedDataBuffer(), RigVMExecuteContext, RigVMMemoryHandles);
+	FWeakDecoratorPtr GraphInstancePtr = VMExecuteContext.GetGraphInstancePtr();
+	const EAnimNextGraphSimulationSteps SimulationSteps = VMExecuteContext.GetSimulationSteps();
+
+	if (EnumHasAnyFlags(SimulationSteps, EAnimNextGraphSimulationSteps::Update))
+	{
+		const FContext& InterfaceContext = VMExecuteContext.GetContext();
+
+		// Call pre/post update on our graph
+		UpdateGraph(Context, GraphInstancePtr, InterfaceContext.GetDeltaTime());
+	}
+
+	if (EnumHasAnyFlags(SimulationSteps, EAnimNextGraphSimulationSteps::Evaluate))
+	{
+		// Call pre/post evaluate on our graph
+		FPoseContainer PoseContainer;
+		EvaluateGraph(Context, GraphInstancePtr, EEvaluationFlags::All, PoseContainer);
+
+		// TODO: Write out our output pose/curves/attributes/etc
+		//ExecuteContext_.GetContext().SetResult<FAnimationGraphNodeResult>(Result);
+	}
+}
diff --git a/Engine/Plugins/Experimental/AnimNext/Source/AnimNext/Private/Graph/RigUnit_AnimNextGraphRoot.cpp b/Engine/Plugins/Experimental/AnimNext/Source/AnimNext/Private/Graph/RigUnit_AnimNextGraphRoot.cpp
--- a/Engine/Plugins/Experimental/AnimNext/Source/AnimNext/Private/Graph/RigUnit_AnimNextGraphRoot.cpp
+++ b/Engine/Plugins/Experimental/AnimNext/Source/AnimNext/Private/Graph/RigUnit_AnimNextGraphRoot.cpp
@@ -1,41 +1,13 @@
 ﻿// Copyright Epic Games, Inc. All Rights Reserved.
 
 #include "Graph/RigUnit_AnimNextGraphRoot.h"
-#include "DecoratorBase/ExecutionContext.h"
-#include "DecoratorInterfaces/IUpdate.h"
-#include "DecoratorInterfaces/IEvaluate.h"
-#include "Context.h"
 
 #include UE_INLINE_GENERATED_CPP_BY_NAME(RigUnit_AnimNextGraphRoot)
 
-FName FRigUnit_AnimNextGraphRoot::EventName = TEXT("Execute");
+FName FRigUnit_AnimNextGraphRoot::EventName = TEXT("DummyExecute");
 
-FRigUnit_AnimNextGraphRoot_Execute()
+FRigUnit_AnimNextGraphRoot_DummyExecute()
 {
-	using namespace UE::AnimNext;
-
-	const FAnimNextExecuteContext& VMExecuteContext = static_cast<const FAnimNextExecuteContext&>(ExecuteContext);
-
-	// Setup what we need to execute
-	FExecutionContext Context(VMExecuteContext.GetSharedDataBuffer());
-	FWeakDecoratorPtr GraphInstancePtr = VMExecuteContext.GetGraphInstancePtr();
-	const EAnimNextGraphSimulationSteps SimulationSteps = VMExecuteContext.GetSimulationSteps();
-
-	if (EnumHasAnyFlags(SimulationSteps, EAnimNextGraphSimulationSteps::Update))
-	{
-		const FContext& InterfaceContext = ExecuteContext.GetContext();
-
-		// Call pre/post update on our graph
-		UpdateGraph(Context, GraphInstancePtr, InterfaceContext.GetDeltaTime());
-	}
-
-	if (EnumHasAnyFlags(SimulationSteps, EAnimNextGraphSimulationSteps::Evaluate))
-	{
-		// Call pre/post evaluate on our graph
-		FPoseContainer PoseContainer;
-		EvaluateGraph(Context, GraphInstancePtr, EEvaluationFlags::All, PoseContainer);
-
-		// TODO: Write out our output pose/curves/attributes/etc
-		//ExecuteContext_.GetContext().SetResult<FAnimationGraphNodeResult>(Result);
-	}
+	// This RigUnit should never execute, it is only used in the editor
+	ensure(false);
 }
diff --git a/Engine/Plugins/Experimental/AnimNext/Source/AnimNext/Private/Graph/RigUnit_AnimNextShimRoot.cpp b/Engine/Plugins/Experimental/AnimNext/Source/AnimNext/Private/Graph/RigUnit_AnimNextShimRoot.cpp
new file mode 100644
--- /dev/null
+++ b/Engine/Plugins/Experimental/AnimNext/Source/AnimNext/Private/Graph/RigUnit_AnimNextShimRoot.cpp
@@ -0,0 +1,13 @@
+﻿// Copyright Epic Games, Inc. All Rights Reserved.
+
+#include "Graph/RigUnit_AnimNextShimRoot.h"
+
+#include UE_INLINE_GENERATED_CPP_BY_NAME(RigUnit_AnimNextShimRoot)
+
+FName FRigUnit_AnimNextShimRoot::EventName = TEXT("Execute");
+
+FRigUnit_AnimNextShimRoot_Execute()
+{
+	// This RigUnit is a no-op entry point, it simply forwards execution to the real entry point
+	// @see FRigUnit_AnimNextRuntimeRoot
+}
diff --git a/Engine/Plugins/Experimental/AnimNext/Source/AnimNext/Public/DecoratorBase/Decorator.h b/Engine/Plugins/Experimental/AnimNext/Source/AnimNext/Public/DecoratorBase/Decorator.h
--- a/Engine/Plugins/Experimental/AnimNext/Source/AnimNext/Public/DecoratorBase/Decorator.h
+++ b/Engine/Plugins/Experimental/AnimNext/Source/AnimNext/Public/DecoratorBase/Decorator.h
@@ -4,6 +4,7 @@
 
 #include "CoreMinimal.h"
 
+#include "DecoratorBase/DecoratorBinding.h"
 #include "DecoratorBase/DecoratorHandle.h"			// Derived types are likely to refer to other decorators as children
 #include "DecoratorBase/DecoratorInstanceData.h"
 #include "DecoratorBase/DecoratorMode.h"
@@ -11,6 +12,9 @@
 #include "DecoratorBase/DecoratorSharedData.h"
 #include "DecoratorBase/DecoratorUID.h"
 #include "DecoratorBase/IDecoratorInterface.h"
+#include "DecoratorBase/LatentPropertyHandle.h"
+
+#include <type_traits>
 
 class FArchive;
 
@@ -25,11 +29,12 @@ class FArchive;
 	static const UE::AnimNext::FDecoratorMemoryLayout DecoratorMemoryDescription; \
 	virtual UE::AnimNext::FDecoratorMemoryLayout GetDecoratorMemoryDescription() const override { return DecoratorMemoryDescription; } \
 	virtual UScriptStruct* GetDecoratorSharedDataStruct() const override { return FSharedData::StaticStruct(); } \
-	virtual void ConstructDecoratorInstance(UE::AnimNext::FExecutionContext& Context, UE::AnimNext::FWeakDecoratorPtr DecoratorPtr, const FAnimNextDecoratorSharedData& DecoratorDesc, UE::AnimNext::FDecoratorInstanceData& DecoratorInstance) const override; \
-	virtual void DestructDecoratorInstance(UE::AnimNext::FExecutionContext& Context, UE::AnimNext::FWeakDecoratorPtr DecoratorPtr, const FAnimNextDecoratorSharedData& DecoratorDesc, UE::AnimNext::FDecoratorInstanceData& DecoratorInstance) const override; \
+	virtual void ConstructDecoratorInstance(UE::AnimNext::FExecutionContext& Context, const UE::AnimNext::FDecoratorBinding& Binding) const override; \
+	virtual void DestructDecoratorInstance(UE::AnimNext::FExecutionContext& Context, const UE::AnimNext::FDecoratorBinding& Binding) const override; \
 	virtual const UE::AnimNext::IDecoratorInterface* GetDecoratorInterface(UE::AnimNext::FDecoratorInterfaceUID InterfaceUID) const override; \
-	static_assert(std::is_base_of<FAnimNextDecoratorSharedData, FSharedData>::value, TEXT("Decorator shared data must derive from FAnimNextDecoratorSharedData")); \
-	static_assert(std::is_base_of<FDecoratorInstanceData, FInstanceData>::value, TEXT("Decorator instance data must derive from FDecoratorInstanceData"));
+	virtual uint32 GetNumLatentDecoratorProperties() const override { return -FSharedData::GetLatentPropertyIndex(~(size_t)0); } \
+	static_assert(std::is_base_of<FAnimNextDecoratorSharedData, FSharedData>::value, "Decorator shared data must derive from FAnimNextDecoratorSharedData"); \
+	static_assert(std::is_base_of<FDecoratorInstanceData, FInstanceData>::value, "Decorator instance data must derive from FDecoratorInstanceData");
 
 #define DECLARE_ABSTRACT_ANIM_DECORATOR(DecoratorName, DecoratorNameHash, SuperDecoratorName) \
 	using DecoratorSuper = SuperDecoratorName; \
@@ -50,16 +55,16 @@ class FArchive;
 #define DEFINE_ANIM_DECORATOR_BEGIN(DecoratorName) \
 	const UE::AnimNext::FDecoratorMemoryLayout DecoratorName::DecoratorMemoryDescription = \
 		UE::AnimNext::FDecoratorMemoryLayout{ sizeof(DecoratorName), alignof(DecoratorName), sizeof(DecoratorName::FSharedData), alignof(DecoratorName::FSharedData), sizeof(DecoratorName::FInstanceData), alignof(DecoratorName::FInstanceData) }; \
-	void DecoratorName::ConstructDecoratorInstance(UE::AnimNext::FExecutionContext& Context, UE::AnimNext::FWeakDecoratorPtr DecoratorPtr, const FAnimNextDecoratorSharedData& DecoratorDesc, UE::AnimNext::FDecoratorInstanceData& DecoratorInstance) const \
+	void DecoratorName::ConstructDecoratorInstance(UE::AnimNext::FExecutionContext& Context, const UE::AnimNext::FDecoratorBinding& Binding) const \
 	{ \
-		FInstanceData* Data = new(&static_cast<FInstanceData&>(DecoratorInstance)) FInstanceData(); \
-		Data->Construct(Context, DecoratorPtr, static_cast<const FSharedData&>(DecoratorDesc)); \
+		FInstanceData* Data = new(Binding.GetInstanceData<FInstanceData>()) FInstanceData(); \
+		Data->Construct(Context, Binding); \
 	} \
-	void DecoratorName::DestructDecoratorInstance(UE::AnimNext::FExecutionContext& Context, UE::AnimNext::FWeakDecoratorPtr DecoratorPtr, const FAnimNextDecoratorSharedData& DecoratorDesc, UE::AnimNext::FDecoratorInstanceData& DecoratorInstance) const \
+	void DecoratorName::DestructDecoratorInstance(UE::AnimNext::FExecutionContext& Context, const UE::AnimNext::FDecoratorBinding& Binding) const \
 	{ \
-		FInstanceData& Data = static_cast<FInstanceData&>(DecoratorInstance); \
-		Data.Destruct(Context, DecoratorPtr, DecoratorDesc); \
-		Data.~FInstanceData(); \
+		FInstanceData* Data = Binding.GetInstanceData<FInstanceData>(); \
+		Data->Destruct(Context, Binding); \
+		Data->~FInstanceData(); \
 	} \
 	const UE::AnimNext::IDecoratorInterface* DecoratorName::GetDecoratorInterface(UE::AnimNext::FDecoratorInterfaceUID InterfaceUID_) const \
 	{
@@ -111,16 +116,16 @@ namespace UE::AnimNext
 		// The alignment in bytes of an instance of the decorator class which derives from FDecorator
 		uint32 DecoratorAlignment = 0;
 
-		// The size in bytes of the shared data for the decorator
+		// The size in bytes of the shared data for the decorator which derives from FAnimNextDecoratorSharedData
 		uint32 SharedDataSize = 0;
 
-		// The alignment in bytes of the shared data for the decorator
+		// The alignment in bytes of the shared data for the decorator which derives from FAnimNextDecoratorSharedData
 		uint32 SharedDataAlignment = 0;
 
-		// The size in bytes of the instance data for the decorator
+		// The size in bytes of the instance data for the decorator which derives from FDecoratorInstanceData
 		uint32 InstanceDataSize = 0;
 
-		// The alignment in bytes of the instance data for the decorator
+		// The alignment in bytes of the instance data for the decorator which derives from FDecoratorInstanceData
 		uint32 InstanceDataAlignment = 0;
 	};
 
@@ -163,15 +168,15 @@ namespace UE::AnimNext
 		virtual FString GetDecoratorName() const { return TEXT("FDecorator"); }
 
 		// Returns the memory requirements of the derived decorator instance
-		virtual FDecoratorMemoryLayout GetDecoratorMemoryDescription() const { return { sizeof(FDecorator), alignof(FDecorator) }; }
+		virtual FDecoratorMemoryLayout GetDecoratorMemoryDescription() const = 0;
 
 		// Returns the UScriptStruct associated with the shared data for the decorator
 		virtual UScriptStruct* GetDecoratorSharedDataStruct() const { return FSharedData::StaticStruct(); }
 
 		// Called when a new instance of the decorator is created or destroyed
 		// Derived types must override this and forward to the instance data constructor/destructor
-		virtual void ConstructDecoratorInstance(FExecutionContext& Context, FWeakDecoratorPtr DecoratorPtr, const FAnimNextDecoratorSharedData& DecoratorDesc, FDecoratorInstanceData& DecoratorInstance) const = 0;
-		virtual void DestructDecoratorInstance(FExecutionContext& Context, FWeakDecoratorPtr DecoratorPtr, const FAnimNextDecoratorSharedData& DecoratorDesc, FDecoratorInstanceData& DecoratorInstance) const = 0;
+		virtual void ConstructDecoratorInstance(FExecutionContext& Context, const FDecoratorBinding& Binding) const = 0;
+		virtual void DestructDecoratorInstance(FExecutionContext& Context, const FDecoratorBinding& Binding) const = 0;
 
 		// Returns the decorator mode.
 		virtual EDecoratorMode GetDecoratorMode() const = 0;
@@ -202,6 +207,9 @@ namespace UE::AnimNext
 			return nullptr;
 		}
 
+		// The number of latent property properties in the shared data of this decorator
+		virtual uint32 GetNumLatentDecoratorProperties() const { return 0; }
+
 		// Called to serialize decorator shared data
 		virtual void SerializeDecoratorSharedData(FArchive& Ar, FAnimNextDecoratorSharedData& SharedData) const;
 
@@ -210,7 +218,10 @@ namespace UE::AnimNext
 		// derived type using UE reflection.
 		// Decorators can override this function to control how editor only properties are coerced into the runtime shared data
 		// instance.
-		virtual void SaveDecoratorSharedData(FDecoratorWriter& Writer, const TFunction<FString(const FString& PropertyName)>& GetDecoratorProperty, FAnimNextDecoratorSharedData& OutSharedData) const;
+		virtual void SaveDecoratorSharedData(const TFunction<FString(const FString& PropertyName)>& GetDecoratorProperty, FAnimNextDecoratorSharedData& OutSharedData) const;
+
+		// Takes the editor properties as authored in the graph and returns the latent property handles using UE reflection.
+		TArray<FLatentPropertyHandle> GetLatentPropertyHandles(bool bFilterEditorOnly, const TFunction<bool(const FString& PropertyName)>& IsDecoratorPropertyLatent, FLatentPropertyHandle& CurrentLatentPropertyHandle) const;
 #endif
 	};
 
@@ -238,7 +249,7 @@ namespace UE::AnimNext
 	 */
 	struct ANIMNEXT_API FDecoratorStaticInitHook final
 	{
-		explicit FDecoratorStaticInitHook(DecoratorConstructorFunc DecoratorConstructor_);
+		explicit FDecoratorStaticInitHook(DecoratorConstructorFunc InDecoratorConstructor);
 		~FDecoratorStaticInitHook();
 
 	private:
diff --git a/Engine/Plugins/Experimental/AnimNext/Source/AnimNext/Public/DecoratorBase/DecoratorBinding.h b/Engine/Plugins/Experimental/AnimNext/Source/AnimNext/Public/DecoratorBase/DecoratorBinding.h
--- a/Engine/Plugins/Experimental/AnimNext/Source/AnimNext/Public/DecoratorBase/DecoratorBinding.h
+++ b/Engine/Plugins/Experimental/AnimNext/Source/AnimNext/Public/DecoratorBase/DecoratorBinding.h
@@ -5,6 +5,9 @@
 #include "CoreMinimal.h"
 #include "DecoratorBase/DecoratorPtr.h"
 #include "DecoratorBase/DecoratorInterfaceUID.h"
+#include "DecoratorBase/DecoratorTemplate.h"
+
+#include <type_traits>
 
 struct FAnimNextDecoratorSharedData;
 
@@ -12,6 +15,7 @@ namespace UE::AnimNext
 {
 	struct FDecoratorInstanceData;
 	struct IDecoratorInterface;
+	struct FNodeDescription;
 
 	/**
 	 * FDecoratorBinding
@@ -24,18 +28,50 @@ namespace UE::AnimNext
 		// Creates an empty binding.
 		FDecoratorBinding() = default;
 
-		// Returns whether or not this binding is valid and set to an interface.
-		bool IsValid() const { return Interface != nullptr; }
+		// Returns whether or not this binding is valid.
+		bool IsValid() const { return DecoratorTemplate != nullptr; }
 
 		// Queries a node for a pointer to its decorator shared data.
 		// If the decorator handle is invalid, a null pointer is returned.
 		template<class SharedDataType>
-		const SharedDataType* GetSharedData() const { return static_cast<const SharedDataType*>(SharedData); }
+		const SharedDataType* GetSharedData() const
+		{
+			static_assert(std::is_base_of<FAnimNextDecoratorSharedData, SharedDataType>::value, "Decorator shared data must derive from FAnimNextDecoratorSharedData");
+
+			if (!IsValid())
+			{
+				return nullptr;
+			}
+
+			return static_cast<const SharedDataType*>(DecoratorTemplate->GetDecoratorDescription(*NodeDescription));
+		}
 
 		// Queries a node for a pointer to its decorator instance data.
 		// If the decorator handle is invalid, a null pointer is returned.
 		template<class InstanceDataType>
-		InstanceDataType* GetInstanceData() const { return static_cast<InstanceDataType*>(InstanceData); }
+		InstanceDataType* GetInstanceData() const
+		{
+			static_assert(std::is_base_of<FDecoratorInstanceData, InstanceDataType>::value, "Decorator instance data must derive from FDecoratorInstanceData");
+
+			if (!IsValid())
+			{
+				return nullptr;
+			}
+
+			return static_cast<InstanceDataType*>(DecoratorTemplate->GetDecoratorInstance(*DecoratorPtr.GetNodeInstance()));
+		}
+
+		// Queries a node for a pointer to its decorator latent properties.
+		// If the decorator handle is invalid or if we have no latent properties, a null pointer is returned.
+		const FLatentPropertyHandle* GetLatentPropertyHandles() const
+		{
+			if (!IsValid() || !DecoratorTemplate->HasLatentProperties())
+			{
+				return nullptr;
+			}
+
+			return DecoratorTemplate->GetDecoratorLatentPropertyHandles(*NodeDescription);
+		}
 
 		// Returns the decorator pointer we are bound to.
 		FWeakDecoratorPtr GetDecoratorPtr() const { return DecoratorPtr; }
@@ -49,23 +85,32 @@ namespace UE::AnimNext
 
 	protected:
 		// Creates a valid binding
-		FDecoratorBinding(const IDecoratorInterface* Interface_, const FAnimNextDecoratorSharedData* SharedData_, FDecoratorInstanceData* InstanceData_, FWeakDecoratorPtr DecoratorPtr_)
-			: Interface(Interface_)
-			, SharedData(SharedData_)
-			, InstanceData(InstanceData_)
-			, DecoratorPtr(DecoratorPtr_)
+		FDecoratorBinding(const IDecoratorInterface* InInterface, const FDecoratorTemplate* InDecoratorTemplate, const FNodeDescription* InNodeDescription, FWeakDecoratorPtr InDecoratorPtr)
+			: Interface(InInterface)
+			, DecoratorTemplate(InDecoratorTemplate)
+			, NodeDescription(InNodeDescription)
+			, DecoratorPtr(InDecoratorPtr)
 		{}
 
 		// Performs a naked cast to the desired interface type
 		template<class DecoratorInterfaceType>
-		const DecoratorInterfaceType* GetInterfaceTyped() const { return static_cast<const DecoratorInterfaceType*>(Interface); }
+		const DecoratorInterfaceType* GetInterfaceTyped() const
+		{
+			static_assert(std::is_base_of<IDecoratorInterface, DecoratorInterfaceType>::value, "Decorator interface data must derive from IDecoratorInterface");
+			return static_cast<const DecoratorInterfaceType*>(Interface);
+		}
+
+		// A pointer to the bound interface or nullptr if we are bound to a decorator but none of its interfaces
+		const IDecoratorInterface*				Interface = nullptr;
 
-		const IDecoratorInterface*		Interface = nullptr;
+		// A pointer to the decorator template that implements the interface we are bound to or nullptr if we are invalid
+		const FDecoratorTemplate*				DecoratorTemplate = nullptr;
 
-		const FAnimNextDecoratorSharedData*		SharedData = nullptr;
-		FDecoratorInstanceData*			InstanceData = nullptr;
+		// A pointer to the node shared data we are bound to
+		const FNodeDescription*					NodeDescription = nullptr;
 
-		FWeakDecoratorPtr				DecoratorPtr;
+		// A weak handle to the decorator instance data we are bound to
+		FWeakDecoratorPtr						DecoratorPtr;
 
 		friend struct FExecutionContext;
 	};
diff --git a/Engine/Plugins/Experimental/AnimNext/Source/AnimNext/Public/DecoratorBase/DecoratorInstanceData.h b/Engine/Plugins/Experimental/AnimNext/Source/AnimNext/Public/DecoratorBase/DecoratorInstanceData.h
--- a/Engine/Plugins/Experimental/AnimNext/Source/AnimNext/Public/DecoratorBase/DecoratorInstanceData.h
+++ b/Engine/Plugins/Experimental/AnimNext/Source/AnimNext/Public/DecoratorBase/DecoratorInstanceData.h
@@ -9,6 +9,7 @@ struct FAnimNextDecoratorSharedData;
 
 namespace UE::AnimNext
 {
+	struct FDecoratorBinding;
 	struct FExecutionContext;
 
 	/**
@@ -24,20 +25,18 @@ namespace UE::AnimNext
 		// Called after the constructor has been called when a new instance is created.
 		// This is called after the default constructor.
 		// You can override this function by adding a new one with the same name on your
-		// derived type. You can also specialize the FAnimNextDecoratorSharedData with the derived
-		// version as long as the reference can coerce to 'const FAnimNextDecoratorSharedData&'.
+		// derived type.
 		// Decorators are constructed from the bottom to the top.
-		void Construct(FExecutionContext& Context, FWeakDecoratorPtr DecoratorPtr, const FAnimNextDecoratorSharedData& DecoratorDesc)
+		void Construct(FExecutionContext& Context, const FDecoratorBinding& Binding) noexcept
 		{
 		}
 
 		// Called before the destructor has been called when an instance is destroyed.
 		// This is called before the default destructor.
 		// You can override this function by adding a new one with the same name on your
-		// derived type. You can also specialize the FAnimNextDecoratorSharedData with the derived
-		// version as long as the reference can coerce to 'const FAnimNextDecoratorSharedData&'.
+		// derived type.
 		// Decorators are destructed from the top to the bottom.
-		void Destruct(FExecutionContext& Context, FWeakDecoratorPtr DecoratorPtr, const FAnimNextDecoratorSharedData& DecoratorDesc)
+		void Destruct(FExecutionContext& Context, const FDecoratorBinding& Binding) noexcept
 		{
 		}
 	};
diff --git a/Engine/Plugins/Experimental/AnimNext/Source/AnimNext/Public/DecoratorBase/DecoratorRegistryHandle.h b/Engine/Plugins/Experimental/AnimNext/Source/AnimNext/Public/DecoratorBase/DecoratorRegistryHandle.h
--- a/Engine/Plugins/Experimental/AnimNext/Source/AnimNext/Public/DecoratorBase/DecoratorRegistryHandle.h
+++ b/Engine/Plugins/Experimental/AnimNext/Source/AnimNext/Public/DecoratorBase/DecoratorRegistryHandle.h
@@ -15,29 +15,29 @@ namespace UE::AnimNext
 	struct FDecoratorRegistryHandle final
 	{
 		// Default constructed handles are invalid
-		FDecoratorRegistryHandle() = default;
+		FDecoratorRegistryHandle() noexcept = default;
 
 		// Returns whether or not this handle points to a valid decorator
-		bool IsValid() const { return HandleValue != 0; }
+		bool IsValid() const noexcept { return HandleValue != 0; }
 
 		// Returns whether or not this handle is valid and points to a static decorator
-		bool IsStatic() const { return IsValid() && HandleValue > 0; }
+		bool IsStatic() const noexcept { return IsValid() && HandleValue > 0; }
 
 		// Returns whether or not this handle is valid and points to a dynamic decorator
-		bool IsDynamic() const { return IsValid() && HandleValue < 0; }
+		bool IsDynamic() const noexcept { return IsValid() && HandleValue < 0; }
 
 		// Returns the static buffer offset for this handle when valid, otherwise INDEX_NONE
-		int32 GetStaticOffset() const { return IsStatic() ? (HandleValue - 1) : INDEX_NONE; }
+		int32 GetStaticOffset() const noexcept { return IsStatic() ? (HandleValue - 1) : INDEX_NONE; }
 
 		// Returns the dynamic array index for this handle when valid, otherwise INDEX_NONE
-		int32 GetDynamicIndex() const { return IsDynamic() ? (-HandleValue - 1) : INDEX_NONE; }
+		int32 GetDynamicIndex() const noexcept { return IsDynamic() ? (-HandleValue - 1) : INDEX_NONE; }
 
 		// Compares for equality and inequality
-		bool operator==(FDecoratorRegistryHandle RHS) const { return HandleValue == RHS.HandleValue; }
-		bool operator!=(FDecoratorRegistryHandle RHS) const { return HandleValue != RHS.HandleValue; }
+		bool operator==(FDecoratorRegistryHandle RHS) const noexcept { return HandleValue == RHS.HandleValue; }
+		bool operator!=(FDecoratorRegistryHandle RHS) const noexcept { return HandleValue != RHS.HandleValue; }
 
 	private:
-		explicit FDecoratorRegistryHandle(int16 HandleValue_)
+		explicit FDecoratorRegistryHandle(int16 HandleValue_) noexcept
 			: HandleValue(HandleValue_)
 		{}
 
diff --git a/Engine/Plugins/Experimental/AnimNext/Source/AnimNext/Public/DecoratorBase/DecoratorSharedData.h b/Engine/Plugins/Experimental/AnimNext/Source/AnimNext/Public/DecoratorBase/DecoratorSharedData.h
--- a/Engine/Plugins/Experimental/AnimNext/Source/AnimNext/Public/DecoratorBase/DecoratorSharedData.h
+++ b/Engine/Plugins/Experimental/AnimNext/Source/AnimNext/Public/DecoratorBase/DecoratorSharedData.h
@@ -3,9 +3,72 @@
 #pragma once
 
 #include "CoreMinimal.h"
+#include "DecoratorBase/DecoratorBinding.h"
+#include "DecoratorBase/ExecutionContext.h"
+#include "DecoratorBase/LatentPropertyHandle.h"
+
+#include <type_traits>
 
 #include "DecoratorSharedData.generated.h"
 
+#define ANIM_NEXT_IMPL_GET_LATENT_PROPERTY_INDEX_FOR_PROPERTY(PropertyName) \
+	LatentPropertyIndex--; \
+	if (LatentPropertyOffset == offsetof(Self, PropertyName)) \
+	{ \
+		return -LatentPropertyIndex; \
+	} \
+
+#define ANIM_NEXT_IMPL_DEFINE_GET_LATENT_PROPERTY_INDEX(EnumeratorMacro) \
+	/** @see FAnimNextDecoratorSharedData::GetLatentPropertyIndex */ \
+	static constexpr int32 GetLatentPropertyIndex(const size_t LatentPropertyOffset) \
+	{ \
+		int32 LatentPropertyIndex = Super::GetLatentPropertyIndex(LatentPropertyOffset); \
+		/* If the value is positive, the property lives in our base type */ \
+		/* Otherwise the value is the negative (or zero) number of latent properties in the base type */ \
+		if (LatentPropertyIndex > 0) \
+		{ \
+			return LatentPropertyIndex; \
+		} \
+		/* If a property in the struct is wrapped with WITH_EDITORONLY_DATA, then this code here */ \
+		/* needs to be wrapped as well */ \
+		EnumeratorMacro(ANIM_NEXT_IMPL_GET_LATENT_PROPERTY_INDEX_FOR_PROPERTY) \
+		/* Latent property wasn't found, return the number of latent properties seen so far */ \
+		return LatentPropertyIndex; \
+	} \
+
+#define ANIM_NEXT_IMPL_DEFINE_LATENT_GETTER(PropertyName) \
+	decltype(PropertyName) Get##PropertyName(UE::AnimNext::FExecutionContext& Context, const UE::AnimNext::FDecoratorBinding& Binding) const \
+	{ \
+		/* We need a mapping of latent property name/offset to latent property index */ \
+		/* This can be built once at runtime using the UE reflection and cached on first call or using a constexpr function, see below */ \
+		constexpr size_t LatentPropertyOffset = offsetof(Self, PropertyName); \
+		constexpr int32 LatentPropertyIndex = GetLatentPropertyIndex(LatentPropertyOffset); \
+		/* A negative or zero property index means the property isn't latent, we can static assert since it isn't possible */ \
+		static_assert(LatentPropertyIndex > 0, "Property " #PropertyName " isn't latent"); \
+		const UE::AnimNext::FLatentPropertyHandle* LatentPropertyHandles = Binding.GetLatentPropertyHandles(); \
+		const UE::AnimNext::FLatentPropertyHandle LatentPropertyHandle = LatentPropertyHandles[LatentPropertyIndex - 1]; \
+		/* An invalid handle means we are inline, otherwise we are latent */ \
+		if (!LatentPropertyHandle.IsValid()) \
+		{ \
+			return PropertyName; \
+		} \
+		else \
+		{ \
+			return Context.EvaluateLatentPin<decltype(PropertyName)>(LatentPropertyHandle); \
+		} \
+	} \
+
+/**
+  * This macro defines the necessary boilerplate for latent property support.
+  * It takes as arguments the shared data struct name that owns the latent properties
+  * and an enumerator macro. The enumerator macro should accept one argument which will
+  * be a macro to apply to each latent property. We then use it to define what we need.
+  */
+#define GENERATE_DECORATOR_LATENT_PROPERTIES(SharedDataType, EnumeratorMacro) \
+	using Self = SharedDataType; \
+	ANIM_NEXT_IMPL_DEFINE_GET_LATENT_PROPERTY_INDEX(EnumeratorMacro) \
+	EnumeratorMacro(ANIM_NEXT_IMPL_DEFINE_LATENT_GETTER) \
+
 /**
  * FAnimNextDecoratorSharedData
  * Decorator shared data represents a unique instance in the authored static graph.
@@ -21,4 +84,12 @@ USTRUCT()
 struct FAnimNextDecoratorSharedData
 {
 	GENERATED_BODY()
+
+	/**
+	  * Returns the latent property index from a latent property offset
+	  * If the property is found within this type, its latent property 1-based index is returned (a positive value)
+	  * If the property isn't found, this function returns the number of latent properties (as a zero or negative value)
+	  * Derived types will have latent property indices higher than their base type
+	  */
+	static constexpr int32 GetLatentPropertyIndex(const size_t LatentPropertyOffset) { return 0; }
 };
diff --git a/Engine/Plugins/Experimental/AnimNext/Source/AnimNext/Public/DecoratorBase/DecoratorUID.h b/Engine/Plugins/Experimental/AnimNext/Source/AnimNext/Public/DecoratorBase/DecoratorUID.h
--- a/Engine/Plugins/Experimental/AnimNext/Source/AnimNext/Public/DecoratorBase/DecoratorUID.h
+++ b/Engine/Plugins/Experimental/AnimNext/Source/AnimNext/Public/DecoratorBase/DecoratorUID.h
@@ -21,7 +21,7 @@ namespace UE::AnimNext
 	struct FDecoratorUID final
 	{
 		// Constructs an invalid UID
-		constexpr FDecoratorUID()
+		constexpr FDecoratorUID() noexcept
 			: UID(INVALID_UID)
 #if !UE_BUILD_SHIPPING
 			, DecoratorName(TEXT("<Invalid decorator UID>"))
@@ -30,7 +30,7 @@ namespace UE::AnimNext
 		}
 
 		// Constructs a decorator UID
-		explicit constexpr FDecoratorUID(FDecoratorUIDRaw InUID, const TCHAR* InDecoratorName = TEXT("<Unknown Decorator Name>"))
+		explicit constexpr FDecoratorUID(FDecoratorUIDRaw InUID, const TCHAR* InDecoratorName = TEXT("<Unknown Decorator Name>")) noexcept
 			: UID(InUID)
 #if !UE_BUILD_SHIPPING
 			, DecoratorName(InDecoratorName)
@@ -40,14 +40,14 @@ namespace UE::AnimNext
 
 #if !UE_BUILD_SHIPPING
 		// Returns a literal string to the interface name
-		constexpr const TCHAR* GetDecoratorName() const { return DecoratorName; }
+		constexpr const TCHAR* GetDecoratorName() const noexcept { return DecoratorName; }
 #endif
 
 		// Returns the decorator global UID
-		constexpr FDecoratorUIDRaw GetUID() const { return UID; }
+		constexpr FDecoratorUIDRaw GetUID() const noexcept { return UID; }
 
 		// Returns whether this UID is valid or not
-		constexpr bool IsValid() const { return UID != INVALID_UID; }
+		constexpr bool IsValid() const noexcept { return UID != INVALID_UID; }
 
 	private:
 		static constexpr FDecoratorUIDRaw INVALID_UID = 0;
@@ -60,10 +60,10 @@ namespace UE::AnimNext
 	};
 
 	// Compares for equality and inequality
-	constexpr bool operator==(FDecoratorUID LHS, FDecoratorUID RHS) { return LHS.GetUID() == RHS.GetUID(); }
-	constexpr bool operator!=(FDecoratorUID LHS, FDecoratorUID RHS) { return LHS.GetUID() != RHS.GetUID(); }
-	constexpr bool operator==(FDecoratorUID LHS, FDecoratorUIDRaw RHS) { return LHS.GetUID() == RHS; }
-	constexpr bool operator!=(FDecoratorUID LHS, FDecoratorUIDRaw RHS) { return LHS.GetUID() != RHS; }
-	constexpr bool operator==(FDecoratorUIDRaw LHS, FDecoratorUID RHS) { return LHS == RHS.GetUID(); }
-	constexpr bool operator!=(FDecoratorUIDRaw LHS, FDecoratorUID RHS) { return LHS != RHS.GetUID(); }
+	constexpr bool operator==(FDecoratorUID LHS, FDecoratorUID RHS) noexcept { return LHS.GetUID() == RHS.GetUID(); }
+	constexpr bool operator!=(FDecoratorUID LHS, FDecoratorUID RHS) noexcept { return LHS.GetUID() != RHS.GetUID(); }
+	constexpr bool operator==(FDecoratorUID LHS, FDecoratorUIDRaw RHS) noexcept { return LHS.GetUID() == RHS; }
+	constexpr bool operator!=(FDecoratorUID LHS, FDecoratorUIDRaw RHS) noexcept { return LHS.GetUID() != RHS; }
+	constexpr bool operator==(FDecoratorUIDRaw LHS, FDecoratorUID RHS) noexcept { return LHS == RHS.GetUID(); }
+	constexpr bool operator!=(FDecoratorUIDRaw LHS, FDecoratorUID RHS) noexcept { return LHS != RHS.GetUID(); }
 }
diff --git a/Engine/Plugins/Experimental/AnimNext/Source/AnimNext/Public/DecoratorBase/ExecutionContext.h b/Engine/Plugins/Experimental/AnimNext/Source/AnimNext/Public/DecoratorBase/ExecutionContext.h
--- a/Engine/Plugins/Experimental/AnimNext/Source/AnimNext/Public/DecoratorBase/ExecutionContext.h
+++ b/Engine/Plugins/Experimental/AnimNext/Source/AnimNext/Public/DecoratorBase/ExecutionContext.h
@@ -7,8 +7,14 @@
 #include "DecoratorBase/DecoratorHandle.h"
 #include "DecoratorBase/DecoratorPtr.h"
 #include "DecoratorBase/DecoratorInterfaceUID.h"
+#include "DecoratorBase/LatentPropertyHandle.h"
 #include "DecoratorBase/NodeHandle.h"
 
+struct FRigVMExtendedExecuteContext;
+struct FRigVMMemoryHandle;
+
+using FRigVMMemoryHandleArray = TArrayView<FRigVMMemoryHandle>;
+
 namespace UE::AnimNext
 {
 	struct FNodeDescription;
@@ -26,7 +32,10 @@ namespace UE::AnimNext
 	struct ANIMNEXT_API FExecutionContext
 	{
 		// Creates an execution context for the specified graph
-		explicit FExecutionContext(TArrayView<const uint8> GraphSharedData);
+		explicit FExecutionContext(TArrayView<const uint8> InGraphSharedData);
+
+		// Creates an execution context for the specified graph with RigVM latent pin support
+		FExecutionContext(TArrayView<const uint8> InGraphSharedData, FRigVMExtendedExecuteContext& InRigVMExecuteContext, FRigVMMemoryHandleArray InRigVMLatentMemoryHandles);
 
 		// Destroys the execution context
 		~FExecutionContext();
@@ -66,12 +75,17 @@ namespace UE::AnimNext
 		template<class TraversalContextType>
 		TraversalContextType& GetTraversalContext() const { return *static_cast<TraversalContextType*>(TraversalContext); }	// TODO: Add a casting check for safety
 
+		// Evaluates the latent pin with the specified handle
+		template<typename LatentPinType>
+		LatentPinType EvaluateLatentPin(FLatentPropertyHandle LatentPropertyHandle) const;
+
 	private:
 		FExecutionContext(const FExecutionContext&) = delete;
 		FExecutionContext(FExecutionContext&&) = delete;
 
 		bool GetInterfaceImpl(FDecoratorInterfaceUID InterfaceUID, FWeakDecoratorPtr DecoratorPtr, FDecoratorBinding& InterfaceBinding) const;
 		bool GetInterfaceSuperImpl(FDecoratorInterfaceUID InterfaceUID, FWeakDecoratorPtr DecoratorPtr, FDecoratorBinding& SuperBinding) const;
+		const void* EvaluateLatentPinImpl(FLatentPropertyHandle LatentPropertyHandle) const;
 
 		const FNodeDescription& GetNodeDescription(FNodeHandle NodeHandle) const;
 		const FNodeTemplate* GetNodeTemplate(const FNodeDescription& NodeDesc) const;
@@ -80,6 +94,8 @@ namespace UE::AnimNext
 		ITraversalContext* TraversalContext = nullptr;
 
 		TArrayView<const uint8> GraphSharedData;
+		FRigVMMemoryHandleArray RigVMLatentMemoryHandles;
+		FRigVMExtendedExecuteContext* RigVMExecuteContext;
 
 		friend struct FScopedTraversalContext;
 	};
@@ -90,27 +106,27 @@ namespace UE::AnimNext
 	//////////////////////////////////////////////////////////////////////////
 
 	template<class DecoratorInterface>
-	bool FExecutionContext::GetInterface(FWeakDecoratorPtr DecoratorPtr, TDecoratorBinding<DecoratorInterface>& InterfaceBinding) const
+	inline bool FExecutionContext::GetInterface(FWeakDecoratorPtr DecoratorPtr, TDecoratorBinding<DecoratorInterface>& InterfaceBinding) const
 	{
 		constexpr FDecoratorInterfaceUID InterfaceUID = DecoratorInterface::InterfaceUID;
 		return GetInterfaceImpl(InterfaceUID, DecoratorPtr, InterfaceBinding);
 	}
 
 	template<class DecoratorInterface>
-	bool FExecutionContext::GetInterface(const FDecoratorBinding& Binding, TDecoratorBinding<DecoratorInterface>& InterfaceBinding) const
+	inline bool FExecutionContext::GetInterface(const FDecoratorBinding& Binding, TDecoratorBinding<DecoratorInterface>& InterfaceBinding) const
 	{
 		return GetInterface<DecoratorInterface>(Binding.GetDecoratorPtr(), InterfaceBinding);
 	}
 
 	template<class DecoratorInterface>
-	bool FExecutionContext::GetInterfaceSuper(FWeakDecoratorPtr DecoratorPtr, TDecoratorBinding<DecoratorInterface>& SuperBinding) const
+	inline bool FExecutionContext::GetInterfaceSuper(FWeakDecoratorPtr DecoratorPtr, TDecoratorBinding<DecoratorInterface>& SuperBinding) const
 	{
 		constexpr FDecoratorInterfaceUID InterfaceUID = DecoratorInterface::InterfaceUID;
 		return GetInterfaceSuperImpl(InterfaceUID, DecoratorPtr, SuperBinding);
 	}
 
 	template<class DecoratorInterface>
-	bool FExecutionContext::GetInterfaceSuper(const FDecoratorBinding& Binding, TDecoratorBinding<DecoratorInterface>& SuperBinding) const
+	inline bool FExecutionContext::GetInterfaceSuper(const FDecoratorBinding& Binding, TDecoratorBinding<DecoratorInterface>& SuperBinding) const
 	{
 		return GetInterfaceSuper<DecoratorInterface>(Binding.GetDecoratorPtr(), SuperBinding);
 	}
@@ -119,4 +135,24 @@ namespace UE::AnimNext
 	{
 		return AllocateNodeInstance(ParentBinding.GetDecoratorPtr(), ChildDecoratorHandle);
 	}
+
+	template<typename LatentPinType>
+	inline LatentPinType FExecutionContext::EvaluateLatentPin(FLatentPropertyHandle LatentPropertyHandle) const
+	{
+		// Latent pin handle needs to be valid
+		ensure(LatentPropertyHandle.IsValid());
+		if (!LatentPropertyHandle.IsValid())
+		{
+			return LatentPinType();
+		}
+
+		// We need latent pin memory handles, also implies we have a valid RigVM execution context
+		ensure(RigVMLatentMemoryHandles.IsValidIndex(LatentPropertyHandle.GetLatentPropertyIndex()));
+		if (!RigVMLatentMemoryHandles.IsValidIndex(LatentPropertyHandle.GetLatentPropertyIndex()))
+		{
+			return LatentPinType();
+		}
+
+		return *reinterpret_cast<const LatentPinType*>(EvaluateLatentPinImpl(LatentPropertyHandle));
+	}
 }
diff --git a/Engine/Plugins/Experimental/AnimNext/Source/AnimNext/Public/DecoratorBase/LatentPropertyHandle.h b/Engine/Plugins/Experimental/AnimNext/Source/AnimNext/Public/DecoratorBase/LatentPropertyHandle.h
new file mode 100644
--- /dev/null
+++ b/Engine/Plugins/Experimental/AnimNext/Source/AnimNext/Public/DecoratorBase/LatentPropertyHandle.h
@@ -0,0 +1,57 @@
+// Copyright Epic Games, Inc. All Rights Reserved.
+
+#pragma once
+
+#include "CoreMinimal.h"
+#include "Serialization/Archive.h"
+
+namespace UE::AnimNext
+{
+	/**
+	 * Latent Property Handle
+	 * Latent property handles represent a 1-based index into the RigVM memory handles array during execution.
+	 *
+	 * @see FExecutionContext
+	 */
+	struct FLatentPropertyHandle final
+	{
+		// Creates an invalid latent property handle
+		constexpr FLatentPropertyHandle() noexcept
+			: HandleValue(INVALID_HANDLE_VALUE)
+		{}
+
+		// Returns true if this latent property handle is valid, false otherwise
+		constexpr bool IsValid() const noexcept { return HandleValue != INVALID_HANDLE_VALUE; }
+
+		// Returns the latent property index represented by this handle if valid, -1 otherwise
+		constexpr int32 GetLatentPropertyIndex() const noexcept { return IsValid() ? (HandleValue - 1) : INDEX_NONE; }
+
+		// The first valid latent handle
+		static constexpr FLatentPropertyHandle GetFirstHandle() noexcept { return FLatentPropertyHandle(1); }
+
+		// Returns the next latent handle
+		// If the current handle is invalid, the next is invalid as well
+		// Note that this function does not handle wrapping if more than 65536 handles are used
+		// If wrapping occurs, the next handle returned will be invalid and all subsequent handles will be invalid
+		constexpr FLatentPropertyHandle GetNextHandle() const noexcept { return IsValid() ? FLatentPropertyHandle(HandleValue + 1) : FLatentPropertyHandle(); }
+
+	private:
+		// Constructs a FLatentPropertyHandle instance from a raw value
+		explicit constexpr FLatentPropertyHandle(uint16 InHandleValue) noexcept : HandleValue(InHandleValue) {}
+
+		// Latent property handles are a 1-based index, 0 is invalid
+		static constexpr uint16 INVALID_HANDLE_VALUE = 0;
+
+		// The raw 1-based latent property index or 0 if invalid
+		uint16		HandleValue;
+
+		friend FArchive& operator<<(FArchive& Ar, FLatentPropertyHandle& Handle);
+	};
+
+	// Serializes a latent handle value
+	inline FArchive& operator<<(FArchive& Ar, FLatentPropertyHandle& Handle)
+	{
+		Ar << Handle.HandleValue;
+		return Ar;
+	}
+}
diff --git a/Engine/Plugins/Experimental/AnimNext/Source/AnimNextEditor/Private/Graph/GraphFactory.cpp b/Engine/Plugins/Experimental/AnimNext/Source/AnimNextEditor/Private/Graph/GraphFactory.cpp
--- a/Engine/Plugins/Experimental/AnimNext/Source/AnimNextEditor/Private/Graph/GraphFactory.cpp
+++ b/Engine/Plugins/Experimental/AnimNext/Source/AnimNextEditor/Private/Graph/GraphFactory.cpp
@@ -34,7 +34,7 @@ UObject* UAnimNextGraphFactory::FactoryCreateNew(UClass* Class, UObject* InParen
 	// Add initial execution unit
 	URigVMController* Controller = EditorData->GetRigVMClient()->GetController(EditorData->RootGraph);
 
-	URigVMUnitNode* MainEntryPointNode = Controller->AddUnitNode(FRigUnit_AnimNextGraphRoot::StaticStruct(), FRigVMStruct::ExecuteName, FVector2D(-400.0f, 0.0f), FString(), false);
+	URigVMUnitNode* MainEntryPointNode = Controller->AddUnitNode(FRigUnit_AnimNextGraphRoot::StaticStruct(), FRigUnit_AnimNextGraphRoot::EventName, FVector2D(-400.0f, 0.0f), FString(), false);
 	URigVMPin* BeginExecutePin = MainEntryPointNode->FindPin(GET_MEMBER_NAME_STRING_CHECKED(FRigUnit_AnimNextGraphRoot, Result));
 	check(BeginExecutePin);
 	check(BeginExecutePin->GetDirection() == ERigVMPinDirection::Input);
diff --git a/Engine/Plugins/Experimental/AnimNext/Source/AnimNextTestSuite/Private/AnimNextDecoratorBaseTest.cpp b/Engine/Plugins/Experimental/AnimNext/Source/AnimNextTestSuite/Private/AnimNextDecoratorBaseTest.cpp
--- a/Engine/Plugins/Experimental/AnimNext/Source/AnimNextTestSuite/Private/AnimNextDecoratorBaseTest.cpp
+++ b/Engine/Plugins/Experimental/AnimNext/Source/AnimNextTestSuite/Private/AnimNextDecoratorBaseTest.cpp
@@ -535,11 +535,19 @@ bool FAnimationAnimNextRuntimeTest_NodeLifetime::RunTest(const FString& InParame
 			[](uint32 DecoratorIndex, const FString& PropertyName)
 			{
 				return FString();
+			},
+			[](uint32 DecoratorIndex, const FString& PropertyName)
+			{
+				return false;
 			});
 		DecoratorWriter.WriteNode(Node1,
 			[](uint32 DecoratorIndex, const FString& PropertyName)
 			{
 				return FString();
+			},
+			[](uint32 DecoratorIndex, const FString& PropertyName)
+			{
+				return false;
 			});
 		DecoratorWriter.EndNodeWriting();
 
@@ -737,6 +745,10 @@ bool FAnimationAnimNextRuntimeTest_GetDecoratorInterface::RunTest(const FString&
 			[](uint32 DecoratorIndex, const FString& PropertyName)
 			{
 				return FString();
+			},
+			[](uint32 DecoratorIndex, const FString& PropertyName)
+			{
+				return false;
 			});
 		DecoratorWriter.EndNodeWriting();
 
@@ -1003,6 +1015,10 @@ bool FAnimationAnimNextRuntimeTest_GetDecoratorInterfaceSuper::RunTest(const FSt
 			[](uint32 DecoratorIndex, const FString& PropertyName)
 			{
 				return FString();
+			},
+			[](uint32 DecoratorIndex, const FString& PropertyName)
+			{
+				return false;
 			});
 		DecoratorWriter.EndNodeWriting();
 
@@ -1341,11 +1357,19 @@ bool FAnimationAnimNextRuntimeTest_DecoratorSerialization::RunTest(const FString
 			[&DecoratorProperties0](uint32 DecoratorIndex, const FString& PropertyName)
 			{
 				return DecoratorProperties0[DecoratorIndex][PropertyName];
+			},
+			[](uint32 DecoratorIndex, const FString& PropertyName)
+			{
+				return false;
 			});
 		DecoratorWriter.WriteNode(Node1,
 			[&DecoratorProperties1](uint32 DecoratorIndex, const FString& PropertyName)
 			{
 				return DecoratorProperties1[DecoratorIndex][PropertyName];
+			},
+			[](uint32 DecoratorIndex, const FString& PropertyName)
+			{
+				return false;
 			});
 		DecoratorWriter.EndNodeWriting();
 
diff --git a/Engine/Plugins/Experimental/AnimNext/Source/AnimNextTestSuite/Private/AnimNextDecoratorBaseTest.h b/Engine/Plugins/Experimental/AnimNext/Source/AnimNextTestSuite/Private/AnimNextDecoratorBaseTest.h
--- a/Engine/Plugins/Experimental/AnimNext/Source/AnimNextTestSuite/Private/AnimNextDecoratorBaseTest.h
+++ b/Engine/Plugins/Experimental/AnimNext/Source/AnimNextTestSuite/Private/AnimNextDecoratorBaseTest.h
@@ -13,7 +13,7 @@ struct FDecoratorA_BaseSharedData : public FAnimNextDecoratorSharedData
 {
 	GENERATED_BODY()
 
-	UPROPERTY()
+	UPROPERTY(meta = (Inline))
 	uint32 DecoratorUID;
 
 	FDecoratorA_BaseSharedData();
@@ -24,7 +24,7 @@ struct FDecoratorAB_AddSharedData : public FAnimNextDecoratorSharedData
 {
 	GENERATED_BODY()
 
-	UPROPERTY()
+	UPROPERTY(meta = (Inline))
 	uint32 DecoratorUID;
 
 	FDecoratorAB_AddSharedData();
@@ -35,7 +35,7 @@ struct FDecoratorAC_AddSharedData : public FAnimNextDecoratorSharedData
 {
 	GENERATED_BODY()
 
-	UPROPERTY()
+	UPROPERTY(meta = (Inline))
 	uint32 DecoratorUID;
 
 	FDecoratorAC_AddSharedData();
@@ -46,28 +46,28 @@ struct FDecoratorSerialization_BaseSharedData : public FAnimNextDecoratorSharedD
 {
 	GENERATED_BODY()
 
-	UPROPERTY()
+	UPROPERTY(meta = (Inline))
 	int32 Integer = 0;
 
-	UPROPERTY()
+	UPROPERTY(meta = (Inline))
 	int32 IntegerArray[4] = { 0, 0, 0, 0 };
 
-	UPROPERTY()
+	UPROPERTY(meta = (Inline))
 	TArray<int32> IntegerTArray;
 
-	UPROPERTY()
+	UPROPERTY(meta = (Inline))
 	FVector Vector = FVector::ZeroVector;
 
-	UPROPERTY()
+	UPROPERTY(meta = (Inline))
 	FVector VectorArray[2] = { FVector::ZeroVector, FVector::ZeroVector };
 
-	UPROPERTY()
+	UPROPERTY(meta = (Inline))
 	TArray<FVector> VectorTArray;
 
-	UPROPERTY()
+	UPROPERTY(meta = (Inline))
 	FString String;
 
-	UPROPERTY()
+	UPROPERTY(meta = (Inline))
 	FName Name;
 };
 
@@ -76,28 +76,28 @@ struct FDecoratorSerialization_AddSharedData : public FAnimNextDecoratorSharedDa
 {
 	GENERATED_BODY()
 
-	UPROPERTY()
+	UPROPERTY(meta = (Inline))
 	int32 Integer = 0;
 
-	UPROPERTY()
+	UPROPERTY(meta = (Inline))
 	int32 IntegerArray[4] = { 0, 0, 0, 0 };
 
-	UPROPERTY()
+	UPROPERTY(meta = (Inline))
 	TArray<int32> IntegerTArray;
 
-	UPROPERTY()
+	UPROPERTY(meta = (Inline))
 	FVector Vector = FVector::ZeroVector;
 
-	UPROPERTY()
+	UPROPERTY(meta = (Inline))
 	FVector VectorArray[2] = { FVector::ZeroVector, FVector::ZeroVector };
 
-	UPROPERTY()
+	UPROPERTY(meta = (Inline))
 	TArray<FVector> VectorTArray;
 
-	UPROPERTY()
+	UPROPERTY(meta = (Inline))
 	FString String;
 
-	UPROPERTY()
+	UPROPERTY(meta = (Inline))
 	FName Name;
 };
 
@@ -106,28 +106,28 @@ struct FDecoratorNativeSerialization_AddSharedData : public FAnimNextDecoratorSh
 {
 	GENERATED_BODY()
 
-	UPROPERTY()
+	UPROPERTY(meta = (Inline))
 	int32 Integer = 0;
 
-	UPROPERTY()
+	UPROPERTY(meta = (Inline))
 	int32 IntegerArray[4] = { 0, 0, 0, 0 };
 
-	UPROPERTY()
+	UPROPERTY(meta = (Inline))
 	TArray<int32> IntegerTArray;
 
-	UPROPERTY()
+	UPROPERTY(meta = (Inline))
 	FVector Vector = FVector::ZeroVector;
 
-	UPROPERTY()
+	UPROPERTY(meta = (Inline))
 	FVector VectorArray[2] = { FVector::ZeroVector, FVector::ZeroVector };
 
-	UPROPERTY()
+	UPROPERTY(meta = (Inline))
 	TArray<FVector> VectorTArray;
 
-	UPROPERTY()
+	UPROPERTY(meta = (Inline))
 	FString String;
 
-	UPROPERTY()
+	UPROPERTY(meta = (Inline))
 	FName Name;
 
 	bool bSerializeCalled = false;
diff --git a/Engine/Plugins/Experimental/AnimNext/Source/AnimNextTestSuite/Private/AnimNextDecoratorGraphTest.cpp b/Engine/Plugins/Experimental/AnimNext/Source/AnimNextTestSuite/Private/AnimNextDecoratorGraphTest.cpp
--- a/Engine/Plugins/Experimental/AnimNext/Source/AnimNextTestSuite/Private/AnimNextDecoratorGraphTest.cpp
+++ b/Engine/Plugins/Experimental/AnimNext/Source/AnimNextTestSuite/Private/AnimNextDecoratorGraphTest.cpp
@@ -16,6 +16,9 @@
 #include "Graph/RigUnit_AnimNextDecoratorStack.h"
 #include "Misc/AutomationTest.h"
 #include "Param/ParamStack.h"
+#include "Param/RigVMDispatch_GetParameter.h"
+#include "RigVMFunctions/Math/RigVMFunction_MathInt.h"
+#include "RigVMFunctions/Math/RigVMFunction_MathFloat.h"
 
 #if WITH_DEV_AUTOMATION_TESTS
 
@@ -49,6 +52,10 @@ namespace UE::AnimNext
 
 			ParamStack.GetMutableParam<int32>("SomeInt32") = SharedData->SomeInt32;
 			ParamStack.GetMutableParam<float>("SomeFloat") = SharedData->SomeFloat;
+
+			ParamStack.GetMutableParam<int32>("SomeLatentInt32") = SharedData->GetSomeLatentInt32(Context, Binding);				// MathAdd with constants, latent
+			ParamStack.GetMutableParam<int32>("SomeOtherLatentInt32") = SharedData->GetSomeOtherLatentInt32(Context, Binding);		// GetParameter, latent
+			ParamStack.GetMutableParam<float>("SomeLatentFloat") = SharedData->GetSomeLatentFloat(Context, Binding);				// Inline value, not latent
 		}
 
 		// IEvaluate impl
@@ -124,18 +131,42 @@ bool FAnimationAnimNextRuntimeTest_GraphAddDecorator::RunTest(const FString& InP
 		DefaultValue, INDEX_NONE, true, true);
 	UE_RETURN_ON_ERROR(DecoratorName == DisplayName, TEXT("FAnimationAnimNextRuntimeTest_GraphAddDecorator -> Unexpected decorator name"));
 
+	URigVMPin* DecoratorPin = DecoratorStackNode->FindPin(*DisplayName);
+	UE_RETURN_ON_ERROR(DecoratorPin != nullptr, TEXT("FAnimationAnimNextRuntimeTest_GraphAddDecorator -> Failed to find decorator pin"));
+
 	// Our first pin is the hard coded output result, decorator pins follow
 	UE_RETURN_ON_ERROR(DecoratorStackNode->GetPins().Num() == 2, TEXT("FAnimationAnimNextRuntimeTest_GraphAddDecorator -> Unexpected number of pins"));
-	UE_RETURN_ON_ERROR(DecoratorStackNode->GetPins()[1]->IsDecoratorPin() == true, TEXT("FAnimationAnimNextRuntimeTest_GraphAddDecorator -> Unexpected pin type"));
-	UE_RETURN_ON_ERROR(DecoratorStackNode->GetPins()[1]->GetFName() == DecoratorName, TEXT("FAnimationAnimNextRuntimeTest_GraphAddDecorator -> Unexpected pin name"));
-	UE_RETURN_ON_ERROR(DecoratorStackNode->GetPins()[1]->GetCPPTypeObject() == FRigDecorator_AnimNextCppDecorator::StaticStruct(), TEXT("FAnimationAnimNextRuntimeTest_GraphAddDecorator -> Unexpected pin type"));
+	UE_RETURN_ON_ERROR(DecoratorPin->IsDecoratorPin() == true, TEXT("FAnimationAnimNextRuntimeTest_GraphAddDecorator -> Unexpected pin type"));
+	UE_RETURN_ON_ERROR(DecoratorPin->GetFName() == DecoratorName, TEXT("FAnimationAnimNextRuntimeTest_GraphAddDecorator -> Unexpected pin name"));
+	UE_RETURN_ON_ERROR(DecoratorPin->GetCPPTypeObject() == FRigDecorator_AnimNextCppDecorator::StaticStruct(), TEXT("FAnimationAnimNextRuntimeTest_GraphAddDecorator -> Unexpected pin type"));
 
 	// Our first sub-pin is the hard coded script struct member that parametrizes the decorator, dynamic decorator sub-pins follow
-	UE_RETURN_ON_ERROR(DecoratorStackNode->GetPins()[1]->GetSubPins().Num() == 3, TEXT("FAnimationAnimNextRuntimeTest_GraphAddDecorator -> Unexpected decorator sub pins"));
-	UE_RETURN_ON_ERROR(DecoratorStackNode->GetPins()[1]->GetSubPins()[1]->GetCPPType() == TEXT("int32"), TEXT("FAnimationAnimNextRuntimeTest_GraphAddDecorator -> Unexpected decorator pin type"));
-	UE_RETURN_ON_ERROR(DecoratorStackNode->GetPins()[1]->GetSubPins()[1]->GetDefaultValue() == TEXT("3"), TEXT("FAnimationAnimNextRuntimeTest_GraphAddDecorator -> Unexpected decorator pin value"));
-	UE_RETURN_ON_ERROR(DecoratorStackNode->GetPins()[1]->GetSubPins()[2]->GetCPPType() == TEXT("float"), TEXT("FAnimationAnimNextRuntimeTest_GraphAddDecorator -> Unexpected decorator pin type"));
-	UE_RETURN_ON_ERROR(DecoratorStackNode->GetPins()[1]->GetSubPins()[2]->GetDefaultValue() == TEXT("34.000000"), TEXT("FAnimationAnimNextRuntimeTest_GraphAddDecorator -> Unexpected decorator pin value"));
+	UE_RETURN_ON_ERROR(DecoratorPin->GetSubPins().Num() == 6, TEXT("FAnimationAnimNextRuntimeTest_GraphAddDecorator -> Unexpected decorator sub pins"));
+
+	// SomeInt32
+	UE_RETURN_ON_ERROR(DecoratorPin->GetSubPins()[1]->GetCPPType() == TEXT("int32"), TEXT("FAnimationAnimNextRuntimeTest_GraphAddDecorator -> Unexpected decorator pin type"));
+	UE_RETURN_ON_ERROR(DecoratorPin->GetSubPins()[1]->GetDefaultValue() == TEXT("3"), TEXT("FAnimationAnimNextRuntimeTest_GraphAddDecorator -> Unexpected decorator pin value"));
+	UE_RETURN_ON_ERROR(!DecoratorPin->GetSubPins()[1]->IsLazy(), TEXT("FAnimationAnimNextRuntimeTest_GraphAddDecorator -> Expected non-lazy decorator pin"));
+
+	// SomeFloat
+	UE_RETURN_ON_ERROR(DecoratorPin->GetSubPins()[2]->GetCPPType() == TEXT("float"), TEXT("FAnimationAnimNextRuntimeTest_GraphAddDecorator -> Unexpected decorator pin type"));
+	UE_RETURN_ON_ERROR(DecoratorPin->GetSubPins()[2]->GetDefaultValue() == TEXT("34.000000"), TEXT("FAnimationAnimNextRuntimeTest_GraphAddDecorator -> Unexpected decorator pin value"));
+	UE_RETURN_ON_ERROR(!DecoratorPin->GetSubPins()[2]->IsLazy(), TEXT("FAnimationAnimNextRuntimeTest_GraphAddDecorator -> Expected non-lazy decorator pin"));
+
+	// SomeLatentInt32
+	UE_RETURN_ON_ERROR(DecoratorPin->GetSubPins()[3]->GetCPPType() == TEXT("int32"), TEXT("FAnimationAnimNextRuntimeTest_GraphAddDecorator -> Unexpected decorator pin type"));
+	UE_RETURN_ON_ERROR(DecoratorPin->GetSubPins()[3]->GetDefaultValue() == TEXT("3"), TEXT("FAnimationAnimNextRuntimeTest_GraphAddDecorator -> Unexpected decorator pin value"));
+	UE_RETURN_ON_ERROR(DecoratorPin->GetSubPins()[3]->IsLazy(), TEXT("FAnimationAnimNextRuntimeTest_GraphAddDecorator -> Expected lazy decorator pin"));
+
+	// SomeOtherLatentInt32
+	UE_RETURN_ON_ERROR(DecoratorPin->GetSubPins()[4]->GetCPPType() == TEXT("int32"), TEXT("FAnimationAnimNextRuntimeTest_GraphAddDecorator -> Unexpected decorator pin type"));
+	UE_RETURN_ON_ERROR(DecoratorPin->GetSubPins()[4]->GetDefaultValue() == TEXT("3"), TEXT("FAnimationAnimNextRuntimeTest_GraphAddDecorator -> Unexpected decorator pin value"));
+	UE_RETURN_ON_ERROR(DecoratorPin->GetSubPins()[4]->IsLazy(), TEXT("FAnimationAnimNextRuntimeTest_GraphAddDecorator -> Expected lazy decorator pin"));
+
+	// SomeLatentFloat
+	UE_RETURN_ON_ERROR(DecoratorPin->GetSubPins()[5]->GetCPPType() == TEXT("float"), TEXT("FAnimationAnimNextRuntimeTest_GraphAddDecorator -> Unexpected decorator pin type"));
+	UE_RETURN_ON_ERROR(DecoratorPin->GetSubPins()[5]->GetDefaultValue() == TEXT("34.000000"), TEXT("FAnimationAnimNextRuntimeTest_GraphAddDecorator -> Unexpected decorator pin value"));
+	UE_RETURN_ON_ERROR(DecoratorPin->GetSubPins()[5]->IsLazy(), TEXT("FAnimationAnimNextRuntimeTest_GraphAddDecorator -> Expected lazy decorator pin"));
 
 	return true;
 }
@@ -156,7 +187,7 @@ bool FAnimationAnimNextRuntimeTest_GraphExecute::RunTest(const FString& InParame
 	UE_RETURN_ON_ERROR(Controller != nullptr, "FAnimationAnimNextRuntimeTest_GraphExecute -> Failed to get RigVM controller");
 
 	// Add graph entry point
-	URigVMUnitNode* MainEntryPointNode = Controller->AddUnitNode(FRigUnit_AnimNextGraphRoot::StaticStruct(), FRigVMStruct::ExecuteName, FVector2D(0.0f, 0.0f), FString(), false);
+	URigVMUnitNode* MainEntryPointNode = Controller->AddUnitNode(FRigUnit_AnimNextGraphRoot::StaticStruct(), FRigUnit_AnimNextGraphRoot::EventName, FVector2D(0.0f, 0.0f), FString(), false);
 	UE_RETURN_ON_ERROR(MainEntryPointNode != nullptr, "FAnimationAnimNextRuntimeTest_GraphExecute -> Failed to find main entry point node");
 
 	URigVMPin* BeginExecutePin = MainEntryPointNode->FindPin(GET_MEMBER_NAME_STRING_CHECKED(FRigUnit_AnimNextGraphRoot, Result));
@@ -204,11 +235,15 @@ bool FAnimationAnimNextRuntimeTest_GraphExecute::RunTest(const FString& InParame
 		DefaultValue, INDEX_NONE, true, true);
 	UE_RETURN_ON_ERROR(DecoratorName == DisplayName, TEXT("FAnimationAnimNextRuntimeTest_GraphExecute -> Unexpected decorator name"));
 
+	URigVMPin* DecoratorPin = DecoratorStackNode->FindPin(*DisplayName);
+	UE_RETURN_ON_ERROR(DecoratorPin != nullptr, TEXT("FAnimationAnimNextRuntimeTest_GraphExecute -> Failed to find decorator pin"));
+
 	// Set some values on our decorator
-	Controller->SetPinDefaultValue(DecoratorStackNode->GetPins()[1]->GetSubPins()[1]->GetPinPath(), TEXT("78"));
-	Controller->SetPinDefaultValue(DecoratorStackNode->GetPins()[1]->GetSubPins()[2]->GetPinPath(), TEXT("142.33"));
+	Controller->SetPinDefaultValue(DecoratorPin->GetSubPins()[1]->GetPinPath(), TEXT("78"));
+	Controller->SetPinDefaultValue(DecoratorPin->GetSubPins()[2]->GetPinPath(), TEXT("142.33"));
 
-	UE::AnimNext::FDecoratorPtr GraphInstance = AnimNextGraph->AllocateInstance();
+	FAnimNextGraphInstance GraphInstance;
+	AnimNextGraph->AllocateInstance(GraphInstance);
 
 	UE::AnimNext::FContext Context(1.0f / 30.0f);
 
@@ -216,7 +251,10 @@ bool FAnimationAnimNextRuntimeTest_GraphExecute::RunTest(const FString& InParame
 		"UpdateCount", (int32)0,
 		"EvaluateCount", (int32)0,
 		"SomeInt32", (int32)0,
-		"SomeFloat", 0.0f
+		"SomeFloat", 0.0f,
+		"SomeLatentInt32", (int32)0,
+		"SomeOtherLatentInt32", (int32)0,
+		"SomeLatentFloat", 0.0f
 	);
 
 	AnimNextGraph->Run(Context, GraphInstance, EAnimNextGraphSimulationSteps::All);
@@ -225,9 +263,145 @@ bool FAnimationAnimNextRuntimeTest_GraphExecute::RunTest(const FString& InParame
 	AddErrorIfFalse(Context.GetParamStack().GetParam<int32>("EvaluateCount") == 1, "FAnimationAnimNextRuntimeTest_GraphExecute -> Unexpected evaluate count");
 	AddErrorIfFalse(Context.GetParamStack().GetParam<int32>("SomeInt32") == 78, "FAnimationAnimNextRuntimeTest_GraphExecute -> Unexpected SomeInt32 value");
 	AddErrorIfFalse(Context.GetParamStack().GetParam<float>("SomeFloat") == 142.33f, "FAnimationAnimNextRuntimeTest_GraphExecute -> Unexpected SomeFloat value");
+	AddErrorIfFalse(Context.GetParamStack().GetParam<int32>("SomeLatentInt32") == 3, "FAnimationAnimNextRuntimeTest_GraphExecute -> Unexpected SomeLatentInt32 value");
+	AddErrorIfFalse(Context.GetParamStack().GetParam<int32>("SomeOtherLatentInt32") == 3, "FAnimationAnimNextRuntimeTest_GraphExecute -> Unexpected SomeOtherLatentInt32 value");
+	AddErrorIfFalse(Context.GetParamStack().GetParam<float>("SomeLatentFloat") == 34.0f, "FAnimationAnimNextRuntimeTest_GraphExecute -> Unexpected SomeLatentFloat value");
+
+	Context.GetMutableParamStack().PopLayer(LayerHandle);
+	GraphInstance.Release();
+
+	return true;
+}
+
+IMPLEMENT_SIMPLE_AUTOMATION_TEST(FAnimationAnimNextRuntimeTest_GraphExecuteLatent, "Animation.AnimNext.Runtime.Graph.ExecuteLatent", EAutomationTestFlags::EditorContext | EAutomationTestFlags::EngineFilter)
+
+bool FAnimationAnimNextRuntimeTest_GraphExecuteLatent::RunTest(const FString& InParameters)
+{
+	UAnimNextGraphTest* AnimNextGraph = NewObject<UAnimNextGraphTest>();
+
+	UAnimNextGraph_EditorData* EditorData = NewObject<UAnimNextGraph_EditorData>(AnimNextGraph, TEXT("EditorData"));
+	AnimNextGraph->SetEditorData(EditorData);
+
+	EditorData->Initialize(/*bRecompileVM*/false);
+	EditorData->GetRigVMClient()->SetExecuteContextStruct(FAnimNextExecuteContext::StaticStruct());
+
+	URigVMController* Controller = EditorData->GetRigVMClient()->GetController(EditorData->GetRootGraph());
+	UE_RETURN_ON_ERROR(Controller != nullptr, "FAnimationAnimNextRuntimeTest_GraphExecuteLatent -> Failed to get RigVM controller");
+
+	// Add graph entry point
+	URigVMUnitNode* MainEntryPointNode = Controller->AddUnitNode(FRigUnit_AnimNextGraphRoot::StaticStruct(), FRigUnit_AnimNextGraphRoot::EventName, FVector2D(0.0f, 0.0f), FString(), false);
+	UE_RETURN_ON_ERROR(MainEntryPointNode != nullptr, "FAnimationAnimNextRuntimeTest_GraphExecuteLatent -> Failed to find main entry point node");
+
+	URigVMPin* BeginExecutePin = MainEntryPointNode->FindPin(GET_MEMBER_NAME_STRING_CHECKED(FRigUnit_AnimNextGraphRoot, Result));
+	UE_RETURN_ON_ERROR(BeginExecutePin != nullptr && BeginExecutePin->GetDirection() == ERigVMPinDirection::Input, "FAnimationAnimNextRuntimeTest_GraphExecuteLatent -> Failed to create entry point");
+
+	// Create an empty decorator stack node
+	URigVMUnitNode* DecoratorStackNode = Controller->AddUnitNode(FRigUnit_AnimNextDecoratorStack::StaticStruct(), FRigVMStruct::ExecuteName, FVector2D(0.0f, 0.0f), FString(), false);
+	UE_RETURN_ON_ERROR(DecoratorStackNode != nullptr, "FAnimationAnimNextRuntimeTest_GraphExecuteLatent -> Failed to create decorator stack node");
+
+	// Link our stack result to our entry point
+	Controller->AddLink(DecoratorStackNode->GetPins()[0], MainEntryPointNode->GetPins()[0]);
+
+	// Add a decorator
+	const UScriptStruct* CppDecoratorStruct = FRigDecorator_AnimNextCppDecorator::StaticStruct();
+	UE_RETURN_ON_ERROR(CppDecoratorStruct != nullptr, "FAnimationAnimNextRuntimeTest_GraphExecuteLatent -> Failed to get find Cpp decorator static struct");
+
+	const UE::AnimNext::FDecorator* Decorator = UE::AnimNext::FDecoratorRegistry::Get().Find(UE::AnimNext::FTestDecorator::DecoratorUID);
+	UE_RETURN_ON_ERROR(Decorator != nullptr, "FAnimationAnimNextRuntimeTest_GraphExecuteLatent -> Failed to find test decorator");
+
+	UScriptStruct* ScriptStruct = Decorator->GetDecoratorSharedDataStruct();
+	UE_RETURN_ON_ERROR(ScriptStruct != nullptr, "FAnimationAnimNextRuntimeTest_GraphExecuteLatent -> Failed to find decorator shared data struct");
+
+	FString DefaultValue;
+	{
+		const FRigDecorator_AnimNextCppDecorator DefaultCppDecoratorStructInstance;
+		FRigDecorator_AnimNextCppDecorator CppDecoratorStructInstance;
+		CppDecoratorStructInstance.DecoratorSharedDataStruct = ScriptStruct;
+
+		UE_RETURN_ON_ERROR(CppDecoratorStructInstance.CanBeAddedToNode(DecoratorStackNode, nullptr), "FAnimationAnimNextRuntimeTest_GraphExecuteLatent -> Decorator cannot be added to decorator stack node");
+
+		const FProperty* Prop = FAnimNextCppDecoratorWrapper::StaticStruct()->FindPropertyByName(GET_MEMBER_NAME_STRING_CHECKED(FAnimNextCppDecoratorWrapper, CppDecorator));
+		UE_RETURN_ON_ERROR(Prop != nullptr, "FAnimationAnimNextRuntimeTest_GraphExecuteLatent -> Failed to find wrapper property");
+
+		Prop->ExportText_Direct(DefaultValue, &CppDecoratorStructInstance, &DefaultCppDecoratorStructInstance, nullptr, PPF_None);
+	}
+
+	FString DisplayNameMetadata;
+	ScriptStruct->GetStringMetaDataHierarchical(FRigVMStruct::DisplayNameMetaName, &DisplayNameMetadata);
+	const FString DisplayName = DisplayNameMetadata.IsEmpty() ? Decorator->GetDecoratorName() : DisplayNameMetadata;
+
+	const FName DecoratorName = Controller->AddDecorator(
+		DecoratorStackNode->GetFName(),
+		*CppDecoratorStruct->GetPathName(),
+		*DisplayName,
+		DefaultValue, INDEX_NONE, true, true);
+	UE_RETURN_ON_ERROR(DecoratorName == DisplayName, TEXT("FAnimationAnimNextRuntimeTest_GraphExecuteLatent -> Unexpected decorator name"));
+
+	// Set some values on our decorator
+	URigVMPin* DecoratorPin = DecoratorStackNode->FindPin(*DisplayName);
+	UE_RETURN_ON_ERROR(DecoratorPin != nullptr, TEXT("FAnimationAnimNextRuntimeTest_GraphExecuteLatent -> Failed to find decorator pin"));
+
+	Controller->SetPinDefaultValue(DecoratorPin->GetSubPins()[1]->GetPinPath(), TEXT("78"));		// SomeInt32
+	Controller->SetPinDefaultValue(DecoratorPin->GetSubPins()[2]->GetPinPath(), TEXT("142.33"));	// SomeFloat
+	Controller->SetPinDefaultValue(DecoratorPin->GetSubPins()[5]->GetPinPath(), TEXT("1123.31"));	// SomeLatentFloat, inline value on latent pin
+
+	// Set some latent values on our decorator
+	{
+		FRigVMFunction_MathIntAdd IntAdd;
+		IntAdd.A = 10;
+		IntAdd.B = 23;
+
+		URigVMUnitNode* IntAddNode = Controller->AddUnitNodeWithDefaults(FRigVMFunction_MathIntAdd::StaticStruct(), FRigStructScope(IntAdd), FRigVMStruct::ExecuteName, FVector2D::ZeroVector, FString(), false);
+		UE_RETURN_ON_ERROR(IntAddNode != nullptr, TEXT("FAnimationAnimNextRuntimeTest_GraphExecuteLatent -> Failed to create Int add node"));
+
+		Controller->AddLink(
+			IntAddNode->FindPin(GET_MEMBER_NAME_STRING_CHECKED(FRigVMFunction_MathIntAdd, Result)),
+			DecoratorPin->GetSubPins()[3]);	// SomeLatentInt32
+	}
+
+	{
+		const FRigVMDispatchFactory* GetParameterFactory = FRigVMRegistry::Get().RegisterFactory(FRigVMDispatch_GetParameter::StaticStruct());
+		const FName GetParameterNotation = GetParameterFactory->GetTemplate()->GetNotation();
+
+		URigVMNode* GetParameterNode = Controller->AddTemplateNode(GetParameterNotation);
+		UE_RETURN_ON_ERROR(GetParameterNode != nullptr, TEXT("FAnimationAnimNextRuntimeTest_GraphExecuteLatent -> Failed to create GetParameter node"));
+
+		UE_RETURN_ON_ERROR(Controller->ResolveWildCardPin(GetParameterNode->FindPin(TEXT("Value")), RigVMTypeUtils::TypeIndex::Int32), TEXT("FAnimationAnimNextRuntimeTest_GraphExecuteLatent -> Failed to resolve wildcard for GetParameter node"));
+		Controller->SetPinDefaultValue(GetParameterNode->FindPin(TEXT("Parameter"))->GetPinPath(), TEXT("SomeSourceInt"));
+
+		Controller->AddLink(
+			GetParameterNode->FindPin(TEXT("Value")),
+			DecoratorPin->GetSubPins()[4]);	// SomeOtherLatentInt32
+	}
+
+	FAnimNextGraphInstance GraphInstance;
+	AnimNextGraph->AllocateInstance(GraphInstance);
+
+	UE::AnimNext::FContext Context(1.0f / 30.0f);
+
+	UE::AnimNext::FParamStack::FPushedLayerHandle LayerHandle = Context.GetMutableParamStack().PushValues(
+		"UpdateCount", (int32)0,
+		"EvaluateCount", (int32)0,
+		"SomeSourceInt", (int32)1223,
+		"SomeInt32", (int32)0,
+		"SomeFloat", 0.0f,
+		"SomeLatentInt32", (int32)0,
+		"SomeOtherLatentInt32", (int32)0,
+		"SomeLatentFloat", 0.0f
+	);
+
+	AnimNextGraph->Run(Context, GraphInstance, EAnimNextGraphSimulationSteps::All);
+
+	AddErrorIfFalse(Context.GetParamStack().GetParam<int32>("UpdateCount") == 1, "FAnimationAnimNextRuntimeTest_GraphExecuteLatent -> Unexpected update count");
+	AddErrorIfFalse(Context.GetParamStack().GetParam<int32>("EvaluateCount") == 1, "FAnimationAnimNextRuntimeTest_GraphExecuteLatent -> Unexpected evaluate count");
+	AddErrorIfFalse(Context.GetParamStack().GetParam<int32>("SomeInt32") == 78, "FAnimationAnimNextRuntimeTest_GraphExecuteLatent -> Unexpected SomeInt32 value");
+	AddErrorIfFalse(Context.GetParamStack().GetParam<float>("SomeFloat") == 142.33f, "FAnimationAnimNextRuntimeTest_GraphExecuteLatent -> Unexpected SomeFloat value");
+	AddErrorIfFalse(Context.GetParamStack().GetParam<int32>("SomeLatentInt32") == 33, "FAnimationAnimNextRuntimeTest_GraphExecuteLatent -> Unexpected SomeLatentInt32 value");
+	AddErrorIfFalse(Context.GetParamStack().GetParam<int32>("SomeOtherLatentInt32") == 1223, "FAnimationAnimNextRuntimeTest_GraphExecuteLatent -> Unexpected SomeOtherLatentInt32 value");
+	AddErrorIfFalse(Context.GetParamStack().GetParam<float>("SomeLatentFloat") == 1123.31f, "FAnimationAnimNextRuntimeTest_GraphExecuteLatent -> Unexpected SomeLatentFloat value");
 
 	Context.GetMutableParamStack().PopLayer(LayerHandle);
-	AnimNextGraph->ReleaseInstance(GraphInstance);
+	GraphInstance.Release();
 
 	return true;
 }
diff --git a/Engine/Plugins/Experimental/AnimNext/Source/AnimNextTestSuite/Private/AnimNextDecoratorGraphTest.h b/Engine/Plugins/Experimental/AnimNext/Source/AnimNextTestSuite/Private/AnimNextDecoratorGraphTest.h
--- a/Engine/Plugins/Experimental/AnimNext/Source/AnimNextTestSuite/Private/AnimNextDecoratorGraphTest.h
+++ b/Engine/Plugins/Experimental/AnimNext/Source/AnimNextTestSuite/Private/AnimNextDecoratorGraphTest.h
@@ -23,9 +23,26 @@ struct FTestDecoratorSharedData : public FAnimNextDecoratorSharedData
 {
 	GENERATED_BODY()
 
-	UPROPERTY(meta = (Input))
+	UPROPERTY(meta = (Input, Inline))
 	int32 SomeInt32 = 3;
 
-	UPROPERTY(meta = (Input))
+	UPROPERTY(meta = (Input, Inline))
 	float SomeFloat = 34.0f;
+
+	UPROPERTY(meta = (Input))
+	int32 SomeLatentInt32 = 3;			// MathAdd with constants, latent
+
+	UPROPERTY(meta = (Input))
+	int32 SomeOtherLatentInt32 = 3;		// GetParameter, latent
+
+	UPROPERTY(meta = (Input))
+	float SomeLatentFloat = 34.0f;		// Inline value, not latent
+
+	#define DECORATOR_LATENT_PROPERTIES_ENUMERATOR(GeneratorMacro) \
+		GeneratorMacro(SomeLatentInt32) \
+		GeneratorMacro(SomeOtherLatentInt32) \
+		GeneratorMacro(SomeLatentFloat) \
+
+	GENERATE_DECORATOR_LATENT_PROPERTIES(FTestDecoratorSharedData, DECORATOR_LATENT_PROPERTIES_ENUMERATOR)
+	#undef DECORATOR_LATENT_PROPERTIES_ENUMERATOR
 };
diff --git a/Engine/Plugins/Experimental/AnimNext/Source/AnimNextTestSuite/Private/AnimNextDecoratorInterfacesTest.cpp b/Engine/Plugins/Experimental/AnimNext/Source/AnimNextTestSuite/Private/AnimNextDecoratorInterfacesTest.cpp
--- a/Engine/Plugins/Experimental/AnimNext/Source/AnimNextTestSuite/Private/AnimNextDecoratorInterfacesTest.cpp
+++ b/Engine/Plugins/Experimental/AnimNext/Source/AnimNextTestSuite/Private/AnimNextDecoratorInterfacesTest.cpp
@@ -104,9 +104,9 @@ namespace UE::AnimNext
 		{
 			FDecoratorPtr Child;
 
-			void Construct(FExecutionContext& Context, FWeakDecoratorPtr DecoratorPtr, const FSharedData& SharedData)
+			void Construct(FExecutionContext& Context, const FDecoratorBinding& Binding)
 			{
-				Child = Context.AllocateNodeInstance(DecoratorPtr, SharedData.Child);
+				Child = Context.AllocateNodeInstance(Binding.GetDecoratorPtr(), Binding.GetSharedData<FSharedData>()->Child);
 			}
 		};
 
@@ -135,10 +135,10 @@ namespace UE::AnimNext
 		{
 			FDecoratorPtr Children[2];
 
-			void Construct(FExecutionContext& Context, FWeakDecoratorPtr DecoratorPtr, const FSharedData& SharedData)
+			void Construct(FExecutionContext& Context, const FDecoratorBinding& Binding)
 			{
-				Children[0] = Context.AllocateNodeInstance(DecoratorPtr, SharedData.Children[0]);
-				Children[1] = Context.AllocateNodeInstance(DecoratorPtr, SharedData.Children[1]);
+				Children[0] = Context.AllocateNodeInstance(Binding.GetDecoratorPtr(), Binding.GetSharedData<FSharedData>()->Children[0]);
+				Children[1] = Context.AllocateNodeInstance(Binding.GetDecoratorPtr(), Binding.GetSharedData<FSharedData>()->Children[1]);
 			}
 		};
 
@@ -285,21 +285,37 @@ bool FAnimationAnimNextRuntimeTest_IHierarchy::RunTest(const FString& InParamete
 			[&DecoratorPropertiesA](uint32 DecoratorIndex, const FString& PropertyName)
 			{
 				return DecoratorPropertiesA[DecoratorIndex][PropertyName];
+			},
+			[](uint32 DecoratorIndex, const FString& PropertyName)
+			{
+				return false;
 			});
 		DecoratorWriter.WriteNode(NodeB,
 			[&DecoratorPropertiesB](uint32 DecoratorIndex, const FString& PropertyName)
 			{
 				return DecoratorPropertiesB[DecoratorIndex][PropertyName];
+			},
+			[](uint32 DecoratorIndex, const FString& PropertyName)
+			{
+				return false;
 			});
 		DecoratorWriter.WriteNode(NodeC,
 			[&DecoratorPropertiesC](uint32 DecoratorIndex, const FString& PropertyName)
 			{
 				return DecoratorPropertiesC[DecoratorIndex][PropertyName];
+			},
+			[](uint32 DecoratorIndex, const FString& PropertyName)
+			{
+				return false;
 			});
 		DecoratorWriter.WriteNode(NodeD,
 			[&DecoratorPropertiesD](uint32 DecoratorIndex, const FString& PropertyName)
 			{
 				return DecoratorPropertiesD[DecoratorIndex][PropertyName];
+			},
+			[](uint32 DecoratorIndex, const FString& PropertyName)
+			{
+				return false;
 			});
 		DecoratorWriter.EndNodeWriting();
 
@@ -465,16 +481,28 @@ bool FAnimationAnimNextRuntimeTest_IUpdate::RunTest(const FString& InParameters)
 			[&DecoratorPropertiesA](uint32 DecoratorIndex, const FString& PropertyName)
 			{
 				return DecoratorPropertiesA[DecoratorIndex][PropertyName];
+			},
+			[](uint32 DecoratorIndex, const FString& PropertyName)
+			{
+				return false;
 			});
 		DecoratorWriter.WriteNode(NodeB,
 			[&DecoratorPropertiesB](uint32 DecoratorIndex, const FString& PropertyName)
 			{
 				return DecoratorPropertiesB[DecoratorIndex][PropertyName];
+			},
+			[](uint32 DecoratorIndex, const FString& PropertyName)
+			{
+				return false;
 			});
 		DecoratorWriter.WriteNode(NodeC,
 			[&DecoratorPropertiesC](uint32 DecoratorIndex, const FString& PropertyName)
 			{
 				return DecoratorPropertiesC[DecoratorIndex][PropertyName];
+			},
+			[](uint32 DecoratorIndex, const FString& PropertyName)
+			{
+				return false;
 			});
 		DecoratorWriter.EndNodeWriting();
 
@@ -602,16 +630,28 @@ bool FAnimationAnimNextRuntimeTest_IEvaluate::RunTest(const FString& InParameter
 			[&DecoratorPropertiesA](uint32 DecoratorIndex, const FString& PropertyName)
 			{
 				return DecoratorPropertiesA[DecoratorIndex][PropertyName];
+			},
+			[](uint32 DecoratorIndex, const FString& PropertyName)
+			{
+				return false;
 			});
 		DecoratorWriter.WriteNode(NodeB,
 			[&DecoratorPropertiesB](uint32 DecoratorIndex, const FString& PropertyName)
 			{
 				return DecoratorPropertiesB[DecoratorIndex][PropertyName];
+			},
+			[](uint32 DecoratorIndex, const FString& PropertyName)
+			{
+				return false;
 			});
 		DecoratorWriter.WriteNode(NodeC,
 			[&DecoratorPropertiesC](uint32 DecoratorIndex, const FString& PropertyName)
 			{
 				return DecoratorPropertiesC[DecoratorIndex][PropertyName];
+			},
+			[](uint32 DecoratorIndex, const FString& PropertyName)
+			{
+				return false;
 			});
 		DecoratorWriter.EndNodeWriting();
 
diff --git a/Engine/Plugins/Experimental/AnimNext/Source/AnimNextUncookedOnly/Internal/Graph/AnimNextGraph_Controller.h b/Engine/Plugins/Experimental/AnimNext/Source/AnimNextUncookedOnly/Internal/Graph/AnimNextGraph_Controller.h
new file mode 100644
--- /dev/null
+++ b/Engine/Plugins/Experimental/AnimNext/Source/AnimNextUncookedOnly/Internal/Graph/AnimNextGraph_Controller.h
@@ -0,0 +1,21 @@
+// Copyright Epic Games, Inc. All Rights Reserved.
+
+#pragma once
+
+#include "CoreMinimal.h"
+#include "RigVMModel/RigVMController.h"
+
+#include "AnimNextGraph_Controller.generated.h"
+
+/**
+  * Implements AnimNext RigVM controller extensions
+  */
+UCLASS(MinimalAPI)
+class UAnimNextGraph_Controller : public URigVMController
+{
+	GENERATED_BODY()
+
+public:
+	// Adds a unit node with a dynamic number of pins
+	URigVMUnitNode* AddUnitNodeWithPins(UScriptStruct* InScriptStruct, const FRigVMPinInfoArray& PinArray, const FName& InMethodName = TEXT("Execute"), const FVector2D& InPosition = FVector2D::ZeroVector, const FString& InNodeName = TEXT(""), bool bSetupUndoRedo = true, bool bPrintPythonCommand = false);
+};
diff --git a/Engine/Plugins/Experimental/AnimNext/Source/AnimNextUncookedOnly/Internal/Graph/AnimNextGraph_UnitNode.h b/Engine/Plugins/Experimental/AnimNext/Source/AnimNextUncookedOnly/Internal/Graph/AnimNextGraph_UnitNode.h
new file mode 100644
--- /dev/null
+++ b/Engine/Plugins/Experimental/AnimNext/Source/AnimNextUncookedOnly/Internal/Graph/AnimNextGraph_UnitNode.h
@@ -0,0 +1,20 @@
+// Copyright Epic Games, Inc. All Rights Reserved.
+
+#pragma once
+
+#include "CoreMinimal.h"
+#include "RigVMModel/Nodes/RigVMUnitNode.h"
+
+#include "AnimNextGraph_UnitNode.generated.h"
+
+/**
+  * Implements AnimNext RigVM unit node extensions
+  */
+UCLASS(MinimalAPI)
+class UAnimNextGraph_UnitNode : public URigVMUnitNode
+{
+	GENERATED_BODY()
+
+public:
+	virtual bool HasNonNativePins() const { return true; }
+};
diff --git a/Engine/Plugins/Experimental/AnimNext/Source/AnimNextUncookedOnly/Private/Graph/AnimNextGraph_Controller.cpp b/Engine/Plugins/Experimental/AnimNext/Source/AnimNextUncookedOnly/Private/Graph/AnimNextGraph_Controller.cpp
new file mode 100644
--- /dev/null
+++ b/Engine/Plugins/Experimental/AnimNext/Source/AnimNextUncookedOnly/Private/Graph/AnimNextGraph_Controller.cpp
@@ -0,0 +1,49 @@
+// Copyright Epic Games, Inc. All Rights Reserved.
+
+#include "Graph/AnimNextGraph_Controller.h"
+#include "Graph/AnimNextGraph_UnitNode.h"
+
+URigVMUnitNode* UAnimNextGraph_Controller::AddUnitNodeWithPins(UScriptStruct* InScriptStruct, const FRigVMPinInfoArray& PinArray, const FName& InMethodName, const FVector2D& InPosition, const FString& InNodeName, bool bSetupUndoRedo, bool bPrintPythonCommand)
+{
+	const bool bHasDynamicPins = PinArray.Num() != 0;
+
+	if (bHasDynamicPins)
+	{
+		OpenUndoBracket(TEXT("Add unit node with pins"));
+	}
+
+	URigVMUnitNode* Node = AddUnitNode(InScriptStruct, UAnimNextGraph_UnitNode::StaticClass(), InMethodName, InPosition, InNodeName, bSetupUndoRedo, bPrintPythonCommand);
+
+	if (Node == nullptr)
+	{
+		if (bHasDynamicPins)
+		{
+			CancelUndoBracket();
+		}
+
+		return nullptr;
+	}
+
+	if (bHasDynamicPins)
+	{
+		const FRigVMRegistry& Registry = FRigVMRegistry::Get();
+		const FRigVMPinInfoArray PreviousPins(Node);
+
+		for (int32 PinIndex = 0; PinIndex < PinArray.Num(); ++PinIndex)
+		{
+			const FString& PinPath = PinArray.GetPinPath(PinIndex);
+			FString ParentPinPath, PinName;
+			UObject* OuterForPin = Node;
+			if (URigVMPin::SplitPinPathAtEnd(PinPath, ParentPinPath, PinName))
+			{
+				OuterForPin = Node->FindPin(ParentPinPath);
+			}
+
+			CreatePinFromPinInfo(Registry, PreviousPins, PinArray[PinIndex], PinPath, OuterForPin);
+		}
+
+		CloseUndoBracket();
+	}
+
+	return Node;
+}
diff --git a/Engine/Plugins/Experimental/AnimNext/Source/AnimNextUncookedOnly/Private/Graph/AnimNextGraph_EditorData.cpp b/Engine/Plugins/Experimental/AnimNext/Source/AnimNextUncookedOnly/Private/Graph/AnimNextGraph_EditorData.cpp
--- a/Engine/Plugins/Experimental/AnimNext/Source/AnimNextUncookedOnly/Private/Graph/AnimNextGraph_EditorData.cpp
+++ b/Engine/Plugins/Experimental/AnimNext/Source/AnimNextUncookedOnly/Private/Graph/AnimNextGraph_EditorData.cpp
@@ -4,6 +4,7 @@
 
 #include "ControlRigDefines.h"
 #include "Graph/AnimNextGraph.h"
+#include "Graph/AnimNextGraph_Controller.h"
 #include "Graph/AnimNextGraph_EdGraphSchema.h"
 #include "UncookedOnlyUtils.h"
 #include "Graph/AnimNextExecuteContext.h"
@@ -19,6 +20,7 @@ UAnimNextGraph_EditorData::UAnimNextGraph_EditorData(const FObjectInitializer& O
 {
 	RigVMClient.Reset();
 	RigVMClient.SetSchemaClass(UAnimNextGraph_Schema::StaticClass());
+	RigVMClient.SetControllerClass(UAnimNextGraph_Controller::StaticClass());
 	RigVMClient.SetOuterClientHost(this, GET_MEMBER_NAME_CHECKED(UAnimNextGraph_EditorData, RigVMClient));
 	{
 		TGuardValue<bool> DisableClientNotifs(RigVMClient.bSuspendNotifications, true);
diff --git a/Engine/Plugins/Experimental/AnimNext/Source/AnimNextUncookedOnly/Private/UncookedOnlyUtils.cpp b/Engine/Plugins/Experimental/AnimNext/Source/AnimNextUncookedOnly/Private/UncookedOnlyUtils.cpp
--- a/Engine/Plugins/Experimental/AnimNext/Source/AnimNextUncookedOnly/Private/UncookedOnlyUtils.cpp
+++ b/Engine/Plugins/Experimental/AnimNext/Source/AnimNextUncookedOnly/Private/UncookedOnlyUtils.cpp
@@ -2,9 +2,12 @@
 
 #include "UncookedOnlyUtils.h"
 #include "Graph/AnimNextGraph.h"
+#include "Graph/AnimNextGraph_Controller.h"
 #include "Graph/AnimNextGraph_EditorData.h"
-#include "Graph/RigUnit_AnimNextGraphRoot.h"
 #include "Graph/RigUnit_AnimNextDecoratorStack.h"
+#include "Graph/RigUnit_AnimNextGraphRoot.h"
+#include "Graph/RigUnit_AnimNextGraphEvaluator.h"
+#include "Graph/RigUnit_AnimNextShimRoot.h"
 #include "Graph/AnimNextExecuteContext.h"
 #include "Param/AnimNextParameter.h"
 #include "Param/AnimNextParameterBlock.h"
@@ -108,7 +111,7 @@ namespace Private
 	FString GetDecoratorProperty(const FDecoratorStackMapping& DecoratorStack, uint32 DecoratorIndex, const FString& PropertyName, const TArray<FDecoratorStackMapping>& DecoratorStackNodes)
 	{
 		const TArray<URigVMPin*>& Pins = DecoratorStack.DecoratorEntries[DecoratorIndex].DecoratorEntryPin->GetSubPins();
-		for (URigVMPin* Pin : Pins)
+		for (const URigVMPin* Pin : Pins)
 		{
 			if (Pin->GetDirection() != ERigVMPinDirection::Input)
 			{
@@ -168,12 +171,37 @@ namespace Private
 		return FString();
 	}
 
+	bool IsDecoratorPropertyLatent(const FDecoratorStackMapping& DecoratorStack, uint32 DecoratorIndex, const FString& PropertyName, const TArray<FDecoratorStackMapping>& DecoratorStackNodes)
+	{
+		const TArray<URigVMPin*>& Pins = DecoratorStack.DecoratorEntries[DecoratorIndex].DecoratorEntryPin->GetSubPins();
+		for (const URigVMPin* Pin : Pins)
+		{
+			if (Pin->GetDirection() != ERigVMPinDirection::Input)
+			{
+				continue;	// We only look for input pins
+			}
+
+			if (Pin->GetName() == PropertyName)
+			{
+				// Lazy pins are latent if they are connected to something
+				return Pin->IsLazy() && !Pin->GetLinks().IsEmpty();
+			}
+		}
+
+		// Unknown property
+		return false;
+	}
+
 	void WriteDecoratorProperties(FDecoratorWriter& DecoratorWriter, const FDecoratorStackMapping& Mapping, const TArray<FDecoratorStackMapping>& DecoratorStackNodes)
 	{
 		DecoratorWriter.WriteNode(Mapping.DecoratorStackNodeHandle,
 			[&Mapping, &DecoratorStackNodes](uint32 DecoratorIndex, const FString& PropertyName)
 			{
 				return GetDecoratorProperty(Mapping, DecoratorIndex, PropertyName, DecoratorStackNodes);
+			},
+			[&Mapping, &DecoratorStackNodes](uint32 DecoratorIndex, const FString& PropertyName)
+			{
+				return IsDecoratorPropertyLatent(Mapping, DecoratorIndex, PropertyName, DecoratorStackNodes);
 			});
 	}
 
@@ -236,6 +264,85 @@ namespace Private
 
 		return DecoratorStackNodes;
 	}
+
+	FRigVMPinInfoArray CollectLatentPins(const TArray<FDecoratorStackMapping>& DecoratorStackNodes, TMap<FName, URigVMPin*>& LatentPinMapping)
+	{
+		const FRigVMRegistry& Registry = FRigVMRegistry::Get();
+
+		FRigVMPinInfoArray LatentPins;
+
+		for (const FDecoratorStackMapping& DecoratorStack : DecoratorStackNodes)
+		{
+			for (const FDecoratorEntryMapping& DecoratorEntry : DecoratorStack.DecoratorEntries)
+			{
+				for (URigVMPin* Pin : DecoratorEntry.DecoratorEntryPin->GetSubPins())
+				{
+					if (Pin->IsLazy() && !Pin->GetLinks().IsEmpty())
+					{
+						// This pin has something linked to it, it is a latent pin
+						const FName LatentPinName(TEXT("LatentPin"), LatentPins.Num());	// Create unique latent pin names
+
+						FRigVMPinInfo PinInfo;
+						PinInfo.Name = LatentPinName;
+						PinInfo.TypeIndex = Pin->GetTypeIndex();
+
+						// All our programmatic pins are lazy inputs
+						PinInfo.Direction = ERigVMPinDirection::Input;
+						PinInfo.bIsLazy = true;
+
+						LatentPins.Pins.Emplace(PinInfo);
+
+						const TArray<URigVMLink*>& PinLinks = Pin->GetLinks();
+						check(PinLinks.Num() == 1);
+
+						LatentPinMapping.Add(LatentPinName, PinLinks[0]->GetSourcePin());
+					}
+				}
+			}
+		}
+
+		return LatentPins;
+	}
+
+	TArray<FRigVMFunctionArgument> GetGraphEvaluatorFunctionArguments(const FRigVMPinInfoArray& LatentPins)
+	{
+		const FRigVMRegistry& Registry = FRigVMRegistry::Get();
+
+		TArray<FRigVMFunctionArgument> Arguments;
+		Arguments.Reserve(LatentPins.Num());
+
+		for (const FRigVMPinInfo& Pin : LatentPins)
+		{
+			const FRigVMTemplateArgumentType& TypeArg = Registry.GetType(Pin.TypeIndex);
+
+			Arguments.Add(FRigVMFunctionArgument(Pin.Name.ToString(), TypeArg.CPPType.ToString(), ERigVMFunctionArgumentDirection::Input));
+		}
+
+		return Arguments;
+	}
+
+	FString GetGraphEvaluatorMethodName(const FRigVMPinInfoArray& LatentPins)
+	{
+		static TMap<uint32, FString> GraphEvaluatorMethodNameCache;
+
+		const uint32 LatentPinListHash = GetTypeHash(LatentPins);
+		if (const FString* MethodName = GraphEvaluatorMethodNameCache.Find(LatentPinListHash))
+		{
+			return *MethodName;
+		}
+
+		// Generate a new method for this argument list
+		const FString MethodName = FString::Printf(TEXT("Execute_%X"), LatentPinListHash);
+		const FString FullExecuteMethodName = FString::Printf(TEXT("FRigUnit_AnimNextGraphEvaluator::%s"), *MethodName);
+
+		const TArray<FRigVMFunctionArgument> GraphEvaluatorArguments = GetGraphEvaluatorFunctionArguments(LatentPins);
+		FRigVMRegistry::Get().Register(*FullExecuteMethodName, &FRigUnit_AnimNextGraphEvaluator::StaticExecute, FRigUnit_AnimNextGraphEvaluator::StaticStruct(), GraphEvaluatorArguments);
+
+		// Cache our result
+		GraphEvaluatorMethodNameCache.Add(LatentPinListHash, MethodName);
+
+		return MethodName;
+	}
 }
 
 void FUtils::Compile(UAnimNextGraph* InGraph)
@@ -267,75 +374,94 @@ void FUtils::Compile(UAnimNextGraph* InGraph)
 	EditorData->CompileLog.NumErrors = EditorData->CompileLog.NumWarnings = 0;
 
 	// We use a temporary graph model to build our final graph that we'll compile
-	URigVMGraph* VMTempGraph = EditorData->GetRigVMClient()->GetModel(TEXT("TempRigVMGraph"));
-	URigVMController* TempController = nullptr;
-
-	TArray<URigVMGraph*> GraphModelsToCompile;
+	FRigVMClient* VMClient = EditorData->GetRigVMClient();
+	URigVMGraph* VMRootGraph = VMClient->GetDefaultModel();
+	URigVMGraph* VMTempGraph = CastChecked<URigVMGraph>(StaticDuplicateObject(VMRootGraph, GetTransientPackage(), VMClient->GetUniqueName(TEXT("TempRigVMGraph"))));
 
-	if (VMTempGraph != nullptr)
-	{
-		TempController = EditorData->GetRigVMClient()->GetOrCreateController(VMTempGraph);
+	UAnimNextGraph_Controller* TempController = CastChecked<UAnimNextGraph_Controller>(VMClient->GetOrCreateController(VMTempGraph));
 
-		// Clear our temporary model
-		TempController->RemoveNodes(VMTempGraph->GetNodes());
+	// Gather our decorator stacks
+	TArray<Private::FDecoratorStackMapping> DecoratorStackNodes = Private::CollectDecoratorStacks(VMTempGraph);
 
-		URigVMGraph* VMRootGraph = EditorData->GetRigVMClient()->GetDefaultModel();
+	// Add our runtime shim root node
+	URigVMUnitNode* TempShimRootNode = TempController->AddUnitNode(FRigUnit_AnimNextShimRoot::StaticStruct(), FRigVMStruct::ExecuteName, FVector2D::ZeroVector, FString(), false);
 
-		// Add our root node, it's always there
-		URigVMUnitNode* TempRootNode = TempController->AddUnitNode(FRigUnit_AnimNextGraphRoot::StaticStruct(), FRigVMStruct::ExecuteName, FVector2D(0.0f, 0.0f), FString(), false);
+	// Add our graph evaluator node
+	TMap<FName, URigVMPin*> LatentPinMapping;
+	const FRigVMPinInfoArray LatentPins = Private::CollectLatentPins(DecoratorStackNodes, LatentPinMapping);
 
-		// Gather our decorator stacks
-		TArray<Private::FDecoratorStackMapping> DecoratorStackNodes = Private::CollectDecoratorStacks(VMRootGraph);
+	// We need a unique method name to match our unique argument list
+	const FString ExecuteMethodName = Private::GetGraphEvaluatorMethodName(LatentPins);
 
-		FDecoratorWriter DecoratorWriter;
+	URigVMUnitNode* GraphEvaluatorNode = TempController->AddUnitNodeWithPins(FRigUnit_AnimNextGraphEvaluator::StaticStruct(), LatentPins, *ExecuteMethodName, FVector2D::ZeroVector, FString(), false);
 
-		// Iterate over every decorator stack and register our node templates
-		for (Private::FDecoratorStackMapping& NodeMapping : DecoratorStackNodes)
-		{
-			NodeMapping.DecoratorStackNodeHandle = Private::RegisterDecoratorNodeTemplate(DecoratorWriter, NodeMapping.DecoratorStackNode);
-		}
+	// Link our shim and evaluator nodes together using the execution context
+	TempController->AddLink(
+		TempShimRootNode->FindPin(GET_MEMBER_NAME_STRING_CHECKED(FRigUnit_AnimNextShimRoot, ExecuteContext)),
+		GraphEvaluatorNode->FindPin(GET_MEMBER_NAME_STRING_CHECKED(FRigUnit_AnimNextGraphEvaluator, ExecuteContext)));
 
-		// Write our node shared data
-		DecoratorWriter.BeginNodeWriting();
+	// Link our latent pins
+	for (const FRigVMPinInfo& LatentPin : LatentPins)
+	{
+		TempController->AddLink(
+			LatentPinMapping[LatentPin.Name],
+			GraphEvaluatorNode->FindPin(LatentPin.Name.ToString()));
+	}
 
-		for (const Private::FDecoratorStackMapping& NodeMapping : DecoratorStackNodes)
-		{
-			Private::WriteDecoratorProperties(DecoratorWriter, NodeMapping, DecoratorStackNodes);
-		}
+	FDecoratorWriter DecoratorWriter;
 
-		DecoratorWriter.EndNodeWriting();
+	// Iterate over every decorator stack and register our node templates
+	for (Private::FDecoratorStackMapping& NodeMapping : DecoratorStackNodes)
+	{
+		NodeMapping.DecoratorStackNodeHandle = Private::RegisterDecoratorNodeTemplate(DecoratorWriter, NodeMapping.DecoratorStackNode);
+	}
 
-		// Find our root node handle, if we have any stack nodes, the first one is our root stack
-		FAnimNextEntryPointHandle RootDecoratorHandle;
-		if (DecoratorStackNodes.Num() != 0)
-		{
-			RootDecoratorHandle = FAnimNextEntryPointHandle(DecoratorStackNodes[0].DecoratorStackNodeHandle);
-		}
+	// Write our node shared data
+	DecoratorWriter.BeginNodeWriting();
 
-		// Cache our compiled metadata
-		InGraph->SharedDataArchiveBuffer = DecoratorWriter.GetGraphSharedData();
-		InGraph->RootDecoratorHandle = RootDecoratorHandle;
+	for (const Private::FDecoratorStackMapping& NodeMapping : DecoratorStackNodes)
+	{
+		Private::WriteDecoratorProperties(DecoratorWriter, NodeMapping, DecoratorStackNodes);
+	}
 
-		// Populate our runtime metadata
-		InGraph->LoadFromArchiveBuffer(InGraph->SharedDataArchiveBuffer);
+	DecoratorWriter.EndNodeWriting();
 
-		// Compile our new graph
-		GraphModelsToCompile.Add(VMTempGraph);
+	// Find our root node handle, if we have any stack nodes, the first one is our root stack
+	FAnimNextEntryPointHandle RootDecoratorHandle;
+	if (DecoratorStackNodes.Num() != 0)
+	{
+		RootDecoratorHandle = FAnimNextEntryPointHandle(DecoratorStackNodes[0].DecoratorStackNodeHandle);
 	}
 
-	if (VMTempGraph == nullptr)
-	{
-		// This is an old graph that was saved before the introduction of the temporary graph, use the root graph instead
-		URigVMGraph* VMRootGraph = EditorData->GetRigVMClient()->GetDefaultModel();
+	// Cache our compiled metadata
+	InGraph->SharedDataArchiveBuffer = DecoratorWriter.GetGraphSharedData();
+	InGraph->RootDecoratorHandle = RootDecoratorHandle;
+
+	// Populate our runtime metadata
+	InGraph->LoadFromArchiveBuffer(InGraph->SharedDataArchiveBuffer);
+
+	// Remove our old root node
+	if (URigVMNode* const* RootNode = VMTempGraph->GetNodes().FindByPredicate(
+		[](URigVMNode* Node)
+		{
+			if (URigVMUnitNode* UnitNode = Cast<URigVMUnitNode>(Node))
+			{
+				return UnitNode->GetScriptStruct() == FRigUnit_AnimNextGraphRoot::StaticStruct();
+			}
 
-		GraphModelsToCompile.Add(VMRootGraph);
-		TempController = EditorData->GetRigVMClient()->GetOrCreateController(VMRootGraph);
+			return false;
+		}))
+	{
+		TempController->RemoveNode(*RootNode, false, false);
 	}
 
 	URigVMCompiler* Compiler = URigVMCompiler::StaticClass()->GetDefaultObject<URigVMCompiler>();
 	EditorData->VMCompileSettings.SetExecuteContextStruct(EditorData->RigVMClient.GetExecuteContextStruct());
 	const FRigVMCompileSettings Settings = (EditorData->bCompileInDebugMode) ? FRigVMCompileSettings::Fast(EditorData->VMCompileSettings.GetExecuteContextStruct()) : EditorData->VMCompileSettings;
-	Compiler->Compile(Settings, GraphModelsToCompile, TempController, InGraph->RigVM, InGraph->ExtendedExecuteContext, InGraph->GetRigVMExternalVariables(), &EditorData->PinToOperandMap);
+	Compiler->Compile(Settings, { VMTempGraph }, TempController, InGraph->RigVM, InGraph->ExtendedExecuteContext, InGraph->GetRigVMExternalVariables(), & EditorData->PinToOperandMap);
+
+	// Initialize right away, in packaged builds we initialize during PostLoad
+	InGraph->RigVM->Initialize(InGraph->ExtendedExecuteContext, InGraph->RigVM->GetLocalMemoryArray(InGraph->ExtendedExecuteContext));
 
 	if (EditorData->bErrorsDuringCompilation)
 	{
@@ -351,6 +477,8 @@ void FUtils::Compile(UAnimNextGraph* InGraph)
 		EditorData->VMCompiledEvent.Broadcast(InGraph, InGraph->RigVM, InGraph->ExtendedExecuteContext);
 	}
 
+	VMClient->RemoveController(VMTempGraph);
+
 #if WITH_EDITOR
 //	RefreshBreakpoints(EditorData);
 #endif
@@ -363,8 +491,7 @@ void FUtils::RecreateVM(UAnimNextGraph* InGraph)
 	// Cooked platforms will load these pointers from disk
 	if (!FPlatformProperties::RequiresCookedData())
 	{
-		// We dont support ERigVMMemoryType::Work memory as we dont operate on an instance
-	//	InGraph->RigVM->GetMemoryByType(ERigVMMemoryType::Work, true);
+		InGraph->RigVM->CreateMemoryByType(InGraph->ExtendedExecuteContext, ERigVMMemoryType::Work);
 		InGraph->RigVM->CreateMemoryByType(InGraph->ExtendedExecuteContext, ERigVMMemoryType::Literal);
 		InGraph->RigVM->CreateMemoryByType(InGraph->ExtendedExecuteContext, ERigVMMemoryType::Debug);
 	}
diff --git a/Engine/Plugins/Runtime/RigVM/Source/RigVM/Public/RigVMCore/RigVMMemoryStorage.h b/Engine/Plugins/Runtime/RigVM/Source/RigVM/Public/RigVMCore/RigVMMemoryStorage.h
--- a/Engine/Plugins/Runtime/RigVM/Source/RigVM/Public/RigVMCore/RigVMMemoryStorage.h
+++ b/Engine/Plugins/Runtime/RigVM/Source/RigVM/Public/RigVMCore/RigVMMemoryStorage.h
@@ -304,7 +304,7 @@ public:
 	 * Computes the data if necessary and returns true if the value is valid
 	 * @return True if the value of the handle is valid after the compute
 	 */
-	bool ComputeLazyValueIfNecessary(FRigVMExtendedExecuteContext& Context, int32 InSliceIndex = INDEX_NONE);
+	RIGVM_API bool ComputeLazyValueIfNecessary(FRigVMExtendedExecuteContext& Context, int32 InSliceIndex = INDEX_NONE);
 
 	// Returns the head property of this handle
 	const FProperty* GetProperty() const
diff --git a/Engine/Plugins/Runtime/RigVM/Source/RigVMDeveloper/Private/RigVMCompiler/RigVMAST.cpp b/Engine/Plugins/Runtime/RigVM/Source/RigVMDeveloper/Private/RigVMCompiler/RigVMAST.cpp
--- a/Engine/Plugins/Runtime/RigVM/Source/RigVMDeveloper/Private/RigVMCompiler/RigVMAST.cpp
+++ b/Engine/Plugins/Runtime/RigVM/Source/RigVMDeveloper/Private/RigVMCompiler/RigVMAST.cpp
@@ -1044,10 +1044,24 @@ TArray<FRigVMExprAST*> FRigVMParserAST::TraversePins(const FRigVMASTProxy& InNod
 				}
 			}
 		}
+
+		// We might have extra non-native pins, add them afterwards
+		if (UnitNode->HasNonNativePins())
+		{
+			for (URigVMPin* Pin : Node->GetPins())
+			{
+				// We skip decorator pins as we don't want to traverse them
+				if (!Pin->IsDecoratorPin())
+				{
+					Pins.AddUnique(Pin);
+				}
+			}
+		}
 	}
 
-	if(Pins.IsEmpty())
+	if (Pins.IsEmpty())
 	{
+		// Just grab whatever pins we have
 		Pins = Node->GetPins();
 	}
 
diff --git a/Engine/Plugins/Runtime/RigVM/Source/RigVMDeveloper/Private/RigVMCompiler/RigVMCompiler.cpp b/Engine/Plugins/Runtime/RigVM/Source/RigVMDeveloper/Private/RigVMCompiler/RigVMCompiler.cpp
--- a/Engine/Plugins/Runtime/RigVM/Source/RigVMDeveloper/Private/RigVMCompiler/RigVMCompiler.cpp
+++ b/Engine/Plugins/Runtime/RigVM/Source/RigVMDeveloper/Private/RigVMCompiler/RigVMCompiler.cpp
@@ -1075,7 +1075,7 @@ bool URigVMCompiler::Compile(const FRigVMCompileSettings& InSettings, TArray<URi
 	if (!WorkData.AST.IsValid())
 	{
 		WorkData.OverrideReportDelegate(bEncounteredASTError, bSurpressedASTError);
-		WorkData.AST = MakeShareable(new FRigVMParserAST(InGraphs, InController, Settings.ASTSettings, InExternalVariables));
+		WorkData.AST = MakeShared<FRigVMParserAST>(InGraphs, InController, Settings.ASTSettings, InExternalVariables);
 		WorkData.RemoveOverrideReportDelegate();
 
 		for(URigVMGraph* Graph : InGraphs)
@@ -1220,7 +1220,7 @@ bool URigVMCompiler::Compile(const FRigVMCompileSettings& InSettings, TArray<URi
 #else
 	for(ERigVMMemoryType MemoryType : MemoryTypes)
 	{
-		UPackage* Package = InGraphs[0]->GetOutermost();
+		UPackage* Package = WorkData.VM->GetOutermost();
 
 		const TArray<FRigVMPropertyDescription>* Properties = WorkData.PropertyDescriptions.Find(MemoryType);
 		if(Properties == nullptr)
diff --git a/Engine/Plugins/Runtime/RigVM/Source/RigVMDeveloper/Private/RigVMModel/RigVMClient.cpp b/Engine/Plugins/Runtime/RigVM/Source/RigVMDeveloper/Private/RigVMModel/RigVMClient.cpp
--- a/Engine/Plugins/Runtime/RigVM/Source/RigVMDeveloper/Private/RigVMModel/RigVMClient.cpp
+++ b/Engine/Plugins/Runtime/RigVM/Source/RigVMDeveloper/Private/RigVMModel/RigVMClient.cpp
@@ -46,6 +46,23 @@ void FRigVMClient::SetSchemaClass(TSubclassOf<URigVMSchema> InSchemaClass)
 	}
 }
 
+void FRigVMClient::SetControllerClass(TSubclassOf<URigVMController> InControllerClass)
+{
+	check(InControllerClass);
+
+	if (InControllerClass == ControllerClass)
+	{
+		return;
+	}
+
+	for (URigVMGraph* Model : Models)
+	{
+		RemoveController(Model);
+	}
+
+	ControllerClass = InControllerClass;
+}
+
 void FRigVMClient::SetOuterClientHost(UObject* InOuterClientHost, const FName& InOuterClientHostPropertyName)
 {
 	OuterClientHost = InOuterClientHost;
@@ -754,7 +771,7 @@ URigVMController* FRigVMClient::CreateController(const URigVMGraph* InModel)
 {
 	const FString ModelName = InModel ? InModel->GetName() : TEXT("__NullGraph");
 	const FName SafeControllerName = GetUniqueName(*FString::Printf(TEXT("%s_Controller"), *ModelName));
-	URigVMController* Controller = NewObject<URigVMController>(GetOuter(), SafeControllerName);
+	URigVMController* Controller = NewObject<URigVMController>(GetOuter(), ControllerClass, SafeControllerName);
 	Controllers.Add(InModel, Controller);
 	Controller->SetSchema(GetOrCreateSchema());
 	Controller->SetActionStack(GetOrCreateActionStack());
diff --git a/Engine/Plugins/Runtime/RigVM/Source/RigVMDeveloper/Private/RigVMModel/RigVMController.cpp b/Engine/Plugins/Runtime/RigVM/Source/RigVMDeveloper/Private/RigVMModel/RigVMController.cpp
--- a/Engine/Plugins/Runtime/RigVM/Source/RigVMDeveloper/Private/RigVMModel/RigVMController.cpp
+++ b/Engine/Plugins/Runtime/RigVM/Source/RigVMDeveloper/Private/RigVMModel/RigVMController.cpp
@@ -362,6 +362,7 @@ FRigVMPinInfo::FRigVMPinInfo()
 	, bIsConstant(false)
 	, bIsDynamicArray(false)
 	, bIsDecorator(false)
+	, bIsLazy(false)
 {
 }
 
@@ -376,6 +377,7 @@ FRigVMPinInfo::FRigVMPinInfo(const URigVMPin* InPin, int32 InParentIndex, ERigVM
 	, bIsConstant(InPin->IsDefinedAsConstant())
 	, bIsDynamicArray(InPin->IsDynamicArray())
 	, bIsDecorator(InPin->IsDecoratorPin() && InPin->IsRootPin())
+	, bIsLazy(InPin->IsLazy() && InPin->IsRootPin())
 {
 	// this method describes the info as currently represented in the model.
 
@@ -400,6 +402,7 @@ FRigVMPinInfo::FRigVMPinInfo(FProperty* InProperty, ERigVMPinDirection InDirecti
 	, bIsConstant(false)
 	, bIsDynamicArray(false)
 	, bIsDecorator(false)
+	, bIsLazy(false)
 {
 	// this method describes the info as needed based on the property structure
 
@@ -443,6 +446,11 @@ FRigVMPinInfo::FRigVMPinInfo(FProperty* InProperty, ERigVMPinDirection InDirecti
 			bIsDynamicArray = false;
 		}
 	}
+
+	if (InProperty->HasMetaData(FRigVMStruct::ComputeLazilyMetaName))
+	{
+		bIsLazy = true;
+	}
 #endif
 
 	UObject* CPPTypeObject = nullptr;
@@ -1285,6 +1293,11 @@ TArray<FString> URigVMController::GetAddDecoratorPythonCommands(URigVMNode* Node
 #if WITH_EDITOR
 
 URigVMUnitNode* URigVMController::AddUnitNode(UScriptStruct* InScriptStruct, const FName& InMethodName, const FVector2D& InPosition, const FString& InNodeName, bool bSetupUndoRedo, bool bPrintPythonCommand)
+{
+	return AddUnitNode(InScriptStruct, URigVMUnitNode::StaticClass(), InMethodName, InPosition, InNodeName, bSetupUndoRedo, bPrintPythonCommand);
+}
+
+URigVMUnitNode* URigVMController::AddUnitNode(UScriptStruct* InScriptStruct, TSubclassOf<URigVMUnitNode> InUnitNodeClass, const FName& InMethodName, const FVector2D& InPosition, const FString& InNodeName, bool bSetupUndoRedo, bool bPrintPythonCommand)
 {
 	if(!IsValidGraph())
 	{
@@ -1433,7 +1446,7 @@ URigVMUnitNode* URigVMController::AddUnitNode(UScriptStruct* InScriptStruct, con
 	}
 	
 	FString Name = GetSchema()->GetValidNodeName(Graph, InNodeName.IsEmpty() ? InScriptStruct->GetName() : InNodeName);
-	URigVMUnitNode* Node = NewObject<URigVMUnitNode>(Graph, *Name);
+	URigVMUnitNode* Node = NewObject<URigVMUnitNode>(Graph, InUnitNodeClass, *Name);
 	Node->ResolvedFunctionName = Function->GetName();
 	Node->Position = InPosition;
 	Node->NodeTitle = InScriptStruct->GetMetaData(TEXT("DisplayName"));
@@ -13766,55 +13779,7 @@ FName URigVMController::AddDecorator(URigVMNode* InNode, UScriptStruct* InDecora
 	Decorator->GetProgrammaticPins(this, INDEX_NONE, ProgrammaticPins);
 
 	const FRigVMRegistry& Registry = FRigVMRegistry::Get();
-	auto CreatePinFromPinInfo = [this, &Registry](const FRigVMPinInfo& InPinInfo, const FString& InPinPath, UObject* InOuter) -> URigVMPin*
-	{
-		check(InOuter);
-		URigVMPin* Pin = NewObject<URigVMPin>(InOuter, InPinInfo.Name);
-		if (InPinInfo.Property)
-		{
-			ConfigurePinFromProperty(InPinInfo.Property, Pin, InPinInfo.Direction);
-		}
-		else
-		{
-			const FRigVMTemplateArgumentType& Type = Registry.GetType(InPinInfo.TypeIndex);
-			Pin->CPPType = Type.CPPType.ToString();
-			Pin->CPPTypeObject = Type.CPPTypeObject;
-			if (Pin->CPPTypeObject)
-			{
-				Pin->CPPTypeObjectPath = *Pin->CPPTypeObject->GetPathName();
-			}
-			if (Registry.IsExecuteType(InPinInfo.TypeIndex))
-			{
-				MakeExecutePin(Pin);
-			}
-
-			Pin->Direction = InPinInfo.Direction;
-			Pin->DisplayName = InPinInfo.DisplayName.IsEmpty() ? NAME_None : FName(*InPinInfo.DisplayName);
-			Pin->bIsConstant = InPinInfo.bIsConstant;
-			Pin->bIsDynamicArray = InPinInfo.bIsDynamicArray;
-			Pin->CustomWidgetName = InPinInfo.CustomWidgetName.IsEmpty() ? NAME_None : FName(*InPinInfo.CustomWidgetName);
-		}
-
-		Pin->bIsExpanded = InPinInfo.bIsExpanded;
-		Pin->DefaultValue = InPinInfo.DefaultValue;
-
-		if (URigVMPin* ParentPin = Cast<URigVMPin>(InOuter))
-		{
-			AddSubPin(ParentPin, Pin);
-		}
-		else if (URigVMNode* OwnerNode = Cast<URigVMNode>(InOuter))
-		{
-			AddNodePin(OwnerNode, Pin);
-		}
-		else
-		{
-			ensureMsgf(false, TEXT("Outer %s of pin info %s is not a pin or a node"), *InOuter->GetPathName(), *InPinPath);
-		}
-
-		Notify(ERigVMGraphNotifType::PinAdded, Pin);
-
-		return Pin;
-	};
+	const FRigVMPinInfoArray PreviousPins;
 
 	for (int32 PinIndex = 0; PinIndex < ProgrammaticPins.Num(); ++PinIndex)
 	{
@@ -13826,7 +13791,7 @@ FName URigVMController::AddDecorator(URigVMNode* InNode, UScriptStruct* InDecora
 			OuterForPin = DecoratorPin->FindSubPin(ParentPinPath);
 		}
 
-		CreatePinFromPinInfo(ProgrammaticPins[PinIndex], PinPath, OuterForPin);
+		CreatePinFromPinInfo(Registry, PreviousPins, ProgrammaticPins[PinIndex], PinPath, OuterForPin);
 	}
 
 	// move the the pin to the right index as required
@@ -14387,6 +14352,12 @@ void URigVMController::AddPinsForTemplate(const FRigVMTemplate* InTemplate, cons
 			}
 		}
 
+		if (Pin->Direction == ERigVMPinDirection::Input &&
+			!InTemplate->GetArgumentMetaData(Arg->Name, FRigVMStruct::ComputeLazilyMetaName).IsEmpty())
+		{
+			Pin->bIsLazy = true;
+		}
+
 		AddNodePin(InNode, Pin);
 
 		if(!Pin->IsWildCard() && !Pin->IsArray())
@@ -14426,7 +14397,7 @@ void URigVMController::AddPinsForTemplate(const FRigVMTemplate* InTemplate, cons
 	AddExecutePins(ERigVMPinDirection::Output);
 }
 
-void URigVMController::ConfigurePinFromProperty(FProperty* InProperty, URigVMPin* InOutPin, ERigVMPinDirection InPinDirection)
+void URigVMController::ConfigurePinFromProperty(FProperty* InProperty, URigVMPin* InOutPin, ERigVMPinDirection InPinDirection) const
 {
 	if (InPinDirection == ERigVMPinDirection::Invalid)
 	{
@@ -14437,6 +14408,12 @@ void URigVMController::ConfigurePinFromProperty(FProperty* InProperty, URigVMPin
 		InOutPin->Direction = InPinDirection;
 	}
 
+	// If this property wants to be explicitly hidden, hide it
+	if (InProperty->HasMetaData(FRigVMStruct::HiddenMetaName))
+	{
+		InOutPin->Direction = ERigVMPinDirection::Hidden;
+	}
+
 #if WITH_EDITOR
 
 	if (!InOutPin->IsArrayElement())
@@ -14490,6 +14467,24 @@ void URigVMController::ConfigurePinFromProperty(FProperty* InProperty, URigVMPin
 			InOutPin->bIsDynamicArray = false;
 		}
 	}
+
+	if (InOutPin->Direction == ERigVMPinDirection::Input)
+	{
+		// fixed array elements are treated as lazy elements
+		// if the original argument is also marked as lazy
+		if (const URigVMPin* ParentPin = InOutPin->GetParentPin())
+		{
+			if (ParentPin->IsFixedSizeArray())
+			{
+				InOutPin->bIsLazy = ParentPin->IsLazy();
+			}
+		}
+
+		if (InProperty->HasMetaData(FRigVMStruct::ComputeLazilyMetaName))
+		{
+			InOutPin->bIsLazy = true;
+		}
+	}
 #endif
 
 	if (FStructProperty* StructProperty = CastField<FStructProperty>(PropertyForType))
@@ -14556,6 +14551,7 @@ void URigVMController::ConfigurePinFromPin(URigVMPin* InOutPin, URigVMPin* InPin
 	InOutPin->CPPTypeObject = InPin->CPPTypeObject;
 	InOutPin->DefaultValue = InPin->DefaultValue;
 	InOutPin->bIsDynamicArray = InPin->bIsDynamicArray;
+	InOutPin->bIsLazy = InPin->bIsLazy;
 	if(bCopyDisplayName)
 	{
 		InOutPin->SetDisplayName(InPin->GetDisplayName());
@@ -15355,66 +15351,6 @@ void URigVMController::RepopulatePinsOnNode(URigVMNode* InNode, bool bFollowCore
 		}
 	}
 
-	auto CreatePinFromPinInfo = [this, &Registry, &PreviousPinInfos](const FRigVMPinInfo& InPinInfo, const FString& InPinPath, UObject* InOuter) -> URigVMPin*
-	{
-		check(InOuter);
-		URigVMPin* Pin = NewObject<URigVMPin>(InOuter, InPinInfo.Name);
-		if(InPinInfo.Property)
-		{
-			ConfigurePinFromProperty(InPinInfo.Property, Pin, InPinInfo.Direction);
-		}
-		else
-		{
-			const FRigVMTemplateArgumentType& Type = Registry.GetType(InPinInfo.TypeIndex);
-			Pin->CPPType = Type.CPPType.ToString();
-			Pin->CPPTypeObject = Type.CPPTypeObject;
-			if(Pin->CPPTypeObject)
-			{
-				Pin->CPPTypeObjectPath = *Pin->CPPTypeObject->GetPathName();
-			}
-			if(Registry.IsExecuteType(InPinInfo.TypeIndex))
-			{
-				MakeExecutePin(Pin);
-			}
-				
-			Pin->Direction = InPinInfo.Direction;
-			Pin->DisplayName = InPinInfo.DisplayName.IsEmpty() ? NAME_None : FName(*InPinInfo.DisplayName);
-			Pin->bIsConstant = InPinInfo.bIsConstant;
-			Pin->bIsDynamicArray = InPinInfo.bIsDynamicArray;
-			Pin->CustomWidgetName = InPinInfo.CustomWidgetName.IsEmpty() ? NAME_None : FName(*InPinInfo.CustomWidgetName);
-		}
-
-		Pin->bIsExpanded = InPinInfo.bIsExpanded;
-		Pin->DefaultValue = InPinInfo.DefaultValue;
-
-		// reuse expansion state and default value
-		if(const FRigVMPinInfo* PreviousPin = PreviousPinInfos.GetPinFromPinPath(InPinPath))
-		{
-			if(PreviousPin->TypeIndex == InPinInfo.TypeIndex)
-			{
-				Pin->bIsExpanded = PreviousPin->bIsExpanded;
-				Pin->DefaultValue = PreviousPin->DefaultValue;
-			}
-		}
-
-		if (URigVMPin* ParentPin = Cast<URigVMPin>(InOuter))
-		{
-			AddSubPin(ParentPin, Pin);
-		}
-		else if (URigVMNode* OwnerNode = Cast<URigVMNode>(InOuter))
-		{
-			AddNodePin(OwnerNode, Pin);
-		}
-		else
-		{
-			ensureMsgf(false, TEXT("Outer %s of pin info %s is not a pin or a node"), *InOuter->GetPathName(), *InPinPath);
-		}
-
-		Notify(ERigVMGraphNotifType::PinAdded, Pin);
-
-		return Pin;
-	};
-
 	// step 1/3: keep a record of the current state of the node's pins
 	TMap<FString, FString> RedirectedPinPaths;
 	if (bFollowCoreRedirectors)
@@ -15506,7 +15442,7 @@ void URigVMController::RepopulatePinsOnNode(URigVMNode* InNode, bool bFollowCore
 			OuterForPin = InNode->FindPin(ParentPinPath);
 		}
 		
-		(void)CreatePinFromPinInfo(NewPinInfos[NewPinsToAdd[Index]], PinPath, OuterForPin);
+		CreatePinFromPinInfo(Registry, PreviousPinInfos, NewPinInfos[NewPinsToAdd[Index]], PinPath, OuterForPin);
 #if UE_RIGVMCONTROLLER_VERBOSE_REPOPULATE
 		UE_LOG(LogRigVMDeveloper, Display, TEXT("Adding new pin '%s'."), *PinPath);
 #endif
@@ -16071,6 +16007,67 @@ void URigVMController::ApplyPinStates(URigVMNode* InNode, const TMap<FString, UR
 	}
 }
 
+URigVMPin* URigVMController::CreatePinFromPinInfo(const FRigVMRegistry& InRegistry, const FRigVMPinInfoArray& InPreviousPinInfos, const FRigVMPinInfo& InPinInfo, const FString& InPinPath, UObject* InOuter) const
+{
+	check(InOuter);
+	URigVMPin* Pin = NewObject<URigVMPin>(InOuter, InPinInfo.Name);
+	if (InPinInfo.Property)
+	{
+		ConfigurePinFromProperty(InPinInfo.Property, Pin, InPinInfo.Direction);
+	}
+	else
+	{
+		const FRigVMTemplateArgumentType& Type = InRegistry.GetType(InPinInfo.TypeIndex);
+		Pin->CPPType = Type.CPPType.ToString();
+		Pin->CPPTypeObject = Type.CPPTypeObject;
+		if (Pin->CPPTypeObject)
+		{
+			Pin->CPPTypeObjectPath = *Pin->CPPTypeObject->GetPathName();
+		}
+		if (InRegistry.IsExecuteType(InPinInfo.TypeIndex))
+		{
+			MakeExecutePin(Pin);
+		}
+
+		Pin->Direction = InPinInfo.Direction;
+		Pin->DisplayName = InPinInfo.DisplayName.IsEmpty() ? NAME_None : FName(*InPinInfo.DisplayName);
+		Pin->bIsConstant = InPinInfo.bIsConstant;
+		Pin->bIsDynamicArray = InPinInfo.bIsDynamicArray;
+		Pin->bIsLazy = InPinInfo.bIsLazy;
+		Pin->CustomWidgetName = InPinInfo.CustomWidgetName.IsEmpty() ? NAME_None : FName(*InPinInfo.CustomWidgetName);
+	}
+
+	Pin->bIsExpanded = InPinInfo.bIsExpanded;
+	Pin->DefaultValue = InPinInfo.DefaultValue;
+
+	// reuse expansion state and default value
+	if (const FRigVMPinInfo* PreviousPin = InPreviousPinInfos.GetPinFromPinPath(InPinPath))
+	{
+		if (PreviousPin->TypeIndex == InPinInfo.TypeIndex)
+		{
+			Pin->bIsExpanded = PreviousPin->bIsExpanded;
+			Pin->DefaultValue = PreviousPin->DefaultValue;
+		}
+	}
+
+	if (URigVMPin* ParentPin = Cast<URigVMPin>(InOuter))
+	{
+		AddSubPin(ParentPin, Pin);
+	}
+	else if (URigVMNode* OwnerNode = Cast<URigVMNode>(InOuter))
+	{
+		AddNodePin(OwnerNode, Pin);
+	}
+	else
+	{
+		ensureMsgf(false, TEXT("Outer %s of pin info %s is not a pin or a node"), *InOuter->GetPathName(), *InPinPath);
+	}
+
+	Notify(ERigVMGraphNotifType::PinAdded, Pin);
+
+	return Pin;
+}
+
 void URigVMController::ReportInfo(const FString& InMessage) const
 {
 	if (URigVMGraph* Graph = GetGraph())
diff --git a/Engine/Plugins/Runtime/RigVM/Source/RigVMDeveloper/Private/RigVMModel/RigVMPin.cpp b/Engine/Plugins/Runtime/RigVM/Source/RigVMDeveloper/Private/RigVMModel/RigVMPin.cpp
--- a/Engine/Plugins/Runtime/RigVM/Source/RigVMDeveloper/Private/RigVMModel/RigVMPin.cpp
+++ b/Engine/Plugins/Runtime/RigVM/Source/RigVMDeveloper/Private/RigVMModel/RigVMPin.cpp
@@ -108,6 +108,7 @@ URigVMPin::URigVMPin()
 	, bIsConstant(false)
 	, bRequiresWatch(false)
 	, bIsDynamicArray(false)
+	, bIsLazy(false)
 	, CPPType(FString())
 	, CPPTypeObject(nullptr)
 	, CPPTypeObjectPath(NAME_None)
@@ -501,24 +502,7 @@ bool URigVMPin::IsDynamicArray() const
 
 bool URigVMPin::IsLazy() const
 {
-	// fixed array elements are treated as lazy elements
-	// if the original argument is also marked as lazy
-	if(const URigVMPin* ParentPin = GetParentPin())
-	{
-		if(ParentPin->IsFixedSizeArray())
-		{
-			return ParentPin->IsLazy();
-		}
-	}
-	
-	if(GetDirection() == ERigVMPinDirection::Input)
-	{
-		if(const URigVMNode* Node = GetNode())
-		{
-			return Node->ShouldInputPinComputeLazily(this);
-		}
-	}
-	return false;
+	return bIsLazy;
 }
 
 int32 URigVMPin::GetPinIndex() const
diff --git a/Engine/Plugins/Runtime/RigVM/Source/RigVMDeveloper/Public/RigVMModel/Nodes/RigVMUnitNode.h b/Engine/Plugins/Runtime/RigVM/Source/RigVMDeveloper/Public/RigVMModel/Nodes/RigVMUnitNode.h
--- a/Engine/Plugins/Runtime/RigVM/Source/RigVMDeveloper/Public/RigVMModel/Nodes/RigVMUnitNode.h
+++ b/Engine/Plugins/Runtime/RigVM/Source/RigVMDeveloper/Public/RigVMModel/Nodes/RigVMUnitNode.h
@@ -96,6 +96,9 @@ public:
 
 	virtual uint32 GetStructureHash() const override;
 
+	// allows the node to support non-native pins
+	virtual bool HasNonNativePins() const { return false; }
+
 protected:
 
 	virtual FText GetToolTipTextForPin(const URigVMPin* InPin) const override;
diff --git a/Engine/Plugins/Runtime/RigVM/Source/RigVMDeveloper/Public/RigVMModel/RigVMClient.h b/Engine/Plugins/Runtime/RigVM/Source/RigVMDeveloper/Public/RigVMModel/RigVMClient.h
--- a/Engine/Plugins/Runtime/RigVM/Source/RigVMDeveloper/Public/RigVMModel/RigVMClient.h
+++ b/Engine/Plugins/Runtime/RigVM/Source/RigVMDeveloper/Public/RigVMModel/RigVMClient.h
@@ -89,6 +89,7 @@ public:
 	FRigVMClient()
 		: SchemaPtr(nullptr)
 		, SchemaClass(URigVMSchema::StaticClass())
+		, ControllerClass(URigVMController::StaticClass())
 		, FunctionLibrary(nullptr)
 		, ActionStack(nullptr)
 		, bSuspendNotifications(false)
@@ -99,6 +100,7 @@ public:
 	}
 
 	void SetSchemaClass(TSubclassOf<URigVMSchema> InSchemaClass);
+	void SetControllerClass(TSubclassOf<URigVMController> InControllerClass);
 	void SetOuterClientHost(UObject* InOuterClientHost, const FName& InOuterClientHostPropertyName);
 	void SetFromDeprecatedData(URigVMGraph* InDefaultGraph, URigVMFunctionLibrary* InFunctionLibrary);
 
@@ -206,6 +208,9 @@ private:
 	UPROPERTY(transient)
 	TSubclassOf<URigVMSchema> SchemaClass;
 
+	UPROPERTY(transient)
+	TSubclassOf<URigVMController> ControllerClass;
+
 	UPROPERTY()
 	TArray<TObjectPtr<URigVMGraph>> Models;
 
diff --git a/Engine/Plugins/Runtime/RigVM/Source/RigVMDeveloper/Public/RigVMModel/RigVMController.h b/Engine/Plugins/Runtime/RigVM/Source/RigVMDeveloper/Public/RigVMModel/RigVMController.h
--- a/Engine/Plugins/Runtime/RigVM/Source/RigVMDeveloper/Public/RigVMModel/RigVMController.h
+++ b/Engine/Plugins/Runtime/RigVM/Source/RigVMDeveloper/Public/RigVMModel/RigVMController.h
@@ -188,6 +188,7 @@ struct RIGVMDEVELOPER_API FRigVMPinInfo
 	bool bIsConstant;
 	bool bIsDynamicArray;
 	bool bIsDecorator;
+	bool bIsLazy;
 	TArray<int32> SubPins;
 
 	friend uint32 GetTypeHash(const FRigVMPinInfo& InPin);
@@ -216,12 +217,14 @@ struct RIGVMDEVELOPER_API FRigVMPinInfoArray
 	const FRigVMPinInfo* GetPinFromPinPath(const FString& InPinPath) const;
 	int32 GetRootIndex(const int32 InIndex) const;
 
-	friend uint32 GetTypeHash(const FRigVMPinInfoArray& InPins);
+	friend RIGVMDEVELOPER_API uint32 GetTypeHash(const FRigVMPinInfoArray& InPins);
 
 	mutable TArray<FRigVMPinInfo> Pins;
 	mutable TMap<FString, int32> PinPathLookup;;
 };
 
+RIGVMDEVELOPER_API uint32 GetTypeHash(const FRigVMPinInfoArray& InPins);
+
 /**
  * The Controller is the sole authority to perform changes
  * on the Graph. The Controller itself is stateless.
@@ -1049,37 +1052,37 @@ public:
 	void SendUserFacingNotification(const FString& InMessage, float InDuration = 0.f, const UObject* InSubject = nullptr, const FName& InBrushName = TEXT("MessageLog.Warning")) const;
 
 	template <typename FmtType, typename... Types>
-	void ReportInfof(const FmtType& Fmt, Types... Args)
+	void ReportInfof(const FmtType& Fmt, Types... Args) const
 	{
 		ReportInfo(FString::Printf(Fmt, Args...));
 	}
 
 	template <typename FmtType, typename... Types>
-	void ReportWarningf(const FmtType& Fmt, Types... Args)
+	void ReportWarningf(const FmtType& Fmt, Types... Args) const
 	{
 		ReportWarning(FString::Printf(Fmt, Args...));
 	}
 
 	template <typename FmtType, typename... Types>
-	void ReportErrorf(const FmtType& Fmt, Types... Args)
+	void ReportErrorf(const FmtType& Fmt, Types... Args) const
 	{
 		ReportError(FString::Printf(Fmt, Args...));
 	}
 
 	template <typename FmtType, typename... Types>
-	void ReportAndNotifyInfof(const FmtType& Fmt, Types... Args)
+	void ReportAndNotifyInfof(const FmtType& Fmt, Types... Args) const
 	{
 		ReportAndNotifyInfo(FString::Printf(Fmt, Args...));
 	}
 
 	template <typename FmtType, typename... Types>
-	void ReportAndNotifyWarningf(const FmtType& Fmt, Types... Args)
+	void ReportAndNotifyWarningf(const FmtType& Fmt, Types... Args) const
 	{
 		ReportAndNotifyWarning(FString::Printf(Fmt, Args...));
 	}
 
 	template <typename FmtType, typename... Types>
-	void ReportAndNotifyErrorf(const FmtType& Fmt, Types... Args)
+	void ReportAndNotifyErrorf(const FmtType& Fmt, Types... Args) const
 	{
 		ReportAndNotifyError(FString::Printf(Fmt, Args...));
 	}
@@ -1107,7 +1110,7 @@ private:
 	void AddPinsForStruct(UStruct* InStruct, URigVMNode* InNode, URigVMPin* InParentPin, ERigVMPinDirection InPinDirection, const FString& InDefaultValue, bool bAutoExpandArrays, const FRigVMPinInfoArray* PreviousPins = nullptr);
 	void AddPinsForArray(FArrayProperty* InArrayProperty, URigVMNode* InNode, URigVMPin* InParentPin, ERigVMPinDirection InPinDirection, const TArray<FString>& InDefaultValues, bool bAutoExpandArrays);
 	void AddPinsForTemplate(const FRigVMTemplate* InTemplate, const FRigVMTemplateTypeMap& InPinTypeMap, URigVMNode* InNode);
-	void ConfigurePinFromProperty(FProperty* InProperty, URigVMPin* InOutPin, ERigVMPinDirection InPinDirection = ERigVMPinDirection::Invalid);
+	void ConfigurePinFromProperty(FProperty* InProperty, URigVMPin* InOutPin, ERigVMPinDirection InPinDirection = ERigVMPinDirection::Invalid) const;
 	void ConfigurePinFromPin(URigVMPin* InOutPin, URigVMPin* InPin, bool bCopyDisplayName = false);
 	void ConfigurePinFromArgument(URigVMPin* InOutPin, const FRigVMGraphFunctionArgument& InArgument, bool bCopyDisplayName = false);
 	bool SetPinDefaultValue(URigVMPin* InPin, const FString& InDefaultValue, bool bResizeArrays, bool bSetupUndoRedo, bool bMergeUndoAction, bool bSetValueOnLinkedPins = true);
@@ -1128,6 +1131,12 @@ private:
 	FName AddDecorator(URigVMNode* InNode, UScriptStruct* InDecoratorScriptStruct, const FName& InDecoratorName, const FString& InDefaultValue, int32 InPinIndex = -1, bool bSetupUndoRedo = true);
 	bool RemoveDecorator(URigVMNode* InNode, const FName& InDecoratorName, bool bSetupUndoRedo = true);
 
+protected:
+
+#if WITH_EDITOR
+	URigVMUnitNode* AddUnitNode(UScriptStruct* InScriptStruct, TSubclassOf<URigVMUnitNode> InUnitNodeClass, const FName& InMethodName, const FVector2D& InPosition, const FString& InNodeName, bool bSetupUndoRedo, bool bPrintPythonCommand);
+#endif
+
 public:
 
 	// try to reconnect source and target pins after a node deletion
@@ -1291,6 +1300,8 @@ public:
 
 protected:
 
+	URigVMPin* CreatePinFromPinInfo(const FRigVMRegistry& InRegistry, const FRigVMPinInfoArray& InPreviousPinInfos, const FRigVMPinInfo& InPinInfo, const FString& InPinPath, UObject* InOuter) const;
+
 	// backwards compatibility code
 	FRigVMClientPatchResult PatchRerouteNodesOnLoad();
 	FRigVMClientPatchResult PatchUnitNodesOnLoad();
diff --git a/Engine/Plugins/Runtime/RigVM/Source/RigVMDeveloper/Public/RigVMModel/RigVMPin.h b/Engine/Plugins/Runtime/RigVM/Source/RigVMDeveloper/Public/RigVMModel/RigVMPin.h
--- a/Engine/Plugins/Runtime/RigVM/Source/RigVMDeveloper/Public/RigVMModel/RigVMPin.h
+++ b/Engine/Plugins/Runtime/RigVM/Source/RigVMDeveloper/Public/RigVMModel/RigVMPin.h
@@ -509,6 +509,9 @@ private:
 	UPROPERTY()
 	bool bIsDynamicArray;
 
+	UPROPERTY()
+	bool bIsLazy;
+
 	UPROPERTY()
 	FString CPPType;
 
