Fix static analysis warning by checking to make sure pin property exists before using it.

diff --git a/Engine/Plugins/Experimental/AnimNext/Source/AnimNext/Private/Graph/RigDecorator_AnimNextCppDecorator.cpp b/Engine/Plugins/Experimental/AnimNext/Source/AnimNext/Private/Graph/RigDecorator_AnimNextCppDecorator.cpp
--- a/Engine/Plugins/Experimental/AnimNext/Source/AnimNext/Private/Graph/RigDecorator_AnimNextCppDecorator.cpp
+++ b/Engine/Plugins/Experimental/AnimNext/Source/AnimNext/Private/Graph/RigDecorator_AnimNextCppDecorator.cpp
@@ -26,6 +26,12 @@ void FRigDecorator_AnimNextCppDecorator::GetProgrammaticPins(URigVMController* I
 	{
 		FRigVMPinInfo& PinInfo = OutPinArray[PinIndex];
 
+		if (PinInfo.Property == nullptr)
+		{
+			// This pin doesn't have a property, we'll have to assume that it has been fully specified by the decorator
+			continue;
+		}
+
 		const bool bIsHidden = PinInfo.Property->HasMetaData(FRigVMStruct::HiddenMetaName);
 		const bool bIsInline = PinInfo.Property->HasMetaData("Inline");
 		const bool bIsDecoratorHandle = PinInfo.Property->GetCPPType() == TEXT("FAnimNextDecoratorHandle");
