Implement AnimNext evaluation virtual machine
Add various tasks to sample sequence,s, blend, etc
Add necessary machinery to execute task lists (evaluation programs)
Implement blending operations using the new TransformArrays
Implement anim sequence player decorator
Implement blend two decorator

diff --git a/Engine/Plugins/Experimental/AnimNext/Source/AnimNext/Internal/Graph/AnimNext_LODPose.h b/Engine/Plugins/Experimental/AnimNext/Source/AnimNext/Internal/Graph/AnimNext_LODPose.h
--- a/Engine/Plugins/Experimental/AnimNext/Source/AnimNext/Internal/Graph/AnimNext_LODPose.h
+++ b/Engine/Plugins/Experimental/AnimNext/Source/AnimNext/Internal/Graph/AnimNext_LODPose.h
@@ -29,14 +29,14 @@ struct FAnimNextGraphLODPose
 
 	FAnimNextGraphLODPose() = default;
 
-	explicit FAnimNextGraphLODPose(const UE::AnimNext::FLODPose& InLODPose)
+	explicit FAnimNextGraphLODPose(const UE::AnimNext::FLODPoseHeap& InLODPose)
 		: LODPose(InLODPose)
 	{
 	}
-	explicit FAnimNextGraphLODPose(UE::AnimNext::FLODPose&& InLODPose)
+	explicit FAnimNextGraphLODPose(UE::AnimNext::FLODPoseHeap&& InLODPose)
 		: LODPose(MoveTemp(InLODPose))
 	{
 	}
 
-	UE::AnimNext::FLODPose LODPose;
+	UE::AnimNext::FLODPoseHeap LODPose;
 };
diff --git a/Engine/Plugins/Experimental/AnimNext/Source/AnimNext/Private/DecompressionTools.cpp b/Engine/Plugins/Experimental/AnimNext/Source/AnimNext/Private/DecompressionTools.cpp
--- a/Engine/Plugins/Experimental/AnimNext/Source/AnimNext/Private/DecompressionTools.cpp
+++ b/Engine/Plugins/Experimental/AnimNext/Source/AnimNext/Private/DecompressionTools.cpp
@@ -170,7 +170,7 @@ void FDecompressionTools::GetBonePose(const UAnimSequence* AnimSequence, const F
 				// Virtual bones are part of the retarget transform pose, so if the pose has not been updated (recently) there might be a mismatch
 				if (SkeletonBoneIndex < NumRawBones || AuthoredOnRefSkeleton.IsValidIndex(SkeletonBoneIndex))
 				{
-					OutAnimationPoseData.LocalTransforms[i] = AuthoredOnRefSkeleton[SkeletonBoneIndex];
+					OutAnimationPoseData.LocalTransformsView[i] = AuthoredOnRefSkeleton[SkeletonBoneIndex];
 				}
 			}
 		}
@@ -368,7 +368,7 @@ void FDecompressionTools::DecompressPose(FLODPose& OutAnimationPoseData,
 		{
 			const int32 TrackIndex = 0;
 			const int32 LODRootBone = 0;
-			FTransform RootAtom = OutAnimationPoseData.LocalTransforms[0];
+			FTransform RootAtom = OutAnimationPoseData.LocalTransformsView[0];
 
 			CompressedData.BoneCompressionCodec->DecompressBone(DecompressionContext, TrackIndex, RootAtom);
 
@@ -407,14 +407,14 @@ void FDecompressionTools::DecompressPose(FLODPose& OutAnimationPoseData,
 				, DecompressionContext.IsAdditiveAnimation()
 				, OutAnimationPoseData.GetDisableRetargeting());
 
-			OutAnimationPoseData.LocalTransforms[0] = RootAtom;
+			OutAnimationPoseData.LocalTransformsView[0] = RootAtom;
 		}
 
 		if (RotationScalePairs.Num() > 0)
 		{
 #if DEFAULT_SOA
 			// get the remaining bone atoms
-			CompressedData.BoneCompressionCodec->DecompressPose(DecompressionContext, UE::Anim::FAnimPoseDecompressionData(RotationScalePairs, TranslationPairs, RotationScalePairs, OutAnimationPoseData.LocalTransforms.Rotations, OutAnimationPoseData.LocalTransforms.Translations, OutAnimationPoseData.LocalTransforms.Scales3D));
+			CompressedData.BoneCompressionCodec->DecompressPose(DecompressionContext, UE::Anim::FAnimPoseDecompressionData(RotationScalePairs, TranslationPairs, RotationScalePairs, OutAnimationPoseData.LocalTransformsView.Rotations, OutAnimationPoseData.LocalTransformsView.Translations, OutAnimationPoseData.LocalTransformsView.Scales3D));
 #else
 			// get the remaining bone atoms
 			TArrayView<FTransform> OutPoseBones = OutAnimationPoseData.LocalTransforms.Transforms;
@@ -433,10 +433,10 @@ void FDecompressionTools::DecompressPose(FLODPose& OutAnimationPoseData,
 			for (int32 LODBoneIndex = bFirstTrackIsRootBone ? 1 : 0; LODBoneIndex < LODNumBones; ++LODBoneIndex)
 			{
 				const int32 TargetSkeletonBoneIndex = ReferencePose.GetSkeletonBoneIndexFromLODBoneIndex(LODBoneIndex);
-				OutAnimationPoseData.LocalTransforms[LODBoneIndex].SetRotation(SkeletonRemapping.RetargetAdditiveRotationToTargetSkeleton(TargetSkeletonBoneIndex, OutAnimationPoseData.LocalTransforms[LODBoneIndex].GetRotation()));
+				OutAnimationPoseData.LocalTransformsView[LODBoneIndex].SetRotation(SkeletonRemapping.RetargetAdditiveRotationToTargetSkeleton(TargetSkeletonBoneIndex, OutAnimationPoseData.LocalTransformsView[LODBoneIndex].GetRotation()));
 				if (TargetSkeleton->GetBoneTranslationRetargetingMode(TargetSkeletonBoneIndex, OutAnimationPoseData.GetDisableRetargeting()) != EBoneTranslationRetargetingMode::Skeleton)
 				{
-					OutAnimationPoseData.LocalTransforms[LODBoneIndex].SetTranslation(SkeletonRemapping.RetargetAdditiveTranslationToTargetSkeleton(TargetSkeletonBoneIndex, OutAnimationPoseData.LocalTransforms[LODBoneIndex].GetTranslation()));
+					OutAnimationPoseData.LocalTransformsView[LODBoneIndex].SetTranslation(SkeletonRemapping.RetargetAdditiveTranslationToTargetSkeleton(TargetSkeletonBoneIndex, OutAnimationPoseData.LocalTransformsView[LODBoneIndex].GetTranslation()));
 				}
 			}
 		}
@@ -445,10 +445,10 @@ void FDecompressionTools::DecompressPose(FLODPose& OutAnimationPoseData,
 			for (int32 LODBoneIndex = (bFirstTrackIsRootBone ? 1 : 0); LODBoneIndex < LODNumBones; ++LODBoneIndex)
 			{
 				const int32 TargetSkeletonBoneIndex = LODBoneIndexes[LODBoneIndex]; // ReferencePose.GetSkeletonBoneIndexFromLODBoneIndex(LODBoneIndex);
-				OutAnimationPoseData.LocalTransforms[LODBoneIndex].SetRotation(SkeletonRemapping.RetargetBoneRotationToTargetSkeleton(TargetSkeletonBoneIndex, OutAnimationPoseData.LocalTransforms[LODBoneIndex].GetRotation()));
+				OutAnimationPoseData.LocalTransformsView[LODBoneIndex].SetRotation(SkeletonRemapping.RetargetBoneRotationToTargetSkeleton(TargetSkeletonBoneIndex, OutAnimationPoseData.LocalTransformsView[LODBoneIndex].GetRotation()));
 				if (TargetSkeleton->GetBoneTranslationRetargetingMode(TargetSkeletonBoneIndex, OutAnimationPoseData.GetDisableRetargeting()) != EBoneTranslationRetargetingMode::Skeleton)
 				{
-					OutAnimationPoseData.LocalTransforms[LODBoneIndex].SetTranslation(SkeletonRemapping.RetargetBoneTranslationToTargetSkeleton(TargetSkeletonBoneIndex, OutAnimationPoseData.LocalTransforms[LODBoneIndex].GetTranslation()));
+					OutAnimationPoseData.LocalTransformsView[LODBoneIndex].SetTranslation(SkeletonRemapping.RetargetBoneTranslationToTargetSkeleton(TargetSkeletonBoneIndex, OutAnimationPoseData.LocalTransformsView[LODBoneIndex].GetTranslation()));
 				}
 			}
 		}
@@ -476,7 +476,7 @@ void FDecompressionTools::DecompressPose(FLODPose& OutAnimationPoseData,
 			if (SourceTranslationLength > UE_KINDA_SMALL_NUMBER)
 			{
 				float const TargetTranslationLength = ReferencePose.GetRefPoseTranslation(LODBoneIndex).Size();
-				OutAnimationPoseData.LocalTransforms[LODBoneIndex].ScaleTranslation(TargetTranslationLength / SourceTranslationLength);
+				OutAnimationPoseData.LocalTransformsView[LODBoneIndex].ScaleTranslation(TargetTranslationLength / SourceTranslationLength);
 			}
 		}
 	}
@@ -503,10 +503,10 @@ void FDecompressionTools::DecompressPose(FLODPose& OutAnimationPoseData,
 			}
 
 			// Apply the retargeting as if it were an additive difference between the current skeleton and the retarget skeleton. 
-			OutAnimationPoseData.LocalTransforms[LODBoneIndex].SetRotation(OutAnimationPoseData.LocalTransforms[LODBoneIndex].GetRotation() * BaseTransform.GetRotation().Inverse() * RefPoseTransform.GetRotation());
-			OutAnimationPoseData.LocalTransforms[LODBoneIndex].SetTranslation(OutAnimationPoseData.LocalTransforms[LODBoneIndex].GetTranslation() + (RefPoseTransform.GetTranslation() - BaseTransform.GetTranslation()));
-			OutAnimationPoseData.LocalTransforms[LODBoneIndex].SetScale3D(OutAnimationPoseData.LocalTransforms[LODBoneIndex].GetScale3D() * (RefPoseTransform.GetScale3D() * BaseTransform.GetSafeScaleReciprocal(BaseTransform.GetScale3D())));
-			OutAnimationPoseData.LocalTransforms[LODBoneIndex].NormalizeRotation();
+			OutAnimationPoseData.LocalTransformsView[LODBoneIndex].SetRotation(OutAnimationPoseData.LocalTransformsView[LODBoneIndex].GetRotation() * BaseTransform.GetRotation().Inverse() * RefPoseTransform.GetRotation());
+			OutAnimationPoseData.LocalTransformsView[LODBoneIndex].SetTranslation(OutAnimationPoseData.LocalTransformsView[LODBoneIndex].GetTranslation() + (RefPoseTransform.GetTranslation() - BaseTransform.GetTranslation()));
+			OutAnimationPoseData.LocalTransformsView[LODBoneIndex].SetScale3D(OutAnimationPoseData.LocalTransformsView[LODBoneIndex].GetScale3D() * (RefPoseTransform.GetScale3D() * BaseTransform.GetSafeScaleReciprocal(BaseTransform.GetScale3D())));
+			OutAnimationPoseData.LocalTransformsView[LODBoneIndex].NormalizeRotation();
 		}
 	}
 
diff --git a/Engine/Plugins/Experimental/AnimNext/Source/AnimNext/Private/DecoratorBase/ExecutionContext.cpp b/Engine/Plugins/Experimental/AnimNext/Source/AnimNext/Private/DecoratorBase/ExecutionContext.cpp
--- a/Engine/Plugins/Experimental/AnimNext/Source/AnimNext/Private/DecoratorBase/ExecutionContext.cpp
+++ b/Engine/Plugins/Experimental/AnimNext/Source/AnimNext/Private/DecoratorBase/ExecutionContext.cpp
@@ -167,6 +167,67 @@ namespace UE::AnimNext
 			return false;
 		}
 
+		// TODO: Revisit this to avoid querying each decorator over and over
+		//
+		// If we have D decorators on a node and I unique interfaces per decorator then finding an interface
+		// has O(D*I) complexity: for each decorator we have to test every interface
+		// However, if the I interfaces are not unique per decorator and instead are unique for the system,
+		// the below code will still have the same complexity. We'll end up testing the same interface multiple
+		// times, once for each decorator. Because we will have a small set of interfaces per node that
+		// multiple decorators implement (e.g IUpdate, IEvaluate), we can do better.
+		// 
+		// It would be much cheaper if instead we could test if we implement the interface first and then
+		// look for a list of decorators (in stack order) that implement it. To do this, when we build the node
+		// template, we could aggregate all interfaces that the decorators implement (sorted by their interface id
+		// for determinism or perhaps some other criteria to put popular/hot interfaces first). Then, in the node
+		// template we store a mapping of InterfaceUID to InterfaceIndex (local to that node template). The interface
+		// index can then be used to index into a list of decorator indices that implement it.
+		// Code would look like:
+		//    * uint32 InterfaceIndex = NodeTemplate.GetInterfaceUIDs().FindIndex(InterfaceUID);	// INDEX_NONE means that none of the decorators implement the interface
+		//    * uint16 DecoratorIndicesStartOffset = NodeTemplate.InterfaceIndexToDecoratorIndices(InterfaceIndex);		// An offset relative to the NodeTemplate*
+		//    * const uint8* DecoratorIndicesForInterface = ((const uint8*NodeTemplate) + DecoratorIndicesStartOffset;
+		//    * uint8 NumDecoratorsWithInterface = DecoratorIndicesForInterface[0];		// Reserve first index for the decorator count or if we use a guard value at the end of the array we have to reserve a value
+		//    * uint8 TopDecoratorIndex = DecoratorIndicesForInterface[1];				// First index is top of stack
+		// Super decorator lookup would be similar:
+		//    * uint8 SuperDecoratorIndex = DecoratorIndicesForInterface[DecoratorIndicesForInterface.FindIndex(CurrentDecoratorIndex) + 1];	// Need to check if we exceed the max count (bottom of stack)
+		// Once we have the decorator index, we use the decorator template as we do now
+		//
+		// The above would have a number of advantages:
+		//    * Querying for an interface that the node doesn't implement would be much cheaper (we test each interface once)
+		//    * InterfaceIndex query can easily be implemented in SIMD, we look for a uint32 in a list sequentially
+		//    * Bulk/interleaved querying of interfaces would be much easier (e.g. query for IEvaluate for 4 nodes)
+		//    * A lot less branching is involved, improving throughput
+		//    * We can handle multiple base decorators by checking the decorator index found which has comparable cost
+		//    * We could store the interface offset from the decorator ptr base to avoid the call to GetInterface()
+		//      Instead of testing to find our interface again to return a custom static_cast, we'd store the interface/decorator offset
+		//      in the node template alongside the decorator index. We can then query for the FDecorator* and add the offset
+		//      improving bulk query feasibility.
+		//
+		// However, querying for a Super interface might be slower since we have to start the search from the start.
+		// In contrast, we now start searching at the current decorator. In order to keep that cost down and amortize
+		// the search, we would have to cache the InterfaceIndex and the current index in the decorator list.
+		// We could use 8-bits for each entry allowing for a max of 256 interfaces per node and 256 decorators per node (existing limitation).
+		//
+		// Adding new data to the decorator binding isn't ideal. It's current size on 64-bit systems is:
+		//    * 8 bytes for interface pointer (to forward function calls to)
+		//    * 8 bytes for decorator template (to get the offsets for shared/instance/latent data)
+		//    * 8 bytes for node description pointer (base of shared data)
+		//    * 8 bytes for node instance pointer (base of instance data, in weak decorator ptr)
+		//    * 4 bytes for decorator index (in weak decorator ptr)
+		//    * 4 bytes of padding (in weak decorator ptr)
+		// We current use 4 bytes for the decorator index but in practice we only need 1 byte. We use 4 bytes because we pay
+		// for padding anyway. We could use 2 extra bytes for our mapping and still have 5 bytes to space in padding (due to alignment).
+		// In practice, decorator bindings always live on the stack and they shouldn't be getting copied/moved around.
+		// We populate them in place when we query, then we use it from its storage on the stack that we first populated.
+		//
+		// To be able to support all of this, instead of having FDecorator::GetInterface() implemented by derived types
+		// we would need instead to implement FDecorator::GetInterfaceUIDs(). Then on load, when we build a node template and
+		// finalize it, we can build the list of unique interface UIDs and sort it and we can build the mapping structures.
+		// On load, as long as we fit within 64KB we are fine. We have to update FNodeTemplate::Finalize and FNodeTemplateBuilder::BuildNodeTemplate.
+		// Perhaps the builder should use the same code and build in a buffer on the stack and copy the final size out.
+		// We must build the mapping on load because the interfaces we implement are only known at runtime (e.g. can be changed through defines/configuration).
+		// We can still implement GetInterface() in terms of GetInterfaceUIDs if it returns the interface offsets as well.
+
 		FNodeInstance* NodeInstance = DecoratorPtr.GetNodeInstance();
 
 		const FNodeDescription& NodeDesc = GetNodeDescription(NodeInstance->GetNodeHandle());
diff --git a/Engine/Plugins/Experimental/AnimNext/Source/AnimNext/Private/DecoratorInterfaces/IEvaluate.cpp b/Engine/Plugins/Experimental/AnimNext/Source/AnimNext/Private/DecoratorInterfaces/IEvaluate.cpp
--- a/Engine/Plugins/Experimental/AnimNext/Source/AnimNext/Private/DecoratorInterfaces/IEvaluate.cpp
+++ b/Engine/Plugins/Experimental/AnimNext/Source/AnimNext/Private/DecoratorInterfaces/IEvaluate.cpp
@@ -7,10 +7,36 @@
 
 namespace UE::AnimNext
 {
-	FEvaluateTraversalContext::FEvaluateTraversalContext(EEvaluationFlags EvaluationFlags_)
-		: EvaluationFlags(EvaluationFlags_)
+	/**
+	 * FScopedEvaluationProgram
+	 *
+	 * Pushes and pops a evaluation program onto the evaluation traversal context.
+	 * During the lifetime of an instance of this object, the traversal context will
+	 * return the provided evaluation program as the current evaluation program.
+	 *
+	 * @see FEvaluateTraversalContext
+	 */
+	struct ANIMNEXT_API FScopedEvaluationProgram final
 	{
-	}
+		FScopedEvaluationProgram(FEvaluateTraversalContext& InTraversalContext, FEvaluationProgram& InEvaluationProgram)
+			: TraversalContext(InTraversalContext)
+			, OldEvaluationProgram(InTraversalContext.EvaluationProgram)
+		{
+			InTraversalContext.EvaluationProgram = &InEvaluationProgram;
+		}
+
+		~FScopedEvaluationProgram()
+		{
+			TraversalContext.EvaluationProgram = OldEvaluationProgram;
+		}
+
+		FScopedEvaluationProgram(const FScopedEvaluationProgram&) = delete;
+		FScopedEvaluationProgram& operator=(const FScopedEvaluationProgram&) = delete;
+
+	private:
+		FEvaluateTraversalContext& TraversalContext;
+		FEvaluationProgram* OldEvaluationProgram;
+	};
 
 	void IEvaluate::PreEvaluate(FExecutionContext& Context, const TDecoratorBinding<IEvaluate>& Binding) const
 	{
@@ -42,13 +68,13 @@ namespace UE::AnimNext
 		EEvaluateStep		DesiredStep = EEvaluateStep::PreEvaluate;
 	};
 
-	void EvaluateGraph(FExecutionContext& Context, FWeakDecoratorPtr GraphRootPtr, EEvaluationFlags EvaluationFlags, FPoseContainer& OutContainer)
+	FEvaluationProgram EvaluateGraph(FExecutionContext& Context, FEvaluateTraversalContext& TraversalContext, FWeakDecoratorPtr GraphRootPtr)
 	{
-		// TODO: Reset container
+		FEvaluationProgram EvaluationProgram;
 
 		if (!GraphRootPtr.IsValid())
 		{
-			return;	// Nothing to update
+			return EvaluationProgram;	// Nothing to update
 		}
 
 		FMemMark Mark(FMemStack::Get());
@@ -59,8 +85,7 @@ namespace UE::AnimNext
 		FChildrenArray Children;
 		Children.Reserve(64);
 
-		FEvaluateTraversalContext TraversalContext(EvaluationFlags);
-
+		FScopedEvaluationProgram ScopedEvaluationProgram(TraversalContext, EvaluationProgram);
 		FScopedTraversalContext ScopedTraversalContext(Context, TraversalContext);
 
 		// Add the graph root to kick start the evaluation process
@@ -116,5 +141,7 @@ namespace UE::AnimNext
 				// It is either a sibling read for its post-evaluate or our parent entry ready for its post-evaluate
 			}
 		}
+
+		return EvaluationProgram;
 	}
 }
diff --git a/Engine/Plugins/Experimental/AnimNext/Source/AnimNext/Private/DecoratorInterfaces/IUpdate.cpp b/Engine/Plugins/Experimental/AnimNext/Source/AnimNext/Private/DecoratorInterfaces/IUpdate.cpp
--- a/Engine/Plugins/Experimental/AnimNext/Source/AnimNext/Private/DecoratorInterfaces/IUpdate.cpp
+++ b/Engine/Plugins/Experimental/AnimNext/Source/AnimNext/Private/DecoratorInterfaces/IUpdate.cpp
@@ -42,7 +42,7 @@ namespace UE::AnimNext
 		EUpdateStep			DesiredStep = EUpdateStep::PreUpdate;
 	};
 
-	void UpdateGraph(FExecutionContext& Context, FWeakDecoratorPtr GraphRootPtr, float DeltaTime)
+	void UpdateGraph(FExecutionContext& Context, FUpdateTraversalContext& TraversalContext, FWeakDecoratorPtr GraphRootPtr)
 	{
 		if (!GraphRootPtr.IsValid())
 		{
@@ -57,8 +57,6 @@ namespace UE::AnimNext
 		FChildrenArray Children;
 		Children.Reserve(64);
 
-		FUpdateTraversalContext TraversalContext(DeltaTime);
-
 		FScopedTraversalContext ScopedTraversalContext(Context, TraversalContext);
 
 		// Add the graph root to kick start the update process
diff --git a/Engine/Plugins/Experimental/AnimNext/Source/AnimNext/Private/Decorators/BlendTwoWay.cpp b/Engine/Plugins/Experimental/AnimNext/Source/AnimNext/Private/Decorators/BlendTwoWay.cpp
--- a/Engine/Plugins/Experimental/AnimNext/Source/AnimNext/Private/Decorators/BlendTwoWay.cpp
+++ b/Engine/Plugins/Experimental/AnimNext/Source/AnimNext/Private/Decorators/BlendTwoWay.cpp
@@ -3,6 +3,7 @@
 #include "Decorators/BlendTwoWay.h"
 
 #include "DecoratorBase/ExecutionContext.h"
+#include "EvaluationVM/Tasks/BlendKeyframes.h"
 
 namespace UE::AnimNext
 {
@@ -16,14 +17,22 @@ namespace UE::AnimNext
 
 	void FBlendTwoWayDecorator::PostEvaluate(FExecutionContext& Context, const TDecoratorBinding<IEvaluate>& Binding) const
 	{
-		// TODO:
-		// Now that our children have finished evaluating, we can pop their two poses and blend it
-		// Children execute in depth first order, as such the poses need to be popped in reverse order
-		// FPose&& pose1 = Context.PopPose();
-		// FPose&& pose0 = Context.PopPose();
-		// FPose resultPose = BlendTwoWay(pose0, pose1, blendWeight);
-		// Context.PushPose(resultPose);
-		// If we have full weight on a child, there is no work to do and we can early exit instead
+		const FInstanceData* InstanceData = Binding.GetInstanceData<FInstanceData>();
+
+		if (InstanceData->ChildA.IsValid() && InstanceData->ChildB.IsValid())
+		{
+			// We have two children, interpolate them
+			const FSharedData* SharedData = Binding.GetSharedData<FSharedData>();
+
+			const float BlendWeight = SharedData->GetBlendWeight(Context, Binding);
+
+			FEvaluateTraversalContext& TraversalContext = Context.GetTraversalContext<FEvaluateTraversalContext>();
+			TraversalContext.AppendTask(FAnimNextBlendTwoKeyframesTask::Make(BlendWeight));
+		}
+		else
+		{
+			// We have only one child that is active, do nothing
+		}
 	}
 
 	void FBlendTwoWayDecorator::PreUpdate(FExecutionContext& Context, const TDecoratorBinding<IUpdate>& Binding) const
diff --git a/Engine/Plugins/Experimental/AnimNext/Source/AnimNext/Private/Decorators/SequencePlayer.cpp b/Engine/Plugins/Experimental/AnimNext/Source/AnimNext/Private/Decorators/SequencePlayer.cpp
--- a/Engine/Plugins/Experimental/AnimNext/Source/AnimNext/Private/Decorators/SequencePlayer.cpp
+++ b/Engine/Plugins/Experimental/AnimNext/Source/AnimNext/Private/Decorators/SequencePlayer.cpp
@@ -3,9 +3,7 @@
 #include "Decorators/SequencePlayer.h"
 
 #include "DecoratorBase/ExecutionContext.h"
-#include "Param/ParamStack.h"
-#include "Graph/AnimNext_LODPose.h"
-#include "DecompressionTools.h"
+#include "EvaluationVM/Tasks/PushAnimSequenceKeyframe.h"
 
 namespace UE::AnimNext
 {
@@ -33,35 +31,13 @@ namespace UE::AnimNext
 		const FSharedData* SharedData = Binding.GetSharedData<FSharedData>();
 		const FInstanceData* InstanceData = Binding.GetInstanceData<FInstanceData>();
 
-		// TODO: Sample pose
-		// FPose foo = ...
-		// Context.PushPose(foo);
+		const bool bInterpolate = true;
 
-		FParamStack& ParamStack = FParamStack::Get();
+		FAnimNextAnimSequenceKeyframeTask Task = FAnimNextAnimSequenceKeyframeTask::MakeFromSampleTime(SharedData->AnimSequence, InstanceData->InternalTimeAccumulator, bInterpolate);
+		Task.bExtractTrajectory = false;	/*Output.AnimInstanceProxy->ShouldExtractRootMotion()*/
 
-		FDeltaTimeRecord DeltaTimeRecord;
-		DeltaTimeRecord.Set(InstanceData->PrevInternalTimeAccumulator, InstanceData->InternalTimeAccumulator);
-
-		const FAnimExtractContext ExtractionContext(static_cast<double>(InstanceData->InternalTimeAccumulator)
-			, false /*Output.AnimInstanceProxy->ShouldExtractRootMotion()*/
-			, DeltaTimeRecord
-			, SharedData->GetbLoop(Context, Binding));
-
-		const FAnimNextGraphReferencePose& GraphReferencePose = ParamStack.GetParam<FAnimNextGraphReferencePose>("GraphReferencePose");
-		const int32 GraphLODLevel = ParamStack.GetParam<int32>("GraphLODLevel");
-		const bool bGraphExpectsAdditive = ParamStack.GetParam<bool>("GraphExpectsAdditive");
-
-		// HACK: Write to our output directly, only works if we have a single node in the graph that outputs a pose
-		FAnimNextGraphLODPose& ResultPose = ParamStack.GetMutableParam<FAnimNextGraphLODPose>("ResultPose");
-
-		ResultPose.LODPose.PrepareForLOD(*GraphReferencePose.ReferencePose, GraphLODLevel, true, bGraphExpectsAdditive);
-
-		// Note : calling FDecompressionTools instead of UAnimSequence / UAnimSequenceBase, as I can not modify the engine code (so I extracted the function)
-		// TODO : Revisit this, so we can plug other sequence types
-		if (SharedData->AnimSequence != nullptr)
-		{
-			FDecompressionTools::GetAnimationPose(SharedData->AnimSequence, ResultPose.LODPose, ExtractionContext);
-		}
+		FEvaluateTraversalContext& TraversalContext = Context.GetTraversalContext<FEvaluateTraversalContext>();
+		TraversalContext.AppendTask(Task);
 	}
 
 	double FSequencePlayerDecorator::GetPlayRate(FExecutionContext& Context, const TDecoratorBinding<ITimeline>& Binding) const
diff --git a/Engine/Plugins/Experimental/AnimNext/Source/AnimNext/Private/EvaluationVM/EvaluationProgram.cpp b/Engine/Plugins/Experimental/AnimNext/Source/AnimNext/Private/EvaluationVM/EvaluationProgram.cpp
new file mode 100644
--- /dev/null
+++ b/Engine/Plugins/Experimental/AnimNext/Source/AnimNext/Private/EvaluationVM/EvaluationProgram.cpp
@@ -0,0 +1,25 @@
+// Copyright Epic Games, Inc. All Rights Reserved.
+
+#include "EvaluationVM/EvaluationProgram.h"
+
+#include "EvaluationVM/EvaluationTask.h"
+#include "EvaluationVM/EvaluationVM.h"
+
+namespace UE::AnimNext
+{
+	bool FEvaluationProgram::IsEmpty() const
+	{
+		return Tasks.IsEmpty();
+	}
+
+	void FEvaluationProgram::Execute(FEvaluationVM& VM) const
+	{
+		if (ensure(VM.IsValid()))
+		{
+			for (const TUniquePtr<FAnimNextEvaluationTask>& Task : Tasks)
+			{
+				Task->Execute(VM);
+			}
+		}
+	}
+}
diff --git a/Engine/Plugins/Experimental/AnimNext/Source/AnimNext/Private/EvaluationVM/EvaluationVM.cpp b/Engine/Plugins/Experimental/AnimNext/Source/AnimNext/Private/EvaluationVM/EvaluationVM.cpp
new file mode 100644
--- /dev/null
+++ b/Engine/Plugins/Experimental/AnimNext/Source/AnimNext/Private/EvaluationVM/EvaluationVM.cpp
@@ -0,0 +1,196 @@
+// Copyright Epic Games, Inc. All Rights Reserved.
+
+#include "EvaluationVM/EvaluationVM.h"
+
+#include "Animation/Skeleton.h"
+#include "BoneContainer.h"
+#include "EvaluationVM/KeyframeState.h"
+
+namespace UE::AnimNext
+{
+	const FEvaluationVMStackName KEYFRAME_STACK_NAME = FName(TEXT("KeyframeStack"));
+
+	FEvaluationVMStack::~FEvaluationVMStack()
+	{
+		FEvaluationVMStackEntry* EntryPtr = Top;
+		while (EntryPtr != nullptr)
+		{
+			if (TypeDestructor)
+			{
+				TypeDestructor(EntryPtr->GetValuePtr());
+			}
+
+			FEvaluationVMStackEntry* PrevEntryPtr = EntryPtr->Prev;
+
+			FMemory::Free(EntryPtr);
+
+			EntryPtr = PrevEntryPtr;
+		}
+	}
+
+	FEvaluationVM::FEvaluationVM(EEvaluationFlags InEvaluationFlags, const UE::AnimNext::FReferencePose& InReferencePose, int32 InCurrentLOD)
+		: ReferencePose(&InReferencePose)
+		, CurrentLOD(InCurrentLOD)
+		, EvaluationFlags(InEvaluationFlags)
+	{
+		const UE::Anim::FCurveFilterSettings CurveFilterSettings(UE::Anim::ECurveFilterMode::DisallowAll);
+		USkeleton* Skeleton = const_cast<USkeleton*>(InReferencePose.Skeleton.Get());	// const_cast because the bone container takes a mutable reference
+		BoneContainer.InitializeTo(InReferencePose.GetLODBoneIndexes(InCurrentLOD), CurveFilterSettings, *Skeleton);
+
+		// TODO: In AnimInstanceProxy this is how we initialize the bone container, we need to get the component somehow or we
+		// gotta figure out how to support ref pose overrides
+#if 0
+		// Use the shared bone container
+		RequiredBones = Component->GetSharedRequiredBones();
+
+		// The first anim instance will initialize the required bones, all others will re-use it
+		if (!RequiredBones->IsValid())
+		{
+			RequiredBones->InitializeTo(Component->RequiredBones, Component->GetCurveFilterSettings(), *Asset);
+
+			// If there is a ref pose override, we want to replace ref pose in RequiredBones
+			// Update ref pose in required bones structure (either set it, or clear it, depending on if one is set on the Component)
+			RequiredBones->SetRefPoseOverride(Component->GetRefPoseOverride());
+		}
+#endif
+	}
+
+	bool FEvaluationVM::IsValid() const
+	{
+		return ReferencePose != nullptr;
+	}
+
+	void FEvaluationVM::Shrink()
+	{
+		InternalStacks.Shrink();
+	}
+
+	FKeyframeState FEvaluationVM::MakeReferenceKeyframe(bool bAdditiveKeyframe) const
+	{
+		// TODO: It would be great if we could support immutable poses that we can push/pop like normal mutable poses
+		// This would allow us to cheaply push reference/identity poses
+		// Tasks that consume poses can then choose to re-use a mutable pose (to avoid allocating a new one) and only
+		// allocate one if all their inputs are immutable.
+		// This would reduce the number of required intermediate poses.
+
+		FKeyframeState Keyframe;
+
+		if (EnumHasAnyFlags(EvaluationFlags, EEvaluationFlags::Bones))
+		{
+			const bool bInitWithRefPose = true;
+
+			Keyframe.Pose.PrepareForLOD(*ReferencePose, CurrentLOD, bInitWithRefPose, bAdditiveKeyframe);
+		}
+
+		if (EnumHasAnyFlags(EvaluationFlags, EEvaluationFlags::Curves))
+		{
+			Keyframe.Curves.InitFrom(BoneContainer);
+		}
+
+		return Keyframe;
+	}
+
+	FKeyframeState FEvaluationVM::MakeUninitializedKeyframe(bool bAdditiveKeyframe) const
+	{
+		FKeyframeState Keyframe;
+
+		if (EnumHasAnyFlags(EvaluationFlags, EEvaluationFlags::Bones))
+		{
+			const bool bInitWithRefPose = false;
+
+			Keyframe.Pose.PrepareForLOD(*ReferencePose, CurrentLOD, bInitWithRefPose, bAdditiveKeyframe);
+		}
+
+		if (EnumHasAnyFlags(EvaluationFlags, EEvaluationFlags::Curves))
+		{
+			Keyframe.Curves.InitFrom(BoneContainer);
+		}
+
+		return Keyframe;
+	}
+
+	const FBoneContainer& FEvaluationVM::GetBoneContainer() const
+	{
+		return BoneContainer;
+	}
+
+	FEvaluationVMStack& FEvaluationVM::GetOrCreateStack(const FEvaluationVMStackName& StackName, uint32 TypeID)
+	{
+		if (FEvaluationVMStack* Stack = InternalStacks.FindByHash(StackName.NameHash, StackName.Name))
+		{
+			return *Stack;
+		}
+
+		FEvaluationVMStack& Stack = InternalStacks.AddByHash(StackName.NameHash, StackName.Name, FEvaluationVMStack());
+		Stack.Name = StackName.Name;
+		Stack.TypeID = TypeID;
+
+		return Stack;
+	}
+
+	FEvaluationVMStack* FEvaluationVM::FindStack(const FEvaluationVMStackName& StackName, uint32 TypeID)
+	{
+		FEvaluationVMStack* Stack = InternalStacks.FindByHash(StackName.NameHash, StackName.Name);
+		if (Stack == nullptr)
+		{
+			return nullptr;
+		}
+
+		if (!ensureMsgf(Stack->TypeID == TypeID, TEXT("Type mismatch! This evaluation stack is being queried with a different type than it was created with")))
+		{
+			return nullptr;
+		}
+
+		return Stack;
+	}
+
+	const FEvaluationVMStack* FEvaluationVM::FindStack(const FEvaluationVMStackName& StackName, uint32 TypeID) const
+	{
+		const FEvaluationVMStack* Stack = InternalStacks.FindByHash(StackName.NameHash, StackName.Name);
+		if (Stack == nullptr)
+		{
+			return nullptr;
+		}
+
+		if (!ensureMsgf(Stack->TypeID == TypeID, TEXT("Type mismatch! This evaluation stack is being queried with a different type than it was created with")))
+		{
+			return nullptr;
+		}
+
+		return Stack;
+	}
+
+	const void* FEvaluationVM::PeekValueImpl(const FEvaluationVMStackName& StackName, uint32 TypeID, uint32 Offset) const
+	{
+		const FEvaluationVMStack* Stack = FindStack(StackName, TypeID);
+		if (Stack == nullptr)
+		{
+			// Stack not found
+			return nullptr;
+		}
+
+		const FEvaluationVMStackEntry* EntryPtr = Stack->Top;
+		if (EntryPtr == nullptr)
+		{
+			// Stack is empty
+			return nullptr;
+		}
+
+		uint32 Count = 0;
+		while (EntryPtr != nullptr && Count < Offset)
+		{
+			EntryPtr = EntryPtr->Prev;
+			Count++;
+		}
+
+		if (Count != Offset)
+		{
+			// Offset is too far, not enough entries
+			check(EntryPtr == nullptr);
+			return nullptr;
+		}
+
+		check(EntryPtr != nullptr);
+		return EntryPtr->GetValuePtr();
+	}
+}
diff --git a/Engine/Plugins/Experimental/AnimNext/Source/AnimNext/Private/EvaluationVM/Tasks/ApplyAdditiveKeyframe.cpp b/Engine/Plugins/Experimental/AnimNext/Source/AnimNext/Private/EvaluationVM/Tasks/ApplyAdditiveKeyframe.cpp
new file mode 100644
--- /dev/null
+++ b/Engine/Plugins/Experimental/AnimNext/Source/AnimNext/Private/EvaluationVM/Tasks/ApplyAdditiveKeyframe.cpp
@@ -0,0 +1,71 @@
+// Copyright Epic Games, Inc. All Rights Reserved.
+
+#include "EvaluationVM/Tasks/ApplyAdditiveKeyframe.h"
+
+#include "EvaluationVM/EvaluationVM.h"
+#include "EvaluationVM/KeyframeState.h"
+#include "TransformArrayOperations.h"
+
+FAnimNextApplyAdditiveKeyframeTask FAnimNextApplyAdditiveKeyframeTask::Make(float BlendWeight)
+{
+	FAnimNextApplyAdditiveKeyframeTask Task;
+	Task.BlendWeight = BlendWeight;
+
+	return Task;
+}
+
+void FAnimNextApplyAdditiveKeyframeTask::Execute(UE::AnimNext::FEvaluationVM& VM) const
+{
+	using namespace UE::AnimNext;
+
+	// Pop our top two poses, we'll re-use the top keyframe for our result
+
+	TUniquePtr<FKeyframeState> AdditiveKeyframe;
+	if (!VM.PopValue(KEYFRAME_STACK_NAME, AdditiveKeyframe))
+	{
+		// We have no inputs, nothing to do
+		return;
+	}
+
+	TUniquePtr<FKeyframeState> BaseKeyframe;
+	if (!VM.PopValue(KEYFRAME_STACK_NAME, BaseKeyframe))
+	{
+		// We have a single input, discard it since it must be the additive pose, either way something went wrong
+		// Push the reference pose since we'll expect a non-additive pose
+		VM.PushValue(KEYFRAME_STACK_NAME, MakeUnique<FKeyframeState>(VM.MakeReferenceKeyframe(false)));
+		return;
+	}
+
+	if (EnumHasAnyFlags(VM.GetFlags(), EEvaluationFlags::Bones))
+	{
+		check(BaseKeyframe->Pose.GetNumBones() == AdditiveKeyframe->Pose.GetNumBones());
+
+		const FTransformArrayView BaseTransformsView = BaseKeyframe->Pose.LocalTransforms.GetView();
+
+#if 0
+		// TODO: Read the additive type from the pose itself
+		if (AdditiveType == AAT_RotationOffsetMeshSpace)
+		{
+			AccumulateMeshSpaceRotationAdditiveToLocalPoseInternal(BaseAnimationPoseData.GetPose(), AdditiveAnimationPoseData.GetPose(), Weight);
+		}
+		else
+#endif
+		{
+			BlendWithIdentityAndAccumulate(BaseTransformsView, AdditiveKeyframe->Pose.LocalTransforms.GetConstView(), BlendWeight);
+		}
+
+		NormalizeRotations(BaseTransformsView);
+	}
+
+	if (EnumHasAnyFlags(VM.GetFlags(), EEvaluationFlags::Curves))
+	{
+		BaseKeyframe->Curves.Accumulate(AdditiveKeyframe->Curves, BlendWeight);
+	}
+
+	if (EnumHasAnyFlags(VM.GetFlags(), EEvaluationFlags::Attributes))
+	{
+		UE::Anim::Attributes::AccumulateAttributes(AdditiveKeyframe->Attributes, BaseKeyframe->Attributes, BlendWeight, AAT_None);
+	}
+
+	VM.PushValue(KEYFRAME_STACK_NAME, MoveTemp(BaseKeyframe));
+}
diff --git a/Engine/Plugins/Experimental/AnimNext/Source/AnimNext/Private/EvaluationVM/Tasks/BlendKeyframes.cpp b/Engine/Plugins/Experimental/AnimNext/Source/AnimNext/Private/EvaluationVM/Tasks/BlendKeyframes.cpp
new file mode 100644
--- /dev/null
+++ b/Engine/Plugins/Experimental/AnimNext/Source/AnimNext/Private/EvaluationVM/Tasks/BlendKeyframes.cpp
@@ -0,0 +1,160 @@
+// Copyright Epic Games, Inc. All Rights Reserved.
+
+#include "EvaluationVM/Tasks/BlendKeyframes.h"
+
+#include "EvaluationVM/EvaluationVM.h"
+#include "EvaluationVM/KeyframeState.h"
+#include "TransformArrayOperations.h"
+
+FAnimNextBlendTwoKeyframesTask FAnimNextBlendTwoKeyframesTask::Make(float InterpolationAlpha)
+{
+	FAnimNextBlendTwoKeyframesTask Task;
+	Task.InterpolationAlpha = InterpolationAlpha;
+
+	return Task;
+}
+
+void FAnimNextBlendTwoKeyframesTask::Execute(UE::AnimNext::FEvaluationVM& VM) const
+{
+	using namespace UE::AnimNext;
+
+	// Pop our top two poses, we'll re-use the top keyframe for our result
+
+	TUniquePtr<FKeyframeState> KeyframeB;
+	if (!VM.PopValue(KEYFRAME_STACK_NAME, KeyframeB))
+	{
+		// We have no inputs, nothing to do
+		return;
+	}
+
+	TUniquePtr<FKeyframeState> KeyframeA;
+	if (!VM.PopValue(KEYFRAME_STACK_NAME, KeyframeA))
+	{
+		// We have a single input, leave it on top of the stack
+		VM.PushValue(KEYFRAME_STACK_NAME, MoveTemp(KeyframeB));
+		return;
+	}
+
+	const float Alpha = FMath::Clamp(InterpolationAlpha, 0.0f, 1.0f);
+	const float WeightOfPoseA = 1.0f - Alpha;
+	const float WeightOfPoseB = Alpha;
+
+	if (EnumHasAnyFlags(VM.GetFlags(), EEvaluationFlags::Bones))
+	{
+		check(KeyframeA->Pose.GetNumBones() == KeyframeB->Pose.GetNumBones());
+
+		const FTransformArrayView KeyframeBTransformsView = KeyframeB->Pose.LocalTransforms.GetView();
+
+		BlendOverwriteWithScale(KeyframeBTransformsView, KeyframeB->Pose.LocalTransforms.GetConstView(), WeightOfPoseB);
+		BlendAddWithScale(KeyframeBTransformsView, KeyframeA->Pose.LocalTransforms.GetConstView(), WeightOfPoseA);
+
+		// Ensure that all of the resulting rotations are normalized
+		NormalizeRotations(KeyframeBTransformsView);
+	}
+
+	if (EnumHasAnyFlags(VM.GetFlags(), EEvaluationFlags::Curves))
+	{
+		// Curves cannot blend in place
+		FBlendedCurve Result;
+		Result.Lerp(KeyframeA->Curves, KeyframeB->Curves, WeightOfPoseB);
+
+		KeyframeB->Curves = MoveTemp(Result);
+	}
+
+	if (EnumHasAnyFlags(VM.GetFlags(), EEvaluationFlags::Attributes))
+	{
+		UE::Anim::Attributes::BlendAttributes({ KeyframeA->Attributes, KeyframeB->Attributes }, { WeightOfPoseA, WeightOfPoseB }, { 0, 1 }, KeyframeB->Attributes);
+	}
+
+	VM.PushValue(KEYFRAME_STACK_NAME, MoveTemp(KeyframeB));
+}
+
+FAnimNextBlendOverwriteKeyframeWithScaleTask FAnimNextBlendOverwriteKeyframeWithScaleTask::Make(float ScaleFactor)
+{
+	FAnimNextBlendOverwriteKeyframeWithScaleTask Task;
+	Task.ScaleFactor = ScaleFactor;
+
+	return Task;
+}
+
+void FAnimNextBlendOverwriteKeyframeWithScaleTask::Execute(UE::AnimNext::FEvaluationVM& VM) const
+{
+	using namespace UE::AnimNext;
+
+	TUniquePtr<FKeyframeState> Keyframe;
+	if (!VM.PopValue(KEYFRAME_STACK_NAME, Keyframe))
+	{
+		// We have no inputs, nothing to do
+		return;
+	}
+
+	if (EnumHasAnyFlags(VM.GetFlags(), EEvaluationFlags::Bones))
+	{
+		BlendOverwriteWithScale(Keyframe->Pose.LocalTransforms.GetView(), Keyframe->Pose.LocalTransforms.GetConstView(), ScaleFactor);
+	}
+
+	if (EnumHasAnyFlags(VM.GetFlags(), EEvaluationFlags::Curves))
+	{
+		// Curves cannot override in place
+		FBlendedCurve Result;
+		Result.Override(Keyframe->Curves, ScaleFactor);
+
+		Keyframe->Curves = MoveTemp(Result);
+	}
+
+	if (EnumHasAnyFlags(VM.GetFlags(), EEvaluationFlags::Attributes))
+	{
+		UE::Anim::Attributes::OverrideAttributes(Keyframe->Attributes, Keyframe->Attributes, ScaleFactor);
+	}
+
+	VM.PushValue(KEYFRAME_STACK_NAME, MoveTemp(Keyframe));
+}
+
+FAnimNextBlendAddKeyframeWithScaleTask FAnimNextBlendAddKeyframeWithScaleTask::Make(float ScaleFactor)
+{
+	FAnimNextBlendAddKeyframeWithScaleTask Task;
+	Task.ScaleFactor = ScaleFactor;
+
+	return Task;
+}
+
+void FAnimNextBlendAddKeyframeWithScaleTask::Execute(UE::AnimNext::FEvaluationVM& VM) const
+{
+	using namespace UE::AnimNext;
+
+	// Pop our top two poses, we'll re-use the top keyframe for our result
+
+	TUniquePtr<FKeyframeState> KeyframeB;
+	if (!VM.PopValue(KEYFRAME_STACK_NAME, KeyframeB))
+	{
+		// We have no inputs, nothing to do
+		return;
+	}
+
+	TUniquePtr<FKeyframeState> KeyframeA;
+	if (!VM.PopValue(KEYFRAME_STACK_NAME, KeyframeA))
+	{
+		// We have a single input, leave it on top of the stack
+		VM.PushValue(KEYFRAME_STACK_NAME, MoveTemp(KeyframeB));
+		return;
+	}
+
+	if (EnumHasAnyFlags(VM.GetFlags(), EEvaluationFlags::Bones))
+	{
+		check(KeyframeA->Pose.GetNumBones() == KeyframeB->Pose.GetNumBones());
+
+		BlendAddWithScale(KeyframeB->Pose.LocalTransforms.GetView(), KeyframeA->Pose.LocalTransforms.GetConstView(), ScaleFactor);
+	}
+
+	if (EnumHasAnyFlags(VM.GetFlags(), EEvaluationFlags::Curves))
+	{
+		KeyframeB->Curves.Accumulate(KeyframeA->Curves, ScaleFactor);
+	}
+
+	if (EnumHasAnyFlags(VM.GetFlags(), EEvaluationFlags::Attributes))
+	{
+		UE::Anim::Attributes::AccumulateAttributes(KeyframeA->Attributes, KeyframeB->Attributes, ScaleFactor, AAT_None);
+	}
+
+	VM.PushValue(KEYFRAME_STACK_NAME, MoveTemp(KeyframeB));
+}
diff --git a/Engine/Plugins/Experimental/AnimNext/Source/AnimNext/Private/EvaluationVM/Tasks/ExecuteProgram.cpp b/Engine/Plugins/Experimental/AnimNext/Source/AnimNext/Private/EvaluationVM/Tasks/ExecuteProgram.cpp
new file mode 100644
--- /dev/null
+++ b/Engine/Plugins/Experimental/AnimNext/Source/AnimNext/Private/EvaluationVM/Tasks/ExecuteProgram.cpp
@@ -0,0 +1,21 @@
+// Copyright Epic Games, Inc. All Rights Reserved.
+
+#include "EvaluationVM/Tasks/ExecuteProgram.h"
+
+#include "EvaluationVM/EvaluationProgram.h"
+#include "EvaluationVM/EvaluationVM.h"
+
+FAnimNextExecuteProgramTask FAnimNextExecuteProgramTask::Make(const UE::AnimNext::FEvaluationProgram* Program)
+{
+	FAnimNextExecuteProgramTask Task;
+	Task.Program = Program;
+	return Task;
+}
+
+void FAnimNextExecuteProgramTask::Execute(UE::AnimNext::FEvaluationVM& VM) const
+{
+	if (Program != nullptr)
+	{
+		Program->Execute(VM);
+	}
+}
diff --git a/Engine/Plugins/Experimental/AnimNext/Source/AnimNext/Private/EvaluationVM/Tasks/NormalizeRotations.cpp b/Engine/Plugins/Experimental/AnimNext/Source/AnimNext/Private/EvaluationVM/Tasks/NormalizeRotations.cpp
new file mode 100644
--- /dev/null
+++ b/Engine/Plugins/Experimental/AnimNext/Source/AnimNext/Private/EvaluationVM/Tasks/NormalizeRotations.cpp
@@ -0,0 +1,26 @@
+// Copyright Epic Games, Inc. All Rights Reserved.
+
+#include "EvaluationVM/Tasks/NormalizeRotations.h"
+
+#include "EvaluationVM/EvaluationVM.h"
+#include "EvaluationVM/KeyframeState.h"
+#include "TransformArrayOperations.h"
+
+void FAnimNextNormalizeKeyframeRotationsTask::Execute(UE::AnimNext::FEvaluationVM& VM) const
+{
+	using namespace UE::AnimNext;
+
+	if (EnumHasAnyFlags(VM.GetFlags(), EEvaluationFlags::Bones))
+	{
+		TUniquePtr<FKeyframeState> Keyframe;
+		if (!VM.PopValue(KEYFRAME_STACK_NAME, Keyframe))
+		{
+			// We have no inputs, nothing to do
+			return;
+		}
+
+		NormalizeRotations(Keyframe->Pose.LocalTransforms.GetView());
+
+		VM.PushValue(KEYFRAME_STACK_NAME, MoveTemp(Keyframe));
+	}
+}
diff --git a/Engine/Plugins/Experimental/AnimNext/Source/AnimNext/Private/EvaluationVM/Tasks/PushAnimSequenceKeyframe.cpp b/Engine/Plugins/Experimental/AnimNext/Source/AnimNext/Private/EvaluationVM/Tasks/PushAnimSequenceKeyframe.cpp
new file mode 100644
--- /dev/null
+++ b/Engine/Plugins/Experimental/AnimNext/Source/AnimNext/Private/EvaluationVM/Tasks/PushAnimSequenceKeyframe.cpp
@@ -0,0 +1,67 @@
+// Copyright Epic Games, Inc. All Rights Reserved.
+
+#include "EvaluationVM/Tasks/PushAnimSequenceKeyframe.h"
+
+#include "Animation/AnimSequence.h"
+#include "BonePose.h"
+#include "DecompressionTools.h"
+#include "EvaluationVM/EvaluationVM.h"
+#include "EvaluationVM/KeyframeState.h"
+
+FAnimNextAnimSequenceKeyframeTask FAnimNextAnimSequenceKeyframeTask::MakeFromSampleTime(TWeakObjectPtr<UAnimSequence> AnimSequence, double SampleTime, bool bInterpolate)
+{
+	FAnimNextAnimSequenceKeyframeTask Task;
+	Task.AnimSequence = AnimSequence;
+	Task.SampleTime = SampleTime;
+	Task.bInterpolate = bInterpolate;
+
+	return Task;
+}
+
+FAnimNextAnimSequenceKeyframeTask FAnimNextAnimSequenceKeyframeTask::MakeFromKeyframeIndex(TWeakObjectPtr<UAnimSequence> AnimSequence, uint32 KeyframeIndex)
+{
+	FAnimNextAnimSequenceKeyframeTask Task;
+	Task.AnimSequence = AnimSequence;
+	Task.KeyframeIndex = KeyframeIndex;
+
+	return Task;
+}
+
+void FAnimNextAnimSequenceKeyframeTask::Execute(UE::AnimNext::FEvaluationVM& VM) const
+{
+	using namespace UE::AnimNext;
+
+	const UAnimSequence* AnimSequencePtr = AnimSequence.Get();
+	const bool bIsAdditive = AnimSequencePtr->IsValidAdditive();
+
+	FDeltaTimeRecord DeltaTimeRecord;	// Not needed
+	const bool bExtractRootMotion = bExtractTrajectory;
+	const bool bLooping = false;		// Not needed
+	const bool bUseRawData = false;
+
+	const FAnimExtractContext ExtractionContext(SampleTime, bExtractRootMotion, DeltaTimeRecord, bLooping);
+
+	FKeyframeState Keyframe = VM.MakeUninitializedKeyframe(bIsAdditive);
+
+	if (EnumHasAnyFlags(VM.GetFlags(), EEvaluationFlags::Bones))
+	{
+		FDecompressionTools::GetAnimationPose(AnimSequencePtr, Keyframe.Pose, ExtractionContext);
+	}
+
+	if (EnumHasAnyFlags(VM.GetFlags(), EEvaluationFlags::Curves))
+	{
+		AnimSequencePtr->EvaluateCurveData(Keyframe.Curves, static_cast<float>(SampleTime), bUseRawData);
+	}
+
+	if (EnumHasAnyFlags(VM.GetFlags(), EEvaluationFlags::Attributes))
+	{
+		FCompactPose Pose;		// Dummy but we need the bone container
+		Pose.SetBoneContainer(&VM.GetBoneContainer());
+
+		FAnimationPoseData PoseData(Pose, Keyframe.Curves, Keyframe.Attributes);
+
+		AnimSequencePtr->EvaluateAttributes(PoseData, ExtractionContext, bUseRawData);
+	}
+
+	VM.PushValue(KEYFRAME_STACK_NAME, MakeUnique<FKeyframeState>(MoveTemp(Keyframe)));
+}
diff --git a/Engine/Plugins/Experimental/AnimNext/Source/AnimNext/Private/EvaluationVM/Tasks/PushReferenceKeyframe.cpp b/Engine/Plugins/Experimental/AnimNext/Source/AnimNext/Private/EvaluationVM/Tasks/PushReferenceKeyframe.cpp
new file mode 100644
--- /dev/null
+++ b/Engine/Plugins/Experimental/AnimNext/Source/AnimNext/Private/EvaluationVM/Tasks/PushReferenceKeyframe.cpp
@@ -0,0 +1,25 @@
+// Copyright Epic Games, Inc. All Rights Reserved.
+
+#include "EvaluationVM/Tasks/PushReferenceKeyframe.h"
+
+#include "EvaluationVM/EvaluationVM.h"
+
+FAnimNextPushReferenceKeyframeTask FAnimNextPushReferenceKeyframeTask::MakeFromSkeleton()
+{
+	return FAnimNextPushReferenceKeyframeTask();
+}
+
+FAnimNextPushReferenceKeyframeTask FAnimNextPushReferenceKeyframeTask::MakeFromAdditiveIdentity()
+{
+	FAnimNextPushReferenceKeyframeTask Task;
+	Task.bIsAdditive = true;
+
+	return Task;
+}
+
+void FAnimNextPushReferenceKeyframeTask::Execute(UE::AnimNext::FEvaluationVM& VM) const
+{
+	using namespace UE::AnimNext;
+
+	VM.PushValue(KEYFRAME_STACK_NAME, MakeUnique<FKeyframeState>(VM.MakeReferenceKeyframe(bIsAdditive)));
+}
diff --git a/Engine/Plugins/Experimental/AnimNext/Source/AnimNext/Private/GenerationTools.cpp b/Engine/Plugins/Experimental/AnimNext/Source/AnimNext/Private/GenerationTools.cpp
--- a/Engine/Plugins/Experimental/AnimNext/Source/AnimNext/Private/GenerationTools.cpp
+++ b/Engine/Plugins/Experimental/AnimNext/Source/AnimNext/Private/GenerationTools.cpp
@@ -401,7 +401,7 @@ struct FCompareBoneIndexType
 			const FCompactPoseBoneIndex CompactBoneIndex = SourcePose.Pose.GetBoneContainer().GetCompactPoseIndexFromSkeletonIndex(SkeletonBoneIndex);
 			if (CompactBoneIndex.GetInt() != INDEX_NONE)
 			{
-				TargetPose.LocalTransforms[i] = SourcePose.Pose[CompactBoneIndex];
+				TargetPose.LocalTransformsView[i] = SourcePose.Pose[CompactBoneIndex];
 			}
 		}
 	}
@@ -424,7 +424,7 @@ struct FCompareBoneIndexType
 			const FCompactPoseBoneIndex CompactBoneIndex = TargetPose.Pose.GetBoneContainer().GetCompactPoseIndexFromSkeletonIndex(SkeletonBoneIndex);
 			if (CompactBoneIndex.GetInt() != INDEX_NONE)
 			{
-				TargetPose.Pose[CompactBoneIndex] = SourcePose.LocalTransforms[i];
+				TargetPose.Pose[CompactBoneIndex] = SourcePose.LocalTransformsView[i];
 			}
 		}
 	}
diff --git a/Engine/Plugins/Experimental/AnimNext/Source/AnimNext/Private/Graph/AnimGraph/AnimNode_AnimNextGraph.cpp b/Engine/Plugins/Experimental/AnimNext/Source/AnimNext/Private/Graph/AnimGraph/AnimNode_AnimNextGraph.cpp
--- a/Engine/Plugins/Experimental/AnimNext/Source/AnimNext/Private/Graph/AnimGraph/AnimNode_AnimNextGraph.cpp
+++ b/Engine/Plugins/Experimental/AnimNext/Source/AnimNext/Private/Graph/AnimGraph/AnimNode_AnimNextGraph.cpp
@@ -145,8 +145,8 @@ void FAnimNode_AnimNextGraph::Evaluate_AnyThread(FPoseContext & Output)
 
 	FAnimNextGraphReferencePose GraphReferencePose(&RefPose);
 
-	FAnimNextGraphLODPose GraphSourceLODPose(FLODPose(RefPose, LODLevel, false, Output.ExpectsAdditivePose()));
-	FAnimNextGraphLODPose ResultPose(FLODPose(RefPose, LODLevel, true, Output.ExpectsAdditivePose()));
+	FAnimNextGraphLODPose GraphSourceLODPose(FLODPoseHeap(RefPose, LODLevel, false, Output.ExpectsAdditivePose()));
+	FAnimNextGraphLODPose ResultPose(FLODPoseHeap(RefPose, LODLevel, true, Output.ExpectsAdditivePose()));
 	FGenerationTools::RemapPose(LODLevel, SourcePose, RefPose, GraphSourceLODPose.LODPose);
 
 	FParamStack& ParamStack = FParamStack::Get();
diff --git a/Engine/Plugins/Experimental/AnimNext/Source/AnimNext/Private/Graph/RigUnit_AnimNextGraphEvaluator.cpp b/Engine/Plugins/Experimental/AnimNext/Source/AnimNext/Private/Graph/RigUnit_AnimNextGraphEvaluator.cpp
--- a/Engine/Plugins/Experimental/AnimNext/Source/AnimNext/Private/Graph/RigUnit_AnimNextGraphEvaluator.cpp
+++ b/Engine/Plugins/Experimental/AnimNext/Source/AnimNext/Private/Graph/RigUnit_AnimNextGraphEvaluator.cpp
@@ -4,7 +4,10 @@
 #include "DecoratorBase/ExecutionContext.h"
 #include "DecoratorInterfaces/IUpdate.h"
 #include "DecoratorInterfaces/IEvaluate.h"
+#include "EvaluationVM/EvaluationVM.h"
 #include "Context.h"
+#include "Param/ParamStack.h"
+#include "Graph/AnimNext_LODPose.h"
 
 #include UE_INLINE_GENERATED_CPP_BY_NAME(RigUnit_AnimNextGraphEvaluator)
 
@@ -24,16 +27,39 @@ void FRigUnit_AnimNextGraphEvaluator::StaticExecute(FRigVMExtendedExecuteContext
 		const FContext& InterfaceContext = VMExecuteContext.GetContext();
 
 		// Call pre/post update on our graph
-		UpdateGraph(Context, GraphInstancePtr, InterfaceContext.GetDeltaTime());
+		FUpdateTraversalContext TraversalContext(InterfaceContext.GetDeltaTime());
+		UpdateGraph(Context, TraversalContext, GraphInstancePtr);
 	}
 
 	if (EnumHasAnyFlags(SimulationSteps, EAnimNextGraphSimulationSteps::Evaluate))
 	{
 		// Call pre/post evaluate on our graph
-		FPoseContainer PoseContainer;
-		EvaluateGraph(Context, GraphInstancePtr, EEvaluationFlags::All, PoseContainer);
+		FEvaluateTraversalContext TraversalContext;
+		FEvaluationProgram EvaluationProgram = EvaluateGraph(Context, TraversalContext, GraphInstancePtr);
 
-		// TODO: Write out our output pose/curves/attributes/etc
-		//ExecuteContext_.GetContext().SetResult<FAnimationGraphNodeResult>(Result);
+		if (!EvaluationProgram.IsEmpty())
+		{
+			FParamStack& ParamStack = FParamStack::Get();
+
+			const FAnimNextGraphReferencePose& GraphReferencePose = ParamStack.GetParam<FAnimNextGraphReferencePose>("GraphReferencePose");
+			const int32 GraphLODLevel = ParamStack.GetParam<int32>("GraphLODLevel");
+
+			FEvaluationVM EvaluationVM(EEvaluationFlags::All, *GraphReferencePose.ReferencePose, GraphLODLevel);
+			EvaluationProgram.Execute(EvaluationVM);
+
+			FAnimNextGraphLODPose& ResultPose = ParamStack.GetMutableParam<FAnimNextGraphLODPose>("ResultPose");
+
+			TUniquePtr<FKeyframeState> EvaluatedKeyframe;
+			if (EvaluationVM.PopValue(KEYFRAME_STACK_NAME, EvaluatedKeyframe))
+			{
+				ResultPose.LODPose.CopyFrom(EvaluatedKeyframe->Pose);
+			}
+			else
+			{
+				// We need to output a valid pose, generate one
+				FKeyframeState ReferenceKeyframe = EvaluationVM.MakeReferenceKeyframe(ParamStack.GetParam<bool>("GraphExpectsAdditive"));
+				ResultPose.LODPose.CopyFrom(ReferenceKeyframe.Pose);
+			}
+		}
 	}
 }
diff --git a/Engine/Plugins/Experimental/AnimNext/Source/AnimNext/Private/TransformArrayOperations_AoS.cpp b/Engine/Plugins/Experimental/AnimNext/Source/AnimNext/Private/TransformArrayOperations_AoS.cpp
new file mode 100644
--- /dev/null
+++ b/Engine/Plugins/Experimental/AnimNext/Source/AnimNext/Private/TransformArrayOperations_AoS.cpp
@@ -0,0 +1,82 @@
+// Copyright Epic Games, Inc. All Rights Reserved.
+
+#include "TransformArrayOperations.h"
+#include "TransformArray.h"
+
+namespace UE::AnimNext
+{
+	void SetIdentity(const FTransformArrayAoSView& Dest, bool bIsAdditive)
+	{
+		const int32 NumTransforms = Dest.Num();
+		const FTransform Identity = bIsAdditive ? TransformAdditiveIdentity : FTransform::Identity;
+
+		for (int32 TransformIndex = 0; TransformIndex < NumTransforms; ++TransformIndex)
+		{
+			Dest[TransformIndex] = Identity;
+		}
+	}
+
+	void CopyTransforms(const FTransformArrayAoSView& Dest, const FTransformArrayAoSConstView& Source, int32 StartIndex, int32 NumToCopy)
+	{
+		const int32 NumTransforms = Dest.Num();
+		const int32 EndIndex = NumToCopy >= 0 ? (StartIndex + NumToCopy) : NumTransforms;
+
+		check(Source.Num() >= NumTransforms);
+		check(StartIndex >= 0 && StartIndex < NumTransforms);
+		check(EndIndex <= NumTransforms);
+
+		for (int32 TransformIndex = StartIndex; TransformIndex < EndIndex; ++TransformIndex)
+		{
+			Dest[TransformIndex] = Source[TransformIndex];
+		}
+	}
+
+	void NormalizeRotations(const FTransformArrayAoSView& Input)
+	{
+		const int32 NumTransforms = Input.Num();
+
+		for (int32 TransformIndex = 0; TransformIndex < NumTransforms; ++TransformIndex)
+		{
+			Input[TransformIndex].NormalizeRotation();
+		}
+	}
+
+	void BlendWithIdentityAndAccumulate(const FTransformArrayAoSView& Base, const FTransformArrayAoSConstView& Additive, const float BlendWeight)
+	{
+		const ScalarRegister VBlendWeight(BlendWeight);
+		const int32 NumTransforms = Base.Num();
+
+		check(Additive.Num() >= NumTransforms);
+
+		for (int32 TransformIndex = 0; TransformIndex < NumTransforms; ++TransformIndex)
+		{
+			FTransform::BlendFromIdentityAndAccumulate(Base[TransformIndex], Additive[TransformIndex], VBlendWeight);
+		}
+	}
+
+	void BlendOverwriteWithScale(const FTransformArrayAoSView& Dest, const FTransformArrayAoSConstView& Source, const float ScaleWeight)
+	{
+		const ScalarRegister VScaleWeight(ScaleWeight);
+		const int32 NumTransforms = Source.Num();
+
+		check(Dest.Num() >= NumTransforms);
+
+		for (int32 TransformIndex = 0; TransformIndex < NumTransforms; ++TransformIndex)
+		{
+			Dest[TransformIndex] = Source[TransformIndex] * VScaleWeight;
+		}
+	}
+
+	void BlendAddWithScale(const FTransformArrayAoSView& Dest, const FTransformArrayAoSConstView& Source, const float ScaleWeight)
+	{
+		const ScalarRegister VScaleWeight(ScaleWeight);
+		const int32 NumTransforms = Source.Num();
+
+		check(Dest.Num() >= NumTransforms);
+
+		for (int32 TransformIndex = 0; TransformIndex < NumTransforms; ++TransformIndex)
+		{
+			Dest[TransformIndex].AccumulateWithShortestRotation(Source[TransformIndex], VScaleWeight);
+		}
+	}
+}
diff --git a/Engine/Plugins/Experimental/AnimNext/Source/AnimNext/Private/TransformArrayOperations_SoA.cpp b/Engine/Plugins/Experimental/AnimNext/Source/AnimNext/Private/TransformArrayOperations_SoA.cpp
new file mode 100644
--- /dev/null
+++ b/Engine/Plugins/Experimental/AnimNext/Source/AnimNext/Private/TransformArrayOperations_SoA.cpp
@@ -0,0 +1,180 @@
+// Copyright Epic Games, Inc. All Rights Reserved.
+
+#include "TransformArrayOperations.h"
+#include "TransformArray.h"
+
+namespace UE::AnimNext
+{
+	// A couple of performance notes:
+	//    * Types all use doubles
+	//    * Because we split transforms into 3 arrays, we pay 3x the range checks
+	//    * We could use pointers directly to avoid range checks, similar to what ISPC would do
+	//    * FTransform uses 3x register wide types (32 bytes each)
+	//    * We could do the same and use register wide types for aligned stores at the expense of
+	//      using more memory
+	//    * Unaligned loads aren't a big deal, perf wise, as long as we are 16 byte aligned
+	//    * FVector loads require zeroing W which we could do without if we ensure we have padding
+	//      on the last array element to avoid an out-of-bounds read
+	//    * FVector stores end up as 2 stores which isn't too bad
+	//    * Array values are very likely to be in the CPU L1 but to help hide the latency we should
+	//      move loads next to one another to bundle them up up front
+	//    * The biggest cost to normalizing is the square root and division which we perform on 4 elements
+	//      even though there is a single value: the length. We could reduce this cost by processing
+	//      4 rotations at a time and computing the square root and division for all 4 at the same time.
+	//      To achieve this, we would have to swizzle the inputs/outputs in true SOA format: XXXX, YYYY, ZZZZ, WWWW.
+	//      Shuffling with 4x4 floats requires 8 shuffles (after the 4 loads) and with doubles we should be
+	//      able to do it with 8 loads and 8 shuffles (perhaps less with AVX).
+	//      Another option might be to split loops into stages:
+	//          * do dot, write it out for 8 entries
+	//          * load, do sqrt, write
+	//          * load, do div, write
+	//          * load, cmp, mul, mask, final store
+	//      Splitting into stages would help ensure that the CPU is saturated with independent instructions
+	//      it can execute out-of-order. The intermediary load/stores would be from the stack if the values don't fit
+	//      in registers or if we wish to re-order things cheaply (shuffle for sqrt from above).
+	//      The idea being to keep loops smaller to maximize overlap of instructions while waiting for long
+	//      latency instructions (cache miss in first batch, sqrt, div, etc).
+
+	void SetIdentity(const FTransformArraySoAView& Dest, bool bIsAdditive)
+	{
+		const int32 NumTransforms = Dest.Num();
+
+		const FTransform& Identity = bIsAdditive ? TransformAdditiveIdentity : FTransform::Identity;
+
+		const FVector IdentityTranslation = Identity.GetTranslation();
+		const FQuat IdentityRotation = Identity.GetRotation();
+		const FVector IdentityScale3D = Identity.GetScale3D();
+
+		for (int32 TransformIndex = 0; TransformIndex < NumTransforms; ++TransformIndex)
+		{
+			Dest.Translations[TransformIndex] = IdentityTranslation;
+			Dest.Rotations[TransformIndex] = IdentityRotation;
+			Dest.Scales3D[TransformIndex] = IdentityScale3D;
+		}
+	}
+
+	void CopyTransforms(const FTransformArraySoAView& Dest, const FTransformArraySoAConstView& Source, int32 StartIndex, int32 NumToCopy)
+	{
+		const int32 NumTransforms = Dest.Num();
+		const int32 EndIndex = NumToCopy >= 0 ? (StartIndex + NumToCopy) : NumTransforms;
+
+		check(Source.Num() >= NumTransforms);
+		check(StartIndex >= 0 && StartIndex < NumTransforms);
+		check(EndIndex <= NumTransforms);
+
+		for (int32 TransformIndex = StartIndex; TransformIndex < EndIndex; ++TransformIndex)
+		{
+			Dest.Translations[TransformIndex] = Source.Translations[TransformIndex];
+			Dest.Rotations[TransformIndex] = Source.Rotations[TransformIndex];
+			Dest.Scales3D[TransformIndex] = Source.Scales3D[TransformIndex];
+		}
+	}
+
+	void NormalizeRotations(const FTransformArraySoAView& Input)
+	{
+		const int32 NumTransforms = Input.Num();
+
+		for (int32 TransformIndex = 0; TransformIndex < NumTransforms; ++TransformIndex)
+		{
+			Input.Rotations[TransformIndex].Normalize();
+		}
+	}
+
+	void BlendWithIdentityAndAccumulate(const FTransformArraySoAView& Base, const FTransformArraySoAConstView& Additive, const float BlendWeight)
+	{
+		const int32 NumTransforms = Base.Num();
+
+		check(Additive.Num() >= NumTransforms);
+
+		using TransformVectorRegister = FTransform::TransformVectorRegister;
+
+		const TransformVectorRegister VBlendWeight(ScalarRegister(BlendWeight).Value);
+		const TransformVectorRegister Zero = VectorZero();
+		const TransformVectorRegister Const0001 = GlobalVectorConstants::Float0001;
+		const TransformVectorRegister ConstNegative0001 = VectorSubtract(Zero, Const0001);
+		const TransformVectorRegister VOneMinusAlpha = VectorSubtract(VectorOne(), VBlendWeight);
+		const TransformVectorRegister DefaultScale = GlobalVectorConstants::Float1110;
+
+		for (int32 TransformIndex = 0; TransformIndex < NumTransforms; ++TransformIndex)
+		{
+			// Blend rotation
+			//     To ensure the 'shortest route', we make sure the dot product between the both rotations is positive.
+			//     const float Bias = (|A.B| >= 0 ? 1 : -1)
+			//     BlendedAtom.Rotation = (B * Alpha) + (A * (Bias * (1.f - Alpha)));
+			//     BlendedAtom.Rotation.QuaternionNormalize();
+			//  Note: A = (0,0,0,1), which simplifies things a lot; only care about sign of B.W now, instead of doing a dot product
+			const TransformVectorRegister RotationB = VectorLoadAligned(&Additive.Rotations[TransformIndex]);
+
+			const TransformVectorRegister QuatRotationDirMask = VectorCompareGE(RotationB, Zero);
+			const TransformVectorRegister BiasTimesA = VectorSelect(QuatRotationDirMask, Const0001, ConstNegative0001);
+			const TransformVectorRegister RotateBTimesWeight = VectorMultiply(RotationB, VBlendWeight);
+			const TransformVectorRegister UnnormalizedRotation = VectorMultiplyAdd(BiasTimesA, VOneMinusAlpha, RotateBTimesWeight);
+
+			// Normalize blended rotation ( result = (Q.Q >= 1e-8) ? (Q / |Q|) : (0,0,0,1) )
+			const TransformVectorRegister BlendedRotation = VectorNormalizeSafe(UnnormalizedRotation, Const0001);
+
+			// FinalAtom.Rotation = BlendedAtom.Rotation * FinalAtom.Rotation;
+			Base.Rotations[TransformIndex] = FQuat::MakeFromVectorRegister(VectorQuaternionMultiply2(BlendedRotation, VectorLoadAligned(&Base.Rotations[TransformIndex])));
+
+			// Blend translation and scale
+			//    BlendedAtom.Translation = Lerp(Zero, Additive.Translation, Alpha);
+			//    BlendedAtom.Scale = Lerp(0, Additive.Scale, Alpha);
+			const TransformVectorRegister BlendedTranslation = FMath::Lerp<TransformVectorRegister>(Zero, VectorLoadFloat3_W0(&Additive.Translations[TransformIndex]), VBlendWeight);
+			const TransformVectorRegister BlendedScale3D = FMath::Lerp<TransformVectorRegister>(Zero, VectorLoadFloat3_W0(&Additive.Scales3D[TransformIndex]), VBlendWeight);
+
+			// Apply translation and scale to final atom
+			//     FinalAtom.Translation += BlendedAtom.Translation
+			//     FinalAtom.Scale *= BlendedAtom.Scale
+			VectorStoreFloat3(VectorAdd(VectorLoadFloat3_W0(&Base.Translations[TransformIndex]), BlendedTranslation), &Base.Translations[TransformIndex]);
+			VectorStoreFloat3(VectorMultiply(VectorLoadFloat3_W0(&Base.Scales3D[TransformIndex]), VectorAdd(DefaultScale, BlendedScale3D)), &Base.Scales3D[TransformIndex]);
+		}
+	}
+
+	void BlendOverwriteWithScale(const FTransformArraySoAView& Dest, const FTransformArraySoAConstView& Source, const float ScaleWeight)
+	{
+		const int32 NumTransforms = Source.Num();
+
+		check(Dest.Num() >= NumTransforms);
+
+		for (int32 TransformIndex = 0; TransformIndex < NumTransforms; ++TransformIndex)
+		{
+			Dest.Translations[TransformIndex] = Source.Translations[TransformIndex] * ScaleWeight;
+			Dest.Rotations[TransformIndex] = Source.Rotations[TransformIndex] * ScaleWeight;
+			Dest.Scales3D[TransformIndex] = Source.Scales3D[TransformIndex] * ScaleWeight;
+		}
+	}
+
+	void BlendAddWithScale(const FTransformArraySoAView& Dest, const FTransformArraySoAConstView& Source, const float ScaleWeight)
+	{
+		const int32 NumTransforms = Source.Num();
+
+		check(Dest.Num() >= NumTransforms);
+
+		using QuatVectorRegister = FQuat::QuatVectorRegister;
+		const QuatVectorRegister Zero = VectorZero();
+		const ScalarRegister VScaleWeight(ScaleWeight);
+
+		for (int32 TransformIndex = 0; TransformIndex < NumTransforms; ++TransformIndex)
+		{
+			const QuatVectorRegister SourceRotation = VectorLoadAligned(&Source.Rotations[TransformIndex]);
+			const QuatVectorRegister DestRotation = VectorLoadAligned(&Dest.Rotations[TransformIndex]);
+
+			const QuatVectorRegister BlendedRotation = VectorMultiply(SourceRotation, VScaleWeight);
+
+			// Blend rotation
+			//     To ensure the 'shortest route', we make sure the dot product between the both rotations is positive.
+			//     const float Bias = (|A.B| >= 0 ? 1 : -1)
+			//     return A + B * Bias;
+
+			const QuatVectorRegister RotationDot = VectorDot4(DestRotation, BlendedRotation);
+			const QuatVectorRegister QuatRotationDirMask = VectorCompareGE(RotationDot, Zero);
+			const QuatVectorRegister NegativeB = VectorSubtract(Zero, BlendedRotation);
+			const QuatVectorRegister BiasTimesB = VectorSelect(QuatRotationDirMask, BlendedRotation, NegativeB);
+
+			Dest.Rotations[TransformIndex] = FQuat::MakeFromVectorRegister(VectorAdd(DestRotation, BiasTimesB));
+
+			Dest.Translations[TransformIndex] += Source.Translations[TransformIndex] * ScaleWeight;
+			Dest.Scales3D[TransformIndex] += Source.Scales3D[TransformIndex] * ScaleWeight;
+		}
+	}
+}
diff --git a/Engine/Plugins/Experimental/AnimNext/Source/AnimNext/Public/DecoratorBase/EvaluationFlags.h b/Engine/Plugins/Experimental/AnimNext/Source/AnimNext/Public/DecoratorBase/EvaluationFlags.h
deleted file mode 100644
diff --git a/Engine/Plugins/Experimental/AnimNext/Source/AnimNext/Public/DecoratorInterfaces/IEvaluate.h b/Engine/Plugins/Experimental/AnimNext/Source/AnimNext/Public/DecoratorInterfaces/IEvaluate.h
--- a/Engine/Plugins/Experimental/AnimNext/Source/AnimNext/Public/DecoratorInterfaces/IEvaluate.h
+++ b/Engine/Plugins/Experimental/AnimNext/Source/AnimNext/Public/DecoratorInterfaces/IEvaluate.h
@@ -3,15 +3,13 @@
 #pragma once
 
 #include "CoreMinimal.h"
-#include "DecoratorBase/EvaluationFlags.h"
 #include "DecoratorBase/IDecoratorInterface.h"
 #include "DecoratorBase/ITraversalContext.h"
+#include "EvaluationVM/EvaluationProgram.h"
+#include "EvaluationVM/KeyframeState.h"
 
 namespace UE::AnimNext
 {
-	// TODO: Hook up our real pose container
-	struct FPoseContainer final {};
-
 	/**
 	 * FEvaluateTraversalContext
 	 *
@@ -19,31 +17,34 @@ namespace UE::AnimNext
 	 */
 	struct ANIMNEXT_API FEvaluateTraversalContext : ITraversalContext
 	{
-		explicit FEvaluateTraversalContext(EEvaluationFlags EvaluationFlags_);
-
-		EEvaluationFlags GetEvaluationFlags() const { return EvaluationFlags; }
+		FEvaluateTraversalContext() = default;
 
-		//void PushPose(FPose&& Pose);
-		//FPose&& PopPose();
+		// Appends a new task into the evaluation program, tasks mutate state in the order they have been appended in
+		// This means that child nodes need to evaluate first, tasks will usually be appended in IEvaluate::PostEvaluate
+		// Tasks are moved into their final memory location, caller can allocate the task anywhere, it is no longer needed after this operation
+		// @see FEvaluationProgram, FEvaluationTask, FEvaluationVM
+		template<class TaskType>
+		void AppendTask(TaskType&& Task) { EvaluationProgram->AppendTask(MoveTemp(Task)); }
 
 	private:
-		EEvaluationFlags EvaluationFlags;
+		FEvaluationProgram* EvaluationProgram = nullptr;
+
+		friend struct FScopedEvaluationProgram;
 	};
 
 	/**
 	 * IEvaluate
 	 * 
-	 * This interface is called during the evaluation traversal. It aims to produce some
-	 * of the following: an output pose, its velocity, the root trajectory, scalar curves, attributes, etc.
+	 * This interface is called during the evaluation traversal. It aims to produce an evaluation program.
 	 * 
-	 * When a node is decorator is visited, PreEvaluate is first called on it. It is responsible for forwarding
+	 * When a node is visited, PreEvaluate is first called on its top decorator. It is responsible for forwarding
 	 * the call to the next decorator that implements this interface on the decorator stack of the node. Once
 	 * all decorators have had the chance to PreEvaluate, the children of the decorator are queried through
 	 * the IHierarchy interface. The children will then evaluate and PostEvaluate will then be called afterwards
 	 * on the original decorator.
 	 * 
 	 * The execution context contains what to evaluate.
-	 * @see EEvaluationFlags
+	 * @see FEvaluationProgram
 	 */
 	struct ANIMNEXT_API IEvaluate : IDecoratorInterface
 	{
@@ -79,7 +80,7 @@ namespace UE::AnimNext
 	};
 
 	/**
-	 * Evaluates a sub-graph starting at the graph root and produces its result.
+	 * Evaluates a sub-graph starting at its root and produces an evaluation program.
 	 * Evaluation starts at the top of the stack that includes the graph root decorator.
 	 *
 	 * For each node:
@@ -90,5 +91,5 @@ namespace UE::AnimNext
 	 *
 	 * @see IEvaluate::PreEvaluate, IEvaluate::PostEvaluate, IHierarchy::GetChildren
 	 */
-	ANIMNEXT_API void EvaluateGraph(FExecutionContext& Context, FWeakDecoratorPtr GraphRootPtr, EEvaluationFlags EvaluationFlags, FPoseContainer& OutContainer);
+	[[nodiscard]] ANIMNEXT_API FEvaluationProgram EvaluateGraph(FExecutionContext& Context, FEvaluateTraversalContext& TraversalContext, FWeakDecoratorPtr GraphRootPtr);
 }
diff --git a/Engine/Plugins/Experimental/AnimNext/Source/AnimNext/Public/DecoratorInterfaces/IUpdate.h b/Engine/Plugins/Experimental/AnimNext/Source/AnimNext/Public/DecoratorInterfaces/IUpdate.h
--- a/Engine/Plugins/Experimental/AnimNext/Source/AnimNext/Public/DecoratorInterfaces/IUpdate.h
+++ b/Engine/Plugins/Experimental/AnimNext/Source/AnimNext/Public/DecoratorInterfaces/IUpdate.h
@@ -13,13 +13,13 @@ namespace UE::AnimNext
 	 * 
 	 * Contains all relevant transient data for an update traversal.
 	 */
-	struct FUpdateTraversalContext : ITraversalContext
+	struct ANIMNEXT_API FUpdateTraversalContext : ITraversalContext
 	{
 		explicit FUpdateTraversalContext(float InDeltaTime);
 
 		float GetDeltaTime() const { return DeltaTime; }
 
-	private:
+	protected:
 		float DeltaTime;
 	};
 
@@ -28,7 +28,7 @@ namespace UE::AnimNext
 	 *
 	 * This interface is called during the update traversal.
 	 *
-	 * When a node is decorator is visited, PreUpdate is first called on it. It is responsible for forwarding
+	 * When a node is visited, PreUpdate is first called on its top decorator. It is responsible for forwarding
 	 * the call to the next decorator that implements this interface on the decorator stack of the node. Once
 	 * all decorators have had the chance to PreUpdate, the children of the decorator are queried through
 	 * the IHierarchy interface. The children will then evaluate and PostUpdate will then be called afterwards
@@ -68,7 +68,7 @@ namespace UE::AnimNext
 	};
 
 	/**
-	 * Updates a sub-graph starting at the graph root.
+	 * Updates a sub-graph starting at its root.
 	 * Update starts at the top of the stack that includes the graph root decorator.
 	 *
 	 * For each node:
@@ -79,5 +79,5 @@ namespace UE::AnimNext
 	 *
 	 * @see IUpdate::PreUpdate, IUpdate::PostUpdate, IHierarchy::GetChildren
 	 */
-	ANIMNEXT_API void UpdateGraph(FExecutionContext& Context, FWeakDecoratorPtr GraphRootPtr, float DeltaTime);
+	ANIMNEXT_API void UpdateGraph(FExecutionContext& Context, FUpdateTraversalContext& TraversalContext, FWeakDecoratorPtr GraphRootPtr);
 }
diff --git a/Engine/Plugins/Experimental/AnimNext/Source/AnimNext/Public/EvaluationVM/EvaluationFlags.h b/Engine/Plugins/Experimental/AnimNext/Source/AnimNext/Public/EvaluationVM/EvaluationFlags.h
new file mode 100644
--- /dev/null
+++ b/Engine/Plugins/Experimental/AnimNext/Source/AnimNext/Public/EvaluationVM/EvaluationFlags.h
@@ -0,0 +1,32 @@
+// Copyright Epic Games, Inc. All Rights Reserved.
+
+#pragma once
+
+#include "CoreMinimal.h"
+
+namespace UE::AnimNext
+{
+	// Flags to control what to compute when evaluating an animation node
+	enum class EEvaluationFlags
+	{
+		// Empty flags
+		None = 0x00,
+
+		// Evaluates the trajectory of the root bone
+		Trajectory = 0x01,
+
+		// Evaluates the animation pose of all bones
+		Bones = 0x02,
+
+		// Evaluates the animation curves
+		Curves = 0x04,
+
+		// Evaluates the animation attributes
+		Attributes = 0x08,
+
+		// Evaluates everything
+		All = 0xff,
+	};
+
+	ENUM_CLASS_FLAGS(EEvaluationFlags)
+}
diff --git a/Engine/Plugins/Experimental/AnimNext/Source/AnimNext/Public/EvaluationVM/EvaluationProgram.h b/Engine/Plugins/Experimental/AnimNext/Source/AnimNext/Public/EvaluationVM/EvaluationProgram.h
new file mode 100644
--- /dev/null
+++ b/Engine/Plugins/Experimental/AnimNext/Source/AnimNext/Public/EvaluationVM/EvaluationProgram.h
@@ -0,0 +1,63 @@
+// Copyright Epic Games, Inc. All Rights Reserved.
+
+#pragma once
+
+#include "CoreMinimal.h"
+#include "EvaluationVM/EvaluationTask.h"
+
+#include <type_traits>
+
+namespace UE::AnimNext
+{
+	struct FEvaluationVM;
+
+	/*
+	 * Evaluation Program
+	 *
+	 * This struct holds a sequence of evaluation tasks that form a program within our
+	 * evaluation virtual machine framework. Programs are immutable once written.
+	 *
+	 * @see FAnimNextEvaluationTask
+	 */
+	struct ANIMNEXT_API FEvaluationProgram
+	{
+		FEvaluationProgram() = default;
+
+		// Allow moving
+		FEvaluationProgram(FEvaluationProgram&&) = default;
+		FEvaluationProgram& operator=(FEvaluationProgram&&) = default;
+
+		// Returns whether or not this program is empty
+		bool IsEmpty() const;
+
+		// Appends a new task into the program, tasks mutate state in the order they have been appended in
+		// This means that child nodes need to evaluate first, tasks will usually be appended in IEvaluate::PostEvaluate
+		// Tasks are moved into their final memory location, caller can allocate the task anywhere, it is no longer needed after this operation
+		template<class TaskType>
+		void AppendTask(TaskType&& Task);
+
+		// Executes the current program on the provided virtual machine
+		void Execute(FEvaluationVM& VM) const;
+
+	private:
+		// Disallow copy
+		FEvaluationProgram(const FEvaluationProgram&) = delete;
+		FEvaluationProgram& operator=(const FEvaluationProgram&) = delete;
+
+		// List of tasks
+		TArray<TUniquePtr<FAnimNextEvaluationTask>> Tasks;
+	};
+
+	//////////////////////////////////////////////////////////////////////////
+	// Implementation of inline functions
+
+	template<class TaskType>
+	inline void FEvaluationProgram::AppendTask(TaskType&& Task)
+	{
+		static_assert(std::is_base_of<FAnimNextEvaluationTask, TaskType>::value, "Task type must derive from FAnimNextEvaluationTask");
+
+		// TODO: Use a fancy allocator to ensure all tasks are contiguous and tightly packed
+		TUniquePtr<FAnimNextEvaluationTask> TaskCopy = MakeUnique<TaskType>(MoveTemp(Task));
+		Tasks.Add(MoveTemp(TaskCopy));
+	}
+}
diff --git a/Engine/Plugins/Experimental/AnimNext/Source/AnimNext/Public/EvaluationVM/EvaluationTask.h b/Engine/Plugins/Experimental/AnimNext/Source/AnimNext/Public/EvaluationVM/EvaluationTask.h
new file mode 100644
--- /dev/null
+++ b/Engine/Plugins/Experimental/AnimNext/Source/AnimNext/Public/EvaluationVM/EvaluationTask.h
@@ -0,0 +1,32 @@
+// Copyright Epic Games, Inc. All Rights Reserved.
+
+#pragma once
+
+#include "CoreMinimal.h"
+
+#include "EvaluationTask.generated.h"
+
+namespace UE::AnimNext
+{
+	struct FEvaluationVM;
+}
+
+/*
+ * Evaluation Task
+ *
+ * This is the base class for our evaluation program tasks. They represent macro
+ * instructions within our evaluation virtual machine system. They operate on the
+ * VM internal state to generate inputs and produce outputs.
+ *
+ * @see FEvaluationProgram, FEvaluationVM
+ */
+USTRUCT()
+struct ANIMNEXT_API FAnimNextEvaluationTask
+{
+	GENERATED_BODY()
+
+	virtual ~FAnimNextEvaluationTask() {}
+
+	// Task entry point
+	virtual void Execute(UE::AnimNext::FEvaluationVM& VM) const { LowLevelFatalError(TEXT("Pure virtual not implemented (FAnimNextEvaluationTask::Execute)")); }
+};
diff --git a/Engine/Plugins/Experimental/AnimNext/Source/AnimNext/Public/EvaluationVM/EvaluationVM.h b/Engine/Plugins/Experimental/AnimNext/Source/AnimNext/Public/EvaluationVM/EvaluationVM.h
new file mode 100644
--- /dev/null
+++ b/Engine/Plugins/Experimental/AnimNext/Source/AnimNext/Public/EvaluationVM/EvaluationVM.h
@@ -0,0 +1,304 @@
+// Copyright Epic Games, Inc. All Rights Reserved.
+
+#pragma once
+
+#include "CoreMinimal.h"
+
+#include "BoneContainer.h"
+#include "EvaluationVM/EvaluationFlags.h"
+#include "EvaluationVM/KeyframeState.h"
+#include "ReferencePose.h"
+#include "Misc/Crc.h"
+
+// If you wish to use a custom type with a stack in the FEvaluationVM, you'll have to enable
+// its usage through this macro to ensure that UE::AnimNext::GetTypeID() is declared.
+// Use this macro inside the UE::AnimNext namespace
+#define ANIM_NEXT_ENABLE_EVALUATION_STACK_USAGE(Type) \
+	template<> \
+	inline uint32 GetTypeID<Type>() \
+	{ \
+		/* TODO: Make this constexpr */ \
+		return FCrc::StrCrc32(#Type); \
+	}
+
+namespace UE::AnimNext
+{
+	// Helper function that returns a UID for the specified type
+	// @see ANIM_NEXT_ENABLE_EVALUATION_STACK_USAGE
+	template<typename Type>
+	inline uint32 GetTypeID()
+	{
+		checkf(false, TEXT("Not implemented! See ANIM_NEXT_ENABLE_EVALUATION_STACK_USAGE for details and register your type"));
+		return 0;
+	}
+
+	// Helper function that destructs a pointer of the provided type
+	template<typename Type>
+	inline void DestructType(void* Ptr)
+	{
+		reinterpret_cast<Type*>(Ptr)->~Type();
+	}
+
+	template<typename Type>
+	inline TFunction<void(void*)> GetTypeDestructor()
+	{
+		if constexpr (std::is_trivially_destructible_v<Type>)
+		{
+			return TFunction<void(void*)>();
+		}
+		else
+		{
+			return TFunction<void(void*)>(&DestructType<Type>);
+		}
+	}
+
+	/*
+	 * Evaluation VM Stack Name
+	 * 
+	 * A small struct to wrap an FName and cache its type hash.
+	 */
+	struct FEvaluationVMStackName final
+	{
+		// Implicitly coerce or construct a stack name from an FName
+		FEvaluationVMStackName(const FName& InName)
+			: Name(InName)
+			, NameHash(GetTypeHash(InName))
+		{}
+
+		FName Name;
+		uint32 NameHash = 0;
+	};
+
+	/*
+	 * Evaluation VM Stack Entry
+	 * 
+	 * A small header for each stack entry.
+	 */
+	struct alignas(16) FEvaluationVMStackEntry final
+	{
+		union
+		{
+			// Pointer to the previous entry beneath us on the stack
+			FEvaluationVMStackEntry* Prev = nullptr;
+
+			// Dummy padding to ensure this structure remains 16 bytes even on 32-bit architectures
+			uint64 DummyPadding;
+		};
+
+		// Allocated size of this entry, includes sizeof(FEvaluationVMStackEntry)
+		uint32 Size = 0;
+
+		// Metadata flags
+		uint32 Flags = 0;
+
+		// Entry's actual data follows in memory, aligned to 16 bytes
+
+		void* GetValuePtr() { return reinterpret_cast<uint8*>(this) + sizeof(FEvaluationVMStackEntry); }
+		const void* GetValuePtr() const { return reinterpret_cast<const uint8*>(this) + sizeof(FEvaluationVMStackEntry); }
+	};
+
+	static_assert(sizeof(FEvaluationVMStackEntry) == 16, "Expected 16 bytes to ensure proper alignment");
+
+	/*
+	 * Evaluation VM Stack
+	 * 
+	 * Represents a named and typed VM stack.
+	 */
+	struct FEvaluationVMStack final
+	{
+		// Name of the stack
+		FName Name;
+
+		// Type ID used for runtime checks to make sure we always read/write using the same type
+		// @see ANIM_NEXT_ENABLE_EVALUATION_STACK_USAGE
+		uint32 TypeID = 0;
+
+		// Pointer to the entry on top of the stack or nullptr if the stack is empty
+		FEvaluationVMStackEntry* Top = nullptr;
+
+		// For non-trivially destructible types, the destructor to call
+		// Only used if the stack is not empty when it is destroyed
+		TFunction<void(void*)> TypeDestructor;
+
+		FEvaluationVMStack() = default;
+
+		// Disable copy but enable move
+		FEvaluationVMStack(const FEvaluationVMStack&) = delete;
+		FEvaluationVMStack(FEvaluationVMStack&&) = default;
+		FEvaluationVMStack& operator=(const FEvaluationVMStack&) = delete;
+		FEvaluationVMStack& operator=(FEvaluationVMStack&&) = default;
+
+		// The destructor will free remaining entries
+		~FEvaluationVMStack();
+	};
+
+	/*
+	 * Evaluation VM
+	 *
+	 * This struct holds the internal state when we execute and FEvaluationProgram.
+	 * Our virtual machine is stack based, values are generally pushed and popped.
+	 * 
+	 * @see FEvaluationProgram
+	 */
+	struct ANIMNEXT_API FEvaluationVM
+	{
+		FEvaluationVM() = default;
+		FEvaluationVM(EEvaluationFlags InEvaluationFlags, const UE::AnimNext::FReferencePose& InReferencePose, int32 InCurrentLOD);
+
+		// Allow move
+		FEvaluationVM(FEvaluationVM&&) = default;
+		FEvaluationVM& operator=(FEvaluationVM&&) = default;
+
+		// Returns true if we are fully initialized
+		[[nodiscard]] bool IsValid() const;
+
+		// Returns the evaluation flags
+		[[nodiscard]] EEvaluationFlags GetFlags() const { return EvaluationFlags; }
+
+		// Pushes and moves the specified value on its evaluation stack
+		template<typename ValueType>
+		void PushValue(const FEvaluationVMStackName& StackName, ValueType&& Value);
+
+		// Pops the top value from the evaluation stack
+		// Returns true if the stack was not empty and a value is returned, false otherwise
+		template<typename ValueType>
+		[[nodiscard]] bool PopValue(const FEvaluationVMStackName& StackName, ValueType& OutValue);
+
+		// Returns an immutable pointer to a value at some offset from the top of the evaluation stack
+		// An offset of 0 means the top of the stack
+		// Returns valid pointer if an entry exists, nullptr otherwise
+		template<typename ValueType>
+		[[nodiscard]] const ValueType* PeekValue(const FEvaluationVMStackName& StackName, uint32 Offset) const;
+
+		// Returns a mutable pointer to a value at some offset from the top of the evaluation stack
+		// An offset of 0 means the top of the stack
+		// Returns valid pointer if an entry exists, nullptr otherwise
+		template<typename ValueType>
+		[[nodiscard]] ValueType* PeekValueMutable(const FEvaluationVMStackName& StackName, uint32 Offset);
+
+		// Trims internal containers to only use as much memory as required to avoid slack
+		void Shrink();
+
+		// Returns a reference keyframe (bind pose or additive identity)
+		[[nodiscard]] FKeyframeState MakeReferenceKeyframe(bool bAdditiveKeyframe) const;
+
+		// Returns an uninitialized keyframe with memory pre-allocated
+		[[nodiscard]] FKeyframeState MakeUninitializedKeyframe(bool bAdditiveKeyframe) const;
+
+		// Returns the bone container we are evaluating with
+		[[nodiscard]] const FBoneContainer& GetBoneContainer() const;
+
+	private:
+		// Disallow copy
+		FEvaluationVM(const FEvaluationVM&) = delete;
+		FEvaluationVM& operator=(const FEvaluationVM&) = delete;
+
+		// Retrieves and optionally creates the stack (if not found) by the specified name
+		FEvaluationVMStack& GetOrCreateStack(const FEvaluationVMStackName& StackName, uint32 TypeID);
+
+		// Returns the specified stack if found, nullptr otherwise
+		FEvaluationVMStack* FindStack(const FEvaluationVMStackName& StackName, uint32 TypeID);
+		const FEvaluationVMStack* FindStack(const FEvaluationVMStackName& StackName, uint32 TypeID) const;
+
+		const void* PeekValueImpl(const FEvaluationVMStackName& StackName, uint32 TypeID, uint32 Offset) const;
+
+		// Various internal stacks that tasks can use (e.g. keyframe state)
+		TMap<FName, FEvaluationVMStack> InternalStacks;
+
+		// The bone container we are evaluating with
+		FBoneContainer BoneContainer;
+
+		// Reference pose
+		const UE::AnimNext::FReferencePose* ReferencePose = nullptr;
+
+		// Current LOD we are evaluating at
+		int32 CurrentLOD = 0;
+
+		// Flags that control what we wish to evaluate
+		EEvaluationFlags EvaluationFlags = EEvaluationFlags::All;
+
+		// TODO: Use a stack based allocator for the stack entries and mark freed entries on Pop and coalesce when we can
+		// Can allocate stack segments from memstack, limit allocation size to 1024 bytes, allocate segments of 8192 bytes
+	};
+
+	//////////////////////////////////////////////////////////////////////////
+	// Inline implementations
+
+	template<typename ValueType>
+	inline void FEvaluationVM::PushValue(const FEvaluationVMStackName& StackName, ValueType&& Value)
+	{
+		static_assert(alignof(ValueType) <= alignof(FEvaluationVMStackEntry), "Expected an alignment smaller or equal to 16");
+
+		const uint32 TypeID = GetTypeID<ValueType>();
+		FEvaluationVMStack& Stack = GetOrCreateStack(StackName, TypeID);
+
+		if (!ensureMsgf(Stack.TypeID == TypeID, TEXT("Type mismatch! This evaluation stack is being queried with a different type than it was created with")))
+		{
+			return;
+		}
+
+		const uint32 EntrySize = sizeof(FEvaluationVMStackEntry) + sizeof(Value);
+		FEvaluationVMStackEntry* EntryPtr = reinterpret_cast<FEvaluationVMStackEntry*>(FMemory::Malloc(EntrySize, alignof(FEvaluationVMStackEntry)));
+		EntryPtr->Size = EntrySize;
+		EntryPtr->Prev = Stack.Top;
+		EntryPtr->Flags = 0;
+
+		ValueType* ValuePtr = reinterpret_cast<ValueType*>(EntryPtr->GetValuePtr());
+		new(ValuePtr) ValueType(MoveTemp(Value));
+
+		if (Stack.Top == nullptr)
+		{
+			Stack.TypeDestructor = GetTypeDestructor<ValueType>();
+		}
+
+		Stack.Top = EntryPtr;
+	}
+
+	template<typename ValueType>
+	[[nodiscard]] inline bool FEvaluationVM::PopValue(const FEvaluationVMStackName& StackName, ValueType& OutValue)
+	{
+		FEvaluationVMStack* Stack = FindStack(StackName, GetTypeID<ValueType>());
+		if (Stack == nullptr)
+		{
+			// Stack not found
+			return false;
+		}
+
+		FEvaluationVMStackEntry* EntryPtr = Stack->Top;
+		if (EntryPtr == nullptr)
+		{
+			// Stack is empty
+			return false;
+		}
+
+		Stack->Top = EntryPtr->Prev;
+
+		ValueType* ValuePtr = reinterpret_cast<ValueType*>(EntryPtr->GetValuePtr());
+		OutValue = MoveTemp(*ValuePtr);
+
+		ValuePtr->~ValueType();
+		FMemory::Free(EntryPtr);
+
+		return true;
+	}
+
+	template<typename ValueType>
+	[[nodiscard]] inline const ValueType* FEvaluationVM::PeekValue(const FEvaluationVMStackName& StackName, uint32 Offset) const
+	{
+		return reinterpret_cast<const ValueType*>(PeekValueImpl(StackName, GetTypeID<ValueType>(), Offset));
+	}
+
+	template<typename ValueType>
+	[[nodiscard]] inline ValueType* FEvaluationVM::PeekValueMutable(const FEvaluationVMStackName& StackName, uint32 Offset)
+	{
+		// Re-use const impl and cast it away
+		return const_cast<ValueType*>(reinterpret_cast<const ValueType*>(PeekValueImpl(StackName, GetTypeID<ValueType>(), Offset)));
+	}
+
+	//////////////////////////////////////////////////////////////////////////
+	// Various commonly used VM stacks
+
+	// A stack of FKeyframeState instances used when sampling sequences and blending their results
+	extern const FEvaluationVMStackName KEYFRAME_STACK_NAME;
+	ANIM_NEXT_ENABLE_EVALUATION_STACK_USAGE(TUniquePtr<FKeyframeState>)
+}
diff --git a/Engine/Plugins/Experimental/AnimNext/Source/AnimNext/Public/EvaluationVM/KeyframeState.h b/Engine/Plugins/Experimental/AnimNext/Source/AnimNext/Public/EvaluationVM/KeyframeState.h
new file mode 100644
--- /dev/null
+++ b/Engine/Plugins/Experimental/AnimNext/Source/AnimNext/Public/EvaluationVM/KeyframeState.h
@@ -0,0 +1,34 @@
+// Copyright Epic Games, Inc. All Rights Reserved.
+
+#pragma once
+
+#include "CoreMinimal.h"
+
+#include "Animation/AnimCurveTypes.h"
+#include "Animation/AttributesRuntime.h"
+#include "LODPose.h"
+
+namespace UE::AnimNext
+{
+	/*
+	 * Key Frame State
+	 * 
+	 * This struct holds sampled keyframe state or interpolated keyframe state.
+	 * It holds a pose, trajectory, curves, attributes, etc that one might find in
+	 * an animation sequence or as an output of an animation node.
+	 */
+	struct ANIMNEXT_API FKeyframeState
+	{
+		// Joint transforms at a particular LOD (on memstack)
+		FLODPoseStack Pose;
+
+		// Float curves at a particular LOD (on memstack)
+		FBlendedCurve Curves;
+
+		// Attributes at a particular LOD (on memstack)
+		UE::Anim::FStackAttributeContainer Attributes;
+
+		[[nodiscard]] bool IsValid() const { return Pose.IsValid(); }
+		[[nodiscard]] bool IsAdditive() const { return Pose.IsAdditive(); }
+	};
+}
diff --git a/Engine/Plugins/Experimental/AnimNext/Source/AnimNext/Public/EvaluationVM/Tasks/ApplyAdditiveKeyframe.h b/Engine/Plugins/Experimental/AnimNext/Source/AnimNext/Public/EvaluationVM/Tasks/ApplyAdditiveKeyframe.h
new file mode 100644
--- /dev/null
+++ b/Engine/Plugins/Experimental/AnimNext/Source/AnimNext/Public/EvaluationVM/Tasks/ApplyAdditiveKeyframe.h
@@ -0,0 +1,30 @@
+// Copyright Epic Games, Inc. All Rights Reserved.
+
+#pragma once
+
+#include "CoreMinimal.h"
+#include "EvaluationVM/EvaluationTask.h"
+
+#include "ApplyAdditiveKeyframe.generated.h"
+
+/*
+ * Apply Additive Keyframe Task
+ *
+ * This pops the top two keyframes from the VM keyframe stack, it applies an additive keyframe onto its base, and pushes
+ * back the result onto the stack.
+ * The top pose should be the additive keyframe and the second to the top the base keyframe.
+ */
+USTRUCT()
+struct ANIMNEXT_API FAnimNextApplyAdditiveKeyframeTask : public FAnimNextEvaluationTask
+{
+	GENERATED_BODY()
+
+	static FAnimNextApplyAdditiveKeyframeTask Make(float BlendWeight);
+
+	// Task entry point
+	virtual void Execute(UE::AnimNext::FEvaluationVM& VM) const override;
+
+	// How much weight between the additive identity and the additive pose to apply: lerp(identity, additive, weight)
+	UPROPERTY()
+	float BlendWeight = 0.0f;
+};
diff --git a/Engine/Plugins/Experimental/AnimNext/Source/AnimNext/Public/EvaluationVM/Tasks/BlendKeyframes.h b/Engine/Plugins/Experimental/AnimNext/Source/AnimNext/Public/EvaluationVM/Tasks/BlendKeyframes.h
new file mode 100644
--- /dev/null
+++ b/Engine/Plugins/Experimental/AnimNext/Source/AnimNext/Public/EvaluationVM/Tasks/BlendKeyframes.h
@@ -0,0 +1,91 @@
+// Copyright Epic Games, Inc. All Rights Reserved.
+
+#pragma once
+
+#include "CoreMinimal.h"
+#include "EvaluationVM/EvaluationTask.h"
+
+#include "BlendKeyframes.generated.h"
+
+/*
+ * Blend Two Keyframes Task
+ *
+ * This pops the top two keyframes from the VM keyframe stack, it blends/interpolates them, and pushes
+ * back the result onto the stack.
+ * Let B be the input keyframe at the top of the stack and A be the second from the top.
+ * Then we have:
+ *     Result = Interpolate(A, B, Alpha)
+ *     Top = Interpolate(Top-1, Top, Alpha)
+ * An Alpha of 0.0 returns A while 1.0 returns B
+ * 
+ * Internally, this breaks down into the following tasks:
+ *		* FAnimNextBlendOverwriteKeyframeWithScaleTask(1.0 - InterpolationAlpha)	// todo clarify alpha
+ *		* FAnimNextBlendAddKeyframeWithScaleTask(InterpolationAlpha)
+ *		* FAnimNextNormalizeKeyframeRotationsTask
+ * 
+ * In order to blend more than two keyframes together, simply add more FAnimNextBlendAddKeyframeWithScaleTask
+ * tasks each with its corresponding weight before normalizing rotations.
+ */
+USTRUCT()
+struct ANIMNEXT_API FAnimNextBlendTwoKeyframesTask : public FAnimNextEvaluationTask
+{
+	GENERATED_BODY()
+
+	static FAnimNextBlendTwoKeyframesTask Make(float InterpolationAlpha);
+
+	// Task entry point
+	virtual void Execute(UE::AnimNext::FEvaluationVM& VM) const override;
+
+	// The interpolation alpha between the two input keyframes
+	UPROPERTY()
+	float InterpolationAlpha = 0.0f;
+};
+
+/*
+ * Blend Overwrite Keyframe With Scale Task
+ *
+ * This pops the top keyframe from the VM keyframe stack, it scales it by a factor, and pushes
+ * back the result onto the stack.
+ * Top = Top * ScaleFactor
+ * 
+ * Note that rotations will not be normalized after this task.
+ */
+USTRUCT()
+struct ANIMNEXT_API FAnimNextBlendOverwriteKeyframeWithScaleTask : public FAnimNextEvaluationTask
+{
+	GENERATED_BODY()
+
+	static FAnimNextBlendOverwriteKeyframeWithScaleTask Make(float ScaleFactor);
+
+	// Task entry point
+	virtual void Execute(UE::AnimNext::FEvaluationVM& VM) const override;
+
+	// The scale factor to apply to the keyframe
+	UPROPERTY()
+	float ScaleFactor = 1.0f;
+};
+
+/*
+ * Blend Add Keyframe With Scale Task
+ *
+ * This pops the top two keyframes (A and B) from the VM keyframe stack (let B be at the top).
+ * B is our intermediary result that we add on top of; while A is the keyframe we scale.
+ * The result is pushed back onto the stack.
+ * Top = Top + (Top-1 * ScaleFactor)
+ * 
+ * Note that rotations will not be normalized after this task.
+ */
+USTRUCT()
+struct ANIMNEXT_API FAnimNextBlendAddKeyframeWithScaleTask : public FAnimNextEvaluationTask
+{
+	GENERATED_BODY()
+
+	static FAnimNextBlendAddKeyframeWithScaleTask Make(float ScaleFactor);
+
+	// Task entry point
+	virtual void Execute(UE::AnimNext::FEvaluationVM& VM) const override;
+
+	// The scale factor to apply to the keyframe
+	UPROPERTY()
+	float ScaleFactor = 0.0f;
+};
diff --git a/Engine/Plugins/Experimental/AnimNext/Source/AnimNext/Public/EvaluationVM/Tasks/ExecuteProgram.h b/Engine/Plugins/Experimental/AnimNext/Source/AnimNext/Public/EvaluationVM/Tasks/ExecuteProgram.h
new file mode 100644
--- /dev/null
+++ b/Engine/Plugins/Experimental/AnimNext/Source/AnimNext/Public/EvaluationVM/Tasks/ExecuteProgram.h
@@ -0,0 +1,33 @@
+// Copyright Epic Games, Inc. All Rights Reserved.
+
+#pragma once
+
+#include "CoreMinimal.h"
+#include "EvaluationVM/EvaluationTask.h"
+
+#include "ExecuteProgram.generated.h"
+
+namespace UE::AnimNext
+{
+	struct FEvaluationProgram;
+}
+
+/*
+ * Execute Program Task
+ *
+ * This allows external caching of evaluation programs by deferring evaluation
+ * or repeated evaluations.
+ */
+USTRUCT()
+struct ANIMNEXT_API FAnimNextExecuteProgramTask : public FAnimNextEvaluationTask
+{
+	GENERATED_BODY()
+
+	static FAnimNextExecuteProgramTask Make(const UE::AnimNext::FEvaluationProgram* Program);
+
+	// Task entry point
+	virtual void Execute(UE::AnimNext::FEvaluationVM& VM) const override;
+
+	// The program to execute
+	const UE::AnimNext::FEvaluationProgram* Program = nullptr;
+};
diff --git a/Engine/Plugins/Experimental/AnimNext/Source/AnimNext/Public/EvaluationVM/Tasks/NormalizeRotations.h b/Engine/Plugins/Experimental/AnimNext/Source/AnimNext/Public/EvaluationVM/Tasks/NormalizeRotations.h
new file mode 100644
--- /dev/null
+++ b/Engine/Plugins/Experimental/AnimNext/Source/AnimNext/Public/EvaluationVM/Tasks/NormalizeRotations.h
@@ -0,0 +1,23 @@
+// Copyright Epic Games, Inc. All Rights Reserved.
+
+#pragma once
+
+#include "CoreMinimal.h"
+#include "EvaluationVM/EvaluationTask.h"
+
+#include "NormalizeRotations.generated.h"
+
+/*
+ * Normalize Keyframe Rotations Task
+ *
+ * This pop the top keyframe from the VM keyframe stack, it normalizes its rotations, and pushes
+ * back the result onto the stack.
+ */
+USTRUCT()
+struct ANIMNEXT_API FAnimNextNormalizeKeyframeRotationsTask : public FAnimNextEvaluationTask
+{
+	GENERATED_BODY()
+
+	// Task entry point
+	virtual void Execute(UE::AnimNext::FEvaluationVM& VM) const override;
+};
diff --git a/Engine/Plugins/Experimental/AnimNext/Source/AnimNext/Public/EvaluationVM/Tasks/PushAnimSequenceKeyframe.h b/Engine/Plugins/Experimental/AnimNext/Source/AnimNext/Public/EvaluationVM/Tasks/PushAnimSequenceKeyframe.h
new file mode 100644
--- /dev/null
+++ b/Engine/Plugins/Experimental/AnimNext/Source/AnimNext/Public/EvaluationVM/Tasks/PushAnimSequenceKeyframe.h
@@ -0,0 +1,50 @@
+// Copyright Epic Games, Inc. All Rights Reserved.
+
+#pragma once
+
+#include "CoreMinimal.h"
+#include "EvaluationVM/EvaluationTask.h"
+
+#include "PushAnimSequenceKeyframe.generated.h"
+
+class UAnimSequence;
+
+/*
+ * Push Anim Sequence Keyframe Task
+ *
+ * This pushes an anim sequence keyframe onto the top of the VM keyframe stack.
+ */
+USTRUCT()
+struct ANIMNEXT_API FAnimNextAnimSequenceKeyframeTask : public FAnimNextEvaluationTask
+{
+	GENERATED_BODY()
+
+	static FAnimNextAnimSequenceKeyframeTask MakeFromSampleTime(TWeakObjectPtr<UAnimSequence> AnimSequence, double SampleTime, bool bInterpolate);
+	static FAnimNextAnimSequenceKeyframeTask MakeFromKeyframeIndex(TWeakObjectPtr<UAnimSequence> AnimSequence, uint32 KeyframeIndex);
+
+	// Task entry point
+	virtual void Execute(UE::AnimNext::FEvaluationVM& VM) const override;
+
+	// Anim Sequence to grab the keyframe from
+	UPROPERTY()
+	TWeakObjectPtr<UAnimSequence> AnimSequence;
+
+	// The point in time within the animation sequence at which we sample the keyframe.
+	// If negative, the sample time hasn't been provided and we should use the keyframe index.
+	UPROPERTY()
+	double SampleTime = -1.0;
+
+	// The specific keyframe within the animation sequence to retrieve.
+	// If ~0, the keyframe index hasn't been provided and we should use the sample time.
+	UPROPERTY()
+	uint32 KeyframeIndex = ~0;
+
+	// Whether to interpolate or step the animation sequence.
+	// Only used when the sample time is used.
+	UPROPERTY()
+	bool bInterpolate = false;
+
+	// Whether to extract trajectory or not
+	UPROPERTY()
+	bool bExtractTrajectory = false;
+};
diff --git a/Engine/Plugins/Experimental/AnimNext/Source/AnimNext/Public/EvaluationVM/Tasks/PushReferenceKeyframe.h b/Engine/Plugins/Experimental/AnimNext/Source/AnimNext/Public/EvaluationVM/Tasks/PushReferenceKeyframe.h
new file mode 100644
--- /dev/null
+++ b/Engine/Plugins/Experimental/AnimNext/Source/AnimNext/Public/EvaluationVM/Tasks/PushReferenceKeyframe.h
@@ -0,0 +1,31 @@
+// Copyright Epic Games, Inc. All Rights Reserved.
+
+#pragma once
+
+#include "CoreMinimal.h"
+#include "EvaluationVM/EvaluationTask.h"
+
+#include "PushReferenceKeyframe.generated.h"
+
+/*
+ * Push Reference Keyframe Task
+ *
+ * This pushes a reference keyframe onto the top of the VM keyframe stack.
+ * This task can be used to push the reference keyframe from a source skeleton
+ * or the additive identity.
+ */
+USTRUCT()
+struct ANIMNEXT_API FAnimNextPushReferenceKeyframeTask : public FAnimNextEvaluationTask
+{
+	GENERATED_BODY()
+
+	static FAnimNextPushReferenceKeyframeTask MakeFromSkeleton();
+	static FAnimNextPushReferenceKeyframeTask MakeFromAdditiveIdentity();
+
+	// Task entry point
+	virtual void Execute(UE::AnimNext::FEvaluationVM& VM) const override;
+
+	// Whether or not the reference pose comes from the skeleton or is the additive identity
+	UPROPERTY()
+	bool bIsAdditive = false;
+};
diff --git a/Engine/Plugins/Experimental/AnimNext/Source/AnimNext/Public/GenerationTools.h b/Engine/Plugins/Experimental/AnimNext/Source/AnimNext/Public/GenerationTools.h
--- a/Engine/Plugins/Experimental/AnimNext/Source/AnimNext/Public/GenerationTools.h
+++ b/Engine/Plugins/Experimental/AnimNext/Source/AnimNext/Public/GenerationTools.h
@@ -88,7 +88,7 @@ public:
 	static void RemapPose(int32 LODLevel
 		, const FPoseContext& SourcePose
 		, const FReferencePose& RefPose
-		, FLODPose& TargetPose); // TODO : see how to deal with heap and stack variants without duplicating functions (and if possible without templating the function)
+		, FLODPose& TargetPose);
 
 	// Converts AnimNext pose to AnimBP Pose
 	// This function expects both poses to have the same LOD (number of bones and indexes)
@@ -96,7 +96,7 @@ public:
 	static void RemapPose(int32 LODLevel
 		, const FReferencePose& RefPose
 		, const FLODPose& SourcePose
-		, FPoseContext& TargetPose); // TODO : see how to deal with heap and stack variants without duplicating functions (and if possible without templating the function)
+		, FPoseContext& TargetPose);
 };
 
 } // namespace UE::AnimNext
diff --git a/Engine/Plugins/Experimental/AnimNext/Source/AnimNext/Public/LODPose.h b/Engine/Plugins/Experimental/AnimNext/Source/AnimNext/Public/LODPose.h
--- a/Engine/Plugins/Experimental/AnimNext/Source/AnimNext/Public/LODPose.h
+++ b/Engine/Plugins/Experimental/AnimNext/Source/AnimNext/Public/LODPose.h
@@ -1,4 +1,4 @@
-﻿// Copyright Epic Games, Inc. All Rights Reserved.
+// Copyright Epic Games, Inc. All Rights Reserved.
 
 #pragma once
 
@@ -6,6 +6,7 @@
 #include "ReferencePose.h"
 #include "BoneIndices.h"
 #include "TransformArray.h"
+#include "TransformArrayOperations.h"
 #include "LODPose.generated.h"
 
 namespace UE::AnimNext
@@ -22,41 +23,31 @@ enum class ELODPoseFlags : uint8
 
 ENUM_CLASS_FLAGS(ELODPoseFlags);
 
-template <typename AllocatorType>
-struct TLODPose
+struct FLODPose
 {
 	static constexpr int32 INVALID_LOD_LEVEL = -1;
 
-	TTransformArray<AllocatorType> LocalTransforms;
+	FTransformArrayView LocalTransformsView;
 	const FReferencePose* RefPose = nullptr;
 	int32 LODLevel = INVALID_LOD_LEVEL;
 	ELODPoseFlags Flags = ELODPoseFlags::None;
 
-	TLODPose() = default;
-
-	TLODPose(const FReferencePose& InRefPose, int32 InLODLevel, bool bSetRefPose = true, bool bAdditive = false)
-	{
-		PrepareForLOD(InRefPose, InLODLevel, bSetRefPose, bAdditive);
-	}
+	FLODPose() = default;
 
-	void PrepareForLOD(const FReferencePose& InRefPose, int32 InLODLevel, bool bSetRefPose = true, bool bAdditive = false)
+	void CopyFrom(const FLODPose& Other)
 	{
-		LODLevel = InLODLevel;
-		RefPose = &InRefPose;
+		check(LODLevel == Other.LODLevel);
+		check(RefPose == Other.RefPose);
+		check(LocalTransformsView.Num() == Other.LocalTransformsView.Num());
 
-		const int32 NumTransforms = InRefPose.GetNumBonesForLOD(InLODLevel);
-		LocalTransforms.SetNum(NumTransforms);
-		Flags = (ELODPoseFlags)(bAdditive ? (Flags | ELODPoseFlags::Additive) : Flags & ELODPoseFlags::Additive);
-
-		if (bSetRefPose && NumTransforms > 0)
-		{
-			SetRefPose(bAdditive);
-		}
+		// Copy over the flags and transforms from our source
+		Flags = Other.Flags;
+		CopyTransforms(LocalTransformsView, Other.LocalTransformsView);
 	}
 
 	void SetRefPose(bool bAdditive = false)
 	{
-		const int32 NumTransforms = LocalTransforms.Num();
+		const int32 NumTransforms = LocalTransformsView.Num();
 
 		if (NumTransforms > 0)
 		{
@@ -67,7 +58,7 @@ struct TLODPose
 			else
 			{
 				check(RefPose != nullptr);
-				LocalTransforms.CopyTransforms(RefPose->ReferenceLocalTransforms, 0, NumTransforms);
+				CopyTransforms(LocalTransformsView, RefPose->ReferenceLocalTransforms.GetConstView(), 0, NumTransforms);
 			}
 		}
 
@@ -82,7 +73,7 @@ struct TLODPose
 
 	void SetIdentity(bool bAdditive = false)
 	{
-		LocalTransforms.SetIdentity(bAdditive);
+		UE::AnimNext::SetIdentity(LocalTransformsView, bAdditive);
 	}
 
 	int32 GetNumBones() const
@@ -119,6 +110,16 @@ struct TLODPose
 		return RefPose != nullptr ? RefPose->Skeleton.Get() : nullptr;
 	}
 
+	bool IsValid() const
+	{
+		return RefPose != nullptr;
+	}
+
+	bool IsAdditive() const
+	{
+		return (Flags & ELODPoseFlags::Additive) != ELODPoseFlags::None;
+	}
+
 	/** Disable Retargeting */
 	void SetDisableRetargeting(bool bDisableRetargeting)
 	{
@@ -154,21 +155,48 @@ struct TLODPose
 	{
 		return (Flags & ELODPoseFlags::UseSourceData) != ELODPoseFlags::None;
 	}
+};
 
+template <typename AllocatorType>
+struct TLODPose : public FLODPose
+{
+	TTransformArray<AllocatorType> LocalTransforms;
+
+	TLODPose() = default;
+
+	TLODPose(const FReferencePose& InRefPose, int32 InLODLevel, bool bSetRefPose = true, bool bAdditive = false)
+	{
+		PrepareForLOD(InRefPose, InLODLevel, bSetRefPose, bAdditive);
+	}
+
+	void PrepareForLOD(const FReferencePose& InRefPose, int32 InLODLevel, bool bSetRefPose = true, bool bAdditive = false)
+	{
+		LODLevel = InLODLevel;
+		RefPose = &InRefPose;
+
+		const int32 NumTransforms = InRefPose.GetNumBonesForLOD(InLODLevel);
+		LocalTransforms.SetNum(NumTransforms);
+		LocalTransformsView = LocalTransforms.GetView();
+
+		Flags = (ELODPoseFlags)(bAdditive ? (Flags | ELODPoseFlags::Additive) : Flags & ELODPoseFlags::Additive);
+
+		if (bSetRefPose && NumTransforms > 0)
+		{
+			SetRefPose(bAdditive);
+		}
+	}
 };
 
 using FLODPoseHeap = TLODPose<FDefaultAllocator>;
 using FLODPoseStack = TLODPose<FAnimStackAllocator>;
 
-using FLODPose = FLODPoseHeap;
-
 }
 
 // USTRUCT wrapper for LOD pose
 USTRUCT()
 struct FAnimNextLODPose
 #if CPP
-	: UE::AnimNext::FLODPose
+	: UE::AnimNext::FLODPoseHeap
 #endif
 {
 	GENERATED_BODY()
diff --git a/Engine/Plugins/Experimental/AnimNext/Source/AnimNext/Public/ReferencePose.h b/Engine/Plugins/Experimental/AnimNext/Source/AnimNext/Public/ReferencePose.h
--- a/Engine/Plugins/Experimental/AnimNext/Source/AnimNext/Public/ReferencePose.h
+++ b/Engine/Plugins/Experimental/AnimNext/Source/AnimNext/Public/ReferencePose.h
@@ -30,9 +30,19 @@ ENUM_CLASS_FLAGS(EReferencePoseGenerationFlags);
 template <typename AllocatorType>
 struct TReferencePose
 {
+	// Transform array of our bind pose sorted by LOD, allows us to truncate the array for a specific LOD
+	// Higher LOD come first
 	TTransformArray<AllocatorType> ReferenceLocalTransforms;
+
+	// List of LOD bone indices for each LOD
+	// Each list of bone indices is a mapping of: BoneIndexAtLOD -> BoneIndexInSkeleton
 	TArray<TArray<FBoneIndexType, AllocatorType>, AllocatorType> LODBoneIndexes;
+
+	// List of skeleton bone indices for each LOD
+	// Each list of skeleton bone indices is a mapping of: BoneIndexInSkeleton -> BoneIndexAtLOD
 	TArray<TArray<FBoneIndexType, AllocatorType>, AllocatorType> SkeletonToLODBoneIndexes;
+
+	// Number of bones for each LOD
 	TArray<int32, AllocatorType> LODNumBones;
 
 	TWeakObjectPtr<const USkeletalMesh> SkeletalMesh = nullptr;
@@ -146,7 +156,7 @@ struct TReferencePose
 		return ReferenceLocalTransforms[LODBoneIndex];
 	}
 
-	const FQuat& GetRefPoseRotation(int32 LODBoneIndex) const
+	const FQuat GetRefPoseRotation(int32 LODBoneIndex) const
 	{
 		const int32 NumBonesLOD0 = LODBoneIndexes[0].Num();
 		check(LODBoneIndex < NumBonesLOD0);
@@ -154,7 +164,7 @@ struct TReferencePose
 		return ReferenceLocalTransforms[LODBoneIndex].GetRotation();
 	}
 
-	const FVector& GetRefPoseTranslation(int32 LODBoneIndex) const
+	const FVector GetRefPoseTranslation(int32 LODBoneIndex) const
 	{
 		const int32 NumBonesLOD0 = LODBoneIndexes[0].Num();
 		check(LODBoneIndex < NumBonesLOD0);
@@ -162,7 +172,7 @@ struct TReferencePose
 		return ReferenceLocalTransforms[LODBoneIndex].GetTranslation();
 	}
 
-	const FVector& GetRefPoseScale3D(int32 LODBoneIndex) const
+	const FVector GetRefPoseScale3D(int32 LODBoneIndex) const
 	{
 		const int32 NumBonesLOD0 = LODBoneIndexes[0].Num();
 		check(LODBoneIndex < NumBonesLOD0);
diff --git a/Engine/Plugins/Experimental/AnimNext/Source/AnimNext/Public/TransformArray.h b/Engine/Plugins/Experimental/AnimNext/Source/AnimNext/Public/TransformArray.h
--- a/Engine/Plugins/Experimental/AnimNext/Source/AnimNext/Public/TransformArray.h
+++ b/Engine/Plugins/Experimental/AnimNext/Source/AnimNext/Public/TransformArray.h
@@ -1,9 +1,10 @@
-﻿// Copyright Epic Games, Inc. All Rights Reserved.
+// Copyright Epic Games, Inc. All Rights Reserved.
 
 #pragma once
 
 #include "CoreMinimal.h"
 #include "Animation/AnimTypes.h"
+#include "TransformArrayView.h"
 #include <algorithm>
 
 #define ANIM_ENABLE_POINTER_ITERATION 1
@@ -81,6 +82,16 @@ struct TTransformArrayAoS
 		return Transforms[Index];
 	}
 
+	FTransformArrayAoSView GetView()
+	{
+		return FTransformArrayAoSView(Transforms);
+	}
+
+	FTransformArrayAoSConstView GetConstView() const
+	{
+		return FTransformArrayAoSConstView(Transforms);
+	}
+
 	/** Set this transform array to the weighted blend of the supplied two transforms. */
 	void Blend(const TTransformArrayAoS& AtomArray1, const TTransformArrayAoS& AtomArray2, float BlendWeight)
 	{
@@ -238,80 +249,6 @@ private:
 using FTransformArrayAoSHeap = TTransformArrayAoS<FDefaultAllocator>;
 using FTransformArrayAoSStack = TTransformArrayAoS<FAnimStackAllocator>;
 
-// This enables using a SoA array with operator[]
-struct FTransformSoAAdapter
-{
-	FQuat& Rotation;
-	FVector& Translation;
-	FVector& Scale3D;
-
-	FTransformSoAAdapter(FQuat& InRotation, FVector& InTranslation, FVector& InScale3D)
-		: Rotation(InRotation)
-		, Translation(InTranslation)
-		, Scale3D(InScale3D)
-	{
-	}
-
-	FORCEINLINE FQuat& GetRotation() const
-	{
-		return Rotation;
-	}
-
-	FORCEINLINE void SetRotation(const FQuat& InRotation)
-	{
-		Rotation = InRotation;
-	}
-
-	FORCEINLINE FVector& GetTranslation() const
-	{
-		return Translation;
-	}
-
-	FORCEINLINE void SetTranslation(const FVector& InTranslation)
-	{
-		Translation = InTranslation;
-	}
-
-	FORCEINLINE FVector& GetScale3D() const
-	{
-		return Scale3D;
-	}
-
-	FORCEINLINE void SetScale3D(const FVector& InScale3D)
-	{
-		Scale3D = InScale3D;
-	}
-
-	FORCEINLINE operator FTransform()
-	{
-		return FTransform(Rotation, Translation, Scale3D);
-	}
-
-	FORCEINLINE operator FTransform() const
-	{
-		return FTransform(Rotation, Translation, Scale3D);
-	}
-
-	FORCEINLINE void operator= (const FTransform& Transform)
-	{
-		Rotation = Transform.GetRotation();
-		Translation = Transform.GetTranslation();
-		Scale3D = Transform.GetScale3D();
-	}
-
-	FORCEINLINE void ScaleTranslation(const FVector::FReal& Scale)
-	{
-		Translation *= Scale;
-		//DiagnosticCheckNaN_Translate();
-	}
-
-	FORCEINLINE void NormalizeRotation()
-	{
-		Rotation.Normalize();
-		//DiagnosticCheckNaN_Rotate();
-	}
-};
-
 /**
 * Transform Array Test using StructOfArrays model
 */
@@ -498,9 +435,27 @@ struct TTransformArraySoA
 		return FTransformSoAAdapter(Rotations[Index], Translations[Index], Scales3D[Index]);
 	}
 
-	const FTransformSoAAdapter operator[] (int Index) const
+	const FTransformSoAAdapterConst operator[] (int Index) const
 	{
-		return FTransformSoAAdapter(Rotations[Index], Translations[Index], Scales3D[Index]);
+		return FTransformSoAAdapterConst(Rotations[Index], Translations[Index], Scales3D[Index]);
+	}
+
+	FTransformArraySoAView GetView()
+	{
+		FTransformArraySoAView View;
+		View.Translations = Translations;
+		View.Rotations = Rotations;
+		View.Scales3D = Scales3D;
+		return View;
+	}
+
+	FTransformArraySoAConstView GetConstView() const
+	{
+		FTransformArraySoAConstView View;
+		View.Translations = Translations;
+		View.Rotations = Rotations;
+		View.Scales3D = Scales3D;
+		return View;
 	}
 
 	bool ContainsNaN() const
diff --git a/Engine/Plugins/Experimental/AnimNext/Source/AnimNext/Public/TransformArrayOperations.h b/Engine/Plugins/Experimental/AnimNext/Source/AnimNext/Public/TransformArrayOperations.h
new file mode 100644
--- /dev/null
+++ b/Engine/Plugins/Experimental/AnimNext/Source/AnimNext/Public/TransformArrayOperations.h
@@ -0,0 +1,39 @@
+// Copyright Epic Games, Inc. All Rights Reserved.
+
+#pragma once
+
+#include "CoreMinimal.h"
+#include "TransformArrayView.h"
+
+namespace UE::AnimNext
+{
+	// Sets the identity (regular or additive) over the whole destination array
+	ANIMNEXT_API void SetIdentity(const FTransformArrayAoSView& Dest, bool bIsAdditive);
+	ANIMNEXT_API void SetIdentity(const FTransformArraySoAView& Dest, bool bIsAdditive);
+
+	// Copies the specified number of transforms from a source into a destination starting at the specified start index
+	// If NumToCopy is -1, we copy until the end
+	ANIMNEXT_API void CopyTransforms(const FTransformArrayAoSView& Dest, const FTransformArrayAoSConstView& Source, int32 StartIndex = 0, int32 NumToCopy = -1);
+	ANIMNEXT_API void CopyTransforms(const FTransformArraySoAView& Dest, const FTransformArraySoAConstView& Source, int32 StartIndex = 0, int32 NumToCopy = -1);
+
+	// Normalizes rotations in input transform view
+	ANIMNEXT_API void NormalizeRotations(const FTransformArrayAoSView& Input);
+	ANIMNEXT_API void NormalizeRotations(const FTransformArraySoAView& Input);
+
+	// The additive transform view is blended with the additive identity using the provided blend weight
+	// We then accumulate the resulting transforms on top of the base transforms
+	// Delta = Blend(Identity, Additive, BlendWeight)
+	// Base.Accumulate(Delta);
+	ANIMNEXT_API void BlendWithIdentityAndAccumulate(const FTransformArrayAoSView& Base, const FTransformArrayAoSConstView& Additive, const float BlendWeight);
+	ANIMNEXT_API void BlendWithIdentityAndAccumulate(const FTransformArraySoAView& Base, const FTransformArraySoAConstView& Additive, const float BlendWeight);
+
+	// The source transforms are scaled by the provided weight and the result is written in the destination
+	// Dest = Source * ScaleWeight
+	ANIMNEXT_API void BlendOverwriteWithScale(const FTransformArrayAoSView& Dest, const FTransformArrayAoSConstView& Source, const float ScaleWeight);
+	ANIMNEXT_API void BlendOverwriteWithScale(const FTransformArraySoAView& Dest, const FTransformArraySoAConstView& Source, const float ScaleWeight);
+
+	// The source transforms are scaled by the provided weight and the result is added to the destination
+	// Dest = Dest + (Source * ScaleWeight)
+	ANIMNEXT_API void BlendAddWithScale(const FTransformArrayAoSView& Dest, const FTransformArrayAoSConstView& Source, const float ScaleWeight);
+	ANIMNEXT_API void BlendAddWithScale(const FTransformArraySoAView& Dest, const FTransformArraySoAConstView& Source, const float ScaleWeight);
+}
diff --git a/Engine/Plugins/Experimental/AnimNext/Source/AnimNext/Public/TransformArrayView.h b/Engine/Plugins/Experimental/AnimNext/Source/AnimNext/Public/TransformArrayView.h
new file mode 100644
--- /dev/null
+++ b/Engine/Plugins/Experimental/AnimNext/Source/AnimNext/Public/TransformArrayView.h
@@ -0,0 +1,138 @@
+// Copyright Epic Games, Inc. All Rights Reserved.
+
+#pragma once
+
+#include "CoreMinimal.h"
+
+namespace UE::AnimNext
+{
+	// This enables using a SoA array with operator[]
+	template <class RotationType, class TranslationType, class ScaleType>
+	struct TTransformSoAAdapter
+	{
+		RotationType& Rotation;
+		TranslationType& Translation;
+		ScaleType& Scale3D;
+
+		TTransformSoAAdapter(RotationType& InRotation, TranslationType& InTranslation, ScaleType& InScale3D)
+			: Rotation(InRotation)
+			, Translation(InTranslation)
+			, Scale3D(InScale3D)
+		{
+		}
+
+		FORCEINLINE RotationType& GetRotation() const
+		{
+			return Rotation;
+		}
+
+		FORCEINLINE void SetRotation(const FQuat& InRotation)
+		{
+			Rotation = InRotation;
+		}
+
+		FORCEINLINE TranslationType& GetTranslation() const
+		{
+			return Translation;
+		}
+
+		FORCEINLINE void SetTranslation(const FVector& InTranslation)
+		{
+			Translation = InTranslation;
+		}
+
+		FORCEINLINE ScaleType& GetScale3D() const
+		{
+			return Scale3D;
+		}
+
+		FORCEINLINE void SetScale3D(const FVector& InScale3D)
+		{
+			Scale3D = InScale3D;
+		}
+
+		FORCEINLINE operator FTransform()
+		{
+			return FTransform(Rotation, Translation, Scale3D);
+		}
+
+		FORCEINLINE operator FTransform() const
+		{
+			return FTransform(Rotation, Translation, Scale3D);
+		}
+
+		FORCEINLINE void operator= (const FTransform& Transform)
+		{
+			Rotation = Transform.GetRotation();
+			Translation = Transform.GetTranslation();
+			Scale3D = Transform.GetScale3D();
+		}
+
+		FORCEINLINE void ScaleTranslation(const FVector::FReal& Scale)
+		{
+			Translation *= Scale;
+			//DiagnosticCheckNaN_Translate();
+		}
+
+		FORCEINLINE void NormalizeRotation()
+		{
+			Rotation.Normalize();
+			//DiagnosticCheckNaN_Rotate();
+		}
+	};
+
+	using FTransformSoAAdapter = TTransformSoAAdapter<FQuat, FVector, FVector>;
+	using FTransformSoAAdapterConst = TTransformSoAAdapter<const FQuat, const FVector, const FVector>;
+
+	using FTransformArrayAoSView = TArrayView<FTransform>;
+	using FTransformArrayAoSConstView = TArrayView<const FTransform>;
+
+	struct FTransformArraySoAView
+	{
+		TArrayView<FVector> Translations;
+		TArrayView<FQuat> Rotations;
+		TArrayView<FVector> Scales3D;
+
+		int32 Num() const { return Rotations.Num(); }
+
+		inline FTransformSoAAdapter operator[](int32 Index)
+		{
+			return FTransformSoAAdapter(Rotations[Index], Translations[Index], Scales3D[Index]);
+		}
+
+		inline const FTransformSoAAdapterConst operator[](int32 Index) const
+		{
+			return FTransformSoAAdapterConst(Rotations[Index], Translations[Index], Scales3D[Index]);
+		}
+	};
+
+	struct FTransformArraySoAConstView
+	{
+		TArrayView<const FVector> Translations;
+		TArrayView<const FQuat> Rotations;
+		TArrayView<const FVector> Scales3D;
+
+		FTransformArraySoAConstView() = default;
+		FTransformArraySoAConstView(const FTransformArraySoAView& Other)	// Safely coerce from mutable view
+			: Translations(Other.Translations)
+			, Rotations(Other.Rotations)
+			, Scales3D(Other.Scales3D)
+		{}
+
+		int32 Num() const { return Rotations.Num(); }
+
+		inline const FTransformSoAAdapterConst operator[](int32 Index) const
+		{
+			return FTransformSoAAdapterConst(Rotations[Index], Translations[Index], Scales3D[Index]);
+		}
+	};
+
+#define DEFAULT_SOA_VIEW 1
+#if DEFAULT_SOA_VIEW
+	using FTransformArrayView = FTransformArraySoAView;
+	using FTransformArrayConstView = FTransformArraySoAConstView;
+#else
+	using FTransformArrayView = FTransformArrayAoSView;
+	using FTransformArrayConstView = FTransformArrayAoSConstView;
+#endif
+}
diff --git a/Engine/Plugins/Experimental/AnimNext/Source/AnimNextTestSuite/Private/AnimNextDecoratorInterfacesTest.cpp b/Engine/Plugins/Experimental/AnimNext/Source/AnimNextTestSuite/Private/AnimNextDecoratorInterfacesTest.cpp
--- a/Engine/Plugins/Experimental/AnimNext/Source/AnimNextTestSuite/Private/AnimNextDecoratorInterfacesTest.cpp
+++ b/Engine/Plugins/Experimental/AnimNext/Source/AnimNextTestSuite/Private/AnimNextDecoratorInterfacesTest.cpp
@@ -539,7 +539,8 @@ bool FAnimationAnimNextRuntimeTest_IUpdate::RunTest(const FString& InParameters)
 		AddErrorIfFalse(NodeCPtr.IsValid(), "FAnimationAnimNextRuntimeTest_IUpdate -> Failed to allocate root node instance");
 
 		// Call pre/post update on our graph
-		UpdateGraph(Context, NodeCPtr, 0.0333f);
+		FUpdateTraversalContext TraversalContext(0.0333f);
+		UpdateGraph(Context, TraversalContext, NodeCPtr);
 
 		AddErrorIfFalse(UpdatedDecorators.Num() == 6, "FAnimationAnimNextRuntimeTest_IUpdate -> Expected 6 nodes to have been visited during the update traversal");
 		AddErrorIfFalse(UpdatedDecorators[0] == FDecoratorWithChildren::DecoratorUID, "FAnimationAnimNextRuntimeTest_IUpdate -> Unexpected update order");		// NodeC
@@ -688,9 +689,8 @@ bool FAnimationAnimNextRuntimeTest_IEvaluate::RunTest(const FString& InParameter
 		AddErrorIfFalse(NodeCPtr.IsValid(), "FAnimationAnimNextRuntimeTest_IEvaluate -> Failed to allocate root node instance");
 
 		// Call pre/post evaluate on our graph
-		const EEvaluationFlags EvaluationFlags = EEvaluationFlags::All;
-		FPoseContainer PoseContainer;
-		EvaluateGraph(Context, NodeCPtr, EvaluationFlags, PoseContainer);
+		FEvaluateTraversalContext TraversalContext;
+		(void)EvaluateGraph(Context, TraversalContext, NodeCPtr);
 
 		AddErrorIfFalse(EvaluatedDecorators.Num() == 6, "FAnimationAnimNextRuntimeTest_IEvaluate -> Expected 6 nodes to have been visited during the evaluate traversal");
 		AddErrorIfFalse(EvaluatedDecorators[0] == FDecoratorWithChildren::DecoratorUID, "FAnimationAnimNextRuntimeTest_IEvaluate -> Unexpected evaluate order");		// NodeC
diff --git a/Engine/Source/Runtime/Engine/Private/BoneContainer.cpp b/Engine/Source/Runtime/Engine/Private/BoneContainer.cpp
--- a/Engine/Source/Runtime/Engine/Private/BoneContainer.cpp
+++ b/Engine/Source/Runtime/Engine/Private/BoneContainer.cpp
@@ -36,7 +36,7 @@ FBoneContainer::FBoneContainer()
 	PoseToSkeletonBoneIndexArray.Empty();
 }
 
-FBoneContainer::FBoneContainer(const TArray<FBoneIndexType>& InRequiredBoneIndexArray, const FCurveEvaluationOption& CurveEvalOption, UObject& InAsset)
+FBoneContainer::FBoneContainer(const TArrayView<const FBoneIndexType>& InRequiredBoneIndexArray, const FCurveEvaluationOption& CurveEvalOption, UObject& InAsset)
 : BoneIndicesArray(InRequiredBoneIndexArray)
 , Asset(&InAsset)
 , AssetSkeletalMesh(nullptr)
@@ -56,7 +56,7 @@ FBoneContainer::FBoneContainer(const TArray<FBoneIndexType>& InRequiredBoneIndex
 	Initialize(CurveFilterSettings);
 }
 
-FBoneContainer::FBoneContainer(const TArray<FBoneIndexType>& InRequiredBoneIndexArray, const UE::Anim::FCurveFilterSettings& InCurveFilterSettings, UObject& InAsset)
+FBoneContainer::FBoneContainer(const TArrayView<const FBoneIndexType>& InRequiredBoneIndexArray, const UE::Anim::FCurveFilterSettings& InCurveFilterSettings, UObject& InAsset)
 	: BoneIndicesArray(InRequiredBoneIndexArray)
 	, Asset(&InAsset)
 	, AssetSkeletalMesh(nullptr)
@@ -74,13 +74,13 @@ FBoneContainer::FBoneContainer(const TArray<FBoneIndexType>& InRequiredBoneIndex
 	Initialize(InCurveFilterSettings);
 }
 
-void FBoneContainer::InitializeTo(const TArray<FBoneIndexType>& InRequiredBoneIndexArray, const FCurveEvaluationOption& CurveEvalOption, UObject& InAsset)
+void FBoneContainer::InitializeTo(const TArrayView<const FBoneIndexType>& InRequiredBoneIndexArray, const FCurveEvaluationOption& CurveEvalOption, UObject& InAsset)
 {
 	const UE::Anim::FCurveFilterSettings CurveFilterSettings(CurveEvalOption.bAllowCurveEvaluation ? UE::Anim::ECurveFilterMode::DisallowFiltered : UE::Anim::ECurveFilterMode::DisallowAll, CurveEvalOption.DisallowedList, CurveEvalOption.LODIndex);
 	InitializeTo(InRequiredBoneIndexArray, CurveFilterSettings, InAsset);
 }
 
-void FBoneContainer::InitializeTo(const TArray<FBoneIndexType>& InRequiredBoneIndexArray, const UE::Anim::FCurveFilterSettings& CurveFilterSettings, UObject& InAsset)
+void FBoneContainer::InitializeTo(const TArrayView<const FBoneIndexType>& InRequiredBoneIndexArray, const UE::Anim::FCurveFilterSettings& CurveFilterSettings, UObject& InAsset)
 {
 	LLM_SCOPE_BYNAME(TEXT("Animation/BoneContainer"));
 	BoneIndicesArray = InRequiredBoneIndexArray;
diff --git a/Engine/Source/Runtime/Engine/Public/BoneContainer.h b/Engine/Source/Runtime/Engine/Public/BoneContainer.h
--- a/Engine/Source/Runtime/Engine/Public/BoneContainer.h
+++ b/Engine/Source/Runtime/Engine/Public/BoneContainer.h
@@ -253,15 +253,15 @@ public:
 
 	ENGINE_API FBoneContainer();
 
-	ENGINE_API FBoneContainer(const TArray<FBoneIndexType>& InRequiredBoneIndexArray, const UE::Anim::FCurveFilterSettings& InCurveFilterSettings, UObject& InAsset);
+	ENGINE_API FBoneContainer(const TArrayView<const FBoneIndexType>& InRequiredBoneIndexArray, const UE::Anim::FCurveFilterSettings& InCurveFilterSettings, UObject& InAsset);
 	
-	ENGINE_API void InitializeTo(const TArray<FBoneIndexType>& InRequiredBoneIndexArray, const UE::Anim::FCurveFilterSettings& InCurveFilterSettings, UObject& InAsset);
+	ENGINE_API void InitializeTo(const TArrayView<const FBoneIndexType>& InRequiredBoneIndexArray, const UE::Anim::FCurveFilterSettings& InCurveFilterSettings, UObject& InAsset);
 	
 	UE_DEPRECATED(5.3, "Please use the constructor that takes a FCurveFilterSettings.")
-	ENGINE_API FBoneContainer(const TArray<FBoneIndexType>& InRequiredBoneIndexArray, const FCurveEvaluationOption& CurveEvalOption, UObject& InAsset);
+	ENGINE_API FBoneContainer(const TArrayView<const FBoneIndexType>& InRequiredBoneIndexArray, const FCurveEvaluationOption& CurveEvalOption, UObject& InAsset);
 
 	UE_DEPRECATED(5.3, "Please use InitializeTo that takes a FCurveFilterSettings.")
-	ENGINE_API void InitializeTo(const TArray<FBoneIndexType>& InRequiredBoneIndexArray, const FCurveEvaluationOption& CurveEvalOption, UObject& InAsset);
+	ENGINE_API void InitializeTo(const TArrayView<const FBoneIndexType>& InRequiredBoneIndexArray, const FCurveEvaluationOption& CurveEvalOption, UObject& InAsset);
 
 	/** Resets the container and reclaims all allocated memory but preserve the serial number. */
 	ENGINE_API void Reset();
