AnimNext: Scheduler, Component and Parameter Improvements

Scheduler

Added schedule asset
Schedule builds a simple bytecode when edited/loaded/saved
Bytecode describes what tasks to run in order with what dependencies
Scheduler acts as an interface to the schedule world subssystem (so schedules can be run against specific worlds)
Schedule entry ultimately owns all the memory for a running schedule, mostly via the FScheduleInstanceData
Schedule is run via a graph of tick functions that are created on handle registration, held on the schedule entry

Component

Bare-bones component, still a few features to add (e.g. stopping starting ticking etc.)
Holds a handle to a parameterized schedule
Can update parameters by pushing them to named scopes in the schedule

Parameter Improvements

UI/UX

Changed parameter block workflow to be closer to how it works under the hood:
- Parameters are local state
- Graphs operate by writing local state (and reading local/global state)
- This means that parameters and graphs are created independently. Still TODO: setting defaults for parameters.
Added custom widget for selecting parameters (on pins at present - details panel still TODO)
Parameter picker can now create new parameters in-line.
Double-clicking a graph now opens it

Runtime

Patched up parameter blocks to work with the last set of RigVM refactors.
Rebased param block node, graph and schema to RigVM.
Re-duplicated compilation logic for parameter blocks. This should be subject to a later refactor to move the compilation our of URigVMBlueprint.
Switched property bag to create object ptr properties rather than raw object properties
Added support for attaching param stacks to anim BP execution, propagated from a schedule instance.
Consolidated specific layer memory ownership into stack layer storage types: FUObjectLayer and FInstancedPropertyBagLayer. Raw data layers are also still available.
Param stacks now have 'parent' stacks that they will differ to if values are not found in queries. This is used for task-level scopes. Tasks are assigned param stacks to run (and the graph of parent scopes is built) via the schedule instance data.
Param stack results API expanded so we can get ricer info from queries
Moved FParamStackLayer out of FParamStack
Added TLayerBuilder to allow raw data layers to be built out-of-line
Updated tests

diff --git a/Engine/Plugins/Experimental/AnimNext/Config/BaseAnimNext.ini b/Engine/Plugins/Experimental/AnimNext/Config/BaseAnimNext.ini
new file mode 100644
--- /dev/null
+++ b/Engine/Plugins/Experimental/AnimNext/Config/BaseAnimNext.ini
@@ -0,0 +1,4 @@
+[/Script/AnimNext.AnimNextConfig]
++ExposedClasses=(Class="/Script/CoreUObject.Class'/Script/Engine.CharacterMovementComponent'",RootParameter="UE_CharacterMovementComponent",FilterType=AllowOnlyFunctions,FilteredFields=,bRegisterTickFunction=True)
++ExposedClasses=(Class="/Script/CoreUObject.Class'/Script/Engine.SkeletalMeshComponent'",RootParameter="UE_SkeletalMeshComponent",FilterType=AllowOnlyFunctions,FilteredFields=,bRegisterTickFunction=True)
+
diff --git a/Engine/Plugins/Experimental/AnimNext/Source/AnimNext/Internal/AnimNextConfig.h b/Engine/Plugins/Experimental/AnimNext/Source/AnimNext/Internal/AnimNextConfig.h
new file mode 100644
--- /dev/null
+++ b/Engine/Plugins/Experimental/AnimNext/Source/AnimNext/Internal/AnimNextConfig.h
@@ -0,0 +1,30 @@
+﻿// Copyright Epic Games, Inc. All Rights Reserved.
+
+#pragma once
+
+#include "Param/AnimNextObjectAdapterConfig.h"
+#include "AnimNextConfig.generated.h"
+
+namespace UE::AnimNext
+{
+	struct FParamId;
+}
+
+UCLASS(Config=AnimNext)
+class ANIMNEXT_API UAnimNextConfig : public UObject
+{
+	GENERATED_BODY()
+
+private:
+	// UObject interface
+#if WITH_EDITOR
+	virtual void PostEditChangeProperty(FPropertyChangedEvent& PropertyChangedEvent) override;
+#endif
+
+private:
+	friend struct UE::AnimNext::FParamId;
+
+	/** The classes that are exposed to AnimNext systems */
+	UPROPERTY(Config, EditAnywhere, Category = "Exposed Classes")
+	TArray<FAnimNextObjectAdapterConfig> ExposedClasses;
+};
\ No newline at end of file
diff --git a/Engine/Plugins/Experimental/AnimNext/Source/AnimNext/Internal/Graph/AnimGraph/AnimNode_AnimNextParameters.h b/Engine/Plugins/Experimental/AnimNext/Source/AnimNext/Internal/Graph/AnimGraph/AnimNode_AnimNextParameters.h
--- a/Engine/Plugins/Experimental/AnimNext/Source/AnimNext/Internal/Graph/AnimGraph/AnimNode_AnimNextParameters.h
+++ b/Engine/Plugins/Experimental/AnimNext/Source/AnimNext/Internal/Graph/AnimGraph/AnimNode_AnimNextParameters.h
@@ -4,10 +4,11 @@
 
 #include "CoreMinimal.h"
 #include "Animation/AnimNodeBase.h"
-#include "Param/ParamStack.h"
+#include "Param/ParamStackLayerHandle.h"
 #include "AnimNode_AnimNextParameters.generated.h"
 
-class UAnimNextParameterBlock;
+class IAnimNextParameterSourceInterface;
+class UAnimNextParameters;
 
 USTRUCT(BlueprintInternalUseOnly)
 struct FAnimNode_AnimNextParameters : public FAnimNode_Base
@@ -19,28 +20,22 @@ struct FAnimNode_AnimNextParameters : public FAnimNode_Base
 	FAnimNode_AnimNextParameters() = default;
 	FAnimNode_AnimNextParameters(const FAnimNode_AnimNextParameters& InOther);
 	FAnimNode_AnimNextParameters& operator=(const FAnimNode_AnimNextParameters& InOther);
-	FAnimNode_AnimNextParameters(FAnimNode_AnimNextParameters&& InOther);
-	FAnimNode_AnimNextParameters& operator=(FAnimNode_AnimNextParameters&& InOther);
-	~FAnimNode_AnimNextParameters() = default;
+	FAnimNode_AnimNextParameters(FAnimNode_AnimNextParameters&& InOther) noexcept;
+	FAnimNode_AnimNextParameters& operator=(FAnimNode_AnimNextParameters&& InOther) noexcept;
+	virtual ~FAnimNode_AnimNextParameters() override = default;
 
 private:
 	UPROPERTY(EditAnywhere, Category = Links)
 	FPoseLink Source;
 
-#if WITH_EDITORONLY_DATA
-	UPROPERTY(EditAnywhere, Category = Settings, meta=(FoldProperty, PinHiddenByDefault))
-	TObjectPtr<UAnimNextParameterBlock> ParameterBlock;
-#endif
+	UPROPERTY(EditAnywhere, Category = Settings)
+	TScriptInterface<IAnimNextParameterSourceInterface> Parameters;
 
 	// Cache previous param block so we know when it changes via pin
-	UAnimNextParameterBlock* PreviousParameterBlock = nullptr;
+	IAnimNextParameterSourceInterface* PreviousParameters = nullptr;
 
 	// Cached layer
-	UE::AnimNext::FParamStack::FLayerHandle ParamLayerHandle;
-
-	// Property bag for layer backing storage
-	UPROPERTY(Transient)
-	FInstancedPropertyBag PropertyBag;
+	UE::AnimNext::FParamStackLayerHandle ParamLayerHandle;
 
 private:
 	// FAnimNode_Base
@@ -49,6 +44,4 @@ private:
 	ANIMNEXT_API virtual void CacheBones_AnyThread(const FAnimationCacheBonesContext& Context) override;
 	ANIMNEXT_API virtual void Evaluate_AnyThread(FPoseContext& Output) override;
 	ANIMNEXT_API virtual void GatherDebugData(FNodeDebugData& DebugData) override;
-
-	UAnimNextParameterBlock* GetParameterBlock() const;
 };
diff --git a/Engine/Plugins/Experimental/AnimNext/Source/AnimNext/Internal/Param/AnimNextObjectAdapterConfig.h b/Engine/Plugins/Experimental/AnimNext/Source/AnimNext/Internal/Param/AnimNextObjectAdapterConfig.h
new file mode 100644
--- /dev/null
+++ b/Engine/Plugins/Experimental/AnimNext/Source/AnimNext/Internal/Param/AnimNextObjectAdapterConfig.h
@@ -0,0 +1,56 @@
+﻿// Copyright Epic Games, Inc. All Rights Reserved.
+
+#pragma once
+
+#include "CoreMinimal.h"
+#include "Templates/SubclassOf.h"
+#include "AnimNextObjectAdapterConfig.generated.h"
+
+UENUM()
+enum class EAnimNextObjectAdapterFilterType : int32
+{
+	// Allow-list of functions and properties
+	AllowList,
+
+	// Deny-list of functions and properties
+	DenyList,
+
+	// Expose all functions but no properties
+	AllowOnlyFunctions,
+
+	// Expose all properties but no functions
+	AllowOnlyProperties,
+
+	// Expose all functions and properties
+	AllowAllPropertiesAndFunctions,
+};
+
+USTRUCT()
+struct FAnimNextObjectAdapterConfig
+{
+	GENERATED_BODY()
+
+	/** The class to expose to AnimNext */
+	UPROPERTY(EditAnywhere, Config, Category="Adapter")
+	TSubclassOf<UObject> Class;
+
+	/**
+	 * The root parameter for the object when exposed to AnimNext
+	 * If Root Parameter is set to 'MyObject', then all properties and functions derived from this class will have the
+	 * form 'MyObject.Param'.
+	 */
+	UPROPERTY(EditAnywhere, Config, Category="Adapter", meta = (CustomWidget = "ParamName"))
+	FName RootParameter;
+
+	/** How to expose the properties and functions of the class */
+	UPROPERTY(EditAnywhere, Config, Category="Adapter")
+	EAnimNextObjectAdapterFilterType FilterType = EAnimNextObjectAdapterFilterType::AllowList;
+
+	/** All functions and properties that are allowed/denied */
+	UPROPERTY(EditAnywhere, Config, Category="Adapter", meta = (EditCondition = "FilterType == EAnimNextObjectAdapterFilterType::AllowList || FilterType == EAnimNextObjectAdapterFilterType::DenyList", EditConditionHides))
+	TArray<FName> FilteredFields;
+
+	/** Whether to expose a tick function (e.g. for a component or actor) along with functions and properties */
+	UPROPERTY(EditAnywhere, Config, Category="Adapter")
+	bool bRegisterTickFunction;
+};
diff --git a/Engine/Plugins/Experimental/AnimNext/Source/AnimNext/Internal/Param/AnimNextParameterBlock.h b/Engine/Plugins/Experimental/AnimNext/Source/AnimNext/Internal/Param/AnimNextParameterBlock.h
--- a/Engine/Plugins/Experimental/AnimNext/Source/AnimNext/Internal/Param/AnimNextParameterBlock.h
+++ b/Engine/Plugins/Experimental/AnimNext/Source/AnimNext/Internal/Param/AnimNextParameterBlock.h
@@ -5,9 +5,13 @@
 #include "CoreMinimal.h"
 #include "PropertyBag.h"
 #include "RigVMCore/RigVMExecuteContext.h"
+#include "Param/IAnimNextParameterSourceInterface.h"
+#include "RigVMHost.h"
 #include "AnimNextParameterBlock.generated.h"
 
 class UEdGraph;
+struct FAnimNextScheduleGraphTask;
+struct FAnimNextScheduleParamScopeTask;
 
 namespace UE::AnimNext
 {
@@ -58,7 +62,7 @@ struct FAnimNextParameterBlockAssetRegistryExports
 
 /** An asset used to define AnimNext parameters and their bindings */
 UCLASS(MinimalAPI, BlueprintType)
-class UAnimNextParameterBlock : public UObject
+class UAnimNextParameterBlock : public URigVMHost, public IAnimNextParameterSourceInterface
 {
 	GENERATED_BODY()
 
@@ -69,26 +73,23 @@ class UAnimNextParameterBlock : public UObject
 	friend class UE::AnimNext::Editor::FParametersEditor;
 	friend struct UE::AnimNext::Editor::FUtils;
 	friend struct FAnimNode_AnimNextParameters;
+	friend struct FAnimNextScheduleGraphTask;
+	friend struct FAnimNextScheduleParamScopeEntryTask;
 
-	void Run(const UE::AnimNext::FContext& Context) const;
+	// IAnimNextParameterSourceInterface interface
+	virtual void UpdateLayer(UE::AnimNext::FParamStackLayerHandle& InHandle) const override;
+	virtual UE::AnimNext::FParamStackLayerHandle CacheLayer() const override;
+	virtual bool ShouldCacheLayer(const UE::AnimNext::FParamStackLayerHandle& InHandle) const override;
 
 	// UObject interface
+	virtual void PostLoad() override;
 	virtual void PostRename(UObject* OldOuter, const FName OldName) override;
 	virtual void GetPreloadDependencies(TArray<UObject*>& OutDeps) override;
 	virtual void GetAssetRegistryTags(TArray<FAssetRegistryTag>& OutTags) const override;
 
-	// Support rig VM execution
-	ANIMNEXT_API TArray<FRigVMExternalVariable> GetRigVMExternalVariables();
-
 	UPROPERTY()
 	TObjectPtr<URigVM> RigVM;
 
-	UPROPERTY(transient)
-	FRigVMExtendedExecuteContext ExtendedExecuteContext;
-
-	UPROPERTY()
-	FRigVMRuntimeSettings VMRuntimeSettings;
-
 	UPROPERTY()
 	FInstancedPropertyBag PropertyBag;
 
diff --git a/Engine/Plugins/Experimental/AnimNext/Source/AnimNext/Internal/Param/AnimNextParameterCollection.h b/Engine/Plugins/Experimental/AnimNext/Source/AnimNext/Internal/Param/AnimNextParameterCollection.h
new file mode 100644
--- /dev/null
+++ b/Engine/Plugins/Experimental/AnimNext/Source/AnimNext/Internal/Param/AnimNextParameterCollection.h
@@ -0,0 +1,20 @@
+// Copyright Epic Games, Inc. All Rights Reserved.
+
+#pragma once
+
+#include "CoreMinimal.h"
+#include "Param/AnimNextParameterSourceRef.h"
+#include "AnimNextParameterCollection.generated.h"
+
+class UAnimNextParameterBlock;
+
+// Struct wrapper to allow an array of parameters to be nested in a container or edited & referenced in bulk
+USTRUCT()
+struct FAnimNextParameterCollection
+{
+	GENERATED_BODY()
+
+	// User-editable parameter blocks
+	UPROPERTY(EditAnywhere, Category = "Parameters")
+	TArray<FAnimNextParameterSourceRef> Parameters;
+};
\ No newline at end of file
diff --git a/Engine/Plugins/Experimental/AnimNext/Source/AnimNext/Internal/Param/AnimNextParameterExecuteContext.h b/Engine/Plugins/Experimental/AnimNext/Source/AnimNext/Internal/Param/AnimNextParameterExecuteContext.h
new file mode 100644
--- /dev/null
+++ b/Engine/Plugins/Experimental/AnimNext/Source/AnimNext/Internal/Param/AnimNextParameterExecuteContext.h
@@ -0,0 +1,55 @@
+// Copyright Epic Games, Inc. All Rights Reserved.
+
+#pragma once
+
+#include "CoreMinimal.h"
+#include "Graph/AnimNextExecuteContext.h"
+#include "AnimNextParameterExecuteContext.generated.h"
+
+class UAnimNextParameterBlock;
+
+namespace UE::AnimNext
+{
+	struct FParamStackLayerHandle;
+}
+
+USTRUCT(BlueprintType)
+struct FAnimNextParameterExecuteContext : public FAnimNextExecuteContext
+{
+	GENERATED_BODY()
+
+	FAnimNextParameterExecuteContext()
+		: FAnimNextExecuteContext()
+	{
+	}
+
+	virtual void Copy(const FRigVMExecuteContext* InOtherContext) override
+	{
+		Super::Copy(InOtherContext);
+
+		const FAnimNextParameterExecuteContext* OtherContext = (const FAnimNextParameterExecuteContext*)InOtherContext;
+		LayerHandle = OtherContext->LayerHandle;
+	}
+
+	UE::AnimNext::FParamStackLayerHandle& GetLayerHandle() const
+	{
+		return *LayerHandle;
+	}
+
+private:
+	friend class UAnimNextParameterBlock;
+
+	void SetParamContextData(UE::AnimNext::FParamStackLayerHandle& InLayerHandle)
+	{
+		LayerHandle = &InLayerHandle;
+	}
+
+	// Parameter scope context
+	UE::AnimNext::FParamStackLayerHandle* LayerHandle = nullptr;
+};
+
+USTRUCT(meta = (ExecuteContext = "FAnimNextParameterExecuteContext"))
+struct FRigUnit_AnimNextParameterBase : public FRigUnit
+{
+	GENERATED_BODY()
+};
diff --git a/Engine/Plugins/Experimental/AnimNext/Source/AnimNext/Internal/Param/AnimNextParameterSourceRef.h b/Engine/Plugins/Experimental/AnimNext/Source/AnimNext/Internal/Param/AnimNextParameterSourceRef.h
new file mode 100644
--- /dev/null
+++ b/Engine/Plugins/Experimental/AnimNext/Source/AnimNext/Internal/Param/AnimNextParameterSourceRef.h
@@ -0,0 +1,65 @@
+// Copyright Epic Games, Inc. All Rights Reserved.
+
+#pragma once
+
+#include "CoreMinimal.h"
+#include "Engine/EngineTypes.h"
+#include "Param/IAnimNextParameterSourceInterface.h"
+#include "AnimNextParameterSourceRef.generated.h"
+
+class UAnimNextParameterBlock;
+class IAnimNextParameterSourceInterface;
+
+UENUM()
+enum class EAnimNextParameterSourceRefType : int32
+{
+	// The hosting component provides parameters via its variables
+	Self,
+
+	// An external parameter block asset
+	Asset,
+
+	// Another component in this actor provides parameters via its variables
+	Component,
+
+	// This actor provides parameters via its variables
+	Actor,
+
+	// Parameters are provided inline
+	Inline
+};
+
+// Wrapper for different kinds of parameter source
+USTRUCT()
+struct FAnimNextParameterSourceRef 
+#if CPP
+	: public IAnimNextParameterSourceInterface
+#endif
+{
+	GENERATED_BODY()
+
+	// Get the parameter source that this reference refers to
+	const IAnimNextParameterSourceInterface* Get(const UObject* InContextObject) const;
+
+#if CPP
+	// IAnimNextParameterSourceInterface interface
+	virtual void UpdateLayer(UE::AnimNext::FParamStackLayerHandle& InHandle) const override {}
+	virtual UE::AnimNext::FParamStackLayerHandle CacheLayer() const override;
+#endif
+
+	// How the parameters are specified
+	UPROPERTY(EditAnywhere, Category = "Parameters")
+	EAnimNextParameterSourceRefType Type = EAnimNextParameterSourceRefType::Asset;
+
+	// External parameter source interface
+	UPROPERTY(EditAnywhere, Category = "Parameters", meta = (EditCondition = "Type == FAnimNextParameterSourceRefType::Asset", EditConditionHides))
+	TScriptInterface<IAnimNextParameterSourceInterface> Asset;
+
+	// Component ref to allow tracking of other components in an actor
+	UPROPERTY(EditAnywhere, Category = "Parameters", meta = (EditCondition = "Type == FAnimNextParameterSourceRefType::Component", EditConditionHides))
+	FComponentReference Component;
+
+	// Properties provided inline
+	UPROPERTY(EditAnywhere, Category = "Parameters", meta = (EditCondition = "Type == FAnimNextParameterSourceRefType::Inline", EditConditionHides))
+	FInstancedPropertyBag InlineParameters;
+};
diff --git a/Engine/Plugins/Experimental/AnimNext/Source/AnimNext/Internal/Param/ParamDefinition.h b/Engine/Plugins/Experimental/AnimNext/Source/AnimNext/Internal/Param/ParamDefinition.h
new file mode 100644
--- /dev/null
+++ b/Engine/Plugins/Experimental/AnimNext/Source/AnimNext/Internal/Param/ParamDefinition.h
@@ -0,0 +1,127 @@
+// Copyright Epic Games, Inc. All Rights Reserved.
+
+#pragma once
+
+#include "CoreMinimal.h"
+#include "Param/ParamId.h"
+#include "Param/ParamType.h"
+#include "Param/ParamTypeHandle.h"
+
+namespace UE::AnimNext
+{
+	struct FParams;
+	struct FParamAdapter;
+}
+
+namespace UE::AnimNext::Editor
+{
+	struct FUtils;
+}
+
+namespace UE::AnimNext
+{
+
+enum class EParamDefinitionFlags : uint8
+{
+	None = 0,
+
+	Mutable = (1 << 0),
+};
+
+ENUM_CLASS_FLAGS(EParamDefinitionFlags);
+
+// Definition of a parameter for reserved or internal use
+struct FParamDefinition
+{
+	FParamDefinition() = default;
+
+	FParamDefinition(FName InName, const FAnimNextParamType& InType, const FText& InTooltip, EParamDefinitionFlags InFlags = EParamDefinitionFlags::None)
+		: Id(InName)
+		, Name(InName)
+		, TypeHandle(InType.GetHandle())
+		, Type(InType)
+#if WITH_EDITORONLY_DATA
+		, Tooltip(InTooltip)
+#endif
+		, Property(nullptr)
+		, Function(nullptr)
+		, Flags(InFlags)
+	{
+		// Param names should not be 'none'
+		check(Name != NAME_None);
+		// Param names should not contain periods. Use underscores.
+		// Periods are only a display concern and we want parameters to be expressible as members of objects & structures. 
+		check(!Name.ToString().Contains(TEXT(".")));
+	}
+	
+	FParamDefinition(FParamId InId, const FAnimNextParamType& InType, const FText& InTooltip, EParamDefinitionFlags InFlags = EParamDefinitionFlags::None)
+		: Id(InId)
+		, Name(InId.ToName())
+		, TypeHandle(InType.GetHandle())
+		, Type(InType)
+#if WITH_EDITORONLY_DATA
+		, Tooltip(InTooltip)
+#endif
+		, Property(nullptr)
+		, Function(nullptr)
+		, Flags(InFlags)
+	{
+	}
+
+	FName GetName() const
+	{
+		return Name;
+	}
+
+	FParamTypeHandle GetTypeHandle() const
+	{
+		return TypeHandle;
+	}
+	
+	const FAnimNextParamType& GetType() const
+	{
+		return Type;
+	}
+
+private:
+	friend struct FParamId;
+	friend struct FParams;
+	friend struct FParamAdapter;
+	friend struct Editor::FUtils;
+
+	FParamDefinition(int32 InIdIndex, FName InName, const FAnimNextParamType& InType, const FText& InTooltip, EParamDefinitionFlags InFlags = EParamDefinitionFlags::None)
+		: Id(InIdIndex)
+		, Name(InName)
+		, TypeHandle(InType.GetHandle())
+		, Type(InType)
+#if WITH_EDITORONLY_DATA
+		, Tooltip(InTooltip)
+#endif
+		, Property(nullptr)
+		, Function(nullptr)
+		, Flags(InFlags)
+	{
+	}
+
+	// Make a parameter definition from its ID and a UObject
+	FParamDefinition(uint32 InIdIndex, FName InName, const UObject* InObject);
+
+	// Make a parameter definition from its ID and a property
+	FParamDefinition(uint32 InIdIndex, FName InName, const FProperty* InProperty);
+
+	// Make a parameter definition from its ID and the return value of a function
+	FParamDefinition(uint32 InIdIndex, FName InName, const UFunction* InFunction);
+
+	FParamId Id;
+	FName Name;
+	FParamTypeHandle TypeHandle;
+	FAnimNextParamType Type;
+#if WITH_EDITORONLY_DATA
+	FText Tooltip;
+#endif
+	const FProperty* Property = nullptr;
+	const UFunction* Function = nullptr;
+	EParamDefinitionFlags Flags = EParamDefinitionFlags::None;
+};
+
+}
diff --git a/Engine/Plugins/Experimental/AnimNext/Source/AnimNext/Internal/Param/Params.h b/Engine/Plugins/Experimental/AnimNext/Source/AnimNext/Internal/Param/Params.h
new file mode 100644
--- /dev/null
+++ b/Engine/Plugins/Experimental/AnimNext/Source/AnimNext/Internal/Param/Params.h
@@ -0,0 +1,41 @@
+// Copyright Epic Games, Inc. All Rights Reserved.
+
+#pragma once
+
+#include "CoreMinimal.h"
+#include "Param/ParamDefinition.h"
+
+namespace UE::AnimNext
+{
+	struct FParameterAdapterCollection;
+}
+
+namespace UE::AnimNext
+{
+
+// Built-in parameter handling
+struct FParams
+{
+	// Register a single built-in parameter
+	static FParamDefinition RegisterBuiltInParameter(const FParamDefinition& InDefinition);
+
+	// Unregister a built-in parameter
+	static void UnregisterBuiltInParameter(FName InName);
+
+	// Check whether a parameter name corresponds to a built-in parameter
+	static bool IsBuiltInParameter(FName InName);
+
+	// Get the number of registered built-in parameters
+	static int32 GetNumBuiltInParameters();
+
+	// Gets a built-in parameter definition. Asserts if it is not present
+	static const FParamDefinition& GetBuiltInParameter(FName InName);
+
+	// Finds a built-in parameter definition
+	static ANIMNEXT_API const FParamDefinition* FindBuiltInParameter(FName InName);
+
+	// Iterates over all built-in parameters
+	static ANIMNEXT_API void ForEachBuiltInParameter(TFunctionRef<void(const FParamDefinition&)> InFunction);
+};
+
+}
\ No newline at end of file
diff --git a/Engine/Plugins/Experimental/AnimNext/Source/AnimNext/Internal/Param/RigUnit_AnimNextParameterBeginExecution.h b/Engine/Plugins/Experimental/AnimNext/Source/AnimNext/Internal/Param/RigUnit_AnimNextParameterBeginExecution.h
new file mode 100644
--- /dev/null
+++ b/Engine/Plugins/Experimental/AnimNext/Source/AnimNext/Internal/Param/RigUnit_AnimNextParameterBeginExecution.h
@@ -0,0 +1,27 @@
+// Copyright Epic Games, Inc. All Rights Reserved.
+
+#pragma once
+
+#include "AnimNextParameterExecuteContext.h"
+#include "RigUnit_AnimNextParameterBeginExecution.generated.h"
+
+/**
+ * Event for driving the skeleton hierarchy with variables and rig elements
+ */
+USTRUCT(meta=(DisplayName="Execute", Category="Events", NodeColor="1, 0, 0", Keywords="Begin,Update,Tick,Forward,Event"))
+struct ANIMNEXT_API FRigUnit_AnimNextParameterBeginExecution : public FRigUnit_AnimNextParameterBase
+{
+	GENERATED_BODY()
+
+	RIGVM_METHOD()
+	void Execute();
+
+	virtual FName GetEventName() const override { return EventName; }
+	virtual bool CanOnlyExistOnce() const override { return true; }
+
+	// The execution result
+	UPROPERTY(EditAnywhere, DisplayName = "Execute", Category = "BeginExecution", meta = (Output))
+	FAnimNextParameterExecuteContext ExecuteContext;
+
+	static FName EventName;
+};
diff --git a/Engine/Plugins/Experimental/AnimNext/Source/AnimNext/Internal/Param/RigVMDispatch_GetLayerParameter.h b/Engine/Plugins/Experimental/AnimNext/Source/AnimNext/Internal/Param/RigVMDispatch_GetLayerParameter.h
new file mode 100644
--- /dev/null
+++ b/Engine/Plugins/Experimental/AnimNext/Source/AnimNext/Internal/Param/RigVMDispatch_GetLayerParameter.h
@@ -0,0 +1,48 @@
+// Copyright Epic Games, Inc. All Rights Reserved.
+
+#pragma once
+
+#include "RigVMCore/RigVMDispatchFactory.h"
+#include "AnimNextParameterExecuteContext.h"
+#include "RigVMDispatch_GetLayerParameter.generated.h"
+
+namespace UE::AnimNext::UncookedOnly
+{
+	struct FUtils;
+}
+
+/*
+ * Gets a parameter's current value from the parameter block value
+ */
+USTRUCT(meta = (DisplayName = "Get Parameter from Block", NodeColor = "0.8, 0, 0.2, 1"))
+struct ANIMNEXT_API FRigVMDispatch_GetLayerParameter : public FRigVMDispatchFactory
+{
+	GENERATED_BODY()
+
+	FRigVMDispatch_GetLayerParameter();
+
+private:
+	friend struct UE::AnimNext::UncookedOnly::FUtils;
+
+	virtual UScriptStruct* GetExecuteContextStruct() const override { return FAnimNextParameterExecuteContext::StaticStruct(); }
+	virtual FName GetArgumentNameForOperandIndex(int32 InOperandIndex, int32 InTotalOperands) const override;
+	virtual const TArray<FRigVMTemplateArgument>& GetArguments() const override;
+#if WITH_EDITOR
+	virtual FString GetArgumentMetaData(const FName& InArgumentName, const FName& InMetaDataKey) const override;
+	virtual FString GetArgumentDefaultValue(const FName& InArgumentName, TRigVMTypeIndex InTypeIndex) const override;
+#endif
+	virtual FRigVMTemplateTypeMap OnNewArgumentType(const FName& InArgumentName, TRigVMTypeIndex InTypeIndex) const override;
+	virtual bool IsSingleton() const override { return true; }
+
+	virtual FRigVMFunctionPtr GetDispatchFunctionImpl(const FRigVMTemplateTypeMap& InTypes) const override
+	{
+		return &FRigVMDispatch_GetLayerParameter::Execute;
+	}
+	static void Execute(FRigVMExtendedExecuteContext& InContext, FRigVMMemoryHandleArray Handles, FRigVMPredicateBranchArray RigVMBranches);
+
+	static const FName ExecuteContextName;
+	static const FName ParameterName;
+	static const FName ParameterIdName;
+	static const FName TypeHandleName;
+	static const FName ValueName;
+};
diff --git a/Engine/Plugins/Experimental/AnimNext/Source/AnimNext/Internal/Param/RigVMDispatch_SetLayerParameter.h b/Engine/Plugins/Experimental/AnimNext/Source/AnimNext/Internal/Param/RigVMDispatch_SetLayerParameter.h
new file mode 100644
--- /dev/null
+++ b/Engine/Plugins/Experimental/AnimNext/Source/AnimNext/Internal/Param/RigVMDispatch_SetLayerParameter.h
@@ -0,0 +1,49 @@
+// Copyright Epic Games, Inc. All Rights Reserved.
+
+#pragma once
+
+#include "RigVMCore/RigVMDispatchFactory.h"
+#include "AnimNextParameterExecuteContext.h"
+#include "RigVMDispatch_SetLayerParameter.generated.h"
+
+namespace UE::AnimNext::UncookedOnly
+{
+struct FUtils;
+}
+
+/*
+ * Sets a parameter's value
+ */
+USTRUCT(meta=(DisplayName = "Set Block Parameter", NodeColor = "0.8, 0, 0.2, 1"))
+struct ANIMNEXT_API FRigVMDispatch_SetLayerParameter : public FRigVMDispatchFactory
+{
+	GENERATED_BODY()
+
+	FRigVMDispatch_SetLayerParameter();
+
+private:
+	friend struct UE::AnimNext::UncookedOnly::FUtils;
+
+	virtual UScriptStruct* GetExecuteContextStruct() const { return FAnimNextParameterExecuteContext::StaticStruct(); }
+	virtual FName GetArgumentNameForOperandIndex(int32 InOperandIndex, int32 InTotalOperands) const override;
+	virtual const TArray<FRigVMTemplateArgument>& GetArguments() const override;
+#if WITH_EDITOR
+	virtual FString GetArgumentMetaData(const FName& InArgumentName, const FName& InMetaDataKey) const override;
+	virtual FString GetArgumentDefaultValue(const FName& InArgumentName, TRigVMTypeIndex InTypeIndex) const override;
+#endif
+	virtual TArray<FRigVMExecuteArgument>& GetExecuteArguments_Impl(const FRigVMDispatchContext& InContext) const override;
+	virtual FRigVMTemplateTypeMap OnNewArgumentType(const FName& InArgumentName, TRigVMTypeIndex InTypeIndex) const override;
+	virtual bool IsSingleton() const override { return true; } 
+
+	virtual FRigVMFunctionPtr GetDispatchFunctionImpl(const FRigVMTemplateTypeMap& InTypes) const override
+	{
+		return &FRigVMDispatch_SetLayerParameter::Execute;
+	}
+	static void Execute(FRigVMExtendedExecuteContext& InContext, FRigVMMemoryHandleArray Handles, FRigVMPredicateBranchArray RigVMBranches);
+
+	static const FName ExecuteContextName;
+	static const FName ParameterName;
+	static const FName ParameterIdName;
+	static const FName TypeHandleName;
+	static const FName ValueName;
+};
diff --git a/Engine/Plugins/Experimental/AnimNext/Source/AnimNext/Internal/Scheduler/AnimNextSchedule.h b/Engine/Plugins/Experimental/AnimNext/Source/AnimNext/Internal/Scheduler/AnimNextSchedule.h
new file mode 100644
--- /dev/null
+++ b/Engine/Plugins/Experimental/AnimNext/Source/AnimNext/Internal/Scheduler/AnimNextSchedule.h
@@ -0,0 +1,206 @@
+// Copyright Epic Games, Inc. All Rights Reserved.
+
+#pragma once
+
+#include "CoreMinimal.h"
+#include "AnimNextScheduleGraphTask.h"
+#include "Scheduler/AnimNextSchedulePortTask.h"
+#include "Scheduler/AnimNextScheduleExternalTask.h"
+#include "Scheduler/AnimNextScheduleParamScopeTask.h"
+#include "Tasks/Task.h"
+#include "AnimNextSchedule.generated.h"
+
+class UAnimNextGraph;
+class UAnimNextParameterBlock;
+class UAnimNextSchedule;
+class UAnimNextSchedulerWorldSubsystem;
+class UAnimNextComponent;
+struct FAnimNextSchedulerEntry;
+
+namespace UE::AnimNext
+{
+	struct FScheduler;
+	struct FSchedulerImpl;
+	struct FScheduleResult;
+	struct FScheduleContext;
+	struct FScheduleInstanceData;
+	struct FScheduleTickFunction;
+}
+
+UCLASS(EditInlineNew, Abstract)
+class UAnimNextScheduleEntry : public UObject
+{
+	GENERATED_BODY()
+
+private:
+	friend class UAnimNextSchedule;
+};
+
+UCLASS()
+class UAnimNextScheduleEntry_AnimNextGraph : public UAnimNextScheduleEntry
+{
+	GENERATED_BODY()
+
+private:
+	friend class UAnimNextSchedule;
+
+	UPROPERTY(EditAnywhere, Category = "Graph", meta = (CustomWidget = "ParamName", ParamType = "FAnimNextScope"))
+	FName Name;
+
+	UPROPERTY(EditAnywhere, Category = "Graph")
+	FName EntryPoint;
+
+	UPROPERTY(EditAnywhere, Category = "Graph")
+	TObjectPtr<UAnimNextGraph> Graph;
+
+	UPROPERTY(EditAnywhere, Category = "Graph")
+	TArray<TObjectPtr<UAnimNextParameterBlock>> ParameterBlocks;
+};
+
+UCLASS()
+class UAnimNextScheduleEntry_Port : public UAnimNextScheduleEntry
+{
+	GENERATED_BODY()
+
+private:
+	friend class UAnimNextSchedule;
+
+	UPROPERTY(EditAnywhere, Category = "Port", meta = (CustomWidget = "ParamName", ParamType = "FAnimNextPort"))
+	FName Name;
+};
+
+UCLASS()
+class UAnimNextScheduleEntry_ExternalTask : public UAnimNextScheduleEntry
+{
+	GENERATED_BODY()
+
+private:
+	friend class UAnimNextSchedule;
+
+	UPROPERTY(EditAnywhere, Category = "External Task", meta = (CustomWidget = "ParamName", ParamType = "FAnimNextTickFunction"))
+	FName Name;
+
+	UPROPERTY(EditAnywhere, Category = "External Task", meta = (CustomWidget = "ParamName", ParamType = "TObjectPtr<UObject>"))
+	FName ObjectName;
+};
+
+UCLASS()
+class UAnimNextScheduleEntry_ParamScope : public UAnimNextScheduleEntry
+{
+	GENERATED_BODY()
+
+private:
+	friend class UAnimNextSchedule;
+
+	UPROPERTY(EditAnywhere, Category = "Parameters", meta = (CustomWidget = "ParamName", ParamType = "FAnimNextScope"))
+	FName Name;
+
+	// Parameters to apply in this scope
+	UPROPERTY(EditAnywhere, Category = "Parameters")
+	TArray<TObjectPtr<UAnimNextParameterBlock>> ParameterBlocks;
+
+	// Entries that are part of this scope
+	UPROPERTY(EditAnywhere, Category = "Parameters", Instanced)
+	TArray<TObjectPtr<UAnimNextScheduleEntry>> SubEntries;
+};
+
+// TEMP: opcode in schedule
+UENUM()
+enum class EAnimNextScheduleScheduleOpcode : uint8
+{
+	None,
+	RunTask,				// Operand = Task Index
+	BeginRunExternalTask,	// Operand = Task Index
+	EndRunExternalTask,		// Operand = Task Index
+	RunPort,				// Operand = Port Index
+	RunParamScopeEntry,		// Operand = Scope Index
+	RunParamScopeExit,		// Operand = Scope Index
+	PrerequisiteTask,		// Operand = Task Index
+	PrerequisiteBeginExternalTask, // Operand = Task Index
+	PrerequisiteEndExternalTask, // Operand = Task Index
+	PrerequisiteScopeEntry,	// Operand = Scope Index
+	PrerequisiteScopeExit,	// Operand = Scope Index
+	Exit,					// Operand = 0
+};
+
+// TEMP: Bytecode instruction in schedule
+USTRUCT()
+struct FAnimNextScheduleInstruction
+{
+	GENERATED_BODY()
+
+	UPROPERTY()
+	EAnimNextScheduleScheduleOpcode Opcode;
+
+	UPROPERTY()
+	int32 Operand;
+};
+
+UCLASS()
+class ANIMNEXT_API UAnimNextSchedule : public UObject
+{
+	GENERATED_BODY()
+
+private:
+	friend struct UE::AnimNext::FScheduler;
+	friend struct UE::AnimNext::FSchedulerImpl;
+	friend struct UE::AnimNext::FScheduleContext;
+	friend struct UE::AnimNext::FScheduleInstanceData;
+	friend struct UE::AnimNext::FScheduleTickFunction;
+	friend struct FAnimNextSchedulerEntry;
+	friend class UAnimNextComponent;
+	friend class UAnimNextSchedulerWorldSubsystem;
+
+	// UObject interface
+	virtual void PostLoad() override;
+#if WITH_EDITOR
+	virtual void PostEditChangeProperty(FPropertyChangedEvent& PropertyChangedEvent) override;
+
+	// Compile the editor data into a compact runtime representation
+	void CompileSchedule();
+#endif
+
+#if WITH_EDITORONLY_DATA
+	// Editor only
+	// TODO: move this into an editor only subobject
+	// TODO: this is currently only a linear list, we want it to be a graph
+	UPROPERTY(EditAnywhere, Category = "Schedule", Instanced)
+	TArray<TObjectPtr<UAnimNextScheduleEntry>> Entries;
+#endif
+
+	// TEMP: Instructions derived from the entries above
+	UPROPERTY()
+	TArray<FAnimNextScheduleInstruction> Instructions;
+
+	// TEMP: Tasks derived from the entries above
+	UPROPERTY()
+	TArray<FAnimNextScheduleGraphTask> Tasks;
+
+	// TEMP: Ports derived from the entries above
+	UPROPERTY()
+	TArray<FAnimNextSchedulePortTask> Ports;
+
+	// TEMP: External tasks derived from the entries above
+	UPROPERTY()
+	TArray<FAnimNextScheduleExternalTask> ExternalTasks;
+
+	// TEMP: Parameter scope entry tasks derived from the entries above
+	UPROPERTY()
+	TArray<FAnimNextScheduleParamScopeEntryTask> ParamScopeEntryTasks;
+
+	// TEMP: Parameter scope exit tasks derived from the entries above
+	UPROPERTY()
+	TArray<FAnimNextScheduleParamScopeExitTask> ParamScopeExitTasks;
+
+	// TEMP: Map from port name to port index in Ports array
+	UPROPERTY()
+	TMap<FName, int32> PortNameIndexMap;
+
+	// TEMP: Count of total number of parameter scopes that this schedule needs to execute
+	UPROPERTY()
+	uint32 NumParameterScopes = 0;
+
+	// TEMP: Count of total number of tick functions that this schedule needs to execute
+	UPROPERTY()
+	uint32 NumTickFunctions = 0;
+};
\ No newline at end of file
diff --git a/Engine/Plugins/Experimental/AnimNext/Source/AnimNext/Internal/Scheduler/AnimNextScheduleExternalTask.h b/Engine/Plugins/Experimental/AnimNext/Source/AnimNext/Internal/Scheduler/AnimNextScheduleExternalTask.h
new file mode 100644
--- /dev/null
+++ b/Engine/Plugins/Experimental/AnimNext/Source/AnimNext/Internal/Scheduler/AnimNextScheduleExternalTask.h
@@ -0,0 +1,48 @@
+// Copyright Epic Games, Inc. All Rights Reserved.
+
+#pragma once
+
+#include "CoreMinimal.h"
+#include "AnimNextScheduleExternalTask.generated.h"
+
+class UAnimNextSchedule;
+class UAnimNextComponent;
+struct FAnimNextSchedulerEntry;
+
+namespace UE::AnimNext
+{
+	struct FScheduleInstanceData;
+	struct FScheduleTickFunction;
+}
+
+USTRUCT()
+struct FAnimNextScheduleExternalTask
+{
+	GENERATED_BODY()
+
+	FAnimNextScheduleExternalTask() = default;
+
+private:
+	friend class UAnimNextSchedule;
+	friend class UAnimNextComponent;
+	friend struct FAnimNextSchedulerEntry;
+	friend struct UE::AnimNext::FScheduleTickFunction;
+	friend struct UE::AnimNext::FScheduleInstanceData;
+
+	UPROPERTY()
+	uint32 TaskIndex = MAX_uint32;
+
+	UPROPERTY()
+	uint32 ParamScopeIndex = MAX_uint32;
+
+	UPROPERTY()
+	uint32 ParamParentScopeIndex = MAX_uint32;
+
+	/** The name of the external task parameter */
+	UPROPERTY()
+	FName Name;
+
+	/** The name of the external task's object parameter */
+	UPROPERTY()
+	FName ObjectName;
+};
diff --git a/Engine/Plugins/Experimental/AnimNext/Source/AnimNext/Internal/Scheduler/AnimNextScheduleGraphTask.h b/Engine/Plugins/Experimental/AnimNext/Source/AnimNext/Internal/Scheduler/AnimNextScheduleGraphTask.h
new file mode 100644
--- /dev/null
+++ b/Engine/Plugins/Experimental/AnimNext/Source/AnimNext/Internal/Scheduler/AnimNextScheduleGraphTask.h
@@ -0,0 +1,55 @@
+// Copyright Epic Games, Inc. All Rights Reserved.
+
+#pragma once
+
+#include "CoreMinimal.h"
+#include "AnimNextScheduleGraphTask.generated.h"
+
+class UAnimNextGraph;
+class UAnimNextParameterBlock;
+
+namespace UE::AnimNext
+{
+	struct FScheduleContext;
+	struct FScheduleInstanceData;
+	struct FScheduleTickFunction;
+}
+
+USTRUCT()
+struct FAnimNextScheduleGraphTask
+{
+	GENERATED_BODY()
+
+	FAnimNextScheduleGraphTask()
+	{}
+
+private:
+	friend class UAnimNextComponent;
+	friend class UAnimNextSchedule;
+	friend struct UE::AnimNext::FScheduleInstanceData;
+	friend struct UE::AnimNext::FScheduleTickFunction;
+
+	void RunGraph(const UE::AnimNext::FScheduleContext& InScheduleContext) const;
+
+private:
+	UPROPERTY()
+	uint32 TaskIndex = MAX_uint32;
+
+	UPROPERTY()
+	uint32 ParamScopeIndex = MAX_uint32;
+
+	UPROPERTY()
+	uint32 ParamParentScopeIndex = MAX_uint32;
+
+	UPROPERTY()
+	FName Name;
+
+	UPROPERTY()
+	FName EntryPoint;
+
+	UPROPERTY()
+	TObjectPtr<UAnimNextGraph> Graph;
+
+	UPROPERTY()
+	TArray<TObjectPtr<UAnimNextParameterBlock>> ParameterBlocks;
+};
diff --git a/Engine/Plugins/Experimental/AnimNext/Source/AnimNext/Internal/Scheduler/AnimNextScheduleParamScopeTask.h b/Engine/Plugins/Experimental/AnimNext/Source/AnimNext/Internal/Scheduler/AnimNextScheduleParamScopeTask.h
new file mode 100644
--- /dev/null
+++ b/Engine/Plugins/Experimental/AnimNext/Source/AnimNext/Internal/Scheduler/AnimNextScheduleParamScopeTask.h
@@ -0,0 +1,79 @@
+// Copyright Epic Games, Inc. All Rights Reserved.
+
+#pragma once
+
+#include "CoreMinimal.h"
+#include "AnimNextScheduleParamScopeTask.generated.h"
+
+class UAnimNextSchedule;
+class UAnimNextParameterBlock;
+class UAnimNextSchedulerWorldSubsystem;
+
+namespace UE::AnimNext
+{
+	struct FScheduleContext;
+	struct FScheduleInstanceData;
+	struct FScheduleTickFunction;
+}
+
+USTRUCT()
+struct FAnimNextScheduleParamScopeEntryTask
+{
+	GENERATED_BODY()
+
+	FAnimNextScheduleParamScopeEntryTask() = default;
+
+private:
+	void RunParamScopeEntry(const UE::AnimNext::FScheduleContext& InScheduleContext) const;
+
+private:
+	friend class UAnimNextSchedule;
+	friend class UAnimNextSchedulerWorldSubsystem;
+	friend struct UE::AnimNext::FScheduleInstanceData;
+	friend struct UE::AnimNext::FScheduleTickFunction;
+
+	UPROPERTY()
+	uint32 TaskIndex = MAX_uint32;
+
+	UPROPERTY()
+	uint32 ParamScopeIndex = MAX_uint32;
+
+	UPROPERTY()
+	uint32 ParamParentScopeIndex = MAX_uint32;
+
+	UPROPERTY()
+	uint32 TickFunctionIndex = MAX_uint32;
+
+	// The name of the scope
+	UPROPERTY()
+	FName Name;
+
+	/** Parameters to apply in this scope */
+	UPROPERTY()
+	TArray<TObjectPtr<UAnimNextParameterBlock>> ParameterBlocks;
+};
+
+USTRUCT()
+struct FAnimNextScheduleParamScopeExitTask
+{
+	GENERATED_BODY()
+
+	FAnimNextScheduleParamScopeExitTask() = default;
+
+private:
+	void RunParamScopeExit(const UE::AnimNext::FScheduleContext& InScheduleContext) const;
+
+private:
+	friend class UAnimNextSchedule;
+	friend struct UE::AnimNext::FScheduleTickFunction;
+
+	UPROPERTY()
+	uint32 TaskIndex = MAX_uint32;
+
+	UPROPERTY()
+	uint32 ParamScopeIndex = MAX_uint32;
+
+	// The name of the scope
+	UPROPERTY()
+	FName Name;
+};
diff --git a/Engine/Plugins/Experimental/AnimNext/Source/AnimNext/Internal/Scheduler/AnimNextSchedulePortTask.h b/Engine/Plugins/Experimental/AnimNext/Source/AnimNext/Internal/Scheduler/AnimNextSchedulePortTask.h
new file mode 100644
--- /dev/null
+++ b/Engine/Plugins/Experimental/AnimNext/Source/AnimNext/Internal/Scheduler/AnimNextSchedulePortTask.h
@@ -0,0 +1,45 @@
+// Copyright Epic Games, Inc. All Rights Reserved.
+
+#pragma once
+
+#include "CoreMinimal.h"
+#include "AnimNextSchedulePortTask.generated.h"
+
+class UAnimNextSchedule;
+
+namespace UE::AnimNext
+{
+	struct FScheduleContext;
+	struct FScheduleTask;
+	struct FScheduleTickFunction;
+}
+
+USTRUCT()
+struct FAnimNextSchedulePortTask
+{
+	GENERATED_BODY()
+
+	friend struct UE::AnimNext::FScheduleTask;
+	friend struct UE::AnimNext::FScheduleTickFunction;
+	friend class UAnimNextSchedule;
+
+	FAnimNextSchedulePortTask() = default;
+
+private:
+	void RunPort(const UE::AnimNext::FScheduleContext& InScheduleContext) const;
+
+private:
+	UPROPERTY()
+	uint32 TaskIndex = MAX_uint32;
+
+	UPROPERTY()
+	uint32 ParamScopeIndex = MAX_uint32;
+
+	// The name of the port
+	UPROPERTY()
+	FName Name;
+
+	// The name of the ports input parameter
+	UPROPERTY()
+	FName InputParameterName;
+};
diff --git a/Engine/Plugins/Experimental/AnimNext/Source/AnimNext/Private/AnimNextConfig.cpp b/Engine/Plugins/Experimental/AnimNext/Source/AnimNext/Private/AnimNextConfig.cpp
new file mode 100644
--- /dev/null
+++ b/Engine/Plugins/Experimental/AnimNext/Source/AnimNext/Private/AnimNextConfig.cpp
@@ -0,0 +1,20 @@
+﻿// Copyright Epic Games, Inc. All Rights Reserved.
+
+#include "AnimNextConfig.h"
+#include "Param/ParamId.h"
+
+#if WITH_EDITOR
+
+void UAnimNextConfig::PostEditChangeProperty(FPropertyChangedEvent& PropertyChangedEvent)
+{
+	using namespace UE::AnimNext;
+	
+	if(PropertyChangedEvent.Property->GetFName() == GET_MEMBER_NAME_CHECKED(UAnimNextConfig, ExposedClasses))
+	{
+		FParamId::RefreshConfigAdapters();
+	}
+
+	SaveConfig();
+}
+
+#endif
\ No newline at end of file
diff --git a/Engine/Plugins/Experimental/AnimNext/Source/AnimNext/Private/Component/AnimNextComponent.cpp b/Engine/Plugins/Experimental/AnimNext/Source/AnimNext/Private/Component/AnimNextComponent.cpp
new file mode 100644
--- /dev/null
+++ b/Engine/Plugins/Experimental/AnimNext/Source/AnimNext/Private/Component/AnimNextComponent.cpp
@@ -0,0 +1,128 @@
+// Copyright Epic Games, Inc. All Rights Reserved.
+
+#include "AnimNextComponent.h"
+#include "Scheduler/Scheduler.h"
+#include "Scheduler/ScheduleContext.h"
+#include "Param/AnimNextParameterSourceRef.h"
+#include "Param/Params.h"
+
+void UAnimNextComponent::OnRegister()
+{
+	using namespace UE::AnimNext;
+
+	Super::OnRegister();
+
+	if (Schedule)
+	{
+		check(!SchedulerHandle.IsValid());
+		SchedulerHandle = FScheduler::AcquireHandle(this, Schedule, Parameters);
+	}
+}
+
+void UAnimNextComponent::OnUnregister()
+{
+	using namespace UE::AnimNext;
+
+	Super::OnUnregister();
+
+	FScheduler::ReleaseHandle(this, SchedulerHandle);
+	SchedulerHandle.Invalidate();
+}
+
+void UAnimNextComponent::UpdateLayer(UE::AnimNext::FParamStackLayerHandle& InHandle) const
+{
+	if (UAnimNextComponent* This = InHandle.As<UAnimNextComponent>())
+	{
+		This->Update();
+	}
+}
+
+UE::AnimNext::FParamStackLayerHandle UAnimNextComponent::CacheLayer() const
+{
+	return UE::AnimNext::FParamStack::MakeReferenceLayer(const_cast<UAnimNextComponent*>(this));
+}
+
+TStructOnScope<FActorComponentInstanceData> UAnimNextComponent::GetComponentInstanceData() const
+{
+	TStructOnScope<FActorComponentInstanceData> InstanceData = MakeStructOnScope<FActorComponentInstanceData, FAnimNextComponentInstanceData>(this);
+	FAnimNextComponentInstanceData* AnimNextComponentInstanceData = InstanceData.Cast<FAnimNextComponentInstanceData>();
+	AnimNextComponentInstanceData->Schedule = Schedule;
+	AnimNextComponentInstanceData->Parameters = Parameters;
+
+	return InstanceData;
+}
+
+void UAnimNextComponent::SetParameterInScope(FName Scope, FName Name, int32 Value)
+{
+	checkNoEntry();
+}
+
+DEFINE_FUNCTION(UAnimNextComponent::execSetParameterInScope)
+{
+	using namespace UE::AnimNext;
+
+	// Read wildcard Value input.
+	Stack.MostRecentPropertyAddress = nullptr;
+	Stack.MostRecentPropertyContainer = nullptr;
+
+	PARAM_PASSED_BY_VAL(Scope, FNameProperty, FName);
+	PARAM_PASSED_BY_VAL(Name, FNameProperty, FName);
+
+	Stack.StepCompiledIn<FProperty>(nullptr);
+	const FProperty* ValueProp = CastField<FProperty>(Stack.MostRecentProperty);
+	const void* ValuePtr = Stack.MostRecentPropertyAddress;
+
+	P_FINISH;
+
+	if (!ValueProp || !ValuePtr)
+	{
+		FBlueprintExceptionInfo ExceptionInfo(
+			EBlueprintExceptionType::AbortExecution,
+			NSLOCTEXT("AnimNextComponent", "AnimNextComponent_SetParameterInScopeError", "Failed to resolve the Value for Set Parameter In Scope")
+		);
+
+		FBlueprintCoreDelegates::ThrowScriptException(P_THIS, Stack, ExceptionInfo);
+	}
+	else if (Scope == NAME_None || Name == NAME_None)
+	{
+		FBlueprintExceptionInfo ExceptionInfo(
+			EBlueprintExceptionType::NonFatalError,
+			NSLOCTEXT("AnimNextComponent", "AnimNextComponent_SetParameterInScopeWarning", "Invalid scope or parameter name supplied to Set Parameter In Scope")
+		);
+
+		FBlueprintCoreDelegates::ThrowScriptException(P_THIS, Stack, ExceptionInfo);
+	}
+	else
+	{
+		P_NATIVE_BEGIN;
+
+		FAnimNextParameterSourceRef ParamSource;
+		ParamSource.Type = EAnimNextParameterSourceRefType::Inline;
+		ParamSource.InlineParameters.AddProperty(Name, ValueProp);
+		ParamSource.InlineParameters.SetValue(Name, ParamSource.InlineParameters.GetPropertyBagStruct()->FindPropertyByName(Name), ValuePtr);
+
+		FScheduler::QueueTask(P_THIS, P_THIS->SchedulerHandle, Scope, [Scope, ParamSource = MoveTemp(ParamSource)](const FScheduleContext& InContext)
+		{
+			FAnimNextParameterCollection& Collection = InContext.GetInstanceData().UserScopes.FindOrAdd(Scope);
+
+			// TODO: pre/post scope distinction
+			Collection.Parameters.Insert(ParamSource, 0);
+		});
+		
+		P_NATIVE_END;
+	}
+}
+
+void UAnimNextComponent::Enable(bool bEnabled)
+{
+	UE::AnimNext::FScheduler::EnableHandle(this, SchedulerHandle, bEnabled);
+}
+
+void FAnimNextComponentInstanceData::ApplyToComponent(UActorComponent* Component, const ECacheApplyPhase CacheApplyPhase) 
+{
+	Super::ApplyToComponent(Component, CacheApplyPhase);
+
+	UAnimNextComponent* AnimNextComponent = CastChecked<UAnimNextComponent>(Component);
+	AnimNextComponent->Schedule = Schedule;
+	AnimNextComponent->Parameters = Parameters;
+}
diff --git a/Engine/Plugins/Experimental/AnimNext/Source/AnimNext/Private/Component/AnimNextComponent.h b/Engine/Plugins/Experimental/AnimNext/Source/AnimNext/Private/Component/AnimNextComponent.h
new file mode 100644
--- /dev/null
+++ b/Engine/Plugins/Experimental/AnimNext/Source/AnimNext/Private/Component/AnimNextComponent.h
@@ -0,0 +1,91 @@
+// Copyright Epic Games, Inc. All Rights Reserved.
+
+#pragma once
+
+#include "Components/ActorComponent.h"
+#include "Param/AnimNextParameterCollection.h"
+#include "Param/IAnimNextParameterSourceInterface.h"
+#include "Scheduler/ScheduleHandle.h"
+#include "AnimNextComponent.generated.h"
+
+class UAnimNextSchedule;
+struct FAnimNextComponentInstanceData;
+
+UCLASS(MinimalAPI, Blueprintable, meta = (BlueprintSpawnableComponent))
+class UAnimNextComponent : public UActorComponent, public IAnimNextParameterSourceInterface
+{
+	GENERATED_BODY()
+
+	friend struct FAnimNextComponentInstanceData;
+
+	// UActorComponent interface
+	virtual void OnRegister() override;
+	virtual void OnUnregister() override;
+	virtual TStructOnScope<FActorComponentInstanceData> GetComponentInstanceData() const override;
+
+	// IAnimNextParameterSourceInterface interface
+	virtual void UpdateLayer(UE::AnimNext::FParamStackLayerHandle& InHandle) const override;
+	virtual UE::AnimNext::FParamStackLayerHandle CacheLayer() const override;
+
+	// Because UHT says: "A Private function cannot be a BlueprintImplementableEvent!"
+protected:
+	// Override point to allow self to be updated
+	UFUNCTION(Category = "AnimNext", BlueprintImplementableEvent, meta=(BlueprintThreadSafe, ForceAsFunction))
+	void Update();
+
+public:
+	// Sets a parameter in the scope. Scopes correspond to an existing scope in a schedule.
+	UFUNCTION(BlueprintCallable, Category = "AnimNext", CustomThunk, meta = (CustomStructureParam = Value, UnsafeDuringActorConstruction))
+	void SetParameterInScope(FName Scope, FName Name, int32 Value);
+
+	// Enable or disable this component's update
+	UFUNCTION(BlueprintCallable, Category = "AnimNext")
+	void Enable(bool bEnabled);
+	
+private:
+	DECLARE_FUNCTION(execSetParameterInScope);
+
+private:
+	// The execution schedule that this component will run
+	UPROPERTY(EditAnywhere, Category="Schedule")
+	UAnimNextSchedule* Schedule = nullptr;
+
+	// Parameters to apply at each scope
+	UPROPERTY(EditAnywhere, Category = "Parameters", meta = (CustomWidget = "ParamName", AllowedParamType = "FAnimNextScope"))
+	TMap<FName, FAnimNextParameterCollection> Parameters;
+
+	// Handle to the registered results/schedule
+	UE::AnimNext::FScheduleHandle SchedulerHandle;
+
+	// Cached layer handle for 'self'
+	mutable UE::AnimNext::FParamStackLayerHandle LayerHandle;
+};
+
+/** Used to store component data during RerunConstructionScripts */
+USTRUCT()
+struct FAnimNextComponentInstanceData : public FActorComponentInstanceData
+{
+	GENERATED_BODY()
+
+	FAnimNextComponentInstanceData() = default;
+	FAnimNextComponentInstanceData(const UAnimNextComponent* InSourceComponent)
+		: FActorComponentInstanceData(InSourceComponent)
+	{}
+
+private:
+	friend class UAnimNextComponent;
+
+	virtual bool ContainsData() const override
+	{
+		return true;
+	}
+
+	virtual void ApplyToComponent(UActorComponent* Component, const ECacheApplyPhase CacheApplyPhase) override;
+
+	UPROPERTY(EditAnywhere, Category = "Schedule")
+	UAnimNextSchedule* Schedule = nullptr;
+
+	UPROPERTY(EditAnywhere, Category = "Parameters")
+	TMap<FName, FAnimNextParameterCollection> Parameters;
+};
+
diff --git a/Engine/Plugins/Experimental/AnimNext/Source/AnimNext/Private/Graph/AnimGraph/AnimGraphParamStackScope.cpp b/Engine/Plugins/Experimental/AnimNext/Source/AnimNext/Private/Graph/AnimGraph/AnimGraphParamStackScope.cpp
new file mode 100644
--- /dev/null
+++ b/Engine/Plugins/Experimental/AnimNext/Source/AnimNext/Private/Graph/AnimGraph/AnimGraphParamStackScope.cpp
@@ -0,0 +1,48 @@
+// Copyright Epic Games, Inc. All Rights Reserved.
+
+#include "AnimGraphParamStackScope.h"
+#include "Param/ParamStack.h"
+#include "Animation/AnimNodeBase.h"
+
+namespace UE::AnimNext
+{
+
+FAnimGraphParamStackScope::FAnimGraphParamStackScope(const FAnimationBaseContext& InContext)
+{
+	TWeakPtr<FParamStack> ExternalParamStack = FParamStack::GetForCurrentThread();
+	if(!ExternalParamStack.IsValid())
+	{
+		// If we havent already attached this thread, we should be using the pending object's stack
+		ComponentObject = InContext.GetAnimInstanceObject()->GetOuter();
+		if (ComponentObject == nullptr || !FParamStack::AttachToCurrentThreadForPendingObject(ComponentObject.Get()))
+		{
+			// No pending stack, we should use an owned one
+			OwnedParamStack = MakeShared<FParamStack>();
+			FParamStack::AttachToCurrentThread(OwnedParamStack);
+		}
+		else
+		{
+			bAttachedPending = true;
+		}
+	}
+}
+
+FAnimGraphParamStackScope::~FAnimGraphParamStackScope()
+{
+	// Check ownership invariance if we have an owned stack
+	if(OwnedParamStack.IsValid())
+	{
+		TWeakPtr<FParamStack> DetachedStack = FParamStack::DetachFromCurrentThread();
+		check(DetachedStack.Pin() == OwnedParamStack);
+		DetachedStack.Reset();
+		check(OwnedParamStack.IsUnique());
+		OwnedParamStack.Reset();
+	}
+	else if(bAttachedPending)
+	{
+		check(ComponentObject.IsValid());
+		FParamStack::DetachFromCurrentThreadForPendingObject(ComponentObject.Get());
+	}
+}
+
+}
\ No newline at end of file
diff --git a/Engine/Plugins/Experimental/AnimNext/Source/AnimNext/Private/Graph/AnimGraph/AnimNodeExposedValueHandler_AnimNextParameters.cpp b/Engine/Plugins/Experimental/AnimNext/Source/AnimNext/Private/Graph/AnimGraph/AnimNodeExposedValueHandler_AnimNextParameters.cpp
--- a/Engine/Plugins/Experimental/AnimNext/Source/AnimNext/Private/Graph/AnimGraph/AnimNodeExposedValueHandler_AnimNextParameters.cpp
+++ b/Engine/Plugins/Experimental/AnimNext/Source/AnimNext/Private/Graph/AnimGraph/AnimNodeExposedValueHandler_AnimNextParameters.cpp
@@ -5,6 +5,7 @@
 #include "Animation/AnimNodeBase.h"
 #include "Animation/AnimSubsystem_PropertyAccess.h"
 #include "Animation/AnimClassInterface.h"
+#include "AnimGraphParamStackScope.h"
 
 void FAnimNodeExposedValueHandler_AnimNextParameters::Initialize(const UClass* InClass)
 {
@@ -432,29 +433,33 @@ void FAnimNodeExposedValueHandler_AnimNextParameters::Execute(const FAnimationBa
 
 	Super::Execute(InContext);
 
-	FParamStack& ParamStack = FParamStack::Get();
-	UObject* Object = InContext.GetAnimInstanceObject();
-
-	for (const FAnimNodeExposedValueHandler_AnimNextParameters_Entry& Entry : Entries)
 	{
-		EPropertyAccessCopyType CopyType = Entry.AccessType;
-		FParamTypeHandle FoundTypeHandle;
-		TConstArrayView<uint8> Value;
-		FParamStack::EGetParamResult Result = ParamStack.GetParamData(Entry.ParamId, Entry.PropertyParamTypeHandle, Value, FoundTypeHandle, FParamCompatibility::IncompatibleWithDataLoss());
-		if (Result == FParamStack::EGetParamResult::CompatibleType)
-		{
-			// Compatibility indicates we need to potentially modify the access type
-			CopyType = Private::GetCopyTypeForCompatibility(CopyType, Entry.PropertyParamTypeHandle, FoundTypeHandle);
-		}
+		FAnimGraphParamStackScope Scope(InContext);
 
-		if(Result == FParamStack::EGetParamResult::CompatibleType || Result == FParamStack::EGetParamResult::Succeeded)
+		FParamStack& ParamStack = FParamStack::Get();
+		UObject* Object = InContext.GetAnimInstanceObject();
+
+		for (const FAnimNodeExposedValueHandler_AnimNextParameters_Entry& Entry : Entries)
 		{
-			const uint8* SrcAddr = Value.GetData();
+			EPropertyAccessCopyType CopyType = Entry.AccessType;
+			FParamTypeHandle FoundTypeHandle;
+			TConstArrayView<uint8> Value;
+			FParamResult Result = ParamStack.GetParamData(Entry.ParamId, Entry.PropertyParamTypeHandle, Value, FoundTypeHandle, FParamCompatibility::IncompatibleWithDataLoss());
+			if (Result.IsOfCompatibleType())
+			{
+				// Compatibility indicates we need to potentially modify the access type
+				CopyType = Private::GetCopyTypeForCompatibility(CopyType, Entry.PropertyParamTypeHandle, FoundTypeHandle);
+			}
 
-			PropertyAccess::GetAccessAddress(Object, *PropertyAccessLibrary, Entry.AccessIndex, [CopyType, SrcAddr](const FProperty* InProperty, void* InAddress)
+			if(Result.IsSuccessful())
 			{
-				Private::PerformCopy(CopyType, InProperty, InAddress, SrcAddr);
-			});
+				const uint8* SrcAddr = Value.GetData();
+
+				PropertyAccess::GetAccessAddress(Object, *PropertyAccessLibrary, Entry.AccessIndex, [CopyType, SrcAddr](const FProperty* InProperty, void* InAddress)
+				{
+					Private::PerformCopy(CopyType, InProperty, InAddress, SrcAddr);
+				});
+			}
 		}
 	}
 }
\ No newline at end of file
diff --git a/Engine/Plugins/Experimental/AnimNext/Source/AnimNext/Private/Graph/AnimGraph/AnimNode_AnimNextParameters.cpp b/Engine/Plugins/Experimental/AnimNext/Source/AnimNext/Private/Graph/AnimGraph/AnimNode_AnimNextParameters.cpp
--- a/Engine/Plugins/Experimental/AnimNext/Source/AnimNext/Private/Graph/AnimGraph/AnimNode_AnimNextParameters.cpp
+++ b/Engine/Plugins/Experimental/AnimNext/Source/AnimNext/Private/Graph/AnimGraph/AnimNode_AnimNextParameters.cpp
@@ -2,54 +2,42 @@
 
 #include "Graph/AnimGraph/AnimNode_AnimNextParameters.h"
 #include "Param/ParamStack.h"
-#include "Param/AnimNextParameterBlock.h"
-#include "Context.h"
+#include "Param/IAnimNextParameterSourceInterface.h"
+#include "AnimGraphParamStackScope.h"
 
 #include UE_INLINE_GENERATED_CPP_BY_NAME(AnimNode_AnimNextParameters)
 
 FAnimNode_AnimNextParameters::FAnimNode_AnimNextParameters(const FAnimNode_AnimNextParameters& InOther)
 	: Source(InOther.Source)
-#if WITH_EDITORONLY_DATA
-	, ParameterBlock(InOther.ParameterBlock)
-#endif
-	, PreviousParameterBlock(nullptr)
+	, Parameters(InOther.Parameters)
+	, PreviousParameters(nullptr)
 	, ParamLayerHandle()
-	, PropertyBag()
 {
 }
 
 FAnimNode_AnimNextParameters& FAnimNode_AnimNextParameters::operator=(const FAnimNode_AnimNextParameters& InOther)
 {
 	Source = InOther.Source;
-#if WITH_EDITORONLY_DATA
-	ParameterBlock = InOther.ParameterBlock;
-#endif
-	PreviousParameterBlock = nullptr;
+	Parameters = InOther.Parameters;
+	PreviousParameters = nullptr;
 	ParamLayerHandle.Invalidate();
-	PropertyBag.Reset();
 	return *this;
 }
 
-FAnimNode_AnimNextParameters::FAnimNode_AnimNextParameters(FAnimNode_AnimNextParameters&& InOther)
+FAnimNode_AnimNextParameters::FAnimNode_AnimNextParameters(FAnimNode_AnimNextParameters&& InOther) noexcept
 {
 	Source = InOther.Source;
-#if WITH_EDITORONLY_DATA
-	ParameterBlock = InOther.ParameterBlock;
-#endif
-	PreviousParameterBlock = nullptr;
+	Parameters = InOther.Parameters;
+	PreviousParameters = nullptr;
 	ParamLayerHandle.Invalidate();
-	PropertyBag.Reset();
 }
 
-FAnimNode_AnimNextParameters& FAnimNode_AnimNextParameters::operator=(FAnimNode_AnimNextParameters&& InOther)
+FAnimNode_AnimNextParameters& FAnimNode_AnimNextParameters::operator=(FAnimNode_AnimNextParameters&& InOther) noexcept
 {
 	Source = InOther.Source;
-#if WITH_EDITORONLY_DATA
-	ParameterBlock = InOther.ParameterBlock;
-#endif
-	PreviousParameterBlock = nullptr;
+	Parameters = InOther.Parameters;
+	PreviousParameters = nullptr;
 	ParamLayerHandle.Invalidate();
-	PropertyBag.Reset();
 	return *this;
 }
 
@@ -64,36 +52,38 @@ void FAnimNode_AnimNextParameters::Update_AnyThread(const FAnimationUpdateContex
 
 	GetEvaluateGraphExposedInputs().Execute(Context);
 
-	UAnimNextParameterBlock* CurrentParameterBlock = GetParameterBlock();
+	IAnimNextParameterSourceInterface* CurrentParameters = Parameters ? Parameters.GetInterface() : nullptr;
 
 	// Reconstruct param block's cached layer if required
-	if (CurrentParameterBlock != PreviousParameterBlock || !ParamLayerHandle.IsValid())
+	if (CurrentParameters != PreviousParameters || !ParamLayerHandle.IsValid())
 	{
 		ParamLayerHandle.Invalidate();
 
-		if (CurrentParameterBlock)
+		if (CurrentParameters)
 		{
-			PropertyBag = CurrentParameterBlock->PropertyBag;
-			ParamLayerHandle = FParamStack::MakeMutableLayer(PropertyBag);
+			ParamLayerHandle = CurrentParameters->CacheLayer();
 		}
 
-		PreviousParameterBlock = CurrentParameterBlock;
+		PreviousParameters = CurrentParameters;
 	}
 
-	FParamStack& ParamStack = FParamStack::Get();
-	FParamStack::FPushedLayerHandle PushedLayerHandle;
-	if (CurrentParameterBlock && ParamLayerHandle.IsValid())
 	{
-		PushedLayerHandle = ParamStack.PushLayer(ParamLayerHandle);
-		UE::AnimNext::FContext AnimNextContext;
-		CurrentParameterBlock->Run(AnimNextContext);
-	}
+		FAnimGraphParamStackScope Scope(Context);
 
-	Source.Update(Context);
+		FParamStack& ParamStack = FParamStack::Get();
+		FParamStack::FPushedLayerHandle PushedLayerHandle;
+		if (CurrentParameters && ParamLayerHandle.IsValid())
+		{
+			CurrentParameters->UpdateLayer(ParamLayerHandle);
+			PushedLayerHandle = ParamStack.PushLayer(ParamLayerHandle);
+		}
 
-	if (PushedLayerHandle.IsValid())
-	{
-		ParamStack.PopLayer(PushedLayerHandle);
+		Source.Update(Context);
+
+		if (PushedLayerHandle.IsValid())
+		{
+			ParamStack.PopLayer(PushedLayerHandle);
+		}
 	}
 }
 
@@ -113,8 +103,3 @@ void FAnimNode_AnimNextParameters::GatherDebugData(FNodeDebugData& DebugData)
 
 	Source.GatherDebugData(DebugData);
 }
-
-UAnimNextParameterBlock* FAnimNode_AnimNextParameters::GetParameterBlock() const
-{
-	return GET_ANIM_NODE_DATA(TObjectPtr<UAnimNextParameterBlock>, ParameterBlock);
-}
\ No newline at end of file
diff --git a/Engine/Plugins/Experimental/AnimNext/Source/AnimNext/Private/Module.cpp b/Engine/Plugins/Experimental/AnimNext/Source/AnimNext/Private/Module.cpp
--- a/Engine/Plugins/Experimental/AnimNext/Source/AnimNext/Private/Module.cpp
+++ b/Engine/Plugins/Experimental/AnimNext/Source/AnimNext/Private/Module.cpp
@@ -1,6 +1,7 @@
 // Copyright Epic Games, Inc. All Rights Reserved.
 
 #include "CoreMinimal.h"
+#include "AnimNextConfig.h"
 #include "Modules/ModuleInterface.h"
 #include "Modules/ModuleManager.h"
 #include "Misc/CoreDelegates.h"
@@ -9,6 +10,9 @@
 #include "DecoratorBase/NodeTemplateRegistry.h"
 #include "RigVMCore/RigVMRegistry.h"
 #include "Animation/AnimSequence.h"
+#include "Scheduler/Scheduler.h"
+#include "Param/Params.h"
+#include "Components/SkeletalMeshComponent.h"
 
 namespace UE::AnimNext
 {
@@ -18,6 +22,10 @@ class FModule : public IModuleInterface
 public:
 	virtual void StartupModule() override
 	{
+		GetMutableDefault<UAnimNextConfig>()->LoadConfig();
+
+		FParamId::Init();
+
 		static TPair<UClass*, FRigVMRegistry::ERegisterObjectOperation> const AllowedObjectTypes[] =
 		{
 			{ UAnimSequence::StaticClass(), FRigVMRegistry::ERegisterObjectOperation::Class },
@@ -29,16 +37,22 @@ public:
 		FDataRegistry::Init();
 		FDecoratorRegistry::Init();
 		FNodeTemplateRegistry::Init();
+		FScheduler::Init();
 	}
 
 	virtual void ShutdownModule() override
 	{
+		FScheduler::Destroy();
 		FNodeTemplateRegistry::Destroy();
 		FDecoratorRegistry::Destroy();
 		FDataRegistry::Destroy();
+
+		FParamId::Destroy();
 	}
 };
 
 }
 
 IMPLEMENT_MODULE(UE::AnimNext::FModule, AnimNext)
+
+#undef LOCTEXT_NAMESPACE
\ No newline at end of file
diff --git a/Engine/Plugins/Experimental/AnimNext/Source/AnimNext/Private/Param/AnimNextParameterBlock.cpp b/Engine/Plugins/Experimental/AnimNext/Source/AnimNext/Private/Param/AnimNextParameterBlock.cpp
--- a/Engine/Plugins/Experimental/AnimNext/Source/AnimNext/Private/Param/AnimNextParameterBlock.cpp
+++ b/Engine/Plugins/Experimental/AnimNext/Source/AnimNext/Private/Param/AnimNextParameterBlock.cpp
@@ -5,15 +5,10 @@
 #include "UObject/ObjectSaveContext.h"
 #include "UObject/Package.h"
 #include "Graph/RigUnit_AnimNextBeginExecution.h"
-#include "Graph/AnimNextExecuteContext.h"
+#include "Param/AnimNextParameterExecuteContext.h"
 
 #include UE_INLINE_GENERATED_CPP_BY_NAME(AnimNextParameterBlock)
 
-TArray<FRigVMExternalVariable> UAnimNextParameterBlock::GetRigVMExternalVariables()
-{
-	return TArray<FRigVMExternalVariable>(); 
-}
-
 namespace UE::AnimNext::Private
 {
 
@@ -36,15 +31,39 @@ static TArray<UClass*> GetClassObjectsInPackage(UPackage* InPackage)
 
 }
 
-void UAnimNextParameterBlock::Run(const UE::AnimNext::FContext& Context) const
+void UAnimNextParameterBlock::UpdateLayer(UE::AnimNext::FParamStackLayerHandle& InHandle) const
+{
+	if (VM)
+	{
+		FRigVMExtendedExecuteContext Context = GetExtendedExecuteContext();
+		FAnimNextParameterExecuteContext& AnimNextParameterContext = Context.GetPublicDataSafe<FAnimNextParameterExecuteContext>();
+		AnimNextParameterContext.SetParamContextData(InHandle);
+		TArray<FRigVMMemoryStorageStruct*> LocalMemory = VM->GetLocalMemoryArray(Context);
+		VM->Execute(Context, LocalMemory, FRigUnit_AnimNextBeginExecution::EventName);
+	}
+}
+
+UE::AnimNext::FParamStackLayerHandle UAnimNextParameterBlock::CacheLayer() const
+{
+	return UE::AnimNext::FParamStack::MakeValueLayer(PropertyBag);
+}
+
+bool UAnimNextParameterBlock::ShouldCacheLayer(const UE::AnimNext::FParamStackLayerHandle& InHandle) const
 {
-	if (RigVM)
+	if(!InHandle.IsValid())
+	{
+		return true;
+	}
+
+#if WITH_EDITOR	// Layout should only be changing in editor
+	const FInstancedPropertyBag* HandlePropertyBag = InHandle.As<FInstancedPropertyBag>();
+	if(HandlePropertyBag == nullptr || HandlePropertyBag->GetPropertyBagStruct() != PropertyBag.GetPropertyBagStruct())
 	{
-		FRigVMExtendedExecuteContext RigVMExtendedExecuteContext;
-		FAnimNextExecuteContext& AnimNextContext = RigVMExtendedExecuteContext.GetPublicDataSafe<FAnimNextExecuteContext>();
-		AnimNextContext.SetContextData(Context);
-		RigVM->Execute(RigVMExtendedExecuteContext, TArray<FRigVMMemoryStorageStruct*>(), FRigUnit_AnimNextBeginExecution::EventName);
+		return true;
 	}
+#endif
+
+	return false;
 }
 
 void UAnimNextParameterBlock::PostRename(UObject* OldOuter, const FName OldName)
@@ -83,6 +102,22 @@ void UAnimNextParameterBlock::GetPreloadDependencies(TArray<UObject*>& OutDeps)
 	}
 }
 
+void UAnimNextParameterBlock::PostLoad()
+{
+	VM = RigVM;
+
+	Super::PostLoad();
+
+	// In packaged builds, initialize the VM
+	// In editor, the VM will be recompiled and initialized at UAnimNextParameterBlock_EditorData::HandlePackageDone::RecompileVM
+#if !WITH_EDITOR
+	if(VM != nullptr)
+	{
+		InitializeVM(FRigUnit_AnimNextBeginExecution::EventName);
+	}
+#endif
+}
+
 void UAnimNextParameterBlock::GetAssetRegistryTags(TArray<FAssetRegistryTag>& OutTags) const
 {
 	Super::GetAssetRegistryTags(OutTags);
diff --git a/Engine/Plugins/Experimental/AnimNext/Source/AnimNext/Private/Param/AnimNextParameterSourceRef.cpp b/Engine/Plugins/Experimental/AnimNext/Source/AnimNext/Private/Param/AnimNextParameterSourceRef.cpp
new file mode 100644
--- /dev/null
+++ b/Engine/Plugins/Experimental/AnimNext/Source/AnimNext/Private/Param/AnimNextParameterSourceRef.cpp
@@ -0,0 +1,42 @@
+// Copyright Epic Games, Inc. All Rights Reserved.
+
+#include "Param/AnimNextParameterSourceRef.h"
+#include "Component/AnimNextComponent.h"
+#include "GameFramework/Actor.h"
+
+const IAnimNextParameterSourceInterface* FAnimNextParameterSourceRef::Get(const UObject* InContextObject) const
+{
+	switch (Type)
+	{
+	case EAnimNextParameterSourceRefType::Self:
+		return Cast<UAnimNextComponent>(InContextObject);
+	case EAnimNextParameterSourceRefType::Asset:
+		return Asset.GetInterface();
+	case EAnimNextParameterSourceRefType::Component:
+	case EAnimNextParameterSourceRefType::Actor:
+		if (const UAnimNextComponent* AnimNextComponent = Cast<UAnimNextComponent>(InContextObject))
+		{
+			if (AActor* Actor = AnimNextComponent->GetOwner())
+			{
+				if (Type == EAnimNextParameterSourceRefType::Actor)
+				{
+					return Cast<IAnimNextParameterSourceInterface>(Actor);
+				}
+				else
+				{
+					return Cast<IAnimNextParameterSourceInterface>(Component.GetComponent(Actor));
+				}
+			}
+		}
+		break;
+	case EAnimNextParameterSourceRefType::Inline:
+		return this;
+	}
+
+	return nullptr;
+}
+
+UE::AnimNext::FParamStackLayerHandle FAnimNextParameterSourceRef::CacheLayer() const
+{
+	return UE::AnimNext::FParamStack::MakeValueLayer(InlineParameters);
+}
\ No newline at end of file
diff --git a/Engine/Plugins/Experimental/AnimNext/Source/AnimNext/Private/Param/ParamAdapter.cpp b/Engine/Plugins/Experimental/AnimNext/Source/AnimNext/Private/Param/ParamAdapter.cpp
new file mode 100644
--- /dev/null
+++ b/Engine/Plugins/Experimental/AnimNext/Source/AnimNext/Private/Param/ParamAdapter.cpp
@@ -0,0 +1,73 @@
+﻿// Copyright Epic Games, Inc. All Rights Reserved.
+
+#include "ParamAdapter.h"
+
+#include "Param/ParamUtils.h"
+#include "Scheduler/AnimNextSchedulerEntry.h"
+#include "Scheduler/ScheduleContext.h"
+
+namespace UE::AnimNext
+{
+
+FParamResult FParamAdapter::GetParamData(FParamTypeHandle InTypeHandle, TConstArrayView<uint8>& OutParamData, FParamTypeHandle& OutParamTypeHandle, FParamCompatibility InRequiredCompatibility) const
+{
+	OutParamTypeHandle = Definition.TypeHandle;
+
+	const FParamCompatibility Compatibility = FParamUtils::GetCompatibility(InTypeHandle, OutParamTypeHandle);
+	if (Compatibility < InRequiredCompatibility)
+	{
+		return EParamResult::TypeError;
+	}
+
+	UObject* ContextObject = FScheduleContext::Get().Entry->Object.Get();
+	if(const uint8* Data = Function(ContextObject, Definition.Id))
+	{
+		OutParamData = TConstArrayView<uint8>(Data, Definition.TypeHandle.GetSize());
+	}
+	else
+	{
+		return EParamResult::NotInScope;
+	}
+
+	if (Compatibility == InRequiredCompatibility)
+	{
+		return EParamResult::Success | EParamResult::TypeCompatible;
+	}
+
+	return EParamResult::Success;
+}
+
+FParamResult FParamAdapter::GetMutableParamData(FParamTypeHandle InTypeHandle, TArrayView<uint8>& OutParamData, FParamTypeHandle& OutParamTypeHandle, FParamCompatibility InRequiredCompatibility) const
+{
+	if(!EnumHasAnyFlags(Definition.Flags, EParamDefinitionFlags::Mutable))
+	{
+		return EParamResult::MutabilityError;
+	}
+
+	OutParamTypeHandle = Definition.TypeHandle;
+
+	const FParamCompatibility Compatibility = FParamUtils::GetCompatibility(InTypeHandle, OutParamTypeHandle);
+	if (Compatibility < InRequiredCompatibility)
+	{
+		return EParamResult::TypeError;
+	}
+
+	UObject* ContextObject = FScheduleContext::Get().Entry->Object.Get();
+	if(uint8* Data = Function(ContextObject, Definition.Id))
+	{
+		OutParamData = TArrayView<uint8>(Data, Definition.TypeHandle.GetSize());
+	}
+	else
+	{
+		return EParamResult::NotInScope;
+	}
+
+	if (Compatibility == InRequiredCompatibility)
+	{
+		return EParamResult::Success | EParamResult::TypeCompatible;
+	}
+
+	return EParamResult::Success;
+}
+
+}
diff --git a/Engine/Plugins/Experimental/AnimNext/Source/AnimNext/Private/Param/ParamAdapter.h b/Engine/Plugins/Experimental/AnimNext/Source/AnimNext/Private/Param/ParamAdapter.h
new file mode 100644
--- /dev/null
+++ b/Engine/Plugins/Experimental/AnimNext/Source/AnimNext/Private/Param/ParamAdapter.h
@@ -0,0 +1,48 @@
+// Copyright Epic Games, Inc. All Rights Reserved.
+
+#pragma once
+
+#include "CoreMinimal.h"
+#include "Param/ParamCompatibility.h"
+#include "UObject/Interface.h"
+#include "Param/ParamDefinition.h"
+#include "Param/ParamResult.h"
+
+struct FAnimNextParamType;
+
+namespace UE::AnimNext
+{
+	struct FParamContext;
+	struct FParamStackLayerHandle;
+	struct FParamDefinition;
+}
+
+namespace UE::AnimNext
+{
+
+using FObjectAdapterFunction = TUniqueFunction<uint8*(UObject*, FParamId)>;
+
+/**
+ * Adapter for a specific parameter. This is designed to wrap native objects and their parameters for the parameter
+ * system. Adapters are deferred to by the param stack if they exist for a non-overriden parameter ID 
+ */
+struct FParamAdapter
+{
+	FParamAdapter(FParamDefinition&& InDefinition, FObjectAdapterFunction&& InFunction)
+		: Definition(MoveTemp(InDefinition))
+		, Function(MoveTemp(InFunction))
+	{
+	}
+
+	// Get the parameter's data for this adapter
+	FParamResult GetParamData(FParamTypeHandle InTypeHandle, TConstArrayView<uint8>& OutParamData, FParamTypeHandle& OutParamTypeHandle, FParamCompatibility InRequiredCompatibility) const;
+
+	// Get the parameter's mutable data for this adapter
+	FParamResult GetMutableParamData(FParamTypeHandle InTypeHandle, TArrayView<uint8>& OutParamData, FParamTypeHandle& OutParamTypeHandle, FParamCompatibility InRequiredCompatibility) const;
+
+	FParamDefinition Definition;
+
+	FObjectAdapterFunction Function;
+};
+
+}
diff --git a/Engine/Plugins/Experimental/AnimNext/Source/AnimNext/Private/Param/ParamDefinition.cpp b/Engine/Plugins/Experimental/AnimNext/Source/AnimNext/Private/Param/ParamDefinition.cpp
new file mode 100644
--- /dev/null
+++ b/Engine/Plugins/Experimental/AnimNext/Source/AnimNext/Private/Param/ParamDefinition.cpp
@@ -0,0 +1,58 @@
+// Copyright Epic Games, Inc. All Rights Reserved.
+
+#include "Param/ParamDefinition.h"
+#include "Param/ParamTypeHandle.h"
+#include "Param/ParamId.h"
+
+namespace UE::AnimNext
+{
+
+FParamDefinition::FParamDefinition(uint32 InIdIndex, FName InName, const UObject* InObject)
+	: Id(FParamId(InIdIndex))
+	, Name(InName)
+	, TypeHandle(FParamTypeHandle::FromObject(InObject))
+	, Type(TypeHandle.GetType())
+#if WITH_EDITORONLY_DATA
+	, Tooltip(InObject->GetClass()->GetToolTipText())
+#endif
+	, Property(nullptr)
+	, Function(nullptr)
+{
+}
+
+FParamDefinition::FParamDefinition(uint32 InIdIndex, FName InName, const FProperty* InProperty)
+	: Id(FParamId(InIdIndex))
+	, Name(InName)
+	, TypeHandle(FParamTypeHandle::FromProperty(InProperty))
+	, Type(TypeHandle.GetType())
+#if WITH_EDITORONLY_DATA
+	, Tooltip(InProperty->GetToolTipText())
+#endif
+	, Property(InProperty)
+	, Function(nullptr)
+{
+	// Param names should not contain periods. Periods are only a display concern. Use underscores.
+	check(!Name.ToString().Contains(TEXT(".")));
+}
+
+FParamDefinition::FParamDefinition(uint32 InIdIndex, FName InName, const UFunction* InFunction)
+	: Id(FParamId(InIdIndex))
+	, Name(InName)
+	, Property(nullptr)
+	, Function(InFunction)
+{
+	// Param names should not contain periods. Periods are only a display concern. Use underscores.
+	check(!Name.ToString().Contains(TEXT(".")));
+
+	const FProperty* ReturnProperty = InFunction->GetReturnProperty();
+	check(ReturnProperty);
+
+	TypeHandle = FParamTypeHandle::FromProperty(ReturnProperty);
+	Type = TypeHandle.GetType();
+	
+#if WITH_EDITORONLY_DATA
+	Tooltip = InFunction->GetToolTipText();
+#endif
+}
+
+}
\ No newline at end of file
diff --git a/Engine/Plugins/Experimental/AnimNext/Source/AnimNext/Private/Param/ParamEntry.cpp b/Engine/Plugins/Experimental/AnimNext/Source/AnimNext/Private/Param/ParamEntry.cpp
new file mode 100644
--- /dev/null
+++ b/Engine/Plugins/Experimental/AnimNext/Source/AnimNext/Private/Param/ParamEntry.cpp
@@ -0,0 +1,117 @@
+// Copyright Epic Games, Inc. All Rights Reserved.
+
+#include "Param/ParamEntry.h"
+#include "Param/ParamHelpers.h"
+
+namespace UE::AnimNext::Private
+{
+
+FParamEntry::FParamEntry(const FParamTypeHandle& InTypeHandle, TArrayView<uint8> InData, bool bInIsReference, bool bInIsMutable)
+	: Data(nullptr)
+	, TypeHandle(InTypeHandle)
+	, Size(InData.Num())
+	, Flags(EParamFlags::None)
+{
+	check(TypeHandle.IsValid());
+	check(InData.Num() > 0 && InData.Num() < 0xffff);
+
+	// If we can store our data inside of a ptr, we do
+	if (!bInIsReference && InData.Num() <= sizeof(void*))
+	{
+		FParamHelpers::Copy(InTypeHandle, InData, TArrayView<uint8>(reinterpret_cast<uint8*>(&Data), sizeof(void*)));
+		Flags |= EParamFlags::Embedded;
+	}
+	else
+	{
+		Data = InData.GetData();
+	}
+
+	if (bInIsReference)
+	{
+		Flags |= EParamFlags::Reference;
+	}
+
+	if (bInIsMutable)
+	{
+		Flags |= EParamFlags::Mutable;
+	}
+}
+
+FParamEntry::~FParamEntry()
+{
+	if (Size > 0)
+	{
+		if (IsEmbedded())
+		{
+			FParamHelpers::Destroy(TypeHandle, TArrayView<uint8>(reinterpret_cast<uint8*>(&Data), sizeof(void*)));
+		}
+	}
+}
+
+FParamEntry::FParamEntry(const FParamEntry& InOtherParam)
+	: TypeHandle(InOtherParam.TypeHandle)
+	, Size(InOtherParam.Size)
+	, Flags(InOtherParam.Flags)
+{
+	if (IsEmbedded())
+	{
+		FParamHelpers::Copy(InOtherParam.TypeHandle, TConstArrayView<uint8>(reinterpret_cast<const uint8*>(&InOtherParam.Data), sizeof(void*)), TArrayView<uint8>(reinterpret_cast<uint8*>(&Data), sizeof(void*)));
+	}
+	else
+	{
+		Data = InOtherParam.Data;
+	}
+}
+
+FParamEntry& FParamEntry::operator=(const FParamEntry& InOtherParam)
+{
+	TypeHandle = InOtherParam.TypeHandle;
+	Size = InOtherParam.Size;
+	Flags = InOtherParam.Flags;
+
+	if (IsEmbedded())
+	{
+		FParamHelpers::Copy(InOtherParam.TypeHandle, TConstArrayView<uint8>(reinterpret_cast<const uint8*>(&InOtherParam.Data), sizeof(void*)), TArrayView<uint8>(reinterpret_cast<uint8*>(&Data), sizeof(void*)));
+	}
+	else
+	{
+		Data = InOtherParam.Data;
+	}
+
+	return *this;
+}
+
+FParamEntry::FParamEntry(FParamEntry&& InOtherParam) noexcept
+	: TypeHandle(InOtherParam.TypeHandle)
+	, Size(InOtherParam.Size)
+	, Flags(InOtherParam.Flags)
+{
+	if (IsEmbedded())
+	{
+		FParamHelpers::Copy(InOtherParam.TypeHandle, TConstArrayView<uint8>(reinterpret_cast<const uint8*>(&InOtherParam.Data), sizeof(void*)), TArrayView<uint8>(reinterpret_cast<uint8*>(&Data), sizeof(void*)));
+	}
+	else
+	{
+		Data = InOtherParam.Data;
+	}
+}
+
+FParamEntry& FParamEntry::operator=(FParamEntry&& InOtherParam) noexcept
+{
+	TypeHandle = InOtherParam.TypeHandle;
+	Size = InOtherParam.Size;
+	Flags = InOtherParam.Flags;
+
+	if (IsEmbedded())
+	{
+		FParamHelpers::Copy(InOtherParam.TypeHandle, TConstArrayView<uint8>(reinterpret_cast<const uint8*>(&InOtherParam.Data), sizeof(void*)), TArrayView<uint8>(reinterpret_cast<uint8*>(&Data), sizeof(void*)));
+	}
+	else
+	{
+		Data = InOtherParam.Data;
+	}
+
+	return *this;
+}
+
+}
\ No newline at end of file
diff --git a/Engine/Plugins/Experimental/AnimNext/Source/AnimNext/Private/Param/ParamId.cpp b/Engine/Plugins/Experimental/AnimNext/Source/AnimNext/Private/Param/ParamId.cpp
--- a/Engine/Plugins/Experimental/AnimNext/Source/AnimNext/Private/Param/ParamId.cpp
+++ b/Engine/Plugins/Experimental/AnimNext/Source/AnimNext/Private/Param/ParamId.cpp
@@ -1,22 +1,53 @@
 // Copyright Epic Games, Inc. All Rights Reserved.
 
 #include "Param/ParamId.h"
+
+#include "AnimNextConfig.h"
 #include "Misc/ScopeRWLock.h"
 #include "HAL/ThreadSingleton.h"
+#include "Param/AnimNextObjectAdapterConfig.h"
+#include "Param/ParamAdapter.h"
+#include "Param/ParamDefinition.h"
+#include "GameFramework/Actor.h"
+#include "Components/ActorComponent.h"
+#include "Param/Params.h"
+#include "Param/ParamStack.h"
+#include "RigVMCore/RigVMRegistry.h"
+#include "Engine/World.h"
+
+#define LOCTEXT_NAMESPACE "ParamId"
 
 namespace UE::AnimNext
 {
 
 static FRWLock GParamIdLock;
 
+struct FParamData
+{
+	FParamData(FName InName)
+		: Name(InName)
+	{}
+
+	FName Name;
+	int32 AdapterIndex = INDEX_NONE;
+	uint32 BufferOffset = MAX_uint32;
+};
+
 struct FParamIdGlobalData
 {
-	TArray<FName> ParamIdToName;
+	TArray<FParamData> ParamData;
+	TPagedArray<FParamAdapter, sizeof(FParamAdapter) * 256> ParamAdapters;
 	TMap<FName, uint32> NameToParamId;
+	uint32 SerialNumber = 0;
+	uint32 ResultBufferSize = 0;
 };
 
 static FParamIdGlobalData GParamIdGlobalData;
 
+// Buffers to hold values for function-wrapped parameter adapters 
+thread_local TArray<uint8> ParamAdapterResultBuffer;
+thread_local uint32 ParamAdapterBufferSerialNumber = 0;
+
 #if WITH_DEV_AUTOMATION_TESTS
 static FParamIdGlobalData GSandboxedParamIdGlobalData;
 static std::atomic<bool> bGParamIdSandboxed = false;
@@ -36,31 +67,397 @@ static FParamIdGlobalData& GetParamIdData()
 	}
 }
 
+void FParamId::Init()
+{
+	RegisterBuiltInAdapters();
+	RefreshConfigAdapters();
+}
+
+void FParamId::Destroy()
+{
+	FRWScopeLock Lock(GParamIdLock, SLT_Write);
+	FParamIdGlobalData& ParamIdData = GetParamIdData();
+	ParamIdData.ParamData.Empty();
+	ParamIdData.ParamAdapters.Empty();
+	ParamIdData.NameToParamId.Empty();
+}
+
+void FParamId::RegisterBuiltInAdapters()
+{
+	FRWScopeLock Lock(GParamIdLock, SLT_Write);
+	FParamIdGlobalData& ParamIdGlobalData = GetParamIdData();
+
+	// Register physics tick function
+	{
+		FObjectAdapterFunction PhysicsTickFunctionFunction = [](UObject* InContextObject, FParamId InId) -> uint8*
+		{
+			if (UWorld* World = InContextObject->GetWorld())
+			{
+				return reinterpret_cast<uint8*>(&World->EndPhysicsTickFunction);
+			}
+			return nullptr;
+		};
+
+		FName PhysicsTickName("UE_Physics_Tick");
+		const uint32 PhysicsTickFunctionIndex = MakeParamId_NoLock(PhysicsTickName);
+		ParamIdGlobalData.ParamData[PhysicsTickFunctionIndex].AdapterIndex = ParamIdGlobalData.ParamAdapters.Num();
+		FText TooltipText = LOCTEXT("PhysicsTickTooltip", "The tick function of the physics scene");
+		FParamAdapter& NewAdapter = ParamIdGlobalData.ParamAdapters.Emplace_GetRef(FParamDefinition(PhysicsTickFunctionIndex, PhysicsTickName, FAnimNextParamType::GetType<FTickFunction>(), TooltipText), MoveTemp(PhysicsTickFunctionFunction));
+		FParams::RegisterBuiltInParameter(NewAdapter.Definition);
+	}
+}
+
+void FParamId::RefreshConfigAdapters()
+{
+	FRWScopeLock Lock(GParamIdLock, SLT_Write);
+	FParamIdGlobalData& ParamIdGlobalData = GetParamIdData();
+
+	ParamIdGlobalData.ResultBufferSize = 0;
+	ParamIdGlobalData.ParamAdapters.Empty();
+
+	// Reset existing IDs adapter data
+	for(FParamData& ParamData : ParamIdGlobalData.ParamData)
+	{
+		ParamData.AdapterIndex = INDEX_NONE;
+		ParamData.BufferOffset = MAX_uint32;
+
+		// Unregister this param as built-in
+		FParams::UnregisterBuiltInParameter(ParamData.Name);
+	}
+
+	for(const FAnimNextObjectAdapterConfig& AdapterConfig : GetDefault<UAnimNextConfig>()->ExposedClasses)
+	{
+		if(AdapterConfig.Class != nullptr && AdapterConfig.RootParameter != NAME_None)
+		{
+			// Ensure we can use the type as a pin in RigVM graphs 
+			FRigVMRegistry::Get().RegisterObjectTypes( {{ AdapterConfig.Class, FRigVMRegistry::ERegisterObjectOperation::Class } });
+
+			const FString RootParam(AdapterConfig.RootParameter.ToString());
+			const FName RootParamName(*RootParam);
+			const uint32 RootParamIndex = MakeParamId_NoLock(RootParamName);
+			const FString NameParam = RootParam + TEXT("_Name");
+			const FName NameParamName(*NameParam);
+			const uint32 NameParamIndex = MakeParamId_NoLock(NameParamName);
+
+			// TODO: add registration mechanism to allow for more than actor/component types
+			FObjectAdapterFunction RootFunction;
+			if(AdapterConfig.Class->IsChildOf(UActorComponent::StaticClass()))
+			{
+				RootFunction = [Class = AdapterConfig.Class, NameParamIndex](UObject* InContextObject, FParamId InId) -> uint8*
+				{
+					if (const UActorComponent* ContextComponent = Cast<UActorComponent>(InContextObject))
+					{
+						if (AActor* Actor = ContextComponent->GetOwner())
+						{
+							// Components can be redirected via their name
+							UActorComponent* ActorComponent = nullptr;
+							if(const FName* ComponentName = FParamStack::Get().GetParamPtr<FName>(FParamId(NameParamIndex)))
+							{
+								const FObjectPropertyBase* ComponentProperty = FindFProperty<FObjectPropertyBase>(Actor->GetClass(), *ComponentName);
+								if(ComponentProperty != nullptr)
+								{
+									ActorComponent = Cast<UActorComponent>(ComponentProperty->GetObjectPropertyValue_InContainer(Actor));
+								}
+							}
+							else
+							{
+								// No name, just find the first component
+								ActorComponent = Actor->FindComponentByClass(Class.Get());
+							}
+
+							UActorComponent** ReturnValue = GetAdapterReturnValue<UActorComponent*>(InId);
+							*ReturnValue = ActorComponent;
+							return reinterpret_cast<uint8*>(*ReturnValue);
+						}
+					}
+
+					return nullptr;
+				};
+			}
+			else if(AdapterConfig.Class->IsChildOf(AActor::StaticClass()))
+			{
+				RootFunction = [Class = AdapterConfig.Class.Get()](UObject* InContextObject, FParamId InId) -> uint8*
+				{
+					if (const UActorComponent* ContextComponent = Cast<UActorComponent>(InContextObject))
+					{
+						// Actors cannot be redirected by name
+						if(AActor* Actor = ContextComponent->GetOwner())
+						{
+							if(Actor->GetClass()->IsChildOf(Class))
+							{
+								AActor** ReturnValue = GetAdapterReturnValue<AActor*>(InId);
+								*ReturnValue = Actor;
+								return reinterpret_cast<uint8*>(*ReturnValue);
+							}
+						}
+					}
+
+					return nullptr;
+				};
+			}
+
+			// Add root param
+			ParamIdGlobalData.ParamData[RootParamIndex].AdapterIndex = ParamIdGlobalData.ParamAdapters.Num(); 
+			FParamAdapter& NewRootAdapter = ParamIdGlobalData.ParamAdapters.Emplace_GetRef(FParamDefinition(RootParamIndex, RootParamName, AdapterConfig.Class->GetDefaultObject()), MoveTemp(RootFunction));
+			FParams::RegisterBuiltInParameter(NewRootAdapter.Definition);
+
+			// Calc buffer size for the ptr's result
+			ParamIdGlobalData.ParamData[RootParamIndex].BufferOffset = Align(ParamIdGlobalData.ResultBufferSize, alignof(UObject*));
+			ParamIdGlobalData.ResultBufferSize += sizeof(UObject*);
+
+			bool bNameParamAdded = false;
+
+			// Add object's BP-exposed properties
+			if(AdapterConfig.FilterType != EAnimNextObjectAdapterFilterType::AllowOnlyFunctions)
+			{
+				auto IsNameAllowed = [&AdapterConfig](FName InName)
+				{
+					return   AdapterConfig.FilterType == EAnimNextObjectAdapterFilterType::AllowOnlyProperties ||
+							 AdapterConfig.FilterType == EAnimNextObjectAdapterFilterType::AllowAllPropertiesAndFunctions ||
+							(AdapterConfig.FilterType == EAnimNextObjectAdapterFilterType::AllowList && AdapterConfig.FilteredFields.Contains(InName)) ||
+							(AdapterConfig.FilterType == EAnimNextObjectAdapterFilterType::DenyList && !AdapterConfig.FilteredFields.Contains(InName));
+				};
+
+				for(TFieldIterator<FProperty> It(AdapterConfig.Class, EFieldIterationFlags::IncludeSuper | EFieldIterationFlags::IncludeInterfaces); It; ++It)
+				{
+					const FProperty* Property = *It;
+					if(Property->HasAnyPropertyFlags(CPF_Edit | CPF_BlueprintVisible))
+					{
+						if(FParamTypeHandle::FromProperty(Property).IsValid() && IsNameAllowed(Property->GetFName()))
+						{
+							TStringBuilder<128> ParameterNameBuilder;
+							ParameterNameBuilder.Appendf(TEXT("%s_%s"), *RootParam, *Property->GetName());
+							if(NameParam.Compare(ParameterNameBuilder.ToString(), ESearchCase::IgnoreCase) == 0)
+							{
+								bNameParamAdded = true;
+							}
+							
+							FObjectAdapterFunction AdapterFunction = [RootParamIndex, Property, Class = AdapterConfig.Class.Get()](UObject* InContextObject, FParamId InId) -> uint8*
+							{
+								// Grab object container
+								if(const UObject* Object = FParamStack::Get().GetParamPtr<UObject>(FParamId(RootParamIndex)))
+								{
+									check(Class->IsChildOf(Property->GetOwnerClass()));
+									return const_cast<uint8*>(Property->ContainerPtrToValuePtr<uint8>(Object));
+								}
+								return nullptr;
+							};
+
+							FName ParameterName(ParameterNameBuilder.ToString());
+							const uint32 NewIdIndex = MakeParamId_NoLock(ParameterName);
+							ParamIdGlobalData.ParamData[NewIdIndex].AdapterIndex = ParamIdGlobalData.ParamAdapters.Num(); 
+							FParamAdapter& NewAdapter = ParamIdGlobalData.ParamAdapters.Emplace_GetRef(FParamDefinition(NewIdIndex, ParameterName, Property), MoveTemp(AdapterFunction));
+							FParams::RegisterBuiltInParameter(NewAdapter.Definition);
+						}
+					}
+				}
+			}
+
+			// Add object's BP-exposed accessor functions
+			if(AdapterConfig.FilterType != EAnimNextObjectAdapterFilterType::AllowOnlyProperties)
+			{
+				auto IsNameAllowed = [&AdapterConfig](FName InName)
+				{
+					return   AdapterConfig.FilterType == EAnimNextObjectAdapterFilterType::AllowOnlyFunctions ||
+							 AdapterConfig.FilterType == EAnimNextObjectAdapterFilterType::AllowAllPropertiesAndFunctions ||
+							(AdapterConfig.FilterType == EAnimNextObjectAdapterFilterType::AllowList && AdapterConfig.FilteredFields.Contains(InName)) ||
+							(AdapterConfig.FilterType == EAnimNextObjectAdapterFilterType::DenyList && !AdapterConfig.FilteredFields.Contains(InName));
+				};
+				
+				for(TFieldIterator<UFunction> It(AdapterConfig.Class, EFieldIterationFlags::IncludeSuper | EFieldIterationFlags::IncludeInterfaces); It; ++It)
+				{
+					UFunction* Function = *It;
+
+					// We add only accessor functions that have valid return types
+					const FProperty* ReturnProperty = Function->GetReturnProperty();
+					if(ReturnProperty != nullptr && Function->NumParms == 1 && Function->HasAnyFunctionFlags(FUNC_BlueprintCallable) && IsNameAllowed(Function->GetFName()))
+					{
+						if(FParamTypeHandle::FromProperty(ReturnProperty).IsValid())
+						{
+							FString FunctionName = Function->GetName();
+							FunctionName.RemoveFromStart(TEXT("Get"));
+							
+							TStringBuilder<128> ParameterNameBuilder;
+							ParameterNameBuilder.Appendf(TEXT("%s_%s"), *RootParam, *FunctionName);
+							if(NameParam.Compare(ParameterNameBuilder.ToString(), ESearchCase::IgnoreCase) == 0)
+							{
+								bNameParamAdded = true;
+							}
+
+							FObjectAdapterFunction AdapterFunction = [RootParamIndex, WeakFunction = TWeakObjectPtr<UFunction>(Function), Class = AdapterConfig.Class.Get()](UObject* InContextObject, FParamId InId)
+							{
+								// Grab object to call the function with
+								if(UFunction* Function = WeakFunction.Get())
+								{
+									if(UObject* Object = FParamStack::Get().GetMutableParamPtr<UObject>(FParamId(RootParamIndex)))
+									{
+										check(Object->GetClass()->IsChildOf(Class));
+										check(Function->GetOuterUClass()->IsChildOf(Class));
+										const FProperty* ReturnProperty = Function->GetReturnProperty();
+										check(ReturnProperty);
+
+										uint8* ReturnValuePtr = GetScratchAreaForParamIdAdapter(InId);
+										ReturnProperty->InitializeValue(ReturnValuePtr);
+										FFrame Stack(Object, Function, nullptr, nullptr, Function->ChildProperties);
+										Function->Invoke(Object, Stack, ReturnValuePtr);
+									}
+								}
+								return nullptr;
+							};
+
+							FName ParameterName(ParameterNameBuilder.ToString());
+							const uint32 NewIdIndex = MakeParamId_NoLock(ParameterName);
+							ParamIdGlobalData.ParamData[NewIdIndex].AdapterIndex = ParamIdGlobalData.ParamAdapters.Num(); 
+							FParamAdapter& NewAdapter = ParamIdGlobalData.ParamAdapters.Emplace_GetRef(FParamDefinition(NewIdIndex, ParameterName, Function), MoveTemp(AdapterFunction));
+							FParams::RegisterBuiltInParameter(NewAdapter.Definition);
+
+							// Calc buffer size for the function's result
+							ParamIdGlobalData.ParamData[NewIdIndex].BufferOffset = Align(ParamIdGlobalData.ResultBufferSize, ReturnProperty->GetMinAlignment());
+							ParamIdGlobalData.ResultBufferSize += ReturnProperty->GetSize();
+						}
+					}
+				}
+			}
+
+			if(!bNameParamAdded)
+			{
+				// Add name param if we haven't already above
+				FText TooltipText;
+#if WITH_EDITOR
+				TooltipText = FText::Format(LOCTEXT("NameTooltipFormat", "The name of the {0}"), AdapterConfig.Class->GetDisplayNameText());
+#endif
+				FParams::RegisterBuiltInParameter(FParamDefinition(NameParamIndex, NameParamName, FAnimNextParamType::GetType<FName>(), TooltipText));
+			}
+
+			if(AdapterConfig.bRegisterTickFunction)
+			{
+				TStringBuilder<128> ParameterNameBuilder;
+				ParameterNameBuilder.Appendf(TEXT("%s_Tick"), *RootParam);
+
+				FObjectAdapterFunction TickFunctionFunction;
+				FAnimNextParamType ParamType = FAnimNextParamType::GetType<FTickFunction>();
+				if(AdapterConfig.Class->IsChildOf(UActorComponent::StaticClass()))
+				{
+					ParamType = FAnimNextParamType::GetType<FActorComponentTickFunction>();
+					TickFunctionFunction = [RootParamIndex](UObject* InContextObject, FParamId InId) -> uint8*
+					{
+						if (const UActorComponent* Component = FParamStack::Get().GetParamPtr<UActorComponent>(FParamId(RootParamIndex)))
+						{
+							return reinterpret_cast<uint8*>(const_cast<FActorComponentTickFunction*>(&Component->PrimaryComponentTick));
+						}
+						return nullptr;
+					};
+				}
+				else if(AdapterConfig.Class->IsChildOf(AActor::StaticClass()))
+				{
+					ParamType = FAnimNextParamType::GetType<FActorTickFunction>();
+					TickFunctionFunction = [RootParamIndex](UObject* InContextObject, FParamId InId) -> uint8*
+					{
+						if (const AActor* Actor = FParamStack::Get().GetParamPtr<AActor>(FParamId(RootParamIndex)))
+						{
+							return reinterpret_cast<uint8*>(const_cast<FActorTickFunction*>(&Actor->PrimaryActorTick));
+						}
+						return nullptr;
+					};
+				}
+
+				FText TooltipText;
+#if WITH_EDITOR
+				TooltipText = FText::Format(LOCTEXT("TickFunctionTooltipFormat", "The tick function of a {0}"), AdapterConfig.Class->GetDisplayNameText());
+#endif
+				FName ParameterName(ParameterNameBuilder.ToString());
+				const uint32 NewIdIndex = MakeParamId_NoLock(ParameterName);
+				ParamIdGlobalData.ParamData[NewIdIndex].AdapterIndex = ParamIdGlobalData.ParamAdapters.Num();
+				FParamAdapter& NewAdapter = ParamIdGlobalData.ParamAdapters.Emplace_GetRef(FParamDefinition(NewIdIndex, ParameterName, ParamType, TooltipText, EParamDefinitionFlags::Mutable), MoveTemp(TickFunctionFunction));
+				FParams::RegisterBuiltInParameter(NewAdapter.Definition);
+			}
+		}
+	}
+
+	// Refresh RigVM registry as we have updated allowed types
+	FRigVMRegistry::Get().RefreshEngineTypes();
+
+	// Increment result buffer serial so we resize on access
+	ParamIdGlobalData.SerialNumber++;
+	if(ParamIdGlobalData.SerialNumber == 0)
+	{
+		ParamIdGlobalData.SerialNumber++;
+	}
+}
+
 FParamId::FParamId(FName InName)
 {
 	FRWScopeLock Lock(GParamIdLock, SLT_Write);
-	if (const uint32* FoundIndex = GetParamIdData().NameToParamId.Find(InName))
+	ParameterIndex = MakeParamId_NoLock(InName);
+}
+
+uint32 FParamId::MakeParamId_NoLock(FName InName)
+{
+	FParamIdGlobalData& ParamIdData = GetParamIdData();
+
+	// Param names should not be 'none'
+	check(InName != NAME_None);
+	// Param names should not contain periods. Use underscores.
+	// Periods are only a display concern and we want parameters to be expressible as members of objects & structures. 
+	check(!InName.ToString().Contains(TEXT(".")));
+
+	uint32 Index = InvalidIndex;
+	if (const uint32* FoundIndex = ParamIdData.NameToParamId.Find(InName))
 	{
-		ParameterIndex = *FoundIndex;
+		Index = *FoundIndex;
 	}
 	else
 	{
-		ParameterIndex = GetParamIdData().ParamIdToName.Num();
-		GetParamIdData().NameToParamId.Add(InName, ParameterIndex);
-		GetParamIdData().ParamIdToName.Add(InName);
+		Index = ParamIdData.ParamData.Num();
+		ParamIdData.NameToParamId.Add(InName, Index);
+		ParamIdData.ParamData.Add(InName);
 	}
+
+	return Index;
 }
 
 FName FParamId::ToName() const
 {
 	FRWScopeLock Lock(GParamIdLock, SLT_ReadOnly);
-	return GetParamIdData().ParamIdToName[ParameterIndex];
+	return GetParamIdData().ParamData[ParameterIndex].Name;
+}
+
+const FParamAdapter* FParamId::GetAdapter(FParamId InId)
+{
+	FRWScopeLock Lock(GParamIdLock, SLT_ReadOnly);
+	const FParamIdGlobalData& ParamIdData = GetParamIdData();
+	const int32 DefinitionIndex = ParamIdData.ParamData[InId.ParameterIndex].AdapterIndex;
+	return DefinitionIndex != INDEX_NONE ? &ParamIdData.ParamAdapters[DefinitionIndex] : nullptr;
+}
+
+uint8* FParamId::GetScratchAreaForParamIdAdapter(FParamId InId)
+{
+	uint32 BufferOffset;
+	uint32 SerialNumber;
+	uint32 ResultBufferSize;
+
+	{
+		FRWScopeLock Lock(GParamIdLock, SLT_ReadOnly);
+		const FParamIdGlobalData& ParamIdData = GetParamIdData();
+		BufferOffset = ParamIdData.ParamData[InId.ParameterIndex].BufferOffset;
+		SerialNumber = ParamIdData.SerialNumber;
+		ResultBufferSize = ParamIdData.ResultBufferSize;
+	}
+
+	if(ParamAdapterBufferSerialNumber != SerialNumber)
+	{
+		// realloc thread-local storage for results
+		ParamAdapterResultBuffer.SetNumZeroed(ResultBufferSize);
+		ParamAdapterBufferSerialNumber = SerialNumber;
+	}
+
+	return &ParamAdapterResultBuffer[BufferOffset];
 }
 
 FParamId FParamId::GetMaxParamId()
 {
 	FRWScopeLock Lock(GParamIdLock, SLT_ReadOnly);
-	return FParamId((uint32)GetParamIdData().ParamIdToName.Num());
+	return FParamId((uint32)GetParamIdData().ParamData.Num());
 }
 
 #if WITH_DEV_AUTOMATION_TESTS
@@ -70,7 +467,7 @@ void FParamId::BeginTestSandbox()
 	check(bGParamIdSandboxed.load() == false);
 	bGParamIdSandboxed.exchange(true);
 	GSandboxedParamIdGlobalData.NameToParamId.Empty();
-	GSandboxedParamIdGlobalData.ParamIdToName.Empty();
+	GSandboxedParamIdGlobalData.ParamData.Empty();
 }
 
 void FParamId::EndTestSandbox()
@@ -79,8 +476,10 @@ void FParamId::EndTestSandbox()
 	check(bGParamIdSandboxed.load() == true);
 	bGParamIdSandboxed.exchange(false);
 	GSandboxedParamIdGlobalData.NameToParamId.Empty();
-	GSandboxedParamIdGlobalData.ParamIdToName.Empty();
+	GSandboxedParamIdGlobalData.ParamData.Empty();
 }
 #endif
 
-}
\ No newline at end of file
+}
+
+#undef LOCTEXT_NAMESPACE
\ No newline at end of file
