AnimNext: Scheduler, Component and Parameter Improvements

diff --git a/Engine/Plugins/Experimental/AnimNext/Source/AnimNext/Private/Param/ParamStack.cpp b/Engine/Plugins/Experimental/AnimNext/Source/AnimNext/Private/Param/ParamStack.cpp
--- a/Engine/Plugins/Experimental/AnimNext/Source/AnimNext/Private/Param/ParamStack.cpp
+++ b/Engine/Plugins/Experimental/AnimNext/Source/AnimNext/Private/Param/ParamStack.cpp
@@ -4,181 +4,169 @@
 #include "Param/ParamHelpers.h"
 #include "PropertyBag.h"
 #include "EngineLogs.h"
-#include "Param/ParamUtils.h"
+#include "Param/ParamAdapter.h"
+#include "UObject/UObjectGlobals.h"
+#include "UObject/Package.h"
+#include "Param/ParamStackLayer.h"
+#include "UObject/ObjectKey.h"
 
 #define LOCTEXT_NAMESPACE "AnimNextParamStack"
 
 namespace UE::AnimNext
 {
 
-FParamStack::FParam::FParam(const FParamTypeHandle& InTypeHandle, TArrayView<uint8> InData, bool bInIsReference, bool bInIsMutable)
-	: Data(nullptr)
-	, TypeHandle(InTypeHandle)
-	, Size(InData.Num())
-	, Flags(EParamFlags::None)
+// Stack layer that can own its own data as a UObject or reference an externally owned object
+struct FUObjectLayer : FParamStackLayer
 {
-	check(TypeHandle.IsValid());
-	check(InData.Num() > 0 && InData.Num() < 0xffff);
+	FUObjectLayer() = delete;
 
-	// If we can store our data inside of a ptr, we do
-	if (!bInIsReference && InData.Num() <= sizeof(void*))
+	explicit FUObjectLayer(UObject* InObject, bool bInMutable)
 	{
-		FParamHelpers::Copy(InTypeHandle, InData, TArrayView<uint8>(reinterpret_cast<uint8*>(&Data), sizeof(void*)));
-		Flags |= EParamFlags::Embedded;
-	}
-	else
-	{
-		Data = InData.GetData();
-	}
+		Object = InObject;
 
-	if (bInIsReference)
-	{
-		Flags |= EParamFlags::Reference;
-	}
+		const UClass* Class = InObject->GetClass();
 
-	if (bInIsMutable)
-	{
-		Flags |= EParamFlags::Mutable;
-	}
-}
+		TArray<FParamId, TInlineAllocator<128>> CachedIds;
+		TArray<FProperty*, TInlineAllocator<128>> CachedProperties;
+		TArray<FParamTypeHandle, TInlineAllocator<128>> CachedParamTypes;
 
-FParamStack::FParam::~FParam()
-{
-	if (Size > 0)
-	{
-		if (IsEmbedded())
+		// Determine param ID range for this layer
+		MinParamId = MAX_uint32;
+		uint32 MaxParamId = 0;
+
+		for (TFieldIterator<FProperty> PropIt(Class, EFieldIteratorFlags::IncludeSuper); PropIt; ++PropIt)
 		{
-			FParamHelpers::Destroy(TypeHandle, TArrayView<uint8>(reinterpret_cast<uint8*>(&Data), sizeof(void*)));
+			FParamTypeHandle ParamTypeHandle = FParamTypeHandle::FromProperty(*PropIt);
+			if(ParamTypeHandle.IsValid())
+			{
+				CachedParamTypes.Add(ParamTypeHandle);
+				CachedProperties.Add(*PropIt);
+				const FParamId& ParamId = CachedIds.Add_GetRef(FParamId(PropIt->GetFName()));
+				MinParamId = FMath::Min(ParamId.ToInt(), MinParamId);
+				MaxParamId = FMath::Max(ParamId.ToInt(), MaxParamId);
+			}
 		}
-	}
-}
 
-FParamStack::FParam::FParam(const FParam& InOtherParam)
-	: TypeHandle(InOtherParam.TypeHandle)
-	, Size(InOtherParam.Size)
-	, Flags(InOtherParam.Flags)
-{
-	if (IsEmbedded())
-	{
-		FParamHelpers::Copy(InOtherParam.TypeHandle, TConstArrayView<uint8>(reinterpret_cast<const uint8*>(&InOtherParam.Data), sizeof(void*)), TArrayView<uint8>(reinterpret_cast<uint8*>(&Data), sizeof(void*)));
+		if (MinParamId <= MaxParamId)
+		{
+			const uint32 ParamRangeSize = (MaxParamId - MinParamId) + 1;
+			Params.SetNumZeroed(ParamRangeSize);
+			for (uint32 PropertyIndex = 0; PropertyIndex < static_cast<uint32>(CachedProperties.Num()); ++PropertyIndex)
+			{
+				const FProperty* Property = CachedProperties[PropertyIndex];
+				const FParamId& ParamId = CachedIds[PropertyIndex];
+				uint8* DataPtr = Property->ContainerPtrToValuePtr<uint8>(InObject);
+				const uint32 LocalParamIndex = ParamId.ToInt() - MinParamId;
+				Params[LocalParamIndex] = Private::FParamEntry(CachedParamTypes[PropertyIndex], TArrayView<uint8>(DataPtr, Property->GetSize()), true, bInMutable);
+			}
+		}
 	}
-	else
+
+	// FParamStackLayer interface
+	virtual UObject* AsUObject() override
 	{
-		Data = InOtherParam.Data;
+		return Object.Get();
 	}
-}
 
-FParamStack::FParam& FParamStack::FParam::operator=(const FParam& InOtherParam)
+	TWeakObjectPtr<UObject> Object;
+};
+
+// Stack layer that can own its own data or reference an external FInstancedPropertyBag
+struct FInstancedPropertyBagLayer : FParamStackLayer, FGCObject
 {
-	TypeHandle = InOtherParam.TypeHandle;
-	Size = InOtherParam.Size;
-	Flags = InOtherParam.Flags;
+	FInstancedPropertyBagLayer() = delete;
 
-	if (IsEmbedded())
+	explicit FInstancedPropertyBagLayer(FInstancedPropertyBag& InPropertyBag, bool bInMutable)
 	{
-		FParamHelpers::Copy(InOtherParam.TypeHandle, TConstArrayView<uint8>(reinterpret_cast<const uint8*>(&InOtherParam.Data), sizeof(void*)), TArrayView<uint8>(reinterpret_cast<uint8*>(&Data), sizeof(void*)));
-	}
-	else
-	{
-		Data = InOtherParam.Data;
-	}
+		if (const UPropertyBag* PropertyBagStruct = InPropertyBag.GetPropertyBagStruct())
+		{
+			TConstArrayView<FPropertyBagPropertyDesc> Descs = PropertyBagStruct->GetPropertyDescs();
 
-	return *this;
-}
+			// Determine param ID range for this layer
+			TArray<FParamId, TInlineAllocator<128>> CachedIds;
+			CachedIds.SetNumUninitialized(Descs.Num());
+			MinParamId = MAX_uint32;
+			uint32 MaxParamId = 0;
+			for (uint32 DescIndex = 0; DescIndex < static_cast<uint32>(Descs.Num()); ++DescIndex)
+			{
+				const FPropertyBagPropertyDesc& Desc = Descs[DescIndex];
+				const FParamId& ParamId = CachedIds[DescIndex] = FParamId(Desc.Name);
+				MinParamId = FMath::Min(ParamId.ToInt(), MinParamId);
+				MaxParamId = FMath::Max(ParamId.ToInt(), MaxParamId);
+			}
 
-FParamStack::FParam::FParam(FParam&& InOtherParam)
-	: TypeHandle(InOtherParam.TypeHandle)
-	, Size(InOtherParam.Size)
-	, Flags(InOtherParam.Flags)
-{
-	if (IsEmbedded())
-	{
-		FParamHelpers::Copy(InOtherParam.TypeHandle, TConstArrayView<uint8>(reinterpret_cast<const uint8*>(&InOtherParam.Data), sizeof(void*)), TArrayView<uint8>(reinterpret_cast<uint8*>(&Data), sizeof(void*)));
+			if (MinParamId <= MaxParamId)
+			{
+				const uint32 ParamRangeSize = (MaxParamId - MinParamId) + 1;
+				Params.SetNumZeroed(ParamRangeSize);
+				FStructView StructView = InPropertyBag.GetMutableValue();
+				for (uint32 DescIndex = 0; DescIndex < static_cast<uint32>(Descs.Num()); ++DescIndex)
+				{
+					const FPropertyBagPropertyDesc& Desc = Descs[DescIndex];
+					const FParamId& ParamId = CachedIds[DescIndex];
+					uint8* DataPtr = StructView.GetMemory() + Desc.CachedProperty->GetOffset_ForInternal();
+					const uint32 LocalParamIndex = ParamId.ToInt() - MinParamId;
+					Params[LocalParamIndex] = Private::FParamEntry(FParamTypeHandle::FromPropertyBagPropertyDesc(Desc), TArrayView<uint8>(DataPtr, Desc.CachedProperty->GetSize()), true, true);
+				}
+			}
+		}
 	}
-	else
+
+	// FGCObject interface
+	virtual FString GetReferencerName() const override
 	{
-		Data = InOtherParam.Data;
+		return TEXT("AnimNext Instanced Property Bag Parameter Layer");
 	}
-}
+};
 
-FParamStack::FParam& FParamStack::FParam::operator=(FParam&& InOtherParam)
+// Stack layer that owns its own data as a FInstancedPropertyBag
+struct FInstancedPropertyBagValueLayer : FInstancedPropertyBagLayer
 {
-	TypeHandle = InOtherParam.TypeHandle;
-	Size = InOtherParam.Size;
-	Flags = InOtherParam.Flags;
+	FInstancedPropertyBagValueLayer() = delete;
+
+	explicit FInstancedPropertyBagValueLayer(FInstancedPropertyBag&& InPropertyBag, bool bInMutable)
+		: FInstancedPropertyBagLayer(InPropertyBag, bInMutable)
+		, PropertyBag(MoveTemp(InPropertyBag))
+	{}
 
-	if (IsEmbedded())
+	// FGCObject interface
+	virtual void AddReferencedObjects(FReferenceCollector& Collector) override
 	{
-		FParamHelpers::Copy(InOtherParam.TypeHandle, TConstArrayView<uint8>(reinterpret_cast<const uint8*>(&InOtherParam.Data), sizeof(void*)), TArrayView<uint8>(reinterpret_cast<uint8*>(&Data), sizeof(void*)));
+		PropertyBag.AddStructReferencedObjects(Collector);
 	}
-	else
+
+	// FParamStackLayer interface
+	virtual FInstancedPropertyBag* AsInstancedPropertyBag() override
 	{
-		Data = InOtherParam.Data;
+		return &PropertyBag;
 	}
 
-	return *this;
-}
+	FInstancedPropertyBag PropertyBag;
+};
 
-FParamStack::FLayer::FLayer(FInstancedPropertyBag& InPropertyBag, bool bInIsMutable)
+// Stack layer that references an external FInstancedPropertyBag
+struct FInstancedPropertyBagReferenceLayer : FInstancedPropertyBagLayer
 {
-	if (const UPropertyBag* PropertyBagStruct = InPropertyBag.GetPropertyBagStruct())
-	{
-		TConstArrayView<FPropertyBagPropertyDesc> Descs = PropertyBagStruct->GetPropertyDescs();
+	FInstancedPropertyBagReferenceLayer() = delete;
 
-		// Determine param ID range for this layer
-		TArray<FParamId> CachedIds;
-		CachedIds.SetNumUninitialized(Descs.Num());
-		MinParamId = MAX_uint32;
-		uint32 MaxParamId = 0;
-		for (uint32 DescIndex = 0; DescIndex < static_cast<uint32>(Descs.Num()); ++DescIndex)
-		{
-			const FPropertyBagPropertyDesc& Desc = Descs[DescIndex];
-			const FParamId& ParamId = CachedIds[DescIndex] = FParamId(Desc.Name);
-			MinParamId = FMath::Min(ParamId.ToInt(), MinParamId);
-			MaxParamId = FMath::Max(ParamId.ToInt(), MaxParamId);
-		}
+	explicit FInstancedPropertyBagReferenceLayer(FInstancedPropertyBag& InPropertyBag, bool bInMutable)
+		: FInstancedPropertyBagLayer(InPropertyBag, bInMutable)
+		, PropertyBag(InPropertyBag)
+	{}
 
-		if (MinParamId <= MaxParamId)
-		{
-			const uint32 ParamRangeSize = (MaxParamId - MinParamId) + 1;
-			Params.SetNumZeroed(ParamRangeSize);
-			FConstStructView StructView = InPropertyBag.GetValue();
-			for (uint32 DescIndex = 0; DescIndex < static_cast<uint32>(Descs.Num()); ++DescIndex)
-			{
-				const FPropertyBagPropertyDesc& Desc = Descs[DescIndex];
-				const FParamId& ParamId = CachedIds[DescIndex];
-				const uint8* DataPtr = StructView.GetMemory() + Desc.CachedProperty->GetOffset_ForInternal();
-				const uint32 LocalParamIndex = ParamId.ToInt() - MinParamId;
-				Params[LocalParamIndex] = FParamStack::FParam(FParamTypeHandle::FromPropertyBagPropertyDesc(Desc), TArrayView<uint8>(const_cast<uint8*>(DataPtr), Desc.CachedProperty->GetSize()), true, bInIsMutable);
-			}
-		}
-	}
-}
+	// FGCObject interface
+	virtual void AddReferencedObjects(FReferenceCollector& Collector) { /* We do not own the references held here, assume they are accounted for elsewhere */ }
 
-FParamStack::FLayer::FLayer(TConstArrayView<TPair<FParamId, FParamStack::FParam>> InParams)
-{
-	MinParamId = MAX_uint32;
-	uint32 MaxParamId = 0;
-	for (uint32 ParamIndex = 0; ParamIndex < static_cast<uint32>(InParams.Num()); ++ParamIndex)
+	// FParamStackLayer interface
+	virtual FInstancedPropertyBag* AsInstancedPropertyBag() override
 	{
-		MinParamId = FMath::Min(InParams[ParamIndex].Key.ToInt(), MinParamId);
-		MaxParamId = FMath::Max(InParams[ParamIndex].Key.ToInt(), MaxParamId);
+		return &PropertyBag;
 	}
 
-	if (MinParamId <= MaxParamId)
-	{
-		const uint32 ParamRangeSize = (MaxParamId - MinParamId) + 1;
-		Params.SetNumZeroed(ParamRangeSize);
-		for (uint32 ParamIndex = 0; ParamIndex < static_cast<uint32>(InParams.Num()); ++ParamIndex)
-		{
-			const TPair<FParamId, FParamStack::FParam>& Pair = InParams[ParamIndex];
-			const uint32 LocalParamIndex = Pair.Key.ToInt() - MinParamId;
-			Params[LocalParamIndex] = Pair.Value;
-		}
-	}
-}
+	FInstancedPropertyBag& PropertyBag;
+};
 
-FParamStack::FPushedLayer::FPushedLayer(FLayer& InLayer, FParamStack& InStack)
+FParamStack::FPushedLayer::FPushedLayer(FParamStackLayer& InLayer, FParamStack& InStack)
 	: Layer(InLayer)
 {
 	SerialNumber = InStack.MakeSerialNumber();
@@ -195,7 +183,7 @@ FParamStack::FPushedLayer::FPushedLayer(FLayer& InLayer, FParamStack& InStack)
 	}
 }
 
-FParamStack::FPushedLayer::FPushedLayer(const FPushedLayer& InPreviousLayer, FLayer& InLayer, FParamStack& InStack)
+FParamStack::FPushedLayer::FPushedLayer(const FPushedLayer& InPreviousLayer, FParamStackLayer& InLayer, FParamStack& InStack)
 	: Layer(InLayer)
 {
 	SerialNumber = InStack.MakeSerialNumber();
@@ -222,10 +210,12 @@ FParamStack::FPushedLayer::FPushedLayer(const FPushedLayer& InPreviousLayer, FLa
 	}
 }
 
-struct FParamStackThreadData : TThreadSingleton<FParamStackThreadData>
-{
-	FParamStack Stack;
-};
+// Current stack assigned to this thread
+static thread_local TWeakPtr<FParamStack> GWeakStack;
+
+// Stacks that are asscociated with objects, pending execution of an object's tick function
+static TMap<TObjectKey<UObject>, TWeakPtr<FParamStack>> GPendingObjects;
+static FRWLock GPendingObjectsLock;
 
 FParamStack::FParamStack()
 {
@@ -234,21 +224,86 @@ FParamStack::FParamStack()
 	PreviousLayerIndices.Reserve(FParamId::GetMaxParamId().ToInt());
 }
 
+FParamStack::~FParamStack()
+{
+}
+
+void FParamStack::SetParent(TWeakPtr<const FParamStack> InParent)
+{
+	WeakParentStack = InParent;
+}
+
 FParamStack& FParamStack::Get()
 {
-	return FParamStackThreadData::Get().Stack;
+	return *GWeakStack.Pin().Get();
+}
+
+TWeakPtr<FParamStack> FParamStack::GetForCurrentThread()
+{
+	return GWeakStack;
 }
 
-FParamStack::FPushedLayerHandle FParamStack::PushLayer(const FLayerHandle& InLayerHandle)
+void FParamStack::AddForPendingObject(const UObject* InObject, TWeakPtr<FParamStack> InStack)
+{
+	FRWScopeLock ScopeLock(GPendingObjectsLock, SLT_Write);
+	if(!GPendingObjects.Contains(InObject))
+	{
+		GPendingObjects.Add(InObject, InStack);
+	}
+}
+
+void FParamStack::RemoveForPendingObject(const UObject* InObject)
+{
+	FRWScopeLock ScopeLock(GPendingObjectsLock, SLT_Write);
+	GPendingObjects.Remove(InObject);
+}
+
+bool FParamStack::AttachToCurrentThreadForPendingObject(const UObject* InObject)
+{
+	FRWScopeLock ScopeLock(GPendingObjectsLock, SLT_ReadOnly);
+	if (TWeakPtr<FParamStack>* PendingStack = GPendingObjects.Find(InObject))
+	{
+		GWeakStack = *PendingStack;
+		return true;
+	}
+
+	return false;
+}
+
+bool FParamStack::DetachFromCurrentThreadForPendingObject(const UObject* InObject)
+{
+	FRWScopeLock ScopeLock(GPendingObjectsLock, SLT_ReadOnly);
+	if (TWeakPtr<FParamStack>* PendingStack = GPendingObjects.Find(InObject))
+	{
+		GWeakStack.Reset();
+		return true;
+	}
+
+	return false;
+}
+
+void FParamStack::AttachToCurrentThread(TWeakPtr<FParamStack> InStack)
+{
+	GWeakStack = InStack;
+}
+
+TWeakPtr<FParamStack> FParamStack::DetachFromCurrentThread()
+{
+	TWeakPtr<FParamStack> Stack = GWeakStack;
+	GWeakStack.Reset();
+	return Stack;
+}
+
+FParamStack::FPushedLayerHandle FParamStack::PushLayer(const FParamStackLayerHandle& InLayerHandle)
 {
 	return PushLayerInternal(*InLayerHandle.Layer.Get());
 }
 
-FParamStack::FPushedLayerHandle FParamStack::PushLayer(TConstArrayView<TPair<FParamId, FParamStack::FParam>> InParams)
+FParamStack::FPushedLayerHandle FParamStack::PushLayer(TConstArrayView<TPair<FParamId, Private::FParamEntry>> InParams)
 {
 	if(Layers.Num() < MAX_uint16)
 	{
-		FLayer& OwnedLayer = OwnedStackLayers.Add_GetRef(FLayer(InParams));
+		FParamStackLayer& OwnedLayer = OwnedStackLayers.Add_GetRef(FParamStackLayer(InParams));
 		OwnedLayer.OwnedStorageOffset = AllocAndCopyOwnedParamStorage(OwnedLayer.Params);
 		return PushLayerInternal(OwnedLayer);
 	}
@@ -259,31 +314,35 @@ FParamStack::FPushedLayerHandle FParamStack::PushLayer(TConstArrayView<TPair<FPa
 	}
 }
 
-FParamStack::FPushedLayerHandle FParamStack::PushLayerInternal(FLayer& InLayer)
+FParamStack::FPushedLayerHandle FParamStack::PushLayerInternal(FParamStackLayer& InLayer)
 {
-	if (Layers.Num() < MAX_uint16)
+	if(InLayer.Params.Num() > 0)
 	{
-		if (Layers.Num())
+		if (Layers.Num() < MAX_uint16)
 		{
-			FPushedLayer& NewPushedLayer = Layers.Emplace_GetRef(Layers.Top(), InLayer, *this);
-			return FPushedLayerHandle(Layers.Num() - 1, NewPushedLayer.SerialNumber);
+			if (Layers.Num())
+			{
+				FPushedLayer& NewPushedLayer = Layers.Emplace_GetRef(Layers.Top(), InLayer, *this);
+				return FPushedLayerHandle(Layers.Num() - 1, NewPushedLayer.SerialNumber);
+			}
+			else
+			{
+				FPushedLayer& NewPushedLayer = Layers.Emplace_GetRef(InLayer, *this);
+				return FPushedLayerHandle(Layers.Num() - 1, NewPushedLayer.SerialNumber);
+			}
 		}
 		else
 		{
-			FPushedLayer& NewPushedLayer = Layers.Emplace_GetRef(InLayer, *this);
-			return FPushedLayerHandle(Layers.Num() - 1, NewPushedLayer.SerialNumber);
+			UE_LOG(LogAnimation, Warning, TEXT("FParamStack: Could not push a layer: Maximum 65535 stack layers."));
 		}
 	}
-	else
-	{
-		UE_LOG(LogAnimation, Warning, TEXT("FParamStack: Could not push a layer: Maximum 65535 stack layers."));
-		return FPushedLayerHandle();
-	}
+
+	return FPushedLayerHandle();
 }
 
 void FParamStack::PopLayer(FPushedLayerHandle InHandle)
 {
-	if(Layers.Num() > 0)
+	if(InHandle.IsValid() && Layers.Num() > 0)
 	{
 		const FPushedLayer& TopLayer = Layers.Top();
 
@@ -307,6 +366,11 @@ void FParamStack::PopLayer(FPushedLayerHandle InHandle)
 
 			PreviousLayerIndices.SetNum(PreviousLayerIndices.Num() - NumParams);
 		}
+		else
+		{
+			// Clear layer indices as this is the last layer to be popped
+			FMemory::Memset(&LayerIndices[0], 0xff, LayerIndices.Num() * LayerIndices.GetTypeSize());
+		}
 
 		// Dont shrink allocs to avoid thrashing
 		constexpr bool bAllowShrinking = false;
@@ -326,108 +390,135 @@ void FParamStack::PopLayer(FPushedLayerHandle InHandle)
 	}
 }
 
-FParamStack::FLayerHandle FParamStack::MakeLayer(const FInstancedPropertyBag& InPropertyBag)
+FParamStackLayerHandle FParamStack::MakeValueLayer(const UClass* InClass)
 {
-	TUniquePtr<FLayer> Layer = TUniquePtr<FLayer>(new FLayer(const_cast<FInstancedPropertyBag&>(InPropertyBag), false));
-	return FLayerHandle(MoveTemp(Layer));
+	check(InClass != nullptr);
+
+	UObject* OwnedObject = NewObject<UObject>(GetTransientPackage(), InClass);
+	TUniquePtr<FParamStackLayer> Layer = MakeUnique<FUObjectLayer>(OwnedObject, true);
+	return FParamStackLayerHandle(MoveTemp(Layer));
 }
 
-FParamStack::FLayerHandle FParamStack::MakeMutableLayer(FInstancedPropertyBag& InPropertyBag)
+FParamStackLayerHandle FParamStack::MakeReferenceLayer(UObject* InObject)
 {
-	TUniquePtr<FLayer> Layer = TUniquePtr<FLayer>(new FLayer(InPropertyBag, true));
-	return FLayerHandle(MoveTemp(Layer));
+	check(InObject != nullptr);
+
+	TUniquePtr<FParamStackLayer> Layer = MakeUnique<FUObjectLayer>(InObject, true);
+	return FParamStackLayerHandle(MoveTemp(Layer));
 }
 
-FParamStack::FLayerHandle FParamStack::MakeLayer(TConstArrayView<TPair<FParamId, FParamStack::FParam>> InParams)
+FParamStackLayerHandle FParamStack::MakeValueLayer(const FInstancedPropertyBag& InPropertyBag)
 {
-	TUniquePtr<FLayer> Layer = TUniquePtr<FLayer>(new FLayer(InParams));
-	return FLayerHandle(MoveTemp(Layer));
+	FInstancedPropertyBag OwnedPropertyBag = InPropertyBag;
+	TUniquePtr<FParamStackLayer> Layer = MakeUnique<FInstancedPropertyBagValueLayer>(MoveTemp(OwnedPropertyBag), true);
+	return FParamStackLayerHandle(MoveTemp(Layer));
 }
 
-FParamStack::EGetParamResult FParamStack::GetParamData(FParamId InId, FParamTypeHandle InTypeHandle, TConstArrayView<uint8>& OutParamData, FParamCompatibility InRequiredCompatibility) const
+FParamStackLayerHandle FParamStack::MakeReferenceLayer(FInstancedPropertyBag& InPropertyBag)
+{
+	TUniquePtr<FParamStackLayer> Layer = MakeUnique<FInstancedPropertyBagReferenceLayer>(InPropertyBag, true);
+	return FParamStackLayerHandle(MoveTemp(Layer));
+}
+
+FParamStackLayerHandle FParamStack::MakeLayer(TConstArrayView<TPair<FParamId, Private::FParamEntry>> InParams)
+{
+	TUniquePtr<FParamStackLayer> Layer = MakeUnique<FParamStackLayer>(InParams);
+	return FParamStackLayerHandle(MoveTemp(Layer));
+}
+
+FParamResult FParamStack::GetParamData(FParamId InId, FParamTypeHandle InTypeHandle, TConstArrayView<uint8>& OutParamData, FParamCompatibility InRequiredCompatibility) const
 {
 	FParamTypeHandle ParamTypeHandle;
 	return GetParamData(InId, InTypeHandle, OutParamData, ParamTypeHandle, InRequiredCompatibility);
 }
 
-FParamStack::EGetParamResult FParamStack::GetParamData(FParamId InId, FParamTypeHandle InTypeHandle, TConstArrayView<uint8>& OutParamData, FParamTypeHandle& OutParamTypeHandle, FParamCompatibility InRequiredCompatibility) const
+FParamResult FParamStack::GetParamData(FParamId InId, FParamTypeHandle InTypeHandle, TConstArrayView<uint8>& OutParamData, FParamTypeHandle& OutParamTypeHandle, FParamCompatibility InRequiredCompatibility) const
 {
-	if (Layers.Num() == 0 || InId.ToInt() >= (uint32)LayerIndices.Num() || LayerIndices[InId.ToInt()] == MAX_uint16)
+	const FParamResult Result = GetParamDataInternal(InId, InTypeHandle, OutParamData, OutParamTypeHandle, InRequiredCompatibility);
+	if (Result.IsInScope())
 	{
-		return EGetParamResult::NotInScope;
+		return Result;
 	}
-
-	const FPushedLayer& Layer = Layers[LayerIndices[InId.ToInt()]];
-	const uint32 LocalParamIndex = InId.ToInt() - Layer.Layer.MinParamId;
-
-	const FParam& Param = Layer.Layer.Params[LocalParamIndex];
-	if (!Param.IsValid())
+	else if (TSharedPtr<const FParamStack> ParentStack = WeakParentStack.Pin())
+	{
+		return ParentStack->GetParamData(InId, InTypeHandle, OutParamData, OutParamTypeHandle, InRequiredCompatibility);
+	}
+	// Check if we have a built in adapter to fall back on
+	else if(const FParamAdapter* Adapter = InId.GetAdapter(InId))
 	{
-		return EGetParamResult::NotInScope;
+		return Adapter->GetParamData(InTypeHandle, OutParamData, OutParamTypeHandle, InRequiredCompatibility);
 	}
 
-	OutParamTypeHandle = Param.GetTypeHandle();
+	return EParamResult::NotInScope;
+}
 
-	FParamCompatibility Compatibility = UE::AnimNext::FParamUtils::GetCompatibility(InTypeHandle, OutParamTypeHandle);
-	if (Compatibility < InRequiredCompatibility)
+FParamResult FParamStack::GetParamDataInternal(FParamId InId, FParamTypeHandle InTypeHandle, TConstArrayView<uint8>& OutParamData, FParamTypeHandle& OutParamTypeHandle, FParamCompatibility InRequiredCompatibility) const
+{
+	if (Layers.Num() == 0 || InId.ToInt() >= (uint32)LayerIndices.Num() || LayerIndices[InId.ToInt()] == MAX_uint16)
 	{
-		return EGetParamResult::IncompatibleType;
+		return EParamResult::NotInScope;
 	}
-	
-	OutParamData = Param.GetData();
 
-	if (Compatibility == InRequiredCompatibility)
+	const FPushedLayer& Layer = Layers[LayerIndices[InId.ToInt()]];
+	const FParamResult Result = Layer.Layer.GetParamData(InId, InTypeHandle, OutParamData, OutParamTypeHandle, InRequiredCompatibility);
+	if(Result.IsInScope())
 	{
-		return EGetParamResult::CompatibleType;
+		return Result;
 	}
 
-	return EGetParamResult::Succeeded;
+
+
+	return EParamResult::NotInScope;
 }
 
-FParamStack::EGetParamResult FParamStack::GetMutableParamData(FParamId InId, FParamTypeHandle InTypeHandle, TArrayView<uint8>& OutParamData, FParamCompatibility InRequiredCompatibility)
+FParamResult FParamStack::GetMutableParamData(FParamId InId, FParamTypeHandle InTypeHandle, TArrayView<uint8>& OutParamData, FParamCompatibility InRequiredCompatibility)
 {
 	FParamTypeHandle ParamTypeHandle;
 	return GetMutableParamData(InId, InTypeHandle, OutParamData, ParamTypeHandle, InRequiredCompatibility);
 }
 
-FParamStack::EGetParamResult FParamStack::GetMutableParamData(FParamId InId, FParamTypeHandle InTypeHandle, TArrayView<uint8>& OutParamData, FParamTypeHandle& OutParamTypeHandle, FParamCompatibility InRequiredCompatibility)
+FParamResult FParamStack::GetMutableParamData(FParamId InId, FParamTypeHandle InTypeHandle, TArrayView<uint8>& OutParamData, FParamTypeHandle& OutParamTypeHandle, FParamCompatibility InRequiredCompatibility)
 {
-	if (Layers.Num() == 0 || InId.ToInt() >= (uint32)LayerIndices.Num() || LayerIndices[InId.ToInt()] == MAX_uint16)
+	const FParamResult Result = GetMutableParamDataInternal(InId, InTypeHandle, OutParamData, OutParamTypeHandle, InRequiredCompatibility);
+	if (Result.IsInScope())
 	{
-		return EGetParamResult::NotInScope;
+		return Result;
 	}
-
-	FPushedLayer& Layer = Layers[LayerIndices[InId.ToInt()]];
-	const uint32 LocalParamIndex = InId.ToInt() - Layer.Layer.MinParamId;
-
-	FParam& Param = Layer.Layer.Params[LocalParamIndex];
-	if (!Param.IsValid())
+	else if (TSharedPtr<const FParamStack> ParentStack = WeakParentStack.Pin())
 	{
-		return EGetParamResult::NotInScope;
+		// we use a dummy here because if the data is present in a parent, it must be immutable anyways and we will early out
+		TConstArrayView<uint8> ParamData;
+		FParamResult ParentResult = ParentStack->GetParamData(InId, InTypeHandle, ParamData, OutParamTypeHandle, InRequiredCompatibility);
+		if (ParentResult.IsInScope())
+		{
+			// Parent data is immutable
+			return ParentResult.Result & EParamResult::MutabilityError;
+		}
 	}
-
-	OutParamTypeHandle = Param.GetTypeHandle();
-
-	EGetParamResult AccessResult = EGetParamResult::Succeeded;
-	FParamCompatibility Compatibility = UE::AnimNext::FParamUtils::GetCompatibility(InTypeHandle, OutParamTypeHandle);
-	if (Compatibility < InRequiredCompatibility)
+	// Check if we have a built in adapter to fall back on
+	else if(const FParamAdapter* Adapter = InId.GetAdapter(InId))
 	{
-		AccessResult |= EGetParamResult::IncompatibleType;
+		return Adapter->GetMutableParamData(InTypeHandle, OutParamData, OutParamTypeHandle, InRequiredCompatibility);
 	}
-	else if (Compatibility == InRequiredCompatibility)
+
+	return EParamResult::NotInScope;
+}
+
+FParamResult FParamStack::GetMutableParamDataInternal(FParamId InId, FParamTypeHandle InTypeHandle, TArrayView<uint8>& OutParamData, FParamTypeHandle& OutParamTypeHandle, FParamCompatibility InRequiredCompatibility)
+{
+	if (Layers.Num() == 0 || InId.ToInt() >= (uint32)LayerIndices.Num() || LayerIndices[InId.ToInt()] == MAX_uint16)
 	{
-		AccessResult |= EGetParamResult::CompatibleType;
+		return EParamResult::NotInScope;
 	}
 
-	AccessResult |= !Param.IsMutable() ? EGetParamResult::Immutable : EGetParamResult::Succeeded;
-	if (AccessResult != EGetParamResult::Succeeded && AccessResult != EGetParamResult::CompatibleType)
+	FPushedLayer& Layer = Layers[LayerIndices[InId.ToInt()]];
+	const FParamResult Result =  Layer.Layer.GetMutableParamData(InId, InTypeHandle, OutParamData, OutParamTypeHandle, InRequiredCompatibility);
+	if(Result.IsInScope())
 	{
-		return AccessResult;
+		return Result;
 	}
 
-	OutParamData = Param.GetMutableData();
-
-	return EGetParamResult::Succeeded;
+	return EParamResult::NotInScope;
 }
 
 bool FParamStack::IsMutableParam(FParamId InId) const
@@ -440,7 +531,7 @@ bool FParamStack::IsMutableParam(FParamId InId) const
 	const FPushedLayer& Layer = Layers[LayerIndices[InId.ToInt()]];
 	const uint32 LocalParamIndex = InId.ToInt() - Layer.Layer.MinParamId;
 
-	const FParam& Param = Layer.Layer.Params[LocalParamIndex];
+	const Private::FParamEntry& Param = Layer.Layer.Params[LocalParamIndex];
 
 	return Param.IsMutable();
 }
@@ -455,17 +546,30 @@ bool FParamStack::IsReferenceParam(FParamId InId) const
 	const FPushedLayer& Layer = Layers[LayerIndices[InId.ToInt()]];
 	const uint32 LocalParamIndex = InId.ToInt() - Layer.Layer.MinParamId;
 
-	const FParam& Param = Layer.Layer.Params[LocalParamIndex];
+	const Private::FParamEntry& Param = Layer.Layer.Params[LocalParamIndex];
 
 	return Param.IsReference();
 }
 
-uint32 FParamStack::AllocAndCopyOwnedParamStorage(TArrayView<FParam> InParams)
+bool FParamStack::LayerContainsParam(const FParamStackLayerHandle& InHandle, FName InKey)
+{
+	const FParamId ParamIdToFind(InKey);
+	const uint32 MinParamId = InHandle.Layer->MinParamId;
+	const uint32 NumParams = InHandle.Layer->Params.Num();
+	if (ParamIdToFind.ToInt() >= MinParamId && ParamIdToFind.ToInt() < MinParamId + NumParams)
+	{
+		const uint32 LocalParamIndex = ParamIdToFind.ToInt() - MinParamId;
+		return InHandle.Layer->Params[LocalParamIndex].IsValid();
+	}
+	return false;
+}
+
+uint32 FParamStack::AllocAndCopyOwnedParamStorage(TArrayView<Private::FParamEntry> InParams)
 {
 	const uint32 CurrentOffset = OwnedLayerParamStorage.Num();
 	const uint32 PageSize = OwnedLayerParamStorage.MaxPerPage();
 
-	for (FParam& Param : InParams)
+	for (Private::FParamEntry& Param : InParams)
 	{
 		if(Param.IsValid())
 		{
diff --git a/Engine/Plugins/Experimental/AnimNext/Source/AnimNext/Private/Param/ParamStackLayer.cpp b/Engine/Plugins/Experimental/AnimNext/Source/AnimNext/Private/Param/ParamStackLayer.cpp
new file mode 100644
--- /dev/null
+++ b/Engine/Plugins/Experimental/AnimNext/Source/AnimNext/Private/Param/ParamStackLayer.cpp
@@ -0,0 +1,141 @@
+// Copyright Epic Games, Inc. All Rights Reserved.
+
+#include "Param/ParamStackLayer.h"
+#include "Param/ParamEntry.h"
+#include "Param/ParamResult.h"
+#include "Param/ParamUtils.h"
+
+namespace UE::AnimNext
+{
+
+FParamStackLayer::FParamStackLayer(TConstArrayView<TPair<FParamId, Private::FParamEntry>> InParams)
+{
+	MinParamId = MAX_uint32;
+	uint32 MaxParamId = 0;
+	for (uint32 ParamIndex = 0; ParamIndex < static_cast<uint32>(InParams.Num()); ++ParamIndex)
+	{
+		MinParamId = FMath::Min(InParams[ParamIndex].Key.ToInt(), MinParamId);
+		MaxParamId = FMath::Max(InParams[ParamIndex].Key.ToInt(), MaxParamId);
+	}
+
+	if (MinParamId <= MaxParamId)
+	{
+		const uint32 ParamRangeSize = (MaxParamId - MinParamId) + 1;
+		Params.SetNumZeroed(ParamRangeSize);
+		for (uint32 ParamIndex = 0; ParamIndex < static_cast<uint32>(InParams.Num()); ++ParamIndex)
+		{
+			const TPair<FParamId, Private::FParamEntry>& Pair = InParams[ParamIndex];
+			const uint32 LocalParamIndex = Pair.Key.ToInt() - MinParamId;
+			Params[LocalParamIndex] = Pair.Value;
+		}
+	}
+}
+
+
+FParamResult FParamStackLayer::GetParamData(FParamId InId, FParamTypeHandle InTypeHandle, TConstArrayView<uint8>& OutParamData) const
+{
+	const uint32 LocalParamIndex = InId.ToInt() - MinParamId;
+
+	const Private::FParamEntry& Param = Params[LocalParamIndex];
+	if (!Param.IsValid())
+	{
+		return EParamResult::NotInScope;
+	}
+
+	FParamCompatibility Compatibility = UE::AnimNext::FParamUtils::GetCompatibility(InTypeHandle, Param.GetTypeHandle());
+	if (!Compatibility.IsCompatible())
+	{
+		return EParamResult::TypeError;
+	}
+
+	OutParamData = Param.GetData();
+	return EParamResult::Success;
+}
+
+FParamResult FParamStackLayer::GetParamData(FParamId InId, FParamTypeHandle InTypeHandle, TConstArrayView<uint8>& OutParamData, FParamTypeHandle& OutParamTypeHandle, FParamCompatibility InRequiredCompatibility) const
+{
+	const uint32 LocalParamIndex = InId.ToInt() - MinParamId;
+
+	const Private::FParamEntry& Param = Params[LocalParamIndex];
+	if (!Param.IsValid())
+	{
+		return EParamResult::NotInScope;
+	}
+
+	OutParamTypeHandle = Param.GetTypeHandle();
+
+	FParamCompatibility Compatibility = UE::AnimNext::FParamUtils::GetCompatibility(InTypeHandle, OutParamTypeHandle);
+	if (Compatibility < InRequiredCompatibility)
+	{
+		return EParamResult::TypeError;
+	}
+
+	OutParamData = Param.GetData();
+
+	if (Compatibility == InRequiredCompatibility)
+	{
+		return EParamResult::Success | EParamResult::TypeCompatible;
+	}
+
+	return EParamResult::Success;
+}
+
+FParamResult FParamStackLayer::GetMutableParamData(FParamId InId, FParamTypeHandle InTypeHandle, TArrayView<uint8>& OutParamData)
+{
+	const uint32 LocalParamIndex = InId.ToInt() - MinParamId;
+
+	Private::FParamEntry& Param = Params[LocalParamIndex];
+	if (!Param.IsValid())
+	{
+		return EParamResult::NotInScope;
+	}
+
+	FParamResult AccessResult = EParamResult::Success;
+	FParamCompatibility Compatibility = UE::AnimNext::FParamUtils::GetCompatibility(InTypeHandle, Param.GetTypeHandle());
+	if (!Compatibility.IsCompatible())
+	{
+		AccessResult.Result |= EParamResult::TypeError;
+	}
+
+	AccessResult.Result |= !Param.IsMutable() ? EParamResult::MutabilityError : EParamResult::Success;
+	if (AccessResult.IsSuccessful())
+	{
+		OutParamData = Param.GetMutableData();
+	}
+
+	return AccessResult;
+}
+
+FParamResult FParamStackLayer::GetMutableParamData(FParamId InId, FParamTypeHandle InTypeHandle, TArrayView<uint8>& OutParamData, FParamTypeHandle& OutParamTypeHandle, FParamCompatibility InRequiredCompatibility)
+{
+	const uint32 LocalParamIndex = InId.ToInt() - MinParamId;
+
+	Private::FParamEntry& Param = Params[LocalParamIndex];
+	if (!Param.IsValid())
+	{
+		return EParamResult::NotInScope;
+	}
+
+	OutParamTypeHandle = Param.GetTypeHandle();
+
+	FParamResult AccessResult = EParamResult::Success;
+	FParamCompatibility Compatibility = UE::AnimNext::FParamUtils::GetCompatibility(InTypeHandle, OutParamTypeHandle);
+	if (Compatibility < InRequiredCompatibility)
+	{
+		AccessResult.Result |= EParamResult::TypeError;
+	}
+	else if (Compatibility == InRequiredCompatibility)
+	{
+		AccessResult.Result |= EParamResult::TypeCompatible;
+	}
+
+	AccessResult.Result |= !Param.IsMutable() ? EParamResult::MutabilityError : EParamResult::Success;
+	if (AccessResult.IsSuccessful())
+	{
+		OutParamData = Param.GetMutableData();
+	}
+
+	return AccessResult;
+}
+
+}
\ No newline at end of file
diff --git a/Engine/Plugins/Experimental/AnimNext/Source/AnimNext/Private/Param/ParamStackLayer.h b/Engine/Plugins/Experimental/AnimNext/Source/AnimNext/Private/Param/ParamStackLayer.h
new file mode 100644
--- /dev/null
+++ b/Engine/Plugins/Experimental/AnimNext/Source/AnimNext/Private/Param/ParamStackLayer.h
@@ -0,0 +1,57 @@
+// Copyright Epic Games, Inc. All Rights Reserved.
+
+#pragma once
+
+#include "CoreMinimal.h"
+#include "Param/ParamId.h"
+#include "Param/ParamCompatibility.h"
+
+struct FInstancedPropertyBag;
+
+namespace UE::AnimNext
+{
+	struct FParamResult;
+	struct FParamTypeHandle;
+}
+
+namespace UE::AnimNext::Private
+{
+	struct FParamEntry;
+}
+
+namespace UE::AnimNext
+{
+
+// Stack layers are what actually get pushed/popped onto the layer stack.
+// They are designed to be held on an instance, their data updated in place.
+// Memory ownership for params is variable depending on the layer's subclass
+struct FParamStackLayer
+{
+	FParamStackLayer() = default;
+	virtual ~FParamStackLayer() = default;
+
+	friend struct FParamStack;
+	friend struct FParamStackLayerHandle;
+
+	explicit FParamStackLayer(TConstArrayView<TPair<FParamId, Private::FParamEntry>> InParams);
+
+	ANIMNEXT_API FParamResult GetParamData(FParamId InId, FParamTypeHandle InTypeHandle, TConstArrayView<uint8>& OutParamData) const;
+	ANIMNEXT_API FParamResult GetParamData(FParamId InId, FParamTypeHandle InTypeHandle, TConstArrayView<uint8>& OutParamData, FParamTypeHandle& OutParamTypeHandle, FParamCompatibility InRequiredCompatibility = FParamCompatibility::Equal()) const;
+
+	ANIMNEXT_API FParamResult GetMutableParamData(FParamId InId, FParamTypeHandle InTypeHandle, TArrayView<uint8>& OutParamData);
+	ANIMNEXT_API FParamResult GetMutableParamData(FParamId InId, FParamTypeHandle InTypeHandle, TArrayView<uint8>& OutParamData, FParamTypeHandle& OutParamTypeHandle, FParamCompatibility InRequiredCompatibility = FParamCompatibility::Equal());
+
+	ANIMNEXT_API virtual UObject* AsUObject() { return nullptr; }
+	ANIMNEXT_API virtual FInstancedPropertyBag* AsInstancedPropertyBag() { return nullptr; }
+
+	// Params that this layer supplies
+	TArray<Private::FParamEntry> Params;
+
+	// ID that the param indices start at. Maps the global param ID range into the range for this layer
+	uint32 MinParamId = 0;
+
+	// Storage offset for this layer if it is owned internally, otherwise MAX_uint32
+	uint32 OwnedStorageOffset = MAX_uint32;
+};
+
+}
\ No newline at end of file
diff --git a/Engine/Plugins/Experimental/AnimNext/Source/AnimNext/Private/Param/ParamStackLayerHandle.cpp b/Engine/Plugins/Experimental/AnimNext/Source/AnimNext/Private/Param/ParamStackLayerHandle.cpp
new file mode 100644
--- /dev/null
+++ b/Engine/Plugins/Experimental/AnimNext/Source/AnimNext/Private/Param/ParamStackLayerHandle.cpp
@@ -0,0 +1,101 @@
+// Copyright Epic Games, Inc. All Rights Reserved.
+
+#include "Param/ParamStackLayerHandle.h"
+#include "Param/ParamStackLayer.h"
+#include "Param/ParamResult.h"
+#include "Param/ParamHelpers.h"
+#include "Param/ParamEntry.h"
+
+namespace UE::AnimNext
+{
+
+FParamStackLayerHandle::FParamStackLayerHandle()
+{
+}
+
+FParamStackLayerHandle::FParamStackLayerHandle(FParamStackLayerHandle&& InLayer) noexcept
+{
+	Layer = MoveTemp(InLayer.Layer);
+}
+
+FParamStackLayerHandle& FParamStackLayerHandle::operator=(FParamStackLayerHandle&& InLayer) noexcept
+{
+	Layer = MoveTemp(InLayer.Layer);
+	return *this;
+}
+
+FParamStackLayerHandle::FParamStackLayerHandle(TUniquePtr<FParamStackLayer>&& InLayer)
+	: Layer(MoveTemp(InLayer))
+{}
+
+FParamStackLayerHandle::~FParamStackLayerHandle()
+{
+	Layer = nullptr;
+}
+
+bool FParamStackLayerHandle::IsValid() const
+{
+	return Layer.IsValid();
+}
+
+void FParamStackLayerHandle::Invalidate()
+{
+	Layer = nullptr;
+}
+
+FParamResult FParamStackLayerHandle::SetValueRaw(FParamId InParamId, FParamTypeHandle InTypeHandle, TConstArrayView<uint8> InData)
+{
+	check(Layer.IsValid());
+
+	TArrayView<uint8> ParamData;
+	FParamResult Result = Layer->GetMutableParamData(InParamId, InTypeHandle, ParamData);
+	if(Result.IsSuccessful())
+	{
+		FParamHelpers::Copy(InTypeHandle, InData, ParamData);
+	}
+
+	return Result;
+}
+
+FParamResult FParamStackLayerHandle::GetValueRaw(FParamId InParamId, FParamTypeHandle InTypeHandle, TConstArrayView<uint8>& OutData)
+{
+	check(Layer.IsValid());
+
+	TArrayView<uint8> ParamData;
+	return Layer->GetParamData(InParamId, InTypeHandle, OutData);
+}
+
+FParamResult FParamStackLayerHandle::SetValuesInternal(TConstArrayView<TPair<FParamId, Private::FParamEntry>> InParams)
+{
+	check(Layer.IsValid());
+
+	const uint32 MinParamId = Layer->MinParamId;
+	for (const TPair<FParamId, Private::FParamEntry>& ParamPair : InParams)
+	{
+		const uint32 LocalParamIndex = ParamPair.Key.ToInt() - MinParamId;
+		if(Layer->Params.IsValidIndex(LocalParamIndex))
+		{
+			Private::FParamEntry& Param = Layer->Params[LocalParamIndex];
+			
+			check(Param.GetTypeHandle() == ParamPair.Value.GetTypeHandle());
+			check(Param.IsMutable());
+
+			FParamHelpers::Copy(ParamPair.Value.GetTypeHandle(), ParamPair.Value.GetData(), Param.GetMutableData());
+			return EParamResult::Success;
+		}
+	}
+
+	return EParamResult::NotInScope;
+}
+
+UObject* FParamStackLayerHandle::GetUObjectFromLayer() const
+{
+	return Layer->AsUObject();
+}
+
+FInstancedPropertyBag* FParamStackLayerHandle::GetInstancedPropertyBagFromLayer() const
+{
+	return Layer->AsInstancedPropertyBag();
+}
+
+}
\ No newline at end of file
diff --git a/Engine/Plugins/Experimental/AnimNext/Source/AnimNext/Private/Param/ParamType.cpp b/Engine/Plugins/Experimental/AnimNext/Source/AnimNext/Private/Param/ParamType.cpp
--- a/Engine/Plugins/Experimental/AnimNext/Source/AnimNext/Private/Param/ParamType.cpp
+++ b/Engine/Plugins/Experimental/AnimNext/Source/AnimNext/Private/Param/ParamType.cpp
@@ -1,4 +1,4 @@
-﻿// Copyright Epic Games, Inc. All Rights Reserved.
+// Copyright Epic Games, Inc. All Rights Reserved.
 
 #include "Param/ParamType.h"
 #include "Param/ParamTypeHandle.h"
@@ -484,3 +484,122 @@ FString FAnimNextParamType::ToString() const
 
 	return StringBuilder.ToString();
 }
+
+FAnimNextParamType FAnimNextParamType::FromString(const FString& InString)
+{
+	auto GetInnerType = [](const FString& InTypeString, FAnimNextParamType& OutType)
+	{
+		static TMap<FString, FAnimNextParamType> BasicTypes =
+		{
+			{ TEXT("bool"),		GetType<bool>() },
+			{ TEXT("uint8"),		GetType<uint8>() },
+			{ TEXT("int32"),		GetType<int32>() },
+			{ TEXT("int64"),		GetType<int64>() },
+			{ TEXT("float"),		GetType<float>() },
+			{ TEXT("double"),		GetType<double>() },
+			{ TEXT("FName"),		GetType<FName>() },
+			{ TEXT("FString"),		GetType<FString>() },
+			{ TEXT("FText"),		GetType<FText>() },
+			{ TEXT("uint32"),		GetType<uint32>() },
+			{ TEXT("uint64"),		GetType<uint64>() },
+		};
+
+		if(const FAnimNextParamType* BasicType = BasicTypes.Find(InTypeString))
+		{
+			OutType = *BasicType;
+			return true;
+		}
+
+		// Check for object/struct/enum
+		EValueType ObjectValueType = EValueType::None;
+		FString ObjectInnerString;
+		if(InTypeString.StartsWith(TEXT("U"), ESearchCase::CaseSensitive))
+		{
+			ObjectInnerString = InTypeString.RightChop(1).TrimStartAndEnd();
+			ObjectValueType = EValueType::Object;
+		}
+		else if(InTypeString.StartsWith(TEXT("TObjectPtr<U"), ESearchCase::CaseSensitive))
+		{
+			ObjectInnerString = InTypeString.RightChop(12).LeftChop(1).TrimStartAndEnd();
+			ObjectValueType = EValueType::Object;
+		}
+		else if(InTypeString.StartsWith(TEXT("TSubClassOf<U"), ESearchCase::CaseSensitive))
+		{
+			ObjectInnerString = InTypeString.RightChop(13).LeftChop(1).TrimStartAndEnd();
+			ObjectValueType = EValueType::Class;
+		}
+		else if(InTypeString.StartsWith(TEXT("F"), ESearchCase::CaseSensitive))
+		{
+			ObjectInnerString = InTypeString.RightChop(1).TrimStartAndEnd();
+			ObjectValueType = EValueType::Struct;
+		}
+		else if(InTypeString.StartsWith(TEXT("E"), ESearchCase::CaseSensitive))
+		{
+			ObjectInnerString = InTypeString.RightChop(1).TrimStartAndEnd();
+			ObjectValueType = EValueType::Enum;
+		}
+		
+		if(UObject* ObjectType = FindFirstObject<UObject>(*ObjectInnerString, EFindFirstObjectOptions::NativeFirst))
+		{
+			OutType.ValueType = ObjectValueType;
+			OutType.ValueTypeObject = nullptr;
+
+			switch(ObjectValueType)
+			{
+			case EPropertyBagPropertyType::Enum:
+				OutType.ValueTypeObject = Cast<UEnum>(ObjectType);
+				break;
+			case EPropertyBagPropertyType::Struct:
+				OutType.ValueTypeObject = Cast<UScriptStruct>(ObjectType);
+				break;
+			case EPropertyBagPropertyType::Object:
+				OutType.ValueTypeObject = Cast<UClass>(ObjectType);
+				break;
+			case EPropertyBagPropertyType::Class:
+				OutType.ValueTypeObject = Cast<UClass>(ObjectType);
+				break;
+			default:
+				break;
+			}
+
+			return OutType.ValueTypeObject != nullptr;
+		}
+
+		return false;
+	};
+
+	{
+		FAnimNextParamType Type;
+		if(GetInnerType(InString, Type))
+		{
+			return Type;
+		}
+	}
+	
+	if(InString.StartsWith(TEXT("TArray<"), ESearchCase::CaseSensitive))
+	{
+		const FString InnerTypeString = InString.RightChop(7).LeftChop(1).TrimStartAndEnd();
+		FAnimNextParamType Type;
+		if(GetInnerType(InnerTypeString, Type))
+		{
+			Type.ContainerType = EContainerType::Array;
+			return Type;
+		}
+	}
+
+	return FAnimNextParamType();
+}
+
+bool FAnimNextParamType::IsObjectType() const
+{
+	switch (ValueType)
+	{
+	case EValueType::Object:
+	case EValueType::SoftObject:
+	case EValueType::Class:
+	case EValueType::SoftClass:
+		return true;
+	default:
+		return false;
+	}
+}
\ No newline at end of file
diff --git a/Engine/Plugins/Experimental/AnimNext/Source/AnimNext/Private/Param/ParamTypeHandle.cpp b/Engine/Plugins/Experimental/AnimNext/Source/AnimNext/Private/Param/ParamTypeHandle.cpp
--- a/Engine/Plugins/Experimental/AnimNext/Source/AnimNext/Private/Param/ParamTypeHandle.cpp
+++ b/Engine/Plugins/Experimental/AnimNext/Source/AnimNext/Private/Param/ParamTypeHandle.cpp
@@ -443,6 +443,35 @@ FParamTypeHandle FParamTypeHandle::FromProperty(const FProperty* InProperty)
 	return Handle;
 }
 
+FParamTypeHandle FParamTypeHandle::FromObject(const UObject* InObject)
+{
+	FParamTypeHandle Handle;
+
+	if(InObject)
+	{
+		Handle.SetParameterType(EParamType::Custom);
+
+		if(const UEnum* Enum = Cast<UEnum>(InObject))
+		{
+			Handle.SetCustomTypeIndex(GetOrAllocateCustomTypeIndex(FAnimNextParamType::EValueType::Enum, FAnimNextParamType::EContainerType::None, Enum));
+		}
+		else if(const UClass* Class = Cast<UClass>(InObject))
+		{
+			Handle.SetCustomTypeIndex(GetOrAllocateCustomTypeIndex(FAnimNextParamType::EValueType::Class, FAnimNextParamType::EContainerType::None, Class));
+		}
+		else if(const UScriptStruct* Struct = Cast<UScriptStruct>(InObject))
+		{
+			Handle.SetCustomTypeIndex(GetOrAllocateCustomTypeIndex(FAnimNextParamType::EValueType::Struct, FAnimNextParamType::EContainerType::None, Struct));
+		}
+		else
+		{
+			Handle.SetCustomTypeIndex(GetOrAllocateCustomTypeIndex(FAnimNextParamType::EValueType::Object, FAnimNextParamType::EContainerType::None, InObject->GetClass()));
+		}
+	}
+
+	return Handle;
+}
+
 FAnimNextParamType FParamTypeHandle::GetType() const
 {
 	FAnimNextParamType ParameterType;
diff --git a/Engine/Plugins/Experimental/AnimNext/Source/AnimNext/Private/Param/Params.cpp b/Engine/Plugins/Experimental/AnimNext/Source/AnimNext/Private/Param/Params.cpp
new file mode 100644
--- /dev/null
+++ b/Engine/Plugins/Experimental/AnimNext/Source/AnimNext/Private/Param/Params.cpp
@@ -0,0 +1,52 @@
+// Copyright Epic Games, Inc. All Rights Reserved.
+
+#include "Param/Params.h"
+
+#define LOCTEXT_NAMESPACE "AnimNextBuiltInParams"
+
+namespace UE::AnimNext
+{
+
+static TMap<FName, FParamDefinition> RegisteredParams;
+
+FParamDefinition FParams::RegisterBuiltInParameter(const FParamDefinition& InDefinition)
+{
+	return RegisteredParams.Add(InDefinition.Name, InDefinition);
+}
+
+void FParams::UnregisterBuiltInParameter(FName InName)
+{
+	RegisteredParams.Remove(InName);
+}
+
+bool FParams::IsBuiltInParameter(FName InName)
+{
+	return RegisteredParams.Contains(InName);
+}
+
+const FParamDefinition& FParams::GetBuiltInParameter(FName InName)
+{
+	return RegisteredParams.FindChecked(InName);
+}
+
+const FParamDefinition* FParams::FindBuiltInParameter(FName InName)
+{
+	return RegisteredParams.Find(InName);
+}
+
+int32 FParams::GetNumBuiltInParameters()
+{
+	return RegisteredParams.Num();
+}
+
+void FParams::ForEachBuiltInParameter(TFunctionRef<void(const FParamDefinition&)> InFunction)
+{
+	for (const TPair<FName, FParamDefinition>& RegisteredParamPair : RegisteredParams)
+	{
+		InFunction(RegisteredParamPair.Value);
+	}
+}
+
+}
+
+#undef LOCTEXT_NAMESPACE
\ No newline at end of file
diff --git a/Engine/Plugins/Experimental/AnimNext/Source/AnimNext/Private/Param/RigUnit_AnimNextParameterBeginExecution.cpp b/Engine/Plugins/Experimental/AnimNext/Source/AnimNext/Private/Param/RigUnit_AnimNextParameterBeginExecution.cpp
new file mode 100644
--- /dev/null
+++ b/Engine/Plugins/Experimental/AnimNext/Source/AnimNext/Private/Param/RigUnit_AnimNextParameterBeginExecution.cpp
@@ -0,0 +1,12 @@
+// Copyright Epic Games, Inc. All Rights Reserved.
+
+#include "Param/RigUnit_AnimNextParameterBeginExecution.h"
+#include "Units/RigUnitContext.h"
+
+#include UE_INLINE_GENERATED_CPP_BY_NAME(RigUnit_AnimNextParameterBeginExecution)
+
+FName FRigUnit_AnimNextParameterBeginExecution::EventName = TEXT("Update");
+
+FRigUnit_AnimNextParameterBeginExecution_Execute()
+{
+}
diff --git a/Engine/Plugins/Experimental/AnimNext/Source/AnimNext/Private/Param/RigVMDispatch_GetLayerParameter.cpp b/Engine/Plugins/Experimental/AnimNext/Source/AnimNext/Private/Param/RigVMDispatch_GetLayerParameter.cpp
new file mode 100644
--- /dev/null
+++ b/Engine/Plugins/Experimental/AnimNext/Source/AnimNext/Private/Param/RigVMDispatch_GetLayerParameter.cpp
@@ -0,0 +1,117 @@
+// Copyright Epic Games, Inc. All Rights Reserved.
+
+#include "Param/RigVMDispatch_GetLayerParameter.h"
+#include "RigVMCore/RigVMStruct.h"
+#include "RigVMCore/RigVM.h"
+#include "Graph/AnimNextExecuteContext.h"
+#include "Param/ParamStack.h"
+
+class UAnimNextParameter;
+
+const FName FRigVMDispatch_GetLayerParameter::ValueName = TEXT("Value");
+const FName FRigVMDispatch_GetLayerParameter::TypeHandleName = TEXT("Type");
+const FName FRigVMDispatch_GetLayerParameter::ParameterName = TEXT("Parameter");
+const FName FRigVMDispatch_GetLayerParameter::ParameterIdName = TEXT("ParameterId");
+
+FRigVMDispatch_GetLayerParameter::FRigVMDispatch_GetLayerParameter()
+{
+	FactoryScriptStruct = StaticStruct();
+}
+
+FName FRigVMDispatch_GetLayerParameter::GetArgumentNameForOperandIndex(int32 InOperandIndex, int32 InTotalOperands) const
+{
+	static const FName ArgumentNames[] =
+	{
+		ParameterName,
+		ValueName,
+		ParameterIdName,
+		TypeHandleName
+	};
+	check(InTotalOperands == UE_ARRAY_COUNT(ArgumentNames));
+	return ArgumentNames[InOperandIndex];
+}
+
+#if WITH_EDITOR
+FString FRigVMDispatch_GetLayerParameter::GetArgumentMetaData(const FName& InArgumentName, const FName& InMetaDataKey) const
+{
+	if ((InArgumentName == TypeHandleName || InArgumentName == ParameterIdName) &&
+		InMetaDataKey == FRigVMStruct::SingletonMetaName)
+	{
+		return TEXT("True");
+	}
+	else if(InArgumentName == ParameterName && InMetaDataKey == FRigVMStruct::CustomWidgetMetaName)
+	{
+		return TEXT("ParamName");
+	}
+
+	return Super::GetArgumentMetaData(InArgumentName, InMetaDataKey);
+}
+
+FString FRigVMDispatch_GetLayerParameter::GetArgumentDefaultValue(const FName& InArgumentName, TRigVMTypeIndex InTypeIndex) const
+{
+	if (InArgumentName == ParameterIdName)
+	{
+		return FString::Printf(TEXT("%u"), 0xffffffff);
+	}
+
+	return Super::GetArgumentDefaultValue(InArgumentName, InTypeIndex);
+}
+#endif
+
+const TArray<FRigVMTemplateArgument>& FRigVMDispatch_GetLayerParameter::GetArguments() const
+{
+	static const TArray<FRigVMTemplateArgument::ETypeCategory> ValueCategories =
+	{
+		FRigVMTemplateArgument::ETypeCategory_SingleAnyValue,
+		FRigVMTemplateArgument::ETypeCategory_ArrayAnyValue
+	};
+
+	static const TArray<FRigVMTemplateArgument> Arguments = {
+		FRigVMTemplateArgument(ParameterName, ERigVMPinDirection::Input, RigVMTypeUtils::TypeIndex::FName),
+		FRigVMTemplateArgument(ValueName, ERigVMPinDirection::Output, ValueCategories),
+		FRigVMTemplateArgument(ParameterIdName, ERigVMPinDirection::Hidden, RigVMTypeUtils::TypeIndex::UInt32),
+		FRigVMTemplateArgument(TypeHandleName, ERigVMPinDirection::Hidden, RigVMTypeUtils::TypeIndex::UInt32),
+	};
+
+	return Arguments;
+}
+
+FRigVMTemplateTypeMap FRigVMDispatch_GetLayerParameter::OnNewArgumentType(const FName& InArgumentName, TRigVMTypeIndex InTypeIndex) const
+{
+	FRigVMTemplateTypeMap Types;
+	Types.Add(ParameterName, RigVMTypeUtils::TypeIndex::FName);
+	Types.Add(ValueName, InTypeIndex);
+	Types.Add(ParameterIdName, RigVMTypeUtils::TypeIndex::UInt32);
+	Types.Add(TypeHandleName, RigVMTypeUtils::TypeIndex::UInt32);
+	return Types;
+}
+
+void FRigVMDispatch_GetLayerParameter::Execute(FRigVMExtendedExecuteContext& InContext, FRigVMMemoryHandleArray Handles, FRigVMPredicateBranchArray RigVMBranches)
+{
+	using namespace UE::AnimNext;
+
+	const FName Parameter = *(FName*)Handles[0].GetData();
+	const FProperty* ValueProperty = Handles[1].GetResolvedProperty();
+	check(ValueProperty);
+	uint8* TargetDataPtr = Handles[1].GetData();
+
+	uint32& ParameterId = *(uint32*)Handles[2].GetData();
+	if (ParameterId == FParamId::InvalidIndex)
+	{
+		ParameterId = FParamId(Parameter).ToInt();
+	}
+
+	uint32& TypeHandle = *(uint32*)Handles[3].GetData();
+	if (TypeHandle == 0)
+	{
+		TypeHandle = FParamTypeHandle::FromProperty(ValueProperty).ToRaw();
+	}
+
+	FAnimNextParameterExecuteContext& ParamContext = InContext.GetPublicData<FAnimNextParameterExecuteContext>();
+	TConstArrayView<uint8> SourceData;
+	if (ParamContext.GetLayerHandle().GetValueRaw(FParamId(ParameterId), FParamTypeHandle::FromRaw(TypeHandle), SourceData).IsSuccessful())
+	{
+		ValueProperty->CopyCompleteValue(TargetDataPtr, SourceData.GetData());
+	}
+}
+
diff --git a/Engine/Plugins/Experimental/AnimNext/Source/AnimNext/Private/Param/RigVMDispatch_GetParameter.cpp b/Engine/Plugins/Experimental/AnimNext/Source/AnimNext/Private/Param/RigVMDispatch_GetParameter.cpp
--- a/Engine/Plugins/Experimental/AnimNext/Source/AnimNext/Private/Param/RigVMDispatch_GetParameter.cpp
+++ b/Engine/Plugins/Experimental/AnimNext/Source/AnimNext/Private/Param/RigVMDispatch_GetParameter.cpp
@@ -41,6 +41,10 @@ FString FRigVMDispatch_GetParameter::GetArgumentMetaData(const FName& InArgument
 	{
 		return TEXT("True");
 	}
+	else if(InArgumentName == ParameterName && InMetaDataKey == FRigVMStruct::CustomWidgetMetaName)
+	{
+		return TEXT("ParamName");
+	}
 
 	return Super::GetArgumentMetaData(InArgumentName, InMetaDataKey);
 }
@@ -77,7 +81,10 @@ const TArray<FRigVMTemplateArgument>& FRigVMDispatch_GetParameter::GetArguments(
 FRigVMTemplateTypeMap FRigVMDispatch_GetParameter::OnNewArgumentType(const FName& InArgumentName, TRigVMTypeIndex InTypeIndex) const
 {
 	FRigVMTemplateTypeMap Types;
+	Types.Add(ParameterName, InTypeIndex);
 	Types.Add(ValueName, InTypeIndex);
+	Types.Add(ValueName, RigVMTypeUtils::TypeIndex::UInt32);
+	Types.Add(ValueName, RigVMTypeUtils::TypeIndex::UInt32);
 	return Types;
 }
 
@@ -88,7 +95,7 @@ void FRigVMDispatch_GetParameter::Execute(FRigVMExtendedExecuteContext& InContex
 	const FName Parameter = *(FName*)Handles[0].GetData();
 	const FProperty* ValueProperty = Handles[1].GetResolvedProperty();
 	check(ValueProperty);
-	uint8* TargetData = Handles[1].GetData();
+	uint8* TargetDataPtr = Handles[1].GetData();
 
 	uint32& ParameterId = *(uint32*)Handles[2].GetData();
 	if (ParameterId == FParamId::InvalidIndex)
@@ -103,9 +110,9 @@ void FRigVMDispatch_GetParameter::Execute(FRigVMExtendedExecuteContext& InContex
 	}
 
 	TConstArrayView<uint8> SourceData;
-	if (FParamStack::Get().GetParamData(FParamId(ParameterId), FParamTypeHandle::FromRaw(TypeHandle), SourceData) == FParamStack::EGetParamResult::Succeeded)
+	if (FParamStack::Get().GetParamData(FParamId(ParameterId), FParamTypeHandle::FromRaw(TypeHandle), SourceData).IsSuccessful())
 	{
-		ValueProperty->CopyCompleteValue(TargetData, SourceData.GetData());
+		ValueProperty->CopyCompleteValue(TargetDataPtr, SourceData.GetData());
 	}
 }
 
diff --git a/Engine/Plugins/Experimental/AnimNext/Source/AnimNext/Private/Param/RigVMDispatch_SetLayerParameter.cpp b/Engine/Plugins/Experimental/AnimNext/Source/AnimNext/Private/Param/RigVMDispatch_SetLayerParameter.cpp
new file mode 100644
--- /dev/null
+++ b/Engine/Plugins/Experimental/AnimNext/Source/AnimNext/Private/Param/RigVMDispatch_SetLayerParameter.cpp
@@ -0,0 +1,128 @@
+// Copyright Epic Games, Inc. All Rights Reserved.
+
+#include "Param/RigVMDispatch_SetLayerParameter.h"
+#include "RigVMCore/RigVMStruct.h"
+#include "RigVMCore/RigVM.h"
+#include "Param/AnimNextParameterExecuteContext.h"
+#include "Param/ParamId.h"
+#include "Param/ParamTypeHandle.h"
+#include "Param/ParamStackLayerHandle.h"
+
+class UAnimNextParameter;
+
+const FName FRigVMDispatch_SetLayerParameter::ExecuteContextName = TEXT("ExecuteContext");
+const FName FRigVMDispatch_SetLayerParameter::ValueName = TEXT("Value");
+const FName FRigVMDispatch_SetLayerParameter::TypeHandleName = TEXT("Type");
+const FName FRigVMDispatch_SetLayerParameter::ParameterName = TEXT("Parameter");
+const FName FRigVMDispatch_SetLayerParameter::ParameterIdName = TEXT("ParameterId");
+
+FRigVMDispatch_SetLayerParameter::FRigVMDispatch_SetLayerParameter()
+{
+	FactoryScriptStruct = StaticStruct();
+}
+
+FName FRigVMDispatch_SetLayerParameter::GetArgumentNameForOperandIndex(int32 InOperandIndex, int32 InTotalOperands) const
+{
+	static const FName ArgumentNames[] = 
+	{
+		ParameterName,
+		ValueName,
+		ParameterIdName,
+		TypeHandleName
+	};
+	check(InTotalOperands == UE_ARRAY_COUNT(ArgumentNames));
+	return ArgumentNames[InOperandIndex];
+}
+
+#if WITH_EDITOR
+FString FRigVMDispatch_SetLayerParameter::GetArgumentMetaData(const FName& InArgumentName, const FName& InMetaDataKey) const
+{
+	if ((InArgumentName == TypeHandleName || InArgumentName == ParameterIdName) &&
+		InMetaDataKey == FRigVMStruct::SingletonMetaName)
+	{
+		return TEXT("True");
+	}
+	else if(InArgumentName == ParameterName && InMetaDataKey == FRigVMStruct::CustomWidgetMetaName)
+	{
+		return TEXT("ParamName");
+	}
+
+	return Super::GetArgumentMetaData(InArgumentName, InMetaDataKey);
+}
+
+FString FRigVMDispatch_SetLayerParameter::GetArgumentDefaultValue(const FName& InArgumentName, TRigVMTypeIndex InTypeIndex) const
+{
+	if (InArgumentName == ParameterIdName)
+	{
+		return FString::Printf(TEXT("%u"), 0xffffffff);
+	}
+
+	return Super::GetArgumentDefaultValue(InArgumentName, InTypeIndex);
+}
+#endif
+
+const TArray<FRigVMTemplateArgument>& FRigVMDispatch_SetLayerParameter::GetArguments() const
+{
+	static const TArray<FRigVMTemplateArgument::ETypeCategory> ValueCategories =
+	{
+		FRigVMTemplateArgument::ETypeCategory_SingleAnyValue,
+		FRigVMTemplateArgument::ETypeCategory_ArrayAnyValue
+	};
+	
+	static const TArray<FRigVMTemplateArgument> Arguments = { 
+		FRigVMTemplateArgument(ParameterName, ERigVMPinDirection::Input, RigVMTypeUtils::TypeIndex::FName),
+		FRigVMTemplateArgument(ValueName, ERigVMPinDirection::Input, ValueCategories),
+		FRigVMTemplateArgument(ParameterIdName, ERigVMPinDirection::Hidden, RigVMTypeUtils::TypeIndex::UInt32),
+		FRigVMTemplateArgument(TypeHandleName, ERigVMPinDirection::Hidden, RigVMTypeUtils::TypeIndex::UInt32),
+	};
+
+	return Arguments;
+}
+
+TArray<FRigVMExecuteArgument>& FRigVMDispatch_SetLayerParameter::GetExecuteArguments_Impl(const FRigVMDispatchContext& InContext) const
+{
+	static TArray<FRigVMExecuteArgument> Arguments = { { ExecuteContextName, ERigVMPinDirection::IO } };
+	return Arguments;
+}
+
+FRigVMTemplateTypeMap FRigVMDispatch_SetLayerParameter::OnNewArgumentType(const FName& InArgumentName, TRigVMTypeIndex InTypeIndex) const
+{
+	FRigVMTemplateTypeMap Types;
+	Types.Add(ParameterName, RigVMTypeUtils::TypeIndex::FName);
+	Types.Add(ValueName, InTypeIndex);
+	Types.Add(ParameterIdName, RigVMTypeUtils::TypeIndex::UInt32);
+	Types.Add(TypeHandleName, RigVMTypeUtils::TypeIndex::UInt32);
+	return Types;
+}
+
+void FRigVMDispatch_SetLayerParameter::Execute(FRigVMExtendedExecuteContext& InContext, FRigVMMemoryHandleArray Handles, FRigVMPredicateBranchArray RigVMBranches)
+{
+	using namespace UE::AnimNext;
+
+	const FName Parameter = *(FName*)Handles[0].GetData();
+	const FProperty* ValueProperty = Handles[1].GetResolvedProperty();
+	check(ValueProperty);
+	const uint8* SourceData =  Handles[1].GetData();
+
+	uint32& RawParameterId = *(uint32*)Handles[2].GetData();
+	if (RawParameterId == FParamId::InvalidIndex && Parameter != NAME_None)
+	{
+		RawParameterId = FParamId(Parameter).ToInt();
+	}
+
+	uint32& RawTypeHandle = *(uint32*)Handles[3].GetData();
+	if (RawTypeHandle == 0)
+	{
+		RawTypeHandle = FParamTypeHandle::FromProperty(ValueProperty).ToRaw();
+	}
+
+	FParamId ParameterId(RawParameterId);
+	FParamTypeHandle TypeHandle = FParamTypeHandle::FromRaw(RawTypeHandle);
+	if(ParameterId.IsValid() && TypeHandle.IsValid())
+	{
+		FAnimNextParameterExecuteContext& ParamContext = InContext.GetPublicData<FAnimNextParameterExecuteContext>();
+		TConstArrayView<uint8> SourceDataView(SourceData, ValueProperty->GetSize());
+		ParamContext.GetLayerHandle().SetValueRaw(ParameterId, TypeHandle, SourceDataView);
+	}
+}
+
diff --git a/Engine/Plugins/Experimental/AnimNext/Source/AnimNext/Private/Scheduler/AnimNextSchedule.cpp b/Engine/Plugins/Experimental/AnimNext/Source/AnimNext/Private/Scheduler/AnimNextSchedule.cpp
new file mode 100644
--- /dev/null
+++ b/Engine/Plugins/Experimental/AnimNext/Source/AnimNext/Private/Scheduler/AnimNextSchedule.cpp
@@ -0,0 +1,187 @@
+// Copyright Epic Games, Inc. All Rights Reserved.
+
+#include "Scheduler/AnimNextSchedule.h"
+#include "Tasks/Task.h"
+#include "Async/TaskGraphInterfaces.h"
+#include "UObject/ObjectSaveContext.h"
+#include "EngineLogs.h"
+
+void UAnimNextSchedule::PostLoad()
+{
+	Super::PostLoad();
+
+#if WITH_EDITOR
+	CompileSchedule();
+#endif
+}
+
+#if WITH_EDITOR
+
+void UAnimNextSchedule::PostEditChangeProperty(FPropertyChangedEvent& PropertyChangedEvent)
+{
+	Super::PostEditChangeProperty(PropertyChangedEvent);
+
+	CompileSchedule();
+}
+
+void UAnimNextSchedule::CompileSchedule()
+{
+	Instructions.Empty();
+	Tasks.Empty();
+	Ports.Empty();
+	ExternalTasks.Empty();
+	ParamScopeEntryTasks.Empty();
+	ParamScopeExitTasks.Empty();
+	PortNameIndexMap.Empty();
+	NumParameterScopes = 0;
+	NumTickFunctions = 0;
+
+	EAnimNextScheduleScheduleOpcode LastOpCode = EAnimNextScheduleScheduleOpcode::None;
+
+	auto Emit = [this, &LastOpCode](EAnimNextScheduleScheduleOpcode InOpCode, int32 InOperand = 0)
+	{
+		FAnimNextScheduleInstruction Instruction;
+		Instruction.Opcode = InOpCode;
+		Instruction.Operand = InOperand;
+		Instructions.Add(Instruction);
+
+		LastOpCode = InOpCode;
+	};
+
+	auto EmitPrerequisite = [this, &Emit, &LastOpCode]()
+	{
+		switch (LastOpCode)
+		{
+		case EAnimNextScheduleScheduleOpcode::RunTask:
+			Emit(EAnimNextScheduleScheduleOpcode::PrerequisiteTask, NumTickFunctions - 1);
+			break;
+		case EAnimNextScheduleScheduleOpcode::BeginRunExternalTask:
+			Emit(EAnimNextScheduleScheduleOpcode::PrerequisiteBeginExternalTask, NumTickFunctions - 1);
+			break;
+		case EAnimNextScheduleScheduleOpcode::EndRunExternalTask:
+			Emit(EAnimNextScheduleScheduleOpcode::PrerequisiteEndExternalTask, NumTickFunctions - 1);
+			break;
+		case EAnimNextScheduleScheduleOpcode::RunParamScopeEntry:
+			Emit(EAnimNextScheduleScheduleOpcode::PrerequisiteScopeEntry, NumTickFunctions - 1);
+			break;
+		case EAnimNextScheduleScheduleOpcode::RunParamScopeExit:
+			Emit(EAnimNextScheduleScheduleOpcode::PrerequisiteScopeExit, NumTickFunctions - 1);
+			break;
+		default:
+			break;
+		}
+	};
+
+	// MAX_uint32 means 'global scope' in this context
+	uint32 ParentScopeIndex = MAX_uint32;
+
+	TFunction<void(const TArray<TObjectPtr<UAnimNextScheduleEntry>>&)> EmitEntries;
+
+	// Iterate over all entries, recursing into scopes
+	EmitEntries = [this, &EmitEntries, &Emit, &EmitPrerequisite, &ParentScopeIndex](const TArray<TObjectPtr<UAnimNextScheduleEntry>>& InEntries)
+	{
+		for (int32 EntryIndex = 0; EntryIndex < InEntries.Num(); ++EntryIndex)
+		{
+			UAnimNextScheduleEntry* Entry = InEntries[EntryIndex];
+
+			if (UAnimNextScheduleEntry_Port* PortEntry = Cast<UAnimNextScheduleEntry_Port>(Entry))
+			{
+				EmitPrerequisite();
+
+				FAnimNextSchedulePortTask PortTask;
+				PortTask.TaskIndex = Ports.Num();
+				PortTask.ParamScopeIndex = ParentScopeIndex;
+				PortTask.Name = PortEntry->Name;
+				int32 PortIndex = Ports.Add(PortTask);
+				PortNameIndexMap.Add(PortEntry->Name, PortIndex);
+
+				NumTickFunctions++;
+
+				Emit(EAnimNextScheduleScheduleOpcode::RunPort, PortIndex);
+			}
+			else if (UAnimNextScheduleEntry_AnimNextGraph* GraphEntry = Cast<UAnimNextScheduleEntry_AnimNextGraph>(Entry))
+			{
+				EmitPrerequisite();
+
+				FAnimNextScheduleGraphTask Task;
+				Task.TaskIndex = Tasks.Num();
+				Task.ParamScopeIndex = NumParameterScopes++;
+				Task.ParamParentScopeIndex = ParentScopeIndex;
+				Task.Name = GraphEntry->Name;
+				Task.EntryPoint = GraphEntry->EntryPoint;
+				Task.Graph = GraphEntry->Graph;
+				Task.ParameterBlocks = GraphEntry->ParameterBlocks;
+				int32 TaskIndex = Tasks.Add(Task);
+
+				NumTickFunctions++;
+
+				Emit(EAnimNextScheduleScheduleOpcode::RunTask, TaskIndex);
+			}
+			else if (UAnimNextScheduleEntry_ExternalTask* ExternalTaskEntry = Cast<UAnimNextScheduleEntry_ExternalTask>(Entry))
+			{
+				EmitPrerequisite();
+
+				FAnimNextScheduleExternalTask ExternalTask;
+				ExternalTask.TaskIndex = ExternalTasks.Num();
+				ExternalTask.ParamScopeIndex = NumParameterScopes++;
+				ExternalTask.ParamParentScopeIndex = ParentScopeIndex;
+				ExternalTask.Name = ExternalTaskEntry->Name;
+				ExternalTask.ObjectName = ExternalTaskEntry->ObjectName;
+				int32 ExternalTaskIndex = ExternalTasks.Add(ExternalTask);
+
+				// Emit the external task
+				Emit(EAnimNextScheduleScheduleOpcode::BeginRunExternalTask, ExternalTaskIndex);
+				NumTickFunctions++;
+
+				EmitPrerequisite();
+
+				Emit(EAnimNextScheduleScheduleOpcode::EndRunExternalTask, ExternalTaskIndex);
+				NumTickFunctions++;
+			}
+			else if (UAnimNextScheduleEntry_ParamScope* ParamScopeTaskEntry = Cast<UAnimNextScheduleEntry_ParamScope>(Entry))
+			{
+				EmitPrerequisite();
+
+				FAnimNextScheduleParamScopeEntryTask ParamScopeEntryTask;
+				ParamScopeEntryTask.TaskIndex = ParamScopeEntryTasks.Num();
+				const uint32 ParamScopeIndex = NumParameterScopes++;
+				ParamScopeEntryTask.ParamScopeIndex = ParamScopeIndex;
+				ParamScopeEntryTask.ParamParentScopeIndex = ParentScopeIndex;
+				ParamScopeEntryTask.TickFunctionIndex = NumTickFunctions;
+				ParamScopeEntryTask.Name = ParamScopeTaskEntry->Name;
+				ParamScopeEntryTask.ParameterBlocks = ParamScopeTaskEntry->ParameterBlocks;
+				int32 ParamScopeTaskEntryIndex = ParamScopeEntryTasks.Add(ParamScopeEntryTask);
+
+				Emit(EAnimNextScheduleScheduleOpcode::RunParamScopeEntry, ParamScopeTaskEntryIndex);
+				NumTickFunctions++;
+
+				// Enter new scope
+				uint32 PreviousParentScope = ParentScopeIndex;
+				ParentScopeIndex = ParamScopeIndex;
+
+				// Emit the subentries
+				EmitEntries(ParamScopeTaskEntry->SubEntries);
+
+				// Exit scope
+				ParentScopeIndex = PreviousParentScope;
+
+				EmitPrerequisite();
+
+				FAnimNextScheduleParamScopeExitTask ParamScopeExitTask;
+				ParamScopeExitTask.TaskIndex = ParamScopeExitTasks.Num();
+				ParamScopeExitTask.ParamScopeIndex = ParamScopeIndex;
+				ParamScopeExitTask.Name = ParamScopeTaskEntry->Name;
+				int32 ParamScopeExitTaskIndex = ParamScopeExitTasks.Add(ParamScopeExitTask);
+
+				Emit(EAnimNextScheduleScheduleOpcode::RunParamScopeExit, ParamScopeExitTaskIndex);
+				NumTickFunctions++;
+			}
+		}
+	};
+
+	EmitEntries(Entries);
+
+	Emit(EAnimNextScheduleScheduleOpcode::Exit);
+}
+
+#endif // #if WITH_EDITOR
\ No newline at end of file
diff --git a/Engine/Plugins/Experimental/AnimNext/Source/AnimNext/Private/Scheduler/AnimNextScheduleGraphTask.cpp b/Engine/Plugins/Experimental/AnimNext/Source/AnimNext/Private/Scheduler/AnimNextScheduleGraphTask.cpp
new file mode 100644
--- /dev/null
+++ b/Engine/Plugins/Experimental/AnimNext/Source/AnimNext/Private/Scheduler/AnimNextScheduleGraphTask.cpp
@@ -0,0 +1,30 @@
+// Copyright Epic Games, Inc. All Rights Reserved.
+
+#include "Scheduler/AnimNextScheduleGraphTask.h"
+#include "Scheduler/ScheduleContext.h"
+#include "Graph/AnimNextGraph.h"
+#include "Context.h"
+#include "Graph/AnimNextExecuteContext.h"
+
+void FAnimNextScheduleGraphTask::RunGraph(const UE::AnimNext::FScheduleContext& InScheduleContext) const
+{
+	using namespace UE::AnimNext;
+
+	if (Graph)
+	{
+		FScheduleInstanceData& InstanceData = InScheduleContext.GetInstanceData();
+
+		// Allocate our graph instance data
+		if (!InstanceData.GraphInstanceData[TaskIndex].IsValid())
+		{
+			 Graph->AllocateInstance(InstanceData.GraphInstanceData[TaskIndex]);
+		}
+
+		// FIXME: Cant run concurrently so disabling for now.
+		if(0)
+		{
+			const FContext Context;
+			Graph->Run(Context, InstanceData.GraphInstanceData[TaskIndex], EAnimNextGraphSimulationSteps::All);
+		}
+	}
+}
diff --git a/Engine/Plugins/Experimental/AnimNext/Source/AnimNext/Private/Scheduler/AnimNextScheduleParamScopeTask.cpp b/Engine/Plugins/Experimental/AnimNext/Source/AnimNext/Private/Scheduler/AnimNextScheduleParamScopeTask.cpp
new file mode 100644
--- /dev/null
+++ b/Engine/Plugins/Experimental/AnimNext/Source/AnimNext/Private/Scheduler/AnimNextScheduleParamScopeTask.cpp
@@ -0,0 +1,99 @@
+// Copyright Epic Games, Inc. All Rights Reserved.
+
+#include "Scheduler/AnimNextScheduleParamScopeTask.h"
+#include "Param/AnimNextParameterBlock.h"
+#include "Scheduler/ScheduleContext.h"
+#include "Scheduler/AnimNextSchedulerEntry.h"
+#include "Param/AnimNextParameterSourceRef.h"
+
+void FAnimNextScheduleParamScopeEntryTask::RunParamScopeEntry(const UE::AnimNext::FScheduleContext& InScheduleContext) const
+{
+	using namespace UE::AnimNext;
+
+	FParamStack& ParamStack = FParamStack::Get();
+
+	FScheduleInstanceData& InstanceData = InScheduleContext.GetInstanceData();
+	FScheduleInstanceData::FScopeCache& ScopeCache = InstanceData.ScopeCaches[TaskIndex];
+	ScopeCache.PushedLayers.Reset();
+
+	const FAnimNextParameterCollection* FoundUserScope = InstanceData.UserScopes.Find(Name);
+	if(FoundUserScope)
+	{
+		ScopeCache.PushedLayers.Reserve(ParameterBlocks.Num() + FoundUserScope->Parameters.Num());
+	}
+	else
+	{
+		ScopeCache.PushedLayers.Reserve(ParameterBlocks.Num());
+	}
+
+	// Resize layer handles appropriately & cache layer
+	if (ScopeCache.StaticHandles.Num() != ParameterBlocks.Num())
+	{
+		ScopeCache.StaticHandles.SetNum(ParameterBlocks.Num());
+		for (int32 LayerHandleIndex = 0; LayerHandleIndex < ScopeCache.StaticHandles.Num(); ++LayerHandleIndex)
+		{
+			ScopeCache.StaticHandles[LayerHandleIndex] = ParameterBlocks[LayerHandleIndex]->CacheLayer();
+		}
+	}
+
+	if (FoundUserScope)
+	{
+		const UObject* ObjectContext = InstanceData.Entry->Object.Get();
+
+		// Resize user scope data appropriately
+		if (ScopeCache.UserHandles.Num() != FoundUserScope->Parameters.Num())
+		{
+			ScopeCache.UserHandles.SetNum(FoundUserScope->Parameters.Num());
+		}
+
+		// Cache any layers if required
+		for (int32 LayerHandleIndex = 0; LayerHandleIndex < ScopeCache.UserHandles.Num(); ++LayerHandleIndex)
+		{
+			if(const IAnimNextParameterSourceInterface* ParameterSource = FoundUserScope->Parameters[LayerHandleIndex].Get(ObjectContext))
+			{
+				if(ParameterSource->ShouldCacheLayer(ScopeCache.UserHandles[LayerHandleIndex]))
+				{
+					ScopeCache.UserHandles[LayerHandleIndex] = ParameterSource->CacheLayer();
+				}
+			}
+		}
+
+		for (int32 ParamBlockIndex = 0; ParamBlockIndex < FoundUserScope->Parameters.Num(); ++ParamBlockIndex)
+		{
+			if(const IAnimNextParameterSourceInterface* ParameterSource = FoundUserScope->Parameters[ParamBlockIndex].Get(ObjectContext))
+			{
+				FParamStackLayerHandle& LayerHandle = ScopeCache.UserHandles[ParamBlockIndex];
+				ParameterSource->UpdateLayer(LayerHandle);
+				ScopeCache.PushedLayers.Add(ParamStack.PushLayer(LayerHandle));
+			}
+		}
+	}
+
+	// TODO: Pre/post scope support
+
+	// Update & push static params
+	for (int32 ParamBlockIndex = 0; ParamBlockIndex < ParameterBlocks.Num(); ++ParamBlockIndex)
+	{
+		const UAnimNextParameterBlock* ParameterBlock = ParameterBlocks[ParamBlockIndex];
+		FParamStackLayerHandle& LayerHandle = ScopeCache.StaticHandles[ParamBlockIndex];
+
+		ParameterBlock->UpdateLayer(LayerHandle);
+		ScopeCache.PushedLayers.Add(ParamStack.PushLayer(LayerHandle));
+	}
+}
+
+void FAnimNextScheduleParamScopeExitTask::RunParamScopeExit(const UE::AnimNext::FScheduleContext& InScheduleContext) const
+{
+	using namespace UE::AnimNext;
+
+	FParamStack& ParamStack = FParamStack::Get();
+
+	FScheduleInstanceData& InstanceData = InScheduleContext.GetInstanceData();
+	FScheduleInstanceData::FScopeCache& ScopeCache = InstanceData.ScopeCaches[TaskIndex];
+	for (int32 LayerIndex = ScopeCache.PushedLayers.Num() - 1; LayerIndex >= 0; --LayerIndex)
+	{
+		ParamStack.PopLayer(ScopeCache.PushedLayers[LayerIndex]);
+	}
+
+	ScopeCache.PushedLayers.Reset();
+}
\ No newline at end of file
diff --git a/Engine/Plugins/Experimental/AnimNext/Source/AnimNext/Private/Scheduler/AnimNextSchedulePort.h b/Engine/Plugins/Experimental/AnimNext/Source/AnimNext/Private/Scheduler/AnimNextSchedulePort.h
new file mode 100644
--- /dev/null
+++ b/Engine/Plugins/Experimental/AnimNext/Source/AnimNext/Private/Scheduler/AnimNextSchedulePort.h
@@ -0,0 +1,31 @@
+// Copyright Epic Games, Inc. All Rights Reserved.
+
+#pragma once
+
+#include "CoreMinimal.h"
+#include "AnimNextSchedulePort.generated.h"
+
+struct FAnimNextSchedulePortTask;
+
+namespace UE::AnimNext
+{
+	struct FScheduleContext;
+}
+
+USTRUCT()
+struct FAnimNextSchedulePort
+{
+	GENERATED_BODY()
+
+private:
+	friend struct FAnimNextSchedulePortTask;
+	friend class UAnimNextComponent;
+
+	// Object the port is bound to
+	UPROPERTY()
+	TWeakObjectPtr<UObject> Object;
+
+	// Static data the port will write to
+	// TODO: This should probably be a TUniqueFunction<TArrayView<uint8>(void)>
+	TArrayView<uint8> Data;
+};
diff --git a/Engine/Plugins/Experimental/AnimNext/Source/AnimNext/Private/Scheduler/AnimNextSchedulePortTask.cpp b/Engine/Plugins/Experimental/AnimNext/Source/AnimNext/Private/Scheduler/AnimNextSchedulePortTask.cpp
new file mode 100644
--- /dev/null
+++ b/Engine/Plugins/Experimental/AnimNext/Source/AnimNext/Private/Scheduler/AnimNextSchedulePortTask.cpp
@@ -0,0 +1,27 @@
+// Copyright Epic Games, Inc. All Rights Reserved.
+
+#include "Scheduler/AnimNextSchedulePortTask.h"
+#include "Scheduler/AnimNextSchedulePort.h"
+#include "Scheduler/SchedulePortDefinition.h"
+#include "Scheduler/Scheduler.h"
+#include "Scheduler/ScheduleContext.h"
+#include "Scheduler/SchedulePortAdapterContext.h"
+
+void FAnimNextSchedulePortTask::RunPort(const UE::AnimNext::FScheduleContext& InScheduleContext) const
+{
+	using namespace UE::AnimNext;
+
+	if (const FSchedulePortDefinition* PortDefinition = FScheduler::FindPortDefinition(Name))
+	{
+		if (PortDefinition->Adapter)
+		{
+			if (const FAnimNextSchedulePort* PortParam = FParamStack::Get().GetParamPtr<FAnimNextSchedulePort>(Name))
+			{
+				FSchedulePortAdapterContext Context;
+				Context.OutputObject = PortParam->Object.Get();
+				Context.OutputData = PortParam->Data;
+				PortDefinition->Adapter(Context);
+			}
+		}
+	}
+}
\ No newline at end of file
diff --git a/Engine/Plugins/Experimental/AnimNext/Source/AnimNext/Private/Scheduler/AnimNextSchedulerEntry.cpp b/Engine/Plugins/Experimental/AnimNext/Source/AnimNext/Private/Scheduler/AnimNextSchedulerEntry.cpp
new file mode 100644
--- /dev/null
+++ b/Engine/Plugins/Experimental/AnimNext/Source/AnimNext/Private/Scheduler/AnimNextSchedulerEntry.cpp
@@ -0,0 +1,186 @@
+// Copyright Epic Games, Inc. All Rights Reserved.
+
+#include "AnimNextSchedulerEntry.h"
+#include "ScheduleInstanceData.h"
+#include "Scheduler/AnimNextSchedule.h"
+#include "Engine/World.h"
+
+FAnimNextSchedulerEntry::FAnimNextSchedulerEntry(const UAnimNextSchedule* InSchedule, UObject* InObject, UE::AnimNext::FScheduleHandle InHandle, const TMap<FName, FAnimNextParameterCollection>& InUserScopes)
+	: Schedule(InSchedule)
+	, UserScopes(InUserScopes)
+	, Object(InObject)
+	, Handle(InHandle)
+	, Context(InSchedule, this)
+{
+	using namespace UE::AnimNext;
+
+	check(InObject)
+	check(InSchedule);
+	check(InHandle.IsValid());
+
+	RootParamStack = MakeShared<FParamStack>();
+
+	FParamStack::AttachToCurrentThread(RootParamStack);
+	FScheduleContext::AttachToCurrentThread(Context);
+
+	// Setup tick function graph
+	if (InSchedule->Instructions.Num() > 0)
+	{
+		BeginTickFunction = MakeUnique<FScheduleBeginTickFunction>(*this);
+		EndTickFunction = MakeUnique<FScheduleEndTickFunction>(*this);
+
+		TArray<TUniquePtr<FScheduleTickFunction>*> Prerequisites;
+		int32 InstructionIndex = 0;
+		while (InstructionIndex < InSchedule->Instructions.Num())
+		{
+			auto AddTickFunctionAndPrerequisites = [this, &Prerequisites, InObject](TUniquePtr<FScheduleTickFunction>&& InTickFunction)
+			{
+				// Add a prereq on the previous tick function we added, if any (we only support linear chains right now so this is OK)
+				if (TickFunctions.Num() == 0)
+				{
+					BeginTickFunction->Subsequent = FTickPrerequisite(InObject, *InTickFunction.Get());
+					InTickFunction->AddPrerequisite(InObject, *BeginTickFunction.Get());
+				}
+
+				for (TUniquePtr<FScheduleTickFunction>* Prerequisite : Prerequisites)
+				{
+					Prerequisite->Get()->Subsequents.Emplace(InObject, *InTickFunction.Get());
+					InTickFunction->AddPrerequisite(InObject, *Prerequisite->Get());
+				}
+				Prerequisites.Reset();
+				TickFunctions.Add(MoveTemp(InTickFunction));
+			};
+
+			const FAnimNextScheduleInstruction& Instruction = InSchedule->Instructions[InstructionIndex++];
+			switch (Instruction.Opcode)
+			{
+			case EAnimNextScheduleScheduleOpcode::BeginRunExternalTask:
+				{
+					AddTickFunctionAndPrerequisites(MakeUnique<FScheduleTickFunction>(Context, Instruction));
+
+					const FName TickFunctionName = InSchedule->ExternalTasks[Instruction.Operand].Name;
+					const FName ObjectName = InSchedule->ExternalTasks[Instruction.Operand].ObjectName;
+					if(TickFunctionName != NAME_None && ObjectName != NAME_None)
+					{
+						FTickFunction* FoundFunction = FParamStack::Get().GetMutableParamPtr<FTickFunction>(TickFunctionName);
+						const UObject* FoundObject = FParamStack::Get().GetParamPtr<UObject>(ObjectName);
+						if(FoundFunction && FoundObject)
+						{
+							FoundFunction->AddPrerequisite(InObject, *TickFunctions.Last().Get());
+							TickFunctions.Last()->TargetObject = const_cast<UObject*>(FoundObject);
+							TickFunctions.Last()->Subsequents.Emplace(const_cast<UObject*>(FoundObject), *FoundFunction);
+						}
+					}
+					break;
+				}
+			case EAnimNextScheduleScheduleOpcode::EndRunExternalTask:
+				{
+					AddTickFunctionAndPrerequisites(MakeUnique<FScheduleTickFunction>(Context, Instruction));
+
+					const FName TickFunctionName = InSchedule->ExternalTasks[Instruction.Operand].Name;
+					const FName ObjectName = InSchedule->ExternalTasks[Instruction.Operand].ObjectName;
+					if(TickFunctionName != NAME_None && ObjectName != NAME_None)
+					{
+						FTickFunction* FoundFunction = FParamStack::Get().GetMutableParamPtr<FTickFunction>(TickFunctionName);
+						const UObject* FoundObject = FParamStack::Get().GetParamPtr<UObject>(ObjectName);
+						if(FoundFunction && FoundObject)
+						{
+							TickFunctions.Last()->TargetObject = const_cast<UObject*>(FoundObject);
+							TickFunctions.Last()->AddPrerequisite(const_cast<UObject*>(FoundObject), *FoundFunction);
+						}
+					}
+					break;
+				}
+			case EAnimNextScheduleScheduleOpcode::RunTask:
+			case EAnimNextScheduleScheduleOpcode::RunPort:
+			case EAnimNextScheduleScheduleOpcode::RunParamScopeEntry:
+			case EAnimNextScheduleScheduleOpcode::RunParamScopeExit:
+				AddTickFunctionAndPrerequisites(MakeUnique<FScheduleTickFunction>(Context, Instruction));
+				break;
+			case EAnimNextScheduleScheduleOpcode::PrerequisiteTask:
+			case EAnimNextScheduleScheduleOpcode::PrerequisiteBeginExternalTask:
+			case EAnimNextScheduleScheduleOpcode::PrerequisiteEndExternalTask:
+			case EAnimNextScheduleScheduleOpcode::PrerequisiteScopeEntry:
+			case EAnimNextScheduleScheduleOpcode::PrerequisiteScopeExit:
+				check(TickFunctions[Instruction.Operand].IsValid());
+				Prerequisites.Add(&TickFunctions[Instruction.Operand]);
+				break;
+			case EAnimNextScheduleScheduleOpcode::Exit:
+				check(Prerequisites.Num() == 0);
+				EndTickFunction->AddPrerequisite(InObject, *TickFunctions.Last().Get());
+				TickFunctions.Last()->Subsequents.Emplace(InObject, *EndTickFunction.Get());
+				break;
+			default:
+				check(false);
+				break;
+			}
+		}
+
+		// Register our tick functions
+		ULevel* Level = InObject->GetWorld()->PersistentLevel;
+		BeginTickFunction->RegisterTickFunction(Level);
+		EndTickFunction->RegisterTickFunction(Level);
+		for (TUniquePtr<FScheduleTickFunction>& TickFunction : TickFunctions)
+		{
+			TickFunction->RegisterTickFunction(Level);
+		}
+	}
+
+	FScheduleContext::DetachFromCurrentThread();
+	FParamStack::DetachFromCurrentThread();
+}
+
+FAnimNextSchedulerEntry::~FAnimNextSchedulerEntry()
+{
+	Invalidate();
+}
+
+void FAnimNextSchedulerEntry::Invalidate()
+{
+	using namespace UE::AnimNext;
+
+	if(BeginTickFunction)
+	{
+		BeginTickFunction->Subsequent.PrerequisiteTickFunction->RemovePrerequisite(Object.Get(), *BeginTickFunction.Get());
+	}
+	
+	for (TUniquePtr<FScheduleTickFunction>& TickFunction : TickFunctions)
+	{
+		for (FTickPrerequisite& Subsequent : TickFunction->Subsequents)
+		{
+			Subsequent.PrerequisiteTickFunction->RemovePrerequisite(Object.Get(), *TickFunction.Get());
+		}
+		TickFunction->UnRegisterTickFunction();
+	}
+
+	if (BeginTickFunction)
+	{
+		BeginTickFunction->UnRegisterTickFunction();
+	}
+	if(EndTickFunction)
+	{
+		EndTickFunction->UnRegisterTickFunction();
+	}
+
+	BeginTickFunction.Reset();
+	EndTickFunction.Reset();
+	TickFunctions.Reset();
+
+	Schedule = nullptr;
+	Object = nullptr;
+	UserScopes.Empty();
+	Handle.Invalidate();
+	Context.InstanceData.Reset();
+}
+
+void FAnimNextSchedulerEntry::Enable(bool bInEnabled)
+{
+	using namespace UE::AnimNext;
+	
+	BeginTickFunction->SetTickFunctionEnable(bInEnabled);
+	for (TUniquePtr<FScheduleTickFunction>& TickFunction : TickFunctions)
+	{
+		TickFunction->SetTickFunctionEnable(bInEnabled);
+	}
+	EndTickFunction->SetTickFunctionEnable(bInEnabled);
+}
\ No newline at end of file
diff --git a/Engine/Plugins/Experimental/AnimNext/Source/AnimNext/Private/Scheduler/AnimNextSchedulerEntry.h b/Engine/Plugins/Experimental/AnimNext/Source/AnimNext/Private/Scheduler/AnimNextSchedulerEntry.h
new file mode 100644
--- /dev/null
+++ b/Engine/Plugins/Experimental/AnimNext/Source/AnimNext/Private/Scheduler/AnimNextSchedulerEntry.h
@@ -0,0 +1,69 @@
+// Copyright Epic Games, Inc. All Rights Reserved.
+
+#pragma once
+
+#include "CoreMinimal.h"
+#include "Scheduler/ScheduleTickFunction.h"
+#include "Scheduler/ScheduleContext.h"
+#include "AnimNextSchedulerEntry.generated.h"
+
+class UAnimNextSchedule;
+
+namespace UE::AnimNext
+{
+	struct FScheduleHandle;
+	struct FParamStack;
+}
+
+// Root memory owner of a parameterized schedule 
+USTRUCT()
+struct FAnimNextSchedulerEntry
+{
+	GENERATED_BODY()
+
+	FAnimNextSchedulerEntry() = default;
+	FAnimNextSchedulerEntry(const UAnimNextSchedule* InSchedule, UObject* InObject, UE::AnimNext::FScheduleHandle InHandle, const TMap<FName, FAnimNextParameterCollection>& InUserScopes);
+	~FAnimNextSchedulerEntry();
+
+	// Used for pooling
+	void Invalidate();
+
+	// Enables/disables the ticking of this entry
+	void Enable(bool bInEnabled);
+	
+	UPROPERTY()
+	TObjectPtr<const UAnimNextSchedule> Schedule = nullptr;
+
+	// User scopes are copied into this entry on construction, but moved out later into instance data
+	// So will be invalid here after first run
+	UPROPERTY()
+	TMap<FName, FAnimNextParameterCollection> UserScopes;
+
+	// Object this entry is bound to
+	TWeakObjectPtr<UObject> Object;
+
+	// Copy of the handle that represents this entry to client systems
+	UE::AnimNext::FScheduleHandle Handle;
+
+	// Root context, passed to all the schedule's tasks
+	UE::AnimNext::FScheduleContext Context;
+
+	// Root param stack for the schedule itself (and globals)
+	TSharedPtr<UE::AnimNext::FParamStack> RootParamStack;
+	
+	// Begin/end tick functions used to wrap the schedule's tick function graph
+	TUniquePtr<UE::AnimNext::FScheduleBeginTickFunction> BeginTickFunction;
+	TUniquePtr<UE::AnimNext::FScheduleEndTickFunction> EndTickFunction;
+
+	// Pre-allocated graph of tick functions
+	TArray<TUniquePtr<UE::AnimNext::FScheduleTickFunction>> TickFunctions;
+};
+
+template<>
+struct TStructOpsTypeTraits<FAnimNextSchedulerEntry> : public TStructOpsTypeTraitsBase2<FAnimNextSchedulerEntry>
+{
+	enum
+	{
+		WithCopy = false
+	};
+};
diff --git a/Engine/Plugins/Experimental/AnimNext/Source/AnimNext/Private/Scheduler/AnimNextSchedulerWorldSubsystem.cpp b/Engine/Plugins/Experimental/AnimNext/Source/AnimNext/Private/Scheduler/AnimNextSchedulerWorldSubsystem.cpp
new file mode 100644
--- /dev/null
+++ b/Engine/Plugins/Experimental/AnimNext/Source/AnimNext/Private/Scheduler/AnimNextSchedulerWorldSubsystem.cpp
@@ -0,0 +1,200 @@
+// Copyright Epic Games, Inc. All Rights Reserved.
+
+#include "AnimNextSchedulerWorldSubsystem.h"
+#include "ScheduleInstanceData.h"
+#include "Scheduler/AnimNextSchedule.h"
+#include "Param/ParamDefinition.h"
+#include "Param/Params.h"
+#include "Engine/World.h"
+
+namespace UE::AnimNext
+{
+static uint32 GEntrySerialNumber = 0;
+}
+
+namespace UE::AnimNext
+{
+
+FSchedulePendingAction::FSchedulePendingAction(EType InType, FScheduleHandle InHandle)
+	: Handle(InHandle)
+	, Type(InType)
+{
+}
+
+}
+
+bool UAnimNextSchedulerWorldSubsystem::DoesSupportWorldType(const EWorldType::Type WorldType) const
+{
+	switch (WorldType)
+	{
+	case EWorldType::Game:
+	case EWorldType::Editor:
+	case EWorldType::PIE:
+	case EWorldType::EditorPreview:
+	case EWorldType::GamePreview:
+		return true;
+	}
+
+	return false;
+}
+
+bool UAnimNextSchedulerWorldSubsystem::IsValidHandle(UE::AnimNext::FScheduleHandle InHandle) const
+{
+	return Entries.IsValidIndex(InHandle.Index) && InHandle.SerialNumber == Entries[InHandle.Index]->Handle.SerialNumber;
+}
+
+void UAnimNextSchedulerWorldSubsystem::FlushPendingActions()
+{
+	using namespace UE::AnimNext;
+
+	FRWScopeLock PendingLockScope(PendingLock, SLT_Write);
+
+	if (PendingActions.Num() > 0)
+	{
+		FRWScopeLock EntriesLockScope(EntriesLock, SLT_Write);
+
+		for (FSchedulePendingAction& PendingAction : PendingActions)
+		{
+			switch (PendingAction.Type)
+			{
+			case FSchedulePendingAction::EType::ReleaseHandle:
+				if (IsValidHandle(PendingAction.Handle))
+				{
+					if (PendingAction.Handle.Index == Entries.Num() - 1)
+					{
+						// Last entry, shrink array
+						Entries.Pop(false);
+					}
+					else
+					{
+						// Not last entry, invalidate and add to free list
+						TUniquePtr<FAnimNextSchedulerEntry>& Entry = Entries[PendingAction.Handle.Index];
+						Entry->Invalidate();
+						FreeEntryIndices.Add(PendingAction.Handle.Index);
+					}
+				}
+				break;
+			case FSchedulePendingAction::EType::EnableHandle:
+				if (IsValidHandle(PendingAction.Handle))
+				{
+					TUniquePtr<FAnimNextSchedulerEntry>& Entry = Entries[PendingAction.Handle.Index];
+					Entry->Enable(true);
+				}
+				break;
+			case FSchedulePendingAction::EType::DisableHandle:
+				if (IsValidHandle(PendingAction.Handle))
+				{
+					TUniquePtr<FAnimNextSchedulerEntry>& Entry = Entries[PendingAction.Handle.Index];
+					Entry->Enable(false);
+				}
+				break;
+			default:
+				break;
+			}
+		}
+
+		PendingActions.Reset();
+	}
+}
+
+UE::AnimNext::FScheduleHandle UAnimNextSchedulerWorldSubsystem::AcquireHandle(UObject* InObject, UAnimNextSchedule* InSchedule, const TMap<FName, FAnimNextParameterCollection>& InUserScopes)
+{
+	using namespace UE::AnimNext;
+
+	FRWScopeLock EntriesLockScope(EntriesLock, SLT_Write);
+
+	FScheduleHandle Handle;
+
+	// Check free list first
+	if (FreeEntryIndices.Num())
+	{
+		Handle.Index = FreeEntryIndices.Last();
+		Handle.SerialNumber = ++GEntrySerialNumber;
+		FreeEntryIndices.Pop(false);
+		new (Entries[Handle.Index].Get()) FAnimNextSchedulerEntry(InSchedule, InObject, Handle, InUserScopes);
+	}
+	// Otherwise append a new entry
+	else
+	{
+		Handle.Index = Entries.Num();
+		Handle.SerialNumber = ++GEntrySerialNumber;
+		Entries.Emplace(MakeUnique<FAnimNextSchedulerEntry>(InSchedule, InObject, Handle, InUserScopes));
+	}
+
+	// Skip 'invalid' 0 serial number
+	if (GEntrySerialNumber == 0)
+	{
+		++GEntrySerialNumber;
+	}
+
+	return Handle;
+}
+
+void UAnimNextSchedulerWorldSubsystem::ReleaseHandle(UE::AnimNext::FScheduleHandle InHandle)
+{
+	using namespace UE::AnimNext;
+
+	if (IsValidHandle(InHandle))
+	{
+		FRWScopeLock EntriesLockScope(EntriesLock, SLT_Write);
+
+		if (InHandle.Index == Entries.Num() - 1)
+		{
+			// Last entry, shrink array
+			Entries.Pop(false);
+		}
+		else
+		{
+			// Not last entry, invalidate and add to free list
+			TUniquePtr<FAnimNextSchedulerEntry>& Entry = Entries[InHandle.Index];
+			Entry->Invalidate();
+			FreeEntryIndices.Add(InHandle.Index);
+		}
+
+		// TODO: do not allow immediate handle releases outside of schedule runs - we can either defer or assert
+	//	PendingActions.Emplace(FAnimNextSchedulePendingAction::EType::ReleaseHandle, InHandle);
+	}
+}
+
+void UAnimNextSchedulerWorldSubsystem::EnableHandle(UE::AnimNext::FScheduleHandle InHandle, bool bInEnabled)
+{
+	using namespace UE::AnimNext;
+	
+	if (IsValidHandle(InHandle))
+	{
+		PendingActions.Emplace(bInEnabled ? FSchedulePendingAction::EType::EnableHandle : FSchedulePendingAction::EType::DisableHandle, InHandle);
+	}
+}
+
+void UAnimNextSchedulerWorldSubsystem::QueueTask(UE::AnimNext::FScheduleHandle InHandle, FName InScheduleTaskName, TUniqueFunction<void(const UE::AnimNext::FScheduleContext&)>&& InTaskFunction, UE::AnimNext::FScheduler::ETaskRunLocation InLocation)
+{
+	using namespace UE::AnimNext;
+
+	if (IsValidHandle(InHandle))
+	{
+		TUniquePtr<FAnimNextSchedulerEntry>& Entry = Entries[InHandle.Index];
+
+		// TODO: Only supporting scope tasks for now
+		const FAnimNextScheduleParamScopeEntryTask* FoundScope = Entry->Schedule->ParamScopeEntryTasks.FindByPredicate([&InScheduleTaskName](const FAnimNextScheduleParamScopeEntryTask& InTask)
+		{
+			return InTask.Name == InScheduleTaskName;
+		});
+
+		if (FoundScope)
+		{
+			switch (InLocation)
+			{
+			case FScheduler::ETaskRunLocation::Before:
+				Entry->TickFunctions[FoundScope->TickFunctionIndex]->PreExecuteTasks.Enqueue(MoveTemp(InTaskFunction));
+				break;
+			case FScheduler::ETaskRunLocation::After:
+				Entry->TickFunctions[FoundScope->TickFunctionIndex]->PostExecuteTasks.Enqueue(MoveTemp(InTaskFunction));
+				break;
+			}
+		}
+		else
+		{
+			UE_LOG(LogAnimation, Warning, TEXT("QueueTask: Could not find scope '%s' in schedule '%s'"), *InScheduleTaskName.ToString(), *Entry->Schedule.GetName());
+		}
+	}
+}
\ No newline at end of file
diff --git a/Engine/Plugins/Experimental/AnimNext/Source/AnimNext/Private/Scheduler/AnimNextSchedulerWorldSubsystem.h b/Engine/Plugins/Experimental/AnimNext/Source/AnimNext/Private/Scheduler/AnimNextSchedulerWorldSubsystem.h
new file mode 100644
--- /dev/null
+++ b/Engine/Plugins/Experimental/AnimNext/Source/AnimNext/Private/Scheduler/AnimNextSchedulerWorldSubsystem.h
@@ -0,0 +1,99 @@
+// Copyright Epic Games, Inc. All Rights Reserved.
+
+#pragma once
+
+#include "CoreMinimal.h"
+#include "Subsystems/WorldSubsystem.h"
+#include "ScheduleHandle.h"
+#include "ScheduleContext.h"
+#include "Param/ParamStack.h"
+#include "Param/ParamStackLayerHandle.h"
+#include "Tasks/Task.h"
+#include "Scheduler/ScheduleTickFunction.h"
+#include "Scheduler/AnimNextSchedulerEntry.h"
+#include "Scheduler/Scheduler.h"
+#include "AnimNextSchedulerWorldSubsystem.generated.h"
+
+class UAnimNextSchedule;
+class UAnimNextParameterBlock;
+
+namespace UE::AnimNext
+{
+	struct FScheduleHandle;
+	struct FParamStackLayerHandle;
+}
+
+namespace UE::AnimNext
+{
+
+// A queued action to complete next frame
+struct FSchedulePendingAction
+{
+	enum class EType : int32
+	{
+		None = 0,
+		ReleaseHandle,
+		EnableHandle,
+		DisableHandle,
+	};
+
+	FSchedulePendingAction() = default;
+
+	FSchedulePendingAction(EType InType, FScheduleHandle InHandle);
+
+	FScheduleHandle Handle;
+
+	EType Type = EType::None;
+};
+
+}
+
+// Represents a scheduler interface to a UWorld
+UCLASS()
+class UAnimNextSchedulerWorldSubsystem : public UWorldSubsystem
+{
+	GENERATED_BODY()
+
+	friend struct UE::AnimNext::FScheduler;
+
+	// UWorldSubsystem interface
+	virtual bool DoesSupportWorldType(const EWorldType::Type WorldType) const override;
+
+	bool IsValidHandle(UE::AnimNext::FScheduleHandle InHandle) const;
+
+	void FlushPendingActions();
+
+	// Acquire a handle that binds a schedule with the supplied parameters
+	UE::AnimNext::FScheduleHandle AcquireHandle(UObject* InObject, UAnimNextSchedule* InSchedule, const TMap<FName, FAnimNextParameterCollection>& InUserScopes);
+
+	// Release an already acquired handle
+	// The full release of the binding referenced by the handle will be deferred after this call is made
+	void ReleaseHandle(UE::AnimNext::FScheduleHandle InHandle);
+
+	// Enables or disables the schedule parameterization represented by the supplied handle
+	// This operation is deferred until the next time the schedule ticks
+	void EnableHandle(UE::AnimNext::FScheduleHandle InHandle, bool bInEnabled);
+
+	// Queue a task to run at a particular point in a schedule
+	// @param	InHandle		The handle to queue the task to
+	// @param	InTaskName		The name of the task in the schedule to run the supplied task relative to
+	// @param	InTaskFunction	The function to run
+	// @param	InLocation		Where to run the task, before or after
+	void QueueTask(UE::AnimNext::FScheduleHandle InHandle, FName InScheduleTaskName, TUniqueFunction<void(const UE::AnimNext::FScheduleContext&)>&& InTaskFunction, UE::AnimNext::FScheduler::ETaskRunLocation InLocation);
+
+	// Currently running entries, pooled
+	TArray<TUniquePtr<FAnimNextSchedulerEntry>> Entries;
+
+	// Free entries for the entries pool
+	TArray<uint32> FreeEntryIndices;
+	
+	// Queued actions
+	TArray<UE::AnimNext::FSchedulePendingAction> PendingActions;
+
+	// Locks for concurrent modifications
+	FRWLock EntriesLock;
+	FRWLock PendingLock;
+
+	// Cached delta time
+	float DeltaTime;
+};
\ No newline at end of file
diff --git a/Engine/Plugins/Experimental/AnimNext/Source/AnimNext/Private/Scheduler/ScheduleContext.cpp b/Engine/Plugins/Experimental/AnimNext/Source/AnimNext/Private/Scheduler/ScheduleContext.cpp
new file mode 100644
--- /dev/null
+++ b/Engine/Plugins/Experimental/AnimNext/Source/AnimNext/Private/Scheduler/ScheduleContext.cpp
@@ -0,0 +1,30 @@
+// Copyright Epic Games, Inc. All Rights Reserved.
+
+#include "ScheduleContext.h"
+#include "Scheduler/AnimNextSchedule.h"
+
+namespace UE::AnimNext
+{
+
+// The currently attached schedule context for this thread
+static thread_local const FScheduleContext* GScheduleContext = nullptr;
+
+void FScheduleContext::AttachToCurrentThread(const FScheduleContext& InContext)
+{
+	check(GScheduleContext == nullptr);
+	GScheduleContext = &InContext;
+}
+
+void FScheduleContext::DetachFromCurrentThread()
+{
+	check(GScheduleContext != nullptr);
+	GScheduleContext = nullptr;
+}
+
+const FScheduleContext& FScheduleContext::Get()
+{
+	check(GScheduleContext != nullptr);
+	return *GScheduleContext;
+}
+
+}
diff --git a/Engine/Plugins/Experimental/AnimNext/Source/AnimNext/Private/Scheduler/ScheduleContext.h b/Engine/Plugins/Experimental/AnimNext/Source/AnimNext/Private/Scheduler/ScheduleContext.h
new file mode 100644
--- /dev/null
+++ b/Engine/Plugins/Experimental/AnimNext/Source/AnimNext/Private/Scheduler/ScheduleContext.h
@@ -0,0 +1,66 @@
+// Copyright Epic Games, Inc. All Rights Reserved.
+
+#pragma once
+
+#include "CoreMinimal.h"
+#include "Param/ParamStack.h"
+#include "Scheduler/ScheduleHandle.h"
+#include "Scheduler/ScheduleInstanceData.h"
+
+class UAnimNextSchedule;
+
+namespace UE::AnimNext
+{
+
+struct FScheduleContext
+{
+public:
+	// Attach a context to the current thread
+	static void AttachToCurrentThread(const FScheduleContext& InContext);
+
+	// Detach a context from the current thread
+	static void DetachFromCurrentThread();
+
+	// Get the context that is attached to the current thread
+	static const FScheduleContext& Get();
+
+	FScheduleContext() = default;
+
+	FScheduleContext(FScheduleContext&& InOther) = default;
+	FScheduleContext& operator=(FScheduleContext&& InOther) = default;
+	FScheduleContext(const FScheduleContext& InOther) = delete;
+	FScheduleContext& operator=(const FScheduleContext& InOther) = delete;
+
+	FScheduleContext(const UAnimNextSchedule* InSchedule, FAnimNextSchedulerEntry* InEntry)
+		: Schedule(InSchedule)
+		, Entry(InEntry)
+	{}
+
+	// Get a string that can identify this schedule for debugging
+	FString GetDebugString() const 
+	{ 
+		if(ensure(InstanceData))
+		{
+			return InstanceData->Handle.ToString(); 
+		}
+		return TEXT("Invalid schedule instance data");
+	}
+
+	// Get the instance data for this running schedule
+	FScheduleInstanceData& GetInstanceData() const 
+	{ 
+		check(InstanceData.IsValid());
+		return *InstanceData.Get(); 
+	}
+
+	// Current schedule
+	const UAnimNextSchedule* Schedule = nullptr;
+
+	// Current entry for this schedule
+	FAnimNextSchedulerEntry* Entry = nullptr;
+
+	// Instance data for the schedule
+	TUniquePtr<FScheduleInstanceData> InstanceData;
+};
+
+}
\ No newline at end of file
diff --git a/Engine/Plugins/Experimental/AnimNext/Source/AnimNext/Private/Scheduler/ScheduleInstanceData.cpp b/Engine/Plugins/Experimental/AnimNext/Source/AnimNext/Private/Scheduler/ScheduleInstanceData.cpp
new file mode 100644
--- /dev/null
+++ b/Engine/Plugins/Experimental/AnimNext/Source/AnimNext/Private/Scheduler/ScheduleInstanceData.cpp
@@ -0,0 +1,61 @@
+// Copyright Epic Games, Inc. All Rights Reserved.
+
+#include "ScheduleInstanceData.h"
+#include "Scheduler/AnimNextSchedule.h"
+#include "ScheduleHandle.h"
+#include "Scheduler/AnimNextSchedulerWorldSubsystem.h"
+#include "Scheduler/AnimNextSchedulerEntry.h"
+
+namespace UE::AnimNext
+{
+
+FScheduleInstanceData::FScheduleInstanceData(const FScheduleContext& InScheduleContext, const UAnimNextSchedule* InSchedule, FScheduleHandle InHandle, FAnimNextSchedulerEntry* InCurrentEntry, TMap<FName, FAnimNextParameterCollection>&& InUserScopes)
+	: Handle(InHandle)
+	, Entry(InCurrentEntry) 
+	, UserScopes(MoveTemp(InUserScopes))
+{
+	// Preallocate data for all scopes & graphs in the schedule
+	ScopeCaches.SetNum(InSchedule->ParamScopeEntryTasks.Num());
+	GraphInstanceData.SetNum(InSchedule->Tasks.Num());
+
+	// Setup param stack graph
+	RootParamStack = InCurrentEntry->RootParamStack;
+	ParamStacks.SetNum(InSchedule->NumParameterScopes);
+	for (TSharedPtr<FParamStack>& ParamStack : ParamStacks)
+	{
+		ParamStack = MakeShared<FParamStack>();
+	}
+
+	for (const FAnimNextScheduleGraphTask& Task : InSchedule->Tasks)
+	{
+		const TSharedPtr<FParamStack> ParentStack = Task.ParamParentScopeIndex != MAX_uint32 ? ParamStacks[Task.ParamParentScopeIndex] : RootParamStack;
+		ParamStacks[Task.ParamScopeIndex]->SetParent(ParentStack);
+	}
+
+	for (const FAnimNextScheduleExternalTask& ExternalTask : InSchedule->ExternalTasks)
+	{
+		const TSharedPtr<FParamStack> ParentStack = ExternalTask.ParamParentScopeIndex != MAX_uint32 ? ParamStacks[ExternalTask.ParamParentScopeIndex] : RootParamStack;
+		ParamStacks[ExternalTask.ParamScopeIndex]->SetParent(ParentStack);
+	}
+
+	for (const FAnimNextScheduleParamScopeEntryTask& ScopeEntryTask : InSchedule->ParamScopeEntryTasks)
+	{
+		const TSharedPtr<FParamStack> ParentStack = ScopeEntryTask.ParamParentScopeIndex != MAX_uint32 ? ParamStacks[ScopeEntryTask.ParamParentScopeIndex] : RootParamStack;
+		ParamStacks[ScopeEntryTask.ParamScopeIndex]->SetParent(ParentStack);
+	}
+}
+
+void FScheduleInstanceData::AddReferencedObjects(FReferenceCollector& Collector)
+{
+	for (TPair<FName, FAnimNextParameterCollection>& ParamPair : UserScopes)
+	{
+		Collector.AddPropertyReferencesWithStructARO(FAnimNextParameterCollection::StaticStruct(), &ParamPair.Value);
+	}
+}
+
+FString FScheduleInstanceData::GetReferencerName() const
+{
+	return TEXT("AnimNextInstanceData");
+}
+
+}
\ No newline at end of file
