AnimNext: Scheduler, Component and Parameter Improvements

diff --git a/Engine/Plugins/Experimental/AnimNext/Source/AnimNext/Private/Scheduler/ScheduleInstanceData.h b/Engine/Plugins/Experimental/AnimNext/Source/AnimNext/Private/Scheduler/ScheduleInstanceData.h
new file mode 100644
--- /dev/null
+++ b/Engine/Plugins/Experimental/AnimNext/Source/AnimNext/Private/Scheduler/ScheduleInstanceData.h
@@ -0,0 +1,67 @@
+// Copyright Epic Games, Inc. All Rights Reserved.
+
+#pragma once
+
+#include "CoreMinimal.h"
+#include "ScheduleHandle.h"
+#include "Param/AnimNextParameterCollection.h"
+#include "Graph/AnimNextGraph.h"
+
+struct FAnimNextGraphInstance;
+class UAnimNextSchedule;
+struct FAnimNextSchedulerEntry;
+
+namespace UE::AnimNext
+{
+	struct FParamStack;
+	struct FParamStackLayerHandle;
+}
+
+namespace UE::AnimNext
+{
+
+// Host for all data needed to run a schedule instance
+struct FScheduleInstanceData : public FGCObject
+{
+	FScheduleInstanceData(const FScheduleContext& InScheduleContext, const UAnimNextSchedule* InSchedule, FScheduleHandle InHandle, FAnimNextSchedulerEntry* InCurrentEntry, TMap<FName, FAnimNextParameterCollection>&& InDynamicScopes);
+
+	// FGCObject interface
+	virtual void AddReferencedObjects(FReferenceCollector& Collector) override;
+	virtual FString GetReferencerName() const override;
+
+	// Handle to the currently executing entry in the schedule
+	FScheduleHandle Handle;
+
+	// Schedule entry that owns this instance
+	FAnimNextSchedulerEntry* Entry = nullptr;
+
+	// Set of dynamic parameter scopes supplied by the user
+	TMap<FName, FAnimNextParameterCollection> UserScopes;
+
+	// Cached data for each parameter scope
+	struct FScopeCache
+	{
+		// Cached handles for scheduled parameter blocks
+		TArray<FParamStackLayerHandle> StaticHandles;
+
+		// Cached handles for user-defined parameter blocks hooked to a scope
+		TArray<FParamStackLayerHandle> UserHandles;
+
+		// Pushed layers, popped when scope exits
+		TArray<FParamStack::FPushedLayerHandle> PushedLayers;
+	};
+
+	// Cached data for all param scopes
+	TArray<FScopeCache> ScopeCaches;
+
+	// Root param stack for the schedule itself (and globals)
+	TSharedPtr<FParamStack> RootParamStack;
+
+	// Param stacks required to run the schedule (one per task that requires a stack)
+	TArray<TSharedPtr<FParamStack>> ParamStacks;
+
+	// Graph instance data for each graph task
+	TArray<FAnimNextGraphInstance> GraphInstanceData;
+};
+
+}
\ No newline at end of file
diff --git a/Engine/Plugins/Experimental/AnimNext/Source/AnimNext/Private/Scheduler/SchedulePortAdapterContext.h b/Engine/Plugins/Experimental/AnimNext/Source/AnimNext/Private/Scheduler/SchedulePortAdapterContext.h
new file mode 100644
--- /dev/null
+++ b/Engine/Plugins/Experimental/AnimNext/Source/AnimNext/Private/Scheduler/SchedulePortAdapterContext.h
@@ -0,0 +1,23 @@
+// Copyright Epic Games, Inc. All Rights Reserved.
+
+#pragma once
+
+#include "CoreMinimal.h"
+#include "Param/ParamTypeHandle.h"
+
+namespace UE::AnimNext
+{
+
+struct FSchedulePortAdapterContext
+{
+	// TODO: templated type-checked data access for data/objects
+
+	TConstArrayView<uint8> InputData;
+	TArrayView<uint8> OutputData;
+	UObject* InputObject = nullptr;
+	UObject* OutputObject = nullptr;
+	FParamTypeHandle InputType;
+	FParamTypeHandle OutputType;
+};
+
+}
\ No newline at end of file
diff --git a/Engine/Plugins/Experimental/AnimNext/Source/AnimNext/Private/Scheduler/SchedulePortDefinition.h b/Engine/Plugins/Experimental/AnimNext/Source/AnimNext/Private/Scheduler/SchedulePortDefinition.h
new file mode 100644
--- /dev/null
+++ b/Engine/Plugins/Experimental/AnimNext/Source/AnimNext/Private/Scheduler/SchedulePortDefinition.h
@@ -0,0 +1,45 @@
+// Copyright Epic Games, Inc. All Rights Reserved.
+
+#pragma once
+
+#include "CoreMinimal.h"
+#include "Param/ParamTypeHandle.h"
+
+struct FAnimNextSchedulePortTask;
+
+namespace UE::AnimNext
+{
+	struct FScheduler;
+	struct FSchedulePortAdapterContext;
+}
+
+namespace UE::AnimNext
+{
+
+// Port definition registered with the scheduler
+struct FSchedulePortDefinition
+{
+	// Function signature for port adapters
+	using FAdapterFunction = TUniqueFunction<void(const FSchedulePortAdapterContext& InContext)>;
+
+	FSchedulePortDefinition(FName InName, const FParamTypeHandle& InInputType, const FParamTypeHandle& InOutputType, FAdapterFunction&& InAdapter)
+		: Name(InName)
+		, InputType(InInputType)
+		, OutputType(InOutputType)
+		, Adapter(MoveTemp(InAdapter))
+	{}
+
+private:
+	friend struct FScheduler;
+	friend struct ::FAnimNextSchedulePortTask;
+
+	FName Name;
+
+	FParamTypeHandle InputType;
+
+	FParamTypeHandle OutputType;
+
+	FAdapterFunction Adapter;
+};
+
+}
\ No newline at end of file
diff --git a/Engine/Plugins/Experimental/AnimNext/Source/AnimNext/Private/Scheduler/ScheduleTickFunction.cpp b/Engine/Plugins/Experimental/AnimNext/Source/AnimNext/Private/Scheduler/ScheduleTickFunction.cpp
new file mode 100644
--- /dev/null
+++ b/Engine/Plugins/Experimental/AnimNext/Source/AnimNext/Private/Scheduler/ScheduleTickFunction.cpp
@@ -0,0 +1,138 @@
+// Copyright Epic Games, Inc. All Rights Reserved.
+
+#include "ScheduleTickFunction.h"
+#include "Scheduler/ScheduleContext.h"
+#include "Scheduler/AnimNextSchedule.h"
+#include "Scheduler/AnimNextScheduleGraphTask.h"
+#include "Scheduler/AnimNextSchedulePortTask.h"
+#include "Scheduler/AnimNextScheduleExternalTask.h"
+#include "Scheduler/AnimNextScheduleParamScopeTask.h"
+#include "Scheduler/ScheduleInstanceData.h"
+#include "Scheduler/AnimNextSchedulerWorldSubsystem.h"
+
+namespace UE::AnimNext
+{
+
+void FScheduleBeginTickFunction::ExecuteTick(float DeltaTime, ELevelTick TickType, ENamedThreads::Type CurrentThread, const FGraphEventRef& MyCompletionGraphEvent)
+{
+	// Allocate instance data if required
+	if (!Entry.Context.InstanceData.IsValid())
+	{
+		Entry.Context.InstanceData = MakeUnique<FScheduleInstanceData>(Entry.Context, Entry.Schedule, Entry.Handle, &Entry, MoveTemp(Entry.UserScopes));
+	}
+}
+
+FString FScheduleBeginTickFunction::DiagnosticMessage()
+{
+	return TEXT("AnimNextScheduleBeginTickFunction");
+}
+
+void FScheduleEndTickFunction::ExecuteTick(float DeltaTime, ELevelTick TickType, ENamedThreads::Type CurrentThread, const FGraphEventRef& MyCompletionGraphEvent)
+{
+}
+
+FString FScheduleEndTickFunction::DiagnosticMessage()
+{
+	return TEXT("AnimNextScheduleEndTickFunction");
+}
+
+void FScheduleTickFunction::ExecuteTick(float DeltaTime, ELevelTick TickType, ENamedThreads::Type CurrentThread, const FGraphEventRef& MyCompletionGraphEvent)
+{
+	FScheduleContext::AttachToCurrentThread(ScheduleContext);
+
+	while (!PreExecuteTasks.IsEmpty())
+	{
+		TOptional<TUniqueFunction<void(const UE::AnimNext::FScheduleContext&)>> Function = PreExecuteTasks.Dequeue();
+		check(Function.IsSet());
+		Function.GetValue()(ScheduleContext);
+	}
+
+	const UAnimNextSchedule* Schedule = ScheduleContext.Schedule;
+
+	switch (Instruction.Opcode)
+	{
+	case EAnimNextScheduleScheduleOpcode::RunTask:
+		{
+			uint32 TaskIndex = Instruction.Operand;
+			FScheduleInstanceData& InstanceData = ScheduleContext.GetInstanceData();
+			FParamStack::AttachToCurrentThread(InstanceData.ParamStacks[Schedule->Tasks[TaskIndex].ParamScopeIndex]);
+
+			Schedule->Tasks[TaskIndex].RunGraph(ScheduleContext);
+
+			FParamStack::DetachFromCurrentThread();
+			break;
+		}
+	case EAnimNextScheduleScheduleOpcode::BeginRunExternalTask:
+		{
+			if (const UObject* Object = TargetObject.Get())
+			{
+				FScheduleInstanceData& InstanceData = ScheduleContext.GetInstanceData();
+				uint32 ExternalTaskIndex = Instruction.Operand;
+				FParamStack::AddForPendingObject(Object, InstanceData.ParamStacks[Schedule->ExternalTasks[ExternalTaskIndex].ParamScopeIndex]);
+			}
+			break;
+		}
+	case EAnimNextScheduleScheduleOpcode::EndRunExternalTask:
+		{
+			if (const UObject* Object = TargetObject.Get())
+			{
+				FScheduleInstanceData& InstanceData = ScheduleContext.GetInstanceData();
+				uint32 ExternalTaskIndex = Instruction.Operand;
+				FParamStack::RemoveForPendingObject(Object);
+			}
+			break;
+		}
+	case EAnimNextScheduleScheduleOpcode::RunPort:
+		{
+			uint32 PortIndex = Instruction.Operand;
+			FScheduleInstanceData& InstanceData = ScheduleContext.GetInstanceData();
+			FParamStack::AttachToCurrentThread(InstanceData.ParamStacks[Schedule->Ports[PortIndex].ParamScopeIndex]);
+
+			Schedule->Ports[PortIndex].RunPort(ScheduleContext);
+
+			FParamStack::DetachFromCurrentThread();
+			break;
+		}
+	case EAnimNextScheduleScheduleOpcode::RunParamScopeEntry:
+		{
+			uint32 ScopeEntryIndex = Instruction.Operand;
+			FScheduleInstanceData& InstanceData = ScheduleContext.GetInstanceData();
+			FParamStack::AttachToCurrentThread(InstanceData.ParamStacks[Schedule->ParamScopeEntryTasks[ScopeEntryIndex].ParamScopeIndex]);
+
+			Schedule->ParamScopeEntryTasks[ScopeEntryIndex].RunParamScopeEntry(ScheduleContext);
+
+			FParamStack::DetachFromCurrentThread();
+			break;
+		}
+	case EAnimNextScheduleScheduleOpcode::RunParamScopeExit:
+		{
+			uint32 ScopeExitIndex = Instruction.Operand;
+			FScheduleInstanceData& InstanceData = ScheduleContext.GetInstanceData();
+			FParamStack::AttachToCurrentThread(InstanceData.ParamStacks[Schedule->ParamScopeExitTasks[ScopeExitIndex].ParamScopeIndex]);
+
+			Schedule->ParamScopeExitTasks[ScopeExitIndex].RunParamScopeExit(ScheduleContext);
+
+			FParamStack::DetachFromCurrentThread();
+			break;
+		}
+	default:
+		check(false);	// Unimplemented instruction
+		break;
+	}
+
+	while (!PostExecuteTasks.IsEmpty())
+	{
+		TOptional<TUniqueFunction<void(const UE::AnimNext::FScheduleContext&)>> Function = PostExecuteTasks.Dequeue();
+		check(Function.IsSet());
+		Function.GetValue()(ScheduleContext);
+	}
+
+	FScheduleContext::DetachFromCurrentThread();
+}
+
+FString FScheduleTickFunction::DiagnosticMessage()
+{
+	return TEXT("AnimNextScheduleTickFunction");
+}
+
+}
\ No newline at end of file
diff --git a/Engine/Plugins/Experimental/AnimNext/Source/AnimNext/Private/Scheduler/ScheduleTickFunction.h b/Engine/Plugins/Experimental/AnimNext/Source/AnimNext/Private/Scheduler/ScheduleTickFunction.h
new file mode 100644
--- /dev/null
+++ b/Engine/Plugins/Experimental/AnimNext/Source/AnimNext/Private/Scheduler/ScheduleTickFunction.h
@@ -0,0 +1,80 @@
+// Copyright Epic Games, Inc. All Rights Reserved.
+
+#pragma once
+
+#include "CoreMinimal.h"
+#include "Engine/EngineBaseTypes.h"
+#include "Containers/SpscQueue.h"
+
+struct FAnimNextScheduleInstruction;
+struct FAnimNextSchedulerEntry;
+
+namespace UE::AnimNext
+{
+	struct FScheduleContext;
+	struct FParamStackLayerHandle;
+}
+
+namespace UE::AnimNext
+{
+
+struct FScheduleBeginTickFunction : public FTickFunction
+{
+	FScheduleBeginTickFunction(FAnimNextSchedulerEntry& InEntry)
+		: Entry(InEntry)
+	{
+		bCanEverTick = true;
+		bStartWithTickEnabled = true;
+		bRunOnAnyThread = true;
+	}
+
+	// FTickFunction interface
+	virtual void ExecuteTick(float DeltaTime, ELevelTick TickType, ENamedThreads::Type CurrentThread, const FGraphEventRef& MyCompletionGraphEvent) override;
+	virtual FString DiagnosticMessage() override;
+
+	FAnimNextSchedulerEntry& Entry;
+	FTickPrerequisite Subsequent;
+};
+
+struct FScheduleEndTickFunction : public FTickFunction
+{
+	FScheduleEndTickFunction(FAnimNextSchedulerEntry& InEntry)
+		: Entry(InEntry)
+	{
+		bCanEverTick = true;
+		bStartWithTickEnabled = true;
+		bRunOnAnyThread = true;
+	}
+
+	// FTickFunction interface
+	virtual void ExecuteTick(float DeltaTime, ELevelTick TickType, ENamedThreads::Type CurrentThread, const FGraphEventRef& MyCompletionGraphEvent) override;
+	virtual FString DiagnosticMessage() override;
+
+	FAnimNextSchedulerEntry& Entry;
+};
+
+// TODO: Run a linear slice of instructions rather than have each instruction a separate tick function
+struct FScheduleTickFunction : public FTickFunction
+{
+	FScheduleTickFunction(const FScheduleContext& InScheduleContext, const FAnimNextScheduleInstruction& InInstruction)
+		: ScheduleContext(InScheduleContext)
+		, Instruction(InInstruction)
+	{
+		bCanEverTick = true;
+		bStartWithTickEnabled = true;
+		bRunOnAnyThread = true;
+	}
+
+	// FTickFunction interface
+	virtual void ExecuteTick(float DeltaTime, ELevelTick TickType, ENamedThreads::Type CurrentThread, const FGraphEventRef& MyCompletionGraphEvent) override;
+	virtual FString DiagnosticMessage() override;
+
+	const FScheduleContext& ScheduleContext;
+	const FAnimNextScheduleInstruction& Instruction;
+	TArray<FTickPrerequisite> Subsequents;
+	TWeakObjectPtr<UObject> TargetObject;
+	TSpscQueue<TUniqueFunction<void(const UE::AnimNext::FScheduleContext&)>> PreExecuteTasks;
+	TSpscQueue<TUniqueFunction<void(const UE::AnimNext::FScheduleContext&)>> PostExecuteTasks;
+};
+
+}
\ No newline at end of file
diff --git a/Engine/Plugins/Experimental/AnimNext/Source/AnimNext/Private/Scheduler/Scheduler.cpp b/Engine/Plugins/Experimental/AnimNext/Source/AnimNext/Private/Scheduler/Scheduler.cpp
new file mode 100644
--- /dev/null
+++ b/Engine/Plugins/Experimental/AnimNext/Source/AnimNext/Private/Scheduler/Scheduler.cpp
@@ -0,0 +1,142 @@
+// Copyright Epic Games, Inc. All Rights Reserved.
+
+#include "Scheduler/Scheduler.h"
+#include "Scheduler/SchedulePortAdapterContext.h"
+#include "Scheduler/SchedulePortDefinition.h"
+#include "Scheduler/ScheduleContext.h"
+#include "Scheduler/AnimNextSchedule.h"
+#include "Scheduler/AnimNextSchedulerWorldSubsystem.h"
+#include "Tasks/Task.h"
+#include "UObject/GCObject.h"
+#include "LODPose.h"
+#include "Param/Params.h"
+#include "Engine/World.h"
+#include "Components/SkeletalMeshComponent.h"
+#include "UObject/ObjectKey.h"
+
+namespace UE::AnimNext
+{
+
+struct FSchedulerImpl
+{
+	void RegisterBuiltInPortDefinitions()
+	{
+/*		RegisteredPortDefinitions.Add(
+			Ports::SkeletalMeshComponentPose.Name,
+			FSchedulePortDefinition(
+				Ports::SkeletalMeshComponentPose.Name,
+				FParamTypeHandle::GetHandle<TArray<FAnimNextLODPose>>(),
+				FParamTypeHandle::GetHandle<TArray<FTransform>>(),
+				[](const FSchedulePortAdapterContext& InContext)
+				{
+					check(InContext.InputType == FParamTypeHandle::GetHandle<FAnimNextLODPose>());
+					check(InContext.OutputType == FParamTypeHandle::GetHandle<TArray<FTransform>>());
+
+					USkeletalMeshComponent* OutputObject = CastChecked<USkeletalMeshComponent>(InContext.OutputObject);
+				
+					// TODO: Perform pose transpose
+				}));*/
+	}
+
+	FDelegateHandle OnWorldPreActorTickHandle;
+
+	TMap<FName, FSchedulePortDefinition> RegisteredPortDefinitions;
+
+	uint32 SerialNumber = 0;
+};
+
+static FSchedulerImpl Impl;
+
+void FScheduler::Init()
+{
+	Impl.RegisterBuiltInPortDefinitions();
+
+	// Kick off root task at the start of each world tick
+	Impl.OnWorldPreActorTickHandle = FWorldDelegates::OnWorldPreActorTick.AddLambda([](UWorld* InWorld, ELevelTick InTickType, float InDeltaSeconds)
+	{
+		if (InTickType == LEVELTICK_All || InTickType == LEVELTICK_ViewportsOnly)
+		{
+			if(UAnimNextSchedulerWorldSubsystem* Subsystem = InWorld->GetSubsystem<UAnimNextSchedulerWorldSubsystem>())
+			{
+				// Flush actions here as they require game thread callbacks (e.g. to reconfigure tick functions)
+				Subsystem->FlushPendingActions();
+				Subsystem->DeltaTime = InDeltaSeconds;
+			}
+		}
+	});
+}
+
+void FScheduler::Destroy()
+{
+	FWorldDelegates::OnWorldPreActorTick.Remove(Impl.OnWorldPreActorTickHandle);
+}
+
+FScheduleHandle FScheduler::AcquireHandle(UObject* InObject, UAnimNextSchedule* InSchedule, const TMap<FName, FAnimNextParameterCollection>& InUserScopes)
+{
+	FScheduleHandle Handle;
+
+	// Check parameters
+	if (InSchedule == nullptr)
+	{
+		UE_LOG(LogAnimation, Warning, TEXT("FScheduler::AcquireHandle: Invalid schedule"));
+		return FScheduleHandle();
+	}
+
+	if (InObject == nullptr)
+	{
+		UE_LOG(LogAnimation, Warning, TEXT("FScheduler::AcquireHandle: Invalid object"));
+		return FScheduleHandle();
+	}
+
+	UWorld* World = InObject->GetWorld();
+	UAnimNextSchedulerWorldSubsystem* Subsystem = World->GetSubsystem<UAnimNextSchedulerWorldSubsystem>();
+	return Subsystem->AcquireHandle(InObject, InSchedule, InUserScopes);
+}
+
+void FScheduler::ReleaseHandle(UObject* InObject, FScheduleHandle& InHandle)
+{
+	if(InHandle.IsValid())
+	{
+		UWorld* World = InObject->GetWorld();
+		UAnimNextSchedulerWorldSubsystem* Subsystem = World->GetSubsystem<UAnimNextSchedulerWorldSubsystem>();
+		Subsystem->ReleaseHandle(InHandle);
+		InHandle.Invalidate();
+	}
+}
+
+void FScheduler::EnableHandle(UObject* InObject, FScheduleHandle InHandle, bool bInEnabled)
+{
+	if (InHandle.IsValid())
+	{
+		UWorld* World = InObject->GetWorld();
+		UAnimNextSchedulerWorldSubsystem* Subsystem = World->GetSubsystem<UAnimNextSchedulerWorldSubsystem>();
+		Subsystem->EnableHandle(InHandle, bInEnabled);
+	}
+}
+
+void FScheduler::QueueTask(UObject* InObject, FScheduleHandle InHandle, FName InScheduleTaskName, TUniqueFunction<void(const FScheduleContext&)>&& InTaskFunction, ETaskRunLocation InLocation)
+{
+	if (InHandle.IsValid())
+	{
+		UWorld* World = InObject->GetWorld();
+		UAnimNextSchedulerWorldSubsystem* Subsystem = World->GetSubsystem<UAnimNextSchedulerWorldSubsystem>();
+		Subsystem->QueueTask(InHandle, InScheduleTaskName, MoveTemp(InTaskFunction), InLocation);
+	}
+}
+
+void FScheduler::RegisterPortDefinition(FSchedulePortDefinition&& InPortDefinition)
+{
+	Impl.RegisteredPortDefinitions.Add(InPortDefinition.Name, MoveTemp(InPortDefinition));
+}
+
+void FScheduler::UnregisterPortDefinition(FName InDefinitionName)
+{
+	Impl.RegisteredPortDefinitions.Remove(InDefinitionName);
+}
+
+const FSchedulePortDefinition* FScheduler::FindPortDefinition(FName InDefinitionName)
+{
+	return Impl.RegisteredPortDefinitions.Find(InDefinitionName);
+}
+
+}
\ No newline at end of file
diff --git a/Engine/Plugins/Experimental/AnimNext/Source/AnimNext/Private/Scheduler/Scheduler.h b/Engine/Plugins/Experimental/AnimNext/Source/AnimNext/Private/Scheduler/Scheduler.h
new file mode 100644
--- /dev/null
+++ b/Engine/Plugins/Experimental/AnimNext/Source/AnimNext/Private/Scheduler/Scheduler.h
@@ -0,0 +1,72 @@
+// Copyright Epic Games, Inc. All Rights Reserved.
+
+#pragma once
+
+#include "CoreMinimal.h"
+#include "Param/ParamTypeHandle.h"
+#include "ScheduleHandle.h"
+
+class UAnimNextGraph;
+class UAnimNextParameterBlock;
+class UAnimNextSchedule;
+struct FAnimNextParameterCollection;
+
+namespace UE::AnimNext
+{
+	struct FScheduler;
+	struct FSchedulerImpl;
+	struct FSchedulePortDefinition;
+	struct FParamStackLayerHandle;
+	struct FScheduleContext;
+}
+
+namespace UE::AnimNext
+{
+
+// Main interface into the AnimNext scheduling system
+struct FScheduler
+{
+	// Start up the scheduler system
+	static void Init();
+
+	// Shut down the scheduler system
+	static void Destroy();
+
+	// Acquire a handle that binds a schedule with the supplied parameters
+	static FScheduleHandle AcquireHandle(UObject* InObject, UAnimNextSchedule* InSchedule, const TMap<FName, FAnimNextParameterCollection>& InUserScopes);
+
+	// Release an already acquired handle
+	// The full release of the binding referenced by the handle map be deferred after this call is made
+	static void ReleaseHandle(UObject* InObject, FScheduleHandle& InHandle);
+
+	// Enables or disables the schedule parameterization represented by the supplied handle
+	// This operation is deferred until the next time the schedule ticks
+	static void EnableHandle(UObject* InObject, FScheduleHandle InHandle, bool bInEnabled);
+
+	enum class ETaskRunLocation : int32
+	{
+		// Run the task before the specified task
+		Before,
+
+		// Run the task after the specified task
+		After,
+	};
+
+	// Queue a task to run at a particular point in a schedule
+	// @param	InHandle		The handle to queue the task to
+	// @param	InTaskName		The name of the task in the schedule to run the supplied task relative to
+	// @param	InTaskFunction	The function to run
+	// @param	InLocation		Where to run the task, before or after
+	static void QueueTask(UObject* InObject, FScheduleHandle InHandle, FName InScheduleTaskName, TUniqueFunction<void(const FScheduleContext&)>&& InTaskFunction, ETaskRunLocation InLocation = ETaskRunLocation::Before);
+
+	// Register a port definition
+	static void RegisterPortDefinition(FSchedulePortDefinition&& InPortDefinition);
+
+	// Unregister a port definition
+	static void UnregisterPortDefinition(FName InDefinitionName);
+
+	// Find a registered port definition
+	static const FSchedulePortDefinition* FindPortDefinition(FName InDefinitionName);
+};
+
+}
\ No newline at end of file
diff --git a/Engine/Plugins/Experimental/AnimNext/Source/AnimNext/Public/Param/IAnimNextParameterSourceInterface.h b/Engine/Plugins/Experimental/AnimNext/Source/AnimNext/Public/Param/IAnimNextParameterSourceInterface.h
new file mode 100644
--- /dev/null
+++ b/Engine/Plugins/Experimental/AnimNext/Source/AnimNext/Public/Param/IAnimNextParameterSourceInterface.h
@@ -0,0 +1,35 @@
+// Copyright Epic Games, Inc. All Rights Reserved.
+
+#pragma once
+
+#include "CoreMinimal.h"
+#include "UObject/Interface.h"
+#include "Param/ParamStack.h"
+#include "IAnimNextParameterSourceInterface.generated.h"
+
+namespace UE::AnimNext
+{
+	struct FParamContext;
+}
+
+UINTERFACE(MinimalAPI, NotBlueprintable)
+class UAnimNextParameterSourceInterface : public UInterface
+{
+	GENERATED_BODY()
+};
+
+/** An interface used to apply AnimNext parameters to a context */
+class IAnimNextParameterSourceInterface
+{
+	GENERATED_BODY()
+
+public:
+	/** Called back by the system to apply this interface's parameters to the layer */
+	virtual void UpdateLayer(UE::AnimNext::FParamStackLayerHandle& InHandle) const = 0;
+
+	/** Cache a (immutable) layer handle for this parameter source */
+	virtual UE::AnimNext::FParamStackLayerHandle CacheLayer() const = 0;
+	
+	/* Whether to cache the supplied layer (e.g. if the layout changes dynamically in editor or the handle is invalid) */
+	virtual bool ShouldCacheLayer(const UE::AnimNext::FParamStackLayerHandle& InHandle) const { return !InHandle.IsValid(); }
+};
\ No newline at end of file
diff --git a/Engine/Plugins/Experimental/AnimNext/Source/AnimNext/Public/Param/ParamEntry.h b/Engine/Plugins/Experimental/AnimNext/Source/AnimNext/Public/Param/ParamEntry.h
new file mode 100644
--- /dev/null
+++ b/Engine/Plugins/Experimental/AnimNext/Source/AnimNext/Public/Param/ParamEntry.h
@@ -0,0 +1,114 @@
+// Copyright Epic Games, Inc. All Rights Reserved.
+
+#pragma once
+
+#include "CoreMinimal.h"
+#include "Param/ParamTypeHandle.h"
+
+namespace UE::AnimNext
+{
+	struct FParamStack;
+	struct FUObjectLayer;
+	struct FInstancedPropertyBagLayer;
+	struct FParamStackLayerHandle;
+}
+
+namespace UE::AnimNext
+{
+	struct FParamStackLayer;
+}
+
+namespace UE::AnimNext::Private
+{
+
+enum class EParamFlags : uint8
+{
+	None = 0,				// No flags
+	Mutable = 1 << 0,		// Parameter is mutable, so can be mutated after it has been created
+	Embedded = 1 << 1,		// Parameter will be stored as a Value, but stored directly on the Data pointer
+	Reference = 1 << 2,		// Parameter is a reference, so will not be embdedded even if it fits in the Data pointer
+};
+
+ENUM_CLASS_FLAGS(EParamFlags);
+
+// Parameter memory wrapper used in the param stack
+struct FParamEntry
+{
+	FParamEntry() = default;
+	ANIMNEXT_API ~FParamEntry();
+	ANIMNEXT_API FParamEntry(const FParamEntry& InOtherParam);
+	ANIMNEXT_API FParamEntry& operator=(const FParamEntry& InOtherParam);
+	ANIMNEXT_API FParamEntry(FParamEntry&& InOtherParam) noexcept;
+	ANIMNEXT_API FParamEntry& operator=(FParamEntry&& InOtherParam) noexcept;
+
+private:
+	friend struct UE::AnimNext::FParamStack;
+	friend struct UE::AnimNext::FParamStackLayer;
+	friend struct UE::AnimNext::FUObjectLayer;
+	friend struct UE::AnimNext::FInstancedPropertyBagLayer;
+	friend struct UE::AnimNext::FParamStackLayerHandle;
+
+	ANIMNEXT_API FParamEntry(const FParamTypeHandle& InTypeHandle, TArrayView<uint8> InData, bool bInIsReference, bool bInIsMutable);
+
+	// Get the type handle of this param 
+	FParamTypeHandle GetTypeHandle() const { return TypeHandle; }
+
+	// Check whether this parameter is able to be mutated
+	bool IsMutable() const { return EnumHasAnyFlags(Flags, EParamFlags::Mutable); }
+
+	// Check whether this parameter is of reference type (i.e. the parameter refers to user data owned outside of the parameter)
+	bool IsReference() const { return EnumHasAnyFlags(Flags, EParamFlags::Reference); }
+
+	// Check whether this parameter is embedded in the parameter (stored internally rather than as a ptr)
+	bool IsEmbedded() const { return EnumHasAnyFlags(Flags, EParamFlags::Embedded); }
+
+	// Get an immutable view of the parameter's data
+	TConstArrayView<uint8> GetData() const
+	{
+		if (EnumHasAnyFlags(Flags, EParamFlags::Embedded))
+		{
+			return TConstArrayView<uint8>(reinterpret_cast<const uint8*>(&Data), Size);
+		}
+		else
+		{
+			return TConstArrayView<uint8>(static_cast<const uint8*>(Data), Size);
+		}
+	}
+
+	// Get an mutable view of the parameter's data, returns an empty array view if this parameter is immutable
+	TArrayView<uint8> GetMutableData()
+	{
+		if (IsMutable())
+		{
+			if (EnumHasAnyFlags(Flags, EParamFlags::Embedded))
+			{
+				return TArrayView<uint8>(reinterpret_cast<uint8*>(&Data), Size);
+			}
+			else
+			{
+				return TArrayView<uint8>(static_cast<uint8*>(Data), Size);
+			}
+		}
+		return TArrayView<uint8>();
+	}
+
+	// Check whether this param represents valid data. Note that this doesn't check the type hande for validity.
+	bool IsValid() const
+	{
+		return Size > 0;
+	}
+
+	// Raw ptr to the data, or the data itself if we have EFlags::Embedded
+	void* Data = nullptr;
+
+	// The type of the param
+	FParamTypeHandle TypeHandle;
+
+	// Size of the data
+	uint16 Size = 0;
+
+	// Internal flags
+	EParamFlags Flags = EParamFlags::None;
+};
+
+}
\ No newline at end of file
diff --git a/Engine/Plugins/Experimental/AnimNext/Source/AnimNext/Public/Param/ParamId.h b/Engine/Plugins/Experimental/AnimNext/Source/AnimNext/Public/Param/ParamId.h
--- a/Engine/Plugins/Experimental/AnimNext/Source/AnimNext/Public/Param/ParamId.h
+++ b/Engine/Plugins/Experimental/AnimNext/Source/AnimNext/Public/Param/ParamId.h
@@ -4,9 +4,17 @@
 
 #include "CoreMinimal.h"
 
-struct FRigVMDispatch_SetParameter;
+class UAnimNextConfig;
+struct FRigVMDispatch_SetLayerParameter;
+struct FRigVMDispatch_GetLayerParameter;
 struct FRigVMDispatch_GetParameter;
-class UAnimNextParametersLibrary;
+
+namespace UE::AnimNext
+{
+	struct FParamAdapter;
+	struct FParamDefinition;
+	class FModule;
+}
 
 namespace UE::AnimNext::Tests
 {
@@ -21,9 +29,13 @@ struct ANIMNEXT_API FParamId
 {
 	friend struct FParamStack;
 	friend class Tests::FParamStackTest;
-	friend struct ::FRigVMDispatch_SetParameter;
+	friend struct ::FRigVMDispatch_SetLayerParameter;
+	friend struct ::FRigVMDispatch_GetLayerParameter;
 	friend struct ::FRigVMDispatch_GetParameter;
-	friend class ::UAnimNextParametersLibrary;
+	friend struct FParamAdapter;
+	friend class ::UAnimNextConfig;
+	friend struct FParamDefinition;
+	friend class FModule;
 
 	static constexpr uint32 InvalidIndex = MAX_uint32;
 
@@ -31,7 +43,6 @@ struct ANIMNEXT_API FParamId
 	FParamId(const FParamId& InName) = default;
 	FParamId& operator=(const FParamId& InName) = default;
 
-
 	// Make a parameter ID from an FName
 	explicit FParamId(FName InName);
 
@@ -44,6 +55,12 @@ struct ANIMNEXT_API FParamId
 	// Get the name that this parameter was created from
 	FName ToName() const;
 
+	// Check if this ID represents a valid parameter
+	bool IsValid() const
+	{
+		return ParameterIndex != InvalidIndex;
+	}
+
 private:
 	// Make a parameter ID from an index (for internal use)
 	explicit FParamId(uint32 InParameterIndex)
@@ -51,11 +68,39 @@ private:
 	{
 	}
 
+	// Make a new parameter ID (internal usage only)
+	static uint32 MakeParamId_NoLock(FName InName);
+
+	// Initialize the parameter ID system
+	static void Init();
+
+	// Shut down the parameter ID system
+	static void Destroy();
+
+	// Register all built-in adapters
+	static void RegisterBuiltInAdapters();
+	
+	// Refresh all adapters defined in config
+	static void RefreshConfigAdapters();
+	
 	// Get the maximum parameter ID that can exist at present
 	// Note that this can change via concurrent modifications when new parameters are created by 
 	// different threads
 	static FParamId GetMaxParamId();
 
+	// Get any built in adapter for the supplied parameter
+	static const FParamAdapter* GetAdapter(FParamId InId);
+
+	// Get the thread-local scratch area for an adapted parameter id
+	static uint8* GetScratchAreaForParamIdAdapter(FParamId InId);
+
+	// Get the thread-local return value for an adapted parameter id
+	template<typename ReturnType>
+	static ReturnType* GetAdapterReturnValue(FParamId InId)
+	{
+		return reinterpret_cast<ReturnType*>(GetScratchAreaForParamIdAdapter(InId));
+	}
+	
 #if WITH_DEV_AUTOMATION_TESTS
 	// Used to isolate global param IDs from automated tests
 	static void BeginTestSandbox();
diff --git a/Engine/Plugins/Experimental/AnimNext/Source/AnimNext/Public/Param/ParamResult.h b/Engine/Plugins/Experimental/AnimNext/Source/AnimNext/Public/Param/ParamResult.h
new file mode 100644
--- /dev/null
+++ b/Engine/Plugins/Experimental/AnimNext/Source/AnimNext/Public/Param/ParamResult.h
@@ -0,0 +1,75 @@
+// Copyright Epic Games, Inc. All Rights Reserved.
+
+#pragma once
+
+#include "CoreMinimal.h"
+
+namespace UE::AnimNext
+{
+
+// Results for param access
+enum class EParamResult : uint32
+{
+	Success = 0,
+
+	// Error flags in low bits
+
+	// Requested param is not in scope
+	NotInScope = 0x00000001,
+
+	// The requested parameter has an incompatible type
+	TypeError = 0x00000002,
+
+	// The requested parameter has an incompatible mutability (parameter is immutable but a mutable request was made of it)
+	MutabilityError = 0x00000004,
+
+	ErrorFlags = NotInScope | TypeError | MutabilityError,
+
+	// Info flags in high bits
+
+	// The type is compatible according to the supplied requirements
+	TypeCompatible = 0x80000000,
+
+	InfoFlags = TypeCompatible,
+};
+
+ENUM_CLASS_FLAGS(EParamResult);
+
+struct FParamResult
+{
+	FParamResult() = default;
+
+	// Implicit constructor
+	FParamResult(EParamResult InResult)
+		: Result(InResult)
+	{}
+
+	bool IsSuccessful() const
+	{
+		return !EnumHasAnyFlags(Result, EParamResult::ErrorFlags);
+	}
+
+	bool IsInScope() const
+	{
+		return !EnumHasAnyFlags(Result, EParamResult::NotInScope);
+	}
+
+	bool IsOfIncompatibleType() const
+	{
+		return EnumHasAnyFlags(Result, EParamResult::TypeError);
+	}
+
+	bool IsOfCompatibleType() const
+	{
+		return EnumHasAnyFlags(Result, EParamResult::TypeCompatible);
+	}
+
+	bool IsOfCompatibleMutability() const
+	{
+		return !EnumHasAnyFlags(Result, EParamResult::MutabilityError);
+	}
+
+	EParamResult Result = EParamResult::Success;
+};
+
+}
diff --git a/Engine/Plugins/Experimental/AnimNext/Source/AnimNext/Public/Param/ParamStack.h b/Engine/Plugins/Experimental/AnimNext/Source/AnimNext/Public/Param/ParamStack.h
--- a/Engine/Plugins/Experimental/AnimNext/Source/AnimNext/Public/Param/ParamStack.h
+++ b/Engine/Plugins/Experimental/AnimNext/Source/AnimNext/Public/Param/ParamStack.h
@@ -10,9 +10,31 @@
 #include "Param/ParamId.h"
 #include "Param/ParamTypeHandle.h"
 #include "Param/ParamCompatibility.h"
-#include "Misc/MemStack.h"
+#include "UObject/GCObject.h"
+#include "Param/ParamStackLayerHandle.h"
+#include "ParamEntry.h"
+#include "Param/ParamResult.h"
 
+class UAnimNextSchedule;
+class UAnimNextSchedulerWorldSubsystem;
 struct FInstancedPropertyBag;
+struct FAnimSubsystem_AnimNextParameters;
+class FAnimationAnimNextRuntimeTest_GraphExecute;
+class FAnimationAnimNextRuntimeTest_GraphExecuteLatent;
+struct FAnimNextSchedulerEntry;
+
+namespace UE::AnimNext
+{
+	struct FScheduleContext;
+	struct FSchedulerImpl;
+	struct FUObjectLayer;
+	struct FInstancedPropertyBagLayer;
+	struct FInstancedPropertyBagValueLayer;
+	struct FAnimGraphParamStackScope;
+	struct FScheduleInstanceData;
+	struct FScheduleTickFunction;
+	struct FParamStackLayer;
+}
 
 namespace UE::AnimNext::Tests
 {
@@ -23,7 +45,7 @@ namespace UE::AnimNext
 {
 
 // Stack of parameter layers.
-// Acts as an associative container - allows retrieval of parameter calues (by ID) that have been
+// Acts as an associative container - allows retrieval of parameter values (by ID) that have been
 // pushed onto the stack in 'layers'.
 // Parameter values (and notably types) on higher stack layers override those on lower layers. 
 // Older values and types are restored when a layer is popped.
@@ -32,120 +54,38 @@ namespace UE::AnimNext
 struct FParamStack
 {
 	friend class Tests::FParamStackTest;
-
+	friend struct FScheduleContext;
+	friend struct FScheduler;
+	friend struct FSchedulerImpl;
+	friend class ::UAnimNextSchedule;
+	friend struct ::FAnimSubsystem_AnimNextParameters;
+	friend struct FParamStackLayer;
+	friend struct FUObjectLayer;
+	friend struct FInstancedPropertyBagLayer;
+	friend struct FInstancedPropertyBagValueLayer;
+	friend struct FAnimGraphParamStackScope;
+	friend struct FScheduleInstanceData;
+	friend class ::UAnimNextSchedulerWorldSubsystem;
+	friend struct FScheduleTickFunction;
+	friend struct ::FAnimNextSchedulerEntry;
+	friend class ::FAnimationAnimNextRuntimeTest_GraphExecute;
+	friend class ::FAnimationAnimNextRuntimeTest_GraphExecuteLatent;
+
+	template<typename AllocatorType> 
+	friend struct TLayerBuilder;
+
+	template<typename T, ESPMode>
+	friend class SharedPointerInternals::TIntrusiveReferenceController;
+	
 private:
-	enum class EParamFlags : uint8
-	{
-		None = 0,				// No flags
-		Mutable = 1 << 0,		// Parameter is mutable, so can be mutated after it has been created
-		Embedded = 1 << 1,		// Parameter will be stored as a Value, but stored directly on the Data pointer
-		Reference = 1 << 2,		// Parameter is a reference, so will not be embdedded even if it fits in the Data pointer
-	};
-
-	FRIEND_ENUM_CLASS_FLAGS(EParamFlags);
-
-	// Parameter memory wrapper
-	struct FParam
-	{
-		FParam() = default;
-		ANIMNEXT_API FParam(const FParam& InOtherParam);
-		ANIMNEXT_API FParam& operator=(const FParam& InOtherParam);
-		ANIMNEXT_API FParam(FParam&& InOtherParam);
-		ANIMNEXT_API FParam& operator=(FParam&& InOtherParam);
-		ANIMNEXT_API ~FParam();
-
-		ANIMNEXT_API FParam(const FParamTypeHandle& InTypeHandle, TArrayView<uint8> InData, bool bInIsReference, bool bInIsMutable);
-
-		// Get the type handle of this param 
-		FParamTypeHandle GetTypeHandle() const { return TypeHandle; }
-
-		// Check whether this parameter is able to be mutated
-		bool IsMutable() const { return EnumHasAnyFlags(Flags, EParamFlags::Mutable); }
-
-		// Check whether this parameter is of reference type (i.e. the parameter refers to user data owned outside of the parameter)
-		bool IsReference() const { return EnumHasAnyFlags(Flags, EParamFlags::Reference); }
-
-		// Check whether this parameter is embedded in the parameter (stored internally rather than as a ptr)
-		bool IsEmbedded() const { return EnumHasAnyFlags(Flags, EParamFlags::Embedded); }
-
-		// Get an immutable view of the parameter's data
-		TConstArrayView<uint8> GetData() const
-		{
-			if (EnumHasAnyFlags(Flags, EParamFlags::Embedded))
-			{
-				return TConstArrayView<uint8>(reinterpret_cast<const uint8*>(&Data), Size);
-			}
-			else
-			{
-				return TConstArrayView<uint8>(static_cast<const uint8*>(Data), Size);
-			}
-		}
-
-		// Get an mutable view of the parameter's data, returns an empty array view if this parameter is immutable
-		TArrayView<uint8> GetMutableData()
-		{
-			if (IsMutable())
-			{
-				if (EnumHasAnyFlags(Flags, EParamFlags::Embedded))
-				{
-					return TArrayView<uint8>(reinterpret_cast<uint8*>(&Data), Size);
-				}
-				else
-				{
-					return TArrayView<uint8>(static_cast<uint8*>(Data), Size);
-				}
-			}
-			return TArrayView<uint8>();
-		}
-
-		// Check whether this param represents valid data. Note that this doesnt check the type hande for validity.
-		bool IsValid() const
-		{
-			return Size > 0;
-		}
-
-		// Raw ptr to the data, or the data itself if we have EFlags::Embedded
-		void* Data = nullptr;
-
-		// The type of the param
-		FParamTypeHandle TypeHandle;
-
-		// Size of the data
-		uint16 Size = 0;
-
-		// Internal flags
-		EParamFlags Flags = EParamFlags::None;
-	};
-
-	// Stack layers are what actually get pushed/popped onto the layer stack.
-	// They are designed to be held on an instance, their data updated in place.
-	// Memory ownership for params is assumed to be outside of the stack and layers
-	struct FLayer
-	{
-		FLayer() = delete;
-
-		explicit FLayer(FInstancedPropertyBag& InPropertyBag, bool bInIsMutable);
-
-		explicit FLayer(TConstArrayView<TPair<FParamId, FParamStack::FParam>> InParams);
-
-		// Params that this layer supplies
-		TArray<FParamStack::FParam> Params;
-
-		// ID that the param indices start at. Maps the global param ID range into the range for this layer
-		uint32 MinParamId = 0;
-
-		// Storage offset for this layer if it is owned internally, otherwise MAX_uint32
-		uint32 OwnedStorageOffset = MAX_uint32;
-	};
-
 	// A layer held on the stack, may be owned by the stack or not
 	struct FPushedLayer
 	{
-		FPushedLayer(FLayer& InLayer, FParamStack& InStack);
+		FPushedLayer(FParamStackLayer& InLayer, FParamStack& InStack);
 
-		FPushedLayer(const FPushedLayer& InPreviousLayer, FLayer& InLayer, FParamStack& InStack);
+		FPushedLayer(const FPushedLayer& InPreviousLayer, FParamStackLayer& InLayer, FParamStack& InStack);
 
-		FLayer& Layer;
+		FParamStackLayer& Layer;
 
 		// Index offset for the previous layer for each active param.
 		// Offset into FParamStack::PreviousLayerIndices.
@@ -187,32 +127,40 @@ public:
 		uint32 SerialNumber = 0;
 	};
 
-	// Opaque handle to a layer that can be held by external systems
-	struct FLayerHandle
+	// Builder struct to help build combined layers of parameters
+	template<typename AllocatorType = FDefaultAllocator>
+	struct TLayerBuilder
 	{
-		FLayerHandle() = default;
+		TLayerBuilder(int32 InNumParams = 0)
+		{
+			Params.Reserve(InNumParams);
+		}
 
-		bool IsValid() const
+		template<typename ValueType>
+		void Add(FName InName, ValueType&& InValue)
 		{
-			return Layer.IsValid();
+			Params.Emplace(FParamId(InName),
+				Private::FParamEntry(
+					FParamTypeHandle::GetHandle<std::remove_reference_t<ValueType>>(),
+					TArrayView<uint8>(const_cast<uint8*>(reinterpret_cast<const uint8*>(&InValue)), sizeof(std::remove_reference_t<ValueType>)),
+					std::is_reference_v<ValueType>,
+					!std::is_const_v<std::remove_reference_t<ValueType>>
+				)
+			);
 		}
 
-		void Invalidate()
+		FParamStackLayerHandle MakeLayer()
 		{
-			Layer = nullptr;
+			return FParamStack::Get().MakeLayer(Params);
 		}
 
 	private:
-		friend struct FParamStack;
-
-		explicit FLayerHandle(TUniquePtr<FLayer>&& InLayer)
-			: Layer(MoveTemp(InLayer))
-		{}
-
-		TUniquePtr<FLayer> Layer;
+		TArray<TPair<FParamId, Private::FParamEntry>, AllocatorType> Params;
 	};
 
-	ANIMNEXT_API FParamStack();
+	using FLayerBuilder = TLayerBuilder<>;
+
+	ANIMNEXT_API ~FParamStack();
 
 	// Get the param stack for this thread
 	static ANIMNEXT_API FParamStack& Get();
@@ -241,40 +189,40 @@ public:
 	FPushedLayerHandle PushValues(Args&&... InValues)
 	{
 		constexpr int32 NumItems = sizeof...(InValues) / 2;
-		TArray<TPair<FParamId, FParam>, TInlineAllocator<NumItems>> ParamIdValues;
+		TArray<TPair<FParamId, Private::FParamEntry>, TInlineAllocator<NumItems>> ParamIdValues;
 		ParamIdValues.Reserve(NumItems);
 		return PushValuesHelper(ParamIdValues, Forward<Args>(InValues)...);
 	}
 
 	// Push a layer
-	ANIMNEXT_API FPushedLayerHandle PushLayer(const FLayerHandle& InLayerHandle);
+	ANIMNEXT_API FPushedLayerHandle PushLayer(const FParamStackLayerHandle& InLayerHandle);
 
-	// Push an internally-owned layer. Copies parameter data to internal storage.
-	ANIMNEXT_API FPushedLayerHandle PushLayer(TConstArrayView<TPair<FParamId, FParam>> InParams);
+	// Create a cached parameter layer from a class. This layer will own an object used to represent the layer.
+	static ANIMNEXT_API FParamStackLayerHandle MakeValueLayer(const UClass* InClass);
 
-	// Create a cached parameter layer from an instanced property bag
-	static ANIMNEXT_API FLayerHandle MakeLayer(const FInstancedPropertyBag& InInstancedPropertyBag);
+	// Create a cached parameter layer from an object. This layer will weakly reference the supplied object.
+	static ANIMNEXT_API FParamStackLayerHandle MakeReferenceLayer(UObject* InObject);
 
-	// Create a cached mutable parameter layer from an instanced property bag
-	static ANIMNEXT_API FLayerHandle MakeMutableLayer(FInstancedPropertyBag& InPropertyBag);
+	// Create a cached parameter layer from an instanced property bag. This layer will own a copy of the property bag used to represent the layer.
+	static ANIMNEXT_API FParamStackLayerHandle MakeValueLayer(const FInstancedPropertyBag& InInstancedPropertyBag);
 
-	// Create a cached parameter layer from set of params. Mutabilty is on a per-parameter basis.
-	static ANIMNEXT_API FLayerHandle MakeLayer(TConstArrayView<TPair<FParamId, FParam>> InParams);
+	// Create a cached parameter layer from an instanced property bag. This layer will reference the suppled propery bag and does not transfer ownership.
+	static ANIMNEXT_API FParamStackLayerHandle MakeReferenceLayer(FInstancedPropertyBag& InInstancedPropertyBag);
 
 	// Make a parameter layer from a value
 	// @param	InParamId			Parameter ID for the parameter
 	// @param	InValue				Value to push
 	template<typename... ValueType>
-	static FLayerHandle MakeValueLayer(FParamId InParamId, ValueType&&... InValue)
+	static FParamStackLayerHandle MakeValueLayer(FParamId InParamId, ValueType&&... InValue)
 	{
 		return MakeValuesLayer(InParamId, Forward<ValueType>(InValue)...);
 	}
 
 	// Make a parameter layer from a value
-	// @param	InParamId			Parameter name for the parameter
+	// @param	InName				Parameter name for the parameter
 	// @param	InValue				Value to push
 	template<typename... ValueType>
-	static FLayerHandle MakeValueLayer(FName InName, ValueType&&... InValue)
+	static FParamStackLayerHandle MakeValueLayer(FName InName, ValueType&&... InValue)
 	{
 		return MakeValuesLayer(FParamId(InName), Forward<ValueType>(InValue)...);
 	}
@@ -282,48 +230,30 @@ public:
 	// Make a parameter layer from an interleaved parameter list of keys and values
 	// @param	InValues	Values to push
 	template<typename... Args>
-	static FLayerHandle MakeValuesLayer(Args&&... InValues)
+	static FParamStackLayerHandle MakeValuesLayer(Args&&... InValues)
 	{
 		constexpr int32 NumItems = sizeof...(InValues) / 2;
-		TArray<TPair<FParamId, FParam>, TInlineAllocator<NumItems>> ParamIdValues;
+		TArray<TPair<FParamId, Private::FParamEntry>, TInlineAllocator<NumItems>> ParamIdValues;
 		ParamIdValues.Reserve(NumItems);
 		return MakeValuesLayerHelper(ParamIdValues, Forward<Args>(InValues)...);
 	}
 
-	// Pop a parameter layer - asserts if the layer supplied is not the top layer
+	// Pop a parameter layer. Attempting to pop an invalid layer handle is supported.
+	// Asserts if the layer supplied is valid and not the top layer
 	ANIMNEXT_API void PopLayer(FPushedLayerHandle InLayer);
 
-	// Results for FParamStack::GetParam
-	enum class EGetParamResult : uint8
-	{
-		// The requested parameter is present in the current scope
-		Succeeded			= 0x00,
-
-		// The requested parameter is present in the current scope, and its type is compatible according to the supplied requirements
-		CompatibleType		= 0x01,
-
-		// The requested parameter is present in the current scope, but its type is not compatible according to the supplied requirements
-		IncompatibleType	= 0x02,
-
-		// The requested parameter is not present in the current scope
-		NotInScope			= 0x04,
-
-		// The requested parameter is immutable but a mutable request was made of it
-		Immutable			= 0x08,
-	};
 
-	FRIEND_ENUM_CLASS_FLAGS(FParamStack::EGetParamResult);
 
 	// Get a pointer to a parameter's value given a FParamId.
 	// @param	InParamId			Parameter ID to find the currently-pushed value for
 	// @param	OutResult			Optional ptr to aresult that will be filled if an error occurs
 	// @return nullptr if the parameter is not mapped
 	template<typename ValueType>
-	const ValueType* GetParamPtr(FParamId InParamId, FParamStack::EGetParamResult* OutResult = nullptr) const
+	const ValueType* GetParamPtr(FParamId InParamId, FParamResult* OutResult = nullptr) const
 	{
 		TConstArrayView<uint8> Data;
 		FParamTypeHandle TypeHandle;
-		FParamStack::EGetParamResult Result = GetParamData(InParamId, FParamTypeHandle::GetHandle<ValueType>(), Data, TypeHandle);
+		FParamResult Result = GetParamData(InParamId, FParamTypeHandle::GetHandle<ValueType>(), Data, TypeHandle);
 		if (OutResult)
 		{
 			*OutResult = Result;
@@ -336,11 +266,11 @@ public:
 	// @param	OutResult			Optional ptr to a result that will be filled if an error occurs
 	// @return nullptr if the parameter is not mapped or mutable
 	template<typename ValueType>
-	ValueType* GetMutableParamPtr(FParamId InParamId, FParamStack::EGetParamResult* OutResult = nullptr)
+	ValueType* GetMutableParamPtr(FParamId InParamId, FParamResult* OutResult = nullptr)
 	{
 		TArrayView<uint8> Data;
 		FParamTypeHandle TypeHandle;
-		FParamStack::EGetParamResult Result = GetMutableParamData(InParamId, FParamTypeHandle::GetHandle<ValueType>(), Data, TypeHandle);
+		FParamResult Result = GetMutableParamData(InParamId, FParamTypeHandle::GetHandle<ValueType>(), Data, TypeHandle);
 		if (OutResult)
 		{
 			*OutResult = Result;
@@ -354,7 +284,7 @@ public:
 	// @param	OutResult			Optional ptr to a result that will be filled if an error occurs
 	// @return nullptr if the parameter is not mapped
 	template<typename ValueType>
-	const ValueType* GetParamPtr(FName InKey, FParamStack::EGetParamResult* OutResult = nullptr) const
+	const ValueType* GetParamPtr(FName InKey, FParamResult* OutResult = nullptr) const
 	{
 		return GetParamPtr<ValueType>(FParamId(InKey), OutResult);
 	}
@@ -365,7 +295,7 @@ public:
 	// @param	OutResult			Optional ptr to a result that will be filled if an error occurs
 	// @return nullptr if the parameter is not mapped
 	template<typename ValueType>
-	ValueType* GetMutableParamPtr(FName InKey, FParamStack::EGetParamResult* OutResult = nullptr)
+	ValueType* GetMutableParamPtr(FName InKey, FParamResult* OutResult = nullptr)
 	{
 		return GetMutableParamPtr<ValueType>(FParamId(InKey), OutResult);
 	}
@@ -375,7 +305,7 @@ public:
 	// @param	OutResult			Optional ptr to a result that will be filled if an error occurs
 	// @return the parameter's value, or a default value if the parameter is not present
 	template<typename ValueType>
-	ValueType GetParamValue(FParamId InParamId, FParamStack::EGetParamResult* OutResult = nullptr) const
+	ValueType GetParamValue(FParamId InParamId, FParamResult* OutResult = nullptr) const
 	{
 		if (const ValueType* ParamPtr = GetParamPtr<ValueType>(InParamId, OutResult))
 		{
@@ -389,7 +319,7 @@ public:
 	// @param	OutResult			Optional ptr to a result that will be filled if an error occurs
 	// @return the parameter's value, or a default value if the parameter is not present
 	template<typename ValueType>
-	ValueType GetParamValue(FName InKey, FParamStack::EGetParamResult* OutResult = nullptr) const
+	ValueType GetParamValue(FName InKey, FParamResult* OutResult = nullptr) const
 	{
 		if (const ValueType* ParamPtr = GetParamPtr<ValueType>(FParamId(InKey), OutResult))
 		{
@@ -403,7 +333,7 @@ public:
 	// @param	OutResult			Optional ptr to a result that will be filled if an error occurs
 	// @return a const reference to the parameter's value. Function asserts if the parameter is not present
 	template<typename ValueType>
-	const ValueType& GetParam(FParamId InParamId, FParamStack::EGetParamResult* OutResult = nullptr) const
+	const ValueType& GetParam(FParamId InParamId, FParamResult* OutResult = nullptr) const
 	{
 		const ValueType* ParamPtr = GetParamPtr<ValueType>(InParamId, OutResult);
 		check(ParamPtr);
@@ -416,7 +346,7 @@ public:
 	// @param	OutResult			Optional ptr to a result that will be filled if an error occurs
 	// @return a const reference to the parameter's value. Function asserts if the parameter is not present
 	template<typename ValueType>
-	const ValueType& GetParam(FName InKey, FParamStack::EGetParamResult* OutResult = nullptr) const
+	const ValueType& GetParam(FName InKey, FParamResult* OutResult = nullptr) const
 	{
 		const ValueType* ParamPtr = GetParamPtr<ValueType>(FParamId(InKey), OutResult);
 		check(ParamPtr);
@@ -428,7 +358,7 @@ public:
 	// @param	OutResult			Optional ptr to a result that will be filled if an error occurs
 	// @return a mutable reference to the parameter's value. Function asserts if the parameter is not present
 	template<typename ValueType>
-	ValueType& GetMutableParam(FParamId InParamId, FParamStack::EGetParamResult* OutResult = nullptr)
+	ValueType& GetMutableParam(FParamId InParamId, FParamResult* OutResult = nullptr)
 	{
 		ValueType* ParamPtr = GetMutableParamPtr<ValueType>(InParamId, OutResult);
 		check(ParamPtr);
@@ -441,7 +371,7 @@ public:
 	// @param	OutResult			Optional ptr to a result that will be filled if an error occurs
 	// @return a mutable reference to the parameter's value. Function asserts if the parameter is not present
 	template<typename ValueType>
-	ValueType& GetMutableParam(FName InKey, FParamStack::EGetParamResult* OutResult = nullptr)
+	ValueType& GetMutableParam(FName InKey, FParamResult* OutResult = nullptr)
 	{
 		ValueType* ParamPtr = GetMutableParamPtr<ValueType>(FParamId(InKey), OutResult);
 		check(ParamPtr);
@@ -456,9 +386,9 @@ public:
 	// @param	OutParamTypeHandle		Handle to the type of the parameter if the parameter is present in the stack.
 	//									This will be valid even if the types are deemed incompatible.
 	// @param	InRequiredCompatibility	The minimum required compatibility level
-	// @return an enum describing the result of the operation @see FParamStack::EGetParamResult
-	ANIMNEXT_API EGetParamResult GetParamData(FParamId InId, FParamTypeHandle InTypeHandle, TConstArrayView<uint8>& OutParamData, FParamCompatibility InRequiredCompatibility = FParamCompatibility::Equal()) const;
-	ANIMNEXT_API EGetParamResult GetParamData(FParamId InId, FParamTypeHandle InTypeHandle, TConstArrayView<uint8>& OutParamData, FParamTypeHandle& OutParamTypeHandle, FParamCompatibility InRequiredCompatibility = FParamCompatibility::Equal()) const;
+	// @return an enum describing the result of the operation @see FParamResult
+	ANIMNEXT_API FParamResult GetParamData(FParamId InId, FParamTypeHandle InTypeHandle, TConstArrayView<uint8>& OutParamData, FParamCompatibility InRequiredCompatibility = FParamCompatibility::Equal()) const;
+	ANIMNEXT_API FParamResult GetParamData(FParamId InId, FParamTypeHandle InTypeHandle, TConstArrayView<uint8>& OutParamData, FParamTypeHandle& OutParamTypeHandle, FParamCompatibility InRequiredCompatibility = FParamCompatibility::Equal()) const;
 
 	// Get a parameter's mutable data given an FParamId.
 	// @param	InParamId				Parameter ID to find the currently-pushed value for
@@ -468,17 +398,52 @@ public:
 	// @param	OutParamTypeHandle		Handle to the type of the parameter if the parameter is present in the stack.
 	//									This will be valid even if the types are deemed incompatible.
 	// @param	InRequiredCompatibility	The minimum required compatibility level
-	// @return an enum describing the result of the operation @see FParamStack::EGetParamResult
-	ANIMNEXT_API EGetParamResult GetMutableParamData(FParamId InId, FParamTypeHandle InTypeHandle, TArrayView<uint8>& OutParamData, FParamCompatibility InRequiredCompatibility = FParamCompatibility::Equal());
-	ANIMNEXT_API EGetParamResult GetMutableParamData(FParamId InId, FParamTypeHandle InTypeHandle, TArrayView<uint8>& OutParamData, FParamTypeHandle& OutParamTypeHandle, FParamCompatibility InRequiredCompatibility = FParamCompatibility::Equal());
+	// @return an enum describing the result of the operation @see FParamResult
+	ANIMNEXT_API FParamResult GetMutableParamData(FParamId InId, FParamTypeHandle InTypeHandle, TArrayView<uint8>& OutParamData, FParamCompatibility InRequiredCompatibility = FParamCompatibility::Equal());
+	ANIMNEXT_API FParamResult GetMutableParamData(FParamId InId, FParamTypeHandle InTypeHandle, TArrayView<uint8>& OutParamData, FParamTypeHandle& OutParamTypeHandle, FParamCompatibility InRequiredCompatibility = FParamCompatibility::Equal());
 
 private:
+	ANIMNEXT_API FParamStack();
+
+	void SetParent(TWeakPtr<const FParamStack> InParent);
+
+	// Get the param stack that is attached to the current thread
+	static ANIMNEXT_API TWeakPtr<FParamStack> GetForCurrentThread();
+
+	// Attach a param stack to the current thread
+	static ANIMNEXT_API void AttachToCurrentThread(TWeakPtr<FParamStack> InStack);
+
+	// Attach a param stack that is 'pending' for the specified object to the current thread.
+	// @return true if the stack was successfully attached
+	static ANIMNEXT_API bool AttachToCurrentThreadForPendingObject(const UObject* InObject);
+
+	// Detach a param stack that is 'pending' for the specified object from the current thread.
+	// @return true if the stack was successfully detached
+	static ANIMNEXT_API bool DetachFromCurrentThreadForPendingObject(const UObject* InObject);
+
+	// Adds a stack asscociated with a 'pending' object - that is, an object that is due to have some
+	// asscociated logic run that needs to access the relevant stack.
+	// If the supplied stack is null, a new stack will be added, parented to the global stack.
+	static ANIMNEXT_API void AddForPendingObject(const UObject* InObject, TWeakPtr<FParamStack> InStack = nullptr);
+
+	// Removes a stack asscociated with a 'pending' object. Asserts if any existing stack still has outstanding references.
+	static ANIMNEXT_API void RemoveForPendingObject(const UObject* InObject);
+
+	// Detach a param stack to the current thread
+	static ANIMNEXT_API TWeakPtr<FParamStack> DetachFromCurrentThread();
+
+	// Create a cached parameter layer from set of params. Mutabilty is on a per-parameter basis.
+	static ANIMNEXT_API FParamStackLayerHandle MakeLayer(TConstArrayView<TPair<FParamId, Private::FParamEntry>> InParams);
+
+	// Push an internally-owned layer. Copies parameter data to internal storage.
+	ANIMNEXT_API FPushedLayerHandle PushLayer(TConstArrayView<TPair<FParamId, Private::FParamEntry>> InParams);
+
 	// Recursive helper function for PushValues
 	template <uint32 NumItems, typename FirstType, typename SecondType, typename... OtherTypes>
-	FPushedLayerHandle PushValuesHelper(TArray<TPair<FParamId, FParam>, TInlineAllocator<NumItems>>& InArray, FirstType&& InFirst, SecondType&& InSecond, OtherTypes&&... InOthers)
+	FPushedLayerHandle PushValuesHelper(TArray<TPair<FParamId, Private::FParamEntry>, TInlineAllocator<NumItems>>& InArray, FirstType&& InFirst, SecondType&& InSecond, OtherTypes&&... InOthers)
 	{
 		InArray.Emplace(FParamId(InFirst),
-			FParam(
+			Private::FParamEntry(
 				FParamTypeHandle::GetHandle<std::remove_reference_t<SecondType>>(),
 				TArrayView<uint8>(const_cast<uint8*>(reinterpret_cast<const uint8*>(&InSecond)), sizeof(std::remove_reference_t<SecondType>)),
 				std::is_reference_v<SecondType>,
@@ -498,10 +463,10 @@ private:
 
 	// Recursive helper function for MakeValuesLayer
 	template <uint32 NumItems, typename FirstType, typename SecondType, typename... OtherTypes>
-	static FLayerHandle MakeValuesLayerHelper(TArray<TPair<FParamId, FParam>, TInlineAllocator<NumItems>>& InArray, FirstType&& InFirst, SecondType&& InSecond, OtherTypes&&... InOthers)
+	static FParamStackLayerHandle MakeValuesLayerHelper(TArray<TPair<FParamId, Private::FParamEntry>, TInlineAllocator<NumItems>>& InArray, FirstType&& InFirst, SecondType&& InSecond, OtherTypes&&... InOthers)
 	{
 		InArray.Emplace(FParamId(InFirst), 
-			FParam(
+			Private::FParamEntry(
 				FParamTypeHandle::GetHandle<std::remove_reference_t<SecondType>>(),
 				TArrayView<uint8>(const_cast<uint8*>(reinterpret_cast<const uint8*>(&InSecond)), sizeof(std::remove_reference_t<SecondType>)),
 				std::is_reference_v<SecondType>,
@@ -520,7 +485,13 @@ private:
 	}
 
 	// Helper function for layer pushes
-	FPushedLayerHandle PushLayerInternal(FLayer& InLayer);
+	FPushedLayerHandle PushLayerInternal(FParamStackLayer& InLayer);
+
+	// Internal helper function for GetParamData
+	FParamResult GetParamDataInternal(FParamId InId, FParamTypeHandle InTypeHandle, TConstArrayView<uint8>& OutParamData, FParamTypeHandle& OutParamTypeHandle, FParamCompatibility InRequiredCompatibility) const;
+
+	// Internal helper function for GetMutableParamData
+	FParamResult GetMutableParamDataInternal(FParamId InId, FParamTypeHandle InTypeHandle, TArrayView<uint8>& OutParamData, FParamTypeHandle& OutParamTypeHandle, FParamCompatibility InRequiredCompatibility);
 
 	// Check whether a given parameter is mutable
 	ANIMNEXT_API bool IsMutableParam(FParamId InId) const;
@@ -540,9 +511,12 @@ private:
 		return IsReferenceParam(FParamId(InKey));
 	}
 
+	// Check the layer to see if it contains the supplied param.
+	static bool LayerContainsParam(const FParamStackLayerHandle& InHandle, FName InKey);
+
 	// Allocates any necessary owned storage for non-reference, non-embeded parameters that a layer holds
 	// @return the base offset of the storage area
-	uint32 AllocAndCopyOwnedParamStorage(TArrayView<FParam> InParams);
+	uint32 AllocAndCopyOwnedParamStorage(TArrayView<Private::FParamEntry> InParams);
 
 	// Frees any owned storage, restores owned storage offset
 	void FreeOwnedParamStorage(uint32 InOffset);
@@ -558,7 +532,7 @@ private:
 
 	// Owned layers, paged for stable addresses as layers reference them by ptr
 	// Grows with each new pushed owned layer
-	TPagedArray<FLayer, 4096> OwnedStackLayers;
+	TPagedArray<FParamStackLayer, 4096> OwnedStackLayers;
 
 	// Parameter storage for owned layers, paged for stable addresses as layer parameters reference them by ptr
 	// Grows with each new pushed layer, free'd when the param stack is destroyed
@@ -573,11 +547,11 @@ private:
 	// Grows with each new pushed layer, free'd when the param stack is destroyed
 	TArray<uint16> PreviousLayerIndices;
 
+	// Parent stack from outer scope
+	TWeakPtr<const FParamStack> WeakParentStack;
+
 	// Serial number for stack layers
 	uint32 SerialNumber = 0;
 };
 
-ENUM_CLASS_FLAGS(FParamStack::EGetParamResult);
-ENUM_CLASS_FLAGS(FParamStack::EParamFlags);
-
 }
diff --git a/Engine/Plugins/Experimental/AnimNext/Source/AnimNext/Public/Param/ParamStackLayerHandle.h b/Engine/Plugins/Experimental/AnimNext/Source/AnimNext/Public/Param/ParamStackLayerHandle.h
new file mode 100644
--- /dev/null
+++ b/Engine/Plugins/Experimental/AnimNext/Source/AnimNext/Public/Param/ParamStackLayerHandle.h
@@ -0,0 +1,142 @@
+// Copyright Epic Games, Inc. All Rights Reserved.
+
+#pragma once
+
+#include "CoreMinimal.h"
+#include "Param/ParamResult.h"
+#include "Param/ParamId.h"
+#include "Param/ParamTypeHandle.h"
+
+class UAnimNextParameterBlock;
+
+namespace UE::AnimNext
+{
+	struct FParamStackLayer;
+}
+
+namespace UE::AnimNext::Private
+{
+	struct FParamEntry;
+}
+
+namespace UE::AnimNext
+{
+
+// Opaque handle to a layer that can be held by external systems
+// Move-only, single ownership
+struct FParamStackLayerHandle
+{
+public:
+	ANIMNEXT_API FParamStackLayerHandle();
+	FParamStackLayerHandle(const FParamStackLayerHandle& InOther) = delete;
+	FParamStackLayerHandle& operator=(const FParamStackLayerHandle& InOther) = delete;
+	ANIMNEXT_API FParamStackLayerHandle(FParamStackLayerHandle&& InLayer) noexcept;
+	ANIMNEXT_API FParamStackLayerHandle& operator=(FParamStackLayerHandle&& InLayer) noexcept;
+	ANIMNEXT_API ~FParamStackLayerHandle();
+
+private:
+	FParamStackLayerHandle(TUniquePtr<FParamStackLayer>&& InLayer);
+
+public:
+	// Check this layer's validity
+	bool IsValid() const;
+
+	// Invalidate the layer
+	void Invalidate();
+
+	// Set a parameter's value in this layer
+	// @param	InParamId			Parameter ID for the parameter
+	// @param	InValue				Value to set
+	template<typename... ValueType>
+	FParamResult SetValue(FParamId InParamId, ValueType&&... InValue)
+	{
+		return SetValues(InParamId, Forward<ValueType>(InValue)...);
+	}
+
+	// Set a parameter's value in this layer
+	// @param	InName				Parameter name for the parameter
+	// @param	InValue				Value to set
+	template<typename... ValueType>
+	FParamResult SetValue(FName InName, ValueType&&... InValue)
+	{
+		return SetValues(FParamId(InName), Forward<ValueType>(InValue)...);
+	}
+
+	// Set a parameter's value in this layer
+	// @param	InParamId			Parameter ID for the parameter
+	// @param	InTypeHandle		The type of the parameter
+	// @param	InData				Value to set
+	ANIMNEXT_API FParamResult SetValueRaw(FParamId InParamId, FParamTypeHandle InTypeHandle, TConstArrayView<uint8> InData);
+
+	// Get a parameter's value from this layer
+	// @param	InParamId			Parameter ID for the parameter
+	// @param	InTypeHandle		The type of the parameter
+	// @param	InData				Value to get
+	ANIMNEXT_API FParamResult GetValueRaw(FParamId InParamId, FParamTypeHandle InTypeHandle, TConstArrayView<uint8>& OutData);
+
+	// Set multiple parameter values in this layer
+	// @param	InValues	Names & values to set, interleaved
+	template<typename... Args>
+	FParamResult SetValues(Args&&... InValues)
+	{
+		constexpr int32 NumItems = sizeof...(InValues) / 2;
+		TArray<TPair<FParamId, Private::FParamEntry>, TInlineAllocator<NumItems>> ParamIdValues;
+		ParamIdValues.Reserve(NumItems);
+		return SetValuesHelper(ParamIdValues, Forward<Args>(InValues)...);
+	}
+
+	// Get the underlying storage for this layer.
+	// @return nullptr if the layer does not use the supplied underlying storage type
+	template<typename WrappedType>
+	WrappedType* As() const
+	{
+		if constexpr (TModels<CStaticClassProvider, WrappedType>::Value)
+		{
+			return Cast<WrappedType>(GetUObjectFromLayer());
+		}
+		else if constexpr (std::is_same_v<WrappedType, FInstancedPropertyBag>)
+		{
+			return GetInstancedPropertyBagFromLayer();
+		}
+
+		return nullptr;
+	}
+
+private:
+	friend struct FParamStack;
+	friend class UAnimNextParameterBlock;
+
+	// Set parameter values
+	ANIMNEXT_API FParamResult SetValuesInternal(TConstArrayView<TPair<FParamId, Private::FParamEntry>> InParams);
+
+	// Helpers for As() to hide Layer
+	ANIMNEXT_API UObject* GetUObjectFromLayer() const;
+	ANIMNEXT_API FInstancedPropertyBag* GetInstancedPropertyBagFromLayer() const;
+
+	// Recursive helper function for SetValues
+	template <uint32 NumItems, typename FirstType, typename SecondType, typename... OtherTypes>
+	FParamResult SetValuesHelper(TArray<TPair<FParamId, Private::FParamEntry>, TInlineAllocator<NumItems>>& InArray, FirstType&& InFirst, SecondType&& InSecond, OtherTypes&&... InOthers)
+	{
+		InArray.Emplace(FParamId(InFirst),
+			Private::FParamEntry(
+				FParamTypeHandle::GetHandle<std::remove_reference_t<SecondType>>(),
+				TArrayView<uint8>(const_cast<uint8*>(reinterpret_cast<const uint8*>(&InSecond)), sizeof(std::remove_reference_t<SecondType>)),
+				std::is_reference_v<SecondType>,
+				!std::is_const_v<std::remove_reference_t<SecondType>>
+			)
+		);
+
+		if constexpr (sizeof...(InOthers) > 0)
+		{
+			return SetValuesHelper(InArray, Forward<OtherTypes>(InOthers)...);
+		}
+		else
+		{
+			return SetValuesInternal(InArray);
+		}
+	}
+
+	TUniquePtr<FParamStackLayer> Layer;
+};
+
+}
\ No newline at end of file
diff --git a/Engine/Plugins/Experimental/AnimNext/Source/AnimNext/Public/Param/ParamType.h b/Engine/Plugins/Experimental/AnimNext/Source/AnimNext/Public/Param/ParamType.h
--- a/Engine/Plugins/Experimental/AnimNext/Source/AnimNext/Public/Param/ParamType.h
+++ b/Engine/Plugins/Experimental/AnimNext/Source/AnimNext/Public/Param/ParamType.h
@@ -180,6 +180,14 @@ private:
 			ParameterType.ValueType = EValueType::SoftClass;
 			ParameterType.ValueTypeObject = ParamType::ElementType::StaticClass();
 		}
+		else if constexpr (std::is_same_v<ParamType, uint32>)
+		{
+			ParameterType.ValueType = EValueType::UInt32;
+		}
+		else if constexpr (std::is_same_v<ParamType, uint64>)
+		{
+			ParameterType.ValueType = EValueType::UInt64;
+		}
 		else
 		{
 			static_assert(sizeof(ParamType) == 0, "Type is not expressible as a FAnimNextParamType for available types.");
@@ -265,6 +273,9 @@ public:
 	/** Get a string representing this type */
 	FString ToString() const;
 
+	/** Get a type from a string */
+	static FAnimNextParamType FromString(const FString& InString);
+
 	/** Equality operator */
 	friend bool operator==(const FAnimNextParamType& InLHS, const FAnimNextParamType& InRHS)
 	{
@@ -291,5 +302,8 @@ public:
 		const bool bHasValidObjectType = (ValueType < EValueType::Enum) || (ValueType >= EValueType::Enum && ValueType <= EValueType::SoftClass && IsValidObject());
 		return bHasValidValueType && bHasValidContainerType && bHasValidObjectType;
 	}
+
+	/** @return whether this type represents an object (object/class/softobject/softclass) */
+	bool IsObjectType() const;
 };
 
diff --git a/Engine/Plugins/Experimental/AnimNext/Source/AnimNext/Public/Param/ParamTypeHandle.h b/Engine/Plugins/Experimental/AnimNext/Source/AnimNext/Public/Param/ParamTypeHandle.h
--- a/Engine/Plugins/Experimental/AnimNext/Source/AnimNext/Public/Param/ParamTypeHandle.h
+++ b/Engine/Plugins/Experimental/AnimNext/Source/AnimNext/Public/Param/ParamTypeHandle.h
@@ -326,6 +326,9 @@ public:
 	/** Make a parameter type handle from a FProperty */
 	static FParamTypeHandle FromProperty(const FProperty* InProperty);
 
+	/** Make a parameter type handle from a UObject */
+	static FParamTypeHandle FromObject(const UObject* InObject);
+	
 	/** Return the raw value that this handle uses to represent itself */
 	uint32 ToRaw() const
 	{
diff --git a/Engine/Plugins/Experimental/AnimNext/Source/AnimNextEditor/AnimNextEditor.Build.cs b/Engine/Plugins/Experimental/AnimNext/Source/AnimNextEditor/AnimNextEditor.Build.cs
--- a/Engine/Plugins/Experimental/AnimNext/Source/AnimNextEditor/AnimNextEditor.Build.cs
+++ b/Engine/Plugins/Experimental/AnimNext/Source/AnimNextEditor/AnimNextEditor.Build.cs
@@ -38,6 +38,7 @@ namespace UnrealBuildTool.Rules
 					"DesktopWidgets",
 					"ContentBrowserFileDataSource",
 					"SubobjectEditor",
+					"Settings",
 				}
 			);
 		}
diff --git a/Engine/Plugins/Experimental/AnimNext/Source/AnimNextEditor/Internal/Param/ParameterPickerArgs.h b/Engine/Plugins/Experimental/AnimNext/Source/AnimNextEditor/Internal/Param/ParameterPickerArgs.h
--- a/Engine/Plugins/Experimental/AnimNext/Source/AnimNextEditor/Internal/Param/ParameterPickerArgs.h
+++ b/Engine/Plugins/Experimental/AnimNext/Source/AnimNextEditor/Internal/Param/ParameterPickerArgs.h
@@ -15,7 +15,7 @@ struct FParameterBindingReference
 {
 	FParameterBindingReference() = default;
 
-	FParameterBindingReference(const FName& InParameter, const FAssetData& InLibrary, const FAssetData& InBlock = FAssetData())
+	FParameterBindingReference(const FName& InParameter, const FAssetData& InLibrary = FAssetData(), const FAssetData& InBlock = FAssetData())
 		: Parameter(InParameter)
 		, Library(InLibrary)
 		, Block(InBlock)
@@ -25,7 +25,7 @@ struct FParameterBindingReference
 	// Parameter name
 	FName Parameter;
 
-	// Library asset
+	// Optional library asset (as this can represent built-in parameters too)
 	FAssetData Library;
 
 	// Optional block asset that the parameter is bound in
@@ -36,7 +36,7 @@ struct FParameterBindingReference
 DECLARE_DELEGATE_OneParam(FOnGetParameterBindings, TArray<FParameterBindingReference>& /*OutParameterBindings*/);
 
 // Delegate called when a parameter has been picked. Block argument is invalid when an unbound parameter is chosen.
-DECLARE_DELEGATE_OneParam(FOnParameterPicked, const FParameterBindingReference& /*OutParameterBinding*/);
+DECLARE_DELEGATE_OneParam(FOnParameterPicked, const FParameterBindingReference& /*InParameterBinding*/);
 
 // Result of a filter operation via FOnFilterParameter
 enum class EFilterParameterResult : int32
@@ -80,11 +80,17 @@ struct FParameterPickerArgs
 	// Whether we should show parameters that are not bound in a parameter block (if bShowBoundParameters is false this will show all parameters)
 	bool bShowUnboundParameters = true;	
 
+	// Whether we should show parameters that are built in
+	bool bShowBuiltInParameters = true;
+
 	// Whether we should show the library alongside parameters
 	bool bShowLibraries = true;
 
 	// Whether we should show the block alongside bound parameters
 	bool bShowBlocks = true;
+
+	// Whether we should allow new parameters to be created by this widget
+	bool bAllowNew = true;
 };
 
 }
diff --git a/Engine/Plugins/Experimental/AnimNext/Source/AnimNextEditor/Internal/Scheduler/AnimNextScheduleFactory.h b/Engine/Plugins/Experimental/AnimNext/Source/AnimNextEditor/Internal/Scheduler/AnimNextScheduleFactory.h
new file mode 100644
--- /dev/null
+++ b/Engine/Plugins/Experimental/AnimNext/Source/AnimNextEditor/Internal/Scheduler/AnimNextScheduleFactory.h
@@ -0,0 +1,22 @@
+// Copyright Epic Games, Inc. All Rights Reserved.
+
+#pragma once
+
+#include "CoreMinimal.h"
+#include "Factories/Factory.h"
+#include "AnimNextScheduleFactory.generated.h"
+
+UCLASS(MinimalAPI, BlueprintType)
+class UAnimNextScheduleFactory : public UFactory
+{
+	GENERATED_BODY()
+
+	UAnimNextScheduleFactory();
+
+	virtual bool ConfigureProperties() override;
+	virtual UObject* FactoryCreateNew(UClass* Class, UObject* InParent, FName Name, EObjectFlags Flags, UObject* Context, FFeedbackContext* Warn, FName CallingContext) override;
+	virtual UObject* FactoryCreateNew(UClass* Class, UObject* InParent, FName Name, EObjectFlags Flags, UObject* Context, FFeedbackContext* Warn) override
+	{
+		return FactoryCreateNew(Class, InParent, Name, Flags, Context, Warn, NAME_None);
+	}
+};
\ No newline at end of file
diff --git a/Engine/Plugins/Experimental/AnimNext/Source/AnimNextEditor/Private/AnimNextEditorModule.cpp b/Engine/Plugins/Experimental/AnimNext/Source/AnimNextEditor/Private/AnimNextEditorModule.cpp
--- a/Engine/Plugins/Experimental/AnimNext/Source/AnimNextEditor/Private/AnimNextEditorModule.cpp
+++ b/Engine/Plugins/Experimental/AnimNext/Source/AnimNextEditor/Private/AnimNextEditorModule.cpp
@@ -1,23 +1,36 @@
 // Copyright Epic Games, Inc. All Rights Reserved.
 
 #include "AnimNextEditorModule.h"
+
+#include "AnimNextConfig.h"
 #include "AssetToolsModule.h"
 #include "IAssetTools.h"
 #include "Graph/AssetTypeActions.h"
 #include "Graph/AnimNextGraphPanelNodeFactory.h"
 #include "Param/ParamTypePropertyCustomization.h"
 #include "Param/ParameterPickerArgs.h"
+#include "Param/ParametersGraphPanelPinFactory.h"
+#include "Param/ParamNamePropertyCustomization.h"
 #include "Param/SParameterPicker.h"
-#include "Param/ParamType.h"
+#include "ISettingsModule.h"
+
+#define LOCTEXT_NAMESPACE "AnimNextEditorModule"
 
 namespace UE::AnimNext::Editor
 {
 
 class FModule : public IModule
 {
-
 	virtual void StartupModule() override
 	{
+		// Register settings for user editing
+		ISettingsModule& SettingsModule = FModuleManager::Get().LoadModuleChecked<ISettingsModule>("Settings");
+		SettingsModule.RegisterSettings("Editor", "General", "AnimNext",
+			LOCTEXT("SettingsName", "AnimNext"),
+			LOCTEXT("SettingsDescription", "Customize AnimNext Settings."),
+			GetMutableDefault<UAnimNextConfig>()
+		);
+		
 		IAssetTools& AssetTools = FModuleManager::LoadModuleChecked<FAssetToolsModule>("AssetTools").Get();
 		AssetTypeActions_AnimNextGraph = MakeShared<FAssetTypeActions_AnimNextGraph>();
 		AssetTools.RegisterAssetTypeActions(AssetTypeActions_AnimNextGraph.ToSharedRef());
@@ -28,8 +41,17 @@ class FModule : public IModule
 			"AnimNextParamType",
 			FOnGetPropertyTypeCustomizationInstance::CreateLambda([] { return MakeShared<FParamTypePropertyTypeCustomization>(); }));
 
+		Identifier = MakeShared<FParamNamePropertyTypeIdentifier>();
+		PropertyModule.RegisterCustomPropertyTypeLayout(
+			FNameProperty::StaticClass()->GetFName(),
+			FOnGetPropertyTypeCustomizationInstance::CreateLambda([] { return MakeShared<FParamNamePropertyTypeCustomization>(); }),
+			Identifier);
+
 		AnimNextGraphPanelNodeFactory = MakeShared<FAnimNextGraphPanelNodeFactory>();
 		FEdGraphUtilities::RegisterVisualNodeFactory(AnimNextGraphPanelNodeFactory);
+
+		ParametersGraphPanelPinFactory = MakeShared<FParametersGraphPanelPinFactory>();
+		FEdGraphUtilities::RegisterVisualPinFactory(ParametersGraphPanelPinFactory);
 	}
 
 	virtual void ShutdownModule() override
@@ -47,6 +69,8 @@ class FModule : public IModule
 		}
 
 		FEdGraphUtilities::UnregisterVisualNodeFactory(AnimNextGraphPanelNodeFactory);
+
+		FEdGraphUtilities::UnregisterVisualPinFactory(ParametersGraphPanelPinFactory);
 	}
 
 	virtual TSharedRef<SWidget> CreateParameterPicker(const FParameterPickerArgs& InArgs) override
@@ -59,8 +83,16 @@ class FModule : public IModule
 
 	/** Node factory for the AnimNext graph */
 	TSharedPtr<FAnimNextGraphPanelNodeFactory> AnimNextGraphPanelNodeFactory;
+	
+	/** Pin factory for parameters */
+	TSharedPtr<FParametersGraphPanelPinFactory> ParametersGraphPanelPinFactory;
+
+	/** Type identifier for parameter names */
+	TSharedPtr<FParamNamePropertyTypeIdentifier> Identifier;
 };
 
 }
 
-IMPLEMENT_MODULE(UE::AnimNext::Editor::FModule, AnimNextEditor);
\ No newline at end of file
+IMPLEMENT_MODULE(UE::AnimNext::Editor::FModule, AnimNextEditor);
+
+#undef LOCTEXT_NAMESPACE
\ No newline at end of file
diff --git a/Engine/Plugins/Experimental/AnimNext/Source/AnimNextEditor/Private/EditorUtils.cpp b/Engine/Plugins/Experimental/AnimNext/Source/AnimNextEditor/Private/EditorUtils.cpp
--- a/Engine/Plugins/Experimental/AnimNext/Source/AnimNextEditor/Private/EditorUtils.cpp
+++ b/Engine/Plugins/Experimental/AnimNext/Source/AnimNextEditor/Private/EditorUtils.cpp
@@ -8,6 +8,9 @@
 #include "Param/AnimNextParameterBlock.h"
 #include "Param/AnimNextParameterLibrary.h"
 #include "PropertyBagDetails.h"
+#include "AssetRegistry/AssetRegistryModule.h"
+#include "AssetRegistry/IAssetRegistry.h"
+#include "Param/Params.h"
 
 namespace UE::AnimNext::Editor
 {
@@ -233,7 +236,7 @@ FAnimNextParamType FUtils::GetParameterTypeFromMetaData(const FStringView& InStr
 void FUtils::GetAllGraphNames(const UAnimNextParameterBlock_EditorData* InEditorData, TSet<FName>& OutNames)
 {
 	// TODO
-	ensure(false);
+//	ensure(false);
 }
 
 FName FUtils::ValidateName(const UAnimNextParameterBlock_EditorData* InEditorData, const FString& InName)
@@ -503,4 +506,37 @@ bool FUtils::GetExportedBindingsForBlock(const FAssetData& InLibraryAsset, FAnim
 	return FAnimNextParameterBlockAssetRegistryExports::StaticStruct()->ImportText(*TagValue, &OutExports, nullptr, PPF_None, nullptr, FAnimNextParameterBlockAssetRegistryExports::StaticStruct()->GetName()) != nullptr;
 }
 
+FAnimNextParamType FUtils::GetParameterTypeFromName(FName InName)
+{
+	// Check built-in params first as they are cheaper
+	if(const FParamDefinition* FoundDefinition = FParams::FindBuiltInParameter(InName))
+	{
+		return FoundDefinition->Type;
+	}
+
+	// Query the asset registry for other params
+	IAssetRegistry& AssetRegistry = FModuleManager::Get().LoadModuleChecked<FAssetRegistryModule>(TEXT("AssetRegistry")).Get();
+
+	FARFilter ARFilter;
+	ARFilter.ClassPaths = { UAnimNextParameterLibrary::StaticClass()->GetClassPathName() };
+
+	TArray<FAssetData> LibraryAssets;
+	AssetRegistry.GetAssets(ARFilter, LibraryAssets);
+
+	for(const FAssetData& LibraryAsset : LibraryAssets)
+	{
+		FAnimNextParameterLibraryAssetRegistryExports Exports;
+		GetExportedParametersForLibrary(LibraryAsset, Exports);
+		for(const FAnimNextParameterLibraryAssetRegistryExportEntry& Export : Exports.Parameters)
+		{
+			if(Export.Name == InName)
+			{
+				return Export.Type;
+			}
+		}
+	}
+
+	return FAnimNextParamType();
+}
+
 }
diff --git a/Engine/Plugins/Experimental/AnimNext/Source/AnimNextEditor/Private/EditorUtils.h b/Engine/Plugins/Experimental/AnimNext/Source/AnimNextEditor/Private/EditorUtils.h
--- a/Engine/Plugins/Experimental/AnimNext/Source/AnimNextEditor/Private/EditorUtils.h
+++ b/Engine/Plugins/Experimental/AnimNext/Source/AnimNextEditor/Private/EditorUtils.h
@@ -44,6 +44,11 @@ struct FUtils
 	static FAnimNextParamType GetParameterTypeFromLibraryExports(FName InName, const FAnimNextParameterLibraryAssetRegistryExports& InExports);
 
 	static bool GetExportedBindingsForBlock(const FAssetData& InBlockAsset, FAnimNextParameterBlockAssetRegistryExports& OutExports);
+
+	// Attempts to determine the type from a parameter name
+	// If the name cannot be found, the returned type will be invalid
+	// Note that this is expensive and can query the asset registry
+	static FAnimNextParamType GetParameterTypeFromName(FName InName);
 };
 
 }
diff --git a/Engine/Plugins/Experimental/AnimNext/Source/AnimNextEditor/Private/Param/ParamNamePropertyCustomization.cpp b/Engine/Plugins/Experimental/AnimNext/Source/AnimNextEditor/Private/Param/ParamNamePropertyCustomization.cpp
new file mode 100644
--- /dev/null
+++ b/Engine/Plugins/Experimental/AnimNext/Source/AnimNextEditor/Private/Param/ParamNamePropertyCustomization.cpp
@@ -0,0 +1,117 @@
+﻿// Copyright Epic Games, Inc. All Rights Reserved.
+
+#include "ParamNamePropertyCustomization.h"
+
+#include "DetailWidgetRow.h"
+#include "EdGraphSchema_K2.h"
+#include "EditorUtils.h"
+#include "PropertyHandle.h"
+#include "SParameterPickerCombo.h"
+#include "UncookedOnlyUtils.h"
+#include "Param/ParamCompatibility.h"
+#include "Param/ParamUtils.h"
+#include "RigVMCore/RigVMStruct.h"
+
+#define LOCTEXT_NAMESPACE "ParamNamePropertyCustomization"
+
+namespace UE::AnimNext::Editor
+{
+
+static const FProperty* GetMetadataProperty(const FProperty* Property)
+{
+	if (FProperty* OuterProperty = Property->GetOwner<FProperty>())
+	{
+		if (OuterProperty->IsA<FArrayProperty>()
+			|| OuterProperty->IsA<FSetProperty>()
+			|| OuterProperty->IsA<FMapProperty>())
+		{
+			return OuterProperty;
+		}
+	}
+
+	return Property;
+}
+
+bool FParamNamePropertyTypeIdentifier::IsPropertyTypeCustomized(const IPropertyHandle& InPropertyHandle) const
+{
+	const FProperty* Property = GetMetadataProperty(InPropertyHandle.GetProperty());
+	if(Property->GetMetaData(FRigVMStruct::CustomWidgetMetaName) == "ParamName")
+	{
+		return true;
+	}
+
+	return false;
+}
+
+void FParamNamePropertyTypeCustomization::CustomizeHeader(TSharedRef<IPropertyHandle> InPropertyHandle, FDetailWidgetRow& InHeaderRow, IPropertyTypeCustomizationUtils& InCustomizationUtils)
+{
+	WeakPropertyHandle = InPropertyHandle;
+	
+	const FProperty* Property = GetMetadataProperty(InPropertyHandle->GetProperty());
+	const FString ParamTypeString = Property->GetMetaData("AllowedParamType");
+	FAnimNextParamType FilterType = FAnimNextParamType::FromString(ParamTypeString);
+	
+	FParameterPickerArgs PickerArgs;
+	PickerArgs.bShowLibraries = false;
+	PickerArgs.bShowBlocks = false;
+	PickerArgs.bMultiSelect = false;
+	PickerArgs.OnFilterParameterType = FOnFilterParameterType::CreateLambda([FilterType](const FAnimNextParamType& InParameterType)
+	{
+		if(!FilterType.IsValid() || FParamUtils::GetCompatibility(FilterType, InParameterType).IsCompatible())
+		{
+			return EFilterParameterResult::Include;
+		}
+		return EFilterParameterResult::Exclude;
+	});
+	PickerArgs.OnParameterPicked = FOnParameterPicked::CreateLambda([this, WeakPropertyHandle = WeakPropertyHandle](const FParameterBindingReference& InParameterBinding)
+	{
+		if(TSharedPtr<IPropertyHandle> PropertyHandle = WeakPropertyHandle.Pin())
+		{
+			PropertyHandle->NotifyPreChange();
+			PropertyHandle->SetValue(InParameterBinding.Parameter);
+			PropertyHandle->NotifyPostChange(EPropertyChangeType::ValueSet);
+
+			Refresh();
+		}
+	});
+	
+	InHeaderRow
+	.NameContent()
+	[
+		InPropertyHandle->CreatePropertyNameWidget()
+	]
+	.ValueContent()
+	[
+		SNew(SParameterPickerCombo)
+		.PickerArgs(PickerArgs)
+		.OnGetParameterName_Lambda([this]()
+		{
+			return CachedName;
+		})
+		.OnGetParameterType_Lambda([this]()
+		{
+			return CachedType;
+		})
+	];
+
+	Refresh();
+}
+
+void FParamNamePropertyTypeCustomization::CustomizeChildren(TSharedRef<IPropertyHandle> PropertyHandle, IDetailChildrenBuilder& ChildBuilder, IPropertyTypeCustomizationUtils& CustomizationUtils)
+{
+}
+
+void FParamNamePropertyTypeCustomization::Refresh()
+{
+	if(TSharedPtr<IPropertyHandle> PropertyHandle = WeakPropertyHandle.Pin())
+	{
+		if(PropertyHandle->GetValue(CachedName) == FPropertyAccess::Success)
+		{
+			CachedType = FUtils::GetParameterTypeFromName(CachedName);
+		}
+	}
+}
+
+}
+
+#undef LOCTEXT_NAMESPACE
\ No newline at end of file
diff --git a/Engine/Plugins/Experimental/AnimNext/Source/AnimNextEditor/Private/Param/ParamNamePropertyCustomization.h b/Engine/Plugins/Experimental/AnimNext/Source/AnimNextEditor/Private/Param/ParamNamePropertyCustomization.h
new file mode 100644
--- /dev/null
+++ b/Engine/Plugins/Experimental/AnimNext/Source/AnimNextEditor/Private/Param/ParamNamePropertyCustomization.h
@@ -0,0 +1,34 @@
+﻿// Copyright Epic Games, Inc. All Rights Reserved.
+
+#pragma once
+
+#include "IPropertyTypeCustomization.h"
+#include "PropertyEditorModule.h"
+#include "Param/ParamType.h"
+
+namespace UE::AnimNext::Editor
+{
+
+class FParamNamePropertyTypeIdentifier : public IPropertyTypeIdentifier
+{
+private:
+	virtual bool IsPropertyTypeCustomized(const IPropertyHandle& InPropertyHandle) const override;
+};
+
+class FParamNamePropertyTypeCustomization : public IPropertyTypeCustomization
+{
+private:
+	// IPropertyTypeCustomization interface
+	virtual void CustomizeHeader(TSharedRef<IPropertyHandle> PropertyHandle, FDetailWidgetRow& HeaderRow, IPropertyTypeCustomizationUtils& CustomizationUtils) override;
+	virtual void CustomizeChildren(TSharedRef<IPropertyHandle> PropertyHandle, IDetailChildrenBuilder& ChildBuilder, IPropertyTypeCustomizationUtils& CustomizationUtils) override;
+
+	void Refresh();
+	
+private:
+	TWeakPtr<IPropertyHandle> WeakPropertyHandle;
+	FName CachedName = NAME_None;
+	FText CachedNameText;
+	FAnimNextParamType CachedType;
+};
+
+}
diff --git a/Engine/Plugins/Experimental/AnimNext/Source/AnimNextEditor/Private/Param/ParameterBlockFactory.cpp b/Engine/Plugins/Experimental/AnimNext/Source/AnimNextEditor/Private/Param/ParameterBlockFactory.cpp
--- a/Engine/Plugins/Experimental/AnimNext/Source/AnimNextEditor/Private/Param/ParameterBlockFactory.cpp
+++ b/Engine/Plugins/Experimental/AnimNext/Source/AnimNextEditor/Private/Param/ParameterBlockFactory.cpp
@@ -3,7 +3,7 @@
 #include "Param/ParameterBlockFactory.h"
 #include "Param/AnimNextParameterBlock.h"
 #include "Param/AnimNextParameterBlock_EditorData.h"
-#include "Graph/AnimNextExecuteContext.h"
+#include "Param/AnimNextParameterExecuteContext.h"
 #include "UncookedOnlyUtils.h"
 
 UAnimNextParameterBlockFactory::UAnimNextParameterBlockFactory()
@@ -25,7 +25,7 @@ UObject* UAnimNextParameterBlockFactory::FactoryCreateNew(UClass* Class, UObject
 	UAnimNextParameterBlock_EditorData* EditorData = NewObject<UAnimNextParameterBlock_EditorData>(NewBlock, TEXT("EditorData"), RF_Transactional);
 	NewBlock->EditorData = EditorData;
 	EditorData->Initialize(/*bRecompileVM*/false);
-	EditorData->GetRigVMClient()->SetExecuteContextStruct(FAnimNextExecuteContext::StaticStruct());
+	EditorData->GetRigVMClient()->SetExecuteContextStruct(FAnimNextParameterExecuteContext::StaticStruct());
 
 	// Compile the initial skeleton
 	UE::AnimNext::UncookedOnly::FUtils::Compile(NewBlock);
diff --git a/Engine/Plugins/Experimental/AnimNext/Source/AnimNextEditor/Private/Param/ParametersGraphPanelPinFactory.cpp b/Engine/Plugins/Experimental/AnimNext/Source/AnimNextEditor/Private/Param/ParametersGraphPanelPinFactory.cpp
new file mode 100644
--- /dev/null
+++ b/Engine/Plugins/Experimental/AnimNext/Source/AnimNextEditor/Private/Param/ParametersGraphPanelPinFactory.cpp
@@ -0,0 +1,37 @@
+﻿// Copyright Epic Games, Inc. All Rights Reserved.
+
+#include "ParametersGraphPanelPinFactory.h"
+#include "SGraphPinParamName.h"
+#include "EdGraphSchema_K2.h"
+#include "EdGraph/RigVMEdGraph.h"
+#include "EdGraph/RigVMEdGraphNode.h"
+
+namespace UE::AnimNext::Editor
+{
+
+FName FParametersGraphPanelPinFactory::GetFactoryName() const
+{
+	return TEXT("ParametersGraphPanelPinFactory");
+}
+
+TSharedPtr<SGraphPin> FParametersGraphPanelPinFactory::CreatePin_Internal(UEdGraphPin* InPin) const
+{
+	if (URigVMEdGraphNode* RigNode = Cast<URigVMEdGraphNode>(InPin->GetOwningNode()))
+	{
+		URigVMEdGraph* RigGraph = Cast<URigVMEdGraph>(RigNode->GetGraph());
+
+		URigVMPin* ModelPin = RigNode->GetModelPinFromPinPath(InPin->GetName());
+		if (ModelPin)
+		{
+			if(ModelPin->GetCustomWidgetName() == "ParamName")
+			{
+				return SNew(SGraphPinParamName, InPin)
+					.ModelPin(ModelPin);
+			}
+		}
+	}
+
+	return FRigVMEdGraphPanelPinFactory::CreatePin_Internal(InPin);
+}
+
+}
diff --git a/Engine/Plugins/Experimental/AnimNext/Source/AnimNextEditor/Private/Param/ParametersGraphPanelPinFactory.h b/Engine/Plugins/Experimental/AnimNext/Source/AnimNextEditor/Private/Param/ParametersGraphPanelPinFactory.h
new file mode 100644
--- /dev/null
+++ b/Engine/Plugins/Experimental/AnimNext/Source/AnimNextEditor/Private/Param/ParametersGraphPanelPinFactory.h
@@ -0,0 +1,18 @@
+﻿// Copyright Epic Games, Inc. All Rights Reserved.
+
+#pragma once
+
+#include "EdGraph/RigVMEdGraphPanelPinFactory.h"
+
+class SGraphPin;
+
+namespace UE::AnimNext::Editor
+{
+
+class FParametersGraphPanelPinFactory : public FRigVMEdGraphPanelPinFactory
+{
+	virtual FName GetFactoryName() const;
+	virtual TSharedPtr<SGraphPin> CreatePin_Internal(UEdGraphPin* InPin) const override;
+};
+
+}
\ No newline at end of file
diff --git a/Engine/Plugins/Experimental/AnimNext/Source/AnimNextEditor/Private/Param/SAddParametersDialog.cpp b/Engine/Plugins/Experimental/AnimNext/Source/AnimNextEditor/Private/Param/SAddParametersDialog.cpp
--- a/Engine/Plugins/Experimental/AnimNext/Source/AnimNextEditor/Private/Param/SAddParametersDialog.cpp
+++ b/Engine/Plugins/Experimental/AnimNext/Source/AnimNextEditor/Private/Param/SAddParametersDialog.cpp
@@ -11,6 +11,7 @@
 #include "EditorUtils.h"
 #include "UncookedOnlyUtils.h"
 #include "PropertyBagDetails.h"
+#include "SParameterPickerCombo.h"
 #include "SPinTypeSelector.h"
 #include "Framework/Application/SlateApplication.h"
 #include "Framework/Docking/TabManager.h"
@@ -38,6 +39,26 @@ static FName Column_Library(TEXT("Library"));
 static FName SelectLibraryMenuName(TEXT("AnimNext.AddParametersDialog.SelectedLibraryMenu"));
 }
 
+bool FParameterToAdd::IsValid(FText& OutReason) const
+{
+	if(Name == NAME_None)
+	{
+		OutReason = LOCTEXT("InvalidParameterName", "Invalid Parameter Name");
+	}
+
+	if(!Type.IsValid())
+	{
+		OutReason = LOCTEXT("InvalidParameterType", "Invalid Parameter Type");
+	}
+
+	if(!Library.IsValid())
+	{
+		OutReason = LOCTEXT("InvalidParameterType", "Invalid Parameter Library");
+	}
+	
+	return true; 
+}
+
 void SAddParametersDialog::Construct(const FArguments& InArgs)
 {
 	using namespace AddParametersDialog;
@@ -47,7 +68,7 @@ void SAddParametersDialog::Construct(const FArguments& InArgs)
 	SWindow::Construct(SWindow::FArguments()
 		.Title(LOCTEXT("WindowTitle", "Add Parameters"))
 		.SizingRule(ESizingRule::UserSized)
-		.ClientSize(FVector2D(500.f, 500.f))
+		.ClientSize(InArgs._AllowMultiple ? FVector2D(500.f, 500.f) : FVector2D(500.f, 100.f))
 		.SupportsMaximize(false)
 		.SupportsMinimize(false)
 		[
@@ -61,6 +82,7 @@ void SAddParametersDialog::Construct(const FArguments& InArgs)
 				.Padding(0.0f, 5.0f)
 				[
 					SNew(SSimpleButton)
+					.Visibility(InArgs._AllowMultiple ? EVisibility::Visible : EVisibility::Collapsed)
 					.Text(LOCTEXT("AddButton", "Add"))
 					.ToolTipText(LOCTEXT("AddButtonTooltip", "Queue a new parameter for adding. New parameters will re-use the settings from the last queued parameter."))
 					.Icon(FAppStyle::Get().GetBrush("Icons.Plus"))
@@ -108,11 +130,36 @@ void SAddParametersDialog::Construct(const FArguments& InArgs)
 						SNew(SButton)
 						.HAlign(HAlign_Center)
 						.ButtonStyle(&FAppStyle::Get().GetWidgetStyle<FButtonStyle>("PrimaryButton"))
+						.IsEnabled_Lambda([this]()
+						{
+							// Check each entry to see if the button can be pressed
+							for(TSharedRef<FParameterToAdd> Entry : Entries)
+							{
+								if(!Entry->IsValid())
+								{
+									return false;
+								}
+							}
+
+							return true;
+						})
 						.Text_Lambda([this]()
 						{
 							return FText::Format(LOCTEXT("AddParametersButtonFormat", "Add {0} {0}|plural(one=Parameter,other=Parameters)"), FText::AsNumber(Entries.Num()));
 						})
-						.ToolTipText(LOCTEXT("AddParametersButtonTooltip", "Add the selected parameters to the current parameter block"))
+						.ToolTipText_Lambda([this]()
+						{
+							// Check each entry to see if the button can be pressed
+							for(TSharedRef<FParameterToAdd> Entry : Entries)
+							{
+								FText Reason;
+								if(!Entry->IsValid(Reason))
+								{
+									return FText::Format(LOCTEXT("AddParametersButtonTooltip_InvalidEntry", "A parameter to add is not valid: {0}"), Reason);
+								}
+							}
+							return LOCTEXT("AddParametersButtonTooltip", "Add the selected parameters to the current parameter block");
+						})
 						.OnClicked_Lambda([this]()
 						{
 							RequestDestroyWindow();
diff --git a/Engine/Plugins/Experimental/AnimNext/Source/AnimNextEditor/Private/Param/SAddParametersDialog.h b/Engine/Plugins/Experimental/AnimNext/Source/AnimNextEditor/Private/Param/SAddParametersDialog.h
--- a/Engine/Plugins/Experimental/AnimNext/Source/AnimNextEditor/Private/Param/SAddParametersDialog.h
+++ b/Engine/Plugins/Experimental/AnimNext/Source/AnimNextEditor/Private/Param/SAddParametersDialog.h
@@ -27,6 +27,8 @@ struct FParameterToAdd
 	{
 		return Name != NAME_None && Type.IsValid() && Library.IsValid(); 
 	}
+
+	bool IsValid(FText& OutReason) const;
 	
 	// Type
 	FAnimNextParamType Type;
@@ -43,9 +45,14 @@ class SAddParametersDialog : public SWindow
 public:
 	SLATE_BEGIN_ARGS(SAddParametersDialog)
 		: _Library(nullptr)
+		, _AllowMultiple(true)
 	{}
 
+	/** The default library to create new parameters in */
 	SLATE_ARGUMENT(UAnimNextParameterLibrary*, Library)
+
+	/** Whether we allow multiple parameters to be added or just one at a time */
+	SLATE_ARGUMENT(bool, AllowMultiple)
 	
 	SLATE_END_ARGS()
 
diff --git a/Engine/Plugins/Experimental/AnimNext/Source/AnimNextEditor/Private/Param/SGraphPinParamName.cpp b/Engine/Plugins/Experimental/AnimNext/Source/AnimNextEditor/Private/Param/SGraphPinParamName.cpp
new file mode 100644
--- /dev/null
+++ b/Engine/Plugins/Experimental/AnimNext/Source/AnimNextEditor/Private/Param/SGraphPinParamName.cpp
@@ -0,0 +1,50 @@
+﻿// Copyright Epic Games, Inc. All Rights Reserved.
+
+#include "SGraphPinParamName.h"
+#include "EditorUtils.h"
+#include "SParameterPickerCombo.h"
+#include "ScopedTransaction.h"
+
+#define LOCTEXT_NAMESPACE "SGraphPinParamName"
+
+namespace UE::AnimNext::Editor
+{
+
+void SGraphPinParamName::Construct(const FArguments& InArgs, UEdGraphPin* InPin)
+{
+	ModelPin = InArgs._ModelPin;
+
+	SGraphPin::Construct(SGraphPin::FArguments(), InPin);
+}
+
+TSharedRef<SWidget> SGraphPinParamName::GetDefaultValueWidget()
+{
+	FParameterPickerArgs Args;
+	Args.bShowLibraries = false;
+	Args.bShowBlocks = false;
+	Args.bMultiSelect = false;
+	Args.OnParameterPicked = FOnParameterPicked::CreateLambda([this](const FParameterBindingReference& InParameterBinding)
+	{
+		if(ModelPin)
+		{
+			FScopedTransaction Transaction(LOCTEXT("SelectParameter", "Select Parameter"));
+			GraphPinObj->Modify();
+			GraphPinObj->GetSchema()->TrySetDefaultValue(*GraphPinObj, InParameterBinding.Parameter.ToString());
+		}
+	});
+	
+	return SNew(SParameterPickerCombo)
+		.PickerArgs(Args)
+		.OnGetParameterName_Lambda([this]()
+		{
+			return FName(*GraphPinObj->DefaultValue);
+		})
+		.OnGetParameterType_Lambda([this]()
+		{
+			return FUtils::GetParameterTypeFromName(FName(*GraphPinObj->DefaultValue));
+		});
+}
+
+}
+
+#undef LOCTEXT_NAMESPACE
\ No newline at end of file
diff --git a/Engine/Plugins/Experimental/AnimNext/Source/AnimNextEditor/Private/Param/SGraphPinParamName.h b/Engine/Plugins/Experimental/AnimNext/Source/AnimNextEditor/Private/Param/SGraphPinParamName.h
new file mode 100644
--- /dev/null
+++ b/Engine/Plugins/Experimental/AnimNext/Source/AnimNextEditor/Private/Param/SGraphPinParamName.h
@@ -0,0 +1,29 @@
+﻿// Copyright Epic Games, Inc. All Rights Reserved.
+
+#pragma once
+
+#include "SGraphPin.h"
+
+class URigVMPin;
+
+namespace UE::AnimNext::Editor
+{
+
+// A pin widget that allows picking using an AnimNext parameter picker
+class SGraphPinParamName : public SGraphPin
+{
+	SLATE_BEGIN_ARGS(SGraphPinParamName) {}
+
+	SLATE_ARGUMENT(URigVMPin*, ModelPin)
+
+	SLATE_END_ARGS()
+
+	void Construct(const FArguments& InArgs, UEdGraphPin* InPin);
+
+private:
+	virtual TSharedRef<SWidget>	GetDefaultValueWidget() override;
+
+	URigVMPin* ModelPin = nullptr;
+};
+
+}
\ No newline at end of file
