AnimNext: Scheduler, Component and Parameter Improvements

diff --git a/Engine/Plugins/Experimental/AnimNext/Source/AnimNextEditor/Private/Param/SParameterBlockView.cpp b/Engine/Plugins/Experimental/AnimNext/Source/AnimNextEditor/Private/Param/SParameterBlockView.cpp
--- a/Engine/Plugins/Experimental/AnimNext/Source/AnimNextEditor/Private/Param/SParameterBlockView.cpp
+++ b/Engine/Plugins/Experimental/AnimNext/Source/AnimNextEditor/Private/Param/SParameterBlockView.cpp
@@ -5,7 +5,6 @@
 #include "Param/AnimNextParameterBlock.h"
 #include "Param/AnimNextParameter.h"
 #include "Param/AnimNextParameterBlock_EditorData.h"
-#include "Param/AnimNextParameterBlockBinding.h"
 #include "Widgets/Layout/SBox.h"
 #include "Widgets/Input/SSearchBox.h"
 #include "Param/AnimNextParameterBlockEntry.h"
@@ -18,7 +17,6 @@
 #include "ParameterBlockViewMenuContext.h"
 #include "PropertyBagDetails.h"
 #include "SAddParametersDialog.h"
-#include "SLinkParametersDialog.h"
 #include "SourceControlOperations.h"
 #include "SPinTypeSelector.h"
 #include "SSimpleButton.h"
@@ -29,6 +27,11 @@
 #include "Param/AnimNextParameterBlockBindingReference.h"
 #include "ToolMenus.h"
 #include "ScopedTransaction.h"
+#include "SParameterPicker.h"
+#include "SSimpleComboButton.h"
+#include "Param/AnimNextParameterBlockParameter.h"
+#include "Param/AnimNextParameterBlockGraph.h"
+#include "Framework/Application/SlateApplication.h"
 
 #define LOCTEXT_NAMESPACE "AnimNextParameterBlockView"
 
@@ -103,9 +106,38 @@ void SParameterBlockView::Construct(const FArguments& InArgs, UAnimNextParameter
 			+SHorizontalBox::Slot()
 			.AutoWidth()
 			.Padding(2.0f)
+			[
+				SNew(SSimpleComboButton)
+				.Text(LOCTEXT("AddParameterButton", "Add Parameter"))
+				.HasDownArrow(true)
+				.OnGetMenuContent_Lambda([this]()
+				{
+					FParameterPickerArgs Args;
+					Args.bMultiSelect = false;
+					Args.bShowLibraries = false;
+					Args.bShowBlocks = false;
+					Args.OnParameterPicked = FOnParameterPicked::CreateLambda([this](const FParameterBindingReference& InParameterBinding)
+					{
+						FSlateApplication::Get().DismissAllMenus();
+
+						FScopedTransaction Transaction(LOCTEXT("AddParameter", "Add parameter"));
+						PendingSelection.Empty();
+
+						// Create a new entry for the parameter
+						UAnimNextParameterLibrary* Library = Cast<UAnimNextParameterLibrary>(InParameterBinding.Library.GetAsset());
+						UAnimNextParameterBlockParameter* Parameter = EditorData->AddParameter(InParameterBinding.Parameter, Library);
+						PendingSelection.Add(Parameter);
+					});
+					return SNew(SParameterPicker)
+						.Args(Args);
+				})
+			]
+			+SHorizontalBox::Slot()
+			.AutoWidth()
+			.Padding(2.0f)
 			[
 				SNew(SSimpleButton)
-				.Text(LOCTEXT("AddNewParameterButton", "New"))
+				.Text(LOCTEXT("AddNewParameterButton", "Add New Parameter"))
 				.Icon(FAppStyle::Get().GetBrush("Icons.Plus"))
 				.OnClicked_Lambda([this]()
 				{
@@ -124,8 +156,8 @@ void SParameterBlockView::Construct(const FArguments& InArgs, UAnimNextParameter
 							UAnimNextParameter* NewParameter = Library->AddParameter(ParameterToAdd.Name, ParameterToAdd.Type);
 
 							// Create a new entry for the parameter
-							UAnimNextParameterBlockBinding* Binding = EditorData->AddBinding(NewParameter->GetFName(), Library);
-							PendingSelection.Add(Binding);
+							UAnimNextParameterBlockParameter* Parameter = EditorData->AddParameter(NewParameter->GetFName(), Library);
+							PendingSelection.Add(Parameter);
 						}
 					}
 
@@ -137,33 +169,17 @@ void SParameterBlockView::Construct(const FArguments& InArgs, UAnimNextParameter
 			.Padding(2.0f)
 			[
 				SNew(SSimpleButton)
-				.Text(LOCTEXT("LinkParameterButton", "Link"))
-				.Icon(FAppStyle::Get().GetBrush("Icons.Link"))
+				.Text(LOCTEXT("AddNewGraphButton", "Graph"))
+				.Icon(FAppStyle::Get().GetBrush("Icons.Plus"))
 				.OnClicked_Lambda([this]()
 				{
-					TSharedRef<SLinkParametersDialog> AddParametersDialog = SNew(SLinkParametersDialog)
-						.OnFilterParameter(this, &SParameterBlockView::HandleFilterLinkedParameter);
-					TArray<FParameterBindingReference> ParametersToLink;
-					if(AddParametersDialog->ShowModal(ParametersToLink))
-					{
-						FScopedTransaction Transaction(FText::FormatOrdered(LOCTEXT("LinkParameters", "Link {0}|plural(one=parameter,other=parameters)"), ParametersToLink.Num()));
+					FScopedTransaction Transaction(LOCTEXT("AddGraph", "Add Graph"));
 
-						PendingSelection.Empty();
+					PendingSelection.Empty();
 
-						for(const FParameterBindingReference& ParameterToLink : ParametersToLink)
-						{
-							check(ParameterToLink.Block.IsValid());
-							check(ParameterToLink.Parameter != NAME_None);
-
-							// Load our assets
-							UAnimNextParameterBlock* ExistingBlock = CastChecked<UAnimNextParameterBlock>(ParameterToLink.Block.GetAsset());
-							UAnimNextParameterLibrary* ExistingLibrary = CastChecked<UAnimNextParameterLibrary>(ParameterToLink.Library.GetAsset());
-
-							// Create a new entry for the supplied parameter
-							UAnimNextParameterBlockBindingReference* BindingReference = EditorData->AddBindingReference(ParameterToLink.Parameter, ExistingLibrary, ExistingBlock);
-							PendingSelection.Add(BindingReference);
-						}
-					}
+					// Create a new entry for the graph
+					UAnimNextParameterBlockGraph* Graph = EditorData->AddGraph(TEXT("NewGraph"));
+					PendingSelection.Add(Graph);
 
 					return FReply::Handled();
 				})
@@ -183,9 +199,10 @@ void SParameterBlockView::Construct(const FArguments& InArgs, UAnimNextParameter
 		+SVerticalBox::Slot()
 		.FillHeight(1.0f)
 		[
-			SAssignNew(EntriesList, SListView<TSharedRef<FParameterBlockViewEntry>>)
-			.ListItemsSource(&FilteredEntries)
+			SAssignNew(EntriesList, STreeView<TSharedRef<FParameterBlockViewEntry>>)
+			.TreeItemsSource(&FilteredEntries)
 			.OnGenerateRow(this, &SParameterBlockView::HandleGenerateRow)
+			.OnGetChildren(this, &SParameterBlockView::HandleGetChildren)
 			.OnItemScrolledIntoView(this, &SParameterBlockView::HandleItemScrolledIntoView)
 			.OnSelectionChanged(this, &SParameterBlockView::HandleSelectionChanged)
 			.ItemHeight(20.0f)
@@ -431,6 +448,20 @@ class SParameterBlockViewRow : public SMultiColumnTableRow<TSharedRef<FParameter
 		SMultiColumnTableRow<TSharedRef<FParameterBlockViewEntry>>::Construct( SMultiColumnTableRow<TSharedRef<FParameterBlockViewEntry>>::FArguments(), InOwnerTableView);
 	}
 
+	virtual FReply OnMouseButtonDoubleClick(const FGeometry& InMyGeometry, const FPointerEvent& InMouseEvent) override
+	{
+		if(IAnimNextParameterBlockGraphInterface* GraphInterface = Cast<IAnimNextParameterBlockGraphInterface>(Entry->WeakEntry.Get()))
+		{
+			if(TSharedPtr<SParameterBlockView> View = WeakView.Pin())
+			{
+				View->OnOpenGraphDelegate.ExecuteIfBound(GraphInterface->GetGraph());
+			}
+			return FReply::Handled();
+		}
+
+		return SMultiColumnTableRow<TSharedRef<FParameterBlockViewEntry>>::OnMouseButtonDoubleClick(InMyGeometry, InMouseEvent);
+	}
+	
 	virtual TSharedRef<SWidget> GenerateWidgetForColumn(const FName& InColumnName) override
 	{
 		using namespace ParameterBlockView;
@@ -495,7 +526,7 @@ class SParameterBlockViewRow : public SMultiColumnTableRow<TSharedRef<FParameter
 								TextBuilder.AppendLine(FText::FromName(ExternalPackage->GetFName()));
 							}
 
-							if(const IAnimNextParameterBlockBindingInterface* Binding = Cast<IAnimNextParameterBlockBindingInterface>(BlockEntry))
+							if(const IAnimNextParameterBlockParameterInterface* Binding = Cast<IAnimNextParameterBlockParameterInterface>(BlockEntry))
 							{
 								if(const UAnimNextParameter* Parameter = Binding->GetParameter())
 								{
@@ -523,7 +554,7 @@ class SParameterBlockViewRow : public SMultiColumnTableRow<TSharedRef<FParameter
 								bIsDirty = true;
 							}
 
-							if(const IAnimNextParameterBlockBindingInterface* Binding = Cast<IAnimNextParameterBlockBindingInterface>(BlockEntry))
+							if(const IAnimNextParameterBlockParameterInterface* Binding = Cast<IAnimNextParameterBlockParameterInterface>(BlockEntry))
 							{
 								if(const UAnimNextParameter* Parameter = Binding->GetParameter())
 								{
@@ -544,7 +575,7 @@ class SParameterBlockViewRow : public SMultiColumnTableRow<TSharedRef<FParameter
 		}
 		else if(InColumnName == Column_Type)
 		{
-			if(Entry->WeakEntry.Get()->Implements<UAnimNextParameterBlockBindingInterface>())
+			if(Entry->WeakEntry.Get()->Implements<UAnimNextParameterBlockParameterInterface>())
 			{
 				return
 					SNew(SBox)
@@ -555,7 +586,7 @@ class SParameterBlockViewRow : public SMultiColumnTableRow<TSharedRef<FParameter
 						SNew(SPinTypeSelector, FGetPinTypeTree::CreateStatic(&Editor::FUtils::GetFilteredVariableTypeTree))
 							.TargetPinType_Lambda([this]()
 							{
-								if(IAnimNextParameterBlockBindingInterface* Binding = Cast<IAnimNextParameterBlockBindingInterface>(Entry->WeakEntry.Get()))
+								if(IAnimNextParameterBlockParameterInterface* Binding = Cast<IAnimNextParameterBlockParameterInterface>(Entry->WeakEntry.Get()))
 								{
 									return UncookedOnly::FUtils::GetPinTypeFromParamType(Binding->GetParamType());
 								}
@@ -572,27 +603,35 @@ class SParameterBlockViewRow : public SMultiColumnTableRow<TSharedRef<FParameter
 		else if(InColumnName == Column_Name)
 		{
 			return
-				SNew(SBox)
+				SNew(SHorizontalBox)
+				+SHorizontalBox::Slot()
 				.HAlign(HAlign_Left)
 				.VAlign(VAlign_Center)
+				.AutoWidth()
+				[
+					SNew(SExpanderArrow, SharedThis(this))
+				]
+				+SHorizontalBox::Slot()
+				.HAlign(HAlign_Left)
+				.VAlign(VAlign_Center)
+				.AutoWidth()
 				[
-					// TODO: make this into a picker appropriately
 					SAssignNew(Entry->NameWidget, SInlineEditableTextBlock)
 					.Font(IDetailLayoutBuilder::GetDetailFont())
 					.IsSelected(this, &SParameterBlockViewRow::IsSelectedExclusively)
 					.IsReadOnly_Lambda([this]()
 					{
-						return Cast<UAnimNextParameterBlockBinding>(Entry->WeakEntry.Get()) == nullptr;
+						return Cast<IAnimNextParameterBlockGraphInterface>(Entry->WeakEntry.Get()) == nullptr;
 					})
 					.OnTextCommitted_Lambda([this](const FText& InNewText, ETextCommit::Type InCommitType)
 					{
 						if(InCommitType == ETextCommit::OnEnter)
 						{
-							if(IAnimNextParameterBlockBindingInterface* Binding = Cast<IAnimNextParameterBlockBindingInterface>(Entry->WeakEntry.Get()))
+							if(IAnimNextParameterBlockGraphInterface* Graph = Cast<IAnimNextParameterBlockGraphInterface>(Entry->WeakEntry.Get()))
 							{
-								FScopedTransaction Transaction(LOCTEXT("SetParameteName", "Set parameter name"));
+								FScopedTransaction Transaction(LOCTEXT("SetGraphName", "Set graph name"));
 
-								Binding->SetParameterName(*InNewText.ToString());
+								Graph->SetGraphName(*InNewText.ToString());
 							}
 						}
 					})
@@ -600,20 +639,13 @@ class SParameterBlockViewRow : public SMultiColumnTableRow<TSharedRef<FParameter
 					{
 						const FString NewString = InNewText.ToString();
 
-						if(IAnimNextParameterBlockBindingInterface* Binding = Cast<IAnimNextParameterBlockBindingInterface>(Entry->WeakEntry.Get()))
+						if(IAnimNextParameterBlockGraphInterface* Binding = Cast<IAnimNextParameterBlockGraphInterface>(Entry->WeakEntry.Get()))
 						{
 							// Make sure the new name only contains valid characters
 							if (!FName::IsValidXName(NewString, INVALID_OBJECTNAME_CHARACTERS INVALID_LONGPACKAGE_CHARACTERS, &OutErrorText))
 							{
 								return false;
 							}
-
-							const FName Name(*NewString);
-							if(!FUtils::DoesParameterExistInLibrary(Binding->GetLibrary(), Name))
-							{
-								OutErrorText = LOCTEXT("Error_NameDoesNotExistInLibrary", "This name does not exist in the specified library");
-								return false;
-							}
 						}
 
 						if(IAnimNextParameterBlockReferenceInterface* Reference = Cast<IAnimNextParameterBlockReferenceInterface>(Entry->WeakEntry.Get()))
@@ -630,9 +662,14 @@ class SParameterBlockViewRow : public SMultiColumnTableRow<TSharedRef<FParameter
 					})
 					.Text_Lambda([this]()
 					{
-						if(UAnimNextParameterBlockEntry* BlockEntry = Entry->WeakEntry.Get())
+						if(IAnimNextParameterBlockParameterInterface* Parameter = Cast<IAnimNextParameterBlockParameterInterface>(Entry->WeakEntry.Get()))
 						{
-							return BlockEntry->GetDisplayName();
+							return UncookedOnly::FUtils::GetParameterDisplayNameText(Parameter->GetParameterName());
+						}
+
+						if(IAnimNextParameterBlockGraphInterface* Graph = Cast<IAnimNextParameterBlockGraphInterface>(Entry->WeakEntry.Get()))
+						{
+							return FText::FromName(Graph->GetGraphName());
 						}
 						return FText::GetEmpty();
 					})
@@ -688,29 +725,8 @@ class SParameterBlockViewRow : public SMultiColumnTableRow<TSharedRef<FParameter
 						})
 					];
 				}
-				
-				if(IAnimNextParameterBlockGraphInterface* GraphInterface = Cast<IAnimNextParameterBlockGraphInterface>(BlockEntry))
-				{
-					URigVMGraph* ReferencedGraph = GraphInterface->GetGraph();
-					HorizontalBox->AddSlot()
-					[
-						SNew(SSimpleButton)
-						.Icon(FAppStyle::GetBrush("Icons.Blueprints"))
-						.OnClicked_Lambda([this, WeakGraph = TWeakObjectPtr<URigVMGraph>(ReferencedGraph)]()
-						{
-							if(URigVMGraph* Graph = WeakGraph.Get())
-							{
-								if(TSharedPtr<SParameterBlockView> View = WeakView.Pin())
-								{
-									View->OnOpenGraphDelegate.ExecuteIfBound(Graph);
-								}
-							}
-							return FReply::Handled();
-						})
-					];
-				}
 			}
-			
+
 			return Widget;
 		}
 
@@ -726,6 +742,11 @@ TSharedRef<ITableRow> SParameterBlockView::HandleGenerateRow(TSharedRef<FParamet
 	return SNew(SParameterBlockViewRow, InOwnerTable, SharedThis(this), InEntry);
 }
 
+void SParameterBlockView::HandleGetChildren(TSharedRef<FParameterBlockViewEntry> InEntry, TArray<TSharedRef<FParameterBlockViewEntry>>& OutChildren)
+{
+
+}
+
 void SParameterBlockView::HandleSelectionChanged(TSharedPtr<FParameterBlockViewEntry> InEntry, ESelectInfo::Type InSelectionType)
 {
 	if(OnSelectionChangedDelegate.IsBound())
diff --git a/Engine/Plugins/Experimental/AnimNext/Source/AnimNextEditor/Private/Param/SParameterBlockView.h b/Engine/Plugins/Experimental/AnimNext/Source/AnimNextEditor/Private/Param/SParameterBlockView.h
--- a/Engine/Plugins/Experimental/AnimNext/Source/AnimNextEditor/Private/Param/SParameterBlockView.h
+++ b/Engine/Plugins/Experimental/AnimNext/Source/AnimNextEditor/Private/Param/SParameterBlockView.h
@@ -71,6 +71,8 @@ private:
 	// Generate a row for the list view
 	TSharedRef<ITableRow> HandleGenerateRow(TSharedRef<FParameterBlockViewEntry> InEntry, const TSharedRef<STableViewBase>& InOwnerTable);
 
+	void HandleGetChildren(TSharedRef<FParameterBlockViewEntry> InEntry, TArray<TSharedRef<FParameterBlockViewEntry>>& OutChildren);
+
 	// Handle rename after scrolling into view
 	void HandleItemScrolledIntoView(TSharedRef<FParameterBlockViewEntry> Entry, const TSharedPtr<ITableRow>& Widget);
 
diff --git a/Engine/Plugins/Experimental/AnimNext/Source/AnimNextEditor/Private/Param/SParameterPicker.cpp b/Engine/Plugins/Experimental/AnimNext/Source/AnimNextEditor/Private/Param/SParameterPicker.cpp
--- a/Engine/Plugins/Experimental/AnimNext/Source/AnimNextEditor/Private/Param/SParameterPicker.cpp
+++ b/Engine/Plugins/Experimental/AnimNext/Source/AnimNextEditor/Private/Param/SParameterPicker.cpp
@@ -12,11 +12,16 @@
 #include "Param/AnimNextParameterLibrary.h"
 #include "DetailLayoutBuilder.h"
 #include "EditorUtils.h"
+#include "SAddParametersDialog.h"
 #include "Param/ParameterPickerArgs.h"
 #include "Widgets/Input/SSearchBox.h"
 #include "String/ParseTokens.h"
 #include "Widgets/Layout/SBox.h"
 #include "Widgets/Images/SImage.h"
+#include "Param/Params.h"
+#include "Widgets/Input/SButton.h"
+#include "Framework/Application/SlateApplication.h"
+#include "ScopedTransaction.h"
 
 #define LOCTEXT_NAMESPACE "SParameterPicker"
 
@@ -29,6 +34,7 @@ static FName Column_Parameter(TEXT("Parameter"));
 static FName Column_Library(TEXT("Library"));
 static FName Column_Block(TEXT("Block"));
 static FName Column_Type(TEXT("Type"));
+static FName Column_New(TEXT("New"));
 }
 
 struct FParameterPickerEntry
@@ -185,6 +191,58 @@ void SParameterPicker::Construct(const FArguments& InArgs)
 			.FillWidth(0.33f));
 	}
 
+	if(Args.bAllowNew)
+	{
+		HeaderRow->AddColumn(
+			SHeaderRow::Column(Column_New)
+			.DefaultLabel(FText::GetEmpty())
+			.HeaderContentPadding(FMargin(0.0f))
+			.FixedWidth(24.0f)
+			.HeaderContent()
+			[
+				SNew(SButton)
+				.ToolTipText(LOCTEXT("AddColumnHeaderTooltip", "Add a new parameter at global scope"))
+				.ButtonStyle(FAppStyle::Get(), "HoverHintOnly")
+				.OnClicked_Lambda([this]()
+				{
+					FSlateApplication::Get().DismissAllMenus();
+					TSharedRef<SAddParametersDialog> AddParametersDialog =
+						SNew(SAddParametersDialog)
+						.AllowMultiple(false);
+					TArray<FParameterToAdd> ParametersToAdd;
+					if(AddParametersDialog->ShowModal(ParametersToAdd))
+					{
+						if(ParametersToAdd.Num() > 0)
+						{
+							FScopedTransaction Transaction(LOCTEXT("AddParameter", "Add parameter"));
+
+							// Create a new parameter in the supplied library
+							UAnimNextParameterLibrary* Library = Cast<UAnimNextParameterLibrary>(ParametersToAdd[0].Library.GetAsset());
+							UAnimNextParameter* NewParameter = Library->AddParameter(ParametersToAdd[0].Name, ParametersToAdd[0].Type);
+							
+							FParameterBindingReference Reference;
+							Reference.Parameter = ParametersToAdd[0].Name;
+							Reference.Library = ParametersToAdd[0].Library;
+							Args.OnParameterPicked.ExecuteIfBound(Reference);
+						}
+					}
+					return FReply::Handled();
+				})
+				[
+					SNew(SBox)
+					.WidthOverride(16.0f)
+					.HeightOverride(16.0f)
+					.VAlign(VAlign_Center)
+					.HAlign(HAlign_Center)
+					[
+						SNew(SImage)
+						.ColorAndOpacity(FSlateColor::UseForeground())
+						.Image(FAppStyle::GetBrush("Icons.Plus"))
+					]
+				]
+			]);
+	}
+	
 	RefreshEntries();
 }
 
@@ -279,6 +337,21 @@ void SParameterPicker::RefreshEntries()
 		}
 	}
 
+	if (Args.bShowBuiltInParameters)
+	{
+		FParams::ForEachBuiltInParameter([this](const FParamDefinition& InDefinition)
+		{
+			FParameterBindingReference NewReference(InDefinition.GetName());
+			if (!Args.OnFilterParameter.IsBound() || Args.OnFilterParameter.Execute(NewReference) == EFilterParameterResult::Include)
+			{
+				if (!Args.OnFilterParameterType.IsBound() || Args.OnFilterParameterType.Execute(InDefinition.GetType()) == EFilterParameterResult::Include)
+				{
+					Entries.Add(MakeShared<FParameterPickerEntry>(NewReference, InDefinition.GetType()));
+				}
+			}
+		});
+	}
+
 	BuildHierarchy();
 
 	RefreshFilter();
@@ -305,7 +378,7 @@ void SParameterPicker::BuildHierarchy()
 
 		TStringBuilder<256> PartialParameterString;
 
-		// We use '_' as a seperator here as:
+		// We use '_' as a separator here as:
 		// - Each param is a UObject in editor and uses its object name, so we cannot use '.'
 		// - This maps nicely to Verse tags that use '_' to hierarchically define their relationship
 		TSharedPtr<FParameterPickerEntry> Parent = nullptr;
diff --git a/Engine/Plugins/Experimental/AnimNext/Source/AnimNextEditor/Private/Param/SParameterPickerCombo.cpp b/Engine/Plugins/Experimental/AnimNext/Source/AnimNextEditor/Private/Param/SParameterPickerCombo.cpp
new file mode 100644
--- /dev/null
+++ b/Engine/Plugins/Experimental/AnimNext/Source/AnimNextEditor/Private/Param/SParameterPickerCombo.cpp
@@ -0,0 +1,99 @@
+﻿// Copyright Epic Games, Inc. All Rights Reserved.
+
+#include "SParameterPickerCombo.h"
+#include "SParameterPicker.h"
+#include "UncookedOnlyUtils.h"
+#include "Kismet2/BlueprintEditorUtils.h"
+#include "Widgets/Input/SComboButton.h"
+#include "Framework/Application/SlateApplication.h"
+#include "Widgets/Text/SInlineEditableTextBlock.h"
+
+namespace UE::AnimNext::Editor
+{
+
+void SParameterPickerCombo::Construct(const FArguments& InArgs)
+{
+	OnGetParameterNameDelegate = InArgs._OnGetParameterName;
+	OnGetParameterTypeDelegate = InArgs._OnGetParameterType;
+	
+	FParameterPickerArgs PickerArgs = InArgs._PickerArgs;
+	PickerArgs.OnParameterPicked = FOnParameterPicked::CreateLambda([this, OrigPickerArgs = InArgs._PickerArgs](const FParameterBindingReference& InReference)
+	{
+		FSlateApplication::Get().DismissAllMenus();
+
+		// Forward to the original delegate
+		OrigPickerArgs.OnParameterPicked.ExecuteIfBound(InReference);
+
+		RequestRefresh();
+	});
+
+	ChildSlot
+	[
+		SNew(SComboButton)
+		.OnGetMenuContent_Lambda([PickerArgs]()
+		{
+			return SNew(SParameterPicker)
+				.Args(PickerArgs);
+		})
+		.ButtonContent()
+		[
+			SNew(SHorizontalBox)
+			+SHorizontalBox::Slot()
+			.AutoWidth()
+			.VAlign(VAlign_Center)
+			.HAlign(HAlign_Center)
+			.Padding(0.0f, 2.0f, 2.0f, 2.0f)
+			[
+				SNew(SImage)
+				.Image_Lambda([this]()
+				{
+					return Icon;
+				})
+				.ColorAndOpacity_Lambda([this]()
+				{
+					return IconColor;
+				})
+			]
+			+SHorizontalBox::Slot()
+			.AutoWidth()
+			.VAlign(VAlign_Center)
+			.HAlign(HAlign_Left)
+			.Padding(0.0f, 2.0f, 0.0f, 2.0f)
+			[
+				SNew(STextBlock)
+				.TextStyle(&FCoreStyle::Get().GetWidgetStyle<FTextBlockStyle>( "SmallText" ))
+				.Text_Lambda([this]()
+				{
+					return ParameterNameText;
+				})
+			]
+		]
+	];
+
+	RequestRefresh();
+}
+
+void SParameterPickerCombo::RequestRefresh()
+{
+	RegisterActiveTimer(1.0f/60.0f, FWidgetActiveTimerDelegate::CreateLambda([this](double InCurrentTime, float InDeltaTime)
+	{
+		if(OnGetParameterNameDelegate.IsBound())
+		{
+			ParameterName = OnGetParameterNameDelegate.Execute();
+		}
+
+		if(OnGetParameterTypeDelegate.IsBound())
+		{
+			ParameterType = OnGetParameterTypeDelegate.Execute();
+		}
+
+		ParameterNameText = UE::AnimNext::UncookedOnly::FUtils::GetParameterDisplayNameText(ParameterName);
+		PinType = UncookedOnly::FUtils::GetPinTypeFromParamType(ParameterType);
+		Icon = FBlueprintEditorUtils::GetIconFromPin(PinType, true);
+		IconColor = GetDefault<UEdGraphSchema_K2>()->GetPinTypeColor(PinType);
+
+		return EActiveTimerReturnType::Stop;
+	}));
+}
+
+}
diff --git a/Engine/Plugins/Experimental/AnimNext/Source/AnimNextEditor/Private/Param/SParameterPickerCombo.h b/Engine/Plugins/Experimental/AnimNext/Source/AnimNextEditor/Private/Param/SParameterPickerCombo.h
new file mode 100644
--- /dev/null
+++ b/Engine/Plugins/Experimental/AnimNext/Source/AnimNextEditor/Private/Param/SParameterPickerCombo.h
@@ -0,0 +1,66 @@
+﻿// Copyright Epic Games, Inc. All Rights Reserved.
+
+#pragma once
+
+#include "Widgets/SCompoundWidget.h"
+#include "Param/ParameterPickerArgs.h"
+#include "EdGraph/EdGraphPin.h"
+
+struct FAnimNextParamType;
+
+namespace UE::AnimNext::Editor
+{
+
+/** Retrieves the parameter name to display */
+using FOnGetParameterName = TDelegate<FName(void)>;
+
+/** Retrieves the parameter type to display */
+using FOnGetParameterType = TDelegate<FAnimNextParamType(void)>;
+
+class SParameterPickerCombo : public SCompoundWidget
+{
+public:
+	SLATE_BEGIN_ARGS(SParameterPickerCombo) {}
+
+	/** Arguments for configuring the picker in the dropdown menu */
+	SLATE_ARGUMENT(FParameterPickerArgs, PickerArgs)
+
+	/** Retrieves the parameter name to display */
+	SLATE_EVENT(FOnGetParameterName, OnGetParameterName)
+
+	/** Retrieves the parameter type to display */
+	SLATE_EVENT(FOnGetParameterType, OnGetParameterType)
+
+	SLATE_END_ARGS()
+
+	void Construct(const FArguments& InArgs);
+
+private:
+	void RequestRefresh();
+
+	// Retrieves the parameter name to display
+	FOnGetParameterName OnGetParameterNameDelegate;
+
+	// Retrieves the parameter type to display
+	FOnGetParameterType OnGetParameterTypeDelegate;
+
+	// Cached pin type
+	FEdGraphPinType PinType;
+
+	// Cached name
+	FName ParameterName;
+
+	// Cached display name
+	FText ParameterNameText;
+
+	// Cached parameter type
+	FAnimNextParamType ParameterType;
+
+	// Cached icon
+	const FSlateBrush* Icon = nullptr;
+
+	// Cached color
+	FSlateColor IconColor = FLinearColor::Gray;
+};
+
+}
\ No newline at end of file
diff --git a/Engine/Plugins/Experimental/AnimNext/Source/AnimNextEditor/Private/Param/SParametersActionMenu.cpp b/Engine/Plugins/Experimental/AnimNext/Source/AnimNextEditor/Private/Param/SParametersActionMenu.cpp
--- a/Engine/Plugins/Experimental/AnimNext/Source/AnimNextEditor/Private/Param/SParametersActionMenu.cpp
+++ b/Engine/Plugins/Experimental/AnimNext/Source/AnimNextEditor/Private/Param/SParametersActionMenu.cpp
@@ -14,7 +14,7 @@
 #include "RigVMCore/RigVMRegistry.h"
 #include "Units/RigUnit.h"
 #include "Widgets/SToolTip.h"
-#include "Graph/AnimNextExecuteContext.h"
+#include "Param/AnimNextParameterExecuteContext.h"
 
 #define LOCTEXT_NAMESPACE "AnimNextParametersEditor"
 
@@ -25,7 +25,8 @@ static void CollectAllAnimNextParameterActions(FGraphContextMenuBuilder& MenuBui
 {
 	static const TArray<UScriptStruct*> AllowedExecuteContexts =
 	{
-		FRigVMExecuteContext::StaticStruct()
+		FRigVMExecuteContext::StaticStruct(),
+		FAnimNextParameterExecuteContext::StaticStruct()
 	};
 
 	for (const FRigVMFunction& Function : FRigVMRegistry::Get().GetFunctions())
@@ -64,7 +65,8 @@ static void CollectAllAnimNextParameterActions(FGraphContextMenuBuilder& MenuBui
 
 	for (const FRigVMDispatchFactory* Factory : FRigVMRegistry::Get().GetFactories())
 	{
-		if (!Factory->SupportsExecuteContextStruct(FRigVMExecuteContext::StaticStruct()))
+		if (!Factory->SupportsExecuteContextStruct(FRigVMExecuteContext::StaticStruct()) &&
+			!Factory->SupportsExecuteContextStruct(FAnimNextParameterExecuteContext::StaticStruct()))
 		{
 			continue;
 		}
diff --git a/Engine/Plugins/Experimental/AnimNext/Source/AnimNextEditor/Private/Scheduler/AnimNextScheduleFactory.cpp b/Engine/Plugins/Experimental/AnimNext/Source/AnimNextEditor/Private/Scheduler/AnimNextScheduleFactory.cpp
new file mode 100644
--- /dev/null
+++ b/Engine/Plugins/Experimental/AnimNext/Source/AnimNextEditor/Private/Scheduler/AnimNextScheduleFactory.cpp
@@ -0,0 +1,23 @@
+// Copyright Epic Games, Inc. All Rights Reserved.
+
+#include "Scheduler/AnimNextScheduleFactory.h"
+#include "Scheduler/AnimNextSchedule.h"
+
+UAnimNextScheduleFactory::UAnimNextScheduleFactory()
+{
+	bCreateNew = true;
+	bEditAfterNew = true;
+	SupportedClass = UAnimNextSchedule::StaticClass();
+}
+
+bool UAnimNextScheduleFactory::ConfigureProperties()
+{
+	return true;
+}
+
+UObject* UAnimNextScheduleFactory::FactoryCreateNew(UClass* Class, UObject* InParent, FName Name, EObjectFlags Flags, UObject* Context, FFeedbackContext* Warn, FName CallingContext)
+{
+	UAnimNextSchedule* NewSchedule = NewObject<UAnimNextSchedule>(InParent, Class, Name, Flags | RF_Public | RF_Standalone | RF_Transactional | RF_LoadCompleted);
+
+	return NewSchedule;
+}
\ No newline at end of file
diff --git a/Engine/Plugins/Experimental/AnimNext/Source/AnimNextEditor/Private/Scheduler/ScheduleAssetDefinition.cpp b/Engine/Plugins/Experimental/AnimNext/Source/AnimNextEditor/Private/Scheduler/ScheduleAssetDefinition.cpp
new file mode 100644
--- /dev/null
+++ b/Engine/Plugins/Experimental/AnimNext/Source/AnimNextEditor/Private/Scheduler/ScheduleAssetDefinition.cpp
@@ -0,0 +1,16 @@
+// Copyright Epic Games, Inc. All Rights Reserved.
+
+#include "ScheduleAssetDefinition.h"
+#include "Toolkits/SimpleAssetEditor.h"
+
+#define LOCTEXT_NAMESPACE "AnimNextAssetDefinitions"
+
+EAssetCommandResult UAssetDefinition_AnimNextSchedule::OpenAssets(const FAssetOpenArgs& OpenArgs) const
+{
+	TArray<UObject*> Objects = OpenArgs.LoadObjects<UObject>();
+	FSimpleAssetEditor::CreateEditor(EToolkitMode::Standalone, OpenArgs.ToolkitHost, Objects);
+
+	return EAssetCommandResult::Handled;
+}
+
+#undef LOCTEXT_NAMESPACE
\ No newline at end of file
diff --git a/Engine/Plugins/Experimental/AnimNext/Source/AnimNextEditor/Private/Scheduler/ScheduleAssetDefinition.h b/Engine/Plugins/Experimental/AnimNext/Source/AnimNextEditor/Private/Scheduler/ScheduleAssetDefinition.h
new file mode 100644
--- /dev/null
+++ b/Engine/Plugins/Experimental/AnimNext/Source/AnimNextEditor/Private/Scheduler/ScheduleAssetDefinition.h
@@ -0,0 +1,29 @@
+// Copyright Epic Games, Inc. All Rights Reserved.
+
+#pragma once
+
+#include "AssetDefinitionDefault.h"
+#include "Scheduler/AnimNextSchedule.h"
+#include "ScheduleAssetDefinition.generated.h"
+
+#define LOCTEXT_NAMESPACE "AnimNextAssetDefinitions"
+
+UCLASS()
+class UAssetDefinition_AnimNextSchedule : public UAssetDefinitionDefault
+{
+	GENERATED_BODY()
+
+public:
+	// UAssetDefinition interface
+	virtual FText GetAssetDisplayName() const override { return LOCTEXT("AnimNextSchedule", "AnimNext Schedule"); }
+	virtual FLinearColor GetAssetColor() const override { return FLinearColor(FColor(128, 64, 64)); }
+	virtual TSoftClassPtr<UObject> GetAssetClass() const override { return UAnimNextSchedule::StaticClass(); }
+	virtual EAssetCommandResult OpenAssets(const FAssetOpenArgs& OpenArgs) const override;
+	virtual TConstArrayView<FAssetCategoryPath> GetAssetCategories() const override
+	{
+		static const auto Categories = { FAssetCategoryPath(EAssetCategoryPaths::Animation, LOCTEXT("AnimNextSubMenu", "AnimNext")) };
+		return Categories;
+	}
+};
+
+#undef LOCTEXT_NAMESPACE
\ No newline at end of file
diff --git a/Engine/Plugins/Experimental/AnimNext/Source/AnimNextTestSuite/Private/AnimNextDecoratorGraphTest.cpp b/Engine/Plugins/Experimental/AnimNext/Source/AnimNextTestSuite/Private/AnimNextDecoratorGraphTest.cpp
--- a/Engine/Plugins/Experimental/AnimNext/Source/AnimNextTestSuite/Private/AnimNextDecoratorGraphTest.cpp
+++ b/Engine/Plugins/Experimental/AnimNext/Source/AnimNextTestSuite/Private/AnimNextDecoratorGraphTest.cpp
@@ -242,6 +242,9 @@ bool FAnimationAnimNextRuntimeTest_GraphExecute::RunTest(const FString& InParame
 	Controller->SetPinDefaultValue(DecoratorPin->GetSubPins()[1]->GetPinPath(), TEXT("78"));
 	Controller->SetPinDefaultValue(DecoratorPin->GetSubPins()[2]->GetPinPath(), TEXT("142.33"));
 
+	TSharedRef<UE::AnimNext::FParamStack> ParamStack = MakeShared<UE::AnimNext::FParamStack>();
+	UE::AnimNext::FParamStack::AttachToCurrentThread(ParamStack);
+
 	FAnimNextGraphInstance GraphInstance;
 	AnimNextGraph->AllocateInstance(GraphInstance);
 
@@ -270,6 +273,8 @@ bool FAnimationAnimNextRuntimeTest_GraphExecute::RunTest(const FString& InParame
 	Context.GetMutableParamStack().PopLayer(LayerHandle);
 	GraphInstance.Release();
 
+	UE::AnimNext::FParamStack::DetachFromCurrentThread();
+
 	return true;
 }
 
@@ -374,6 +379,9 @@ bool FAnimationAnimNextRuntimeTest_GraphExecuteLatent::RunTest(const FString& In
 			DecoratorPin->GetSubPins()[4]);	// SomeOtherLatentInt32
 	}
 
+	TSharedRef<UE::AnimNext::FParamStack> ParamStack = MakeShared<UE::AnimNext::FParamStack>();
+	UE::AnimNext::FParamStack::AttachToCurrentThread(ParamStack);
+
 	FAnimNextGraphInstance GraphInstance;
 	AnimNextGraph->AllocateInstance(GraphInstance);
 
@@ -403,6 +411,8 @@ bool FAnimationAnimNextRuntimeTest_GraphExecuteLatent::RunTest(const FString& In
 	Context.GetMutableParamStack().PopLayer(LayerHandle);
 	GraphInstance.Release();
 
+	UE::AnimNext::FParamStack::DetachFromCurrentThread();
+
 	return true;
 }
 
diff --git a/Engine/Plugins/Experimental/AnimNext/Source/AnimNextTestSuite/Private/AnimNextEditorTest.cpp b/Engine/Plugins/Experimental/AnimNext/Source/AnimNextTestSuite/Private/AnimNextEditorTest.cpp
--- a/Engine/Plugins/Experimental/AnimNext/Source/AnimNextTestSuite/Private/AnimNextEditorTest.cpp
+++ b/Engine/Plugins/Experimental/AnimNext/Source/AnimNextTestSuite/Private/AnimNextEditorTest.cpp
@@ -135,7 +135,12 @@ bool FAnimationAnimNextParametersEditorTest_Block::RunTest(const FString& InPara
 
 	GEditor->UndoTransaction();
 	AddErrorIfFalse(EditorData->Entries.Num() == 1, FString::Printf(TEXT("Unexpected entry count found in parameter block (Have %d, expected 1)."), EditorData->Entries.Num()));
-	
+
+	// FindBinding
+	AddErrorIfFalse(EditorData->FindBinding(TEXT("TestParam")) != nullptr, TEXT("Could not find binding in block."));
+
+	GEditor->RedoTransaction();
+
 	UAnimNextParameterBlock* OtherBlock = Cast<UAnimNextParameterBlock>(BlockFactory->FactoryCreateNew(UAnimNextParameterBlock::StaticClass(), GetTransientPackage(), TEXT("TestAnimNextParameterBlock2"), RF_Transient, nullptr, nullptr, NAME_None));
 	if(OtherBlock == nullptr)
 	{
@@ -164,9 +169,36 @@ bool FAnimationAnimNextParametersEditorTest_Block::RunTest(const FString& InPara
 	AddErrorIfFalse(OtherEditorData->Entries.Num() == 1, FString::Printf(TEXT("Unexpected entry count found in parameter block (Have %d, expected 1)."), OtherEditorData->Entries.Num()));
 
 	// FindBinding
-	AddErrorIfFalse(EditorData->FindBinding(TEXT("TestParam")) != nullptr, TEXT("Could not find binding in block."));
 	AddErrorIfFalse(OtherEditorData->FindBinding(TEXT("TestParam")) != nullptr, TEXT("Could not find binding refernce in block."));
 
+	UAnimNextParameterBlockParameter* Parameter = nullptr;
+	{
+		FScopedTransaction Transaction(FText::GetEmpty());
+		Parameter = EditorData->AddParameter(TEXT("TestParam"), Library);
+		AddErrorIfFalse(Parameter != nullptr, TEXT("Could not create new parameter in block."));
+	}
+
+	GEditor->UndoTransaction();
+	AddErrorIfFalse(EditorData->Entries.Num() == 0, FString::Printf(TEXT("Unexpected entry count found in parameter block (Have %d, expected 0)."), EditorData->Entries.Num()));
+
+	GEditor->RedoTransaction();
+	AddErrorIfFalse(EditorData->Entries.Num() == 1, FString::Printf(TEXT("Unexpected entry count found in parameter block (Have %d, expected 1)."), EditorData->Entries.Num()));
+	GEditor->UndoTransaction();
+	
+	UAnimNextParameterBlockGraph* Graph = nullptr;
+	{
+		FScopedTransaction Transaction(FText::GetEmpty());
+		Graph = EditorData->AddGraph(TEXT("TestGraph"));
+		AddErrorIfFalse(Graph != nullptr, TEXT("Could not create new graph in block."));
+	}
+
+	GEditor->UndoTransaction();
+	AddErrorIfFalse(EditorData->Entries.Num() == 0, FString::Printf(TEXT("Unexpected entry count found in parameter block (Have %d, expected 0)."), EditorData->Entries.Num()));
+
+	GEditor->RedoTransaction();
+	AddErrorIfFalse(EditorData->Entries.Num() == 1, FString::Printf(TEXT("Unexpected entry count found in parameter block (Have %d, expected 1)."), EditorData->Entries.Num()));
+	GEditor->UndoTransaction();
+
 	CollectGarbage(GARBAGE_COLLECTION_KEEPFLAGS);
 
 	return true;
@@ -186,6 +218,8 @@ bool FAnimationAnimNextParametersEditorTest_Python::RunTest(const FString& InPar
 		"library.add_parameter(name = \"TestParam\", value_type = unreal.PropertyBagPropertyType.BOOL, container_type = unreal.PropertyBagContainerType.NONE)\n"
 		"block.add_binding(name = \"TestParam\", library = library)\n"
 		"other_block.add_binding_reference(name = \"TestParam\", library = library, referenced_block = block)\n"
+		"block.add_parameter(name = \"TestParam\", library = library)\n"
+		"block.add_graph(name = \"TestGraph\")\n"
 		"unreal.EditorAssetLibrary.delete_loaded_asset(library)\n"
 		"unreal.EditorAssetLibrary.delete_loaded_asset(block)\n"
 		"unreal.EditorAssetLibrary.delete_loaded_asset(other_block)\n"
diff --git a/Engine/Plugins/Experimental/AnimNext/Source/AnimNextTestSuite/Private/AnimNextParametersTest.cpp b/Engine/Plugins/Experimental/AnimNext/Source/AnimNextTestSuite/Private/AnimNextParametersTest.cpp
--- a/Engine/Plugins/Experimental/AnimNext/Source/AnimNextTestSuite/Private/AnimNextParametersTest.cpp
+++ b/Engine/Plugins/Experimental/AnimNext/Source/AnimNextTestSuite/Private/AnimNextParametersTest.cpp
@@ -330,11 +330,11 @@ bool FParamStackTest::RunTest(const FString& InParameters)
 		const FDateTime Now = FDateTime::Now();
 		Stack.PushValue("Param0", Now);
 
-		FParamStack::EGetParamResult Result;
+		FParamResult Result;
 		AddErrorIfFalse(Stack.GetParamPtr<FDateTime>("Param0", &Result) != nullptr, "Const ptr cannot be accessed");
-		AddErrorIfFalse(Result == FParamStack::EGetParamResult::Succeeded, "Result != FParamStack::EGetParamResult::Succeeded");
+		AddErrorIfFalse(Result.IsSuccessful(), "Unexpected result");
 		AddErrorIfFalse(Stack.GetMutableParamPtr<FDateTime>("Param0", &Result) == nullptr, "Mutable ptr can be accessed");
-		AddErrorIfFalse(Result == FParamStack::EGetParamResult::Immutable, "Result != FParamStack::EGetParamResult::Immutable");
+		AddErrorIfFalse(!Result.IsOfCompatibleMutability(), "Unexpected result");
 	}
 
 	// Check type variance
@@ -345,11 +345,11 @@ bool FParamStackTest::RunTest(const FString& InParameters)
 		const float Float = 1.0f;
 		Handle = Stack.PushValue("Param0", Float);
 
-		FParamStack::EGetParamResult Result;
+		FParamResult Result;
 		AddErrorIfFalse(Stack.GetParamPtr<double>("Param0", &Result) == nullptr, "Promoted typed access is non-null");
-		AddErrorIfFalse(Result == FParamStack::EGetParamResult::IncompatibleType, "Result != FParamStack::EGetParamResult::IncompatibleType");
+		AddErrorIfFalse(Result.IsOfIncompatibleType(), "Unexpected result");
 		AddErrorIfFalse(Stack.GetParamPtr<uint8>("Param0", &Result) == nullptr, "Demoted typed access is non-null");
-		AddErrorIfFalse(Result == FParamStack::EGetParamResult::IncompatibleType, "Result != FParamStack::EGetParamResult::IncompatibleType");
+		AddErrorIfFalse(Result.IsOfIncompatibleType(), "Unexpected result");
 
 		Stack.PopLayer(Handle);
 
@@ -357,11 +357,11 @@ bool FParamStackTest::RunTest(const FString& InParameters)
 		Handle = Stack.PushValue("Param0", Blue);
 
 		AddErrorIfFalse(Stack.GetParamPtr<FDateTime>("Param0", &Result) == nullptr, "Incorrectly typed access is non-null");
-		AddErrorIfFalse(Result == FParamStack::EGetParamResult::IncompatibleType, "Result != FParamStack::EGetParamResult::IncompatibleType");
+		AddErrorIfFalse(Result.IsOfIncompatibleType(), "Unexpected result");
 		AddErrorIfFalse(Stack.GetMutableParamPtr<FDateTime>("Param0", &Result) == nullptr, "Incorrectly typed mutable access is non-null");
-		AddErrorIfFalse(Result == (FParamStack::EGetParamResult::IncompatibleType | FParamStack::EGetParamResult::Immutable), "Result != (FParamStack::EGetParamResult::IncompatibleType | FParamStack::EGetParamResult::Immutable)");
+		AddErrorIfFalse(Result.IsOfIncompatibleType() && !Result.IsOfCompatibleMutability(), "Result.IsOfIncompatibleType() && !Result.IsOfCompatibleMutability()");
 		AddErrorIfFalse(Stack.GetParamPtr<FLinearColor>("Param0", &Result) != nullptr, "Correctly typed access is null");
-		AddErrorIfFalse(Result == FParamStack::EGetParamResult::Succeeded, "Result != FParamStack::EGetParamResult::Succeeded");
+		AddErrorIfFalse(Result.IsSuccessful(), "Unexpected result");
 
 		Stack.PopLayer(Handle);
 
@@ -369,13 +369,13 @@ bool FParamStackTest::RunTest(const FString& InParameters)
 		Handle = Stack.PushValue("Param0", AnimSequence);
 
 		AddErrorIfFalse(Stack.GetParamPtr<UAnimSequence>("Param0", &Result) != nullptr, "Correctly typed access is null");
-		AddErrorIfFalse(Result == FParamStack::EGetParamResult::Succeeded, "Result != FParamStack::EGetParamResult::Succeeded");
+		AddErrorIfFalse(Result.IsSuccessful(), "Unexpected result");
 		AddErrorIfFalse(Stack.GetParamPtr<UAnimSequenceBase>("Param0", &Result) != nullptr, "Base class access is null");
-		AddErrorIfFalse(Result == FParamStack::EGetParamResult::CompatibleType, "Result != FParamStack::EGetParamResult::Succeeded");
+		AddErrorIfFalse(Result.IsOfCompatibleType(), "Unexpected result");
 		AddErrorIfFalse(Stack.GetParamPtr<UAnimationAsset>("Param0", &Result) != nullptr, "Ancestor class access is null");
-		AddErrorIfFalse(Result == FParamStack::EGetParamResult::CompatibleType, "Result != FParamStack::EGetParamResult::Succeeded");
+		AddErrorIfFalse(Result.IsOfCompatibleType(), "Unexpected result");
 		AddErrorIfFalse(Stack.GetParamPtr<UAnimMontage>("Param0", &Result) == nullptr, "Sibling class access is non-null");
-		AddErrorIfFalse(Result == FParamStack::EGetParamResult::IncompatibleType, "Result != FParamStack::EGetParamResult::IncompatibleType");
+		AddErrorIfFalse(Result.IsOfIncompatibleType(), "Unexpected result");
 
 		Stack.PopLayer(Handle);
 	}
@@ -384,9 +384,9 @@ bool FParamStackTest::RunTest(const FString& InParameters)
 	{
 		FParamStack Stack;
 
-		FParamStack::FLayerHandle Layer0 = FParamStack::MakeValueLayer("Param0", 1.0f);
-		FParamStack::FLayerHandle Layer1 = FParamStack::MakeValuesLayer("Param1", 2.0, "Param3", 5);
-		FParamStack::FLayerHandle Layer2 = FParamStack::MakeValueLayer(ParamIds[2], true);
+		FParamStackLayerHandle Layer0 = FParamStack::MakeValueLayer("Param0", 1.0f);
+		FParamStackLayerHandle Layer1 = FParamStack::MakeValuesLayer("Param1", 2.0, "Param3", 5);
+		FParamStackLayerHandle Layer2 = FParamStack::MakeValueLayer(ParamIds[2], true);
 
 		Stack.PushLayer(Layer0);
 		Stack.PushLayer(Layer1);
@@ -519,7 +519,7 @@ bool FParamStackTest::RunTest(const FString& InParameters)
 		PropertyBag.SetValueBool("Param1", true);
 		PropertyBag.SetValueInt32("Param2", 5);
 
-		FParamStack::FLayerHandle Layer = FParamStack::MakeLayer(PropertyBag);
+		FParamStackLayerHandle Layer = FParamStack::MakeReferenceLayer(PropertyBag);
 
 		FParamStack Stack;
 
diff --git a/Engine/Plugins/Experimental/AnimNext/Source/AnimNextUncookedOnly/Internal/Param/AnimNextParameter.h b/Engine/Plugins/Experimental/AnimNext/Source/AnimNextUncookedOnly/Internal/Param/AnimNextParameter.h
--- a/Engine/Plugins/Experimental/AnimNext/Source/AnimNextUncookedOnly/Internal/Param/AnimNextParameter.h
+++ b/Engine/Plugins/Experimental/AnimNext/Source/AnimNextUncookedOnly/Internal/Param/AnimNextParameter.h
@@ -7,6 +7,7 @@
 #include "AnimNextParameter.generated.h"
 
 class UAnimNextParameter;
+class UAnimNextParameterBlockParameter;
 
 namespace UE::AnimNext::Editor
 {
@@ -31,13 +32,13 @@ class UAnimNextParameter : public UObject
 #if WITH_EDITORONLY_DATA
 	friend class UAnimNextParameterBlock_EditorData;
 	friend class UAnimNextParameterLibrary;
+	friend class UAnimNextParameterBlockParameter;
 	friend class UAnimNextParameterBlockBinding;
 	friend class UAnimNextParameterBlockBindingReference;
 	friend struct UE::AnimNext::Editor::FUtils;
 	friend class UE::AnimNext::Editor::SParameterBlockView;
 	friend struct UE::AnimNext::Editor::FParameterPickerEntry;
 	friend class UE::AnimNext::Editor::SParameterLibraryViewRow;
-	friend struct FRigVMDispatch_SetParameter;
 
 	// Get the type of this parameter
 	const FAnimNextParamType& GetType() const { return Type; }
diff --git a/Engine/Plugins/Experimental/AnimNext/Source/AnimNextUncookedOnly/Internal/Param/AnimNextParameterBlockBinding.h b/Engine/Plugins/Experimental/AnimNext/Source/AnimNextUncookedOnly/Internal/Param/AnimNextParameterBlockBinding.h
--- a/Engine/Plugins/Experimental/AnimNext/Source/AnimNextUncookedOnly/Internal/Param/AnimNextParameterBlockBinding.h
+++ b/Engine/Plugins/Experimental/AnimNext/Source/AnimNextUncookedOnly/Internal/Param/AnimNextParameterBlockBinding.h
@@ -1,11 +1,11 @@
-﻿// Copyright Epic Games, Inc. All Rights Reserved.
+// Copyright Epic Games, Inc. All Rights Reserved.
 
 #pragma once
 
 #include "CoreMinimal.h"
 #include "AnimNextParameterBlockEntry.h"
 #include "IAnimNextParameterBlockGraphInterface.h"
-#include "IAnimNextParameterBlockBindingInterface.h"
+#include "IAnimNextParameterBlockParameterInterface.h"
 #include "AnimNextParameterBlockBinding.generated.h"
 
 class UAssetDefinition_AnimNextParameterBlockBinding;
@@ -20,7 +20,7 @@ namespace UE::AnimNext::Editor
 
 /** Parameter binding block entry */
 UCLASS(MinimalAPI, BlueprintType)
-class UAnimNextParameterBlockBinding : public UAnimNextParameterBlockEntry, public IAnimNextParameterBlockBindingInterface, public IAnimNextParameterBlockGraphInterface
+class UAnimNextParameterBlockBinding : public UAnimNextParameterBlockEntry, public IAnimNextParameterBlockParameterInterface, public IAnimNextParameterBlockGraphInterface
 {
 	GENERATED_BODY()
 
@@ -28,7 +28,7 @@ class UAnimNextParameterBlockBinding : public UAnimNextParameterBlockEntry, publ
 	friend class UAssetDefinition_AnimNextParameterBlockBinding;
 	friend class UE::AnimNext::Editor::SParameterBlockViewRow;
 
-	// IAnimNextParameterBlockBindingInterface interface
+	// IAnimNextParameterBlockParameterInterface interface
 	virtual FAnimNextParamType GetParamType() const override;
 	virtual void SetParameterName(FName InName, bool bSetupUndoRedo = true) override;
 	virtual FName GetParameterName() const override;
@@ -37,6 +37,8 @@ class UAnimNextParameterBlockBinding : public UAnimNextParameterBlockEntry, publ
 
 	// IAnimNextParameterBlockGraphInterface interface
 	virtual URigVMGraph* GetGraph() const override { return BindingGraph; }
+	virtual FName GetGraphName() const override { return ParameterName; }
+	virtual void SetGraphName(FName InName, bool bSetupUndoRedo = true) override { SetParameterName(InName, bSetupUndoRedo); };
 
 	// UAnimNextParameterBlockEntry interface
 	virtual FText GetDisplayName() const override;
diff --git a/Engine/Plugins/Experimental/AnimNext/Source/AnimNextUncookedOnly/Internal/Param/AnimNextParameterBlockBindingReference.h b/Engine/Plugins/Experimental/AnimNext/Source/AnimNextUncookedOnly/Internal/Param/AnimNextParameterBlockBindingReference.h
--- a/Engine/Plugins/Experimental/AnimNext/Source/AnimNextUncookedOnly/Internal/Param/AnimNextParameterBlockBindingReference.h
+++ b/Engine/Plugins/Experimental/AnimNext/Source/AnimNextUncookedOnly/Internal/Param/AnimNextParameterBlockBindingReference.h
@@ -1,11 +1,11 @@
-﻿// Copyright Epic Games, Inc. All Rights Reserved.
+// Copyright Epic Games, Inc. All Rights Reserved.
 
 #pragma once
 
 #include "CoreMinimal.h"
 #include "AnimNextParameterBlockEntry.h"
 #include "IAnimNextParameterBlockReferenceInterface.h"
-#include "IAnimNextParameterBlockBindingInterface.h"
+#include "IAnimNextParameterBlockParameterInterface.h"
 #include "AnimNextParameterBlockBindingReference.generated.h"
 
 class UAnimNextParameter;
@@ -21,14 +21,14 @@ namespace UE::AnimNext::Editor
 
 /** Parameter binding block entry */
 UCLASS(MinimalAPI, BlueprintType)
-class UAnimNextParameterBlockBindingReference : public UAnimNextParameterBlockEntry, public IAnimNextParameterBlockBindingInterface, public IAnimNextParameterBlockReferenceInterface
+class UAnimNextParameterBlockBindingReference : public UAnimNextParameterBlockEntry, public IAnimNextParameterBlockParameterInterface, public IAnimNextParameterBlockReferenceInterface
 {
 	GENERATED_BODY()
 
 	friend class UAnimNextParameterBlock_EditorData;
 	friend class UAssetDefinition_AnimNextParameterBlockBindingReference;
 
-	// IAnimNextParameterBlockBindingInterface interface
+	// IAnimNextParameterBlockParameterInterface interface
 	virtual FAnimNextParamType GetParamType() const override;
 	virtual FName GetParameterName() const override;
 	virtual void SetParameterName(FName InName, bool bSetupUndoRedo = true) override;
diff --git a/Engine/Plugins/Experimental/AnimNext/Source/AnimNextUncookedOnly/Internal/Param/AnimNextParameterBlockGraph.h b/Engine/Plugins/Experimental/AnimNext/Source/AnimNextUncookedOnly/Internal/Param/AnimNextParameterBlockGraph.h
new file mode 100644
--- /dev/null
+++ b/Engine/Plugins/Experimental/AnimNext/Source/AnimNextUncookedOnly/Internal/Param/AnimNextParameterBlockGraph.h
@@ -0,0 +1,37 @@
+﻿// Copyright Epic Games, Inc. All Rights Reserved.
+
+#pragma once
+
+#include "CoreMinimal.h"
+#include "AnimNextParameterBlockEntry.h"
+#include "IAnimNextParameterBlockGraphInterface.h"
+#include "AnimNextParameterBlockGraph.generated.h"
+
+class UAnimNextParameterLibrary;
+class UAnimNextParameterBlock_EditorData;
+class URigVMGraph;
+
+UCLASS()
+class UAnimNextParameterBlockGraph : public UAnimNextParameterBlockEntry, public IAnimNextParameterBlockGraphInterface
+{
+	GENERATED_BODY()
+
+	friend class UAnimNextParameterBlock_EditorData;
+
+	// UAnimNextParameterBlockEntry interface
+	virtual FText GetDisplayName() const override;
+	virtual FText GetDisplayNameTooltip() const override;
+
+	// IAnimNextParameterBlockGraphInterface interface
+	virtual URigVMGraph* GetGraph() const override { return Graph; }
+	virtual FName GetGraphName() const override { return GraphName; }
+	virtual void SetGraphName(FName InName, bool bSetupUndoRedo = true) override;
+
+	/** The name of the graph */
+	UPROPERTY(VisibleAnywhere, Category = Parameter)
+	FName GraphName;
+
+	/** Graph */
+	UPROPERTY()
+	TObjectPtr<URigVMGraph> Graph;
+};
\ No newline at end of file
diff --git a/Engine/Plugins/Experimental/AnimNext/Source/AnimNextUncookedOnly/Internal/Param/AnimNextParameterBlockParameter.h b/Engine/Plugins/Experimental/AnimNext/Source/AnimNextUncookedOnly/Internal/Param/AnimNextParameterBlockParameter.h
new file mode 100644
--- /dev/null
+++ b/Engine/Plugins/Experimental/AnimNext/Source/AnimNextUncookedOnly/Internal/Param/AnimNextParameterBlockParameter.h
@@ -0,0 +1,38 @@
+﻿// Copyright Epic Games, Inc. All Rights Reserved.
+
+#pragma once
+
+#include "CoreMinimal.h"
+#include "AnimNextParameterBlockEntry.h"
+#include "IAnimNextParameterBlockParameterInterface.h"
+#include "AnimNextParameterBlockParameter.generated.h"
+
+class UAnimNextParameterLibrary;
+class UAnimNextParameterBlock_EditorData;
+
+UCLASS()
+class UAnimNextParameterBlockParameter : public UAnimNextParameterBlockEntry, public IAnimNextParameterBlockParameterInterface
+{
+	GENERATED_BODY()
+
+	friend class UAnimNextParameterBlock_EditorData;
+	
+	// UAnimNextParameterBlockEntry interface
+	virtual FText GetDisplayName() const override;
+	virtual FText GetDisplayNameTooltip() const override;
+
+	// IAnimNextParameterBlockParameterInterface interface
+	virtual FAnimNextParamType GetParamType() const override;
+	virtual void SetParameterName(FName InName, bool bSetupUndoRedo = true) override;
+	virtual FName GetParameterName() const override;
+	virtual const UAnimNextParameter* GetParameter() const override;
+	virtual const UAnimNextParameterLibrary* GetLibrary() const override;
+	
+	/** Parameter name we reference */
+	UPROPERTY(VisibleAnywhere, Category = Parameter)
+	FName ParameterName;
+
+	/** Parameter library we reference */
+	UPROPERTY(VisibleAnywhere, Category = Parameter)
+	TObjectPtr<UAnimNextParameterLibrary> Library;
+};
\ No newline at end of file
diff --git a/Engine/Plugins/Experimental/AnimNext/Source/AnimNextUncookedOnly/Internal/Param/AnimNextParameterBlock_EdGraph.h b/Engine/Plugins/Experimental/AnimNext/Source/AnimNextUncookedOnly/Internal/Param/AnimNextParameterBlock_EdGraph.h
--- a/Engine/Plugins/Experimental/AnimNext/Source/AnimNextUncookedOnly/Internal/Param/AnimNextParameterBlock_EdGraph.h
+++ b/Engine/Plugins/Experimental/AnimNext/Source/AnimNextUncookedOnly/Internal/Param/AnimNextParameterBlock_EdGraph.h
@@ -3,7 +3,7 @@
 #pragma once
 
 #include "CoreMinimal.h"
-#include "Graph/ControlRigGraph.h"
+#include "EdGraph/RigVMEdGraph.h"
 #include "AnimNextParameterBlock_EdGraph.generated.h"
 
 class UAnimNextParameterBlock_EditorData;
@@ -14,7 +14,7 @@ namespace UE::AnimNext::UncookedOnly
 }
 
 UCLASS(MinimalAPI)
-class UAnimNextParameterBlock_EdGraph : public UControlRigGraph
+class UAnimNextParameterBlock_EdGraph : public URigVMEdGraph
 {
 	GENERATED_BODY()
 
diff --git a/Engine/Plugins/Experimental/AnimNext/Source/AnimNextUncookedOnly/Internal/Param/AnimNextParameterBlock_EdGraphNode.h b/Engine/Plugins/Experimental/AnimNext/Source/AnimNextUncookedOnly/Internal/Param/AnimNextParameterBlock_EdGraphNode.h
--- a/Engine/Plugins/Experimental/AnimNext/Source/AnimNextUncookedOnly/Internal/Param/AnimNextParameterBlock_EdGraphNode.h
+++ b/Engine/Plugins/Experimental/AnimNext/Source/AnimNextUncookedOnly/Internal/Param/AnimNextParameterBlock_EdGraphNode.h
@@ -3,11 +3,11 @@
 #pragma once
 
 #include "CoreMinimal.h"
-#include "Graph/ControlRigGraphNode.h"
+#include "EdGraph/RigVMEdGraphNode.h"
 #include "AnimNextParameterBlock_EdGraphNode.generated.h"
 
 UCLASS(MinimalAPI)
-class UAnimNextParameterBlock_EdGraphNode : public UControlRigGraphNode
+class UAnimNextParameterBlock_EdGraphNode : public URigVMEdGraphNode
 {
 	GENERATED_BODY()
 };
\ No newline at end of file
diff --git a/Engine/Plugins/Experimental/AnimNext/Source/AnimNextUncookedOnly/Internal/Param/AnimNextParameterBlock_EdGraphSchema.h b/Engine/Plugins/Experimental/AnimNext/Source/AnimNextUncookedOnly/Internal/Param/AnimNextParameterBlock_EdGraphSchema.h
--- a/Engine/Plugins/Experimental/AnimNext/Source/AnimNextUncookedOnly/Internal/Param/AnimNextParameterBlock_EdGraphSchema.h
+++ b/Engine/Plugins/Experimental/AnimNext/Source/AnimNextUncookedOnly/Internal/Param/AnimNextParameterBlock_EdGraphSchema.h
@@ -3,18 +3,15 @@
 #pragma once
 
 #include "CoreMinimal.h"
-#include "Graph/ControlRigGraphSchema.h"
+#include "EdGraph/RigVMEdGraphSchema.h"
 #include "AnimNextParameterBlock_EdGraphNode.h"
 #include "AnimNextParameterBlock_EdGraphSchema.generated.h"
 
 UCLASS(MinimalAPI)
-class UAnimNextParameterBlock_EdGraphSchema : public UControlRigGraphSchema
+class UAnimNextParameterBlock_EdGraphSchema : public URigVMEdGraphSchema
 {
 	GENERATED_BODY()
 
-	// UControlRigGraphSchema interface
-	virtual TSubclassOf<URigVMEdGraphNode> GetGraphNodeClass(const URigVMEdGraph* InGraph) const override { return UAnimNextParameterBlock_EdGraphNode::StaticClass(); }
-
-	// UEdGraphSchema interface
-	virtual void TrySetDefaultValue(UEdGraphPin& InPin, const FString& InNewDefaultValue, bool bMarkAsModified) const override;
+	// URigVMEdGraphSchema interface
+	virtual TSubclassOf<URigVMEdGraphNode> GetGraphNodeClass(const URigVMEdGraph* InGraph) const override { return UAnimNextParameterBlock_EdGraphNode::StaticClass(); }	
 };
\ No newline at end of file
diff --git a/Engine/Plugins/Experimental/AnimNext/Source/AnimNextUncookedOnly/Internal/Param/AnimNextParameterBlock_EditorData.h b/Engine/Plugins/Experimental/AnimNext/Source/AnimNextUncookedOnly/Internal/Param/AnimNextParameterBlock_EditorData.h
--- a/Engine/Plugins/Experimental/AnimNext/Source/AnimNextUncookedOnly/Internal/Param/AnimNextParameterBlock_EditorData.h
+++ b/Engine/Plugins/Experimental/AnimNext/Source/AnimNextUncookedOnly/Internal/Param/AnimNextParameterBlock_EditorData.h
@@ -7,10 +7,12 @@
 #include "AnimNextParameterBlock_EdGraph.h"
 #include "RigVMCore/RigVMGraphFunctionHost.h"
 #include "RigVMBlueprint.h"
+#include "Kismet/BlueprintFunctionLibrary.h"
 #include "AnimNextParameterBlock_EditorData.generated.h"
 
 class UAnimNextParameterBlock;
 class UAnimNextParameterBlockEntry;
+class UAnimNextParameterBlockParameter;
 class UAnimNextParameter;
 class UAnimNextParameterBlock_Controller;
 enum class ERigVMGraphNotifType : uint8;
@@ -54,6 +56,12 @@ UCLASS()
 class ANIMNEXTUNCOOKEDONLY_API UAnimNextParameterBlockLibrary : public UBlueprintFunctionLibrary
 {
 	GENERATED_BODY()
+
+	UFUNCTION(BlueprintCallable, Category = "AnimNext|Parameter Block", meta=(ScriptMethod))
+	static UAnimNextParameterBlockParameter* AddParameter(UAnimNextParameterBlock* InBlock, FName InName, UAnimNextParameterLibrary* InLibrary, bool bSetupUndoRedo = true, bool bPrintPythonCommand = true);
+
+	UFUNCTION(BlueprintCallable, Category = "AnimNext|Parameter Block", meta=(ScriptMethod))
+	static UAnimNextParameterBlockGraph* AddGraph(UAnimNextParameterBlock* InBlock, FName InName, bool bSetupUndoRedo = true, bool bPrintPythonCommand = true);
 	
 	UFUNCTION(BlueprintCallable, Category = "AnimNext|Parameter Block", meta=(ScriptMethod))
 	static UAnimNextParameterBlockBinding* AddBinding(UAnimNextParameterBlock* InBlock, FName InName, UAnimNextParameterLibrary* InLibrary, bool bSetupUndoRedo = true, bool bPrintPythonCommand = true);
@@ -100,6 +108,10 @@ class UAnimNextParameterBlock_EditorData : public UObject, public IRigVMClientHo
 	friend class UE::AnimNext::Editor::FParameterBlockTabSummoner;
 	friend class FAnimationAnimNextParametersEditorTest_Block;
 
+	ANIMNEXTUNCOOKEDONLY_API UAnimNextParameterBlockParameter* AddParameter(FName InName, UAnimNextParameterLibrary* InLibrary, bool bSetupUndoRedo = true, bool bPrintPythonCommand = true);
+
+	ANIMNEXTUNCOOKEDONLY_API UAnimNextParameterBlockGraph* AddGraph(FName InName, bool bSetupUndoRedo = true, bool bPrintPythonCommand = true);
+	
 	ANIMNEXTUNCOOKEDONLY_API UAnimNextParameterBlockBinding* AddBinding(FName InName, UAnimNextParameterLibrary* InLibrary, bool bSetupUndoRedo = true, bool bPrintPythonCommand = true);
 
 	ANIMNEXTUNCOOKEDONLY_API UAnimNextParameterBlockBindingReference* AddBindingReference(FName InName, UAnimNextParameterLibrary* InLibrary, UAnimNextParameterBlock* InBlock, bool bSetupUndoRedo = true, bool bPrintPythonCommand = true);
@@ -171,6 +183,8 @@ class UAnimNextParameterBlock_EditorData : public UObject, public IRigVMClientHo
 
 	void ReportError(const TCHAR* InMessage) const;
 
+	void ReconstructAllNodes();
+	
 	UPROPERTY()
 	TArray<TObjectPtr<UAnimNextParameterBlock_EdGraph>> Graphs;
 
diff --git a/Engine/Plugins/Experimental/AnimNext/Source/AnimNextUncookedOnly/Internal/Param/AnimNextParameterLibrary.h b/Engine/Plugins/Experimental/AnimNext/Source/AnimNextUncookedOnly/Internal/Param/AnimNextParameterLibrary.h
--- a/Engine/Plugins/Experimental/AnimNext/Source/AnimNextUncookedOnly/Internal/Param/AnimNextParameterLibrary.h
+++ b/Engine/Plugins/Experimental/AnimNext/Source/AnimNextUncookedOnly/Internal/Param/AnimNextParameterLibrary.h
@@ -8,6 +8,7 @@
 
 class UAnimNextParameter;
 class UAnimNextParameterLibrary;
+class UAnimNextParameterBlockParameter;
 class UAnimNextParameterBlockBinding;
 class UAnimNextParameterBlockBindingReference;
 class UAnimNextParameterBlock_EditorData;
@@ -19,6 +20,7 @@ namespace UE::AnimNext::Editor
 	struct FUtils;
 	class SParameterBlockView;
 	class SParameterLibraryView;
+	class SParameterPicker;
 }
 
 namespace UE::AnimNext::UncookedOnly
@@ -69,7 +71,9 @@ class UAnimNextParameterLibrary : public UObject
 	friend struct UE::AnimNext::Editor::FUtils;
 	friend class UE::AnimNext::Editor::SParameterBlockView;
 	friend class UE::AnimNext::Editor::SParameterLibraryView;
+	friend class UE::AnimNext::Editor::SParameterPicker;
 	friend class UAnimNextParameterBlockBinding;
+	friend class UAnimNextParameterBlockParameter;
 	friend class UAnimNextParameterBlockBindingReference;
 	friend class UAnimNextParameterBlock_EditorData;
 	friend class FAnimationAnimNextParametersEditorTest_Library;
diff --git a/Engine/Plugins/Experimental/AnimNext/Source/AnimNextUncookedOnly/Internal/Param/IAnimNextParameterBlockGraphInterface.h b/Engine/Plugins/Experimental/AnimNext/Source/AnimNextUncookedOnly/Internal/Param/IAnimNextParameterBlockGraphInterface.h
--- a/Engine/Plugins/Experimental/AnimNext/Source/AnimNextUncookedOnly/Internal/Param/IAnimNextParameterBlockGraphInterface.h
+++ b/Engine/Plugins/Experimental/AnimNext/Source/AnimNextUncookedOnly/Internal/Param/IAnimNextParameterBlockGraphInterface.h
@@ -37,6 +37,12 @@ class ANIMNEXTUNCOOKEDONLY_API IAnimNextParameterBlockGraphInterface
 	friend struct UE::AnimNext::UncookedOnly::FUtilsPrivate;
 	friend class UE::AnimNext::Editor::FParameterBlockTabSummoner;
 
+	// Get the graph name
+	virtual FName GetGraphName() const = 0;
+
+	// Set the graph name
+	virtual void SetGraphName(FName InName, bool bSetupUndoRedo = true) = 0;
+
 	// Get the graph
 	virtual URigVMGraph* GetGraph() const = 0;
 };
\ No newline at end of file
diff --git a/Engine/Plugins/Experimental/AnimNext/Source/AnimNextUncookedOnly/Internal/Param/IAnimNextParameterBlockParameterInterface.h b/Engine/Plugins/Experimental/AnimNext/Source/AnimNextUncookedOnly/Internal/Param/IAnimNextParameterBlockParameterInterface.h
new file mode 100644
--- /dev/null
+++ b/Engine/Plugins/Experimental/AnimNext/Source/AnimNextUncookedOnly/Internal/Param/IAnimNextParameterBlockParameterInterface.h
@@ -0,0 +1,58 @@
+﻿// Copyright Epic Games, Inc. All Rights Reserved.
+
+#pragma once
+
+#include "Param/ParamType.h"
+#include "UObject/Interface.h"
+#include "IAnimNextParameterBlockParameterInterface.generated.h"
+
+class UAnimNextParameter;
+class UAnimNextParameterLibrary;
+class UAnimNextParameterBlock_EditorData;
+
+namespace UE::AnimNext
+{
+	struct FParamDefinition;
+}
+
+namespace UE::AnimNext::Editor
+{
+	class SParameterBlockView;
+	class SParameterBlockViewRow;
+}
+
+namespace UE::AnimNext::UncookedOnly
+{
+	struct FUtils;
+}
+
+UINTERFACE(meta=(CannotImplementInterfaceInBlueprint))
+class ANIMNEXTUNCOOKEDONLY_API UAnimNextParameterBlockParameterInterface : public UInterface
+{
+	GENERATED_BODY()
+};
+
+class ANIMNEXTUNCOOKEDONLY_API IAnimNextParameterBlockParameterInterface
+{
+	GENERATED_BODY()
+
+	friend class UAnimNextParameterBlock_EditorData;
+	friend class UE::AnimNext::Editor::SParameterBlockView;
+	friend class UE::AnimNext::Editor::SParameterBlockViewRow;
+	friend struct UE::AnimNext::UncookedOnly::FUtils;
+
+	// Get the parameter type
+	virtual FAnimNextParamType GetParamType() const = 0;
+
+	// Get the parameter name
+	virtual FName GetParameterName() const = 0;
+
+	// Set the parameter name
+	virtual void SetParameterName(FName InName, bool bSetupUndoRedo = true) = 0;
+
+	// Resolve the underlying parameter in the library
+	virtual const UAnimNextParameter* GetParameter() const = 0;
+
+	// Get the parameter library hosting the parameter
+	virtual const UAnimNextParameterLibrary* GetLibrary() const = 0;
+};
\ No newline at end of file
diff --git a/Engine/Plugins/Experimental/AnimNext/Source/AnimNextUncookedOnly/Internal/UncookedOnlyUtils.h b/Engine/Plugins/Experimental/AnimNext/Source/AnimNextUncookedOnly/Internal/UncookedOnlyUtils.h
--- a/Engine/Plugins/Experimental/AnimNext/Source/AnimNextUncookedOnly/Internal/UncookedOnlyUtils.h
+++ b/Engine/Plugins/Experimental/AnimNext/Source/AnimNextUncookedOnly/Internal/UncookedOnlyUtils.h
@@ -3,6 +3,7 @@
 #pragma once
 
 #include "CoreMinimal.h"
+#include "Param/AnimNextParameterBlock_EditorData.h"
 #include "Param/ParamTypeHandle.h"
 #include "RigVMCore/RigVMTemplate.h"
 
@@ -67,11 +68,30 @@ struct ANIMNEXTUNCOOKEDONLY_API FUtils
 	static FRigVMTemplateArgumentType GetRigVMArgTypeFromParamTypeHandle(const FParamTypeHandle& InParamTypeHandle);
 	static FRigVMTemplateArgumentType GetRigVMArgTypeFromParamType(const FAnimNextParamType& InParamType);
 
+	/** Set up a simple graph */
+	static void SetupGraph(URigVMController* InController);
+	
 	/** Set up a binding graph given the type to set */
+	static void SetupBindingGraph(URigVMController* InController, FName InParameterName, const FAnimNextParamType& InParamType);
+
+	/** Set up a binding graph for a simple literal value given the type to set */
 	static void SetupBindingGraphForLiteral(URigVMController* InController, FName InParameterName, const FAnimNextParamType& InParamType);
 
 	/** Converts the Verse-tag-like snake_case_parameter_name to a period-separated display name similar to a gameplay tag */
 	static FText GetParameterDisplayNameText(FName InParameterName);
+
+	/** Returns all nodes in all graphs of the specified class */
+	template<class T>
+	static void GetAllNodesOfClass(const UAnimNextParameterBlock_EditorData* InEditorData, TArray<T*>& OutNodes)
+	{
+		for(const UAnimNextParameterBlock_EdGraph* Graph : InEditorData->Graphs)
+		{
+			check(Graph);
+			TArray<T*> GraphNodes;
+			Graph->GetNodesOfClass<T>(GraphNodes);
+			OutNodes.Append(GraphNodes);
+		}
+	}
 };
 
 }
\ No newline at end of file
diff --git a/Engine/Plugins/Experimental/AnimNext/Source/AnimNextUncookedOnly/Private/Graph/AnimGraph/AnimGraphNode_AnimNextParameters.cpp b/Engine/Plugins/Experimental/AnimNext/Source/AnimNextUncookedOnly/Private/Graph/AnimGraph/AnimGraphNode_AnimNextParameters.cpp
--- a/Engine/Plugins/Experimental/AnimNext/Source/AnimNextUncookedOnly/Private/Graph/AnimGraph/AnimGraphNode_AnimNextParameters.cpp
+++ b/Engine/Plugins/Experimental/AnimNext/Source/AnimNextUncookedOnly/Private/Graph/AnimGraph/AnimGraphNode_AnimNextParameters.cpp
@@ -1,6 +1,7 @@
 // Copyright Epic Games, Inc. All Rights Reserved.
 
 #include "Graph/AnimGraph/AnimGraphNode_AnimNextParameters.h"
+#include "Graph/AnimGraph/AnimBlueprintExtension_AnimNextParameters.h"
 
 #define LOCTEXT_NAMESPACE "AnimGraphNode_AnimNextParameters"
 
@@ -19,4 +20,9 @@ FText UAnimGraphNode_AnimNextParameters::GetMenuCategory() const
 	return LOCTEXT("Category", "AnimNext");
 }
 
+void UAnimGraphNode_AnimNextParameters::GetRequiredExtensions(TArray<TSubclassOf<UAnimBlueprintExtension>>& OutExtensions) const
+{
+	OutExtensions.Add(UAnimBlueprintExtension_AnimNextParameters::StaticClass());
+}
+
 #undef LOCTEXT_NAMESPACE
\ No newline at end of file
diff --git a/Engine/Plugins/Experimental/AnimNext/Source/AnimNextUncookedOnly/Private/Param/AnimNextParameterBlockGraph.cpp b/Engine/Plugins/Experimental/AnimNext/Source/AnimNextUncookedOnly/Private/Param/AnimNextParameterBlockGraph.cpp
new file mode 100644
--- /dev/null
+++ b/Engine/Plugins/Experimental/AnimNext/Source/AnimNextUncookedOnly/Private/Param/AnimNextParameterBlockGraph.cpp
@@ -0,0 +1,25 @@
+﻿// Copyright Epic Games, Inc. All Rights Reserved.
+
+#include "Param/AnimNextParameterBlockGraph.h"
+
+FText UAnimNextParameterBlockGraph::GetDisplayName() const
+{
+	return FText::FromName(GraphName);
+}
+
+FText UAnimNextParameterBlockGraph::GetDisplayNameTooltip() const
+{
+	return FText::FromName(GraphName);
+}
+
+void UAnimNextParameterBlockGraph::SetGraphName(FName InName, bool bSetupUndoRedo)
+{
+	if(bSetupUndoRedo)
+	{
+		Modify();
+	}
+
+	GraphName = InName;
+
+	BroadcastModified();
+}
\ No newline at end of file
diff --git a/Engine/Plugins/Experimental/AnimNext/Source/AnimNextUncookedOnly/Private/Param/AnimNextParameterBlockParameter.cpp b/Engine/Plugins/Experimental/AnimNext/Source/AnimNextUncookedOnly/Private/Param/AnimNextParameterBlockParameter.cpp
new file mode 100644
--- /dev/null
+++ b/Engine/Plugins/Experimental/AnimNext/Source/AnimNextUncookedOnly/Private/Param/AnimNextParameterBlockParameter.cpp
@@ -0,0 +1,71 @@
+﻿// Copyright Epic Games, Inc. All Rights Reserved.
+
+#include "Param/AnimNextParameterBlockParameter.h"
+#include "Param/AnimNextParameter.h"
+#include "Param/AnimNextParameterLibrary.h"
+#include "Param/Params.h"
+
+FAnimNextParamType UAnimNextParameterBlockParameter::GetParamType() const
+{
+	using namespace UE::AnimNext;
+	
+	if(Library)
+	{
+		if(const UAnimNextParameter* Parameter = Library->FindParameter(ParameterName))
+		{
+			return Parameter->GetType();
+		}
+	}
+	else
+	{
+		// Look in built-in parameters
+		if(const FParamDefinition* ParamDefinition = FParams::FindBuiltInParameter(ParameterName))
+		{
+			return ParamDefinition->GetType();
+		}
+	}
+
+	return FAnimNextParamType();
+}
+
+FName UAnimNextParameterBlockParameter::GetParameterName() const
+{
+	return ParameterName;
+}
+
+void UAnimNextParameterBlockParameter::SetParameterName(FName InName, bool bSetupUndoRedo)
+{
+	if(bSetupUndoRedo)
+	{
+		Modify();
+	}
+
+	ParameterName = InName;
+
+	BroadcastModified();
+}
+
+const UAnimNextParameter* UAnimNextParameterBlockParameter::GetParameter() const
+{
+	if(Library)
+	{
+		return Library->FindParameter(ParameterName);
+	}
+
+	return nullptr;
+}
+
+const UAnimNextParameterLibrary* UAnimNextParameterBlockParameter::GetLibrary() const
+{
+	return Library;
+}
+
+FText UAnimNextParameterBlockParameter::GetDisplayName() const
+{
+	return FText::FromName(ParameterName);
+}
+
+FText UAnimNextParameterBlockParameter::GetDisplayNameTooltip() const
+{
+	return FText::FromName(ParameterName);
+}
\ No newline at end of file
diff --git a/Engine/Plugins/Experimental/AnimNext/Source/AnimNextUncookedOnly/Private/Param/AnimNextParameterBlock_EdGraphSchema.cpp b/Engine/Plugins/Experimental/AnimNext/Source/AnimNextUncookedOnly/Private/Param/AnimNextParameterBlock_EdGraphSchema.cpp
--- a/Engine/Plugins/Experimental/AnimNext/Source/AnimNextUncookedOnly/Private/Param/AnimNextParameterBlock_EdGraphSchema.cpp
+++ b/Engine/Plugins/Experimental/AnimNext/Source/AnimNextUncookedOnly/Private/Param/AnimNextParameterBlock_EdGraphSchema.cpp
@@ -1,42 +1,4 @@
 ﻿// Copyright Epic Games, Inc. All Rights Reserved.
 
 #include "Param/AnimNextParameterBlock_EdGraphSchema.h"
-#include "Editor.h"
 
-void UAnimNextParameterBlock_EdGraphSchema::TrySetDefaultValue(UEdGraphPin& InPin, const FString& InNewDefaultValue, bool bMarkAsModified) const
-{
-#if WITH_EDITOR
-	if (GEditor)
-	{
-		GEditor->CancelTransaction(0);
-	}
-#endif
-
-	FString UseDefaultValue;
-	TObjectPtr<UObject> UseDefaultObject = nullptr;
-	FText UseDefaultText;
-
-	GetDefault<UEdGraphSchema_K2>()->GetPinDefaultValuesFromString(InPin.PinType, InPin.GetOwningNodeUnchecked(), InNewDefaultValue, UseDefaultValue, UseDefaultObject, UseDefaultText, /*bPreserveTextIdentity*/false);
-
-	// Check the default value and make it an error if it's bogus
-	if (GetDefault<UEdGraphSchema_K2>()->DefaultValueSimpleValidation(InPin.PinType, InPin.PinName, UseDefaultValue, UseDefaultObject, UseDefaultText))
-	{
-		InPin.DefaultObject = UseDefaultObject;
-		InPin.DefaultValue = UseDefaultValue;
-		InPin.DefaultTextValue = UseDefaultText;
-
-		UEdGraphNode* Node = InPin.GetOwningNode();
-		Node->PinDefaultValueChanged(&InPin);
-
-		// If the default value is manually set then treat it as if the value was reset to default and remove the orphaned InPin
-		if (InPin.bOrphanedPin && InPin.DoesDefaultValueMatchAutogenerated())
-		{
-			Node->PinConnectionListChanged(&InPin);
-		}
-
-		if (bMarkAsModified)
-		{
-			Node->MarkPackageDirty();
-		}
-	}
-}
diff --git a/Engine/Plugins/Experimental/AnimNext/Source/AnimNextUncookedOnly/Private/Param/AnimNextParameterBlock_EditorData.cpp b/Engine/Plugins/Experimental/AnimNext/Source/AnimNextUncookedOnly/Private/Param/AnimNextParameterBlock_EditorData.cpp
--- a/Engine/Plugins/Experimental/AnimNext/Source/AnimNextUncookedOnly/Private/Param/AnimNextParameterBlock_EditorData.cpp
+++ b/Engine/Plugins/Experimental/AnimNext/Source/AnimNextUncookedOnly/Private/Param/AnimNextParameterBlock_EditorData.cpp
@@ -6,7 +6,7 @@
 #include "Param/AnimNextParameterBlock_EdGraph.h"
 #include "Param/AnimNextParameterBlock_EdGraphSchema.h"
 #include "UncookedOnlyUtils.h"
-#include "Graph/AnimNextExecuteContext.h"
+#include "Param/AnimNextParameterExecuteContext.h"
 #include "Rigs/RigHierarchyPose.h"
 #include "RigVMModel/RigVMFunctionLibrary.h"
 #include "RigVMModel/RigVMNotifications.h"
@@ -21,6 +21,8 @@
 #include "Param/AnimNextParameterLibrary.h"
 #include "Misc/TransactionObjectEvent.h"
 #include "ObjectTools.h"
+#include "Param/AnimNextParameterBlockGraph.h"
+#include "Param/AnimNextParameterBlockParameter.h"
 
 #if WITH_EDITORONLY_DATA
 
@@ -58,7 +60,7 @@ UAnimNextParameterBlock_EditorData::UAnimNextParameterBlock_EditorData(const FOb
 {
 	RigVMClient.Reset();
 	RigVMClient.SetSchemaClass(UAnimNextParameterBlockLibrary_Schema::StaticClass());
-	RigVMClient.SetExecuteContextStruct(FAnimNextExecuteContext::StaticStruct());
+	RigVMClient.SetExecuteContextStruct(FAnimNextParameterExecuteContext::StaticStruct());
 	RigVMClient.SetOuterClientHost(this, GET_MEMBER_NAME_CHECKED(UAnimNextParameterBlock_EditorData, RigVMClient));
 	{
 		TGuardValue<bool> DisableClientNotifs(RigVMClient.bSuspendNotifications, true);
@@ -69,6 +71,125 @@ UAnimNextParameterBlock_EditorData::UAnimNextParameterBlock_EditorData(const FOb
 
 #if WITH_EDITORONLY_DATA
 
+UAnimNextParameterBlockParameter* UAnimNextParameterBlockLibrary::AddParameter(UAnimNextParameterBlock* InBlock, FName InName, UAnimNextParameterLibrary* InLibrary, bool bSetupUndoRedo, bool bPrintPythonCommand)
+{
+	return UE::AnimNext::UncookedOnly::FUtils::GetEditorData(InBlock)->AddParameter(InName, InLibrary, bSetupUndoRedo, bPrintPythonCommand);
+}
+
+UAnimNextParameterBlockParameter* UAnimNextParameterBlock_EditorData::AddParameter(FName InName, UAnimNextParameterLibrary* InLibrary, bool bSetupUndoRedo, bool bPrintPythonCommand)
+{
+	if(InName == NAME_None)
+	{
+		ReportError(TEXT("UAnimNextParameterBlock_EditorData::AddParameter: Invalid parameter name supplied."));
+		return nullptr;
+	}
+
+	if(InLibrary != nullptr)
+	{
+		// Check if parameter exists in library
+		UAnimNextParameter* Parameter = InLibrary->FindParameter(InName);
+		if(Parameter == nullptr)
+		{
+			ReportError(TEXT("UAnimNextParameterBlock_EditorData::AddParameter: Parameter does not exist in library."));
+			return nullptr;
+		}
+	}
+
+	// Check for duplicate parameter
+	const bool bAlreadyExists = Entries.ContainsByPredicate([InName](const UAnimNextParameterBlockEntry* InEntry)
+	{
+		if(const UAnimNextParameterBlockParameter* Parameter = Cast<UAnimNextParameterBlockParameter>(InEntry))
+		{
+			return Parameter->ParameterName == InName;
+		}
+		return false;
+	});
+
+	if(bAlreadyExists)
+	{
+		ReportError(TEXT("UAnimNextParameterBlock_EditorData::AddParameter: A parameter already exists for the supplied parameter name."));
+		return nullptr;
+	}
+
+	UAnimNextParameterBlockParameter* NewEntry = UE::AnimNext::Parameters::Private::CreateNewParameterBlockEntry<UAnimNextParameterBlockParameter>(this);
+	NewEntry->ParameterName = InName;
+	NewEntry->Library = InLibrary;
+
+	if(bSetupUndoRedo)
+	{
+		NewEntry->Modify();
+		Modify();
+	}
+	
+	Entries.Add(NewEntry);
+
+	BroadcastModified();
+
+	return NewEntry;
+}
+
+
+UAnimNextParameterBlockGraph* UAnimNextParameterBlockLibrary::AddGraph(UAnimNextParameterBlock* InBlock, FName InName, bool bSetupUndoRedo, bool bPrintPythonCommand)
+{
+	return UE::AnimNext::UncookedOnly::FUtils::GetEditorData(InBlock)->AddGraph(InName, bSetupUndoRedo, bPrintPythonCommand);
+}
+
+UAnimNextParameterBlockGraph* UAnimNextParameterBlock_EditorData::AddGraph(FName InName, bool bSetupUndoRedo, bool bPrintPythonCommand)
+{
+	if(InName == NAME_None)
+	{
+		ReportError(TEXT("UAnimNextParameterBlock_EditorData::AddGraph: Invalid graph name supplied."));
+		return nullptr;
+	}
+
+	// Check for duplicate name
+	FName NewGraphName = InName;
+	auto DuplicateNamePredicate = [&NewGraphName](const UAnimNextParameterBlockEntry* InEntry)
+	{
+		if(const IAnimNextParameterBlockGraphInterface* GraphInterface = Cast<IAnimNextParameterBlockGraphInterface>(InEntry))
+		{
+			return GraphInterface->GetGraphName() == NewGraphName;
+		}
+		return false;
+	};
+
+	bool bAlreadyExists = Entries.ContainsByPredicate(DuplicateNamePredicate);
+	int32 NameNumber = InName.GetNumber() + 1;
+	while(bAlreadyExists)
+	{
+		NewGraphName = FName(InName, NameNumber);
+		bAlreadyExists =  Entries.ContainsByPredicate(DuplicateNamePredicate);
+	}
+
+	UAnimNextParameterBlockGraph* NewEntry = UE::AnimNext::Parameters::Private::CreateNewParameterBlockEntry<UAnimNextParameterBlockGraph>(this);
+	NewEntry->GraphName = NewGraphName;
+
+	if(bSetupUndoRedo)
+	{
+		NewEntry->Modify();
+		Modify();
+	}
+	
+	Entries.Add(NewEntry);
+
+	// Add new graph
+	{
+		TGuardValue<bool> EnablePythonPrint(bSuspendPythonMessagesForRigVMClient, !bPrintPythonCommand);
+		TGuardValue<bool> DisableAutoCompile(bAutoRecompileVM, false);
+		URigVMGraph* NewGraph = RigVMClient.AddModel(NewGraphName, bSetupUndoRedo);
+		ensure(NewGraph);
+		NewEntry->Graph = NewGraph;
+
+		URigVMController* Controller = RigVMClient.GetController(NewGraph);
+		UE::AnimNext::UncookedOnly::FUtils::SetupGraph(Controller);
+	}
+	
+	BroadcastModified();
+
+	return NewEntry;
+}
+
+
 UAnimNextParameterBlockBinding* UAnimNextParameterBlockLibrary::AddBinding(UAnimNextParameterBlock* InBlock, FName InName, UAnimNextParameterLibrary* InLibrary, bool bSetupUndoRedo, bool bPrintPythonCommand)
 {
 	return UE::AnimNext::UncookedOnly::FUtils::GetEditorData(InBlock)->AddBinding(InName, InLibrary, bSetupUndoRedo, bPrintPythonCommand);
@@ -99,7 +220,7 @@ UAnimNextParameterBlockBinding* UAnimNextParameterBlock_EditorData::AddBinding(F
 	// Check for duplicate bindings
 	const bool bAlreadyExists = Entries.ContainsByPredicate([InName, InLibrary](const UAnimNextParameterBlockEntry* InEntry)
 	{
-		if(const IAnimNextParameterBlockBindingInterface* Binding = Cast<IAnimNextParameterBlockBindingInterface>(InEntry))
+		if(const IAnimNextParameterBlockParameterInterface* Binding = Cast<IAnimNextParameterBlockParameterInterface>(InEntry))
 		{
 			return Binding->GetParameterName() == InName && Binding->GetLibrary() == InLibrary;
 		}
@@ -133,7 +254,7 @@ UAnimNextParameterBlockBinding* UAnimNextParameterBlock_EditorData::AddBinding(F
 		NewEntry->BindingGraph = NewGraph;
 
 		URigVMController* Controller = RigVMClient.GetController(NewGraph);
-		UE::AnimNext::UncookedOnly::FUtils::SetupBindingGraphForLiteral(Controller, InName, Parameter->GetType());
+		UE::AnimNext::UncookedOnly::FUtils::SetupBindingGraph(Controller, InName, Parameter->GetType());
 	}
 
 	BroadcastModified();
@@ -177,7 +298,7 @@ UAnimNextParameterBlockBindingReference* UAnimNextParameterBlock_EditorData::Add
 	// Check for duplicates
 	const bool bAlreadyExists = Entries.ContainsByPredicate([InName, InLibrary](const UAnimNextParameterBlockEntry* InEntry)
 	{
-		if(const IAnimNextParameterBlockBindingInterface* Binding = Cast<IAnimNextParameterBlockBindingInterface>(InEntry))
+		if(const IAnimNextParameterBlockParameterInterface* Binding = Cast<IAnimNextParameterBlockParameterInterface>(InEntry))
 		{
 			return Binding->GetParameterName() == InName && Binding->GetLibrary() == InLibrary;
 		}
@@ -225,7 +346,7 @@ bool UAnimNextParameterBlock_EditorData::RemoveAllBindings(FName InName, bool bS
 	// Check parameter binding exists
 	TObjectPtr<UAnimNextParameterBlockEntry>* EntryToRemove = Entries.FindByPredicate([InName](const UAnimNextParameterBlockEntry* InEntry)
 	{
-		if(const IAnimNextParameterBlockBindingInterface* Binding = Cast<IAnimNextParameterBlockBindingInterface>(InEntry))
+		if(const IAnimNextParameterBlockParameterInterface* Binding = Cast<IAnimNextParameterBlockParameterInterface>(InEntry))
 		{
 			return Binding->GetParameterName() == InName;
 		}
@@ -247,7 +368,7 @@ bool UAnimNextParameterBlock_EditorData::RemoveAllBindings(FName InName, bool bS
 
 	Entries.RemoveAll([InName](const UAnimNextParameterBlockEntry* InEntry)
 	{
-		if(const IAnimNextParameterBlockBindingInterface* Binding = Cast<IAnimNextParameterBlockBindingInterface>(InEntry))
+		if(const IAnimNextParameterBlockParameterInterface* Binding = Cast<IAnimNextParameterBlockParameterInterface>(InEntry))
 		{
 			return Binding->GetParameterName() == InName;
 		}
@@ -355,7 +476,7 @@ UAnimNextParameterBlockEntry* UAnimNextParameterBlock_EditorData::FindBinding(FN
 
 	const TObjectPtr<UAnimNextParameterBlockEntry>* FoundEntry = Entries.FindByPredicate([InName](const UAnimNextParameterBlockEntry* InEntry)
 	{
-		if(const IAnimNextParameterBlockBindingInterface* Binding = Cast<IAnimNextParameterBlockBindingInterface>(InEntry))
+		if(const IAnimNextParameterBlockParameterInterface* Binding = Cast<IAnimNextParameterBlockParameterInterface>(InEntry))
 		{
 			return Binding->GetParameterName() == InName;
 		}
@@ -382,6 +503,40 @@ void UAnimNextParameterBlock_EditorData::ReportError(const TCHAR* InMessage) con
 #endif
 }
 
+void UAnimNextParameterBlock_EditorData::ReconstructAllNodes()
+{
+#ifdef WITH_EDITORONLY_DATA
+	// Avoid refreshing EdGraph nodes during cook
+	if (GIsCookerLoadingPackage)
+	{
+		return;
+	}
+	
+	if (GetRigVMClient()->GetDefaultModel() == nullptr)
+	{
+		return;
+	}
+
+	TArray<URigVMEdGraphNode*> AllNodes;
+	UE::AnimNext::UncookedOnly::FUtils::GetAllNodesOfClass(this, AllNodes);
+
+	for (URigVMEdGraphNode* Node : AllNodes)
+	{
+		Node->SetFlags(RF_Transient);
+	}
+
+	for(URigVMEdGraphNode* Node : AllNodes)
+	{
+		Node->ReconstructNode();
+	}
+
+	for (URigVMEdGraphNode* Node : AllNodes)
+	{
+		Node->ClearFlags(RF_Transient);
+	}
+#endif	
+}
+
 void UAnimNextParameterBlock_EditorData::Serialize(FArchive& Ar)
 {
 	RigVMClient.SetOuterClientHost(this, GET_MEMBER_NAME_CHECKED(UAnimNextParameterBlock_EditorData, RigVMClient));
@@ -462,7 +617,7 @@ void UAnimNextParameterBlock_EditorData::GetAssetRegistryTags(TArray<FAssetRegis
 
 	for(const UAnimNextParameterBlockEntry* Entry : Entries)
 	{
-		if(const IAnimNextParameterBlockBindingInterface* Binding = Cast<IAnimNextParameterBlockBindingInterface>(Entry))
+		if(const IAnimNextParameterBlockParameterInterface* Binding = Cast<IAnimNextParameterBlockParameterInterface>(Entry))
 		{
 			Exports.Bindings.Emplace(Binding->GetParameterName(), FSoftObjectPath(Binding->GetLibrary()));
 		}
@@ -489,6 +644,8 @@ void UAnimNextParameterBlock_EditorData::HandlePackageDone()
 	FCoreUObjectDelegates::OnEndLoadPackage.RemoveAll(this);
 
 	RecompileVM();
+
+	ReconstructAllNodes();
 }
 
 void UAnimNextParameterBlock_EditorData::RefreshAllModels(EAnimNextParameterLoadType InLoadType)
@@ -632,7 +789,7 @@ void UAnimNextParameterBlock_EditorData::HandleRigVMGraphAdded(const FRigVMClien
 {
 	if(URigVMGraph* RigVMGraph = InClient->GetModel(InNodePath))
 	{
-		RigVMGraph->SetExecuteContextStruct(FAnimNextExecuteContext::StaticStruct());
+		RigVMGraph->SetExecuteContextStruct(FAnimNextParameterExecuteContext::StaticStruct());
 
 		if(!HasAnyFlags(RF_ClassDefaultObject | RF_NeedInitialization | RF_NeedLoad | RF_NeedPostLoad) &&
 			GetOuter() != GetTransientPackage())
@@ -685,7 +842,7 @@ void UAnimNextParameterBlock_EditorData::HandleConfigureRigVMController(const FR
 			{
 				if (UAnimNextParameterBlock* ParameterBlock = EditorData->GetTypedOuter<UAnimNextParameterBlock>())
 				{
-					return ParameterBlock->GetRigVMExternalVariables();
+					return ParameterBlock->GetExternalVariables();
 				}
 			}
 		}
@@ -699,9 +856,9 @@ void UAnimNextParameterBlock_EditorData::HandleConfigureRigVMController(const FR
 		{
 			if(UAnimNextParameterBlock* ParameterBlock = WeakThis->GetTypedOuter<UAnimNextParameterBlock>())
 			{
-				if (ParameterBlock->RigVM)
+				if (ParameterBlock->VM)
 				{
-					return &ParameterBlock->RigVM->GetByteCode();
+					return &ParameterBlock->VM->GetByteCode();
 				}
 			}
 		}
@@ -944,7 +1101,7 @@ UEdGraph* UAnimNextParameterBlock_EditorData::CreateEdGraph(URigVMGraph* InRigVM
 
 	if(GraphName.IsEmpty())
 	{
-		GraphName = UControlRigGraphSchema::GraphName_ControlRig.ToString();
+		GraphName = URigVMEdGraphSchema::GraphName_RigVM.ToString();
 	}
 
 	GraphName = RigVMClient.GetUniqueName(*GraphName).ToString();
@@ -976,4 +1133,4 @@ bool UAnimNextParameterBlock_EditorData::RemoveEdGraph(URigVMGraph* InModel)
 	return false;
 }
 
-#endif // #if WITH_EDITORONLY_DATA
\ No newline at end of file
+#endif // #if WITH_EDI
\ No newline at end of file
diff --git a/Engine/Plugins/Experimental/AnimNext/Source/AnimNextUncookedOnly/Private/UncookedOnlyUtils.cpp b/Engine/Plugins/Experimental/AnimNext/Source/AnimNextUncookedOnly/Private/UncookedOnlyUtils.cpp
--- a/Engine/Plugins/Experimental/AnimNext/Source/AnimNextUncookedOnly/Private/UncookedOnlyUtils.cpp
+++ b/Engine/Plugins/Experimental/AnimNext/Source/AnimNextUncookedOnly/Private/UncookedOnlyUtils.cpp
@@ -14,10 +14,11 @@
 #include "Param/AnimNextParameterBlock_EditorData.h"
 #include "Param/AnimNextParameterBlock_EdGraph.h"
 #include "Graph/RigDecorator_AnimNextCppDecorator.h"
-#include "Graph/RigUnit_AnimNextBeginExecution.h"
-#include "Param/RigVMDispatch_SetParameter.h"
+#include "Param/RigUnit_AnimNextParameterBeginExecution.h"
+#include "Param/RigVMDispatch_GetParameter.h"
+#include "Param/RigVMDispatch_SetLayerParameter.h"
 #include "Param/AnimNextParameterBlockEntry.h"
-#include "Param/IAnimNextParameterBlockBindingInterface.h"
+#include "Param/IAnimNextParameterBlockParameterInterface.h"
 #include "DecoratorBase/DecoratorReader.h"
 #include "DecoratorBase/DecoratorWriter.h"
 #include "DecoratorBase/NodeTemplate.h"
@@ -27,6 +28,7 @@
 #include "DecoratorBase/NodeInstance.h"
 #include "DecoratorBase/DecoratorRegistry.h"
 #include "DecoratorBase/Decorator.h"
+#include "Graph/RigUnit_AnimNextBeginExecution.h"
 #include "Serialization/MemoryReader.h"
 
 #include "RigVMCompiler/RigVMCompiler.h"
@@ -639,23 +641,37 @@ void FUtils::CompileVM(UAnimNextParameterBlock* InParameterBlock)
 
 	URigVMCompiler* Compiler = URigVMCompiler::StaticClass()->GetDefaultObject<URigVMCompiler>();
 	EditorData->VMCompileSettings.SetExecuteContextStruct(EditorData->RigVMClient.GetExecuteContextStruct());
+	FRigVMExtendedExecuteContext& CDOContext = InParameterBlock->GetExtendedExecuteContext();
 	const FRigVMCompileSettings Settings = (EditorData->bCompileInDebugMode) ? FRigVMCompileSettings::Fast(EditorData->VMCompileSettings.GetExecuteContextStruct()) : EditorData->VMCompileSettings;
 	URigVMController* RootController = EditorData->GetRigVMClient()->GetOrCreateController(EditorData->GetRigVMClient()->GetDefaultModel());
-	Compiler->Compile(Settings, EditorData->GetRigVMClient()->GetAllModels(false, false), RootController, InParameterBlock->RigVM, InParameterBlock->ExtendedExecuteContext, InParameterBlock->GetRigVMExternalVariables(), &EditorData->PinToOperandMap);
+	Compiler->Compile(Settings, EditorData->GetRigVMClient()->GetAllModels(false, false), RootController, InParameterBlock->RigVM, InParameterBlock->GetExtendedExecuteContext(), InParameterBlock->GetExternalVariables(), &EditorData->PinToOperandMap);
+
+	InParameterBlock->RigVM->SetVMHash(InParameterBlock->RigVM->ComputeVMHash(CDOContext));
+	CDOContext.VMHash = InParameterBlock->RigVM->GetVMHash();
+	InParameterBlock->RigVM->Initialize(CDOContext, InParameterBlock->RigVM->GetLocalMemoryArray(CDOContext));
+	InParameterBlock->GenerateUserDefinedDependenciesData(CDOContext);
+
+	// Notable difference with vanilla RigVM host behavior - we init the VM here at the moment as we only have one 'instance'
+	InParameterBlock->bTEMP_CopyDefaultsFromCDO = false;
+	InParameterBlock->VM = InParameterBlock->RigVM;
+	InParameterBlock->InitializeVM(FRigUnit_AnimNextBeginExecution::EventName);
 
 	if (EditorData->bErrorsDuringCompilation)
 	{
 		if(Settings.SurpressErrors)
 		{
-			Settings.Reportf(EMessageSeverity::Info, InParameterBlock, TEXT("Compilation Errors may be suppressed for AnimNext Interface Graph: %s. See VM Compile Settings for more Details"), *InParameterBlock->GetName());
+			Settings.Reportf(EMessageSeverity::Info, InParameterBlock,
+				TEXT("Compilation Errors may be suppressed for ControlRigBlueprint: %s. See VM Compile Setting in Class Settings for more Details"), *InParameterBlock->GetName());
 		}
+		EditorData->bVMRecompilationRequired = false;
+		if(InParameterBlock->RigVM)
+		{
+			EditorData->RigVMCompiledEvent.Broadcast(InParameterBlock, InParameterBlock->RigVM, InParameterBlock->GetExtendedExecuteContext());
+		}
+		return;
 	}
 
-	EditorData->bVMRecompilationRequired = false;
-	if(InParameterBlock->RigVM)
-	{
-		EditorData->RigVMCompiledEvent.Broadcast(InParameterBlock, InParameterBlock->RigVM, InParameterBlock->ExtendedExecuteContext);
-	}
+//	InitializeArchetypeInstances();
 
 #if WITH_EDITOR
 //	RefreshBreakpoints(EditorData);
@@ -680,10 +696,10 @@ void FUtils::CompileStruct(UAnimNextParameterBlock* InParameterBlock)
 	TArray<FPropertyBagPropertyDesc> PropertyDescs;
 	PropertyDescs.Reserve(EditorData->Entries.Num());
 	
-	// Gather all properties in this block
+	// Gather all parameters in this block
 	for(const UAnimNextParameterBlockEntry* Entry : EditorData->Entries)
 	{
-		if(const IAnimNextParameterBlockBindingInterface* Binding = Cast<IAnimNextParameterBlockBindingInterface>(Entry))
+		if(const IAnimNextParameterBlockParameterInterface* Binding = Cast<IAnimNextParameterBlockParameterInterface>(Entry))
 		{
 			if(const UAnimNextParameter* Parameter = Binding->GetParameter())
 			{
@@ -718,11 +734,11 @@ void FUtils::RecreateVM(UAnimNextParameterBlock* InParameterBlock)
 	{
 		// We dont support ERigVMMemoryType::Work memory as we dont operate on an instance
 	//	InParameterBlock->RigVM->GetMemoryByType(ERigVMMemoryType::Work, true);
-		InParameterBlock->RigVM->CreateMemoryByType(InParameterBlock->ExtendedExecuteContext, ERigVMMemoryType::Literal);
-		InParameterBlock->RigVM->CreateMemoryByType(InParameterBlock->ExtendedExecuteContext, ERigVMMemoryType::Debug);
+		InParameterBlock->RigVM->CreateMemoryByType(InParameterBlock->GetExtendedExecuteContext(), ERigVMMemoryType::Literal);
+		InParameterBlock->RigVM->CreateMemoryByType(InParameterBlock->GetExtendedExecuteContext(), ERigVMMemoryType::Debug);
 	}
 
-	InParameterBlock->RigVM->Reset(InParameterBlock->ExtendedExecuteContext);
+	InParameterBlock->RigVM->Reset(InParameterBlock->GetExtendedExecuteContext());
 }
 
 UAnimNextParameterBlock_EditorData* FUtils::GetEditorData(const UAnimNextParameterBlock* InParameterBlock)
@@ -926,7 +942,6 @@ FEdGraphPinType FUtils::GetPinTypeFromParamType(const FAnimNextParamType& InPara
 		PinType.PinSubCategoryObject = const_cast<UObject*>(InParamType.ValueTypeObject.Get());
 		break;
 	default:
-		ensureMsgf(false, TEXT("Unhandled value type %d"), InParamType.ValueType);
 		break;
 	}
 
@@ -983,14 +998,15 @@ FRigVMTemplateArgumentType FUtils::GetRigVMArgTypeFromParamType(const FAnimNextP
 		ArgType.CPPTypeObject = const_cast<UObject*>(InParamType.ValueTypeObject.Get());
 		break;
 	case EPropertyBagPropertyType::Object:
-		CPPTypeString = RigVMTypeUtils::CPPTypeFromObject(Cast<UClass>(InParamType.ValueTypeObject.Get()));
+		CPPTypeString = RigVMTypeUtils::CPPTypeFromObject(Cast<UClass>(InParamType.ValueTypeObject.Get()), RigVMTypeUtils::EClassArgType::AsObject);
 		ArgType.CPPTypeObject = const_cast<UObject*>(InParamType.ValueTypeObject.Get());
 		break;
 	case EPropertyBagPropertyType::SoftObject:
 		ensureMsgf(false, TEXT("Unhandled value type %d"), InParamType.ValueType);
 		break;
 	case EPropertyBagPropertyType::Class:
-		ensureMsgf(false, TEXT("Unhandled value type %d"), InParamType.ValueType);
+		CPPTypeString = RigVMTypeUtils::CPPTypeFromObject(Cast<UClass>(InParamType.ValueTypeObject.Get()), RigVMTypeUtils::EClassArgType::AsClass);
+		ArgType.CPPTypeObject = const_cast<UObject*>(InParamType.ValueTypeObject.Get());
 		break;
 	case EPropertyBagPropertyType::SoftClass:
 		ensureMsgf(false, TEXT("Unhandled value type %d"), InParamType.ValueType);
@@ -1018,6 +1034,21 @@ FRigVMTemplateArgumentType FUtils::GetRigVMArgTypeFromParamType(const FAnimNextP
 	return ArgType;
 }
 
+
+void FUtils::SetupGraph(URigVMController* InController)
+{
+	// Clear the graph
+	InController->RemoveNodes(InController->GetGraph()->GetNodes());
+
+	// Add entry point
+	InController->AddUnitNode(FRigUnit_AnimNextParameterBeginExecution::StaticStruct(), FRigVMStruct::ExecuteName, FVector2D(-200.0f, 0.0f), FString(), false);
+}
+
+void FUtils::SetupBindingGraph(URigVMController* InController, FName InParameterName, const FAnimNextParamType& InParamType)
+{
+	SetupBindingGraphForLiteral(InController, InParameterName, InParamType);
+}
+
 void FUtils::SetupBindingGraphForLiteral(URigVMController* InController, FName InParameterName, const FAnimNextParamType& InParamType)
 {
 	FRigVMTemplateArgumentType ArgType = GetRigVMArgTypeFromParamType(InParamType);
@@ -1027,15 +1058,15 @@ void FUtils::SetupBindingGraphForLiteral(URigVMController* InController, FName I
 	InController->RemoveNodes(InController->GetGraph()->GetNodes());
 
 	// Add new nodes for a simple literal binding
-	URigVMUnitNode* EntryPointNode = InController->AddUnitNode(FRigUnit_AnimNextBeginExecution::StaticStruct(), FRigVMStruct::ExecuteName, FVector2D(-200.0f, 0.0f), FString(), false);
+	URigVMUnitNode* EntryPointNode = InController->AddUnitNode(FRigUnit_AnimNextParameterBeginExecution::StaticStruct(), FRigVMStruct::ExecuteName, FVector2D(-200.0f, 0.0f), FString(), false);
 
-	const FName FactoryName = FRigVMDispatch_SetParameter().GetFactoryName();
-	FRigVMDispatch_SetParameter* Factory = static_cast<FRigVMDispatch_SetParameter*>(FRigVMRegistry::Get().FindDispatchFactory(FactoryName));
+	const FName FactoryName = FRigVMDispatch_SetLayerParameter().GetFactoryName();
+	FRigVMDispatch_SetLayerParameter* Factory = static_cast<FRigVMDispatch_SetLayerParameter*>(FRigVMRegistry::Get().FindDispatchFactory(FactoryName));
 	URigVMTemplateNode* SetParameterNode = InController->AddTemplateNode(Factory->GetTemplate()->GetNotation(), FVector2D(200.0f, 0.0f));
-	InController->SetPinDefaultValue(SetParameterNode->FindPin(FRigVMDispatch_SetParameter::ParameterName.ToString())->GetPinPath(), InParameterName.ToString());
-	InController->ResolveWildCardPin(SetParameterNode->FindPin(FRigVMDispatch_SetParameter::ValueName.ToString()), TypeIndex);
+	InController->SetPinDefaultValue(SetParameterNode->FindPin(FRigVMDispatch_SetLayerParameter::ParameterName.ToString())->GetPinPath(), InParameterName.ToString());
+	InController->ResolveWildCardPin(SetParameterNode->FindPin(FRigVMDispatch_SetLayerParameter::ValueName.ToString()), TypeIndex);
 
-	InController->AddLink(EntryPointNode->FindPin(GET_MEMBER_NAME_STRING_CHECKED(FRigUnit_AnimNextBeginExecution, ExecuteContext)), SetParameterNode->FindPin(FRigVMDispatch_SetParameter::ExecuteContextName.ToString()));
+	InController->AddLink(EntryPointNode->FindPin(GET_MEMBER_NAME_STRING_CHECKED(FRigUnit_AnimNextParameterBeginExecution, ExecuteContext)), SetParameterNode->FindPin(FRigVMDispatch_SetLayerParameter::ExecuteContextName.ToString()));
 }
 
 FText FUtils::GetParameterDisplayNameText(FName InParameterName)
@@ -1045,4 +1076,6 @@ FText FUtils::GetParameterDisplayNameText(FName InParameterName)
 	return FText::FromString(NameAsString);
 }
 
+
+
 }
diff --git a/Engine/Plugins/Experimental/StructUtils/Source/StructUtils/Private/PropertyBag.cpp b/Engine/Plugins/Experimental/StructUtils/Source/StructUtils/Private/PropertyBag.cpp
--- a/Engine/Plugins/Experimental/StructUtils/Source/StructUtils/Private/PropertyBag.cpp
+++ b/Engine/Plugins/Experimental/StructUtils/Source/StructUtils/Private/PropertyBag.cpp
@@ -157,6 +157,10 @@ namespace UE::StructUtils::Private
 		{
 			return EPropertyBagPropertyType::Object;
 		}
+		if (CastField<FObjectPtrProperty>(InSourceProperty))
+		{
+			return EPropertyBagPropertyType::Object;
+		}
 		if (CastField<FSoftObjectProperty>(InSourceProperty))
 		{
 			return EPropertyBagPropertyType::SoftObject;
@@ -200,6 +204,10 @@ namespace UE::StructUtils::Private
 		{
 			return ObjectProperty->PropertyClass;
 		}
+		if (const FObjectPtrProperty* ObjectProperty = CastField<FObjectPtrProperty>(InSourceProperty))
+		{
+			return ObjectProperty->PropertyClass;
+		}
 		if (const FSoftObjectProperty* SoftObjectProperty = CastField<FSoftObjectProperty>(InSourceProperty))
 		{
 			return SoftObjectProperty->PropertyClass;
@@ -365,7 +373,7 @@ namespace UE::StructUtils::Private
 		case EPropertyBagPropertyType::Object:
 			if (const UClass* Class = Cast<UClass>(Desc.ValueTypeObject))
 			{
-				FObjectProperty* Prop = new FObjectProperty(PropertyScope, Desc.Name, RF_Public);
+				FObjectPtrProperty* Prop = new FObjectPtrProperty(PropertyScope, Desc.Name, RF_Public);
 				if (Class->HasAnyClassFlags(CLASS_DefaultToInstanced))
 				{
 					Prop->SetPropertyFlags(CPF_InstancedReference);
diff --git a/Engine/Source/Editor/BlueprintGraph/Private/EdGraphSchema_K2.cpp b/Engine/Source/Editor/BlueprintGraph/Private/EdGraphSchema_K2.cpp
--- a/Engine/Source/Editor/BlueprintGraph/Private/EdGraphSchema_K2.cpp
+++ b/Engine/Source/Editor/BlueprintGraph/Private/EdGraphSchema_K2.cpp
@@ -3004,11 +3004,6 @@ FString UEdGraphSchema_K2::IsPinDefaultValid(const UEdGraphPin* Pin, const FStri
 	MessageArgs.Add(TEXT("PinName"), Pin->GetDisplayName());
 
 	const UBlueprint* OwningBP = FBlueprintEditorUtils::FindBlueprintForNode(Pin->GetOwningNodeUnchecked());
-	if (!OwningBP)
-	{
-		FText MsgFormat = LOCTEXT("NoBlueprintFoundForPin", "No Blueprint was found for the pin '{PinName}'.");
-		return FText::Format(MsgFormat, MessageArgs).ToString();
-	}
 
 	const bool bIsArray = Pin->PinType.IsArray();
 	const bool bIsSet = Pin->PinType.IsSet();
@@ -3016,7 +3011,7 @@ FString UEdGraphSchema_K2::IsPinDefaultValid(const UEdGraphPin* Pin, const FStri
 	const bool bIsReference = Pin->PinType.bIsReference;
 	const bool bIsAutoCreateRefTerm = IsAutoCreateRefTerm(Pin);
 
-	if (OwningBP->BlueprintType != BPTYPE_Interface)
+	if (OwningBP == nullptr || OwningBP->BlueprintType != BPTYPE_Interface)
 	{
 		if( !bIsAutoCreateRefTerm )
 		{
