Use correct bone indices when using the new LOD pose API
Adding comments and clarifications
Minor tweaks

diff --git a/Engine/Plugins/Experimental/AnimNext/Source/AnimNext/Private/DecompressionTools.cpp b/Engine/Plugins/Experimental/AnimNext/Source/AnimNext/Private/DecompressionTools.cpp
--- a/Engine/Plugins/Experimental/AnimNext/Source/AnimNext/Private/DecompressionTools.cpp
+++ b/Engine/Plugins/Experimental/AnimNext/Source/AnimNext/Private/DecompressionTools.cpp
@@ -276,13 +276,12 @@ void FDecompressionTools::DecompressPose(FLODPose& OutAnimationPoseData,
 										const FRootMotionReset& RootMotionReset)
 {
 	const FReferencePose& ReferencePose = OutAnimationPoseData.GetRefPose();
-	const TArrayView<const FBoneIndexType>LODBoneIndexes = OutAnimationPoseData.GetLODBoneIndexes();
-	const TArrayView<const FBoneIndexType> SkeletoonToLODBoneIndexes = ReferencePose.GetSkeletonToLODBoneIndexes(0); // Full list of Skeleton to LOD conversion
+	const TArrayView<const FBoneIndexType> LODBoneIndexes = OutAnimationPoseData.GetLODBoneIndexes();
+	const TArrayView<const FBoneIndexType> SkeletonToLODBoneIndexes = ReferencePose.GetSkeletonToLODBoneIndexes(0); // Full list of Skeleton to LOD conversion
 	const int32 NumLODBoneIndexes = LODBoneIndexes.Num();
 
 	const int32 NumTracks = CompressedData.CompressedTrackToSkeletonMapTable.Num();
 
-
 	const USkeleton* TargetSkeleton = OutAnimationPoseData.GetSkeletonAsset();
 	const FSkeletonRemapping& SkeletonRemapping = UE::Anim::FSkeletonRemappingRegistry::Get().GetRemapping(DecompressionContext.GetSourceSkeleton(), TargetSkeleton);
 
@@ -315,7 +314,7 @@ void FDecompressionTools::DecompressPose(FLODPose& OutAnimationPoseData,
 
 			if (TargetSkeletonBoneIndex != INDEX_NONE)
 			{
-				const int32 LODBoneIndex = TargetSkeletonBoneIndex < SkeletoonToLODBoneIndexes.Num() ? SkeletoonToLODBoneIndexes[TargetSkeletonBoneIndex] : INDEX_NONE;
+				const int32 LODBoneIndex = TargetSkeletonBoneIndex < SkeletonToLODBoneIndexes.Num() ? SkeletonToLODBoneIndexes[TargetSkeletonBoneIndex] : INDEX_NONE;
 
 				if (LODBoneIndex != INDEX_NONE && LODBoneIndex < NumLODBoneIndexes) // skip bones not in current LOD
 				{
diff --git a/Engine/Plugins/Experimental/AnimNext/Source/AnimNext/Private/DecoratorBase/DecoratorPtr.cpp b/Engine/Plugins/Experimental/AnimNext/Source/AnimNext/Private/DecoratorBase/DecoratorPtr.cpp
--- a/Engine/Plugins/Experimental/AnimNext/Source/AnimNext/Private/DecoratorBase/DecoratorPtr.cpp
+++ b/Engine/Plugins/Experimental/AnimNext/Source/AnimNext/Private/DecoratorBase/DecoratorPtr.cpp
@@ -68,18 +68,18 @@ namespace UE::AnimNext
 		// The new handle will remain weak
 		if (!DecoratorPtr.IsWeak())
 		{
-			if (FNodeInstance* Node = DecoratorPtr.GetNodeInstance())
+			if (FNodeInstance* NewNode = DecoratorPtr.GetNodeInstance())
 			{
-				Node->AddReference();
+				NewNode->AddReference();
 			}
 		}
 
 		// Only decrement the reference count if we aren't a weak handle
 		if (!IsWeak())
 		{
-			if (FNodeInstance* Node = GetNodeInstance())
+			if (FNodeInstance* OldNode = GetNodeInstance())
 			{
-				Node->ReleaseReference();
+				OldNode->ReleaseReference();
 			}
 		}
 
diff --git a/Engine/Plugins/Experimental/AnimNext/Source/AnimNext/Private/GenerationTools.cpp b/Engine/Plugins/Experimental/AnimNext/Source/AnimNext/Private/GenerationTools.cpp
--- a/Engine/Plugins/Experimental/AnimNext/Source/AnimNext/Private/GenerationTools.cpp
+++ b/Engine/Plugins/Experimental/AnimNext/Source/AnimNext/Private/GenerationTools.cpp
@@ -105,9 +105,11 @@ struct FCompareBoneIndexType
 
 			if (bCanGenerateSingleBonesList)
 			{
+				const USkeleton* Skeleton = SkeletalMesh->GetSkeleton();
+
 				OutAnimationReferencePose.GenerationFlags = EReferencePoseGenerationFlags::FastPath;
 				OutAnimationReferencePose.SkeletalMesh = SkeletalMesh;
-				OutAnimationReferencePose.Skeleton = SkeletalMesh->GetSkeleton();
+				OutAnimationReferencePose.Skeleton = Skeleton;
 
 				TArray<int32> LODNumBones;
 				LODNumBones.Reset(NumLODs);
@@ -116,14 +118,27 @@ struct FCompareBoneIndexType
 					LODNumBones.Add(GenerationLODData[LODIndex].RequiredBones.Num());
 				}
 
+				// Removing const here because the linkup lazily builds the mapping and caches it
+				const int32 SkelMeshLinkupIndex = const_cast<USkeleton*>(Skeleton)->GetMeshLinkupIndex(SkeletalMesh);
+				check(SkelMeshLinkupIndex != INDEX_NONE);
+
+				const FSkeletonToMeshLinkup& LinkupTable = Skeleton->LinkupCache[SkelMeshLinkupIndex];
+
 				// Generate a Skeleton to LOD look up table
 				const int32 NumOrderedBones = OrderedBoneList.Num();
 
 				TArray<FBoneIndexType> SkeletonToLODBoneList;
 				SkeletonToLODBoneList.SetNumZeroed(NumOrderedBones);
-				for (int i = 0; i < NumOrderedBones; ++i)
+				for (int32 LODBoneIndex = 0; LODBoneIndex < NumOrderedBones; ++LODBoneIndex)
 				{
-					SkeletonToLODBoneList[OrderedBoneList[i]] = i;
+					// The ordered list contains skeletal mesh bone indices sorted by LOD
+					const FMeshPoseBoneIndex MeshBoneIndex(OrderedBoneList[LODBoneIndex]);
+
+					// Remap our skeletal mesh bone index into the skeleton bone index we output for
+					const FSkeletonPoseBoneIndex SkeletonBoneIndex(LinkupTable.MeshToSkeletonTable[MeshBoneIndex.GetInt()]);
+					ensure(SkeletonBoneIndex.IsValid());	// We expect the skeletal mesh bone to map to a valid skeleton bone
+
+					SkeletonToLODBoneList[static_cast<FBoneIndexType>(SkeletonBoneIndex.GetInt())] = LODBoneIndex;
 				}
 
 				OutAnimationReferencePose.Initialize(SkeletalMesh->GetRefSkeleton(), { OrderedBoneList }, { SkeletonToLODBoneList }, LODNumBones, bCanGenerateSingleBonesList);
@@ -387,22 +402,34 @@ struct FCompareBoneIndexType
 // Converts AnimBP pose to AnimNext Pose
 // This function expects both poses to have the same LOD (number of bones and indexes)
 // The target pose should be assigned to the correct reference pose prior to this call
-/*static*/ void FGenerationTools::RemapPose(int32 LODLevel, const FPoseContext& SourcePose, const FReferencePose& RefPose, FLODPose& TargetPose)
+/*static*/ void FGenerationTools::RemapPose(const FPoseContext& SourcePose, FLODPose& TargetPose)
 {
-	const TArrayView<const FBoneIndexType> LODBoneIndexes = RefPose.GetLODBoneIndexes(LODLevel);
-	const int32 NumBones = LODBoneIndexes.Num();
+	const FBoneContainer& BoneContainer = SourcePose.Pose.GetBoneContainer();
+	const FReferencePose& RefPose = TargetPose.GetRefPose();
+	const TArrayView<const FBoneIndexType> LODBoneIndexes = RefPose.GetLODBoneIndexes(TargetPose.LODLevel);
+	const int32 NumLODBones = LODBoneIndexes.Num();
+
+	check(TargetPose.GetNumBones() == NumLODBones);
 
-	if (TargetPose.GetNumBones() == NumBones)
+	for (int32 LODBoneIndex = 0; LODBoneIndex < NumLODBones; ++LODBoneIndex)
 	{
-		for (int i = 0; i < NumBones; ++i)
-		{
-			const auto& SkeletonBoneIndex = LODBoneIndexes[i];
+		// Reference pose holds a list of skeletal mesh bone indices sorted by LOD
+		const FMeshPoseBoneIndex MeshBoneIndex(LODBoneIndexes[LODBoneIndex]);
 
-			const FCompactPoseBoneIndex CompactBoneIndex = SourcePose.Pose.GetBoneContainer().GetCompactPoseIndexFromSkeletonIndex(SkeletonBoneIndex);
-			if (CompactBoneIndex.GetInt() != INDEX_NONE)
-			{
-				TargetPose.LocalTransformsView[i] = SourcePose.Pose[CompactBoneIndex];
-			}
+		// Remap our skeletal mesh bone index into the skeleton bone index we output for
+		const FSkeletonPoseBoneIndex SkeletonBoneIndex = BoneContainer.GetSkeletonPoseIndexFromMeshPoseIndex(MeshBoneIndex);
+		ensure(SkeletonBoneIndex.IsValid());	// We expect the skeletal mesh bone to map to a valid skeleton bone
+
+		// Remap our skeleton bone index into the compact pose bone index we output for
+		const FCompactPoseBoneIndex CompactBoneIndex = BoneContainer.GetCompactPoseIndexFromSkeletonPoseIndex(SkeletonBoneIndex);
+
+		if (ensure(CompactBoneIndex.IsValid()))	// We expect the skeleton bone to map to a valid compact pose bone
+		{
+			TargetPose.LocalTransformsView[LODBoneIndex] = SourcePose.Pose[CompactBoneIndex];
+		}
+		else
+		{
+			// This bone is part of the LOD but isn't part of the required bones
 		}
 	}
 }
@@ -410,22 +437,34 @@ struct FCompareBoneIndexType
 // Converts AnimNext pose to AnimBP Pose
 // This function expects both poses to have the same LOD (number of bones and indexes)
 // The target pose should be assigned to the correct reference pose prior to this call
-/*static*/ void FGenerationTools::RemapPose(int32 LODLevel, const FReferencePose& RefPose, const FLODPose& SourcePose, FPoseContext& TargetPose)
+/*static*/ void FGenerationTools::RemapPose(const FLODPose& SourcePose, FPoseContext& TargetPose)
 {
-	const TArrayView<const FBoneIndexType> LODBoneIndexes = RefPose.GetLODBoneIndexes(LODLevel);
-	const int32 NumBones = LODBoneIndexes.Num();
+	const FBoneContainer& BoneContainer = TargetPose.Pose.GetBoneContainer();
+	const FReferencePose& RefPose = SourcePose.GetRefPose();
+	const TArrayView<const FBoneIndexType> LODBoneIndexes = RefPose.GetLODBoneIndexes(SourcePose.LODLevel);
+	const int32 NumLODBones = LODBoneIndexes.Num();
+
+	check(SourcePose.GetNumBones() == NumLODBones);
 
-	if (SourcePose.GetNumBones() == NumBones)
+	for (int32 LODBoneIndex = 0; LODBoneIndex < NumLODBones; ++LODBoneIndex)
 	{
-		for (int i = 0; i < NumBones; ++i)
-		{
-			const auto& SkeletonBoneIndex = LODBoneIndexes[i];
+		// Reference pose holds a list of skeletal mesh bone indices sorted by LOD
+		const FMeshPoseBoneIndex MeshBoneIndex(LODBoneIndexes[LODBoneIndex]);
 
-			const FCompactPoseBoneIndex CompactBoneIndex = TargetPose.Pose.GetBoneContainer().GetCompactPoseIndexFromSkeletonIndex(SkeletonBoneIndex);
-			if (CompactBoneIndex.GetInt() != INDEX_NONE)
-			{
-				TargetPose.Pose[CompactBoneIndex] = SourcePose.LocalTransformsView[i];
-			}
+		// Remap our skeletal mesh bone index into the skeleton bone index we output for
+		const FSkeletonPoseBoneIndex SkeletonBoneIndex = BoneContainer.GetSkeletonPoseIndexFromMeshPoseIndex(MeshBoneIndex);
+		ensure(SkeletonBoneIndex.IsValid());	// We expect the skeletal mesh bone to map to a valid skeleton bone
+
+		// Remap our skeleton bone index into the compact pose bone index we output for
+		const FCompactPoseBoneIndex CompactBoneIndex = BoneContainer.GetCompactPoseIndexFromSkeletonPoseIndex(SkeletonBoneIndex);
+
+		if (ensure(CompactBoneIndex.IsValid()))	// We expect the skeleton bone to map to a valid compact pose bone
+		{
+			TargetPose.Pose[CompactBoneIndex] = SourcePose.LocalTransformsView[LODBoneIndex];
+		}
+		else
+		{
+			// This bone is part of the LOD but isn't part of the required bones
 		}
 	}
 }
diff --git a/Engine/Plugins/Experimental/AnimNext/Source/AnimNext/Private/Graph/AnimGraph/AnimNode_AnimNextGraph.cpp b/Engine/Plugins/Experimental/AnimNext/Source/AnimNext/Private/Graph/AnimGraph/AnimNode_AnimNextGraph.cpp
--- a/Engine/Plugins/Experimental/AnimNext/Source/AnimNext/Private/Graph/AnimGraph/AnimNode_AnimNextGraph.cpp
+++ b/Engine/Plugins/Experimental/AnimNext/Source/AnimNext/Private/Graph/AnimGraph/AnimNode_AnimNextGraph.cpp
@@ -23,6 +23,8 @@
 #include "Editor.h"
 #endif
 
+TAutoConsoleVariable<int32> CVarAnimNextForceAnimBP(TEXT("a.AnimNextForceAnimBP"), 0, TEXT("If != 0, then we use the input pose of the AnimNext AnimBP node instead of the AnimNext graph."));
+
 FAnimNode_AnimNextGraph::FAnimNode_AnimNextGraph()
 	: FAnimNode_CustomProperty()
 	, AnimNextGraph(nullptr)
@@ -54,7 +56,7 @@ void FAnimNode_AnimNextGraph::Update_AnyThread(const FAnimationUpdateContext& Co
 
 	SourceLink.Update(Context);
 
-	if (IsLODEnabled(Context.AnimInstanceProxy))
+	if (IsLODEnabled(Context.AnimInstanceProxy) && !CVarAnimNextForceAnimBP.GetValueOnAnyThread() && GraphInstance.IsValid())
 	{
 		GetEvaluateGraphExposedInputs().Execute(Context);
 
@@ -93,22 +95,19 @@ void FAnimNode_AnimNextGraph::Initialize_AnyThread(const FAnimationInitializeCon
 	if (!GraphInstance.IsValid() && AnimNextGraph)
 	{
 		// If we don't have an instance yet, create one
+		UE::AnimNext::FAnimGraphParamStackScope Scope(Context);
 
-		{
-			UE::AnimNext::FAnimGraphParamStackScope Scope(Context);
-
-			// Populate our param stack since our instance data might need it during construction
-			const int32 LODLevel = Context.AnimInstanceProxy->GetLODLevel();
+		// Populate our param stack since our instance data might need it during construction
+		const int32 LODLevel = Context.AnimInstanceProxy->GetLODLevel();
 
-			FParamStack& ParamStack = FParamStack::Get();
-			FParamStack::FPushedLayerHandle LayerHandle = ParamStack.PushValues(
-				"GraphLODLevel", LODLevel
-			);
+		FParamStack& ParamStack = FParamStack::Get();
+		FParamStack::FPushedLayerHandle LayerHandle = ParamStack.PushValues(
+			"GraphLODLevel", LODLevel
+		);
 
-			AnimNextGraph->AllocateInstance(GraphInstance);
+		AnimNextGraph->AllocateInstance(GraphInstance);
 
-			ParamStack.PopLayer(LayerHandle);
-		}
+		ParamStack.PopLayer(LayerHandle);
 	}
 
 	FAnimNode_CustomProperty::Initialize_AnyThread(Context);
@@ -123,59 +122,51 @@ void FAnimNode_AnimNextGraph::CacheBones_AnyThread(const FAnimationCacheBonesCon
 	SourceLink.CacheBones(Context);
 }
 
-void FAnimNode_AnimNextGraph::Evaluate_AnyThread(FPoseContext & Output)
+void FAnimNode_AnimNextGraph::Evaluate_AnyThread(FPoseContext& Output)
 {
 	DECLARE_SCOPE_HIERARCHICAL_COUNTER_FUNC()
 
 	using namespace UE::AnimNext;
 
-	FPoseContext SourcePose(Output);
-
-	if (SourceLink.GetLinkNode())
+	if (!CVarAnimNextForceAnimBP.GetValueOnAnyThread() && GraphInstance.IsValid())
 	{
-		SourceLink.Evaluate(SourcePose);
-	}
-	else
-	{
-		SourcePose.ResetToRefPose();
-	}
-
-	USkeletalMeshComponent* SkeletalMeshComponent = Output.AnimInstanceProxy->GetSkelMeshComponent();
-	check(SkeletalMeshComponent != nullptr);
+		USkeletalMeshComponent* SkeletalMeshComponent = Output.AnimInstanceProxy->GetSkelMeshComponent();
+		check(SkeletalMeshComponent != nullptr);
 
-	FDataHandle RefPoseHandle = FDataRegistry::Get()->GetOrGenerateReferencePose(SkeletalMeshComponent);
-	const UE::AnimNext::FReferencePose& RefPose = RefPoseHandle.GetRef<UE::AnimNext::FReferencePose>();
+		FDataHandle RefPoseHandle = FDataRegistry::Get()->GetOrGenerateReferencePose(SkeletalMeshComponent);
+		const UE::AnimNext::FReferencePose& RefPose = RefPoseHandle.GetRef<UE::AnimNext::FReferencePose>();
+		FAnimNextGraphReferencePose GraphReferencePose(&RefPose);
 
-	const int32 LODLevel = Output.AnimInstanceProxy->GetLODLevel();
-	
-	UE::AnimNext::FContext Context(0.0f);
+		const int32 LODLevel = Output.AnimInstanceProxy->GetLODLevel();
 
-	FAnimNextGraphReferencePose GraphReferencePose(&RefPose);
+		UE::AnimNext::FContext Context(0.0f);
 
-	FAnimNextGraphLODPose GraphSourceLODPose(FLODPoseHeap(RefPose, LODLevel, false, Output.ExpectsAdditivePose()));
-	FAnimNextGraphLODPose ResultPose(FLODPoseHeap(RefPose, LODLevel, true, Output.ExpectsAdditivePose()));
-	FGenerationTools::RemapPose(LODLevel, SourcePose, RefPose, GraphSourceLODPose.LODPose);
+		FAnimNextGraphLODPose ResultPose(FLODPoseHeap(RefPose, LODLevel, true, Output.ExpectsAdditivePose()));
 
-	{
 		UE::AnimNext::FAnimGraphParamStackScope Scope(Output);
-
 		FParamStack& ParamStack = FParamStack::Get();
 		FParamStack::FPushedLayerHandle LayerHandle = ParamStack.PushValues(
 			"GraphReferencePose", GraphReferencePose,
 			"ResultPose", ResultPose,
 			"GraphLODLevel", LODLevel,
-			"GraphExpectsAdditive", Output.ExpectsAdditivePose(),
-			"SourcePose", GraphSourceLODPose						// TODO : Pass this as a external variable maybe ? When we have support for variables in the rigvm graph
+			"GraphExpectsAdditive", Output.ExpectsAdditivePose()
 		);
 
 		AnimNextGraph->Run(Context, GraphInstance, EAnimNextGraphSimulationSteps::Evaluate);
 
-		FGenerationTools::RemapPose(LODLevel, RefPose, ResultPose.LODPose, Output);
-
-		FAnimNode_CustomProperty::Evaluate_AnyThread(Output);
+		FGenerationTools::RemapPose(ResultPose.LODPose, Output);
 
 		ParamStack.PopLayer(LayerHandle);
 	}
+	else
+	{
+		if (SourceLink.GetLinkNode())
+		{
+			SourceLink.Evaluate(Output);
+		}
+	}
+
+	FAnimNode_CustomProperty::Evaluate_AnyThread(Output);
 }
 
 void FAnimNode_AnimNextGraph::PostSerialize(const FArchive& Ar)
diff --git a/Engine/Plugins/Experimental/AnimNext/Source/AnimNext/Public/GenerationTools.h b/Engine/Plugins/Experimental/AnimNext/Source/AnimNext/Public/GenerationTools.h
--- a/Engine/Plugins/Experimental/AnimNext/Source/AnimNext/Public/GenerationTools.h
+++ b/Engine/Plugins/Experimental/AnimNext/Source/AnimNext/Public/GenerationTools.h
@@ -85,18 +85,12 @@ public:
 	// Converts AnimBP pose to AnimNext Pose
 	// This function expects both poses to have the same LOD (number of bones and indexes)
 	// The target pose should be assigned to the correct reference pose prior to this call
-	static void RemapPose(int32 LODLevel
-		, const FPoseContext& SourcePose
-		, const FReferencePose& RefPose
-		, FLODPose& TargetPose);
+	static void RemapPose(const FPoseContext& SourcePose, FLODPose& TargetPose);
 
 	// Converts AnimNext pose to AnimBP Pose
 	// This function expects both poses to have the same LOD (number of bones and indexes)
 	// The target pose should be assigned to the correct reference pose prior to this call
-	static void RemapPose(int32 LODLevel
-		, const FReferencePose& RefPose
-		, const FLODPose& SourcePose
-		, FPoseContext& TargetPose);
+	static void RemapPose(const FLODPose& SourcePose, FPoseContext& TargetPose);
 };
 
 } // namespace UE::AnimNext
diff --git a/Engine/Plugins/Experimental/AnimNext/Source/AnimNext/Public/ReferencePose.h b/Engine/Plugins/Experimental/AnimNext/Source/AnimNext/Public/ReferencePose.h
--- a/Engine/Plugins/Experimental/AnimNext/Source/AnimNext/Public/ReferencePose.h
+++ b/Engine/Plugins/Experimental/AnimNext/Source/AnimNext/Public/ReferencePose.h
@@ -35,11 +35,11 @@ struct TReferencePose
 	TTransformArray<AllocatorType> ReferenceLocalTransforms;
 
 	// List of LOD bone indices for each LOD
-	// Each list of bone indices is a mapping of: BoneIndexAtLOD -> BoneIndexInSkeleton
+	// Each list of bone indices is a mapping of: LODSortedBoneIndex -> SkeletalMeshBoneIndex
 	TArray<TArray<FBoneIndexType, AllocatorType>, AllocatorType> LODBoneIndexes;
 
 	// List of skeleton bone indices for each LOD
-	// Each list of skeleton bone indices is a mapping of: BoneIndexInSkeleton -> BoneIndexAtLOD
+	// Each list of skeleton bone indices is a mapping of: SkeletonBoneIndex -> LODSortedBoneIndex
 	TArray<TArray<FBoneIndexType, AllocatorType>, AllocatorType> SkeletonToLODBoneIndexes;
 
 	// Number of bones for each LOD
diff --git a/Engine/Source/Runtime/AnimationCore/Public/BoneIndices.h b/Engine/Source/Runtime/AnimationCore/Public/BoneIndices.h
--- a/Engine/Source/Runtime/AnimationCore/Public/BoneIndices.h
+++ b/Engine/Source/Runtime/AnimationCore/Public/BoneIndices.h
@@ -157,18 +157,21 @@ struct FBoneIndexWithOperators : public FBoneIndexBase
 	}
 };
 
+// This represents a compact pose bone index. A compact pose is held by a bone container and can have a different ordering than either the skeleton or skeletal mesh.
 struct FCompactPoseBoneIndex : public FBoneIndexWithOperators < FCompactPoseBoneIndex >
 {
 public:
 	explicit FCompactPoseBoneIndex(int32 InBoneIndex) { BoneIndex = InBoneIndex; }
 };
 
+// This represents a skeletal mesh bone index which may differ from the skeleton bone index it corresponds to.
 struct FMeshPoseBoneIndex : public FBoneIndexWithOperators < FMeshPoseBoneIndex >
 {
 public:
 	explicit FMeshPoseBoneIndex(int32 InBoneIndex) { BoneIndex = InBoneIndex; }
 };
 
+// This represents a skeleton bone index which may differ from the skeletal mesh bone index it corresponds to.
 struct FSkeletonPoseBoneIndex : public FBoneIndexWithOperators < FSkeletonPoseBoneIndex >
 {
 public:
diff --git a/Engine/Source/Runtime/Engine/Private/Animation/AnimSequence.cpp b/Engine/Source/Runtime/Engine/Private/Animation/AnimSequence.cpp
--- a/Engine/Source/Runtime/Engine/Private/Animation/AnimSequence.cpp
+++ b/Engine/Source/Runtime/Engine/Private/Animation/AnimSequence.cpp
@@ -1629,7 +1629,7 @@ void UAnimSequence::GetBonePose(FAnimationPoseData& OutAnimationPoseData, const
 		{
 			TArray<FTransform> const& AuthoredOnRefSkeleton = GetRetargetTransforms();
 			// Map from this sequence its Skeleton to target
-			const FSkeletonRemapping& SkeletonRemapping = UE::Anim::FSkeletonRemappingRegistry::Get().GetRemapping(GetSkeleton(), RequiredBones.GetSkeletonAsset());
+			const FSkeletonRemapping& SkeletonRemapping = UE::Anim::FSkeletonRemappingRegistry::Get().GetRemapping(MySkeleton, RequiredBones.GetSkeletonAsset());
 			for (FCompactPoseBoneIndex PoseBoneIndex : OutPose.ForEachBoneIndex())
 			{
 				int32 SkeletonBoneIndex = RequiredBones.GetSkeletonPoseIndexFromCompactPoseIndex(PoseBoneIndex).GetInt();
