Fix loctext issues

diff --git a/Engine/Plugins/Experimental/AnimNext/Source/AnimNext/Private/Module.cpp b/Engine/Plugins/Experimental/AnimNext/Source/AnimNext/Private/Module.cpp
--- a/Engine/Plugins/Experimental/AnimNext/Source/AnimNext/Private/Module.cpp
+++ b/Engine/Plugins/Experimental/AnimNext/Source/AnimNext/Private/Module.cpp
@@ -54,5 +54,3 @@ public:
 }
 
 IMPLEMENT_MODULE(UE::AnimNext::FModule, AnimNext)
-
-#undef LOCTEXT_NAMESPACE
\ No newline at end of file
diff --git a/Engine/Plugins/Experimental/AnimNext/Source/AnimNextEditor/Private/Param/SAddParametersDialog.cpp b/Engine/Plugins/Experimental/AnimNext/Source/AnimNextEditor/Private/Param/SAddParametersDialog.cpp
--- a/Engine/Plugins/Experimental/AnimNext/Source/AnimNextEditor/Private/Param/SAddParametersDialog.cpp
+++ b/Engine/Plugins/Experimental/AnimNext/Source/AnimNextEditor/Private/Param/SAddParametersDialog.cpp
@@ -53,7 +53,7 @@ bool FParameterToAdd::IsValid(FText& OutReason) const
 
 	if(!Library.IsValid())
 	{
-		OutReason = LOCTEXT("InvalidParameterType", "Invalid Parameter Library");
+		OutReason = LOCTEXT("InvalidParameterLibrary", "Invalid Parameter Library");
 	}
 	
 	return true; 
