AnimNext: Small fixes

Updated metadata to use correct tag on schedule entries
Exposed tags & scopes to Blueprint to allow their type to be chosen in pickers

diff --git a/Engine/Plugins/Experimental/AnimNext/Source/AnimNext/Internal/Scheduler/AnimNextSchedule.h b/Engine/Plugins/Experimental/AnimNext/Source/AnimNext/Internal/Scheduler/AnimNextSchedule.h
--- a/Engine/Plugins/Experimental/AnimNext/Source/AnimNext/Internal/Scheduler/AnimNextSchedule.h
+++ b/Engine/Plugins/Experimental/AnimNext/Source/AnimNext/Internal/Scheduler/AnimNextSchedule.h
@@ -44,10 +44,10 @@ class UAnimNextScheduleEntry_AnimNextGraph : public UAnimNextScheduleEntry
 private:
 	friend class UAnimNextSchedule;
 
-	UPROPERTY(EditAnywhere, Category = "Graph", meta = (CustomWidget = "ParamName", ParamType = "FAnimNextScope"))
+	UPROPERTY(EditAnywhere, Category = "Graph", meta = (CustomWidget = "ParamName", AllowedParamType = "FAnimNextScope"))
 	FName Name;
 
-	UPROPERTY(EditAnywhere, Category = "Graph")
+	UPROPERTY(EditAnywhere, Category = "Graph", meta = (CustomWidget = "ParamName", AllowedParamType = "FName"))
 	FName EntryPoint;
 
 	UPROPERTY(EditAnywhere, Category = "Graph")
@@ -65,7 +65,7 @@ class UAnimNextScheduleEntry_Port : public UAnimNextScheduleEntry
 private:
 	friend class UAnimNextSchedule;
 
-	UPROPERTY(EditAnywhere, Category = "Port", meta = (CustomWidget = "ParamName", ParamType = "FAnimNextPort"))
+	UPROPERTY(EditAnywhere, Category = "Port", meta = (CustomWidget = "ParamName", AllowedParamType = "FAnimNextPort"))
 	FName Name;
 };
 
@@ -77,10 +77,10 @@ class UAnimNextScheduleEntry_ExternalTask : public UAnimNextScheduleEntry
 private:
 	friend class UAnimNextSchedule;
 
-	UPROPERTY(EditAnywhere, Category = "External Task", meta = (CustomWidget = "ParamName", ParamType = "FAnimNextTickFunction"))
+	UPROPERTY(EditAnywhere, Category = "External Task", meta = (CustomWidget = "ParamName", AllowedParamType = "FTickFunction"))
 	FName Name;
 
-	UPROPERTY(EditAnywhere, Category = "External Task", meta = (CustomWidget = "ParamName", ParamType = "TObjectPtr<UObject>"))
+	UPROPERTY(EditAnywhere, Category = "External Task", meta = (CustomWidget = "ParamName", AllowedParamType = "FName"))
 	FName ObjectName;
 };
 
@@ -92,7 +92,7 @@ class UAnimNextScheduleEntry_ParamScope : public UAnimNextScheduleEntry
 private:
 	friend class UAnimNextSchedule;
 
-	UPROPERTY(EditAnywhere, Category = "Parameters", meta = (CustomWidget = "ParamName", ParamType = "FAnimNextScope"))
+	UPROPERTY(EditAnywhere, Category = "Parameters", meta = (CustomWidget = "ParamName", AllowedParamType = "FAnimNextScope"))
 	FName Name;
 
 	// Parameters to apply in this scope
diff --git a/Engine/Plugins/Experimental/AnimNext/Source/AnimNext/Public/Param/AnimNextTag.h b/Engine/Plugins/Experimental/AnimNext/Source/AnimNext/Public/Param/AnimNextTag.h
--- a/Engine/Plugins/Experimental/AnimNext/Source/AnimNext/Public/Param/AnimNextTag.h
+++ b/Engine/Plugins/Experimental/AnimNext/Source/AnimNext/Public/Param/AnimNextTag.h
@@ -5,14 +5,14 @@
 #include "AnimNextTag.generated.h"
 
 // Empty struct to act as a named tag for parameters
-USTRUCT()
+USTRUCT(BlueprintType)
 struct FAnimNextTag
 {
 	GENERATED_BODY()
 };
 
 // Empty struct to tag named scopes
-USTRUCT()
+USTRUCT(BlueprintType)
 struct FAnimNextScope : public FAnimNextTag
 {
 	GENERATED_BODY()
