Added missing copyright boilerplate to cpp

#rnx

diff --git a/Engine/Plugins/Experimental/AnimNext/Source/AnimNext/Private/RigVMRuntimeDataRegistry.cpp b/Engine/Plugins/Experimental/AnimNext/Source/AnimNext/Private/RigVMRuntimeDataRegistry.cpp
--- a/Engine/Plugins/Experimental/AnimNext/Source/AnimNext/Private/RigVMRuntimeDataRegistry.cpp
+++ b/Engine/Plugins/Experimental/AnimNext/Source/AnimNext/Private/RigVMRuntimeDataRegistry.cpp
@@ -1,3 +1,5 @@
+// Copyright Epic Games, Inc. All Rights Reserved.
+
 #include "RigVMRuntimeDataRegistry.h"
 #include "Misc/ScopeRWLock.h"
 
