AnimNext: Workspace editor

Added a new editor that incorporates editing functionality for multiple assets at once. Assets are represented as multiple sub-document tabs across multiple tab wells. A CVar controls whether to open in the new workspace or in seperate sub-editors.
Reworked the AnimNextGraph a little to be more in-line with parameter blocks - we may be able to go further and merge them to a common base too.
Contains minor changes for RigVM to allow IRigVMClientHost to get corresponding RigVM graphs for editor graphs.

diff --git a/Engine/Plugins/Experimental/AnimNext/Source/AnimNextEditor/AnimNextEditor.Build.cs b/Engine/Plugins/Experimental/AnimNext/Source/AnimNextEditor/AnimNextEditor.Build.cs
--- a/Engine/Plugins/Experimental/AnimNext/Source/AnimNextEditor/AnimNextEditor.Build.cs
+++ b/Engine/Plugins/Experimental/AnimNext/Source/AnimNextEditor/AnimNextEditor.Build.cs
@@ -39,6 +39,7 @@ namespace UnrealBuildTool.Rules
 					"ContentBrowserFileDataSource",
 					"SubobjectEditor",
 					"Settings",
+					"EditorWidgets",
 				}
 			);
 		}
diff --git a/Engine/Plugins/Experimental/AnimNext/Source/AnimNextEditor/Internal/Param/ParameterPickerArgs.h b/Engine/Plugins/Experimental/AnimNext/Source/AnimNextEditor/Internal/Param/ParameterPickerArgs.h
--- a/Engine/Plugins/Experimental/AnimNext/Source/AnimNextEditor/Internal/Param/ParameterPickerArgs.h
+++ b/Engine/Plugins/Experimental/AnimNext/Source/AnimNextEditor/Internal/Param/ParameterPickerArgs.h
@@ -91,6 +91,9 @@ struct FParameterPickerArgs
 
 	// Whether we should allow new parameters to be created by this widget
 	bool bAllowNew = true;
+
+	// Whether the search box should be focussed on widget creation
+	bool bFocusSearchWidget = true;
 };
 
 }
diff --git a/Engine/Plugins/Experimental/AnimNext/Source/AnimNextEditor/Private/AnimNextEditorModule.cpp b/Engine/Plugins/Experimental/AnimNext/Source/AnimNextEditor/Private/AnimNextEditorModule.cpp
--- a/Engine/Plugins/Experimental/AnimNext/Source/AnimNextEditor/Private/AnimNextEditorModule.cpp
+++ b/Engine/Plugins/Experimental/AnimNext/Source/AnimNextEditor/Private/AnimNextEditorModule.cpp
@@ -3,9 +3,6 @@
 #include "AnimNextEditorModule.h"
 
 #include "AnimNextConfig.h"
-#include "AssetToolsModule.h"
-#include "IAssetTools.h"
-#include "Graph/AssetTypeActions.h"
 #include "Graph/AnimNextGraphPanelNodeFactory.h"
 #include "Param/ParamTypePropertyCustomization.h"
 #include "Param/ParameterPickerArgs.h"
@@ -13,6 +10,18 @@
 #include "Param/ParamNamePropertyCustomization.h"
 #include "Param/SParameterPicker.h"
 #include "ISettingsModule.h"
+#include "UncookedOnlyUtils.h"
+#include "Graph/AnimNextGraph.h"
+#include "Param/AnimNextParameterBlock.h"
+#include "Param/AnimNextParameterBlockEntry.h"
+#include "Param/AnimNextParameterLibrary.h"
+#include "Param/IAnimNextParameterBlockGraphInterface.h"
+#include "Param/SParameterBlockView.h"
+#include "Param/SParameterLibraryView.h"
+#include "Graph/SAnimNextGraphView.h"
+#include "Scheduler/AnimNextSchedule.h"
+#include "Workspace/AnimNextWorkspaceEditor.h"
+
 
 #define LOCTEXT_NAMESPACE "AnimNextEditorModule"
 
@@ -30,10 +39,6 @@ class FModule : public IModule
 			LOCTEXT("SettingsDescription", "Customize AnimNext Settings."),
 			GetMutableDefault<UAnimNextConfig>()
 		);
-		
-		IAssetTools& AssetTools = FModuleManager::LoadModuleChecked<FAssetToolsModule>("AssetTools").Get();
-		AssetTypeActions_AnimNextGraph = MakeShared<FAssetTypeActions_AnimNextGraph>();
-		AssetTools.RegisterAssetTypeActions(AssetTypeActions_AnimNextGraph.ToSharedRef());
 
 		FPropertyEditorModule& PropertyModule = FModuleManager::LoadModuleChecked<FPropertyEditorModule>("PropertyEditor");
 
@@ -52,16 +57,106 @@ class FModule : public IModule
 
 		ParametersGraphPanelPinFactory = MakeShared<FParametersGraphPanelPinFactory>();
 		FEdGraphUtilities::RegisterVisualPinFactory(ParametersGraphPanelPinFactory);
+
+		FWorkspaceEditor::RegisterAssetDocumentWidget(UAnimNextParameterBlock::StaticClass()->GetFName(), [](TSharedRef<FWorkspaceEditor> InEditor, UObject* InAsset)
+		{
+			UAnimNextParameterBlock* ParameterBlock = CastChecked<UAnimNextParameterBlock>(InAsset);
+			UAnimNextParameterBlock_EditorData* EditorData = UncookedOnly::FUtils::GetEditorData(ParameterBlock);
+			return SNew(SParameterBlockView, EditorData)
+				.OnSelectionChanged_Lambda([WeakEditor = TWeakPtr<FWorkspaceEditor>(InEditor)](const TArray<UObject*>& InObjects)
+				{
+					if(TSharedPtr<FWorkspaceEditor> Editor = WeakEditor.Pin())
+					{
+						Editor->SetSelectedObjects(InObjects);
+					}
+				})
+				.OnOpenGraph_Lambda([WeakEditor = TWeakPtr<FWorkspaceEditor>(InEditor)](URigVMGraph* InGraph)
+				{
+					if(TSharedPtr<FWorkspaceEditor> Editor = WeakEditor.Pin())
+					{
+						if(IRigVMClientHost* RigVMClientHost = InGraph->GetImplementingOuter<IRigVMClientHost>())
+						{
+							if(UObject* EditorObject = RigVMClientHost->GetEditorObjectForRigVMGraph(InGraph))
+							{
+								Editor->OpenDocument(EditorObject, FDocumentTracker::EOpenDocumentCause::OpenNewDocument);
+							}
+						}
+					}
+				})
+				.OnDeleteEntries_Lambda([WeakEditor = TWeakPtr<FWorkspaceEditor>(InEditor)](const TArray<UAnimNextParameterBlockEntry*>& InEntries)
+				{
+					if(InEntries.Num() > 0)
+					{
+						if(TSharedPtr<FWorkspaceEditor> Editor = WeakEditor.Pin())
+						{
+							if(IRigVMClientHost* RigVMClientHost = InEntries[0]->GetImplementingOuter<IRigVMClientHost>())
+							{
+								for(UAnimNextParameterBlockEntry* Entry : InEntries)
+								{
+									if(IAnimNextParameterBlockGraphInterface* GraphInterface = Cast<IAnimNextParameterBlockGraphInterface>(Entry))
+									{
+										if(URigVMGraph* RigVMGraph = GraphInterface->GetGraph())
+										{
+											if (UObject* EditorObject = RigVMClientHost->GetEditorObjectForRigVMGraph(RigVMGraph))
+											{
+												Editor->CloseDocumentTab(EditorObject);
+											}
+										}
+									}
+								}
+							}
+						}
+					}
+				});
+		});
+
+		FWorkspaceEditor::RegisterAssetDocumentWidget(UAnimNextParameterLibrary::StaticClass()->GetFName(), [](TSharedRef<FWorkspaceEditor> InEditor, UObject* InAsset)
+		{
+			UAnimNextParameterLibrary* ParameterLibrary = CastChecked<UAnimNextParameterLibrary>(InAsset);
+			return SNew(SParameterLibraryView, ParameterLibrary)
+				.OnSelectionChanged_Lambda([WeakEditor = TWeakPtr<FWorkspaceEditor>(InEditor)](const TArray<UObject*>& InObjects)
+				{
+					if(TSharedPtr<FWorkspaceEditor> Editor = WeakEditor.Pin())
+					{
+						Editor->SetSelectedObjects(InObjects);
+					}
+				});
+		});
+
+		FWorkspaceEditor::RegisterAssetDocumentWidget(UAnimNextSchedule::StaticClass()->GetFName(), [](TSharedRef<FWorkspaceEditor> InEditor, UObject* InAsset)
+		{
+			UAnimNextSchedule* Schedule = CastChecked<UAnimNextSchedule>(InAsset);
+			FPropertyEditorModule& PropertyEditorModule = FModuleManager::GetModuleChecked<FPropertyEditorModule>( "PropertyEditor" );
+			FDetailsViewArgs DetailsViewArgs;
+			DetailsViewArgs.NameAreaSettings = FDetailsViewArgs::HideNameArea;
+			TSharedRef<IDetailsView> DetailsView = PropertyEditorModule.CreateDetailView(DetailsViewArgs);
+			DetailsView->SetObject(Schedule);
+			return DetailsView;
+		});
+
+		FWorkspaceEditor::RegisterAssetDocumentWidget(UAnimNextGraph::StaticClass()->GetFName(), [](TSharedRef<FWorkspaceEditor> InEditor, UObject* InAsset)
+		{
+			UAnimNextGraph* Graph = CastChecked<UAnimNextGraph>(InAsset);
+			UAnimNextGraph_EditorData* EditorData = UncookedOnly::FUtils::GetEditorData(Graph);
+			return SNew(SAnimNextGraphView, EditorData)
+				.OnOpenGraph_Lambda([WeakEditor = TWeakPtr<FWorkspaceEditor>(InEditor)](URigVMGraph* InGraph)
+				{
+					if(TSharedPtr<FWorkspaceEditor> Editor = WeakEditor.Pin())
+					{
+						if(IRigVMClientHost* RigVMClientHost = InGraph->GetImplementingOuter<IRigVMClientHost>())
+						{
+							if(UObject* EditorObject = RigVMClientHost->GetEditorObjectForRigVMGraph(InGraph))
+							{
+								Editor->OpenDocument(EditorObject, FDocumentTracker::EOpenDocumentCause::OpenNewDocument);
+							}
+						}
+					}
+				});
+		});
 	}
 
 	virtual void ShutdownModule() override
 	{
-		if(FModuleManager::Get().IsModuleLoaded("AssetTools"))
-		{
-			IAssetTools& AssetTools = FModuleManager::GetModuleChecked<FAssetToolsModule>("AssetTools").Get();
-			AssetTools.UnregisterAssetTypeActions(AssetTypeActions_AnimNextGraph.ToSharedRef());
-		}
-	
 		if(FModuleManager::Get().IsModuleLoaded("PropertyEditor"))
 		{
 			FPropertyEditorModule& PropertyModule = FModuleManager::LoadModuleChecked<FPropertyEditorModule>("PropertyEditor");
@@ -79,8 +174,6 @@ class FModule : public IModule
 			.Args(InArgs);
 	}
 
-	TSharedPtr<FAssetTypeActions_AnimNextGraph> AssetTypeActions_AnimNextGraph;
-
 	/** Node factory for the AnimNext graph */
 	TSharedPtr<FAnimNextGraphPanelNodeFactory> AnimNextGraphPanelNodeFactory;
 	
diff --git a/Engine/Plugins/Experimental/AnimNext/Source/AnimNextEditor/Private/EditorCVars.cpp b/Engine/Plugins/Experimental/AnimNext/Source/AnimNextEditor/Private/EditorCVars.cpp
new file mode 100644
--- /dev/null
+++ b/Engine/Plugins/Experimental/AnimNext/Source/AnimNextEditor/Private/EditorCVars.cpp
@@ -0,0 +1,13 @@
+﻿// Copyright Epic Games, Inc. All Rights Reserved.
+
+#include "EditorCVars.h"
+
+namespace UE::AnimNext::Editor::CVars
+{
+
+TAutoConsoleVariable<bool> GUseWorkspaceEditor(
+	TEXT("a.AnimNext.UseWorkspaceEditor"),
+	true,
+	TEXT("Use the workspace editor to edit AnimNext assets rather than standalone editors"));
+
+}
\ No newline at end of file
diff --git a/Engine/Plugins/Experimental/AnimNext/Source/AnimNextEditor/Private/EditorCVars.h b/Engine/Plugins/Experimental/AnimNext/Source/AnimNextEditor/Private/EditorCVars.h
new file mode 100644
--- /dev/null
+++ b/Engine/Plugins/Experimental/AnimNext/Source/AnimNextEditor/Private/EditorCVars.h
@@ -0,0 +1,12 @@
+﻿// Copyright Epic Games, Inc. All Rights Reserved.
+
+#pragma once
+
+#include "HAL/IConsoleManager.h"
+
+namespace UE::AnimNext::Editor::CVars
+{
+
+extern TAutoConsoleVariable<bool> GUseWorkspaceEditor;
+
+}
\ No newline at end of file
diff --git a/Engine/Plugins/Experimental/AnimNext/Source/AnimNextEditor/Private/EditorUtils.cpp b/Engine/Plugins/Experimental/AnimNext/Source/AnimNextEditor/Private/EditorUtils.cpp
--- a/Engine/Plugins/Experimental/AnimNext/Source/AnimNextEditor/Private/EditorUtils.cpp
+++ b/Engine/Plugins/Experimental/AnimNext/Source/AnimNextEditor/Private/EditorUtils.cpp
@@ -8,17 +8,23 @@
 #include "Param/AnimNextParameterBlock.h"
 #include "Param/AnimNextParameterLibrary.h"
 #include "PropertyBagDetails.h"
+#include "Graph/AnimNextGraphEntry.h"
 #include "AssetRegistry/AssetRegistryModule.h"
 #include "AssetRegistry/IAssetRegistry.h"
+#include "Param/IAnimNextParameterBlockGraphInterface.h"
+#include "Param/IAnimNextParameterBlockParameterInterface.h"
 #include "Param/Params.h"
+#include "Workspace/AnimNextWorkspace.h"
 
 namespace UE::AnimNext::Editor
 {
 
 void FUtils::GetAllGraphNames(const UAnimNextGraph_EditorData* InEditorData, TSet<FName>& OutNames)
 {
-	// TODO
-	ensure(false);
+	for(UAnimNextGraphEntry* Entry : InEditorData->Entries)
+	{
+		OutNames.Add(Entry->GraphName);
+	}
 }
 
 static const int32 MaxNameLength = 100;
@@ -233,10 +239,19 @@ FAnimNextParamType FUtils::GetParameterTypeFromMetaData(const FStringView& InStr
 }
 
 
-void FUtils::GetAllGraphNames(const UAnimNextParameterBlock_EditorData* InEditorData, TSet<FName>& OutNames)
+void FUtils::GetAllEntryNames(const UAnimNextParameterBlock_EditorData* InEditorData, TSet<FName>& OutNames)
 {
-	// TODO
-//	ensure(false);
+	for(UAnimNextParameterBlockEntry* Entry : InEditorData->Entries)
+	{
+		if(IAnimNextParameterBlockParameterInterface* Parameter = Cast<IAnimNextParameterBlockParameterInterface>(Entry))
+		{
+			OutNames.Add(Parameter->GetParameterName());
+		}
+		else if(IAnimNextParameterBlockGraphInterface* Graph = Cast<IAnimNextParameterBlockGraphInterface>(Entry))
+		{
+			OutNames.Add(Graph->GetGraphName());
+		}
+	}
 }
 
 FName FUtils::ValidateName(const UAnimNextParameterBlock_EditorData* InEditorData, const FString& InName)
@@ -246,7 +261,7 @@ FName FUtils::ValidateName(const UAnimNextParameterBlock_EditorData* InEditorDat
 		FNameValidator(const UAnimNextParameterBlock_EditorData* InEditorData)
 			: EditorData(InEditorData)
 		{
-			GetAllGraphNames(EditorData, Names);
+			GetAllEntryNames(EditorData, Names);
 		}
 		
 		virtual EValidatorResult IsValid (const FName& Name, bool bOriginal = false) override
@@ -539,4 +554,11 @@ FAnimNextParamType FUtils::GetParameterTypeFromName(FName InName)
 	return FAnimNextParamType();
 }
 
+bool FUtils::GetExportedAssetsForWorkspace(const FAssetData& InWorkspaceAsset, FAnimNextWorkspaceAssetRegistryExports& OutExports)
+{
+	const FString TagValue = InWorkspaceAsset.GetTagValueRef<FString>(UAnimNextWorkspace::ExportsAssetRegistryTag);
+	return FAnimNextWorkspaceAssetRegistryExports::StaticStruct()->ImportText(*TagValue, &OutExports, nullptr, PPF_None, nullptr, FAnimNextWorkspaceAssetRegistryExports::StaticStruct()->GetName()) != nullptr;
+}
+
+
 }
diff --git a/Engine/Plugins/Experimental/AnimNext/Source/AnimNextEditor/Private/EditorUtils.h b/Engine/Plugins/Experimental/AnimNext/Source/AnimNextEditor/Private/EditorUtils.h
--- a/Engine/Plugins/Experimental/AnimNext/Source/AnimNextEditor/Private/EditorUtils.h
+++ b/Engine/Plugins/Experimental/AnimNext/Source/AnimNextEditor/Private/EditorUtils.h
@@ -17,6 +17,7 @@ class UAnimNextParameter;
 class URigVMController;
 struct FAnimNextParameterLibraryAssetRegistryExports;
 struct FAnimNextParameterBlockAssetRegistryExports;
+struct FAnimNextWorkspaceAssetRegistryExports;
 
 namespace UE::AnimNext::Editor
 {
@@ -31,7 +32,7 @@ struct FUtils
 
 	static FName ValidateName(const UAnimNextParameterBlock_EditorData* InEditorData, const FString& InName);
 
-	static void GetAllGraphNames(const UAnimNextParameterBlock_EditorData* InEditorData, TSet<FName>& OutNames);
+	static void GetAllEntryNames(const UAnimNextParameterBlock_EditorData* InEditorData, TSet<FName>& OutNames);
 
 	static void GetFilteredVariableTypeTree(TArray<TSharedPtr<UEdGraphSchema_K2::FPinTypeTreeInfo>>& TypeTree, ETypeTreeFilter TypeTreeFilter);
 
@@ -49,6 +50,8 @@ struct FUtils
 	// If the name cannot be found, the returned type will be invalid
 	// Note that this is expensive and can query the asset registry
 	static FAnimNextParamType GetParameterTypeFromName(FName InName);
+
+	static bool GetExportedAssetsForWorkspace(const FAssetData& InWorkspaceAsset, FAnimNextWorkspaceAssetRegistryExports& OutExports);
 };
 
 }
diff --git a/Engine/Plugins/Experimental/AnimNext/Source/AnimNextEditor/Private/Graph/AnimNextGraphAssetDefinition.cpp b/Engine/Plugins/Experimental/AnimNext/Source/AnimNextEditor/Private/Graph/AnimNextGraphAssetDefinition.cpp
new file mode 100644
--- /dev/null
+++ b/Engine/Plugins/Experimental/AnimNext/Source/AnimNextEditor/Private/Graph/AnimNextGraphAssetDefinition.cpp
@@ -0,0 +1,26 @@
+﻿// Copyright Epic Games, Inc. All Rights Reserved.
+
+#include "AnimNextGraphAssetDefinition.h"
+#include "AnimNextGraphEditor.h"
+#include "Workspace/AnimNextWorkspaceEditor.h"
+#include "EditorCVars.h"
+
+EAssetCommandResult UAssetDefinition_AnimNextGraph::OpenAssets(const FAssetOpenArgs& OpenArgs) const
+{
+	using namespace UE::AnimNext::Editor;
+	
+	for (UAnimNextGraph* Asset : OpenArgs.LoadObjects<UAnimNextGraph>())
+	{
+		if(CVars::GUseWorkspaceEditor.GetValueOnGameThread())
+		{
+			FWorkspaceEditor::OpenWorkspaceForAsset(Asset, FWorkspaceEditor::EOpenWorkspaceMethod::Default);
+		}
+		else
+		{
+			TSharedRef<FGraphEditor> GraphEditor = MakeShared<FGraphEditor>();
+			GraphEditor->InitEditor(OpenArgs.GetToolkitMode(), OpenArgs.ToolkitHost, Asset);
+		}
+	}
+
+	return EAssetCommandResult::Handled;
+}
diff --git a/Engine/Plugins/Experimental/AnimNext/Source/AnimNextEditor/Private/Graph/AnimNextGraphAssetDefinition.h b/Engine/Plugins/Experimental/AnimNext/Source/AnimNextEditor/Private/Graph/AnimNextGraphAssetDefinition.h
new file mode 100644
--- /dev/null
+++ b/Engine/Plugins/Experimental/AnimNext/Source/AnimNextEditor/Private/Graph/AnimNextGraphAssetDefinition.h
@@ -0,0 +1,29 @@
+﻿// Copyright Epic Games, Inc. All Rights Reserved.
+
+#pragma once
+
+#include "Graph/AnimNextGraph.h"
+#include "AssetDefinitionDefault.h"
+#include "AnimNextGraphAssetDefinition.generated.h"
+
+#define LOCTEXT_NAMESPACE "AnimNextAssetDefinitions"
+
+UCLASS()
+class UAssetDefinition_AnimNextGraph : public UAssetDefinitionDefault
+{
+	GENERATED_BODY()
+
+public:
+	// UAssetDefinition interface
+	virtual FText GetAssetDisplayName() const override { return LOCTEXT("AnimNextGraph", "AnimNext Graph"); }
+	virtual FLinearColor GetAssetColor() const override { return FLinearColor(FColor(128,128,64)); }
+	virtual TSoftClassPtr<UObject> GetAssetClass() const override { return UAnimNextGraph::StaticClass(); }
+	virtual EAssetCommandResult OpenAssets(const FAssetOpenArgs& OpenArgs) const override;
+	virtual TConstArrayView<FAssetCategoryPath> GetAssetCategories() const override
+	{
+		static const auto Categories = { FAssetCategoryPath(EAssetCategoryPaths::Animation, LOCTEXT("AnimNextSubMenu", "AnimNext")) };
+		return Categories;
+	}
+};
+
+#undef LOCTEXT_NAMESPACE
\ No newline at end of file
diff --git a/Engine/Plugins/Experimental/AnimNext/Source/AnimNextEditor/Private/Graph/AnimNextGraphDocumentSummoner.cpp b/Engine/Plugins/Experimental/AnimNext/Source/AnimNextEditor/Private/Graph/AnimNextGraphDocumentSummoner.cpp
new file mode 100644
--- /dev/null
+++ b/Engine/Plugins/Experimental/AnimNext/Source/AnimNextEditor/Private/Graph/AnimNextGraphDocumentSummoner.cpp
@@ -0,0 +1,35 @@
+﻿// Copyright Epic Games, Inc. All Rights Reserved.
+
+#include "AnimNextGraphDocumentSummoner.h"
+#include "SActionMenu.h"
+#include "Graph/AnimNextGraph_EdGraph.h"
+#include "Workspace/AnimNextWorkspaceEditor.h"
+
+namespace UE::AnimNext::Editor
+{
+
+FAnimNextGraphDocumentSummoner::FAnimNextGraphDocumentSummoner(FName InIdentifier, TSharedPtr<FWorkspaceEditor> InHostingApp)
+	: FGraphDocumentSummoner(InIdentifier, InHostingApp)
+{
+}
+
+FActionMenuContent FAnimNextGraphDocumentSummoner::OnCreateGraphActionMenu(UEdGraph* InGraph, const FVector2D& InNodePosition, const TArray<UEdGraphPin*>& InDraggedPins, bool bAutoExpand, SGraphEditor::FActionMenuClosed InOnMenuClosed) const
+{
+	TSharedRef<SActionMenu> ActionMenu = SNew(SActionMenu)
+		.AutoExpandActionMenu(bAutoExpand)
+		.Graph(InGraph)
+		.NewNodePosition(InNodePosition)
+		.DraggedFromPins(InDraggedPins)
+		.OnClosedCallback(InOnMenuClosed);
+
+	TSharedPtr<SWidget> FilterTextBox = StaticCastSharedRef<SWidget>(ActionMenu->GetFilterTextBox());
+	return FActionMenuContent(StaticCastSharedRef<SWidget>(ActionMenu), FilterTextBox);
+}
+
+bool FAnimNextGraphDocumentSummoner::IsPayloadSupported(TSharedRef<FTabPayload> Payload) const
+{
+	UObject* Object = Payload->IsValid() ? FTabPayload_UObject::CastChecked<UObject>(Payload) : nullptr;
+	return Object && Object->IsA<UAnimNextGraph_EdGraph>();
+}
+
+}
diff --git a/Engine/Plugins/Experimental/AnimNext/Source/AnimNextEditor/Private/Graph/AnimNextGraphDocumentSummoner.h b/Engine/Plugins/Experimental/AnimNext/Source/AnimNextEditor/Private/Graph/AnimNextGraphDocumentSummoner.h
new file mode 100644
--- /dev/null
+++ b/Engine/Plugins/Experimental/AnimNext/Source/AnimNextEditor/Private/Graph/AnimNextGraphDocumentSummoner.h
@@ -0,0 +1,22 @@
+﻿// Copyright Epic Games, Inc. All Rights Reserved.
+
+#pragma once
+
+#include "CoreMinimal.h"
+#include "Workspace/GraphDocumentSummoner.h"
+
+namespace UE::AnimNext::Editor
+{
+
+class FAnimNextGraphDocumentSummoner : public FGraphDocumentSummoner
+{
+public:
+	explicit FAnimNextGraphDocumentSummoner(FName InIdentifier, TSharedPtr<FWorkspaceEditor> InHostingApp);
+
+private:
+	// FGraphDocumentSummoner interface
+	virtual FActionMenuContent OnCreateGraphActionMenu(UEdGraph* InGraph, const FVector2D& InNodePosition, const TArray<UEdGraphPin*>& InDraggedPins, bool bAutoExpand, SGraphEditor::FActionMenuClosed InOnMenuClosed) const override;
+	virtual bool IsPayloadSupported(TSharedRef<FTabPayload> Payload) const override;
+};
+
+}
diff --git a/Engine/Plugins/Experimental/AnimNext/Source/AnimNextEditor/Private/Graph/AnimNextGraphEditor.cpp b/Engine/Plugins/Experimental/AnimNext/Source/AnimNextEditor/Private/Graph/AnimNextGraphEditor.cpp
--- a/Engine/Plugins/Experimental/AnimNext/Source/AnimNextEditor/Private/Graph/AnimNextGraphEditor.cpp
+++ b/Engine/Plugins/Experimental/AnimNext/Source/AnimNextEditor/Private/Graph/AnimNextGraphEditor.cpp
@@ -62,10 +62,10 @@ void FGraphEditor::InitEditor(const EToolkitMode::Type InMode, const TSharedPtr<
 	RegenerateMenusAndToolbars();
 
 	// Open initial document
-	DocumentManager->OpenDocument(FTabPayload_UObject::Make(AnimNextGraph_EditorData->RootGraph), FDocumentTracker::EOpenDocumentCause::OpenNewDocument);
-
-	// Make sure editor graph matches the model
-	AnimNextGraph_EditorData->RebuildEdGraphFromModel();
+	if(AnimNextGraph_EditorData->Graphs.Num() > 0)
+	{
+		DocumentManager->OpenDocument(FTabPayload_UObject::Make(AnimNextGraph_EditorData->Graphs[0]), FDocumentTracker::EOpenDocumentCause::OpenNewDocument);
+	}
 }
 
 void FGraphEditor::BindCommands()
@@ -145,7 +145,7 @@ UEdGraph* FGraphEditor::GetFocusedGraph() const
 URigVMGraph* FGraphEditor::GetFocusedVMGraph() const
 {
 	UAnimNextGraph_EdGraph* EdGraph = Cast<UAnimNextGraph_EdGraph>(GetFocusedGraph());
-	return AnimNextGraph_EditorData->GetVMGraphForEdGraph(EdGraph);
+	return AnimNextGraph_EditorData->GetRigVMGraphForEditorObject(EdGraph);
 }
 
 URigVMController* FGraphEditor::GetFocusedVMController() const
diff --git a/Engine/Plugins/Experimental/AnimNext/Source/AnimNextEditor/Private/Graph/AssetTypeActions.cpp b/Engine/Plugins/Experimental/AnimNext/Source/AnimNextEditor/Private/Graph/AssetTypeActions.cpp
deleted file mode 100644
diff --git a/Engine/Plugins/Experimental/AnimNext/Source/AnimNextEditor/Private/Graph/AssetTypeActions.h b/Engine/Plugins/Experimental/AnimNext/Source/AnimNextEditor/Private/Graph/AssetTypeActions.h
deleted file mode 100644
diff --git a/Engine/Plugins/Experimental/AnimNext/Source/AnimNextEditor/Private/Graph/GraphFactory.cpp b/Engine/Plugins/Experimental/AnimNext/Source/AnimNextEditor/Private/Graph/GraphFactory.cpp
--- a/Engine/Plugins/Experimental/AnimNext/Source/AnimNextEditor/Private/Graph/GraphFactory.cpp
+++ b/Engine/Plugins/Experimental/AnimNext/Source/AnimNextEditor/Private/Graph/GraphFactory.cpp
@@ -25,72 +25,18 @@ UObject* UAnimNextGraphFactory::FactoryCreateNew(UClass* Class, UObject* InParen
 {
 	UAnimNextGraph* NewGraph = NewObject<UAnimNextGraph>(InParent, Class, Name, Flags | RF_Public | RF_Standalone | RF_Transactional | RF_LoadCompleted);
 
-	// Create internal editor data & graphs
+	// Create internal editor data
 	UAnimNextGraph_EditorData* EditorData = NewObject<UAnimNextGraph_EditorData>(NewGraph, TEXT("EditorData"));
 	NewGraph->EditorData = EditorData;
 	EditorData->Initialize(/*bRecompileVM*/false);
-	EditorData->GetRigVMClient()->SetExecuteContextStruct(FAnimNextExecuteContext::StaticStruct());
 
-	// Add initial execution unit
-	URigVMController* Controller = EditorData->GetRigVMClient()->GetController(EditorData->RootGraph);
-
-	URigVMUnitNode* MainEntryPointNode = Controller->AddUnitNode(FRigUnit_AnimNextGraphRoot::StaticStruct(), FRigUnit_AnimNextGraphRoot::EventName, FVector2D(-400.0f, 0.0f), FString(), false);
-	URigVMPin* BeginExecutePin = MainEntryPointNode->FindPin(GET_MEMBER_NAME_STRING_CHECKED(FRigUnit_AnimNextGraphRoot, Result));
-	check(BeginExecutePin);
-	check(BeginExecutePin->GetDirection() == ERigVMPinDirection::Input);
-
-	//// Add function to function lib
-	//{
-	//	FRigVMControllerGraphGuard LibraryGraphGuard(Controller, EditorData->GetRigVMClient()->GetFunctionLibrary());
-	//	EditorData->EntryPoint = Controller->AddFunctionToLibrary(UE::AnimNext::Graph::EntryPointName, true, FVector2D::ZeroVector, false, false);
-
-	//	// Add exposed pin for result to function
-	//	{
-	//		FRigVMControllerGraphGuard FunctionGraphGuard(Controller, EditorData->EntryPoint->GetContainedGraph());
-
-	//		// TODO: using float for now, but needs to be user-driven
-	//		const FString CPPType = TEXT("float");
-	//		const FName CPPTypeObjectPath = NAME_None;
-	//		const FString DefaultValue = TEXT("0.0");
-
-	//		//FunctionController->AddExposedPin(UE::AnimNext::Graph::ResultName, ERigVMPinDirection::Output, CPPType, CPPTypeObjectPath, DefaultValue, false, false);
-	//		Controller->AddExposedPin(UE::AnimNext::Graph::ResultName, ERigVMPinDirection::Output, CPPType, CPPTypeObjectPath, DefaultValue, false, false);
-	//	}
-	//}
-
-	//// Add function call
-	//URigVMFunctionReferenceNode* FunctionCallNode = Controller->AddFunctionReferenceNode(EditorData->EntryPoint, FVector2D::ZeroVector, UE::AnimNext::Graph::EntryPointName.ToString(), false, false);
-	//URigVMPin* FunctionExecutePin = FunctionCallNode->FindPin(FRigVMStruct::ExecuteContextName.ToString());
-	//check(FunctionExecutePin);
-	//check(FunctionExecutePin->GetDirection() == ERigVMPinDirection::IO);
-
-	//// Link entry point to function call
-	//Controller->AddLink(BeginExecutePin, FunctionExecutePin, false);
-
-	//URigVMPin* FunctionResultPin = FunctionCallNode->FindPin(UE::AnimNext::Graph::ResultName.ToString());
-	//check(FunctionResultPin);
-	//check(FunctionResultPin->GetDirection() == ERigVMPinDirection::Output);
-
-	//// Add end-execution unit of the correct type
-	//// TODO: using a float for now, but need to use a registry to determine correct type
-	//URigVMUnitNode* MainExitPointNode = Controller->AddUnitNode(FRigUnit_AnimNextEndExecution_Float::StaticStruct(), FRigVMStruct::ExecuteName, FVector2D(400.0f, 0.0f), FString(), false);
-	//URigVMPin* EndExecutePin = MainExitPointNode->FindPin(GET_MEMBER_NAME_STRING_CHECKED(FRigUnit_AnimNextEndExecution, ExecuteContext));
-	//check(EndExecutePin);
-	//check(EndExecutePin->GetDirection() == ERigVMPinDirection::Input);
-
-	//URigVMPin* EndExecuteResultPin = MainExitPointNode->FindPin(GET_MEMBER_NAME_STRING_CHECKED(FRigUnit_AnimNextEndExecution_Float, Result));
-	//check(EndExecuteResultPin);
-	//check(EndExecuteResultPin->GetDirection() == ERigVMPinDirection::Input);
-
-	//// Link function call to exit point 
-	//Controller->AddLink(FunctionExecutePin, EndExecutePin, false);
-
-	//// Link result pins 
-	//Controller->AddLink(FunctionResultPin, EndExecuteResultPin, false);
+	// Add root graph
+	EditorData->AddGraph(TEXT("Root"));
+	check(EditorData->Graphs.Num() > 0);
 
 	// Compile the initial skeleton
 	UE::AnimNext::UncookedOnly::FUtils::Compile(NewGraph);
 	check(!EditorData->bErrorsDuringCompilation);
-	
+
 	return NewGraph;
 }
\ No newline at end of file
diff --git a/Engine/Plugins/Experimental/AnimNext/Source/AnimNextEditor/Private/Graph/SAnimNextGraphView.cpp b/Engine/Plugins/Experimental/AnimNext/Source/AnimNextEditor/Private/Graph/SAnimNextGraphView.cpp
new file mode 100644
--- /dev/null
+++ b/Engine/Plugins/Experimental/AnimNext/Source/AnimNextEditor/Private/Graph/SAnimNextGraphView.cpp
@@ -0,0 +1,38 @@
+﻿// Copyright Epic Games, Inc. All Rights Reserved.
+
+#include "SAnimNextGraphView.h"
+#include "Widgets/SBoxPanel.h"
+#include "SPrimaryButton.h"
+#include "Graph/AnimNextGraph_EditorData.h"
+
+#define LOCTEXT_NAMESPACE "SAnimNextGraphView"
+
+namespace UE::AnimNext::Editor
+{
+
+void SAnimNextGraphView::Construct(const FArguments& InArgs, UAnimNextGraph_EditorData* InEditorData)
+{
+	EditorData = InEditorData;
+
+	ChildSlot
+	[
+		SNew(SVerticalBox)
+		+SVerticalBox::Slot()
+		.AutoHeight()
+		.Padding(5.0f)
+		[
+			SNew(SPrimaryButton)
+			.Text(LOCTEXT("OpenGraph", "Open Graph"))
+			.OnClicked_Lambda([this, OnOpenGraph = InArgs._OnOpenGraph]()
+			{
+				OnOpenGraph.ExecuteIfBound(EditorData->GetRigVMGraphForEditorObject(EditorData->Graphs[0]));
+
+				return FReply::Handled();
+			})
+		]
+	];
+}
+
+}
+
+#undef LOCTEXT_NAMESPACE
diff --git a/Engine/Plugins/Experimental/AnimNext/Source/AnimNextEditor/Private/Graph/SAnimNextGraphView.h b/Engine/Plugins/Experimental/AnimNext/Source/AnimNextEditor/Private/Graph/SAnimNextGraphView.h
new file mode 100644
--- /dev/null
+++ b/Engine/Plugins/Experimental/AnimNext/Source/AnimNextEditor/Private/Graph/SAnimNextGraphView.h
@@ -0,0 +1,30 @@
+﻿// Copyright Epic Games, Inc. All Rights Reserved.
+
+#pragma once
+
+#include "Widgets/SCompoundWidget.h"
+
+class UAnimNextGraph_EditorData;
+class URigVMGraph;
+
+namespace UE::AnimNext::Editor
+{
+
+class SAnimNextGraphView : public SCompoundWidget
+{
+public:
+	DECLARE_DELEGATE_OneParam(FOnOpenGraph, URigVMGraph* /*InGraph*/);
+
+	SLATE_BEGIN_ARGS(SAnimNextGraphView) {}
+
+	SLATE_EVENT(SAnimNextGraphView::FOnOpenGraph, OnOpenGraph)
+
+	SLATE_END_ARGS()
+
+	void Construct(const FArguments& InArgs, UAnimNextGraph_EditorData* InEditorData);
+
+private:
+	UAnimNextGraph_EditorData* EditorData = nullptr;
+};
+
+}
\ No newline at end of file
diff --git a/Engine/Plugins/Experimental/AnimNext/Source/AnimNextEditor/Private/Param/AnimNextParameterBlockEditor.cpp b/Engine/Plugins/Experimental/AnimNext/Source/AnimNextEditor/Private/Param/AnimNextParameterBlockEditor.cpp
--- a/Engine/Plugins/Experimental/AnimNext/Source/AnimNextEditor/Private/Param/AnimNextParameterBlockEditor.cpp
+++ b/Engine/Plugins/Experimental/AnimNext/Source/AnimNextEditor/Private/Param/AnimNextParameterBlockEditor.cpp
@@ -192,7 +192,7 @@ UEdGraph* FParameterBlockEditor::GetFocusedGraph() const
 URigVMGraph* FParameterBlockEditor::GetFocusedVMGraph() const
 {
 	UAnimNextParameterBlock_EdGraph* EdGraph = Cast<UAnimNextParameterBlock_EdGraph>(GetFocusedGraph());
-	return EditorData->GetVMGraphForEdGraph(EdGraph);
+	return EditorData->GetRigVMGraphForEditorObject(EdGraph);
 }
 
 URigVMController* FParameterBlockEditor::GetFocusedVMController() const
diff --git a/Engine/Plugins/Experimental/AnimNext/Source/AnimNextEditor/Private/Param/ParameterBlockGraphDocumentSummoner.cpp b/Engine/Plugins/Experimental/AnimNext/Source/AnimNextEditor/Private/Param/ParameterBlockGraphDocumentSummoner.cpp
new file mode 100644
--- /dev/null
+++ b/Engine/Plugins/Experimental/AnimNext/Source/AnimNextEditor/Private/Param/ParameterBlockGraphDocumentSummoner.cpp
@@ -0,0 +1,36 @@
+﻿// Copyright Epic Games, Inc. All Rights Reserved.
+
+#include "ParameterBlockGraphDocumentSummoner.h"
+#include "SParametersActionMenu.h"
+#include "Graph/AnimNextGraph.h"
+#include "Param/AnimNextParameterBlock_EdGraph.h"
+#include "Workspace/AnimNextWorkspaceEditor.h"
+
+namespace UE::AnimNext::Editor
+{
+
+FParameterBlockGraphDocumentSummoner::FParameterBlockGraphDocumentSummoner(FName InIdentifier, TSharedPtr<FWorkspaceEditor> InHostingApp)
+	: FGraphDocumentSummoner(InIdentifier, InHostingApp)
+{
+}
+
+FActionMenuContent FParameterBlockGraphDocumentSummoner::OnCreateGraphActionMenu(UEdGraph* InGraph, const FVector2D& InNodePosition, const TArray<UEdGraphPin*>& InDraggedPins, bool bAutoExpand, SGraphEditor::FActionMenuClosed InOnMenuClosed) const
+{
+	TSharedRef<SParametersActionMenu> ActionMenu = SNew(SParametersActionMenu)
+		.AutoExpandActionMenu(bAutoExpand)
+		.Graph(InGraph)
+		.NewNodePosition(InNodePosition)
+		.DraggedFromPins(InDraggedPins)
+		.OnClosedCallback(InOnMenuClosed);
+
+	TSharedPtr<SWidget> FilterTextBox = StaticCastSharedRef<SWidget>(ActionMenu->GetFilterTextBox());
+	return FActionMenuContent(StaticCastSharedRef<SWidget>(ActionMenu), FilterTextBox);
+}
+
+bool FParameterBlockGraphDocumentSummoner::IsPayloadSupported(TSharedRef<FTabPayload> Payload) const
+{
+	UObject* Object = Payload->IsValid() ? FTabPayload_UObject::CastChecked<UObject>(Payload) : nullptr;
+	return Object && Object->IsA<UAnimNextParameterBlock_EdGraph>();
+}
+
+}
diff --git a/Engine/Plugins/Experimental/AnimNext/Source/AnimNextEditor/Private/Param/ParameterBlockGraphDocumentSummoner.h b/Engine/Plugins/Experimental/AnimNext/Source/AnimNextEditor/Private/Param/ParameterBlockGraphDocumentSummoner.h
new file mode 100644
--- /dev/null
+++ b/Engine/Plugins/Experimental/AnimNext/Source/AnimNextEditor/Private/Param/ParameterBlockGraphDocumentSummoner.h
@@ -0,0 +1,22 @@
+﻿// Copyright Epic Games, Inc. All Rights Reserved.
+
+#pragma once
+
+#include "CoreMinimal.h"
+#include "Workspace/GraphDocumentSummoner.h"
+
+namespace UE::AnimNext::Editor
+{
+
+class FParameterBlockGraphDocumentSummoner : public FGraphDocumentSummoner
+{
+public:
+	explicit FParameterBlockGraphDocumentSummoner(FName InIdentifier, TSharedPtr<FWorkspaceEditor> InHostingApp);
+
+private:
+	// FGraphDocumentSummoner interface
+	virtual FActionMenuContent OnCreateGraphActionMenu(UEdGraph* InGraph, const FVector2D& InNodePosition, const TArray<UEdGraphPin*>& InDraggedPins, bool bAutoExpand, SGraphEditor::FActionMenuClosed InOnMenuClosed) const override;
+	virtual bool IsPayloadSupported(TSharedRef<FTabPayload> Payload) const override;
+};
+
+}
diff --git a/Engine/Plugins/Experimental/AnimNext/Source/AnimNextEditor/Private/Param/ParametersAssetDefinitions.cpp b/Engine/Plugins/Experimental/AnimNext/Source/AnimNextEditor/Private/Param/ParametersAssetDefinitions.cpp
--- a/Engine/Plugins/Experimental/AnimNext/Source/AnimNextEditor/Private/Param/ParametersAssetDefinitions.cpp
+++ b/Engine/Plugins/Experimental/AnimNext/Source/AnimNextEditor/Private/Param/ParametersAssetDefinitions.cpp
@@ -5,15 +5,26 @@
 #include "AnimNextParameterLibraryEditor.h"
 #include "Param/AnimNextParameterBlock.h"
 #include "Param/AnimNextParameterLibrary.h"
+#include "Workspace/AnimNextWorkspaceEditor.h"
+#include "EditorCVars.h"
 
 #define LOCTEXT_NAMESPACE "AnimNextAssetDefinitions"
 
 EAssetCommandResult UAssetDefinition_AnimNextParameterLibrary::OpenAssets(const FAssetOpenArgs& OpenArgs) const
 {
+	using namespace UE::AnimNext::Editor;
+
 	for (UAnimNextParameterLibrary* Asset : OpenArgs.LoadObjects<UAnimNextParameterLibrary>())
 	{
-		TSharedRef<UE::AnimNext::Editor::FParameterLibraryEditor> Editor = MakeShared<UE::AnimNext::Editor::FParameterLibraryEditor>();
-		Editor->InitEditor(OpenArgs.GetToolkitMode(), OpenArgs.ToolkitHost, Asset);
+		if(CVars::GUseWorkspaceEditor.GetValueOnGameThread())
+		{
+			FWorkspaceEditor::OpenWorkspaceForAsset(Asset, FWorkspaceEditor::EOpenWorkspaceMethod::Default);
+		}
+		else
+		{
+			TSharedRef<FParameterLibraryEditor> Editor = MakeShared<UE::AnimNext::Editor::FParameterLibraryEditor>();
+			Editor->InitEditor(OpenArgs.GetToolkitMode(), OpenArgs.ToolkitHost, Asset);
+		}
 	}
 
 	return EAssetCommandResult::Handled;
@@ -32,10 +43,19 @@ FText UAssetDefinition_AnimNextParameter::GetObjectDisplayNameText(UObject* Obje
 
 EAssetCommandResult UAssetDefinition_AnimNextParameterBlock::OpenAssets(const FAssetOpenArgs& OpenArgs) const
 {
+	using namespace UE::AnimNext::Editor;
+
 	for (UAnimNextParameterBlock* Asset : OpenArgs.LoadObjects<UAnimNextParameterBlock>())
 	{
-		TSharedRef<UE::AnimNext::Editor::FParameterBlockEditor> Editor = MakeShared<UE::AnimNext::Editor::FParameterBlockEditor>();
-		Editor->InitEditor(OpenArgs.GetToolkitMode(), OpenArgs.ToolkitHost, Asset);
+		if(CVars::GUseWorkspaceEditor.GetValueOnGameThread())
+		{
+			FWorkspaceEditor::OpenWorkspaceForAsset(Asset, FWorkspaceEditor::EOpenWorkspaceMethod::Default);
+		}
+		else
+		{
+			TSharedRef<FParameterBlockEditor> Editor = MakeShared<UE::AnimNext::Editor::FParameterBlockEditor>();
+			Editor->InitEditor(OpenArgs.GetToolkitMode(), OpenArgs.ToolkitHost, Asset);
+		}
 	}
 
 	return EAssetCommandResult::Handled;
diff --git a/Engine/Plugins/Experimental/AnimNext/Source/AnimNextEditor/Private/Param/SParameterPicker.cpp b/Engine/Plugins/Experimental/AnimNext/Source/AnimNextEditor/Private/Param/SParameterPicker.cpp
--- a/Engine/Plugins/Experimental/AnimNext/Source/AnimNextEditor/Private/Param/SParameterPicker.cpp
+++ b/Engine/Plugins/Experimental/AnimNext/Source/AnimNextEditor/Private/Param/SParameterPicker.cpp
@@ -117,6 +117,23 @@ void SParameterPicker::Construct(const FArguments& InArgs)
 		Args.OnGetParameterBindings->BindSP(this, &SParameterPicker::HandleGetParameterBindings);
 	}
 
+	if(Args.bFocusSearchWidget)
+	{
+		RegisterActiveTimer( 0.f, FWidgetActiveTimerDelegate::CreateLambda([this](double InCurrentTime, float InDeltaTime)
+		{
+			if (SearchBox.IsValid())
+			{
+				FWidgetPath WidgetToFocusPath;
+				FSlateApplication::Get().GeneratePathToWidgetUnchecked(SearchBox.ToSharedRef(), WidgetToFocusPath);
+				FSlateApplication::Get().SetKeyboardFocus(WidgetToFocusPath, EFocusCause::SetDirectly);
+				WidgetToFocusPath.GetWindow()->SetWidgetToFocusOnActivate(SearchBox);
+				return EActiveTimerReturnType::Stop;
+			}
+
+			return EActiveTimerReturnType::Continue;
+		}));
+	}
+
 	TSharedPtr<SHeaderRow> HeaderRow;
 
 	ChildSlot
@@ -126,7 +143,7 @@ void SParameterPicker::Construct(const FArguments& InArgs)
 		.AutoHeight()
 		.Padding(2.0f)
 		[
-			SNew(SSearchBox)
+			SAssignNew(SearchBox, SSearchBox)
 			.OnTextChanged_Lambda([this](FText InText)
 			{
 				FilterText = InText;
diff --git a/Engine/Plugins/Experimental/AnimNext/Source/AnimNextEditor/Private/Param/SParameterPicker.h b/Engine/Plugins/Experimental/AnimNext/Source/AnimNextEditor/Private/Param/SParameterPicker.h
--- a/Engine/Plugins/Experimental/AnimNext/Source/AnimNextEditor/Private/Param/SParameterPicker.h
+++ b/Engine/Plugins/Experimental/AnimNext/Source/AnimNextEditor/Private/Param/SParameterPicker.h
@@ -6,6 +6,8 @@
 #include "Param/ParameterPickerArgs.h"
 #include "Widgets/Views/STreeView.h"
 
+class SSearchBox;
+
 namespace UE::AnimNext::Editor
 {
 
@@ -55,6 +57,8 @@ private:
 	TArray<TSharedRef<FParameterPickerEntry>> FilteredHierarchy;
 
 	FText FilterText;
+
+	TSharedPtr<SSearchBox> SearchBox;
 };
 
 }
\ No newline at end of file
diff --git a/Engine/Plugins/Experimental/AnimNext/Source/AnimNextEditor/Private/Param/SParameterPickerCombo.cpp b/Engine/Plugins/Experimental/AnimNext/Source/AnimNextEditor/Private/Param/SParameterPickerCombo.cpp
--- a/Engine/Plugins/Experimental/AnimNext/Source/AnimNextEditor/Private/Param/SParameterPickerCombo.cpp
+++ b/Engine/Plugins/Experimental/AnimNext/Source/AnimNextEditor/Private/Param/SParameterPickerCombo.cpp
@@ -7,6 +7,7 @@
 #include "Widgets/Input/SComboButton.h"
 #include "Framework/Application/SlateApplication.h"
 #include "Widgets/Text/SInlineEditableTextBlock.h"
+#include "Widgets/Input/SSearchBox.h"
 
 namespace UE::AnimNext::Editor
 {
@@ -15,7 +16,7 @@ void SParameterPickerCombo::Construct(const FArguments& InArgs)
 {
 	OnGetParameterNameDelegate = InArgs._OnGetParameterName;
 	OnGetParameterTypeDelegate = InArgs._OnGetParameterType;
-	
+
 	FParameterPickerArgs PickerArgs = InArgs._PickerArgs;
 	PickerArgs.OnParameterPicked = FOnParameterPicked::CreateLambda([this, OrigPickerArgs = InArgs._PickerArgs](const FParameterBindingReference& InReference)
 	{
@@ -26,11 +27,12 @@ void SParameterPickerCombo::Construct(const FArguments& InArgs)
 
 		RequestRefresh();
 	});
+	PickerArgs.bFocusSearchWidget = true;
 
 	ChildSlot
 	[
 		SNew(SComboButton)
-		.OnGetMenuContent_Lambda([PickerArgs]()
+		.OnGetMenuContent_Lambda([PickerArgs, this]()
 		{
 			return SNew(SParameterPicker)
 				.Args(PickerArgs);
diff --git a/Engine/Plugins/Experimental/AnimNext/Source/AnimNextEditor/Private/Param/SParameterPickerCombo.h b/Engine/Plugins/Experimental/AnimNext/Source/AnimNextEditor/Private/Param/SParameterPickerCombo.h
--- a/Engine/Plugins/Experimental/AnimNext/Source/AnimNextEditor/Private/Param/SParameterPickerCombo.h
+++ b/Engine/Plugins/Experimental/AnimNext/Source/AnimNextEditor/Private/Param/SParameterPickerCombo.h
@@ -8,6 +8,11 @@
 
 struct FAnimNextParamType;
 
+namespace UE::AnimNext::Editor
+{
+class SParameterPicker;
+}
+
 namespace UE::AnimNext::Editor
 {
 
@@ -61,6 +66,9 @@ private:
 
 	// Cached color
 	FSlateColor IconColor = FLinearColor::Gray;
+
+	// Picker widget used to focus after the popup is displayed
+	TWeakPtr<SParameterPicker> PickerWidget;
 };
 
 }
\ No newline at end of file
diff --git a/Engine/Plugins/Experimental/AnimNext/Source/AnimNextEditor/Private/Scheduler/ScheduleAssetDefinition.cpp b/Engine/Plugins/Experimental/AnimNext/Source/AnimNextEditor/Private/Scheduler/ScheduleAssetDefinition.cpp
--- a/Engine/Plugins/Experimental/AnimNext/Source/AnimNextEditor/Private/Scheduler/ScheduleAssetDefinition.cpp
+++ b/Engine/Plugins/Experimental/AnimNext/Source/AnimNextEditor/Private/Scheduler/ScheduleAssetDefinition.cpp
@@ -2,13 +2,26 @@
 
 #include "ScheduleAssetDefinition.h"
 #include "Toolkits/SimpleAssetEditor.h"
+#include "EditorCVars.h"
+#include "Workspace/AnimNextWorkspaceEditor.h"
 
 #define LOCTEXT_NAMESPACE "AnimNextAssetDefinitions"
 
 EAssetCommandResult UAssetDefinition_AnimNextSchedule::OpenAssets(const FAssetOpenArgs& OpenArgs) const
 {
-	TArray<UObject*> Objects = OpenArgs.LoadObjects<UObject>();
-	FSimpleAssetEditor::CreateEditor(EToolkitMode::Standalone, OpenArgs.ToolkitHost, Objects);
+	using namespace UE::AnimNext::Editor;
+	
+	for(UObject* Asset : OpenArgs.LoadObjects<UObject>())
+	{
+		if(CVars::GUseWorkspaceEditor.GetValueOnGameThread())
+		{
+			FWorkspaceEditor::OpenWorkspaceForAsset(Asset, FWorkspaceEditor::EOpenWorkspaceMethod::Default);
+		}
+		else
+		{
+			FSimpleAssetEditor::CreateEditor(EToolkitMode::Standalone, OpenArgs.ToolkitHost, Asset);
+		}
+	}
 
 	return EAssetCommandResult::Handled;
 }
diff --git a/Engine/Plugins/Experimental/AnimNext/Source/AnimNextEditor/Private/Workspace/AnimNextWorkspace.cpp b/Engine/Plugins/Experimental/AnimNext/Source/AnimNextEditor/Private/Workspace/AnimNextWorkspace.cpp
new file mode 100644
--- /dev/null
+++ b/Engine/Plugins/Experimental/AnimNext/Source/AnimNextEditor/Private/Workspace/AnimNextWorkspace.cpp
@@ -0,0 +1,233 @@
+﻿// Copyright Epic Games, Inc. All Rights Reserved.
+
+#include "AnimNextWorkspace.h"
+
+#include "Graph/AnimNextGraph.h"
+#include "Param/AnimNextParameterBlock.h"
+#include "Param/AnimNextParameterLibrary.h"
+#include "Scheduler/AnimNextSchedule.h"
+
+const FName UAnimNextWorkspace::ExportsAssetRegistryTag = TEXT("Exports");
+
+bool UAnimNextWorkspace::AddAsset(const FAssetData& InAsset, bool bSetupUndoRedo, bool bPrintPythonCommand)
+{
+	if(!IsAssetSupported(InAsset))
+	{
+		ReportError(TEXT("UAnimNextWorkspace::AddAsset: Unsupported asset supplied."));
+		return false;
+	}
+
+	if(bSetupUndoRedo)
+	{
+		Modify();
+	}
+
+	const int32 NewIndex = Assets.AddUnique(TSoftObjectPtr<UObject>(InAsset.GetSoftObjectPath()));
+	if(NewIndex != INDEX_NONE)
+	{
+		BroadcastModified();
+	}
+
+	return NewIndex != INDEX_NONE;
+}
+
+bool UAnimNextWorkspace::AddAsset(UObject* InAsset, bool bSetupUndoRedo, bool bPrintPythonCommand)
+{
+	if(InAsset == nullptr)
+	{
+		ReportError(TEXT("UAnimNextWorkspace::AddAsset: Invalid asset supplied."));
+		return false;
+	}
+
+	return AddAsset(FAssetData(InAsset), bSetupUndoRedo, bPrintPythonCommand);
+}
+
+bool UAnimNextWorkspace::AddAssets(TConstArrayView<FAssetData> InAssets, bool bSetupUndoRedo, bool bPrintPythonCommand)
+{
+	if(InAssets.Num() == 0)
+	{
+		ReportError(TEXT("UAnimNextWorkspace::AddAssets: No assets supplied."));
+		return false;
+	}
+
+	bool bAdded = false;
+	{
+		TGuardValue<bool> DisableNotifications(bSuspendNotifications, true);
+		for(const FAssetData& Asset : InAssets)
+		{
+			bAdded |= AddAsset(Asset, bSetupUndoRedo, bPrintPythonCommand);
+		}
+	}
+
+	if(bAdded)
+	{
+		BroadcastModified();
+	}
+
+	return bAdded;
+}
+
+bool UAnimNextWorkspace::AddAssets(const TArray<UObject*>& InAssets, bool bSetupUndoRedo, bool bPrintPythonCommand)
+{
+	if(InAssets.Num() == 0)
+	{
+		ReportError(TEXT("UAnimNextWorkspace::AddAssets: No assets supplied."));
+		return false;
+	}
+
+	bool bAdded = false;
+	{
+		TGuardValue<bool> DisableNotifications(bSuspendNotifications, true);
+		for(UObject* Asset : InAssets)
+		{
+			bAdded |= AddAsset(FAssetData(Asset), bSetupUndoRedo, bPrintPythonCommand);
+		}
+	}
+
+	if(bAdded)
+	{
+		BroadcastModified();
+	}
+
+	return bAdded;
+}
+
+bool UAnimNextWorkspace::RemoveAsset(const FAssetData& InAsset, bool bSetupUndoRedo, bool bPrintPythonCommand)
+{
+	if(!IsAssetSupported(InAsset))
+	{
+		ReportError(TEXT("UAnimNextWorkspace::RemoveAsset: Unsupported asset supplied."));
+		return false;
+	}
+
+	if(bSetupUndoRedo)
+	{
+		Modify();
+	}
+
+	const int32 NumRemoved = Assets.Remove(TSoftObjectPtr<UObject>(InAsset.GetSoftObjectPath()));
+	if(NumRemoved > 0)
+	{
+		BroadcastModified();
+	}
+
+	return NumRemoved > 0;
+}
+
+bool UAnimNextWorkspace::RemoveAsset(UObject* InAsset, bool bSetupUndoRedo, bool bPrintPythonCommand)
+{
+	if(InAsset == nullptr)
+	{
+		ReportError(TEXT("UAnimNextWorkspace::RemoveAsset: Invalid asset supplied."));
+		return false;
+	}
+
+	return RemoveAsset(FAssetData(InAsset), bSetupUndoRedo, bPrintPythonCommand);
+}
+
+bool UAnimNextWorkspace::RemoveAssets(TArray<UObject*> InAssets, bool bSetupUndoRedo, bool bPrintPythonCommand)
+{
+	if(InAssets.Num() == 0)
+	{
+		ReportError(TEXT("UAnimNextWorkspace::RemoveAssets: No assets supplied."));
+		return false;
+	}
+
+	bool bRemoved = false;
+	{
+		TGuardValue<bool> DisableNotifications(bSuspendNotifications, true);
+		for(UObject* Asset : InAssets)
+		{
+			bRemoved |= RemoveAsset(FAssetData(Asset), bSetupUndoRedo, bPrintPythonCommand);
+		}
+	}
+
+	if(bRemoved)
+	{
+		BroadcastModified();
+	}
+
+	return bRemoved;
+}
+
+bool UAnimNextWorkspace::RemoveAssets(TConstArrayView<FAssetData> InAssets, bool bSetupUndoRedo, bool bPrintPythonCommand)
+{
+	if(InAssets.Num() == 0)
+	{
+		ReportError(TEXT("UAnimNextWorkspace::RemoveAssets: No assets supplied."));
+		return false;
+	}
+
+	bool bRemoved = false;
+	{
+		TGuardValue<bool> DisableNotifications(bSuspendNotifications, true);
+		for(const FAssetData& Asset : InAssets)
+		{
+			bRemoved |= RemoveAsset(Asset, bSetupUndoRedo, bPrintPythonCommand);
+		}
+	}
+
+	if(bRemoved)
+	{
+		BroadcastModified();
+	}
+
+	return bRemoved;
+}
+
+const TArray<FTopLevelAssetPath>& UAnimNextWorkspace::GetSupportedAssetClassPaths()
+{
+	static TArray<FTopLevelAssetPath> SupportedAssets = {
+		UAnimNextParameterLibrary::StaticClass()->GetClassPathName(),
+		UAnimNextParameterBlock::StaticClass()->GetClassPathName(),
+		UAnimNextGraph::StaticClass()->GetClassPathName(),
+		UAnimNextSchedule::StaticClass()->GetClassPathName(),
+	};
+
+	return SupportedAssets;
+}
+
+bool UAnimNextWorkspace::IsAssetSupported(const FAssetData& InAsset)
+{
+	const TArray<FTopLevelAssetPath>& SupportedAssets = GetSupportedAssetClassPaths();
+	return SupportedAssets.Contains(InAsset.AssetClassPath);
+}
+
+void UAnimNextWorkspace::BroadcastModified()
+{
+	if(!bSuspendNotifications)
+	{
+		ModifiedDelegate.Broadcast(this);
+	}
+}
+
+void UAnimNextWorkspace::ReportError(const TCHAR* InMessage) const
+{
+#if WITH_EDITOR
+	FScriptExceptionHandler::Get().HandleException(ELogVerbosity::Error, InMessage, TEXT(""));
+#endif
+}
+
+void UAnimNextWorkspace::PostTransacted(const FTransactionObjectEvent& TransactionEvent)
+{
+	Super::PostTransacted(TransactionEvent);
+	BroadcastModified();
+}
+
+void UAnimNextWorkspace::GetAssetRegistryTags(TArray<FAssetRegistryTag>& OutTags) const
+{
+	Super::GetAssetRegistryTags(OutTags);
+
+	FAnimNextWorkspaceAssetRegistryExports Exports;
+	Exports.Assets.Reserve(Assets.Num());
+
+	for(const TSoftObjectPtr<UObject>& Asset : Assets)
+	{
+		Exports.Assets.Emplace(Asset.GetUniqueID());
+	}
+
+	FString TagValue;
+	FAnimNextWorkspaceAssetRegistryExports::StaticStruct()->ExportText(TagValue, &Exports, nullptr, nullptr, PPF_None, nullptr);
+
+	OutTags.Add(FAssetRegistryTag(ExportsAssetRegistryTag, TagValue, FAssetRegistryTag::TT_Hidden));
+}
diff --git a/Engine/Plugins/Experimental/AnimNext/Source/AnimNextEditor/Private/Workspace/AnimNextWorkspace.h b/Engine/Plugins/Experimental/AnimNext/Source/AnimNextEditor/Private/Workspace/AnimNextWorkspace.h
new file mode 100644
--- /dev/null
+++ b/Engine/Plugins/Experimental/AnimNext/Source/AnimNextEditor/Private/Workspace/AnimNextWorkspace.h
@@ -0,0 +1,110 @@
+﻿// Copyright Epic Games, Inc. All Rights Reserved.
+
+#pragma once
+
+#include "CoreMinimal.h"
+#include "AnimNextWorkspace.generated.h"
+
+struct FEditedDocumentInfo;
+class UAnimNextWorkspace;
+
+namespace UE::AnimNext::Editor
+{
+	class FWorkspaceEditor;
+	class SWorkspaceView;
+	struct FUtils;
+}
+
+namespace UE::AnimNext::Editor
+{
+	// A delegate for subscribing / reacting to workspace modifications.
+	DECLARE_MULTICAST_DELEGATE_OneParam(FOnWorkspaceModified, UAnimNextWorkspace* /* InWorkspace */);
+}
+
+// Workspace entry used to export to asset registry
+USTRUCT()
+struct FAnimNextWorkspaceAssetRegistryExportEntry
+{
+	GENERATED_BODY()
+
+	FAnimNextWorkspaceAssetRegistryExportEntry() = default;
+	
+	FAnimNextWorkspaceAssetRegistryExportEntry(const FSoftObjectPath& InAsset)
+		: Asset(InAsset)
+	{}
+
+	UPROPERTY()
+	FSoftObjectPath Asset;
+};
+
+// Workspace used to export to asset registry
+USTRUCT()
+struct FAnimNextWorkspaceAssetRegistryExports
+{
+	GENERATED_BODY()
+
+	UPROPERTY()
+	TArray<FAnimNextWorkspaceAssetRegistryExportEntry> Assets;
+};
+
+UCLASS()
+class UAnimNextWorkspace : public UObject
+{
+	GENERATED_BODY()
+
+	friend class UE::AnimNext::Editor::FWorkspaceEditor;
+	friend class UE::AnimNext::Editor::SWorkspaceView;
+	friend struct UE::AnimNext::Editor::FUtils;
+
+	ANIMNEXTEDITOR_API static const FName ExportsAssetRegistryTag;
+
+	// Adds an asset to the workspace
+	// @return true if the asset was added
+	UFUNCTION(BlueprintCallable, Category = "AnimNext|Workspace")
+	bool AddAsset(UObject* InAsset, bool bSetupUndoRedo = true, bool bPrintPythonCommand = true);
+	bool AddAsset(const FAssetData& InAsset, bool bSetupUndoRedo = true, bool bPrintPythonCommand = true);
+
+	// Adds assets to the workspace
+	// @return true if an asset was added
+	UFUNCTION(BlueprintCallable, Category = "AnimNext|Workspace")
+	bool AddAssets(const TArray<UObject*>& InAssets, bool bSetupUndoRedo = true, bool bPrintPythonCommand = true);
+	bool AddAssets(TConstArrayView<FAssetData> InAssets, bool bSetupUndoRedo = true, bool bPrintPythonCommand = true);
+
+	// Removes an asset from the workspace
+	// @return true if the asset was removed
+	UFUNCTION(BlueprintCallable, Category = "AnimNext|Workspace")
+	bool RemoveAsset(UObject* InAsset, bool bSetupUndoRedo = true, bool bPrintPythonCommand = true);
+	bool RemoveAsset(const FAssetData& InAsset, bool bSetupUndoRedo = true, bool bPrintPythonCommand = true);
+
+	// Removes assets from the workspace
+	// @return true if the asset was removed
+	UFUNCTION(BlueprintCallable, Category = "AnimNext|Workspace")
+	bool RemoveAssets(TArray<UObject*> InAssets, bool bSetupUndoRedo = true, bool bPrintPythonCommand = true);
+	bool RemoveAssets(TConstArrayView<FAssetData> InAssets, bool bSetupUndoRedo = true, bool bPrintPythonCommand = true);
+
+	static const TArray<FTopLevelAssetPath>& GetSupportedAssetClassPaths();
+
+	static bool IsAssetSupported(const FAssetData& InAsset);
+
+	void ReportError(const TCHAR* InMessage) const;
+
+	void BroadcastModified();
+
+	// UObject interface
+	virtual void GetAssetRegistryTags(TArray<FAssetRegistryTag>& OutTags) const override;
+	virtual void PostTransacted(const FTransactionObjectEvent& TransactionEvent) override;
+	virtual bool IsEditorOnly() const override { return true; }
+
+	// All of the assets referenced by this workspace
+	UPROPERTY()
+	TArray<TSoftObjectPtr<UObject>> Assets;
+
+	// Documents this workspace was editing
+	UPROPERTY()
+	TArray<FEditedDocumentInfo> LastEditedDocuments;
+
+	// Delegate to subscribe to modifications
+	UE::AnimNext::Editor::FOnWorkspaceModified ModifiedDelegate;
+
+	bool bSuspendNotifications = false;
+};
\ No newline at end of file
diff --git a/Engine/Plugins/Experimental/AnimNext/Source/AnimNextEditor/Private/Workspace/AnimNextWorkspaceAssetDefinition.cpp b/Engine/Plugins/Experimental/AnimNext/Source/AnimNextEditor/Private/Workspace/AnimNextWorkspaceAssetDefinition.cpp
new file mode 100644
--- /dev/null
+++ b/Engine/Plugins/Experimental/AnimNext/Source/AnimNextEditor/Private/Workspace/AnimNextWorkspaceAssetDefinition.cpp
@@ -0,0 +1,15 @@
+﻿// Copyright Epic Games, Inc. All Rights Reserved.
+
+#include "AnimNextWorkspaceAssetDefinition.h"
+#include "AnimNextWorkspaceEditor.h"
+
+EAssetCommandResult UAssetDefinition_AnimNextWorkspace::OpenAssets(const FAssetOpenArgs& OpenArgs) const
+{
+	for (UAnimNextWorkspace* Asset : OpenArgs.LoadObjects<UAnimNextWorkspace>())
+	{
+		TSharedRef<UE::AnimNext::Editor::FWorkspaceEditor> Editor = MakeShared<UE::AnimNext::Editor::FWorkspaceEditor>();
+		Editor->InitEditor(OpenArgs.GetToolkitMode(), OpenArgs.ToolkitHost, Asset);
+	}
+
+	return EAssetCommandResult::Handled;
+}
diff --git a/Engine/Plugins/Experimental/AnimNext/Source/AnimNextEditor/Private/Workspace/AnimNextWorkspaceAssetDefinition.h b/Engine/Plugins/Experimental/AnimNext/Source/AnimNextEditor/Private/Workspace/AnimNextWorkspaceAssetDefinition.h
new file mode 100644
--- /dev/null
+++ b/Engine/Plugins/Experimental/AnimNext/Source/AnimNextEditor/Private/Workspace/AnimNextWorkspaceAssetDefinition.h
@@ -0,0 +1,29 @@
+﻿// Copyright Epic Games, Inc. All Rights Reserved.
+
+#pragma once
+
+#include "AnimNextWorkspace.h"
+#include "AssetDefinitionDefault.h"
+#include "AnimNextWorkspaceAssetDefinition.generated.h"
+
+#define LOCTEXT_NAMESPACE "AnimNextAssetDefinitions"
+
+UCLASS()
+class UAssetDefinition_AnimNextWorkspace : public UAssetDefinitionDefault
+{
+	GENERATED_BODY()
+
+public:
+	// UAssetDefinition interface
+	virtual FText GetAssetDisplayName() const override { return LOCTEXT("AnimNextWorkspace", "AnimNext Workspace"); }
+	virtual FLinearColor GetAssetColor() const override { return FLinearColor(FColor(128,64,64)); }
+	virtual TSoftClassPtr<UObject> GetAssetClass() const override { return UAnimNextWorkspace::StaticClass(); }
+	virtual EAssetCommandResult OpenAssets(const FAssetOpenArgs& OpenArgs) const override;
+	virtual TConstArrayView<FAssetCategoryPath> GetAssetCategories() const override
+	{
+		static const auto Categories = { FAssetCategoryPath(EAssetCategoryPaths::Animation, LOCTEXT("AnimNextSubMenu", "AnimNext")) };
+		return Categories;
+	}
+};
+
+#undef LOCTEXT_NAMESPACE
\ No newline at end of file
diff --git a/Engine/Plugins/Experimental/AnimNext/Source/AnimNextEditor/Private/Workspace/AnimNextWorkspaceEditor.cpp b/Engine/Plugins/Experimental/AnimNext/Source/AnimNextEditor/Private/Workspace/AnimNextWorkspaceEditor.cpp
new file mode 100644
--- /dev/null
+++ b/Engine/Plugins/Experimental/AnimNext/Source/AnimNextEditor/Private/Workspace/AnimNextWorkspaceEditor.cpp
@@ -0,0 +1,542 @@
+﻿// Copyright Epic Games, Inc. All Rights Reserved.
+
+#include "AnimNextWorkspaceEditor.h"
+
+#include "GraphDocumentSummoner.h"
+#include "AnimNextWorkspaceEditorMode.h"
+#include "AnimNextWorkspace.h"
+#include "AnimNextWorkspaceFactory.h"
+#include "AssetDocumentSummoner.h"
+#include "EdGraphNode_Comment.h"
+#include "EditorUtils.h"
+#include "ExternalPackageHelper.h"
+#include "Framework/Commands/GenericCommands.h"
+#include "RigVMModel/RigVMController.h"
+#include "Kismet2/BlueprintEditorUtils.h"
+#include "ScopedTransaction.h"
+#include "Widgets/Docking/SDockTab.h"
+#include "WorkflowOrientedApp/WorkflowUObjectDocuments.h"
+#include "SWorkspacePicker.h"
+#include "Dialog/SCustomDialog.h"
+#include "Graph/AnimNextGraphDocumentSummoner.h"
+#include "Param/AnimNextParameterBlock.h"
+#include "Param/AnimNextParameterBlock_EditorData.h"
+#include "Param/AnimNextParameterLibrary.h"
+#include "Param/ParameterBlockGraphDocumentSummoner.h"
+#include "Scheduler/AnimNextSchedule.h"
+#include "Framework/Application/SlateApplication.h"
+#include "Graph/AnimNextGraph.h"
+
+#define LOCTEXT_NAMESPACE "AnimNextWorkspaceEditor"
+
+namespace UE::AnimNext::Editor
+{
+
+namespace WorkspaceModes
+{
+	const FName WorkspaceEditor("AnimNextWorkspaceEditorMode");
+}
+
+namespace WorkspaceTabs
+{
+	const FName Details("DetailsTab");
+	const FName WorkspaceView("WorkspaceView");
+	const FName LeftAssetDocument("LeftAssetDocument");
+	const FName MiddleAssetDocument("MiddleAssetDocument");
+	const FName AnimNextGraphDocument("AnimNextGraphDocument");
+	const FName ParameterBlockGraphDocument("ParameterBlockGraphDocument");
+}
+
+const FName WorkspaceAppIdentifier("AnimNextWorkspaceEditor");
+
+TMap<FName, FWorkspaceEditor::FAssetDocumentWidgetFactoryFunc> FWorkspaceEditor::AssetDocumentWidgetFactories;
+
+FWorkspaceEditor::FWorkspaceEditor()
+{
+}
+
+FWorkspaceEditor::~FWorkspaceEditor()
+{
+}
+
+void FWorkspaceEditor::InitEditor(const EToolkitMode::Type InMode, const TSharedPtr<IToolkitHost>& InInitToolkitHost, UAnimNextWorkspace* InWorkspace)
+{
+	Workspace = InWorkspace;
+
+	DocumentManager = MakeShared<FDocumentTracker>(NAME_None);
+	DocumentManager->Initialize(SharedThis(this));
+
+	TSharedRef<FParameterBlockGraphDocumentSummoner> ParameterGraphDocumentSummoner = MakeShared<FParameterBlockGraphDocumentSummoner>(WorkspaceTabs::ParameterBlockGraphDocument, SharedThis(this));
+	ParameterGraphDocumentSummoner->OnSaveGraphState().BindSP(this, &FWorkspaceEditor::HandleSaveGraphState);
+	DocumentManager->RegisterDocumentFactory(ParameterGraphDocumentSummoner);
+
+	TSharedRef<FAnimNextGraphDocumentSummoner> GraphDocumentSummoner = MakeShared<FAnimNextGraphDocumentSummoner>(WorkspaceTabs::AnimNextGraphDocument, SharedThis(this));
+	GraphDocumentSummoner->OnSaveGraphState().BindSP(this, &FWorkspaceEditor::HandleSaveGraphState);
+	DocumentManager->RegisterDocumentFactory(GraphDocumentSummoner);
+
+	TSharedRef<FAssetDocumentSummoner> LeftAssetDocumentSummoner = MakeShared<FAssetDocumentSummoner>(WorkspaceTabs::LeftAssetDocument, SharedThis(this));
+	LeftAssetDocumentSummoner->SetAllowedAssetClassPaths(
+	{
+		UAnimNextGraph::StaticClass()->GetClassPathName(),
+		UAnimNextParameterBlock::StaticClass()->GetClassPathName(),
+		UAnimNextParameterLibrary::StaticClass()->GetClassPathName(),
+	});
+	LeftAssetDocumentSummoner->OnSaveDocumentState().BindSP(this, &FWorkspaceEditor::HandleSaveDocumentState);
+	DocumentManager->RegisterDocumentFactory(LeftAssetDocumentSummoner);
+
+	TSharedRef<FAssetDocumentSummoner> MiddleAssetDocumentSummoner = MakeShared<FAssetDocumentSummoner>(WorkspaceTabs::MiddleAssetDocument, SharedThis(this));
+	MiddleAssetDocumentSummoner->SetAllowedAssetClassPaths(
+	{
+		UAnimNextSchedule::StaticClass()->GetClassPathName(),
+	});
+	MiddleAssetDocumentSummoner->OnSaveDocumentState().BindSP(this, &FWorkspaceEditor::HandleSaveDocumentState);
+	DocumentManager->RegisterDocumentFactory(MiddleAssetDocumentSummoner);
+
+	constexpr bool bCreateDefaultStandaloneMenu = true;
+	constexpr bool bCreateDefaultToolbar = true;
+	InitAssetEditor(InMode, InInitToolkitHost, WorkspaceAppIdentifier, FTabManager::FLayout::NullLayout, bCreateDefaultStandaloneMenu, bCreateDefaultToolbar, InWorkspace);
+
+	BindCommands();
+
+	AddApplicationMode(WorkspaceModes::WorkspaceEditor, MakeShared<FWorkspaceEditorMode>(SharedThis(this)));
+	SetCurrentMode(WorkspaceModes::WorkspaceEditor);
+
+	ExtendMenu();
+	ExtendToolbar();
+	RegenerateMenusAndToolbars();
+}
+
+void FWorkspaceEditor::RegisterAssetDocumentWidget(FName InAssetClassName, FAssetDocumentWidgetFactoryFunc&& InFunction)
+{
+	AssetDocumentWidgetFactories.Add(InAssetClassName, MoveTemp(InFunction));
+}
+
+void FWorkspaceEditor::UnregisterAssetDocumentWidget(FName InAssetClassName)
+{
+	AssetDocumentWidgetFactories.Remove(InAssetClassName);
+}
+
+void FWorkspaceEditor::OpenWorkspaceForAsset(UObject* InAsset, EOpenWorkspaceMethod InOpenMethod)
+{
+	FAssetRegistryModule& AssetRegistryModule = FModuleManager::LoadModuleChecked<FAssetRegistryModule>(TEXT("AssetRegistry"));
+
+	TArray<FAssetData> RelevantWorkspaceAssets;
+
+	if(InOpenMethod != EOpenWorkspaceMethod::AlwaysOpenNewWorkspace)
+	{
+		// Look for existing workspaces that export this asset
+		FARFilter ARFilter;
+		ARFilter.ClassPaths.Add(UAnimNextWorkspace::StaticClass()->GetClassPathName());
+		ARFilter.bRecursiveClasses = true;
+
+		TArray<FAssetData> AllWorkspaceAssets;
+		AssetRegistryModule.Get().GetAssets(ARFilter, AllWorkspaceAssets);
+
+		for(const FAssetData& WorkspaceAsset : AllWorkspaceAssets)
+		{
+			FAnimNextWorkspaceAssetRegistryExports Exports;
+			FUtils::GetExportedAssetsForWorkspace(WorkspaceAsset, Exports);
+
+			FSoftObjectPath ObjectPath(InAsset);
+			for(const FAnimNextWorkspaceAssetRegistryExportEntry& ExportEntry : Exports.Assets)
+			{
+				if(ExportEntry.Asset == ObjectPath)
+				{
+					RelevantWorkspaceAssets.Add(WorkspaceAsset);
+					break;
+				}
+			}
+		}
+	}
+
+	auto HandleNewWorkspace = [InAsset]()
+	{
+		UAnimNextWorkspaceFactory* Factory = NewObject<UAnimNextWorkspaceFactory>();
+		UPackage* Package = CreatePackage(nullptr);
+		FName PackageName = *FPaths::GetBaseFilename(Package->GetName());
+		UAnimNextWorkspace* NewWorkspace = CastChecked<UAnimNextWorkspace>(Factory->FactoryCreateNew(UAnimNextWorkspace::StaticClass(), Package, PackageName, RF_Public | RF_Standalone, NULL, GWarn));
+		NewWorkspace->AddAsset(InAsset, false);
+		NewWorkspace->MarkPackageDirty();
+		TSharedRef<FWorkspaceEditor> Editor = MakeShared<FWorkspaceEditor>();
+		Editor->InitEditor(EToolkitMode::Standalone, nullptr, NewWorkspace);
+	};
+
+	auto HandleExistingWorkspace = [](const FAssetData& InAssetData)
+	{
+		if(UAnimNextWorkspace* ExistingWorkspace = Cast<UAnimNextWorkspace>(InAssetData.GetAsset()))
+		{
+			GEditor->GetEditorSubsystem<UAssetEditorSubsystem>()->OpenEditorForAsset(ExistingWorkspace);
+		}
+	};
+	
+	if(InOpenMethod == EOpenWorkspaceMethod::AlwaysOpenNewWorkspace || RelevantWorkspaceAssets.Num() == 0)
+	{
+		// No relevant workspaces, so open a new one and add the asset
+		HandleNewWorkspace();
+	}
+	else if(RelevantWorkspaceAssets.Num() == 1)
+	{
+		// One existing workspace, open it
+		HandleExistingWorkspace(RelevantWorkspaceAssets[0]);
+	}
+	else
+	{
+		// Multiple existing workspaces, present a window to let the user choose one to open with
+		TSharedRef<SWorkspacePicker> WorkspacePicker = SNew(SWorkspacePicker)
+			.WorkspaceAssets(RelevantWorkspaceAssets)
+			.OnAssetSelected_Lambda(HandleExistingWorkspace)
+			.OnNewAsset_Lambda(HandleNewWorkspace);
+
+		WorkspacePicker->ShowModal();
+	}
+}
+
+void FWorkspaceEditor::RestoreEditedObjectState()
+{
+	for (const FEditedDocumentInfo& Document : Workspace->LastEditedDocuments)
+	{
+		if (UObject* Obj = Document.EditedObjectPath.ResolveObject())
+		{
+			if(TSharedPtr<SDockTab> DockTab = OpenDocument(Obj, FDocumentTracker::RestorePreviousDocument))
+			{
+				if(Obj->IsA<UEdGraph>())
+				{
+					TSharedRef<SGraphEditor> GraphEditor = StaticCastSharedRef<SGraphEditor>(DockTab->GetContent());
+					GraphEditor->SetViewLocation(Document.SavedViewOffset, Document.SavedZoomAmount);
+				}
+			}
+		}
+	}
+}
+
+void FWorkspaceEditor::SaveEditedObjectState()
+{
+	// Clear currently edited documents
+	Workspace->LastEditedDocuments.Empty();
+
+	// Ask all open documents to save their state, which will update LastEditedGraphDocuments
+	DocumentManager->SaveAllState();
+}
+
+TSharedPtr<SDockTab> FWorkspaceEditor::OpenDocument(const UObject* InForObject, FDocumentTracker::EOpenDocumentCause InCause)
+{
+	if(InCause != FDocumentTracker::RestorePreviousDocument)
+	{
+		Workspace->LastEditedDocuments.AddUnique(const_cast<UObject*>(InForObject));
+	}
+
+	TSharedRef<FTabPayload_UObject> Payload = FTabPayload_UObject::Make(InForObject);
+	return DocumentManager->OpenDocument(Payload, InCause);
+}
+
+void FWorkspaceEditor::OpenAssets(TConstArrayView<FAssetData> InAssets)
+{
+	for(const FAssetData& Asset : InAssets)
+	{
+		if(UObject* LoadedAsset = Asset.GetAsset())
+		{
+			OpenDocument(LoadedAsset, FDocumentTracker::EOpenDocumentCause::OpenNewDocument);
+		}
+	}
+}
+
+void FWorkspaceEditor::BindCommands()
+{
+}
+
+void FWorkspaceEditor::ExtendMenu()
+{
+	
+}
+
+void FWorkspaceEditor::ExtendToolbar()
+{
+	
+}
+
+void FWorkspaceEditor::RegisterTabSpawners(const TSharedRef<FTabManager>& InTabManager)
+{
+	DocumentManager->SetTabManager(InTabManager);
+
+	FWorkflowCentricApplication::RegisterTabSpawners(InTabManager);
+}
+
+void FWorkspaceEditor::UnregisterTabSpawners(const TSharedRef<FTabManager>& InTabManager)
+{
+	FWorkflowCentricApplication::UnregisterTabSpawners(InTabManager);
+}
+
+FName FWorkspaceEditor::GetToolkitFName() const
+{
+	return FName("AnimNextWorkspaceEditor");
+}
+
+FText FWorkspaceEditor::GetBaseToolkitName() const
+{
+	return LOCTEXT("AppLabel", "AnimNextWorkspaceEditor");
+}
+
+FString FWorkspaceEditor::GetWorldCentricTabPrefix() const
+{
+	return LOCTEXT("WorldCentricTabPrefix", "AnimNextWorkspaceEditor ").ToString();
+}
+
+FLinearColor FWorkspaceEditor::GetWorldCentricTabColorScale() const
+{
+	return FLinearColor(0.3f, 0.2f, 0.5f, 0.5f);
+}
+
+void FWorkspaceEditor::InitToolMenuContext(FToolMenuContext& InMenuContext)
+{
+}
+
+void FWorkspaceEditor::SaveAsset_Execute()
+{
+	// If asset is a default 'Untitled' workspace, redirect to the 'save as' flow
+	FString AssetPath = Workspace->GetOutermost()->GetPathName();
+	if(AssetPath.StartsWith(TEXT("/Temp/Untitled")))
+	{
+		// Ensure we dont also 'save as' other externally linked assets at this point
+		TGuardValue<bool> SaveWorkspaceOnly(bSavingWorkspaceOnly, true);
+
+		SaveAssetAs_Execute();
+	}
+	else
+	{
+		FWorkflowCentricApplication::SaveAsset_Execute();
+	}
+}
+
+void FWorkspaceEditor::SetFocusedGraphEditor(TSharedPtr<SGraphEditor> InGraphEditor)
+{
+	// Update the graph editor that is currently focused
+	FocusedGraphEdPtr = InGraphEditor;
+}
+
+UEdGraph* FWorkspaceEditor::GetFocusedGraph() const
+{
+	if (FocusedGraphEdPtr.IsValid())
+	{
+		if (UEdGraph* Graph = FocusedGraphEdPtr.Pin()->GetCurrentGraph())
+		{
+			return Graph;
+		}
+	}
+	return nullptr;
+}
+
+URigVMGraph* FWorkspaceEditor::GetFocusedVMGraph() const
+{
+	if(UEdGraph* EdGraph = GetFocusedGraph())
+	{
+		if(IRigVMClientHost* RigVMClientHost = EdGraph->GetImplementingOuter<IRigVMClientHost>())
+		{
+			return Cast<URigVMGraph>(RigVMClientHost->GetRigVMGraphForEditorObject(EdGraph));
+		}
+	}
+	return nullptr;
+}
+
+URigVMController* FWorkspaceEditor::GetFocusedVMController() const
+{
+	if(UEdGraph* EdGraph = GetFocusedGraph())
+	{
+		if(IRigVMClientHost* RigVMClientHost = EdGraph->GetImplementingOuter<IRigVMClientHost>())
+		{
+			return RigVMClientHost->GetRigVMClient()->GetController(EdGraph);
+		}
+	}
+	return nullptr;
+}
+
+void FWorkspaceEditor::OnNodeTitleCommitted(const FText& NewText, ETextCommit::Type CommitInfo, UEdGraphNode* NodeBeingChanged)
+{
+	if (UEdGraphNode_Comment* CommentBeingChanged = Cast<UEdGraphNode_Comment>(NodeBeingChanged))
+	{
+		GetFocusedVMController()->SetCommentTextByName(CommentBeingChanged->GetFName(), NewText.ToString(), CommentBeingChanged->FontSize, CommentBeingChanged->bCommentBubbleVisible, CommentBeingChanged->bColorCommentBubble, true, true);
+	}
+}
+
+void FWorkspaceEditor::CloseDocumentTab(const UObject* DocumentID)
+{
+	Workspace->LastEditedDocuments.Remove(const_cast<UObject*>(DocumentID));
+
+	TSharedRef<FTabPayload_UObject> Payload = FTabPayload_UObject::Make(DocumentID);
+	DocumentManager->CloseTab(Payload);
+}
+
+bool FWorkspaceEditor::InEditingMode() const
+{
+	// @TODO: disallow editing when debugging when implemented
+
+	return true;
+}
+
+bool FWorkspaceEditor::IsEditable(UEdGraph* InGraph) const
+{
+	return InGraph && InEditingMode() && InGraph->bEditable;
+}
+
+FGraphPanelSelectionSet FWorkspaceEditor::GetSelectedNodes() const
+{
+	FGraphPanelSelectionSet CurrentSelection;
+	TSharedPtr<SGraphEditor> FocusedGraphEd = FocusedGraphEdPtr.Pin();
+	if (FocusedGraphEd.IsValid())
+	{
+		CurrentSelection = FocusedGraphEd->GetSelectedNodes();
+	}
+	return CurrentSelection;
+}
+
+void FWorkspaceEditor::DeleteSelectedNodes()
+{
+	TSharedPtr<SGraphEditor> FocusedGraphEd = FocusedGraphEdPtr.Pin();
+	if (!FocusedGraphEd.IsValid())
+	{
+		return;
+	}
+
+	const FScopedTransaction Transaction(FGenericCommands::Get().Delete->GetDescription());
+	FocusedGraphEd->GetCurrentGraph()->Modify();
+	
+	const FGraphPanelSelectionSet SelectedNodes = GetSelectedNodes();
+	
+	if(FocusedGraphEd)
+	{
+		FocusedGraphEd->ClearSelectionSet();
+	}
+
+	// Some nodes have sub-objects that are represented as other tabs.
+	// Close them here as a pre-pass before we remove their nodes. If the documents are left open they
+	// may reference dangling data and function incorrectly in cases such as FindBlueprintforNodeChecked
+	for (FGraphPanelSelectionSet::TConstIterator NodeIt( SelectedNodes ); NodeIt; ++NodeIt)
+	{
+		if (UEdGraphNode* Node = Cast<UEdGraphNode>(*NodeIt))
+		{
+			if (Node->CanUserDeleteNode())
+			{
+				auto CloseAllDocumentsTab = [this](const UEdGraphNode* InNode)
+				{
+					TArray<UObject*> NodesToClose;
+					GetObjectsWithOuter(InNode, NodesToClose);
+					for (UObject* Node : NodesToClose)
+					{
+						UEdGraph* NodeGraph = Cast<UEdGraph>(Node);
+						if (NodeGraph)
+						{
+							CloseDocumentTab(NodeGraph);
+						}
+					}
+				};
+				
+				if (Node->GetSubGraphs().Num() > 0)
+				{
+					CloseAllDocumentsTab(Node);
+				}
+			}
+		}
+	}
+
+	// Now remove the selected nodes
+	for (FGraphPanelSelectionSet::TConstIterator NodeIt( SelectedNodes ); NodeIt; ++NodeIt)
+	{
+		if (UEdGraphNode* Node = Cast<UEdGraphNode>(*NodeIt))
+		{
+			if (Node->CanUserDeleteNode())
+			{
+				if (Node->GetSubGraphs().Num() > 0)
+				{
+					DocumentManager->CleanInvalidTabs();
+				}
+
+				FBlueprintEditorUtils::RemoveNode(nullptr, Node);
+			}
+		}
+	}
+}
+
+bool FWorkspaceEditor::CanDeleteSelectedNodes()
+{
+	const FGraphPanelSelectionSet SelectedNodes = GetSelectedNodes();
+
+	bool bCanUserDeleteNode = false;
+
+	if(IsEditable(GetFocusedGraph()) && SelectedNodes.Num() > 0)
+	{
+		for(UObject* NodeObject : SelectedNodes)
+		{
+			// If any nodes allow deleting, then do not disable the delete option
+			UEdGraphNode* Node = Cast<UEdGraphNode>(NodeObject);
+			if(Node->CanUserDeleteNode())
+			{
+				bCanUserDeleteNode = true;
+				break;
+			}
+		}
+	}
+
+	return bCanUserDeleteNode;
+}
+
+void FWorkspaceEditor::SetSelectedObjects(TArray<UObject*> InObjects)
+{
+	if(DetailsView.IsValid())
+	{
+		DetailsView->SetObjects(InObjects);
+	}
+}
+
+void FWorkspaceEditor::GetSaveableObjects(TArray<UObject*>& OutObjects) const
+{
+	// Base class will pick up edited object
+	FWorkflowCentricApplication::GetSaveableObjects(OutObjects);
+
+	if(!bSavingWorkspaceOnly)
+	{
+		for(TSoftObjectPtr<UObject>& SoftAsset : Workspace->Assets)
+		{
+			if(UObject* Asset = SoftAsset.Get())
+			{
+				// Add object referenced by workspace
+				OutObjects.Add(Asset);
+
+				// Get external objects too
+				FExternalPackageHelper::GetExternalSaveableObjects(Asset, OutObjects);
+			}
+		}
+	}
+}
+
+void FWorkspaceEditor::HandleSaveGraphState(UEdGraph* InGraph, FVector2D InViewOffset, float InZoomAmount)
+{
+	Workspace->LastEditedDocuments.AddUnique(FEditedDocumentInfo(InGraph, InViewOffset, InZoomAmount));
+}
+
+void FWorkspaceEditor::HandleSaveDocumentState(UObject* InObject)
+{
+	Workspace->LastEditedDocuments.AddUnique(FEditedDocumentInfo(InObject));
+}
+
+bool FWorkspaceEditor::OnRequestClose(EAssetEditorCloseReason InCloseReason)
+{
+	auto RequiresSave = [this]()
+	{
+		UPackage* Package = Workspace->GetOutermost();
+		return Package->GetPathName().StartsWith(TEXT("/Temp/Untitled"));
+	};
+
+	// Give the user opportunity to save temp workspaces
+	if(RequiresSave())
+	{
+		// Ensure we dont also 'save as' other externally linked assets at this point
+		TGuardValue<bool> SaveWorkspaceOnly(bSavingWorkspaceOnly, true);
+
+		SaveAssetAs_Execute();
+	}
+
+	return true;
+}
+
+}
+
+#undef LOCTEXT_NAMESPACE
\ No newline at end of file
diff --git a/Engine/Plugins/Experimental/AnimNext/Source/AnimNextEditor/Private/Workspace/AnimNextWorkspaceEditor.h b/Engine/Plugins/Experimental/AnimNext/Source/AnimNextEditor/Private/Workspace/AnimNextWorkspaceEditor.h
new file mode 100644
--- /dev/null
+++ b/Engine/Plugins/Experimental/AnimNext/Source/AnimNextEditor/Private/Workspace/AnimNextWorkspaceEditor.h
@@ -0,0 +1,165 @@
+﻿// Copyright Epic Games, Inc. All Rights Reserved.
+
+#pragma once
+
+#include "CoreMinimal.h"
+#include "WorkflowOrientedApp/WorkflowCentricApplication.h"
+#include "GraphEditor.h"
+#include "WorkflowOrientedApp/WorkflowTabManager.h"
+
+enum class ERigVMGraphNotifType : uint8;
+class UAnimNextWorkspace;
+class FDocumentTracker;
+class FDocumentTabFactory;
+class FTabInfo;
+class FTabManager;
+class IToolkitHost;
+class URigVMGraph;
+class URigVMController;
+class UDetailsViewWrapperObject;
+
+namespace UE::AnimNext::Editor
+{
+
+class FWorkspaceEditorMode;
+struct FGraphDocumentSummoner;
+struct FWorkspaceTabSummoner;
+struct FAssetDocumentSummoner;
+class FModule;
+
+namespace WorkspaceModes
+{
+	extern const FName WorkspaceEditor;
+}
+
+namespace WorkspaceTabs
+{
+	extern const FName Details;
+	extern const FName WorkspaceView;
+	extern const FName LeftAssetDocument;
+	extern const FName MiddleAssetDocument;
+	extern const FName AnimNextGraphDocument;
+	extern const FName ParameterBlockGraphDocument;
+}
+
+class FWorkspaceEditor : public FWorkflowCentricApplication
+{
+public:
+	FWorkspaceEditor();
+	virtual ~FWorkspaceEditor() override;
+
+	/** Edits the specified asset */
+	void InitEditor(const EToolkitMode::Type InMode, const TSharedPtr<IToolkitHost>& InInitToolkitHost, UAnimNextWorkspace* InAnimNextWorkspace);
+
+	using FAssetDocumentWidgetFactoryFunc = TUniqueFunction<TSharedRef<SWidget>(TSharedRef<FWorkspaceEditor>, UObject*)>;
+
+	static void RegisterAssetDocumentWidget(FName InAssetClassName, FAssetDocumentWidgetFactoryFunc&&  InFunction);
+	static void UnregisterAssetDocumentWidget(FName InAssetClassName);
+
+	enum class EOpenWorkspaceMethod : int32
+	{
+		// If the asset is already used in a workspace, open that (if not already opened)
+		// If the asset is already used in more than one workspace, let the user choose the workspace to open it in
+		// If the asset is not yet in a workspace, create a default workspace, add the asset and open the workspace
+		Default,
+
+		// Always open a new workspace asset and add the asset to it
+		AlwaysOpenNewWorkspace,
+	};
+
+	// Open an asset inside the workspace editor.
+	static void OpenWorkspaceForAsset(UObject* InAsset, EOpenWorkspaceMethod InOpenMethod);
+
+private:
+	friend class FWorkspaceEditorMode;
+	friend struct FGraphDocumentSummoner;
+	friend struct FWorkspaceTabSummoner;
+	friend struct FAssetDocumentSummoner;
+	friend class FModule;
+
+	// IToolkit interface
+	virtual void RegisterTabSpawners(const TSharedRef<FTabManager>& InTabManager) override;
+	virtual void UnregisterTabSpawners(const TSharedRef<FTabManager>& InTabManager) override;
+	virtual FName GetToolkitFName() const override;
+	virtual FText GetBaseToolkitName() const override;
+	virtual FString GetWorldCentricTabPrefix() const override;
+	virtual FLinearColor GetWorldCentricTabColorScale() const override;
+	virtual void InitToolMenuContext(FToolMenuContext& InMenuContext) override;
+	virtual void SaveAsset_Execute() override;
+	virtual bool OnRequestClose(EAssetEditorCloseReason InCloseReason) override;
+
+	// FAssetEditorToolkit interface
+	virtual void GetSaveableObjects(TArray<UObject*>& OutObjects) const override;
+
+	void BindCommands();
+	
+	void ExtendMenu();
+
+	void ExtendToolbar();
+
+	void SetFocusedGraphEditor(TSharedPtr<SGraphEditor> InGraphEditor);
+
+	UEdGraph* GetFocusedGraph() const;
+
+	URigVMGraph* GetFocusedVMGraph() const;
+
+	URigVMController* GetFocusedVMController() const;
+
+	// Handle node titles being changed
+	void OnNodeTitleCommitted(const FText& NewText, ETextCommit::Type CommitInfo, UEdGraphNode* NodeBeingChanged);
+
+	void CloseDocumentTab(const UObject* DocumentID);
+
+	bool InEditingMode() const;
+
+	bool IsEditable(UEdGraph* InGraph) const;
+	
+	FGraphPanelSelectionSet GetSelectedNodes() const;
+	
+	void DeleteSelectedNodes();
+	
+	bool CanDeleteSelectedNodes();
+
+	void SetSelectedObjects(TArray<UObject*> InObjects);
+
+	void HandleDetailsViewCreated(TSharedRef<IDetailsView> InDetailsView)
+	{
+		DetailsView = InDetailsView;
+	}
+
+	void RestoreEditedObjectState();
+
+	void SaveEditedObjectState();
+
+	TSharedPtr<SDockTab> OpenDocument(const UObject* InForObject, FDocumentTracker::EOpenDocumentCause InCause);
+
+	// Open the supplied assets for editing within the workspace editor
+	void OpenAssets(TConstArrayView<FAssetData> InAssets);
+
+	void HandleSaveGraphState(UEdGraph* InGraph, FVector2D InViewOffset, float InZoomAmount);
+
+	void HandleSaveDocumentState(UObject* InObject);
+
+	// The asset we are editing
+	UAnimNextWorkspace* Workspace = nullptr;
+
+	// Currently focused graph editor
+	TWeakPtr<SGraphEditor> FocusedGraphEdPtr;
+
+	// Document tracker
+	TSharedPtr<FDocumentTracker> DocumentManager;
+
+	// Command list for this editor
+	TSharedPtr<FUICommandList> CommandList;
+
+	// Our details panel
+	TSharedPtr<IDetailsView> DetailsView;
+
+	TArray<TStrongObjectPtr<UDetailsViewWrapperObject>> WrapperObjects;
+
+	static TMap<FName, FAssetDocumentWidgetFactoryFunc> AssetDocumentWidgetFactories;
+
+	bool bSavingWorkspaceOnly = false;
+};
+
+}
\ No newline at end of file
diff --git a/Engine/Plugins/Experimental/AnimNext/Source/AnimNextEditor/Private/Workspace/AnimNextWorkspaceEditorMode.cpp b/Engine/Plugins/Experimental/AnimNext/Source/AnimNextEditor/Private/Workspace/AnimNextWorkspaceEditorMode.cpp
new file mode 100644
--- /dev/null
+++ b/Engine/Plugins/Experimental/AnimNext/Source/AnimNextEditor/Private/Workspace/AnimNextWorkspaceEditorMode.cpp
@@ -0,0 +1,111 @@
+﻿// Copyright Epic Games, Inc. All Rights Reserved.
+
+#include "AnimNextWorkspaceEditorMode.h"
+
+#include "AnimNextWorkspaceEditor.h"
+#include "DetailsTabSummoner.h"
+#include "WorkspaceTabSummoner.h"
+#include "Widgets/Docking/SDockTab.h"
+#include "EdGraph/EdGraph.h"
+#include "EdGraph/EdGraphSchema.h"
+#include "Widgets/Layout/SSpacer.h"
+#include "Modules/ModuleManager.h"
+#include "Graph/GraphEditorMode.h"
+
+#define LOCTEXT_NAMESPACE "WorkspaceEditorMode"
+
+namespace UE::AnimNext::Editor
+{
+
+FWorkspaceEditorMode::FWorkspaceEditorMode(TSharedRef<FWorkspaceEditor> InHostingApp)
+	: FApplicationMode(WorkspaceModes::WorkspaceEditor)
+	, HostingAppPtr(InHostingApp)
+{
+	TSharedRef<FWorkspaceEditor> WorkspaceEditor = StaticCastSharedRef<FWorkspaceEditor>(InHostingApp);
+	
+	TabFactories.RegisterFactory(MakeShared<FDetailsTabSummoner>(WorkspaceEditor, FOnDetailsViewCreated::CreateSP(&WorkspaceEditor.Get(), &FWorkspaceEditor::HandleDetailsViewCreated)));
+	TabFactories.RegisterFactory(MakeShared<FWorkspaceTabSummoner>(WorkspaceEditor));
+
+	TabLayout = FTabManager::NewLayout("Standalone_AnimNextWorkspaceEditor_Layout_v1.1")
+		->AddArea
+		(
+			FTabManager::NewPrimaryArea()
+			->SetOrientation(Orient_Vertical)
+			->Split
+			(
+				FTabManager::NewSplitter()
+				->SetSizeCoefficient(1.0f)
+				->SetOrientation(Orient_Horizontal)
+				->Split
+				(
+					FTabManager::NewStack()
+					->SetSizeCoefficient(0.25f)
+					->SetHideTabWell(false)
+					->AddTab(WorkspaceTabs::WorkspaceView, ETabState::OpenedTab)
+					->AddTab(WorkspaceTabs::LeftAssetDocument, ETabState::ClosedTab)
+				)
+				->Split
+				(
+					FTabManager::NewStack()
+					->SetSizeCoefficient(0.5f)
+					->SetHideTabWell(false)
+					->AddTab(WorkspaceTabs::ParameterBlockGraphDocument, ETabState::ClosedTab)
+					->AddTab(WorkspaceTabs::AnimNextGraphDocument, ETabState::ClosedTab)
+					->AddTab(WorkspaceTabs::MiddleAssetDocument, ETabState::ClosedTab)
+				)
+				->Split
+				(
+					FTabManager::NewStack()
+					->SetSizeCoefficient(0.25f)
+					->SetHideTabWell(false)
+					->AddTab(WorkspaceTabs::Details, ETabState::OpenedTab)
+				)
+			)
+		);
+
+	WorkspaceEditor->RegisterModeToolbarIfUnregistered(GetModeName());
+}
+
+void FWorkspaceEditorMode::RegisterTabFactories(TSharedPtr<FTabManager> InTabManager)
+{
+	TSharedPtr<FWorkspaceEditor> HostingApp = HostingAppPtr.Pin();
+	HostingApp->PushTabFactories(TabFactories);
+
+	FApplicationMode::RegisterTabFactories(InTabManager);
+}
+
+void FWorkspaceEditorMode::AddTabFactory(FCreateWorkflowTabFactory FactoryCreator)
+{
+	if (FactoryCreator.IsBound())
+	{
+		TabFactories.RegisterFactory(FactoryCreator.Execute(HostingAppPtr.Pin()));
+	}
+}
+
+void FWorkspaceEditorMode::RemoveTabFactory(FName TabFactoryID)
+{
+	TabFactories.UnregisterFactory(TabFactoryID);
+}
+
+void FWorkspaceEditorMode::PreDeactivateMode()
+{
+	FApplicationMode::PreDeactivateMode();
+
+	TSharedPtr<FWorkspaceEditor> WorkspaceEditor = HostingAppPtr.Pin();
+
+	WorkspaceEditor->SaveEditedObjectState();
+}
+
+void FWorkspaceEditorMode::PostActivateMode()
+{
+	// Reopen any documents that were open when the blueprint was last saved
+	TSharedPtr<FWorkspaceEditor> WorkspaceEditor = HostingAppPtr.Pin();
+	WorkspaceEditor->RestoreEditedObjectState();
+
+	FApplicationMode::PostActivateMode();
+}
+
+
+}
+
+#undef LOCTEXT_NAMESPACE
\ No newline at end of file
diff --git a/Engine/Plugins/Experimental/AnimNext/Source/AnimNextEditor/Private/Workspace/AnimNextWorkspaceEditorMode.h b/Engine/Plugins/Experimental/AnimNext/Source/AnimNextEditor/Private/Workspace/AnimNextWorkspaceEditorMode.h
new file mode 100644
--- /dev/null
+++ b/Engine/Plugins/Experimental/AnimNext/Source/AnimNextEditor/Private/Workspace/AnimNextWorkspaceEditorMode.h
@@ -0,0 +1,42 @@
+﻿// Copyright Epic Games, Inc. All Rights Reserved.
+
+#pragma once
+
+#include "CoreMinimal.h"
+#include "EdGraph/EdGraph.h"
+#include "WorkflowOrientedApp/ApplicationMode.h"
+#include "WorkflowOrientedApp/WorkflowTabManager.h"
+
+class FWorkflowCentricApplication;
+class SGraphEditor;
+
+namespace UE::AnimNext::Editor
+{
+	class FWorkspaceEditor;
+}
+
+namespace UE::AnimNext::Editor
+{
+
+class FWorkspaceEditorMode : public FApplicationMode
+{
+public:
+	FWorkspaceEditorMode(TSharedRef<FWorkspaceEditor> InHostingApp);
+
+private:
+	// FApplicationMode interface
+	virtual void RegisterTabFactories(TSharedPtr<FTabManager> InTabManager) override;
+	virtual void AddTabFactory(FCreateWorkflowTabFactory FactoryCreator) override;
+	virtual void RemoveTabFactory(FName TabFactoryID) override;
+	virtual void PreDeactivateMode() override;
+	virtual void PostActivateMode() override;
+
+	// The hosting app
+	TWeakPtr<FWorkspaceEditor> HostingAppPtr;
+
+	// The tab factories we support
+	FWorkflowAllowedTabSet TabFactories;
+};
+
+}
+
diff --git a/Engine/Plugins/Experimental/AnimNext/Source/AnimNextEditor/Private/Workspace/AnimNextWorkspaceFactory.cpp b/Engine/Plugins/Experimental/AnimNext/Source/AnimNextEditor/Private/Workspace/AnimNextWorkspaceFactory.cpp
new file mode 100644
--- /dev/null
+++ b/Engine/Plugins/Experimental/AnimNext/Source/AnimNextEditor/Private/Workspace/AnimNextWorkspaceFactory.cpp
@@ -0,0 +1,28 @@
+﻿// Copyright Epic Games, Inc. All Rights Reserved.
+
+#include "AnimNextWorkspaceFactory.h"
+#include "AnimNextWorkspace.h"
+#include "UObject/Package.h"
+
+UAnimNextWorkspaceFactory::UAnimNextWorkspaceFactory()
+{
+	bCreateNew = true;
+	bEditAfterNew = true;
+	SupportedClass = UAnimNextWorkspace::StaticClass();
+}
+
+bool UAnimNextWorkspaceFactory::ConfigureProperties()
+{
+	return true;
+}
+
+UObject* UAnimNextWorkspaceFactory::FactoryCreateNew(UClass* Class, UObject* InParent, FName Name, EObjectFlags Flags, UObject* Context, FFeedbackContext* Warn, FName CallingContext)
+{
+	UAnimNextWorkspace* NewParameter = NewObject<UAnimNextWorkspace>(InParent, Class, Name, Flags | RF_Public | RF_Standalone | RF_Transactional | RF_LoadCompleted);
+
+	// make sure the package is never cooked.
+	UPackage* Package = NewParameter->GetOutermost();
+	Package->SetPackageFlags(Package->GetPackageFlags() | PKG_EditorOnly);
+
+	return NewParameter;
+}
\ No newline at end of file
diff --git a/Engine/Plugins/Experimental/AnimNext/Source/AnimNextEditor/Private/Workspace/AnimNextWorkspaceFactory.h b/Engine/Plugins/Experimental/AnimNext/Source/AnimNextEditor/Private/Workspace/AnimNextWorkspaceFactory.h
new file mode 100644
--- /dev/null
+++ b/Engine/Plugins/Experimental/AnimNext/Source/AnimNextEditor/Private/Workspace/AnimNextWorkspaceFactory.h
@@ -0,0 +1,30 @@
+﻿// Copyright Epic Games, Inc. All Rights Reserved.
+
+#pragma once
+
+#include "CoreMinimal.h"
+#include "Factories/Factory.h"
+#include "AnimNextWorkspaceFactory.generated.h"
+
+namespace UE::AnimNext::Editor
+{
+	struct FUtils;
+	class FWorkspaceEditor;
+}
+
+UCLASS(MinimalAPI, BlueprintType)
+class UAnimNextWorkspaceFactory : public UFactory
+{
+	GENERATED_BODY()
+
+	friend class UE::AnimNext::Editor::FWorkspaceEditor;
+
+	UAnimNextWorkspaceFactory();
+
+	virtual bool ConfigureProperties() override;
+	virtual UObject* FactoryCreateNew(UClass* Class, UObject* InParent, FName Name, EObjectFlags Flags, UObject* Context, FFeedbackContext* Warn, FName CallingContext) override;
+	virtual UObject* FactoryCreateNew(UClass* Class, UObject* InParent, FName Name, EObjectFlags Flags, UObject* Context, FFeedbackContext* Warn) override
+	{
+		return FactoryCreateNew(Class, InParent, Name, Flags, Context, Warn, NAME_None);
+	}
+};
\ No newline at end of file
diff --git a/Engine/Plugins/Experimental/AnimNext/Source/AnimNextEditor/Private/Workspace/AssetDocumentSummoner.cpp b/Engine/Plugins/Experimental/AnimNext/Source/AnimNextEditor/Private/Workspace/AssetDocumentSummoner.cpp
new file mode 100644
--- /dev/null
+++ b/Engine/Plugins/Experimental/AnimNext/Source/AnimNextEditor/Private/Workspace/AssetDocumentSummoner.cpp
@@ -0,0 +1,126 @@
+﻿// Copyright Epic Games, Inc. All Rights Reserved.
+
+#include "AssetDocumentSummoner.h"
+
+#include "AnimNextWorkspaceEditor.h"
+#include "AssetDefinitionRegistry.h"
+#include "Widgets/Docking/SDockTab.h"
+#include "Widgets/Layout/SSpacer.h"
+
+#define LOCTEXT_NAMESPACE "AnimNextAssetDocumentSummoner"
+
+namespace UE::AnimNext::Editor
+{
+
+FAssetDocumentSummoner::FAssetDocumentSummoner(FName InIdentifier, TSharedPtr<FWorkspaceEditor> InHostingApp)
+	: FDocumentTabFactoryForObjects<UObject>(InIdentifier, InHostingApp)
+	, HostingAppPtr(InHostingApp)
+{
+}
+
+void FAssetDocumentSummoner::SetAllowedAssetClassPaths(TConstArrayView<FTopLevelAssetPath> InAllowedAssetClassPaths)
+{
+	AllowedAssetClassPaths = InAllowedAssetClassPaths;
+}
+
+void FAssetDocumentSummoner::OnTabActivated(TSharedPtr<SDockTab> Tab) const
+{
+}
+
+void FAssetDocumentSummoner::OnTabBackgrounded(TSharedPtr<SDockTab> Tab) const
+{
+}
+
+void FAssetDocumentSummoner::OnTabRefreshed(TSharedPtr<SDockTab> Tab) const
+{
+}
+
+void FAssetDocumentSummoner::SaveState(TSharedPtr<SDockTab> Tab, TSharedPtr<FTabPayload> Payload) const
+{
+	UObject* Object = Payload->IsValid() ? FTabPayload_UObject::CastChecked<UObject>(Payload) : nullptr;
+	OnSaveDocumentStateDelegate.ExecuteIfBound(Object);
+}
+
+TAttribute<FText> FAssetDocumentSummoner::ConstructTabNameForObject(UObject* DocumentID) const
+{
+	if (DocumentID)
+	{
+		return MakeAttributeLambda([WeakObject = TWeakObjectPtr<UObject>(DocumentID)]()
+		{
+			if(UObject* Object = WeakObject.Get())
+			{
+				return FText::FromName(Object->GetFName());
+			}
+
+			return FText::GetEmpty();
+		});
+	}
+
+	return LOCTEXT("UnknownGraphName", "Unknown");
+}
+
+bool FAssetDocumentSummoner::IsPayloadSupported(TSharedRef<FTabPayload> Payload) const
+{
+	UObject* Object = Payload->IsValid() ? FTabPayload_UObject::CastChecked<UObject>(Payload) : nullptr;
+	if(Object && Object->IsAsset())
+	{
+		if(FWorkspaceEditor::FAssetDocumentWidgetFactoryFunc* FactoryFunc = FWorkspaceEditor::AssetDocumentWidgetFactories.Find(Object->GetClass()->GetFName()))
+		{
+			return AllowedAssetClassPaths.IsEmpty() || AllowedAssetClassPaths.Contains(Object->GetClass()->GetClassPathName());
+		}
+	}
+
+	return false;
+}
+
+TAttribute<FText> FAssetDocumentSummoner::ConstructTabLabelSuffix(const FWorkflowTabSpawnInfo& Info) const
+{
+	UObject* Object = Info.Payload->IsValid() ? FTabPayload_UObject::CastChecked<UObject>(Info.Payload) : nullptr;
+	if(Object)
+	{
+		return MakeAttributeLambda([WeakObject = TWeakObjectPtr<UObject>(Object)]()
+		{
+			if(UObject* Object = WeakObject.Get())
+			{
+				if(Object->GetPackage()->IsDirty())
+				{
+					return LOCTEXT("TabSuffixAsterisk", "*");
+				}
+			}
+
+			return FText::GetEmpty();
+		});
+	}
+
+	return FText::GetEmpty();
+}
+
+TSharedRef<SWidget> FAssetDocumentSummoner::CreateTabBodyForObject(const FWorkflowTabSpawnInfo& Info, UObject* DocumentID) const
+{
+	if(TSharedPtr<FWorkspaceEditor> WorkspaceEditor = HostingAppPtr.Pin())
+	{
+		if(FWorkspaceEditor::FAssetDocumentWidgetFactoryFunc* FactoryFunc = FWorkspaceEditor::AssetDocumentWidgetFactories.Find(DocumentID->GetClass()->GetFName()))
+		{
+			return (*FactoryFunc)(WorkspaceEditor.ToSharedRef(), DocumentID);
+		}
+	}
+
+	return SNew(SSpacer);
+}
+
+const FSlateBrush* FAssetDocumentSummoner::GetTabIconForObject(const FWorkflowTabSpawnInfo& Info, UObject* DocumentID) const
+{
+	if (UAssetDefinitionRegistry* AssetDefinitionRegistry = UAssetDefinitionRegistry::Get())
+	{
+		const FAssetData AssetData(DocumentID);
+		if(const UAssetDefinition* AssetDefinition = AssetDefinitionRegistry->GetAssetDefinitionForAsset(AssetData))
+		{
+			return AssetDefinition->GetIconBrush(AssetData, AssetData.AssetClassPath.GetAssetName());
+		}
+	}
+	return nullptr;
+}
+
+}
+
+#undef LOCTEXT_NAMESPACE
\ No newline at end of file
diff --git a/Engine/Plugins/Experimental/AnimNext/Source/AnimNextEditor/Private/Workspace/AssetDocumentSummoner.h b/Engine/Plugins/Experimental/AnimNext/Source/AnimNextEditor/Private/Workspace/AssetDocumentSummoner.h
new file mode 100644
--- /dev/null
+++ b/Engine/Plugins/Experimental/AnimNext/Source/AnimNextEditor/Private/Workspace/AssetDocumentSummoner.h
@@ -0,0 +1,58 @@
+﻿// Copyright Epic Games, Inc. All Rights Reserved.
+
+#pragma once
+
+#include "CoreMinimal.h"
+#include "WorkflowOrientedApp/WorkflowUObjectDocuments.h"
+
+class FWorkflowCentricApplication;
+class SGraphEditor;
+
+namespace UE::AnimNext::Editor
+{
+	class FWorkspaceEditor;
+}
+
+namespace UE::AnimNext::Editor
+{
+
+struct FAssetDocumentSummoner : public FDocumentTabFactoryForObjects<UObject>
+{
+public:
+	// Delegate called to save the state of a document
+	DECLARE_DELEGATE_OneParam(FOnSaveDocumentState, UObject*);
+
+	FAssetDocumentSummoner(FName InIdentifier, TSharedPtr<FWorkspaceEditor> InHostingApp);
+
+	FOnSaveDocumentState& OnSaveDocumentState() { return OnSaveDocumentStateDelegate; }
+
+	void SetAllowedAssetClassPaths(TConstArrayView<FTopLevelAssetPath> InAllowedAssetClassPaths);
+
+private:
+	// FWorkflowTabFactory interface
+	virtual void OnTabActivated(TSharedPtr<SDockTab> Tab) const override;
+	virtual void OnTabBackgrounded(TSharedPtr<SDockTab> Tab) const override;
+	virtual void OnTabRefreshed(TSharedPtr<SDockTab> Tab) const override;
+	virtual void SaveState(TSharedPtr<SDockTab> Tab, TSharedPtr<FTabPayload> Payload) const override;
+	virtual bool IsPayloadSupported(TSharedRef<FTabPayload> Payload) const override;
+	virtual TAttribute<FText> ConstructTabLabelSuffix(const FWorkflowTabSpawnInfo& Info) const override;
+
+	// FDocumentTabFactoryForObjects interface
+	virtual TAttribute<FText> ConstructTabNameForObject(UObject* DocumentID) const override;
+	virtual const FSlateBrush* GetTabIconForObject(const FWorkflowTabSpawnInfo& Info, UObject* DocumentID) const override;
+	virtual TSharedRef<SWidget> CreateTabBodyForObject(const FWorkflowTabSpawnInfo& Info, UObject* DocumentID) const override;
+
+	// The hosting app
+	TWeakPtr<FWorkspaceEditor> HostingAppPtr;
+
+	// Delegate called to save the state of a graph
+	FOnSaveDocumentState OnSaveDocumentStateDelegate;
+
+	// Command list
+	TSharedPtr<FUICommandList> CommandList;
+
+	// Allowed asset types
+	TArray<FTopLevelAssetPath> AllowedAssetClassPaths;
+};
+
+}
diff --git a/Engine/Plugins/Experimental/AnimNext/Source/AnimNextEditor/Private/Workspace/DetailsTabSummoner.cpp b/Engine/Plugins/Experimental/AnimNext/Source/AnimNextEditor/Private/Workspace/DetailsTabSummoner.cpp
new file mode 100644
--- /dev/null
+++ b/Engine/Plugins/Experimental/AnimNext/Source/AnimNextEditor/Private/Workspace/DetailsTabSummoner.cpp
@@ -0,0 +1,43 @@
+﻿// Copyright Epic Games, Inc. All Rights Reserved.
+
+#include "DetailsTabSummoner.h"
+#include "AnimNextWorkspaceEditor.h"
+#include "PropertyEditorModule.h"
+
+#define LOCTEXT_NAMESPACE "DetailsTabSummoner"
+
+namespace UE::AnimNext::Editor
+{
+
+FDetailsTabSummoner::FDetailsTabSummoner(TSharedPtr<FWorkspaceEditor> InHostingApp, FOnDetailsViewCreated InOnDetailsViewCreated)
+	: FWorkflowTabFactory(WorkspaceTabs::Details, InHostingApp)
+{
+	TabLabel = LOCTEXT("DetailsTabLabel", "Details");
+	TabIcon = FSlateIcon("EditorStyle", "LevelEditor.Tabs.Details");
+	ViewMenuDescription = LOCTEXT("DetailsTabMenuDescription", "Details");
+	ViewMenuTooltip = LOCTEXT("DetailsTabToolTip", "Shows the details tab for selected objects.");
+	bIsSingleton = true;
+
+	FDetailsViewArgs DetailsViewArgs;
+	DetailsViewArgs.NameAreaSettings = FDetailsViewArgs::HideNameArea;
+	DetailsViewArgs.bHideSelectionTip = true;
+
+	FPropertyEditorModule& PropertyEditorModule = FModuleManager::LoadModuleChecked<FPropertyEditorModule>("PropertyEditor");
+	DetailsView = PropertyEditorModule.CreateDetailView(DetailsViewArgs);
+
+	InOnDetailsViewCreated.ExecuteIfBound(DetailsView.ToSharedRef());
+}
+
+TSharedRef<SWidget> FDetailsTabSummoner::CreateTabBody(const FWorkflowTabSpawnInfo& Info) const
+{
+	return DetailsView.ToSharedRef();
+}
+
+FText FDetailsTabSummoner::GetTabToolTipText(const FWorkflowTabSpawnInfo& Info) const
+{
+	return ViewMenuTooltip;
+}
+
+}
+
+#undef LOCTEXT_NAMESPACE
\ No newline at end of file
diff --git a/Engine/Plugins/Experimental/AnimNext/Source/AnimNextEditor/Private/Workspace/DetailsTabSummoner.h b/Engine/Plugins/Experimental/AnimNext/Source/AnimNextEditor/Private/Workspace/DetailsTabSummoner.h
new file mode 100644
--- /dev/null
+++ b/Engine/Plugins/Experimental/AnimNext/Source/AnimNextEditor/Private/Workspace/DetailsTabSummoner.h
@@ -0,0 +1,36 @@
+﻿// Copyright Epic Games, Inc. All Rights Reserved.
+
+#pragma once
+
+#include "WorkflowOrientedApp/WorkflowTabFactory.h"
+
+class IDetailsView;
+
+namespace UE::AnimNext::Editor
+{
+	class FWorkspaceEditor;
+}
+
+namespace UE::AnimNext::Editor
+{
+
+DECLARE_DELEGATE_OneParam(FOnDetailsViewCreated, TSharedRef<IDetailsView>);
+
+class FDetailsTabSummoner : public FWorkflowTabFactory
+{
+public:
+	FDetailsTabSummoner(TSharedPtr<FWorkspaceEditor> InHostingApp, FOnDetailsViewCreated InOnDetailsViewCreated);
+
+	TSharedPtr<IDetailsView> GetDetailsView() const
+	{
+		return DetailsView;
+	}
+
+private:
+	virtual TSharedRef<SWidget> CreateTabBody(const FWorkflowTabSpawnInfo& Info) const override;
+	virtual FText GetTabToolTipText(const FWorkflowTabSpawnInfo& Info) const override;
+
+	TSharedPtr<IDetailsView> DetailsView;
+};
+
+}
\ No newline at end of file
diff --git a/Engine/Plugins/Experimental/AnimNext/Source/AnimNextEditor/Private/Workspace/GraphDocumentSummoner.cpp b/Engine/Plugins/Experimental/AnimNext/Source/AnimNextEditor/Private/Workspace/GraphDocumentSummoner.cpp
new file mode 100644
--- /dev/null
+++ b/Engine/Plugins/Experimental/AnimNext/Source/AnimNextEditor/Private/Workspace/GraphDocumentSummoner.cpp
@@ -0,0 +1,179 @@
+﻿// Copyright Epic Games, Inc. All Rights Reserved.
+
+#include "GraphDocumentSummoner.h"
+
+#include "AnimNextWorkspaceEditor.h"
+#include "WorkflowOrientedApp/WorkflowTabManager.h"
+#include "Widgets/Docking/SDockTab.h"
+#include "Widgets/Layout/SSpacer.h"
+#include "EdGraph/EdGraphSchema.h"
+#include "Framework/Commands/GenericCommands.h"
+
+#define LOCTEXT_NAMESPACE "AnimNextGraphDocumentSummoner"
+
+namespace UE::AnimNext::Editor
+{
+
+FGraphDocumentSummoner::FGraphDocumentSummoner(FName InIdentifier, TSharedPtr<FWorkspaceEditor> InHostingApp)
+	: FDocumentTabFactoryForObjects<UEdGraph>(InIdentifier, InHostingApp)
+	, HostingAppPtr(InHostingApp)
+{
+	CommandList = MakeShared<FUICommandList>();
+
+	CommandList->MapAction(FGenericCommands::Get().Delete,
+			FExecuteAction::CreateSP(HostingAppPtr.Pin().Get(), &FWorkspaceEditor::DeleteSelectedNodes),
+			FCanExecuteAction::CreateSP(HostingAppPtr.Pin().Get(), &FWorkspaceEditor::CanDeleteSelectedNodes));
+}
+
+void FGraphDocumentSummoner::OnTabActivated(TSharedPtr<SDockTab> Tab) const
+{
+	if(TSharedPtr<FWorkspaceEditor> WorkspaceEditor = HostingAppPtr.Pin())
+	{
+		WorkspaceEditor->SetFocusedGraphEditor(StaticCastSharedRef<SGraphEditor>(Tab->GetContent()));
+	}
+}
+
+void FGraphDocumentSummoner::OnTabBackgrounded(TSharedPtr<SDockTab> Tab) const
+{
+	if(TSharedPtr<FWorkspaceEditor> WorkspaceEditor = HostingAppPtr.Pin())
+	{
+		WorkspaceEditor->SetFocusedGraphEditor(nullptr);
+	}
+}
+
+void FGraphDocumentSummoner::OnTabRefreshed(TSharedPtr<SDockTab> Tab) const
+{
+	TSharedRef<SGraphEditor> GraphEditor = StaticCastSharedRef<SGraphEditor>(Tab->GetContent());
+	GraphEditor->NotifyGraphChanged();
+}
+
+void FGraphDocumentSummoner::SaveState(TSharedPtr<SDockTab> Tab, TSharedPtr<FTabPayload> Payload) const
+{
+	TSharedRef<SGraphEditor> GraphEditor = StaticCastSharedRef<SGraphEditor>(Tab->GetContent());
+
+	FVector2D ViewLocation;
+	float ZoomAmount;
+	GraphEditor->GetViewLocation(ViewLocation, ZoomAmount);
+
+	UEdGraph* Graph = Payload->IsValid() ? FTabPayload_UObject::CastChecked<UEdGraph>(Payload) : nullptr;
+
+	OnSaveGraphStateDelegate.ExecuteIfBound(Graph, ViewLocation, ZoomAmount);
+}
+
+TAttribute<FText> FGraphDocumentSummoner::ConstructTabNameForObject(UEdGraph* DocumentID) const
+{
+	if (DocumentID)
+	{
+		if (const UEdGraphSchema* Schema = DocumentID->GetSchema())
+		{
+			FGraphDisplayInfo Info;
+			Schema->GetGraphDisplayInformation(*DocumentID, /*out*/ Info);
+			return Info.DisplayName;
+		}
+		else
+		{
+			// if we don't have a schema, we're dealing with a malformed (or incomplete graph)...
+			// possibly in the midst of some transaction - here we return the object's outer path 
+			// so we can at least get some context as to which graph we're referring
+			return FText::FromString(DocumentID->GetPathName());
+		}
+	}
+
+	return LOCTEXT("UnknownGraphName", "Unknown");
+}
+
+TSharedRef<SWidget> FGraphDocumentSummoner::CreateTabBodyForObject(const FWorkflowTabSpawnInfo& Info, UEdGraph* DocumentID) const
+{
+	SGraphEditor::FGraphEditorEvents Events;
+	Events.OnCreateActionMenu = SGraphEditor::FOnCreateActionMenu::CreateSP(this, &FGraphDocumentSummoner::OnCreateGraphActionMenu);
+	
+	return SNew(SGraphEditor)
+		.AdditionalCommands(CommandList)
+		.IsEditable(HostingAppPtr.Pin().Get(), &FWorkspaceEditor::IsEditable, DocumentID)
+		.GraphToEdit(DocumentID)
+		.GraphEvents(Events)
+		.AssetEditorToolkit(HostingAppPtr);
+}
+
+bool FGraphDocumentSummoner::IsPayloadSupported(TSharedRef<FTabPayload> Payload) const
+{
+	UObject* Object = Payload->IsValid() ? FTabPayload_UObject::CastChecked<UObject>(Payload) : nullptr;
+	return Object && Object->IsA<UEdGraph>();
+}
+
+const FSlateBrush* FGraphDocumentSummoner::GetTabIconForObject(const FWorkflowTabSpawnInfo& Info, UEdGraph* DocumentID) const
+{
+	return FAppStyle::Get().GetBrush(TEXT("GraphEditor.EventGraph_16x"));
+}
+
+TSharedRef<FGenericTabHistory> FGraphDocumentSummoner::CreateTabHistoryNode(TSharedPtr<FTabPayload> Payload)
+{
+	struct FLocalGraphTabHistory : public FGenericTabHistory
+	{
+	public:
+		FLocalGraphTabHistory(TSharedPtr<FDocumentTabFactory> InFactory, TSharedPtr<FTabPayload> InPayload)
+			: FGenericTabHistory(InFactory, InPayload)
+			, SavedLocation(FVector2D::ZeroVector)
+			, SavedZoomAmount(INDEX_NONE)
+		{
+		}
+
+	private:
+		// FGenericTabHistory interface
+		virtual void EvokeHistory(TSharedPtr<::FTabInfo> InTabInfo, bool bPrevTabMatches) override
+		{
+			FWorkflowTabSpawnInfo SpawnInfo;
+			SpawnInfo.Payload = Payload;
+			SpawnInfo.TabInfo = InTabInfo;
+
+			if(bPrevTabMatches)
+			{
+				TSharedPtr<SDockTab> DockTab = InTabInfo->GetTab().Pin();
+				GraphEditor = StaticCastSharedRef<SGraphEditor>(DockTab->GetContent());
+			}
+			else
+			{
+				TSharedRef<SGraphEditor> GraphEditorRef = StaticCastSharedRef<SGraphEditor>(FactoryPtr.Pin()->CreateTabBody(SpawnInfo));
+				GraphEditor = GraphEditorRef;
+				FactoryPtr.Pin()->UpdateTab(InTabInfo->GetTab().Pin(), SpawnInfo, GraphEditorRef);
+			}
+		}
+	
+		virtual void SaveHistory() override
+		{
+			if (IsHistoryValid())
+			{
+				check(GraphEditor.IsValid());
+				GraphEditor.Pin()->GetViewLocation(SavedLocation, SavedZoomAmount);
+				GraphEditor.Pin()->GetViewBookmark(SavedBookmarkId);
+			}
+		}
+	
+		virtual void RestoreHistory() override
+		{
+			if (IsHistoryValid())
+			{
+				check(GraphEditor.IsValid());
+				GraphEditor.Pin()->SetViewLocation(SavedLocation, SavedZoomAmount, SavedBookmarkId);
+			}
+		}
+	
+		/** The graph editor represented by this history node. While this node is inactive, the graph editor is invalid */
+		TWeakPtr<SGraphEditor> GraphEditor;
+	
+		/** Saved location the graph editor was at when this history node was last visited */
+		FVector2D SavedLocation;
+	
+		/** Saved zoom the graph editor was at when this history node was last visited */
+		float SavedZoomAmount;
+	
+		/** Saved bookmark ID the graph editor was at when this history node was last visited */
+		FGuid SavedBookmarkId;
+	};
+
+	return MakeShared<FLocalGraphTabHistory>(SharedThis(this), Payload);
+}
+
+}
+
+#undef LOCTEXT_NAMESPACE
\ No newline at end of file
diff --git a/Engine/Plugins/Experimental/AnimNext/Source/AnimNextEditor/Private/Workspace/GraphDocumentSummoner.h b/Engine/Plugins/Experimental/AnimNext/Source/AnimNextEditor/Private/Workspace/GraphDocumentSummoner.h
new file mode 100644
--- /dev/null
+++ b/Engine/Plugins/Experimental/AnimNext/Source/AnimNextEditor/Private/Workspace/GraphDocumentSummoner.h
@@ -0,0 +1,58 @@
+﻿// Copyright Epic Games, Inc. All Rights Reserved.
+
+#pragma once
+
+#include "CoreMinimal.h"
+#include "WorkflowOrientedApp/WorkflowUObjectDocuments.h"
+#include "EdGraph/EdGraph.h"
+#include "GraphEditor.h"
+
+class FWorkflowCentricApplication;
+class SGraphEditor;
+
+namespace UE::AnimNext::Editor
+{
+	class FWorkspaceEditor;
+}
+
+namespace UE::AnimNext::Editor
+{
+
+struct FGraphDocumentSummoner : public FDocumentTabFactoryForObjects<UEdGraph>
+{
+public:
+	// Delegate called to save the state of a graph
+	DECLARE_DELEGATE_ThreeParams(FOnSaveGraphState, UEdGraph*, FVector2D, float);
+
+	FGraphDocumentSummoner(FName InIdentifier, TSharedPtr<FWorkspaceEditor> InHostingApp);
+
+	FOnSaveGraphState& OnSaveGraphState() { return OnSaveGraphStateDelegate; }
+
+private:
+	// FWorkflowTabFactory interface
+	virtual void OnTabActivated(TSharedPtr<SDockTab> Tab) const override;
+	virtual void OnTabBackgrounded(TSharedPtr<SDockTab> Tab) const override;
+	virtual void OnTabRefreshed(TSharedPtr<SDockTab> Tab) const override;
+	virtual void SaveState(TSharedPtr<SDockTab> Tab, TSharedPtr<FTabPayload> Payload) const override;
+	
+	// FDocumentTabFactoryForObjects interface
+	virtual TAttribute<FText> ConstructTabNameForObject(UEdGraph* DocumentID) const override;
+	virtual const FSlateBrush* GetTabIconForObject(const FWorkflowTabSpawnInfo& Info, UEdGraph* DocumentID) const override;
+	virtual TSharedRef<SWidget> CreateTabBodyForObject(const FWorkflowTabSpawnInfo& Info, UEdGraph* DocumentID) const override;
+	virtual TSharedRef<FGenericTabHistory> CreateTabHistoryNode(TSharedPtr<FTabPayload> Payload) override;
+	virtual bool IsPayloadSupported(TSharedRef<FTabPayload> Payload) const override;
+
+	// Create the context 'action' menu used for graphs
+	virtual FActionMenuContent OnCreateGraphActionMenu(UEdGraph* InGraph, const FVector2D& InNodePosition, const TArray<UEdGraphPin*>& InDraggedPins, bool bAutoExpand, SGraphEditor::FActionMenuClosed InOnMenuClosed) const = 0;
+
+	// The hosting app
+	TWeakPtr<FWorkspaceEditor> HostingAppPtr;
+
+	// Delegate called to save the state of a graph
+	FOnSaveGraphState OnSaveGraphStateDelegate;
+
+	// Command list for graphs
+	TSharedPtr<FUICommandList> CommandList;
+};
+
+}
diff --git a/Engine/Plugins/Experimental/AnimNext/Source/AnimNextEditor/Private/Workspace/SWorkspacePicker.cpp b/Engine/Plugins/Experimental/AnimNext/Source/AnimNextEditor/Private/Workspace/SWorkspacePicker.cpp
new file mode 100644
--- /dev/null
+++ b/Engine/Plugins/Experimental/AnimNext/Source/AnimNextEditor/Private/Workspace/SWorkspacePicker.cpp
@@ -0,0 +1,102 @@
+﻿// Copyright Epic Games, Inc. All Rights Reserved.
+
+#include "SWorkspacePicker.h"
+#include "AnimNextWorkspace.h"
+#include "ContentBrowserModule.h"
+#include "IContentBrowserSingleton.h"
+#include "SPrimaryButton.h"
+#include "Widgets/Layout/SBox.h"
+#include "Widgets/Text/STextBlock.h"
+#include "Framework/Application/SlateApplication.h"
+
+#define LOCTEXT_NAMESPACE "SWorkspacePicker"
+
+namespace UE::AnimNext::Editor
+{
+
+void SWorkspacePicker::Construct(const FArguments& InArgs)
+{
+	WorkspaceAssets = InArgs._WorkspaceAssets;
+
+	FContentBrowserModule& ContentBrowserModule = FModuleManager::Get().LoadModuleChecked<FContentBrowserModule>(TEXT("ContentBrowser"));
+
+	FAssetPickerConfig AssetPickerConfig;
+	AssetPickerConfig.InitialAssetViewType = EAssetViewType::List;
+	AssetPickerConfig.OnAssetSelected = FOnAssetSelected::CreateLambda([this, OnAssetSelected = InArgs._OnAssetSelected](const FAssetData& InAssetData)
+	{
+		OnAssetSelected.ExecuteIfBound(InAssetData);
+
+		if(TSharedPtr<SWindow> Window = WeakWindow.Pin())
+		{
+			Window->RequestDestroyWindow();
+		}
+	});
+	AssetPickerConfig.Filter.ClassPaths.Add(UAnimNextWorkspace::StaticClass()->GetClassPathName());
+	AssetPickerConfig.Filter.bRecursiveClasses = true;
+	AssetPickerConfig.OnShouldFilterAsset = FOnShouldFilterAsset::CreateLambda([this](const FAssetData& InAssetData)
+	{
+		return !WorkspaceAssets.IsEmpty() && !WorkspaceAssets.Contains(InAssetData);
+	});
+
+	ChildSlot
+	[
+		SNew(SBox)
+		.WidthOverride(400.0f)
+		.HeightOverride(400.0f)
+		[
+			SNew(SVerticalBox)
+			+SVerticalBox::Slot()
+			.AutoHeight()
+			.Padding(5.0f)
+			.HAlign(HAlign_Left)
+			[
+				SNew(STextBlock)
+				.Text(LOCTEXT("ChooseExistingWorkspaceToDetails", "Selected asset is part of multiple workspaces.\nPlease select the workspace you want to open the asset with"))
+			]
+			+SVerticalBox::Slot()
+			.FillHeight(1.0f)
+			.Padding(5.0f)
+			[
+				ContentBrowserModule.Get().CreateAssetPicker(AssetPickerConfig)
+			]
+			+SVerticalBox::Slot()
+			.AutoHeight()
+			.Padding(5.0f)
+			.HAlign(HAlign_Fill)
+			[
+				SNew(SPrimaryButton)
+				.Text(LOCTEXT("CreateNewWorkspace", "Create New Workspace"))
+				.OnClicked_Lambda([this, OnNewAsset = InArgs._OnNewAsset]()
+				{
+					OnNewAsset.ExecuteIfBound();
+
+					if(TSharedPtr<SWindow> Window = WeakWindow.Pin())
+					{
+						Window->RequestDestroyWindow();
+					}
+					return FReply::Handled();
+				})
+			]
+		]
+	];
+}
+
+void SWorkspacePicker::ShowModal()
+{
+	TSharedRef<SWindow> Window = SNew(SWindow)
+	.Title(LOCTEXT("ChooseExistingWorkspaceToOpen", "Choose Existing Workspace to Open"))
+	.SizingRule(ESizingRule::Autosized)
+	.SupportsMaximize(false)
+	.SupportsMinimize(false)
+	[
+		AsShared()
+	];
+
+	WeakWindow = Window;
+
+	FSlateApplication::Get().AddModalWindow(Window, FGlobalTabmanager::Get()->GetRootWindow());
+}
+
+}
+
+#undef LOCTEXT_NAMESPACE
\ No newline at end of file
diff --git a/Engine/Plugins/Experimental/AnimNext/Source/AnimNextEditor/Private/Workspace/SWorkspacePicker.h b/Engine/Plugins/Experimental/AnimNext/Source/AnimNextEditor/Private/Workspace/SWorkspacePicker.h
new file mode 100644
--- /dev/null
+++ b/Engine/Plugins/Experimental/AnimNext/Source/AnimNextEditor/Private/Workspace/SWorkspacePicker.h
@@ -0,0 +1,38 @@
+﻿// Copyright Epic Games, Inc. All Rights Reserved.
+
+#pragma once
+
+#include "ContentBrowserDelegates.h"
+#include "Widgets/SCompoundWidget.h"
+
+class UAnimNextWorkspace;
+
+namespace UE::AnimNext::Editor
+{
+
+class SWorkspacePicker : public SCompoundWidget
+{
+	using FOnNewAsset = TDelegate<void(TConstArrayView<FAssetData>)>; 
+	
+	SLATE_BEGIN_ARGS(SWorkspacePicker) {}
+
+	SLATE_EVENT(FOnAssetSelected, OnAssetSelected)
+
+	SLATE_EVENT(FSimpleDelegate, OnNewAsset)
+
+	// The set of workspace assets the user can choose from
+	SLATE_ARGUMENT(TArray<FAssetData>, WorkspaceAssets)
+
+	SLATE_END_ARGS()
+
+	void Construct(const FArguments& InArgs);
+
+	void ShowModal();
+
+private:
+	TArray<FAssetData> WorkspaceAssets;
+
+	TWeakPtr<SWindow> WeakWindow;
+};
+
+};
\ No newline at end of file
diff --git a/Engine/Plugins/Experimental/AnimNext/Source/AnimNextEditor/Private/Workspace/SWorkspaceView.cpp b/Engine/Plugins/Experimental/AnimNext/Source/AnimNextEditor/Private/Workspace/SWorkspaceView.cpp
new file mode 100644
--- /dev/null
+++ b/Engine/Plugins/Experimental/AnimNext/Source/AnimNextEditor/Private/Workspace/SWorkspaceView.cpp
@@ -0,0 +1,132 @@
+﻿// Copyright Epic Games, Inc. All Rights Reserved.
+
+#include "SWorkspaceView.h"
+#include "AnimNextWorkspace.h"
+#include "ContentBrowserModule.h"
+#include "IContentBrowserSingleton.h"
+#include "SAssetDropTarget.h"
+#include "ScopedTransaction.h"
+#include "Framework/Commands/GenericCommands.h"
+
+#define LOCTEXT_NAMESPACE "SWorkspaceView"
+
+namespace UE::AnimNext::Editor
+{
+
+void SWorkspaceView::Construct(const FArguments& InArgs, UAnimNextWorkspace* InWorkspace)
+{
+	Workspace = InWorkspace;
+	OnAssetsOpened = InArgs._OnAssetsOpened;
+
+	UICommandList = MakeShared<FUICommandList>();
+
+	UICommandList->MapAction(FGenericCommands::Get().Delete,
+		FExecuteAction::CreateSP(this, &SWorkspaceView::HandleDelete),
+		FCanExecuteAction::CreateSP(this, &SWorkspaceView::HasValidSelection));
+
+	Workspace->ModifiedDelegate.AddSP(this, &SWorkspaceView::HandleWorkspaceModified);
+
+	FContentBrowserModule& ContentBrowserModule = FModuleManager::Get().LoadModuleChecked<FContentBrowserModule>(TEXT("ContentBrowser"));
+
+	FAssetPickerConfig AssetPickerConfig;
+	AssetPickerConfig.InitialAssetViewType = EAssetViewType::List;
+	AssetPickerConfig.OnAssetsActivated = FOnAssetsActivated::CreateLambda([this](const TArray<FAssetData>& InAssets, EAssetTypeActivationMethod::Type InActivationMethod)
+	{
+		if( InActivationMethod == EAssetTypeActivationMethod::DoubleClicked ||
+			InActivationMethod == EAssetTypeActivationMethod::Opened)
+		{
+			OnAssetsOpened.ExecuteIfBound(InAssets);
+		}
+	});
+	AssetPickerConfig.RefreshAssetViewDelegates.Add(&RefreshAssetViewDelegate);
+	AssetPickerConfig.GetCurrentSelectionDelegates.Add(&GetCurrentSelectionDelegate);
+	AssetPickerConfig.Filter = MakeARFilter();
+	AssetPickerConfig.OnShouldFilterAsset = FOnShouldFilterAsset::CreateLambda([this](const FAssetData& InAsset)
+	{
+		return !Workspace->Assets.Contains(TSoftObjectPtr<UObject>(InAsset.GetSoftObjectPath()));
+	});
+	AssetPickerConfig.AssetShowWarningText = LOCTEXT("EmptyWorkspaceText", "Workspace is empty.\nDrag-drop to add assets to this workspace.");
+
+	ChildSlot
+	[
+		SNew(SAssetDropTarget)
+		.bSupportsMultiDrop(true)
+		.OnAssetsDropped_Lambda([this](const FDragDropEvent& InEvent, TArrayView<FAssetData> InAssets)
+		{
+			FScopedTransaction Transaction(LOCTEXT("AddAssets", "Add assets to workspace"));
+
+			Workspace->AddAssets(InAssets);
+		})
+		.OnAreAssetsAcceptableForDropWithReason_Lambda([](TArrayView<FAssetData> InAssets, FText& OutText)
+		{
+			for(const FAssetData& Asset : InAssets)
+			{
+				if(UAnimNextWorkspace::IsAssetSupported(Asset))
+				{
+					return true;
+				}
+			}
+
+			OutText = LOCTEXT("AssetsUnsupportedInWorkspace", "Assets are not supported by this workspace");
+			return false;
+		})
+		.Content()
+		[
+			ContentBrowserModule.Get().CreateAssetPicker(AssetPickerConfig)
+		]
+	];
+}
+
+FARFilter SWorkspaceView::MakeARFilter()
+{
+	FARFilter Filter;
+	Filter.ClassPaths = UAnimNextWorkspace::GetSupportedAssetClassPaths();
+	Filter.bRecursiveClasses = true;
+
+	return Filter;
+}
+
+FReply SWorkspaceView::OnKeyDown(const FGeometry& MyGeometry, const FKeyEvent& InKeyEvent)
+{
+	if (UICommandList.IsValid() && UICommandList->ProcessCommandBindings(InKeyEvent))
+	{
+		return FReply::Handled();
+	}
+	return FReply::Unhandled();
+}
+
+void SWorkspaceView::HandleDelete()
+{
+	if(GetCurrentSelectionDelegate.IsBound())
+	{
+		TArray<FAssetData> Selection = GetCurrentSelectionDelegate.Execute();
+
+		if(Selection.Num() > 0)
+		{
+			FScopedTransaction Transaction(LOCTEXT("RemoveAssets", "Remove assets from workspace"));
+
+			Workspace->RemoveAssets(Selection);
+		}
+	}
+}
+
+bool SWorkspaceView::HasValidSelection() const 
+{
+	if(GetCurrentSelectionDelegate.IsBound())
+	{
+		TArray<FAssetData> Selection = GetCurrentSelectionDelegate.Execute();
+		return Selection.Num() > 0;
+	}
+	return false;
+}
+
+void SWorkspaceView::HandleWorkspaceModified(UAnimNextWorkspace* InWorkspace)
+{
+	check(InWorkspace == Workspace);
+
+	RefreshAssetViewDelegate.ExecuteIfBound(true);
+}
+
+}
+
+#undef LOCTEXT_NAMESPACE
\ No newline at end of file
diff --git a/Engine/Plugins/Experimental/AnimNext/Source/AnimNextEditor/Private/Workspace/SWorkspaceView.h b/Engine/Plugins/Experimental/AnimNext/Source/AnimNextEditor/Private/Workspace/SWorkspaceView.h
new file mode 100644
--- /dev/null
+++ b/Engine/Plugins/Experimental/AnimNext/Source/AnimNextEditor/Private/Workspace/SWorkspaceView.h
@@ -0,0 +1,48 @@
+﻿// Copyright Epic Games, Inc. All Rights Reserved.
+
+#pragma once
+
+#include "ContentBrowserDelegates.h"
+#include "Widgets/SCompoundWidget.h"
+
+class UAnimNextWorkspace;
+
+namespace UE::AnimNext::Editor
+{
+
+class SWorkspaceView : public SCompoundWidget
+{
+	using FOnAssetsOpened = TDelegate<void(TConstArrayView<FAssetData>)>; 
+
+	SLATE_BEGIN_ARGS(SWorkspaceView) {}
+
+	SLATE_EVENT(FOnAssetsOpened, OnAssetsOpened)
+
+	SLATE_END_ARGS()
+
+	void Construct(const FArguments& InArgs, UAnimNextWorkspace* InWorkspace);
+
+private:
+	// SWidget interface
+	virtual FReply OnKeyDown(const FGeometry& MyGeometry, const FKeyEvent& InKeyEvent) override;
+
+	void HandleDelete();
+
+	bool HasValidSelection() const;
+
+	void HandleWorkspaceModified(UAnimNextWorkspace* InWorkspace);
+	
+	static FARFilter MakeARFilter();
+
+	UAnimNextWorkspace* Workspace = nullptr;
+
+	FOnAssetsOpened OnAssetsOpened;
+
+	FRefreshAssetViewDelegate RefreshAssetViewDelegate;
+
+	FGetCurrentSelectionDelegate GetCurrentSelectionDelegate;
+
+	TSharedPtr<FUICommandList> UICommandList;
+};
+
+};
\ No newline at end of file
diff --git a/Engine/Plugins/Experimental/AnimNext/Source/AnimNextEditor/Private/Workspace/WorkspaceTabSummoner.cpp b/Engine/Plugins/Experimental/AnimNext/Source/AnimNextEditor/Private/Workspace/WorkspaceTabSummoner.cpp
new file mode 100644
--- /dev/null
+++ b/Engine/Plugins/Experimental/AnimNext/Source/AnimNextEditor/Private/Workspace/WorkspaceTabSummoner.cpp
@@ -0,0 +1,45 @@
+﻿// Copyright Epic Games, Inc. All Rights Reserved.
+
+#include "WorkspaceTabSummoner.h"
+#include "AnimNextWorkspaceEditor.h"
+#include "SWorkspaceView.h"
+
+#define LOCTEXT_NAMESPACE "WorkspaceTabSummoner"
+
+namespace UE::AnimNext::Editor
+{
+
+FWorkspaceTabSummoner::FWorkspaceTabSummoner(TSharedPtr<FWorkspaceEditor> InHostingApp)
+	: FWorkflowTabFactory(WorkspaceTabs::WorkspaceView, InHostingApp)
+{
+	HostingAppPtr = InHostingApp;
+
+	TabLabel = LOCTEXT("WorkspaceTabLabel", "Workspace");
+	TabIcon = FSlateIcon("EditorStyle", "LevelEditor.Tabs.Outliner");
+	ViewMenuDescription = LOCTEXT("WorkspaceTabMenuDescription", "Workspace");
+	ViewMenuTooltip = LOCTEXT("WorkspaceTabToolTip", "Shows the workspace outliner tab.");
+	bIsSingleton = true;
+
+	WorkspaceView = SNew(SWorkspaceView, InHostingApp->Workspace)
+		.OnAssetsOpened_Lambda([this](TConstArrayView<FAssetData> InAssets)
+		{
+			if(TSharedPtr<FWorkspaceEditor> HostingApp = HostingAppPtr.Pin())
+			{
+				HostingApp->OpenAssets(InAssets);
+			}
+		});
+}
+
+TSharedRef<SWidget> FWorkspaceTabSummoner::CreateTabBody(const FWorkflowTabSpawnInfo& Info) const
+{
+	return WorkspaceView.ToSharedRef();
+}
+
+FText FWorkspaceTabSummoner::GetTabToolTipText(const FWorkflowTabSpawnInfo& Info) const
+{
+	return ViewMenuTooltip;
+}
+
+}
+
+#undef LOCTEXT_NAMESPACE
\ No newline at end of file
diff --git a/Engine/Plugins/Experimental/AnimNext/Source/AnimNextEditor/Private/Workspace/WorkspaceTabSummoner.h b/Engine/Plugins/Experimental/AnimNext/Source/AnimNextEditor/Private/Workspace/WorkspaceTabSummoner.h
new file mode 100644
--- /dev/null
+++ b/Engine/Plugins/Experimental/AnimNext/Source/AnimNextEditor/Private/Workspace/WorkspaceTabSummoner.h
@@ -0,0 +1,33 @@
+﻿// Copyright Epic Games, Inc. All Rights Reserved.
+
+#pragma once
+
+#include "WorkflowOrientedApp/WorkflowTabFactory.h"
+
+namespace UE::AnimNext::Editor
+{
+	class FWorkspaceEditor;
+	class SWorkspaceView;
+}
+
+namespace UE::AnimNext::Editor
+{
+
+struct FWorkspaceTabSummoner : public FWorkflowTabFactory
+{
+public:
+	FWorkspaceTabSummoner(TSharedPtr<FWorkspaceEditor> InHostingApp);
+
+private:
+	// FWorkflowTabFactory interface
+	virtual TSharedRef<SWidget> CreateTabBody(const FWorkflowTabSpawnInfo& Info) const override;
+	virtual FText GetTabToolTipText(const FWorkflowTabSpawnInfo& Info) const override;
+
+	// The hosting app
+	TWeakPtr<FWorkspaceEditor> HostingAppPtr;
+
+	// The widget this tab spawner wraps
+	TSharedPtr<SWorkspaceView> WorkspaceView;
+};
+
+}
\ No newline at end of file
diff --git a/Engine/Plugins/Experimental/AnimNext/Source/AnimNextTestSuite/Private/AnimNextDecoratorGraphTest.cpp b/Engine/Plugins/Experimental/AnimNext/Source/AnimNextTestSuite/Private/AnimNextDecoratorGraphTest.cpp
--- a/Engine/Plugins/Experimental/AnimNext/Source/AnimNextTestSuite/Private/AnimNextDecoratorGraphTest.cpp
+++ b/Engine/Plugins/Experimental/AnimNext/Source/AnimNextTestSuite/Private/AnimNextDecoratorGraphTest.cpp
@@ -89,7 +89,7 @@ bool FAnimationAnimNextRuntimeTest_GraphAddDecorator::RunTest(const FString& InP
 	EditorData->Initialize(/*bRecompileVM*/false);
 	EditorData->GetRigVMClient()->SetExecuteContextStruct(FAnimNextExecuteContext::StaticStruct());
 
-	URigVMController* Controller = EditorData->GetRigVMClient()->GetController(EditorData->GetRootGraph());
+	URigVMController* Controller = EditorData->GetRigVMClient()->GetController(EditorData->Graphs[0]);
 	UE_RETURN_ON_ERROR(Controller != nullptr, "FAnimationAnimNextRuntimeTest_GraphAddDecorator -> Failed to get RigVM controller");
 
 	// Create an empty decorator stack node
@@ -183,7 +183,7 @@ bool FAnimationAnimNextRuntimeTest_GraphExecute::RunTest(const FString& InParame
 	EditorData->Initialize(/*bRecompileVM*/false);
 	EditorData->GetRigVMClient()->SetExecuteContextStruct(FAnimNextExecuteContext::StaticStruct());
 
-	URigVMController* Controller = EditorData->GetRigVMClient()->GetController(EditorData->GetRootGraph());
+	URigVMController* Controller = EditorData->GetRigVMClient()->GetController(EditorData->Graphs[0]);
 	UE_RETURN_ON_ERROR(Controller != nullptr, "FAnimationAnimNextRuntimeTest_GraphExecute -> Failed to get RigVM controller");
 
 	// Add graph entry point
@@ -290,7 +290,7 @@ bool FAnimationAnimNextRuntimeTest_GraphExecuteLatent::RunTest(const FString& In
 	EditorData->Initialize(/*bRecompileVM*/false);
 	EditorData->GetRigVMClient()->SetExecuteContextStruct(FAnimNextExecuteContext::StaticStruct());
 
-	URigVMController* Controller = EditorData->GetRigVMClient()->GetController(EditorData->GetRootGraph());
+	URigVMController* Controller = EditorData->GetRigVMClient()->GetController(EditorData->Graphs[0]);
 	UE_RETURN_ON_ERROR(Controller != nullptr, "FAnimationAnimNextRuntimeTest_GraphExecuteLatent -> Failed to get RigVM controller");
 
 	// Add graph entry point
diff --git a/Engine/Plugins/Experimental/AnimNext/Source/AnimNextUncookedOnly/Internal/Graph/AnimNextGraphEntry.h b/Engine/Plugins/Experimental/AnimNext/Source/AnimNextUncookedOnly/Internal/Graph/AnimNextGraphEntry.h
new file mode 100644
--- /dev/null
+++ b/Engine/Plugins/Experimental/AnimNext/Source/AnimNextUncookedOnly/Internal/Graph/AnimNextGraphEntry.h
@@ -0,0 +1,49 @@
+﻿// Copyright Epic Games, Inc. All Rights Reserved.
+
+#pragma once
+
+#include "CoreMinimal.h"
+#include "AnimNextGraphEntry.generated.h"
+
+class UAnimNextGraph_EditorData;
+enum class ERigVMGraphNotifType : uint8;
+class URigVMGraph;
+
+namespace UE::AnimNext::Editor
+{
+	struct FUtils;
+}
+
+/** A single entry in an AnimNext graph asset */
+UCLASS(MinimalAPI, BlueprintType)
+class UAnimNextGraphEntry : public UObject
+{
+	GENERATED_BODY()
+
+	friend class UAnimNextGraph_EditorData;
+	friend struct UE::AnimNext::Editor::FUtils;
+	
+	void Initialize(UAnimNextGraph_EditorData* InEditorData);
+
+	void HandleRigVMGraphModifiedEvent(ERigVMGraphNotifType InNotifType, URigVMGraph* InGraph, UObject* InSubject);
+	
+	// Get the name to be displayed in the UI for this entry
+	virtual FText GetDisplayName() const;
+
+	// Get the tooltip to be displayed for the name in the UI for this entry
+	virtual FText GetDisplayNameTooltip() const;
+
+	// UObject interface
+	virtual bool IsAsset() const override;
+
+protected:
+	void BroadcastModified();
+
+	/** The name of the graph */
+	UPROPERTY(VisibleAnywhere, Category = Parameter)
+	FName GraphName;
+
+	/** Graph */
+	UPROPERTY()
+	TObjectPtr<URigVMGraph> Graph;
+};
\ No newline at end of file
diff --git a/Engine/Plugins/Experimental/AnimNext/Source/AnimNextUncookedOnly/Internal/Graph/AnimNextGraph_EditorData.h b/Engine/Plugins/Experimental/AnimNext/Source/AnimNextUncookedOnly/Internal/Graph/AnimNextGraph_EditorData.h
--- a/Engine/Plugins/Experimental/AnimNext/Source/AnimNextUncookedOnly/Internal/Graph/AnimNextGraph_EditorData.h
+++ b/Engine/Plugins/Experimental/AnimNext/Source/AnimNextUncookedOnly/Internal/Graph/AnimNextGraph_EditorData.h
@@ -7,10 +7,15 @@
 #include "AnimNextGraph_EdGraph.h"
 #include "RigVMCore/RigVMGraphFunctionHost.h"
 #include "RigVMBlueprint.h"
+#include "Kismet/BlueprintFunctionLibrary.h"
 #include "AnimNextGraph_EditorData.generated.h"
 
 class UAnimNextGraph;
 enum class ERigVMGraphNotifType : uint8;
+class FAnimationAnimNextRuntimeTest_GraphAddDecorator;
+class FAnimationAnimNextRuntimeTest_GraphExecute;
+class FAnimationAnimNextRuntimeTest_GraphExecuteLatent;
+
 namespace UE::AnimNext::UncookedOnly
 {
 	struct FUtils;
@@ -19,6 +24,8 @@ namespace UE::AnimNext::UncookedOnly
 namespace UE::AnimNext::Editor
 {
 	class FGraphEditor;
+	class SAnimNextGraphView;
+	struct FUtils;
 }
 
 enum class EAnimNextGraphLoadType : uint8
@@ -27,6 +34,12 @@ enum class EAnimNextGraphLoadType : uint8
 	CheckUserDefinedStructs
 };
 
+namespace UE::AnimNext::UncookedOnly
+{
+// A delegate for subscribing / reacting to graph modifications.
+DECLARE_MULTICAST_DELEGATE_OneParam(FOnGraphModified, UAnimNextGraph_EditorData* /* InEditorData */);
+}
+
 /**
  * The Schema is used to determine which actions are allowed
  * on a graph. This includes any topological change.
@@ -37,6 +50,17 @@ class UAnimNextGraph_Schema : public URigVMSchema
 	GENERATED_BODY()
 };
 
+// Script-callable editor API hoisted onto UAnimNextGraph
+UCLASS()
+class ANIMNEXTUNCOOKEDONLY_API UAnimNextGraphLibrary : public UBlueprintFunctionLibrary
+{
+	GENERATED_BODY()
+
+	/** Adds a graph to an AnimNext Graph asset */
+	UFUNCTION(BlueprintCallable, Category = "AnimNext|Graph", meta=(ScriptMethod))
+	static UAnimNextGraphEntry* AddGraph(UAnimNextGraph* InGraph, FName InName, bool bSetupUndoRedo = true, bool bPrintPythonCommand = true);
+};
+
 /**
   * Implements a RigVM client host and RigVM graph function store host
   * 
@@ -56,10 +80,18 @@ public:
 	
 	friend class UAnimNextGraphFactory;
 	friend class UAnimNextGraph_EdGraph;
+	friend class UAnimNextGraphEntry;
 	friend struct UE::AnimNext::UncookedOnly::FUtils;
+	friend struct UE::AnimNext::Editor::FUtils;
 	friend class UE::AnimNext::Editor::FGraphEditor;
+	friend class UE::AnimNext::Editor::SAnimNextGraphView;
 	friend struct FAnimNextGraphSchemaAction_RigUnit;
 	friend struct FAnimNextGraphSchemaAction_DispatchFactory;
+	friend class FAnimationAnimNextRuntimeTest_GraphAddDecorator;
+	friend class FAnimationAnimNextRuntimeTest_GraphExecute;
+	friend class FAnimationAnimNextRuntimeTest_GraphExecuteLatent;
+
+	ANIMNEXTUNCOOKEDONLY_API UAnimNextGraphEntry* AddGraph(FName InName, bool bSetupUndoRedo = true, bool bPrintPythonCommand = true);
 
 	// UObject interface
 	virtual void Serialize(FArchive& Ar) override;
@@ -76,22 +108,24 @@ public:
 	virtual IRigVMGraphFunctionHost* GetRigVMGraphFunctionHost() override;
 	virtual const IRigVMGraphFunctionHost* GetRigVMGraphFunctionHost() const override;
 	virtual void HandleRigVMGraphAdded(const FRigVMClient* InClient, const FString& InNodePath) override;
-	virtual void HandleRigVMGraphRemoved(const FRigVMClient* InClient, const FString& InNodePath) override {}
+	virtual void HandleRigVMGraphRemoved(const FRigVMClient* InClient, const FString& InNodePath) override;
 	virtual void HandleRigVMGraphRenamed(const FRigVMClient* InClient, const FString& InOldNodePath, const FString& InNewNodePath) override {}
 	virtual void HandleConfigureRigVMController(const FRigVMClient* InClient, URigVMController* InControllerToConfigure) override;
 	virtual UObject* GetEditorObjectForRigVMGraph(URigVMGraph* InVMGraph) const override;
-
+	virtual URigVMGraph* GetRigVMGraphForEditorObject(UObject* InObject) const override;
+	
 	// IRigVMGraphFunctionHost interface
 	virtual FRigVMGraphFunctionStore* GetRigVMGraphFunctionStore() override;
 	virtual const FRigVMGraphFunctionStore* GetRigVMGraphFunctionStore() const override;
 
+	UEdGraph* CreateEdGraph(URigVMGraph* InRigVMGraph, bool bForce);
+	bool RemoveEdGraph(URigVMGraph* InModel);
+
 	ANIMNEXTUNCOOKEDONLY_API void Initialize(bool bRecompileVM);
-	ANIMNEXTUNCOOKEDONLY_API UAnimNextGraph_EdGraph* GetRootGraph() const;
 
 protected:
 #if WITH_EDITOR
 	void RefreshAllModels(EAnimNextGraphLoadType InLoadType);
-	ANIMNEXTUNCOOKEDONLY_API void RebuildEdGraphFromModel();
 	void GetAllGraphs(TArray<UEdGraph*>& Graphs) const;
 #endif
 
@@ -107,20 +141,16 @@ protected:
 	
 	void HandleModifiedEvent(ERigVMGraphNotifType InNotifType, URigVMGraph* InGraph, UObject* InSubject);
 
-	ANIMNEXTUNCOOKEDONLY_API URigVMGraph* GetVMGraphForEdGraph(const UEdGraph* InGraph) const;
-
 	void CreateEdGraphForCollapseNode(URigVMCollapseNode* InNode);
 
-	bool IsNodeExecConnected(const UEdGraphNode* Node) const;
+	void BroadcastModified();
 
-	UPROPERTY()
-	TObjectPtr<UAnimNextGraph_EdGraph> RootGraph;
+	void ReportError(const TCHAR* InMessage) const;
 
-	UPROPERTY()
-	TObjectPtr<UAnimNextGraph_EdGraph> EntryPointGraph;
+	void ReconstructAllNodes();
 
 	UPROPERTY()
-	TObjectPtr<UAnimNextGraph_EdGraph> FunctionLibraryEdGraph;
+	TArray<TObjectPtr<UAnimNextGraph_EdGraph>> Graphs;
 
 	UPROPERTY()
 	FRigVMClient RigVMClient;
@@ -146,6 +176,10 @@ protected:
 	UPROPERTY(transient, DuplicateTransient)
 	TMap<FString, FRigVMOperand> PinToOperandMap;
 
+	/** All entries in this graph asset - not saved, either serialized or discovered at load time */
+	UPROPERTY(transient)
+	TArray<TObjectPtr<UAnimNextGraphEntry>> Entries;
+	
 	UPROPERTY(transient, DuplicateTransient)
 	int32 VMRecompilationBracket = 0;
 
@@ -157,13 +191,18 @@ protected:
 	
 	FCompilerResultsLog CompileLog;
 
-	FOnRigVMCompiledEvent VMCompiledEvent;
-	FRigVMGraphModifiedEvent ModifiedEvent;
+	FOnRigVMCompiledEvent RigVMCompiledEvent;
+	FRigVMGraphModifiedEvent RigVMGraphModifiedEvent;
+
+	// Delegate to subscribe to modifications to this graph
+	UE::AnimNext::UncookedOnly::FOnGraphModified ModifiedDelegate;
 
 	bool bAutoRecompileVM = true;
 	bool bErrorsDuringCompilation = false;
 	bool bSuspendModelNotificationsForSelf = false;
 	bool bSuspendModelNotificationsForOthers = false;
+	bool bSuspendPythonMessagesForRigVMClient = false;
 	bool bSuspendAllNotifications = false;
+	bool bSuspendGraphNotifications = false;
 	bool bCompileInDebugMode = false;
 };
diff --git a/Engine/Plugins/Experimental/AnimNext/Source/AnimNextUncookedOnly/Internal/Param/AnimNextParameterBlock_EditorData.h b/Engine/Plugins/Experimental/AnimNext/Source/AnimNextUncookedOnly/Internal/Param/AnimNextParameterBlock_EditorData.h
--- a/Engine/Plugins/Experimental/AnimNext/Source/AnimNextUncookedOnly/Internal/Param/AnimNextParameterBlock_EditorData.h
+++ b/Engine/Plugins/Experimental/AnimNext/Source/AnimNextUncookedOnly/Internal/Param/AnimNextParameterBlock_EditorData.h
@@ -35,8 +35,8 @@ namespace UE::AnimNext::Editor
 
 namespace UE::AnimNext::UncookedOnly
 {
-// A delegate for subscribing / reacting to parameter block modifications.
-DECLARE_MULTICAST_DELEGATE_OneParam(FOnParameterBlockModified, UAnimNextParameterBlock_EditorData* /* InEditorData */);
+	// A delegate for subscribing / reacting to parameter block modifications.
+	DECLARE_MULTICAST_DELEGATE_OneParam(FOnParameterBlockModified, UAnimNextParameterBlock_EditorData* /* InEditorData */);
 }
 
 enum class EAnimNextParameterLoadType : uint8
@@ -145,6 +145,7 @@ class UAnimNextParameterBlock_EditorData : public UObject, public IRigVMClientHo
 	virtual void HandleRigVMGraphRenamed(const FRigVMClient* InClient, const FString& InOldNodePath, const FString& InNewNodePath) override {}
 	virtual void HandleConfigureRigVMController(const FRigVMClient* InClient, URigVMController* InControllerToConfigure) override;
 	virtual UObject* GetEditorObjectForRigVMGraph(URigVMGraph* InVMGraph) const override;
+	virtual URigVMGraph* GetRigVMGraphForEditorObject(UObject* InObject) const override;
 
 	// IRigVMGraphFunctionHost interface
 	virtual FRigVMGraphFunctionStore* GetRigVMGraphFunctionStore() override;
@@ -172,10 +173,8 @@ class UAnimNextParameterBlock_EditorData : public UObject, public IRigVMClientHo
 	void IncrementVMRecompileBracket();
 
 	void DecrementVMRecompileBracket();
-	
-	void HandleModifiedEvent(ERigVMGraphNotifType InNotifType, URigVMGraph* InGraph, UObject* InSubject);
 
-	ANIMNEXTUNCOOKEDONLY_API URigVMGraph* GetVMGraphForEdGraph(const UEdGraph* InGraph) const;
+	void HandleModifiedEvent(ERigVMGraphNotifType InNotifType, URigVMGraph* InGraph, UObject* InSubject);
 
 	void CreateEdGraphForCollapseNode(URigVMCollapseNode* InNode);
 
diff --git a/Engine/Plugins/Experimental/AnimNext/Source/AnimNextUncookedOnly/Internal/Param/IAnimNextParameterBlockGraphInterface.h b/Engine/Plugins/Experimental/AnimNext/Source/AnimNextUncookedOnly/Internal/Param/IAnimNextParameterBlockGraphInterface.h
--- a/Engine/Plugins/Experimental/AnimNext/Source/AnimNextUncookedOnly/Internal/Param/IAnimNextParameterBlockGraphInterface.h
+++ b/Engine/Plugins/Experimental/AnimNext/Source/AnimNextUncookedOnly/Internal/Param/IAnimNextParameterBlockGraphInterface.h
@@ -14,6 +14,8 @@ namespace UE::AnimNext::Editor
 	class SParameterBlockView;
 	class SParameterBlockViewRow;
 	class FParameterBlockTabSummoner;
+	class FModule;
+	struct FUtils;
 }
 
 namespace UE::AnimNext::UncookedOnly
@@ -36,6 +38,8 @@ class ANIMNEXTUNCOOKEDONLY_API IAnimNextParameterBlockGraphInterface
 	friend class UE::AnimNext::Editor::SParameterBlockViewRow;
 	friend struct UE::AnimNext::UncookedOnly::FUtilsPrivate;
 	friend class UE::AnimNext::Editor::FParameterBlockTabSummoner;
+	friend class UE::AnimNext::Editor::FModule;
+	friend struct UE::AnimNext::Editor::FUtils;
 
 	// Get the graph name
 	virtual FName GetGraphName() const = 0;
diff --git a/Engine/Plugins/Experimental/AnimNext/Source/AnimNextUncookedOnly/Internal/Param/IAnimNextParameterBlockParameterInterface.h b/Engine/Plugins/Experimental/AnimNext/Source/AnimNextUncookedOnly/Internal/Param/IAnimNextParameterBlockParameterInterface.h
--- a/Engine/Plugins/Experimental/AnimNext/Source/AnimNextUncookedOnly/Internal/Param/IAnimNextParameterBlockParameterInterface.h
+++ b/Engine/Plugins/Experimental/AnimNext/Source/AnimNextUncookedOnly/Internal/Param/IAnimNextParameterBlockParameterInterface.h
@@ -19,6 +19,7 @@ namespace UE::AnimNext::Editor
 {
 	class SParameterBlockView;
 	class SParameterBlockViewRow;
+	struct FUtils;
 }
 
 namespace UE::AnimNext::UncookedOnly
@@ -40,6 +41,7 @@ class ANIMNEXTUNCOOKEDONLY_API IAnimNextParameterBlockParameterInterface
 	friend class UE::AnimNext::Editor::SParameterBlockView;
 	friend class UE::AnimNext::Editor::SParameterBlockViewRow;
 	friend struct UE::AnimNext::UncookedOnly::FUtils;
+	friend struct UE::AnimNext::Editor::FUtils;
 
 	// Get the parameter type
 	virtual FAnimNextParamType GetParamType() const = 0;
diff --git a/Engine/Plugins/Experimental/AnimNext/Source/AnimNextUncookedOnly/Internal/UncookedOnlyUtils.h b/Engine/Plugins/Experimental/AnimNext/Source/AnimNextUncookedOnly/Internal/UncookedOnlyUtils.h
--- a/Engine/Plugins/Experimental/AnimNext/Source/AnimNextUncookedOnly/Internal/UncookedOnlyUtils.h
+++ b/Engine/Plugins/Experimental/AnimNext/Source/AnimNextUncookedOnly/Internal/UncookedOnlyUtils.h
@@ -4,6 +4,7 @@
 
 #include "CoreMinimal.h"
 #include "Param/AnimNextParameterBlock_EditorData.h"
+#include "Graph/AnimNextGraph_EditorData.h"
 #include "Param/ParamTypeHandle.h"
 #include "RigVMCore/RigVMTemplate.h"
 
@@ -27,6 +28,8 @@ struct ANIMNEXTUNCOOKEDONLY_API FUtils
 	
 	static UAnimNextGraph_EditorData* GetEditorData(const UAnimNextGraph* InAnimNextGraph);
 	
+	static UAnimNextGraph* GetGraph(const UAnimNextGraph_EditorData* InEditorData);
+	
 	static void RecreateVM(UAnimNextGraph* InGraph);
 
 	/**
@@ -68,8 +71,11 @@ struct ANIMNEXTUNCOOKEDONLY_API FUtils
 	static FRigVMTemplateArgumentType GetRigVMArgTypeFromParamTypeHandle(const FParamTypeHandle& InParamTypeHandle);
 	static FRigVMTemplateArgumentType GetRigVMArgTypeFromParamType(const FAnimNextParamType& InParamType);
 
-	/** Set up a simple graph */
-	static void SetupGraph(URigVMController* InController);
+	/** Set up a simple animation graph */
+	static void SetupAnimGraph(URigVMController* InController);
+	
+	/** Set up a simple parameter graph */
+	static void SetupParameterGraph(URigVMController* InController);
 	
 	/** Set up a binding graph given the type to set */
 	static void SetupBindingGraph(URigVMController* InController, FName InParameterName, const FAnimNextParamType& InParamType);
@@ -92,6 +98,19 @@ struct ANIMNEXTUNCOOKEDONLY_API FUtils
 			OutNodes.Append(GraphNodes);
 		}
 	}
+
+	/** Returns all nodes in all graphs of the specified class */
+	template<class T>
+	static void GetAllNodesOfClass(const UAnimNextGraph_EditorData* InEditorData, TArray<T*>& OutNodes)
+	{
+		for(const UAnimNextGraph_EdGraph* Graph : InEditorData->Graphs)
+		{
+			check(Graph);
+			TArray<T*> GraphNodes;
+			Graph->GetNodesOfClass<T>(GraphNodes);
+			OutNodes.Append(GraphNodes);
+		}
+	}
 };
 
 }
\ No newline at end of file
diff --git a/Engine/Plugins/Experimental/AnimNext/Source/AnimNextUncookedOnly/Private/Graph/AnimNextGraphEntry.cpp b/Engine/Plugins/Experimental/AnimNext/Source/AnimNextUncookedOnly/Private/Graph/AnimNextGraphEntry.cpp
new file mode 100644
--- /dev/null
+++ b/Engine/Plugins/Experimental/AnimNext/Source/AnimNextUncookedOnly/Private/Graph/AnimNextGraphEntry.cpp
@@ -0,0 +1,40 @@
+﻿// Copyright Epic Games, Inc. All Rights Reserved.
+
+#include "Graph/AnimNextGraphEntry.h"
+#include "Graph/AnimNextGraph_EditorData.h"
+
+void UAnimNextGraphEntry::Initialize(UAnimNextGraph_EditorData* InEditorData)
+{
+	InEditorData->RigVMGraphModifiedEvent.RemoveAll(this);
+	InEditorData->RigVMGraphModifiedEvent.AddUObject(this, &UAnimNextGraphEntry::HandleRigVMGraphModifiedEvent);
+}
+
+void UAnimNextGraphEntry::HandleRigVMGraphModifiedEvent(ERigVMGraphNotifType InNotifType, URigVMGraph* InGraph, UObject* InSubject)
+{
+	
+}
+
+bool UAnimNextGraphEntry::IsAsset() const
+{
+	// Entries are considered assets to allow using the asset logic for save dialogs, etc.
+	// Also, they return true even if pending kill, in order to show up as deleted in these dialogs.
+	return IsPackageExternal() && !GetPackage()->HasAnyFlags(RF_Transient) && !HasAnyFlags(RF_Transient | RF_ClassDefaultObject);
+}
+
+FText UAnimNextGraphEntry::GetDisplayName() const
+{
+	return FText::FromName(GraphName);
+}
+
+FText UAnimNextGraphEntry::GetDisplayNameTooltip() const
+{
+	return FText::FromName(GraphName);
+}
+
+void UAnimNextGraphEntry::BroadcastModified()
+{
+	if(UAnimNextGraph_EditorData* EditorData = Cast<UAnimNextGraph_EditorData>(GetOuter()))
+	{
+		EditorData->BroadcastModified();
+	}
+}
\ No newline at end of file
diff --git a/Engine/Plugins/Experimental/AnimNext/Source/AnimNextUncookedOnly/Private/Graph/AnimNextGraph_EdGraph.cpp b/Engine/Plugins/Experimental/AnimNext/Source/AnimNextUncookedOnly/Private/Graph/AnimNextGraph_EdGraph.cpp
--- a/Engine/Plugins/Experimental/AnimNext/Source/AnimNextUncookedOnly/Private/Graph/AnimNextGraph_EdGraph.cpp
+++ b/Engine/Plugins/Experimental/AnimNext/Source/AnimNextUncookedOnly/Private/Graph/AnimNextGraph_EdGraph.cpp
@@ -5,10 +5,10 @@
 
 void UAnimNextGraph_EdGraph::Initialize(UAnimNextGraph_EditorData* InEditorData)
 {
-	InEditorData->ModifiedEvent.RemoveAll(this);
-	InEditorData->ModifiedEvent.AddUObject(this, &UAnimNextGraph_EdGraph::HandleModifiedEvent);
-	InEditorData->VMCompiledEvent.RemoveAll(this);
-	InEditorData->VMCompiledEvent.AddUObject(this, &UAnimNextGraph_EdGraph::HandleVMCompiledEvent);
+	InEditorData->RigVMGraphModifiedEvent.RemoveAll(this);
+	InEditorData->RigVMGraphModifiedEvent.AddUObject(this, &UAnimNextGraph_EdGraph::HandleModifiedEvent);
+	InEditorData->RigVMCompiledEvent.RemoveAll(this);
+	InEditorData->RigVMCompiledEvent.AddUObject(this, &UAnimNextGraph_EdGraph::HandleVMCompiledEvent);
 }
 
 FRigVMClient* UAnimNextGraph_EdGraph::GetRigVMClient() const
diff --git a/Engine/Plugins/Experimental/AnimNext/Source/AnimNextUncookedOnly/Private/Graph/AnimNextGraph_EditorData.cpp b/Engine/Plugins/Experimental/AnimNext/Source/AnimNextUncookedOnly/Private/Graph/AnimNextGraph_EditorData.cpp
--- a/Engine/Plugins/Experimental/AnimNext/Source/AnimNextUncookedOnly/Private/Graph/AnimNextGraph_EditorData.cpp
+++ b/Engine/Plugins/Experimental/AnimNext/Source/AnimNextUncookedOnly/Private/Graph/AnimNextGraph_EditorData.cpp
@@ -3,6 +3,7 @@
 #include "Graph/AnimNextGraph_EditorData.h"
 
 #include "ControlRigDefines.h"
+#include "ExternalPackageHelper.h"
 #include "Graph/AnimNextGraph.h"
 #include "Graph/AnimNextGraph_Controller.h"
 #include "Graph/AnimNextGraph_EdGraphSchema.h"
@@ -15,6 +16,31 @@
 #include "Curves/CurveFloat.h"
 #include "UObject/ObjectSaveContext.h"
 #include "Kismet2/BlueprintEditorUtils.h"
+#include "Graph/AnimNextGraphEntry.h"
+
+namespace UE::AnimNext::Graph::Private
+{
+
+static UAnimNextGraphEntry* CreateNewGraphEntry(UAnimNextGraph_EditorData* InEditorData, TSubclassOf<UAnimNextGraphEntry> InClass)
+{
+	UAnimNextGraphEntry* NewEntry = NewObject<UAnimNextGraphEntry>(InEditorData, InClass.Get(), NAME_None, RF_Transactional);
+	// If we are a transient asset, dont use external packages
+	UAnimNextGraph* Graph = UncookedOnly::FUtils::GetGraph(InEditorData);
+	check(Graph);
+	if(!Graph->HasAnyFlags(RF_Transient))
+	{
+		FExternalPackageHelper::SetPackagingMode(NewEntry, InEditorData, true, false, PKG_None);
+	}
+	return NewEntry;
+}
+
+template<typename EntryClassType>
+static EntryClassType* CreateNewGraphEntry(UAnimNextGraph_EditorData* InEditorData)
+{
+	return CastChecked<EntryClassType>(CreateNewGraphEntry(InEditorData, EntryClassType::StaticClass()));
+}
+
+}
 
 UAnimNextGraph_EditorData::UAnimNextGraph_EditorData(const FObjectInitializer& ObjectInitializer)
 {
@@ -25,26 +51,61 @@ UAnimNextGraph_EditorData::UAnimNextGraph_EditorData(const FObjectInitializer& O
 	{
 		TGuardValue<bool> DisableClientNotifs(RigVMClient.bSuspendNotifications, true);
 		RigVMClient.AddModel(TEXT("RigVMGraph"), false, &ObjectInitializer);
-		RigVMClient.AddModel(TEXT("TempRigVMGraph"), false, &ObjectInitializer);	// Used during compilation, otherwise empty
 		RigVMClient.GetOrCreateFunctionLibrary(false, &ObjectInitializer);
 	}
 	RigVMClient.SetExecuteContextStruct(FAnimNextExecuteContext::StaticStruct());
+}
+
+
+void UAnimNextGraph_EditorData::BroadcastModified()
+{
+	RecompileVM();
+
+	if(!bSuspendGraphNotifications)
+	{
+		ModifiedDelegate.Broadcast(this);
+	}
+}
+
+void UAnimNextGraph_EditorData::ReportError(const TCHAR* InMessage) const
+{
+#if WITH_EDITOR
+	FScriptExceptionHandler::Get().HandleException(ELogVerbosity::Error, InMessage, TEXT(""));
+#endif
+}
+
+void UAnimNextGraph_EditorData::ReconstructAllNodes()
+{
+#ifdef WITH_EDITORONLY_DATA
+	// Avoid refreshing EdGraph nodes during cook
+	if (GIsCookerLoadingPackage)
+	{
+		return;
+	}
 	
-	auto MakeEdGraph = [this, &ObjectInitializer](FName InName) -> UAnimNextGraph_EdGraph*
+	if (GetRigVMClient()->GetDefaultModel() == nullptr)
 	{
-		UAnimNextGraph_EdGraph* EdGraph = ObjectInitializer.CreateDefaultSubobject<UAnimNextGraph_EdGraph>(this, InName);
-		EdGraph->Schema = UAnimNextGraph_EdGraphSchema::StaticClass();
-		EdGraph->bAllowRenaming = false;
-		EdGraph->bEditable = true;
-		EdGraph->bAllowDeletion = false;
-		EdGraph->bIsFunctionDefinition = false;
-		EdGraph->Initialize(this);
+		return;
+	}
 
-		return EdGraph;
-	};
+	TArray<URigVMEdGraphNode*> AllNodes;
+	UE::AnimNext::UncookedOnly::FUtils::GetAllNodesOfClass(this, AllNodes);
 
-	RootGraph = MakeEdGraph(TEXT("RootEdGraph"));
-	FunctionLibraryEdGraph = MakeEdGraph(TEXT("RigVMFunctionLibraryEdGraph"));
+	for (URigVMEdGraphNode* Node : AllNodes)
+	{
+		Node->SetFlags(RF_Transient);
+	}
+
+	for(URigVMEdGraphNode* Node : AllNodes)
+	{
+		Node->ReconstructNode();
+	}
+
+	for (URigVMEdGraphNode* Node : AllNodes)
+	{
+		Node->ClearFlags(RF_Transient);
+	}
+#endif	
 }
 
 void UAnimNextGraph_EditorData::Serialize(FArchive& Ar)
@@ -65,7 +126,6 @@ void UAnimNextGraph_EditorData::Initialize(bool bRecompileVM)
 		
 		RigVMClient.GetOrCreateController(RigVMClient.GetDefaultModel());
 		RigVMClient.GetOrCreateController(RigVMClient.GetFunctionLibrary());
-		RigVMClient.GetOrCreateController(TEXT("TempRigVMGraph"));
 
 		// Init function library controllers
 		for(URigVMLibraryNode* LibraryNode : RigVMClient.GetFunctionLibrary()->GetFunctions())
@@ -79,19 +139,12 @@ void UAnimNextGraph_EditorData::Initialize(bool bRecompileVM)
 		}
 	}
 
-	RootGraph->Initialize(this);
-	FunctionLibraryEdGraph->Initialize(this);
-	if(EntryPointGraph)
+	for(UAnimNextGraph_EdGraph* Graph : Graphs)
 	{
-		EntryPointGraph->Initialize(this);
+		Graph->Initialize(this);
 	}
 }
 
-UAnimNextGraph_EdGraph* UAnimNextGraph_EditorData::GetRootGraph() const
-{
-	return RootGraph;
-}
-
 void UAnimNextGraph_EditorData::PostLoad()
 {
 	Super::PostLoad();
@@ -100,6 +153,12 @@ void UAnimNextGraph_EditorData::PostLoad()
 	RefreshAllModels(EAnimNextGraphLoadType::PostLoad);
 
 #if WITH_EDITOR
+	FExternalPackageHelper::LoadObjectsFromExternalPackages<UAnimNextGraphEntry>(this, [this](UAnimNextGraphEntry* InLoadedEntry)
+	{
+		check(IsValid(InLoadedEntry));
+		Entries.Add(InLoadedEntry);
+	});
+	
 	if (GIsEditor)
 	{
 		// delay compilation until the package has been loaded
@@ -125,6 +184,8 @@ void UAnimNextGraph_EditorData::HandlePackageDone()
 	FCoreUObjectDelegates::OnEndLoadPackage.RemoveAll(this);
 
 	RecompileVM();
+
+	ReconstructAllNodes();
 }
 
 void UAnimNextGraph_EditorData::RefreshAllModels(EAnimNextGraphLoadType InLoadType)
@@ -242,55 +303,10 @@ void UAnimNextGraph_EditorData::RefreshAllModels(EAnimNextGraphLoadType InLoadTy
 	}
 }
 
-void UAnimNextGraph_EditorData::RebuildEdGraphFromModel()
-{
-	DECLARE_SCOPE_HIERARCHICAL_COUNTER_FUNC()
-
-	TGuardValue<bool> SelfGuard(bSuspendModelNotificationsForSelf, true);
-	TGuardValue<bool> ClientIgnoreModificationsGuard(RigVMClient.bIgnoreModelNotifications, true);
-
-	verify(RigVMClient.GetOrCreateController((URigVMGraph*)nullptr));
-
-	TArray<UEdGraph*> EdGraphs;
-	GetAllGraphs(EdGraphs);
-
-	for (UEdGraph* Graph : EdGraphs)
-	{
-		TArray<UEdGraphNode*> Nodes = Graph->Nodes;
-		for (UEdGraphNode* Node : Nodes)
-		{
-			Graph->RemoveNode(Node);
-		}
-	}
-
-	if (FunctionLibraryEdGraph && RigVMClient.GetFunctionLibrary())
-	{
-		FunctionLibraryEdGraph->ModelNodePath = RigVMClient.GetFunctionLibrary()->GetNodePath();
-	}
-
-	TArray<URigVMGraph*> RigGraphs = RigVMClient.GetAllModels(true, true);
-
-	for (int32 RigGraphIndex = 0; RigGraphIndex < RigGraphs.Num(); RigGraphIndex++)
-	{
-		RigVMClient.GetOrCreateController(RigGraphs[RigGraphIndex])->ResendAllNotifications();
-	}
-}
-
-void UAnimNextGraph_EditorData::GetAllGraphs(TArray<UEdGraph*>& Graphs) const
+void UAnimNextGraph_EditorData::GetAllGraphs(TArray<UEdGraph*>& OutGraphs) const
 {
-	Graphs.Reset();
-
-	Graphs.Add(RootGraph);
-
-	if (EntryPointGraph)
-	{
-		Graphs.Add(EntryPointGraph);
-	}
-
-	if (FunctionLibraryEdGraph)
-	{
-		Graphs.Add(FunctionLibraryEdGraph);
-	}
+	OutGraphs.Reset();
+	OutGraphs.Append(Graphs);
 }
 
 #endif // WITH_EDITOR
@@ -321,6 +337,22 @@ void UAnimNextGraph_EditorData::HandleRigVMGraphAdded(const FRigVMClient* InClie
 	if(URigVMGraph* RigVMGraph = InClient->GetModel(InNodePath))
 	{
 		RigVMGraph->SetExecuteContextStruct(FAnimNextExecuteContext::StaticStruct());
+
+		if(!HasAnyFlags(RF_ClassDefaultObject | RF_NeedInitialization | RF_NeedLoad | RF_NeedPostLoad) &&
+			GetOuter() != GetTransientPackage())
+		{
+			CreateEdGraph(RigVMGraph, true);
+			RecompileVM();
+		}
+	}
+}
+
+void UAnimNextGraph_EditorData::HandleRigVMGraphRemoved(const FRigVMClient* InClient, const FString& InNodePath)
+{
+	if(URigVMGraph* RigVMGraph = InClient->GetModel(InNodePath))
+	{
+		RemoveEdGraph(RigVMGraph);
+		RecompileVM();
 	}
 }
 
@@ -381,8 +413,7 @@ UObject* UAnimNextGraph_EditorData::GetEditorObjectForRigVMGraph(URigVMGraph* In
 {
 	if(InVMGraph)
 	{
-		TArray<UAnimNextGraph_EdGraph*> AllGraphs = {RootGraph, EntryPointGraph, FunctionLibraryEdGraph};
-		for(UAnimNextGraph_EdGraph* EdGraph : AllGraphs)
+		for(UAnimNextGraph_EdGraph* EdGraph : Graphs)
 		{
 			if(EdGraph)
 			{
@@ -396,6 +427,26 @@ UObject* UAnimNextGraph_EditorData::GetEditorObjectForRigVMGraph(URigVMGraph* In
 	return nullptr;
 }
 
+URigVMGraph* UAnimNextGraph_EditorData::GetRigVMGraphForEditorObject(UObject* InObject) const
+{
+	if(const UAnimNextGraph_EdGraph* Graph = Cast<UAnimNextGraph_EdGraph>(InObject))
+	{
+		if (Graph->bIsFunctionDefinition)
+		{
+			if (URigVMLibraryNode* LibraryNode = RigVMClient.GetFunctionLibrary()->FindFunction(*Graph->ModelNodePath))
+			{
+				return LibraryNode->GetContainedGraph();
+			}
+		}
+		else
+		{
+			return RigVMClient.GetModel(Graph->ModelNodePath);
+		}
+	}
+
+	return nullptr;
+}
+
 FRigVMGraphFunctionStore* UAnimNextGraph_EditorData::GetRigVMGraphFunctionStore()
 {
 	return &GraphFunctionStore;
@@ -406,7 +457,6 @@ const FRigVMGraphFunctionStore* UAnimNextGraph_EditorData::GetRigVMGraphFunction
 	return &GraphFunctionStore;
 }
 
-
 void UAnimNextGraph_EditorData::RecompileVM()
 {
 	UE::AnimNext::UncookedOnly::FUtils::Compile(GetTypedOuter<UAnimNextGraph>());
@@ -520,75 +570,132 @@ void UAnimNextGraph_EditorData::HandleModifiedEvent(ERigVMGraphNotifType InNotif
 	// if the notification still has to be sent...
 	if (bNotifForOthersPending && !bSuspendModelNotificationsForOthers)
 	{
-		if (ModifiedEvent.IsBound())
+		if (RigVMGraphModifiedEvent.IsBound())
 		{
-			ModifiedEvent.Broadcast(InNotifType, InGraph, InSubject);
+			RigVMGraphModifiedEvent.Broadcast(InNotifType, InGraph, InSubject);
 		}
 	}
 }
 
-URigVMGraph* UAnimNextGraph_EditorData::GetVMGraphForEdGraph(const UEdGraph* InGraph) const
+void UAnimNextGraph_EditorData::CreateEdGraphForCollapseNode(URigVMCollapseNode* InNode)
 {
-	if (InGraph == RootGraph)
+	if (InNode->GetGraph()->IsA<URigVMFunctionLibrary>())
 	{
-		return RigVMClient.GetDefaultModel();
+		if (URigVMGraph* ContainedGraph = InNode->GetContainedGraph())
+		{
+			// create a sub graph
+			UAnimNextGraph_EdGraph* RigFunctionGraph = NewObject<UAnimNextGraph_EdGraph>(this, *InNode->GetName(), RF_Transactional);
+			RigFunctionGraph->Schema = UAnimNextGraph_EdGraphSchema::StaticClass();
+			RigFunctionGraph->bAllowRenaming = true;
+			RigFunctionGraph->bEditable = true;
+			RigFunctionGraph->bAllowDeletion = true;
+			RigFunctionGraph->ModelNodePath = ContainedGraph->GetNodePath();
+			RigFunctionGraph->bIsFunctionDefinition = true;
+		
+			RigFunctionGraph->Initialize(this);
+
+			RigVMClient.GetOrCreateController(ContainedGraph)->ResendAllNotifications();
+		}
 	}
-	else
+}
+
+UEdGraph* UAnimNextGraph_EditorData::CreateEdGraph(URigVMGraph* InRigVMGraph, bool bForce)
+{
+	check(InRigVMGraph);
+
+	if(bForce)
 	{
-		const UAnimNextGraph_EdGraph* Graph = Cast<UAnimNextGraph_EdGraph>(InGraph);
-		check(Graph);
+		RemoveEdGraph(InRigVMGraph);
+	}
 
-		if (Graph->bIsFunctionDefinition)
-		{
-			if (URigVMLibraryNode* LibraryNode = RigVMClient.GetFunctionLibrary()->FindFunction(*Graph->ModelNodePath))
-			{
-				return LibraryNode->GetContainedGraph();
-			}
-		}
+	FString GraphName = InRigVMGraph->GetName();
+
+	if(GraphName.IsEmpty())
+	{
+		GraphName = URigVMEdGraphSchema::GraphName_RigVM.ToString();
 	}
 
-	return nullptr;
+	GraphName = RigVMClient.GetUniqueName(*GraphName).ToString();
+
+	UAnimNextGraph_EdGraph* RigFunctionGraph = NewObject<UAnimNextGraph_EdGraph>(this, *GraphName, RF_Transactional);
+	RigFunctionGraph->Schema = UAnimNextGraph_EdGraphSchema::StaticClass();
+	RigFunctionGraph->bAllowDeletion = true;
+	RigFunctionGraph->bIsFunctionDefinition = false;
+	RigFunctionGraph->ModelNodePath = InRigVMGraph->GetNodePath();
+	RigFunctionGraph->Initialize(this);
+
+	Graphs.Add(RigFunctionGraph);
+
+	return RigFunctionGraph;
 }
 
-void UAnimNextGraph_EditorData::CreateEdGraphForCollapseNode(URigVMCollapseNode* InNode)
+bool UAnimNextGraph_EditorData::RemoveEdGraph(URigVMGraph* InModel)
 {
-	if (InNode->GetGraph()->IsA<URigVMFunctionLibrary>())
+	if(UAnimNextGraph_EdGraph* RigFunctionGraph = Cast<UAnimNextGraph_EdGraph>(GetEditorObjectForRigVMGraph(InModel)))
 	{
-		if (URigVMGraph* ContainedGraph = InNode->GetContainedGraph())
+		if(Graphs.Contains(RigFunctionGraph))
 		{
-			if (EntryPointGraph == nullptr)
-			{
-				// create a sub graph
-				UAnimNextGraph_EdGraph* RigFunctionGraph = NewObject<UAnimNextGraph_EdGraph>(this, *InNode->GetName(), RF_Transactional);
-				RigFunctionGraph->Schema = UAnimNextGraph_EdGraphSchema::StaticClass();
-				RigFunctionGraph->bAllowRenaming = true;
-				RigFunctionGraph->bEditable = true;
-				RigFunctionGraph->bAllowDeletion = true;
-				RigFunctionGraph->ModelNodePath = ContainedGraph->GetNodePath();
-				RigFunctionGraph->bIsFunctionDefinition = true;
-
-				EntryPointGraph = RigFunctionGraph;
-				RigFunctionGraph->Initialize(this);
-
-				RigVMClient.GetOrCreateController(ContainedGraph)->ResendAllNotifications();
-			}
+			Modify();
+			Graphs.Remove(RigFunctionGraph);
 		}
+		RigVMClient.DestroyObject(RigFunctionGraph);
+		return true;
 	}
+	return false;
+}
+
+UAnimNextGraphEntry* UAnimNextGraphLibrary::AddGraph(UAnimNextGraph* InGraph, FName InName, bool bSetupUndoRedo, bool bPrintPythonCommand)
+{
+	return UE::AnimNext::UncookedOnly::FUtils::GetEditorData(InGraph)->AddGraph(InName, bSetupUndoRedo, bPrintPythonCommand);
 }
 
-bool UAnimNextGraph_EditorData::IsNodeExecConnected(const UEdGraphNode* Node) const
+UAnimNextGraphEntry* UAnimNextGraph_EditorData::AddGraph(FName InName, bool bSetupUndoRedo, bool bPrintPythonCommand)
 {
-	static const TCHAR* ExecuteContextStr = TEXT(".ExecuteContext");
+	if(InName == NAME_None)
+	{
+		ReportError(TEXT("UAnimNextGraph_EditorData::AddGraph: Invalid graph name supplied."));
+		return nullptr;
+	}
 
-	for (const UEdGraphPin* Pin : Node->Pins)
+	// Check for duplicate name
+	FName NewGraphName = InName;
+	auto DuplicateNamePredicate = [&NewGraphName](const UAnimNextGraphEntry* InEntry)
 	{
-		const FString PinName = Pin->PinName.ToString();
+		return InEntry->GraphName == NewGraphName;
+	};
 
-		if (PinName.Contains(ExecuteContextStr) && Pin->LinkedTo.Num() > 0)
-		{
-			return true;
-		}
+	bool bAlreadyExists = Entries.ContainsByPredicate(DuplicateNamePredicate);
+	int32 NameNumber = InName.GetNumber() + 1;
+	while(bAlreadyExists)
+	{
+		NewGraphName = FName(InName, NameNumber);
+		bAlreadyExists =  Entries.ContainsByPredicate(DuplicateNamePredicate);
 	}
 
-	return false;
-}
+	UAnimNextGraphEntry* NewEntry = UE::AnimNext::Graph::Private::CreateNewGraphEntry<UAnimNextGraphEntry>(this);
+	NewEntry->GraphName = NewGraphName;
+
+	if(bSetupUndoRedo)
+	{
+		NewEntry->Modify();
+		Modify();
+	}
+	
+	Entries.Add(NewEntry);
+
+	// Add new graph
+	{
+		TGuardValue<bool> EnablePythonPrint(bSuspendPythonMessagesForRigVMClient, !bPrintPythonCommand);
+		TGuardValue<bool> DisableAutoCompile(bAutoRecompileVM, false);
+		URigVMGraph* NewGraph = RigVMClient.AddModel(NewGraphName, bSetupUndoRedo);
+		ensure(NewGraph);
+		NewEntry->Graph = NewGraph;
+
+		URigVMController* Controller = RigVMClient.GetController(NewGraph);
+		UE::AnimNext::UncookedOnly::FUtils::SetupAnimGraph(Controller);
+	}
+
+	BroadcastModified();
+
+	return NewEntry;
+}
\ No newline at end of file
diff --git a/Engine/Plugins/Experimental/AnimNext/Source/AnimNextUncookedOnly/Private/Param/AnimNextParameterBlock_EditorData.cpp b/Engine/Plugins/Experimental/AnimNext/Source/AnimNextUncookedOnly/Private/Param/AnimNextParameterBlock_EditorData.cpp
--- a/Engine/Plugins/Experimental/AnimNext/Source/AnimNextUncookedOnly/Private/Param/AnimNextParameterBlock_EditorData.cpp
+++ b/Engine/Plugins/Experimental/AnimNext/Source/AnimNextUncookedOnly/Private/Param/AnimNextParameterBlock_EditorData.cpp
@@ -181,7 +181,7 @@ UAnimNextParameterBlockGraph* UAnimNextParameterBlock_EditorData::AddGraph(FName
 		NewEntry->Graph = NewGraph;
 
 		URigVMController* Controller = RigVMClient.GetController(NewGraph);
-		UE::AnimNext::UncookedOnly::FUtils::SetupGraph(Controller);
+		UE::AnimNext::UncookedOnly::FUtils::SetupParameterGraph(Controller);
 	}
 	
 	BroadcastModified();
@@ -577,7 +577,7 @@ void UAnimNextParameterBlock_EditorData::Initialize(bool bRecompileVM)
 void UAnimNextParameterBlock_EditorData::PostLoad()
 {
 	Super::PostLoad();
-	
+
 	Initialize(/*bRecompileVM*/false);
 	RefreshAllModels(EAnimNextParameterLoadType::PostLoad);
 
@@ -895,6 +895,26 @@ UObject* UAnimNextParameterBlock_EditorData::GetEditorObjectForRigVMGraph(URigVM
 	return nullptr;
 }
 
+URigVMGraph* UAnimNextParameterBlock_EditorData::GetRigVMGraphForEditorObject(UObject* InObject) const
+{
+	if(const UAnimNextParameterBlock_EdGraph* Graph = Cast<UAnimNextParameterBlock_EdGraph>(InObject))
+	{
+		if (Graph->bIsFunctionDefinition)
+		{
+			if (URigVMLibraryNode* LibraryNode = RigVMClient.GetFunctionLibrary()->FindFunction(*Graph->ModelNodePath))
+			{
+				return LibraryNode->GetContainedGraph();
+			}
+		}
+		else
+		{
+			return RigVMClient.GetModel(Graph->ModelNodePath);
+		}
+	}
+
+	return nullptr;
+}
+
 FRigVMGraphFunctionStore* UAnimNextParameterBlock_EditorData::GetRigVMGraphFunctionStore()
 {
 	return &GraphFunctionStore;
@@ -1041,22 +1061,6 @@ void UAnimNextParameterBlock_EditorData::HandleModifiedEvent(ERigVMGraphNotifTyp
 	}
 }
 
-URigVMGraph* UAnimNextParameterBlock_EditorData::GetVMGraphForEdGraph(const UEdGraph* InGraph) const
-{
-	const UAnimNextParameterBlock_EdGraph* Graph = Cast<UAnimNextParameterBlock_EdGraph>(InGraph);
-	check(Graph);
-
-	if (Graph->bIsFunctionDefinition)
-	{
-		if (URigVMLibraryNode* LibraryNode = RigVMClient.GetFunctionLibrary()->FindFunction(*Graph->ModelNodePath))
-		{
-			return LibraryNode->GetContainedGraph();
-		}
-	}
-
-	return nullptr;
-}
-
 void UAnimNextParameterBlock_EditorData::CreateEdGraphForCollapseNode(URigVMCollapseNode* InNode)
 {
 	if (InNode->GetGraph()->IsA<URigVMFunctionLibrary>())
@@ -1109,7 +1113,7 @@ UEdGraph* UAnimNextParameterBlock_EditorData::CreateEdGraph(URigVMGraph* InRigVM
 	UAnimNextParameterBlock_EdGraph* RigFunctionGraph = NewObject<UAnimNextParameterBlock_EdGraph>(this, *GraphName, RF_Transactional);
 	RigFunctionGraph->Schema = UAnimNextParameterBlock_EdGraphSchema::StaticClass();
 	RigFunctionGraph->bAllowDeletion = true;
-	RigFunctionGraph->bIsFunctionDefinition = true;
+	RigFunctionGraph->bIsFunctionDefinition = false;
 	RigFunctionGraph->ModelNodePath = InRigVMGraph->GetNodePath();
 	RigFunctionGraph->Initialize(this);
 
diff --git a/Engine/Plugins/Experimental/AnimNext/Source/AnimNextUncookedOnly/Private/UncookedOnlyUtils.cpp b/Engine/Plugins/Experimental/AnimNext/Source/AnimNextUncookedOnly/Private/UncookedOnlyUtils.cpp
--- a/Engine/Plugins/Experimental/AnimNext/Source/AnimNextUncookedOnly/Private/UncookedOnlyUtils.cpp
+++ b/Engine/Plugins/Experimental/AnimNext/Source/AnimNextUncookedOnly/Private/UncookedOnlyUtils.cpp
@@ -477,7 +477,7 @@ void FUtils::Compile(UAnimNextGraph* InGraph)
 	EditorData->bVMRecompilationRequired = false;
 	if(InGraph->RigVM)
 	{
-		EditorData->VMCompiledEvent.Broadcast(InGraph, InGraph->RigVM, InGraph->ExtendedExecuteContext);
+		EditorData->RigVMCompiledEvent.Broadcast(InGraph, InGraph->RigVM, InGraph->ExtendedExecuteContext);
 	}
 
 	VMClient->RemoveController(VMTempGraph);
@@ -500,6 +500,13 @@ UAnimNextGraph_EditorData* FUtils::GetEditorData(const UAnimNextGraph* InAnimNex
 	return CastChecked<UAnimNextGraph_EditorData>(InAnimNextGraph->EditorData);
 }
 
+UAnimNextGraph* FUtils::GetGraph(const UAnimNextGraph_EditorData* InEditorData)
+{
+	check(InEditorData);
+
+	return CastChecked<UAnimNextGraph>(InEditorData->GetOuter());
+}
+
 FParamTypeHandle FUtils::GetParameterHandleFromPin(const FEdGraphPinType& InPinType)
 {
 	FAnimNextParamType::EValueType ValueType = FAnimNextParamType::EValueType::None;
@@ -1012,8 +1019,19 @@ FRigVMTemplateArgumentType FUtils::GetRigVMArgTypeFromParamType(const FAnimNextP
 	return ArgType;
 }
 
+void FUtils::SetupAnimGraph(URigVMController* InController)
+{
+	// Clear the graph
+	InController->RemoveNodes(InController->GetGraph()->GetNodes());
+
+	// Add root node
+	URigVMUnitNode* MainEntryPointNode = InController->AddUnitNode(FRigUnit_AnimNextGraphRoot::StaticStruct(), FRigUnit_AnimNextGraphRoot::EventName, FVector2D(-400.0f, 0.0f), FString(), false);
+	URigVMPin* BeginExecutePin = MainEntryPointNode->FindPin(GET_MEMBER_NAME_STRING_CHECKED(FRigUnit_AnimNextGraphRoot, Result));
+	check(BeginExecutePin);
+	check(BeginExecutePin->GetDirection() == ERigVMPinDirection::Input);
+}
 
-void FUtils::SetupGraph(URigVMController* InController)
+void FUtils::SetupParameterGraph(URigVMController* InController)
 {
 	// Clear the graph
 	InController->RemoveNodes(InController->GetGraph()->GetNodes());
diff --git a/Engine/Plugins/Runtime/RigVM/Source/RigVMDeveloper/Private/RigVMBlueprint.cpp b/Engine/Plugins/Runtime/RigVM/Source/RigVMDeveloper/Private/RigVMBlueprint.cpp
--- a/Engine/Plugins/Runtime/RigVM/Source/RigVMDeveloper/Private/RigVMBlueprint.cpp
+++ b/Engine/Plugins/Runtime/RigVM/Source/RigVMDeveloper/Private/RigVMBlueprint.cpp
@@ -421,6 +421,26 @@ UObject* URigVMBlueprint::GetEditorObjectForRigVMGraph(URigVMGraph* InVMGraph) c
 	return nullptr;
 }
 
+URigVMGraph* URigVMBlueprint::GetRigVMGraphForEditorObject(UObject* InObject) const
+{
+	if(URigVMEdGraph* RigVMEdGraph = Cast<URigVMEdGraph>(InObject))
+	{
+		if (RigVMEdGraph->bIsFunctionDefinition)
+		{
+			if (URigVMLibraryNode* LibraryNode = RigVMClient.GetFunctionLibrary()->FindFunction(*RigVMEdGraph->ModelNodePath))
+			{
+				return LibraryNode->GetContainedGraph();
+			}
+		}
+		else
+		{
+			return RigVMClient.GetModel(RigVMEdGraph->ModelNodePath);
+		}
+	}
+
+	return nullptr;
+}
+
 void URigVMBlueprint::HandleRigVMGraphAdded(const FRigVMClient* InClient, const FString& InNodePath)
 {
 	if(URigVMGraph* Model = InClient->GetModel(InNodePath))
diff --git a/Engine/Plugins/Runtime/RigVM/Source/RigVMDeveloper/Public/RigVMBlueprint.h b/Engine/Plugins/Runtime/RigVM/Source/RigVMDeveloper/Public/RigVMBlueprint.h
--- a/Engine/Plugins/Runtime/RigVM/Source/RigVMDeveloper/Public/RigVMBlueprint.h
+++ b/Engine/Plugins/Runtime/RigVM/Source/RigVMDeveloper/Public/RigVMBlueprint.h
@@ -310,6 +310,7 @@ public:
 	virtual IRigVMGraphFunctionHost* GetRigVMGraphFunctionHost() override;
 	virtual const IRigVMGraphFunctionHost* GetRigVMGraphFunctionHost() const override;
 	virtual UObject* GetEditorObjectForRigVMGraph(URigVMGraph* InVMGraph) const override;
+	virtual URigVMGraph* GetRigVMGraphForEditorObject(UObject* InObject) const override;
 	virtual void HandleRigVMGraphAdded(const FRigVMClient* InClient, const FString& InNodePath) override;
 	virtual void HandleRigVMGraphRemoved(const FRigVMClient* InClient, const FString& InNodePath) override;
 	virtual void HandleRigVMGraphRenamed(const FRigVMClient* InClient, const FString& InOldNodePath, const FString& InNewNodePath) override;
diff --git a/Engine/Plugins/Runtime/RigVM/Source/RigVMDeveloper/Public/RigVMModel/RigVMClient.h b/Engine/Plugins/Runtime/RigVM/Source/RigVMDeveloper/Public/RigVMModel/RigVMClient.h
--- a/Engine/Plugins/Runtime/RigVM/Source/RigVMDeveloper/Public/RigVMModel/RigVMClient.h
+++ b/Engine/Plugins/Runtime/RigVM/Source/RigVMDeveloper/Public/RigVMModel/RigVMClient.h
@@ -35,9 +35,12 @@ public:
 	// Returns the rigvm function host
 	virtual const IRigVMGraphFunctionHost* GetRigVMGraphFunctionHost() const = 0;
 
-	// Returns the editor object corresponding with the supplied editor object
+	// Returns the editor object corresponding with the supplied RigVM graph
 	virtual UObject* GetEditorObjectForRigVMGraph(URigVMGraph* InVMGraph) const = 0;
 
+	// Returns the RigVM graph corresponding with the supplied editor object
+	virtual URigVMGraph* GetRigVMGraphForEditorObject(UObject* InObject) const = 0;
+
 	// Reacts to adding a graph
 	virtual void HandleRigVMGraphAdded(const FRigVMClient* InClient, const FString& InNodePathOrName) = 0;
 
