FIX CIS ensure at RigVMHost, due to VM being serialized were we expect the pointer to be the CDO VM

#rnx

diff --git a/Engine/Plugins/Experimental/AnimNext/Source/AnimNext/Internal/Param/AnimNextParameterBlock.h b/Engine/Plugins/Experimental/AnimNext/Source/AnimNext/Internal/Param/AnimNextParameterBlock.h
--- a/Engine/Plugins/Experimental/AnimNext/Source/AnimNext/Internal/Param/AnimNextParameterBlock.h
+++ b/Engine/Plugins/Experimental/AnimNext/Source/AnimNext/Internal/Param/AnimNextParameterBlock.h
@@ -91,6 +91,10 @@ class UAnimNextParameterBlock : public URigVMHost, public IAnimNextParameterSour
 	
 	FRigVMExtendedExecuteContext BaseRigVMContext;
 
+	UPROPERTY()
+	TObjectPtr<URigVM> RigVM;
+
+
 	UPROPERTY()
 	FInstancedPropertyBag PropertyBag;
 
diff --git a/Engine/Plugins/Experimental/AnimNext/Source/AnimNext/Private/Param/AnimNextParameterBlock.cpp b/Engine/Plugins/Experimental/AnimNext/Source/AnimNext/Private/Param/AnimNextParameterBlock.cpp
--- a/Engine/Plugins/Experimental/AnimNext/Source/AnimNext/Private/Param/AnimNextParameterBlock.cpp
+++ b/Engine/Plugins/Experimental/AnimNext/Source/AnimNext/Private/Param/AnimNextParameterBlock.cpp
@@ -119,6 +119,8 @@ void UAnimNextParameterBlock::PostLoad()
 {
 	Super::PostLoad();
 
+	VM = RigVM;
+
 	// In packaged builds, initialize the VM
 	// In editor, the VM will be recompiled and initialized at UAnimNextParameterBlock_EditorData::HandlePackageDone::RecompileVM
 #if !WITH_EDITOR
diff --git a/Engine/Plugins/Experimental/AnimNext/Source/AnimNext/Private/RigVMRuntimeDataRegistry.cpp b/Engine/Plugins/Experimental/AnimNext/Source/AnimNext/Private/RigVMRuntimeDataRegistry.cpp
--- a/Engine/Plugins/Experimental/AnimNext/Source/AnimNext/Private/RigVMRuntimeDataRegistry.cpp
+++ b/Engine/Plugins/Experimental/AnimNext/Source/AnimNext/Private/RigVMRuntimeDataRegistry.cpp
@@ -12,7 +12,7 @@ namespace Private
 static bool bInitialized = false;
 static FDelegateHandle PostGarbageCollectHandle;
 
-static std::atomic<uint32> GCCycle = 0; // Main thread GC counter, incremented when a VM has been deleted
+static std::atomic<uint32> GCCycle = 0; // Main thread GC counter, incremented on each main thread GC cycle
 
 static thread_local uint32 LocalGCCycle = 0;	// Local thread GC counter, used to compare with main and trigger compaction if different
 static thread_local TMap<FRigVMRuntimeDataID, FRigVMRuntimeData> RuntimeDataStorage;
diff --git a/Engine/Plugins/Experimental/AnimNext/Source/AnimNextUncookedOnly/Private/UncookedOnlyUtils.cpp b/Engine/Plugins/Experimental/AnimNext/Source/AnimNextUncookedOnly/Private/UncookedOnlyUtils.cpp
--- a/Engine/Plugins/Experimental/AnimNext/Source/AnimNextUncookedOnly/Private/UncookedOnlyUtils.cpp
+++ b/Engine/Plugins/Experimental/AnimNext/Source/AnimNextUncookedOnly/Private/UncookedOnlyUtils.cpp
@@ -724,6 +724,7 @@ void FUtils::RecreateVM(UAnimNextParameterBlock* InParameterBlock)
 {
 	InParameterBlock->VM = NewObject<URigVM>(InParameterBlock, TEXT("VM"), RF_NoFlags);
 	InParameterBlock->VM->Reset(InParameterBlock->GetRigVMExtendedExecuteContext());
+	InParameterBlock->RigVM = InParameterBlock->VM; // Local serialization
 }
 
 UAnimNextParameterBlock_EditorData* FUtils::GetEditorData(const UAnimNextParameterBlock* InParameterBlock)
diff --git a/Engine/Plugins/Runtime/RigVM/Source/RigVM/Public/RigVMHost.h b/Engine/Plugins/Runtime/RigVM/Source/RigVM/Public/RigVMHost.h
--- a/Engine/Plugins/Runtime/RigVM/Source/RigVM/Public/RigVMHost.h
+++ b/Engine/Plugins/Runtime/RigVM/Source/RigVM/Public/RigVMHost.h
@@ -325,7 +325,7 @@ protected:
 	/** true if we should increase the AbsoluteTime */
 	bool bAccumulateTime;
 
-	UPROPERTY()
+	UPROPERTY(Transient)
 	TObjectPtr<URigVM> VM;
 
 #if WITH_EDITOR
