Clean up LOD pose API to make it clearer and avoid confusion with mesh/skeleton bone indices
Fix incorrect LOD pose usage

diff --git a/Engine/Plugins/Experimental/AnimNext/Source/AnimNext/Private/DecompressionTools.cpp b/Engine/Plugins/Experimental/AnimNext/Source/AnimNext/Private/DecompressionTools.cpp
--- a/Engine/Plugins/Experimental/AnimNext/Source/AnimNext/Private/DecompressionTools.cpp
+++ b/Engine/Plugins/Experimental/AnimNext/Source/AnimNext/Private/DecompressionTools.cpp
@@ -127,7 +127,7 @@ void FDecompressionTools::GetBonePose(const UAnimSequence* AnimSequence, const F
 	SCOPE_CYCLE_COUNTER(STAT_AnimSeq_GetBonePose);
 	CSV_SCOPED_TIMING_STAT(Animation, AnimSeq_GetBonePose);
 
-	const TArrayView<const FBoneIndexType> LODBoneIndexes = OutAnimationPoseData.GetLODBoneIndexes();
+	const TArrayView<const FBoneIndexType> LODBoneIndexToSkeletonBoneIndexMap = OutAnimationPoseData.GetLODBoneIndexToSkeletonBoneIndexMap();
 
 	check(!bForceUseRawData || CanEvaluateRawAnimationData(AnimSequence));
 	const bool bUseRawDataForPoseExtraction = (CanEvaluateRawAnimationData(AnimSequence) && bForceUseRawData) ||
@@ -158,19 +158,19 @@ void FDecompressionTools::GetBonePose(const UAnimSequence* AnimSequence, const F
 		// if retargeting is disabled, we initialize pose with 'Retargeting Source' ref pose.
 		if (bDisableRetargeting)
 		{
-			TArray<FTransform> const& AuthoredOnRefSkeleton = AnimSequence->GetRetargetTransforms();
+			const TArray<FTransform>& AuthoredOnRefSkeleton = AnimSequence->GetRetargetTransforms();
 
-			const int32 NumBones = LODBoneIndexes.Num();
-			const int32 NumRawBones = AnimSequence->GetSkeleton()->GetReferenceSkeleton().GetRawBoneNum();
+			const int32 NumLODBones = LODBoneIndexToSkeletonBoneIndexMap.Num();
+			const int32 NumRawSkeletonBones = AnimSequence->GetSkeleton()->GetReferenceSkeleton().GetRawBoneNum();
 
-			for (int i = 0; i < NumBones; i++)
+			for (int LODBoneIndex = 0; LODBoneIndex < NumLODBones; LODBoneIndex++)
 			{
-				const int32 SkeletonBoneIndex = LODBoneIndexes[i];
+				const int32 SkeletonBoneIndex = LODBoneIndexToSkeletonBoneIndexMap[LODBoneIndex];
 
 				// Virtual bones are part of the retarget transform pose, so if the pose has not been updated (recently) there might be a mismatch
-				if (SkeletonBoneIndex < NumRawBones || AuthoredOnRefSkeleton.IsValidIndex(SkeletonBoneIndex))
+				if (SkeletonBoneIndex < NumRawSkeletonBones || AuthoredOnRefSkeleton.IsValidIndex(SkeletonBoneIndex))
 				{
-					OutAnimationPoseData.LocalTransformsView[i] = AuthoredOnRefSkeleton[SkeletonBoneIndex];
+					OutAnimationPoseData.LocalTransformsView[LODBoneIndex] = AuthoredOnRefSkeleton[SkeletonBoneIndex];
 				}
 			}
 		}
@@ -276,9 +276,9 @@ void FDecompressionTools::DecompressPose(FLODPose& OutAnimationPoseData,
 										const FRootMotionReset& RootMotionReset)
 {
 	const FReferencePose& ReferencePose = OutAnimationPoseData.GetRefPose();
-	const TArrayView<const FBoneIndexType> LODBoneIndexes = OutAnimationPoseData.GetLODBoneIndexes();
-	const TArrayView<const FBoneIndexType> SkeletonToLODBoneIndexes = ReferencePose.GetSkeletonToLODBoneIndexes(0); // Full list of Skeleton to LOD conversion
-	const int32 NumLODBoneIndexes = LODBoneIndexes.Num();
+	const TArrayView<const FBoneIndexType> LODBoneIndexToSkeletonBoneIndexMap = OutAnimationPoseData.GetLODBoneIndexToMeshBoneIndexMap();
+	const TArrayView<const FBoneIndexType> SkeletonToLODBoneIndexes = ReferencePose.GetSkeletonBoneIndexToLODBoneIndexMap(0); // Full list of Skeleton to LOD conversion
+	const int32 NumLODBoneIndexes = LODBoneIndexToSkeletonBoneIndexMap.Num();
 
 	const int32 NumTracks = CompressedData.CompressedTrackToSkeletonMapTable.Num();
 
@@ -299,7 +299,7 @@ void FDecompressionTools::DecompressPose(FLODPose& OutAnimationPoseData,
 	OrientAndScaleRetargetingPairs.Reset();
 
 	// Optimization: assuming first index is root bone. That should always be the case in Skeletons.
-	checkSlow((LODBoneIndexes[0] == FMeshPoseBoneIndex(0).GetInt()));
+	checkSlow((LODBoneIndexToSkeletonBoneIndexMap[0] == FMeshPoseBoneIndex(0).GetInt()));
 	// this is not guaranteed for AnimSequences though... If Root is not animated, Track will not exist.
 	const bool bFirstTrackIsRootBone = (CompressedData.GetSkeletonIndexFromTrackIndex(0) == 0);
 
@@ -443,7 +443,7 @@ void FDecompressionTools::DecompressPose(FLODPose& OutAnimationPoseData,
 		{
 			for (int32 LODBoneIndex = (bFirstTrackIsRootBone ? 1 : 0); LODBoneIndex < LODNumBones; ++LODBoneIndex)
 			{
-				const int32 TargetSkeletonBoneIndex = LODBoneIndexes[LODBoneIndex]; // ReferencePose.GetSkeletonBoneIndexFromLODBoneIndex(LODBoneIndex);
+				const int32 TargetSkeletonBoneIndex = LODBoneIndexToSkeletonBoneIndexMap[LODBoneIndex]; // ReferencePose.GetSkeletonBoneIndexFromLODBoneIndex(LODBoneIndex);
 				OutAnimationPoseData.LocalTransformsView[LODBoneIndex].SetRotation(SkeletonRemapping.RetargetBoneRotationToTargetSkeleton(TargetSkeletonBoneIndex, OutAnimationPoseData.LocalTransformsView[LODBoneIndex].GetRotation()));
 				if (TargetSkeleton->GetBoneTranslationRetargetingMode(TargetSkeletonBoneIndex, OutAnimationPoseData.GetDisableRetargeting()) != EBoneTranslationRetargetingMode::Skeleton)
 				{
diff --git a/Engine/Plugins/Experimental/AnimNext/Source/AnimNext/Private/EvaluationVM/EvaluationVM.cpp b/Engine/Plugins/Experimental/AnimNext/Source/AnimNext/Private/EvaluationVM/EvaluationVM.cpp
--- a/Engine/Plugins/Experimental/AnimNext/Source/AnimNext/Private/EvaluationVM/EvaluationVM.cpp
+++ b/Engine/Plugins/Experimental/AnimNext/Source/AnimNext/Private/EvaluationVM/EvaluationVM.cpp
@@ -35,7 +35,7 @@ namespace UE::AnimNext
 	{
 		const UE::Anim::FCurveFilterSettings CurveFilterSettings(UE::Anim::ECurveFilterMode::DisallowAll);
 		USkeleton* Skeleton = const_cast<USkeleton*>(InReferencePose.Skeleton.Get());	// const_cast because the bone container takes a mutable reference
-		BoneContainer.InitializeTo(InReferencePose.GetLODBoneIndexes(InCurrentLOD), CurveFilterSettings, *Skeleton);
+		BoneContainer.InitializeTo(InReferencePose.GetLODBoneIndexToMeshBoneIndexMap(InCurrentLOD), CurveFilterSettings, *Skeleton);
 
 		// TODO: In AnimInstanceProxy this is how we initialize the bone container, we need to get the component somehow or we
 		// gotta figure out how to support ref pose overrides
diff --git a/Engine/Plugins/Experimental/AnimNext/Source/AnimNext/Private/GenerationTools.cpp b/Engine/Plugins/Experimental/AnimNext/Source/AnimNext/Private/GenerationTools.cpp
--- a/Engine/Plugins/Experimental/AnimNext/Source/AnimNext/Private/GenerationTools.cpp
+++ b/Engine/Plugins/Experimental/AnimNext/Source/AnimNext/Private/GenerationTools.cpp
@@ -93,14 +93,14 @@ struct FCompareBoneIndexType
 			bCanGenerateSingleBonesList &= CheckExcludedBones(NumLODs, GenerationLODData, SkeletalMesh);
 			//bCanGenerateSingleBonesList &= CheckExcludedBones(NumLODs, ComponentSpaceGenerationLODData, SkeletalMesh); // Commented : right now we only support skeletal meshes with all the sockets set to always animate
 
-			TArray<FBoneIndexType> OrderedBoneList;
+			TArray<FBoneIndexType> LODBoneIndexToMeshBoneIndexMap;
 			TArray<FBoneIndexType> ComponentSpaceOrderedBoneList;
 
 			if (bCanGenerateSingleBonesList)
 			{
-				bCanGenerateSingleBonesList &= GenerateOrderedBoneList(SkeletalMesh, GenerationLODData, OrderedBoneList);
+				bCanGenerateSingleBonesList &= GenerateOrderedBoneList(SkeletalMesh, GenerationLODData, LODBoneIndexToMeshBoneIndexMap);
 				//bCanGenerateSingleBonesList &= GenerateOrderedBoneList(SkeletalMesh, ComponentSpaceGenerationLODData, ComponentSpaceOrderedBoneList); // Commented : right now we only support skeletal meshes with all the sockets set to always animate
-				//bCanGenerateSingleBonesList &= OrderedBoneList == ComponentSpaceOrderedBoneList; // Commented : right now we only support skeletal meshes with all the sockets set to always animate
+				//bCanGenerateSingleBonesList &= LODBoneIndexToMeshBoneIndexMap == ComponentSpaceOrderedBoneList; // Commented : right now we only support skeletal meshes with all the sockets set to always animate
 			}
 
 			if (bCanGenerateSingleBonesList)
@@ -125,23 +125,27 @@ struct FCompareBoneIndexType
 				const FSkeletonToMeshLinkup& LinkupTable = Skeleton->LinkupCache[SkelMeshLinkupIndex];
 
 				// Generate a Skeleton to LOD look up table
-				const int32 NumOrderedBones = OrderedBoneList.Num();
+				const int32 NumOrderedBones = LODBoneIndexToMeshBoneIndexMap.Num();
+
+				TArray<FBoneIndexType> SkeletonBoneIndexToLODBoneIndexMap;
+				TArray<FBoneIndexType> LODBoneIndexToSkeletonBoneIndexMap;
+				SkeletonBoneIndexToLODBoneIndexMap.SetNumZeroed(NumOrderedBones);
+				LODBoneIndexToSkeletonBoneIndexMap.SetNumZeroed(NumOrderedBones);
 
-				TArray<FBoneIndexType> SkeletonToLODBoneList;
-				SkeletonToLODBoneList.SetNumZeroed(NumOrderedBones);
 				for (int32 LODBoneIndex = 0; LODBoneIndex < NumOrderedBones; ++LODBoneIndex)
 				{
 					// The ordered list contains skeletal mesh bone indices sorted by LOD
-					const FMeshPoseBoneIndex MeshBoneIndex(OrderedBoneList[LODBoneIndex]);
+					const FMeshPoseBoneIndex MeshBoneIndex(LODBoneIndexToMeshBoneIndexMap[LODBoneIndex]);
 
 					// Remap our skeletal mesh bone index into the skeleton bone index we output for
 					const FSkeletonPoseBoneIndex SkeletonBoneIndex(LinkupTable.MeshToSkeletonTable[MeshBoneIndex.GetInt()]);
 					ensure(SkeletonBoneIndex.IsValid());	// We expect the skeletal mesh bone to map to a valid skeleton bone
 
-					SkeletonToLODBoneList[static_cast<FBoneIndexType>(SkeletonBoneIndex.GetInt())] = LODBoneIndex;
+					SkeletonBoneIndexToLODBoneIndexMap[static_cast<FBoneIndexType>(SkeletonBoneIndex.GetInt())] = LODBoneIndex;
+					LODBoneIndexToSkeletonBoneIndexMap[LODBoneIndex] = static_cast<FBoneIndexType>(SkeletonBoneIndex.GetInt());
 				}
 
-				OutAnimationReferencePose.Initialize(SkeletalMesh->GetRefSkeleton(), { OrderedBoneList }, { SkeletonToLODBoneList }, LODNumBones, bCanGenerateSingleBonesList);
+				OutAnimationReferencePose.Initialize(SkeletalMesh->GetRefSkeleton(), { LODBoneIndexToMeshBoneIndexMap }, { LODBoneIndexToSkeletonBoneIndexMap }, { SkeletonBoneIndexToLODBoneIndexMap }, LODNumBones, bCanGenerateSingleBonesList);
 
 				ReferencePoseGenerated = true;
 			}
@@ -406,7 +410,7 @@ struct FCompareBoneIndexType
 {
 	const FBoneContainer& BoneContainer = SourcePose.Pose.GetBoneContainer();
 	const FReferencePose& RefPose = TargetPose.GetRefPose();
-	const TArrayView<const FBoneIndexType> LODBoneIndexes = RefPose.GetLODBoneIndexes(TargetPose.LODLevel);
+	const TArrayView<const FBoneIndexType> LODBoneIndexes = RefPose.GetLODBoneIndexToMeshBoneIndexMap(TargetPose.LODLevel);
 	const int32 NumLODBones = LODBoneIndexes.Num();
 
 	check(TargetPose.GetNumBones() == NumLODBones);
@@ -441,7 +445,7 @@ struct FCompareBoneIndexType
 {
 	const FBoneContainer& BoneContainer = TargetPose.Pose.GetBoneContainer();
 	const FReferencePose& RefPose = SourcePose.GetRefPose();
-	const TArrayView<const FBoneIndexType> LODBoneIndexes = RefPose.GetLODBoneIndexes(SourcePose.LODLevel);
+	const TArrayView<const FBoneIndexType> LODBoneIndexes = RefPose.GetLODBoneIndexToMeshBoneIndexMap(SourcePose.LODLevel);
 	const int32 NumLODBones = LODBoneIndexes.Num();
 
 	check(SourcePose.GetNumBones() == NumLODBones);
diff --git a/Engine/Plugins/Experimental/AnimNext/Source/AnimNext/Public/LODPose.h b/Engine/Plugins/Experimental/AnimNext/Source/AnimNext/Public/LODPose.h
--- a/Engine/Plugins/Experimental/AnimNext/Source/AnimNext/Public/LODPose.h
+++ b/Engine/Plugins/Experimental/AnimNext/Source/AnimNext/Public/LODPose.h
@@ -81,11 +81,11 @@ struct FLODPose
 		return RefPose != nullptr ? RefPose->GetNumBonesForLOD(LODLevel) : 0;
 	}
 
-	const TArrayView<const FBoneIndexType> GetLODBoneIndexes() const
+	const TArrayView<const FBoneIndexType> GetLODBoneIndexToMeshBoneIndexMap() const
 	{
 		if (LODLevel != INVALID_LOD_LEVEL && RefPose != nullptr)
 		{
-			return RefPose->GetLODBoneIndexes(LODLevel);
+			return RefPose->GetLODBoneIndexToMeshBoneIndexMap(LODLevel);
 		}
 		else
 		{
@@ -93,11 +93,23 @@ struct FLODPose
 		}
 	}
 
-	const TArrayView<const FBoneIndexType> GetSkeletonToLODBoneIndexes() const
+	const TArrayView<const FBoneIndexType> GetLODBoneIndexToSkeletonBoneIndexMap() const
 	{
 		if (LODLevel != INVALID_LOD_LEVEL && RefPose != nullptr)
 		{
-			return RefPose->GetSkeletonToLODBoneIndexes(LODLevel);
+			return RefPose->GetLODBoneIndexToSkeletonBoneIndexMap(LODLevel);
+		}
+		else
+		{
+			return TArrayView<const FBoneIndexType>();
+		}
+	}
+
+	const TArrayView<const FBoneIndexType> GetSkeletonBoneIndexToLODBoneIndexMap() const
+	{
+		if (LODLevel != INVALID_LOD_LEVEL && RefPose != nullptr)
+		{
+			return RefPose->GetSkeletonBoneIndexToLODBoneIndexMap(LODLevel);
 		}
 		else
 		{
diff --git a/Engine/Plugins/Experimental/AnimNext/Source/AnimNext/Public/ReferencePose.h b/Engine/Plugins/Experimental/AnimNext/Source/AnimNext/Public/ReferencePose.h
--- a/Engine/Plugins/Experimental/AnimNext/Source/AnimNext/Public/ReferencePose.h
+++ b/Engine/Plugins/Experimental/AnimNext/Source/AnimNext/Public/ReferencePose.h
@@ -34,13 +34,20 @@ struct TReferencePose
 	// Higher LOD come first
 	TTransformArray<AllocatorType> ReferenceLocalTransforms;
 
-	// List of LOD bone indices for each LOD
+	// A mapping of LOD sorted bone indices to skeletal mesh indices per LOD
 	// Each list of bone indices is a mapping of: LODSortedBoneIndex -> SkeletalMeshBoneIndex
-	TArray<TArray<FBoneIndexType, AllocatorType>, AllocatorType> LODBoneIndexes;
+	// When fast path is enabled, we have a single LOD entry that we truncate to the number of bones for each LOD
+	TArray<TArray<FBoneIndexType, AllocatorType>, AllocatorType> LODBoneIndexToMeshBoneIndexMapPerLOD;
+
+	// A mapping of LOD sorted bone indices to skeleton indices per LOD
+	// Each list of bone indices is a mapping of: LODSortedBoneIndex -> SkeletonBoneIndex
+	// When fast path is enabled, we have a single LOD entry that we truncate to the number of bones for each LOD
+	TArray<TArray<FBoneIndexType, AllocatorType>, AllocatorType> LODBoneIndexToSkeletonBoneIndexMapPerLOD;
 
 	// List of skeleton bone indices for each LOD
 	// Each list of skeleton bone indices is a mapping of: SkeletonBoneIndex -> LODSortedBoneIndex
-	TArray<TArray<FBoneIndexType, AllocatorType>, AllocatorType> SkeletonToLODBoneIndexes;
+	// When fast path is enabled, we have a single LOD entry that we truncate to the number of bones for each LOD
+	TArray<TArray<FBoneIndexType, AllocatorType>, AllocatorType> SkeletonBoneIndexToLODBoneIndexMapPerLOD;
 
 	// Number of bones for each LOD
 	TArray<int32, AllocatorType> LODNumBones;
@@ -69,88 +76,104 @@ struct TReferencePose
 	}
 
 	void Initialize(const FReferenceSkeleton& RefSkeleton
-		, const TArray<TArray<FBoneIndexType>>& InLODBoneIndexes
-		, const TArray<TArray<FBoneIndexType>>& InSkeletonToLODBoneIndexes
+		, const TArray<TArray<FBoneIndexType>>& InLODBoneIndexToMeshBoneIndexMapPerLOD
+		, const TArray<TArray<FBoneIndexType>>& InLODBoneIndexToSkeletonBoneIndexMapPerLOD
+		, const TArray<TArray<FBoneIndexType>>& InSkeletonBoneIndexToLODBoneIndexMapPerLOD
 		, const TArray<int32, AllocatorType>& InLODNumBones
 		, bool bFastPath = false)
 	{
-		const int32 NumBonesLOD0 = (InLODNumBones.Num()) > 0 ? InLODNumBones[0] : 0;
+		const int32 NumBonesLOD0 = !InLODNumBones.IsEmpty() ? InLODNumBones[0] : 0;
 
 		ReferenceLocalTransforms.SetNum(NumBonesLOD0);
-		LODBoneIndexes = InLODBoneIndexes;
-		SkeletonToLODBoneIndexes = InSkeletonToLODBoneIndexes;
+		LODBoneIndexToMeshBoneIndexMapPerLOD = InLODBoneIndexToMeshBoneIndexMapPerLOD;
+		LODBoneIndexToSkeletonBoneIndexMapPerLOD = InLODBoneIndexToSkeletonBoneIndexMapPerLOD;
+		SkeletonBoneIndexToLODBoneIndexMapPerLOD = InSkeletonBoneIndexToLODBoneIndexMapPerLOD;
 		LODNumBones = InLODNumBones;
 		
 		const TArray<FTransform>& RefBonePose = RefSkeleton.GetRefBonePose();
-		const TArray<FBoneIndexType>& InBoneIndexes = InLODBoneIndexes[0]; // Fill the transforms with the LOD0 indexes
+		const TArray<FBoneIndexType>& BoneLODIndexToMeshIndexMap0 = InLODBoneIndexToMeshBoneIndexMapPerLOD[0]; // Fill the transforms with the LOD0 indexes
 
-		for (int32 i = 0; i < NumBonesLOD0; ++i)
+		for (int32 LODBoneIndex = 0; LODBoneIndex < NumBonesLOD0; ++LODBoneIndex)
 		{
-			ReferenceLocalTransforms[i] = RefBonePose[InBoneIndexes[i]]; // TODO : For SoA this is un-optimal, as we are using a TransformAdapter. Evaluate using a specific SoA iterator
+			// TODO : For SoA this is un-optimal, as we are using a TransformAdapter. Evaluate using a specific SoA iterator
+			ReferenceLocalTransforms[LODBoneIndex] = RefBonePose[BoneLODIndexToMeshIndexMap0[LODBoneIndex]];
 		}
 
 		GenerationFlags = bFastPath ? EReferencePoseGenerationFlags::FastPath : EReferencePoseGenerationFlags::None;
 	}
 
-	const TArrayView<const FBoneIndexType> GetLODBoneIndexes(int32 LODLevel) const
+	// Returns a list of LOD sorted skeletal mesh bone indices, a mapping of: LODSortedBoneIndex -> SkeletalMeshBoneIndex
+	const TArrayView<const FBoneIndexType> GetLODBoneIndexToMeshBoneIndexMap(int32 LODLevel) const
 	{
 		TArrayView<const FBoneIndexType> ArrayView;
 
-		if (LODLevel >= 0 && (IsFastPath() || LODLevel < LODBoneIndexes.Num()))
+		if (LODLevel >= 0 && (IsFastPath() || LODLevel < LODBoneIndexToMeshBoneIndexMapPerLOD.Num()))
 		{
 			const int32 NumBonesForLOD = GetNumBonesForLOD(LODLevel);
+			const int32 LODIndex = IsFastPath() ? 0 : LODLevel;
 
-			if (IsFastPath())
-			{
-				ArrayView = MakeArrayView(LODBoneIndexes[0].GetData(), NumBonesForLOD);
-			}
-			else
-			{
-				ArrayView = MakeArrayView(LODBoneIndexes[LODLevel].GetData(), NumBonesForLOD);
-			}
+			ArrayView = MakeArrayView(LODBoneIndexToMeshBoneIndexMapPerLOD[LODIndex].GetData(), NumBonesForLOD);
 		}
 
 		return ArrayView;
 	}
 
-	const TArrayView<const FBoneIndexType> GetSkeletonToLODBoneIndexes(int32 LODLevel) const
+	// Returns a list of LOD sorted skeleton bone indices, a mapping of: LODSortedBoneIndex -> SkeletonBoneIndex
+	const TArrayView<const FBoneIndexType> GetLODBoneIndexToSkeletonBoneIndexMap(int32 LODLevel) const
 	{
 		TArrayView<const FBoneIndexType> ArrayView;
 
-		if (LODLevel >= 0 && (IsFastPath() || LODLevel < SkeletonToLODBoneIndexes.Num()))
+		if (LODLevel >= 0 && (IsFastPath() || LODLevel < LODBoneIndexToSkeletonBoneIndexMapPerLOD.Num()))
 		{
 			const int32 NumBonesForLOD = GetNumBonesForLOD(LODLevel);
+			const int32 LODIndex = IsFastPath() ? 0 : LODLevel;
 
-			if (IsFastPath())
-			{
-				return MakeArrayView(SkeletonToLODBoneIndexes[0].GetData(), NumBonesForLOD);
-			}
-			else
-			{
-				return MakeArrayView(SkeletonToLODBoneIndexes[LODLevel].GetData(), NumBonesForLOD);
-			}
+			ArrayView = MakeArrayView(LODBoneIndexToSkeletonBoneIndexMapPerLOD[LODIndex].GetData(), NumBonesForLOD);
 		}
 
 		return ArrayView;
 	}
 
+	// Returns a list of LOD bone indices, a mapping of: SkeletonBoneIndex -> LODSortedBoneIndex
+	const TArrayView<const FBoneIndexType> GetSkeletonBoneIndexToLODBoneIndexMap(int32 LODLevel) const
+	{
+		TArrayView<const FBoneIndexType> ArrayView;
+
+		if (LODLevel >= 0 && (IsFastPath() || LODLevel < SkeletonBoneIndexToLODBoneIndexMapPerLOD.Num()))
+		{
+			const int32 NumBonesForLOD = GetNumBonesForLOD(LODLevel);
+			const int32 LODIndex = IsFastPath() ? 0 : LODLevel;
+
+			ArrayView = MakeArrayView(SkeletonBoneIndexToLODBoneIndexMapPerLOD[LODIndex].GetData(), NumBonesForLOD);
+		}
+
+		return ArrayView;
+	}
+
+	int32 GetMeshBoneIndexFromLODBoneIndex(int32 LODBoneIndex) const
+	{
+		const int32 NumBonesLOD0 = LODBoneIndexToMeshBoneIndexMapPerLOD[0].Num();
+		check(LODBoneIndex < NumBonesLOD0);
+		return LODBoneIndexToMeshBoneIndexMapPerLOD[0][LODBoneIndex];
+	}
+
 	int32 GetSkeletonBoneIndexFromLODBoneIndex(int32 LODBoneIndex) const
 	{
-		const int32 NumBonesLOD0 = LODBoneIndexes[0].Num();
+		const int32 NumBonesLOD0 = LODBoneIndexToSkeletonBoneIndexMapPerLOD[0].Num();
 		check(LODBoneIndex < NumBonesLOD0);
-		return LODBoneIndexes[0][LODBoneIndex];
+		return LODBoneIndexToSkeletonBoneIndexMapPerLOD[0][LODBoneIndex];
 	}
 
 	int32 GetLODBoneIndexFromSkeletonBoneIndex(int32 SkeletionBoneIndex) const
 	{
-		const int32 NumBonesLOD0 = SkeletonToLODBoneIndexes[0].Num();
+		const int32 NumBonesLOD0 = SkeletonBoneIndexToLODBoneIndexMapPerLOD[0].Num();
 		check(SkeletionBoneIndex < NumBonesLOD0);
-		return SkeletonToLODBoneIndexes[0][SkeletionBoneIndex];
+		return SkeletonBoneIndexToLODBoneIndexMapPerLOD[0][SkeletionBoneIndex];
 	}
 
 	FTransform GetRefPoseTransform(int32 LODBoneIndex) const
 	{
-		const int32 NumBonesLOD0 = LODBoneIndexes[0].Num();
+		const int32 NumBonesLOD0 = LODBoneIndexToMeshBoneIndexMapPerLOD[0].Num();
 		check(LODBoneIndex < NumBonesLOD0);
 
 		return ReferenceLocalTransforms[LODBoneIndex];
@@ -158,7 +181,7 @@ struct TReferencePose
 
 	const FQuat GetRefPoseRotation(int32 LODBoneIndex) const
 	{
-		const int32 NumBonesLOD0 = LODBoneIndexes[0].Num();
+		const int32 NumBonesLOD0 = LODBoneIndexToMeshBoneIndexMapPerLOD[0].Num();
 		check(LODBoneIndex < NumBonesLOD0);
 
 		return ReferenceLocalTransforms[LODBoneIndex].GetRotation();
@@ -166,7 +189,7 @@ struct TReferencePose
 
 	const FVector GetRefPoseTranslation(int32 LODBoneIndex) const
 	{
-		const int32 NumBonesLOD0 = LODBoneIndexes[0].Num();
+		const int32 NumBonesLOD0 = LODBoneIndexToMeshBoneIndexMapPerLOD[0].Num();
 		check(LODBoneIndex < NumBonesLOD0);
 
 		return ReferenceLocalTransforms[LODBoneIndex].GetTranslation();
@@ -174,12 +197,11 @@ struct TReferencePose
 
 	const FVector GetRefPoseScale3D(int32 LODBoneIndex) const
 	{
-		const int32 NumBonesLOD0 = LODBoneIndexes[0].Num();
+		const int32 NumBonesLOD0 = LODBoneIndexToMeshBoneIndexMapPerLOD[0].Num();
 		check(LODBoneIndex < NumBonesLOD0);
 
 		return ReferenceLocalTransforms[LODBoneIndex].GetScale3D();
 	}
-
 };
 
 /**
