Implement a debug function to convert an evaluation program into a string

diff --git a/Engine/Plugins/Experimental/AnimNext/Source/AnimNext/Private/EvaluationVM/EvaluationProgram.cpp b/Engine/Plugins/Experimental/AnimNext/Source/AnimNext/Private/EvaluationVM/EvaluationProgram.cpp
--- a/Engine/Plugins/Experimental/AnimNext/Source/AnimNext/Private/EvaluationVM/EvaluationProgram.cpp
+++ b/Engine/Plugins/Experimental/AnimNext/Source/AnimNext/Private/EvaluationVM/EvaluationProgram.cpp
@@ -5,6 +5,8 @@
 #include "EvaluationVM/EvaluationTask.h"
 #include "EvaluationVM/EvaluationVM.h"
 
+#include "Misc/StringBuilder.h"
+
 namespace UE::AnimNext
 {
 	bool FEvaluationProgram::IsEmpty() const
@@ -22,4 +24,43 @@ namespace UE::AnimNext
 			}
 		}
 	}
+
+	FString FEvaluationProgram::ToString() const
+	{
+		FStringBuilderBase Result;
+
+		FString PropertyValue;
+		for (const TUniquePtr<FAnimNextEvaluationTask>& Task : Tasks)
+		{
+			const UScriptStruct* TaskStruct = Task->GetStruct();
+
+			Result.Appendf(TEXT("%s\n"), *TaskStruct->GetName());
+
+			for (TFieldIterator<FProperty> It(TaskStruct); It; ++It)
+			{
+				if (const FProperty* Property = *It)
+				{
+					PropertyValue.Reset();
+
+					const bool bIsDefaultValue = !Property->ExportText_InContainer(0, PropertyValue, Task.Get(), nullptr, nullptr, PPF_None);
+
+					if (bIsDefaultValue)
+					{
+						if (Property->IsA<FBoolProperty>())
+						{
+							PropertyValue = TEXT("False");
+						}
+						else if (Property->IsA<FNumericProperty>())
+						{
+							PropertyValue = TEXT("0");
+						}
+					}
+
+					Result.Appendf(TEXT("\t%s: %s\n"), *Property->GetName(), *PropertyValue);
+				}
+			}
+		}
+
+		return Result.ToString();
+	}
 }
diff --git a/Engine/Plugins/Experimental/AnimNext/Source/AnimNext/Public/EvaluationVM/EvaluationProgram.h b/Engine/Plugins/Experimental/AnimNext/Source/AnimNext/Public/EvaluationVM/EvaluationProgram.h
--- a/Engine/Plugins/Experimental/AnimNext/Source/AnimNext/Public/EvaluationVM/EvaluationProgram.h
+++ b/Engine/Plugins/Experimental/AnimNext/Source/AnimNext/Public/EvaluationVM/EvaluationProgram.h
@@ -39,6 +39,9 @@ namespace UE::AnimNext
 		// Executes the current program on the provided virtual machine
 		void Execute(FEvaluationVM& VM) const;
 
+		// Returns the program as a string suitable for debug purposes
+		FString ToString() const;
+
 	private:
 		// Disallow copy
 		FEvaluationProgram(const FEvaluationProgram&) = delete;
diff --git a/Engine/Plugins/Experimental/AnimNext/Source/AnimNext/Public/EvaluationVM/EvaluationTask.h b/Engine/Plugins/Experimental/AnimNext/Source/AnimNext/Public/EvaluationVM/EvaluationTask.h
--- a/Engine/Plugins/Experimental/AnimNext/Source/AnimNext/Public/EvaluationVM/EvaluationTask.h
+++ b/Engine/Plugins/Experimental/AnimNext/Source/AnimNext/Public/EvaluationVM/EvaluationTask.h
@@ -11,6 +11,10 @@ namespace UE::AnimNext
 	struct FEvaluationVM;
 }
 
+// Implements various implementation details for an animation evaluation task
+#define DECLARE_ANIM_EVALUATION_TASK(TaskName) \
+	virtual const UScriptStruct* GetStruct() const override { return TaskName::StaticStruct(); }
+
 /*
  * Evaluation Task
  *
@@ -27,6 +31,9 @@ struct ANIMNEXT_API FAnimNextEvaluationTask
 
 	virtual ~FAnimNextEvaluationTask() {}
 
+	// Returns the derived type for the task
+	virtual const UScriptStruct* GetStruct() const { return FAnimNextEvaluationTask::StaticStruct(); }
+
 	// Task entry point
 	virtual void Execute(UE::AnimNext::FEvaluationVM& VM) const { LowLevelFatalError(TEXT("Pure virtual not implemented (FAnimNextEvaluationTask::Execute)")); }
 };
diff --git a/Engine/Plugins/Experimental/AnimNext/Source/AnimNext/Public/EvaluationVM/Tasks/ApplyAdditiveKeyframe.h b/Engine/Plugins/Experimental/AnimNext/Source/AnimNext/Public/EvaluationVM/Tasks/ApplyAdditiveKeyframe.h
--- a/Engine/Plugins/Experimental/AnimNext/Source/AnimNext/Public/EvaluationVM/Tasks/ApplyAdditiveKeyframe.h
+++ b/Engine/Plugins/Experimental/AnimNext/Source/AnimNext/Public/EvaluationVM/Tasks/ApplyAdditiveKeyframe.h
@@ -19,6 +19,8 @@ struct ANIMNEXT_API FAnimNextApplyAdditiveKeyframeTask : public FAnimNextEvaluat
 {
 	GENERATED_BODY()
 
+	DECLARE_ANIM_EVALUATION_TASK(FAnimNextApplyAdditiveKeyframeTask)
+
 	static FAnimNextApplyAdditiveKeyframeTask Make(float BlendWeight);
 
 	// Task entry point
diff --git a/Engine/Plugins/Experimental/AnimNext/Source/AnimNext/Public/EvaluationVM/Tasks/BlendKeyframes.h b/Engine/Plugins/Experimental/AnimNext/Source/AnimNext/Public/EvaluationVM/Tasks/BlendKeyframes.h
--- a/Engine/Plugins/Experimental/AnimNext/Source/AnimNext/Public/EvaluationVM/Tasks/BlendKeyframes.h
+++ b/Engine/Plugins/Experimental/AnimNext/Source/AnimNext/Public/EvaluationVM/Tasks/BlendKeyframes.h
@@ -31,6 +31,8 @@ struct ANIMNEXT_API FAnimNextBlendTwoKeyframesTask : public FAnimNextEvaluationT
 {
 	GENERATED_BODY()
 
+	DECLARE_ANIM_EVALUATION_TASK(FAnimNextBlendTwoKeyframesTask)
+
 	static FAnimNextBlendTwoKeyframesTask Make(float InterpolationAlpha);
 
 	// Task entry point
@@ -55,6 +57,8 @@ struct ANIMNEXT_API FAnimNextBlendOverwriteKeyframeWithScaleTask : public FAnimN
 {
 	GENERATED_BODY()
 
+	DECLARE_ANIM_EVALUATION_TASK(FAnimNextBlendOverwriteKeyframeWithScaleTask)
+
 	static FAnimNextBlendOverwriteKeyframeWithScaleTask Make(float ScaleFactor);
 
 	// Task entry point
@@ -80,6 +84,8 @@ struct ANIMNEXT_API FAnimNextBlendAddKeyframeWithScaleTask : public FAnimNextEva
 {
 	GENERATED_BODY()
 
+	DECLARE_ANIM_EVALUATION_TASK(FAnimNextBlendAddKeyframeWithScaleTask)
+
 	static FAnimNextBlendAddKeyframeWithScaleTask Make(float ScaleFactor);
 
 	// Task entry point
diff --git a/Engine/Plugins/Experimental/AnimNext/Source/AnimNext/Public/EvaluationVM/Tasks/ExecuteProgram.h b/Engine/Plugins/Experimental/AnimNext/Source/AnimNext/Public/EvaluationVM/Tasks/ExecuteProgram.h
--- a/Engine/Plugins/Experimental/AnimNext/Source/AnimNext/Public/EvaluationVM/Tasks/ExecuteProgram.h
+++ b/Engine/Plugins/Experimental/AnimNext/Source/AnimNext/Public/EvaluationVM/Tasks/ExecuteProgram.h
@@ -23,6 +23,8 @@ struct ANIMNEXT_API FAnimNextExecuteProgramTask : public FAnimNextEvaluationTask
 {
 	GENERATED_BODY()
 
+	DECLARE_ANIM_EVALUATION_TASK(FAnimNextExecuteProgramTask)
+
 	static FAnimNextExecuteProgramTask Make(const UE::AnimNext::FEvaluationProgram* Program);
 
 	// Task entry point
diff --git a/Engine/Plugins/Experimental/AnimNext/Source/AnimNext/Public/EvaluationVM/Tasks/NormalizeRotations.h b/Engine/Plugins/Experimental/AnimNext/Source/AnimNext/Public/EvaluationVM/Tasks/NormalizeRotations.h
--- a/Engine/Plugins/Experimental/AnimNext/Source/AnimNext/Public/EvaluationVM/Tasks/NormalizeRotations.h
+++ b/Engine/Plugins/Experimental/AnimNext/Source/AnimNext/Public/EvaluationVM/Tasks/NormalizeRotations.h
@@ -18,6 +18,8 @@ struct ANIMNEXT_API FAnimNextNormalizeKeyframeRotationsTask : public FAnimNextEv
 {
 	GENERATED_BODY()
 
+	DECLARE_ANIM_EVALUATION_TASK(FAnimNextNormalizeKeyframeRotationsTask)
+
 	// Task entry point
 	virtual void Execute(UE::AnimNext::FEvaluationVM& VM) const override;
 };
diff --git a/Engine/Plugins/Experimental/AnimNext/Source/AnimNext/Public/EvaluationVM/Tasks/PushAnimSequenceKeyframe.h b/Engine/Plugins/Experimental/AnimNext/Source/AnimNext/Public/EvaluationVM/Tasks/PushAnimSequenceKeyframe.h
--- a/Engine/Plugins/Experimental/AnimNext/Source/AnimNext/Public/EvaluationVM/Tasks/PushAnimSequenceKeyframe.h
+++ b/Engine/Plugins/Experimental/AnimNext/Source/AnimNext/Public/EvaluationVM/Tasks/PushAnimSequenceKeyframe.h
@@ -19,6 +19,8 @@ struct ANIMNEXT_API FAnimNextAnimSequenceKeyframeTask : public FAnimNextEvaluati
 {
 	GENERATED_BODY()
 
+	DECLARE_ANIM_EVALUATION_TASK(FAnimNextAnimSequenceKeyframeTask)
+
 	static FAnimNextAnimSequenceKeyframeTask MakeFromSampleTime(TWeakObjectPtr<UAnimSequence> AnimSequence, double SampleTime, bool bInterpolate);
 	static FAnimNextAnimSequenceKeyframeTask MakeFromKeyframeIndex(TWeakObjectPtr<UAnimSequence> AnimSequence, uint32 KeyframeIndex);
 
diff --git a/Engine/Plugins/Experimental/AnimNext/Source/AnimNext/Public/EvaluationVM/Tasks/PushReferenceKeyframe.h b/Engine/Plugins/Experimental/AnimNext/Source/AnimNext/Public/EvaluationVM/Tasks/PushReferenceKeyframe.h
--- a/Engine/Plugins/Experimental/AnimNext/Source/AnimNext/Public/EvaluationVM/Tasks/PushReferenceKeyframe.h
+++ b/Engine/Plugins/Experimental/AnimNext/Source/AnimNext/Public/EvaluationVM/Tasks/PushReferenceKeyframe.h
@@ -19,6 +19,8 @@ struct ANIMNEXT_API FAnimNextPushReferenceKeyframeTask : public FAnimNextEvaluat
 {
 	GENERATED_BODY()
 
+	DECLARE_ANIM_EVALUATION_TASK(FAnimNextPushReferenceKeyframeTask)
+
 	static FAnimNextPushReferenceKeyframeTask MakeFromSkeleton();
 	static FAnimNextPushReferenceKeyframeTask MakeFromAdditiveIdentity();
 
