Fix RigVM decorator array properties not restoring on load

diff --git a/Engine/Plugins/Experimental/AnimNext/Source/AnimNext/Internal/Graph/RigDecorator_AnimNextCppDecorator.h b/Engine/Plugins/Experimental/AnimNext/Source/AnimNext/Internal/Graph/RigDecorator_AnimNextCppDecorator.h
--- a/Engine/Plugins/Experimental/AnimNext/Source/AnimNext/Internal/Graph/RigDecorator_AnimNextCppDecorator.h
+++ b/Engine/Plugins/Experimental/AnimNext/Source/AnimNext/Internal/Graph/RigDecorator_AnimNextCppDecorator.h
@@ -22,7 +22,7 @@ struct ANIMNEXT_API FRigDecorator_AnimNextCppDecorator : public FRigVMDecorator
 	UScriptStruct* DecoratorSharedDataStruct = nullptr;
 
 #if WITH_EDITOR
-	virtual void GetProgrammaticPins(URigVMController* InController, int32 InParentPinIndex, struct FRigVMPinInfoArray& OutPinArray) const override;
+	virtual void GetProgrammaticPins(URigVMController* InController, int32 InParentPinIndex, const FString& InDefaultValue, struct FRigVMPinInfoArray& OutPinArray) const override;
 
 	const UE::AnimNext::FDecorator* GetDecorator() const;
 #endif
diff --git a/Engine/Plugins/Experimental/AnimNext/Source/AnimNext/Private/Graph/RigDecorator_AnimNextCppDecorator.cpp b/Engine/Plugins/Experimental/AnimNext/Source/AnimNext/Private/Graph/RigDecorator_AnimNextCppDecorator.cpp
--- a/Engine/Plugins/Experimental/AnimNext/Source/AnimNext/Private/Graph/RigDecorator_AnimNextCppDecorator.cpp
+++ b/Engine/Plugins/Experimental/AnimNext/Source/AnimNext/Private/Graph/RigDecorator_AnimNextCppDecorator.cpp
@@ -10,7 +10,7 @@
 #include UE_INLINE_GENERATED_CPP_BY_NAME(RigDecorator_AnimNextCppDecorator)
 
 #if WITH_EDITOR
-void FRigDecorator_AnimNextCppDecorator::GetProgrammaticPins(URigVMController* InController, int32 InParentPinIndex, FRigVMPinInfoArray& OutPinArray) const
+void FRigDecorator_AnimNextCppDecorator::GetProgrammaticPins(URigVMController* InController, int32 InParentPinIndex, const FString& InDefaultValue, FRigVMPinInfoArray& OutPinArray) const
 {
 	if (DecoratorSharedDataStruct == nullptr)
 	{
@@ -19,6 +19,12 @@ void FRigDecorator_AnimNextCppDecorator::GetProgrammaticPins(URigVMController* I
 
 	FStructOnScope DefaultValueMemoryScope(DecoratorSharedDataStruct);
 
+	if (!InDefaultValue.IsEmpty())
+	{
+		FRigVMPinDefaultValueImportErrorContext ErrorPipe;
+		DecoratorSharedDataStruct->ImportText(*InDefaultValue, DefaultValueMemoryScope.GetStructMemory(), nullptr, PPF_None, &ErrorPipe, DecoratorSharedDataStruct->GetName());
+	}
+
 	const int32 StartPinIndex = OutPinArray.Num();
 	OutPinArray.AddPins(DecoratorSharedDataStruct, InController, ERigVMPinDirection::Invalid, InParentPinIndex, DefaultValueMemoryScope.GetStructMemory(), true);
 
diff --git a/Engine/Plugins/Runtime/RigVM/Source/RigVM/Public/RigVMCore/RigVMDecorator.h b/Engine/Plugins/Runtime/RigVM/Source/RigVM/Public/RigVMCore/RigVMDecorator.h
--- a/Engine/Plugins/Runtime/RigVM/Source/RigVM/Public/RigVMCore/RigVMDecorator.h
+++ b/Engine/Plugins/Runtime/RigVM/Source/RigVM/Public/RigVMCore/RigVMDecorator.h
@@ -63,7 +63,7 @@ public:
 	virtual void OnDecoratorAdded(URigVMController* InController, URigVMNode* InNode) {}
 
 	// allows the decorator to return dynamic pins (parent pin index must be INDEX_NONE or point to a valid index of the parent pin in the OutPinArray)
-	virtual void GetProgrammaticPins(URigVMController* InController, int32 InParentPinIndex, struct FRigVMPinInfoArray& OutPinArray) const {}
+	virtual void GetProgrammaticPins(URigVMController* InController, int32 InParentPinIndex, const FString& InDefaultValue, struct FRigVMPinInfoArray& OutPinArray) const {}
 
 #endif
 
diff --git a/Engine/Plugins/Runtime/RigVM/Source/RigVMDeveloper/Private/RigVMModel/RigVMController.cpp b/Engine/Plugins/Runtime/RigVM/Source/RigVMDeveloper/Private/RigVMModel/RigVMController.cpp
--- a/Engine/Plugins/Runtime/RigVM/Source/RigVMDeveloper/Private/RigVMModel/RigVMController.cpp
+++ b/Engine/Plugins/Runtime/RigVM/Source/RigVMDeveloper/Private/RigVMModel/RigVMController.cpp
@@ -13793,7 +13793,7 @@ FName URigVMController::AddDecorator(URigVMNode* InNode, UScriptStruct* InDecora
 	AddPinsForStruct(InDecoratorScriptStruct, InNode, DecoratorPin, DecoratorPin->GetDirection(), InDefaultValue, true);
 
 	FRigVMPinInfoArray ProgrammaticPins;
-	Decorator->GetProgrammaticPins(this, INDEX_NONE, ProgrammaticPins);
+	Decorator->GetProgrammaticPins(this, INDEX_NONE, InDefaultValue, ProgrammaticPins);
 
 	const FRigVMRegistry& Registry = FRigVMRegistry::Get();
 	const FRigVMPinInfoArray PreviousPins;
@@ -15313,7 +15313,7 @@ bool URigVMController::GenerateNewPinInfos(const FRigVMRegistry& Registry, URigV
 				TSharedPtr<FStructOnScope> DecoratorScope = Pin->GetDecoratorInstance();
 				FRigVMDecorator* VMDecorator = (FRigVMDecorator*)DecoratorScope->GetStructMemory();
 
-				VMDecorator->GetProgrammaticPins(this, NewPinIndex, NewPinInfos);
+				VMDecorator->GetProgrammaticPins(this, NewPinIndex, Pin->GetDefaultValue(), NewPinInfos);
 			}
 		}
 	}
