Fix EngineTest unit test failure
Use the correct graph when compiling
Dynamically register/unregister test decorator

diff --git a/Engine/Plugins/Experimental/AnimNext/Source/AnimNext/Public/DecoratorBase/Decorator.h b/Engine/Plugins/Experimental/AnimNext/Source/AnimNext/Public/DecoratorBase/Decorator.h
--- a/Engine/Plugins/Experimental/AnimNext/Source/AnimNext/Public/DecoratorBase/Decorator.h
+++ b/Engine/Plugins/Experimental/AnimNext/Source/AnimNext/Public/DecoratorBase/Decorator.h
@@ -81,9 +81,10 @@ class FArchive;
 	}
 
 // Allows a decorator to auto-register and unregister within the current execution scope
+// The decorator must be found in the current scope without a namespace qualification
 #define AUTO_REGISTER_ANIM_DECORATOR(DecoratorName) \
 	UE::AnimNext::FDecoratorStaticInitHook DecoratorName##Hook( \
-		[](void* DestPtr, FDecoratorMemoryLayout& MemoryDesc) -> FDecorator* \
+		[](void* DestPtr, UE::AnimNext::FDecoratorMemoryLayout& MemoryDesc) -> UE::AnimNext::FDecorator* \
 		{ \
 			MemoryDesc = DecoratorName::DecoratorMemoryDescription; \
 			return DestPtr != nullptr ? new(DestPtr) DecoratorName() : nullptr; \
diff --git a/Engine/Plugins/Experimental/AnimNext/Source/AnimNextTestSuite/Private/AnimNextDecoratorGraphTest.cpp b/Engine/Plugins/Experimental/AnimNext/Source/AnimNextTestSuite/Private/AnimNextDecoratorGraphTest.cpp
--- a/Engine/Plugins/Experimental/AnimNext/Source/AnimNextTestSuite/Private/AnimNextDecoratorGraphTest.cpp
+++ b/Engine/Plugins/Experimental/AnimNext/Source/AnimNextTestSuite/Private/AnimNextDecoratorGraphTest.cpp
@@ -2,6 +2,7 @@
 
 #include "AnimNextDecoratorGraphTest.h"
 
+#include "AnimNextRuntimeTest.h"
 #include "AssetToolsModule.h"
 #include "Context.h"
 #include "UncookedOnlyUtils.h"
@@ -74,13 +75,17 @@ namespace UE::AnimNext
 		DEFINE_ANIM_DECORATOR_IMPLEMENTS_INTERFACE(IEvaluate)
 		DEFINE_ANIM_DECORATOR_IMPLEMENTS_INTERFACE(IUpdate)
 	DEFINE_ANIM_DECORATOR_END(FTestDecorator)
-	AUTO_REGISTER_ANIM_DECORATOR(FTestDecorator)
 }
 
 IMPLEMENT_SIMPLE_AUTOMATION_TEST(FAnimationAnimNextRuntimeTest_GraphAddDecorator, "Animation.AnimNext.Runtime.Graph.AddDecorator", EAutomationTestFlags::EditorContext | EAutomationTestFlags::EngineFilter)
 
 bool FAnimationAnimNextRuntimeTest_GraphAddDecorator::RunTest(const FString& InParameters)
 {
+	using namespace UE::AnimNext;
+	AUTO_REGISTER_ANIM_DECORATOR(FTestDecorator)
+
+	FScopedClearNodeTemplateRegistry ScopedClearNodeTemplateRegistry;
+
 	UAnimNextGraphTest* AnimNextGraph = NewObject<UAnimNextGraphTest>();
 
 	UAnimNextGraph_EditorData* EditorData = NewObject<UAnimNextGraph_EditorData>(AnimNextGraph, TEXT("EditorData"));
@@ -89,7 +94,7 @@ bool FAnimationAnimNextRuntimeTest_GraphAddDecorator::RunTest(const FString& InP
 	EditorData->Initialize(/*bRecompileVM*/false);
 	EditorData->GetRigVMClient()->SetExecuteContextStruct(FAnimNextExecuteContext::StaticStruct());
 
-	URigVMController* Controller = EditorData->GetRigVMClient()->GetController(EditorData->Graphs[0]);
+	URigVMController* Controller = EditorData->GetRigVMClient()->GetController(EditorData->GetRigVMClient()->GetDefaultModel());
 	UE_RETURN_ON_ERROR(Controller != nullptr, "FAnimationAnimNextRuntimeTest_GraphAddDecorator -> Failed to get RigVM controller");
 
 	// Create an empty decorator stack node
@@ -100,7 +105,7 @@ bool FAnimationAnimNextRuntimeTest_GraphAddDecorator::RunTest(const FString& InP
 	const UScriptStruct* CppDecoratorStruct = FRigDecorator_AnimNextCppDecorator::StaticStruct();
 	UE_RETURN_ON_ERROR(CppDecoratorStruct != nullptr, "FAnimationAnimNextRuntimeTest_GraphAddDecorator -> Failed to get find Cpp decorator static struct");
 
-	const UE::AnimNext::FDecorator* Decorator = UE::AnimNext::FDecoratorRegistry::Get().Find(UE::AnimNext::FTestDecorator::DecoratorUID);
+	const FDecorator* Decorator = FDecoratorRegistry::Get().Find(FTestDecorator::DecoratorUID);
 	UE_RETURN_ON_ERROR(Decorator != nullptr, "FAnimationAnimNextRuntimeTest_GraphAddDecorator -> Failed to find test decorator");
 
 	UScriptStruct* ScriptStruct = Decorator->GetDecoratorSharedDataStruct();
@@ -175,6 +180,11 @@ IMPLEMENT_SIMPLE_AUTOMATION_TEST(FAnimationAnimNextRuntimeTest_GraphExecute, "An
 
 bool FAnimationAnimNextRuntimeTest_GraphExecute::RunTest(const FString& InParameters)
 {
+	using namespace UE::AnimNext;
+	AUTO_REGISTER_ANIM_DECORATOR(FTestDecorator)
+
+	FScopedClearNodeTemplateRegistry ScopedClearNodeTemplateRegistry;
+
 	UAnimNextGraphTest* AnimNextGraph = NewObject<UAnimNextGraphTest>();
 
 	UAnimNextGraph_EditorData* EditorData = NewObject<UAnimNextGraph_EditorData>(AnimNextGraph, TEXT("EditorData"));
@@ -183,7 +193,7 @@ bool FAnimationAnimNextRuntimeTest_GraphExecute::RunTest(const FString& InParame
 	EditorData->Initialize(/*bRecompileVM*/false);
 	EditorData->GetRigVMClient()->SetExecuteContextStruct(FAnimNextExecuteContext::StaticStruct());
 
-	URigVMController* Controller = EditorData->GetRigVMClient()->GetController(EditorData->Graphs[0]);
+	URigVMController* Controller = EditorData->GetRigVMClient()->GetController(EditorData->GetRigVMClient()->GetDefaultModel());
 	UE_RETURN_ON_ERROR(Controller != nullptr, "FAnimationAnimNextRuntimeTest_GraphExecute -> Failed to get RigVM controller");
 
 	// Add graph entry point
@@ -204,7 +214,7 @@ bool FAnimationAnimNextRuntimeTest_GraphExecute::RunTest(const FString& InParame
 	const UScriptStruct* CppDecoratorStruct = FRigDecorator_AnimNextCppDecorator::StaticStruct();
 	UE_RETURN_ON_ERROR(CppDecoratorStruct != nullptr, "FAnimationAnimNextRuntimeTest_GraphExecute -> Failed to get find Cpp decorator static struct");
 
-	const UE::AnimNext::FDecorator* Decorator = UE::AnimNext::FDecoratorRegistry::Get().Find(UE::AnimNext::FTestDecorator::DecoratorUID);
+	const FDecorator* Decorator = FDecoratorRegistry::Get().Find(FTestDecorator::DecoratorUID);
 	UE_RETURN_ON_ERROR(Decorator != nullptr, "FAnimationAnimNextRuntimeTest_GraphExecute -> Failed to find test decorator");
 
 	UScriptStruct* ScriptStruct = Decorator->GetDecoratorSharedDataStruct();
@@ -242,15 +252,15 @@ bool FAnimationAnimNextRuntimeTest_GraphExecute::RunTest(const FString& InParame
 	Controller->SetPinDefaultValue(DecoratorPin->GetSubPins()[1]->GetPinPath(), TEXT("78"));
 	Controller->SetPinDefaultValue(DecoratorPin->GetSubPins()[2]->GetPinPath(), TEXT("142.33"));
 
-	TSharedRef<UE::AnimNext::FParamStack> ParamStack = MakeShared<UE::AnimNext::FParamStack>();
-	UE::AnimNext::FParamStack::AttachToCurrentThread(ParamStack);
+	TSharedRef<FParamStack> ParamStack = MakeShared<FParamStack>();
+	FParamStack::AttachToCurrentThread(ParamStack);
 
 	FAnimNextGraphInstance GraphInstance;
 	AnimNextGraph->AllocateInstance(GraphInstance);
 
-	UE::AnimNext::FContext Context(1.0f / 30.0f);
+	FContext Context(1.0f / 30.0f);
 
-	UE::AnimNext::FParamStack::FPushedLayerHandle LayerHandle = Context.GetMutableParamStack().PushValues(
+	FParamStack::FPushedLayerHandle LayerHandle = Context.GetMutableParamStack().PushValues(
 		"UpdateCount", (int32)0,
 		"EvaluateCount", (int32)0,
 		"SomeInt32", (int32)0,
@@ -273,7 +283,7 @@ bool FAnimationAnimNextRuntimeTest_GraphExecute::RunTest(const FString& InParame
 	Context.GetMutableParamStack().PopLayer(LayerHandle);
 	GraphInstance.Release();
 
-	UE::AnimNext::FParamStack::DetachFromCurrentThread();
+	FParamStack::DetachFromCurrentThread();
 
 	return true;
 }
@@ -282,6 +292,11 @@ IMPLEMENT_SIMPLE_AUTOMATION_TEST(FAnimationAnimNextRuntimeTest_GraphExecuteLaten
 
 bool FAnimationAnimNextRuntimeTest_GraphExecuteLatent::RunTest(const FString& InParameters)
 {
+	using namespace UE::AnimNext;
+	AUTO_REGISTER_ANIM_DECORATOR(FTestDecorator)
+
+	FScopedClearNodeTemplateRegistry ScopedClearNodeTemplateRegistry;
+
 	UAnimNextGraphTest* AnimNextGraph = NewObject<UAnimNextGraphTest>();
 
 	UAnimNextGraph_EditorData* EditorData = NewObject<UAnimNextGraph_EditorData>(AnimNextGraph, TEXT("EditorData"));
@@ -290,7 +305,7 @@ bool FAnimationAnimNextRuntimeTest_GraphExecuteLatent::RunTest(const FString& In
 	EditorData->Initialize(/*bRecompileVM*/false);
 	EditorData->GetRigVMClient()->SetExecuteContextStruct(FAnimNextExecuteContext::StaticStruct());
 
-	URigVMController* Controller = EditorData->GetRigVMClient()->GetController(EditorData->Graphs[0]);
+	URigVMController* Controller = EditorData->GetRigVMClient()->GetController(EditorData->GetRigVMClient()->GetDefaultModel());
 	UE_RETURN_ON_ERROR(Controller != nullptr, "FAnimationAnimNextRuntimeTest_GraphExecuteLatent -> Failed to get RigVM controller");
 
 	// Add graph entry point
@@ -311,7 +326,7 @@ bool FAnimationAnimNextRuntimeTest_GraphExecuteLatent::RunTest(const FString& In
 	const UScriptStruct* CppDecoratorStruct = FRigDecorator_AnimNextCppDecorator::StaticStruct();
 	UE_RETURN_ON_ERROR(CppDecoratorStruct != nullptr, "FAnimationAnimNextRuntimeTest_GraphExecuteLatent -> Failed to get find Cpp decorator static struct");
 
-	const UE::AnimNext::FDecorator* Decorator = UE::AnimNext::FDecoratorRegistry::Get().Find(UE::AnimNext::FTestDecorator::DecoratorUID);
+	const FDecorator* Decorator = FDecoratorRegistry::Get().Find(FTestDecorator::DecoratorUID);
 	UE_RETURN_ON_ERROR(Decorator != nullptr, "FAnimationAnimNextRuntimeTest_GraphExecuteLatent -> Failed to find test decorator");
 
 	UScriptStruct* ScriptStruct = Decorator->GetDecoratorSharedDataStruct();
@@ -379,15 +394,15 @@ bool FAnimationAnimNextRuntimeTest_GraphExecuteLatent::RunTest(const FString& In
 			DecoratorPin->GetSubPins()[4]);	// SomeOtherLatentInt32
 	}
 
-	TSharedRef<UE::AnimNext::FParamStack> ParamStack = MakeShared<UE::AnimNext::FParamStack>();
-	UE::AnimNext::FParamStack::AttachToCurrentThread(ParamStack);
+	TSharedRef<FParamStack> ParamStack = MakeShared<FParamStack>();
+	FParamStack::AttachToCurrentThread(ParamStack);
 
 	FAnimNextGraphInstance GraphInstance;
 	AnimNextGraph->AllocateInstance(GraphInstance);
 
-	UE::AnimNext::FContext Context(1.0f / 30.0f);
+	FContext Context(1.0f / 30.0f);
 
-	UE::AnimNext::FParamStack::FPushedLayerHandle LayerHandle = Context.GetMutableParamStack().PushValues(
+	FParamStack::FPushedLayerHandle LayerHandle = Context.GetMutableParamStack().PushValues(
 		"UpdateCount", (int32)0,
 		"EvaluateCount", (int32)0,
 		"SomeSourceInt", (int32)1223,
@@ -411,7 +426,7 @@ bool FAnimationAnimNextRuntimeTest_GraphExecuteLatent::RunTest(const FString& In
 	Context.GetMutableParamStack().PopLayer(LayerHandle);
 	GraphInstance.Release();
 
-	UE::AnimNext::FParamStack::DetachFromCurrentThread();
+	FParamStack::DetachFromCurrentThread();
 
 	return true;
 }
