Introduce continuous blend interface
Introduce discrete blend interface
Introduce smooth blend interface
Introduce blend by bool base decorator
Introduce blend smoother additive decorator

diff --git a/Engine/Plugins/Experimental/AnimNext/Source/AnimNext/Private/DecoratorInterfaces/IContinuousBlend.cpp b/Engine/Plugins/Experimental/AnimNext/Source/AnimNext/Private/DecoratorInterfaces/IContinuousBlend.cpp
new file mode 100644
--- /dev/null
+++ b/Engine/Plugins/Experimental/AnimNext/Source/AnimNext/Private/DecoratorInterfaces/IContinuousBlend.cpp
@@ -0,0 +1,19 @@
+// Copyright Epic Games, Inc. All Rights Reserved.
+
+#include "DecoratorInterfaces/IContinuousBlend.h"
+
+#include "DecoratorBase/ExecutionContext.h"
+
+namespace UE::AnimNext
+{
+	float IContinuousBlend::GetBlendWeight(FExecutionContext& Context, const TDecoratorBinding<IContinuousBlend>& Binding, int32 ChildIndex) const
+	{
+		TDecoratorBinding<IContinuousBlend> SuperBinding;
+		if (Context.GetInterfaceSuper(Binding, SuperBinding))
+		{
+			return SuperBinding.GetBlendWeight(Context, ChildIndex);
+		}
+
+		return -1.0f;
+	}
+}
diff --git a/Engine/Plugins/Experimental/AnimNext/Source/AnimNext/Private/DecoratorInterfaces/IDiscreteBlend.cpp b/Engine/Plugins/Experimental/AnimNext/Source/AnimNext/Private/DecoratorInterfaces/IDiscreteBlend.cpp
new file mode 100644
--- /dev/null
+++ b/Engine/Plugins/Experimental/AnimNext/Source/AnimNext/Private/DecoratorInterfaces/IDiscreteBlend.cpp
@@ -0,0 +1,68 @@
+// Copyright Epic Games, Inc. All Rights Reserved.
+
+#include "DecoratorInterfaces/IDiscreteBlend.h"
+
+#include "DecoratorBase/ExecutionContext.h"
+
+namespace UE::AnimNext
+{
+	float IDiscreteBlend::GetBlendWeight(FExecutionContext& Context, const TDecoratorBinding<IDiscreteBlend>& Binding, int32 ChildIndex) const
+	{
+		TDecoratorBinding<IDiscreteBlend> SuperBinding;
+		if (Context.GetInterfaceSuper(Binding, SuperBinding))
+		{
+			return SuperBinding.GetBlendWeight(Context, ChildIndex);
+		}
+
+		return -1.0f;
+	}
+
+	const FAlphaBlend* IDiscreteBlend::GetBlendState(FExecutionContext& Context, const TDecoratorBinding<IDiscreteBlend>& Binding, int32 ChildIndex) const
+	{
+		TDecoratorBinding<IDiscreteBlend> SuperBinding;
+		if (Context.GetInterfaceSuper(Binding, SuperBinding))
+		{
+			return SuperBinding.GetBlendState(Context, ChildIndex);
+		}
+
+		return nullptr;
+	}
+
+	int32 IDiscreteBlend::GetBlendDestinationChildIndex(FExecutionContext& Context, const TDecoratorBinding<IDiscreteBlend>& Binding) const
+	{
+		TDecoratorBinding<IDiscreteBlend> SuperBinding;
+		if (Context.GetInterfaceSuper(Binding, SuperBinding))
+		{
+			return SuperBinding.GetBlendDestinationChildIndex(Context);
+		}
+
+		return INDEX_NONE;
+	}
+
+	void IDiscreteBlend::OnBlendTransition(FExecutionContext& Context, const TDecoratorBinding<IDiscreteBlend>& Binding, int32 OldChildIndex, int32 NewChildIndex) const
+	{
+		TDecoratorBinding<IDiscreteBlend> SuperBinding;
+		if (Context.GetInterfaceSuper(Binding, SuperBinding))
+		{
+			return SuperBinding.OnBlendTransition(Context, OldChildIndex, NewChildIndex);
+		}
+	}
+
+	void IDiscreteBlend::OnBlendInitiated(FExecutionContext& Context, const TDecoratorBinding<IDiscreteBlend>& Binding, int32 ChildIndex) const
+	{
+		TDecoratorBinding<IDiscreteBlend> SuperBinding;
+		if (Context.GetInterfaceSuper(Binding, SuperBinding))
+		{
+			return SuperBinding.OnBlendInitiated(Context, ChildIndex);
+		}
+	}
+
+	void IDiscreteBlend::OnBlendTerminated(FExecutionContext& Context, const TDecoratorBinding<IDiscreteBlend>& Binding, int32 ChildIndex) const
+	{
+		TDecoratorBinding<IDiscreteBlend> SuperBinding;
+		if (Context.GetInterfaceSuper(Binding, SuperBinding))
+		{
+			return SuperBinding.OnBlendTerminated(Context, ChildIndex);
+		}
+	}
+}
diff --git a/Engine/Plugins/Experimental/AnimNext/Source/AnimNext/Private/DecoratorInterfaces/IHierarchy.cpp b/Engine/Plugins/Experimental/AnimNext/Source/AnimNext/Private/DecoratorInterfaces/IHierarchy.cpp
--- a/Engine/Plugins/Experimental/AnimNext/Source/AnimNext/Private/DecoratorInterfaces/IHierarchy.cpp
+++ b/Engine/Plugins/Experimental/AnimNext/Source/AnimNext/Private/DecoratorInterfaces/IHierarchy.cpp
@@ -6,6 +6,17 @@
 
 namespace UE::AnimNext
 {
+	uint32 IHierarchy::GetNumChildren(FExecutionContext& Context, const TDecoratorBinding<IHierarchy>& Binding) const
+	{
+		TDecoratorBinding<IHierarchy> SuperBinding;
+		if (Context.GetInterfaceSuper(Binding, SuperBinding))
+		{
+			return SuperBinding.GetNumChildren(Context);
+		}
+
+		return 0;
+	}
+
 	void IHierarchy::GetChildren(FExecutionContext& Context, const TDecoratorBinding<IHierarchy>& Binding, FChildrenArray& Children) const
 	{
 		TDecoratorBinding<IHierarchy> SuperBinding;
diff --git a/Engine/Plugins/Experimental/AnimNext/Source/AnimNext/Private/DecoratorInterfaces/ISmoothBlend.cpp b/Engine/Plugins/Experimental/AnimNext/Source/AnimNext/Private/DecoratorInterfaces/ISmoothBlend.cpp
new file mode 100644
--- /dev/null
+++ b/Engine/Plugins/Experimental/AnimNext/Source/AnimNext/Private/DecoratorInterfaces/ISmoothBlend.cpp
@@ -0,0 +1,19 @@
+// Copyright Epic Games, Inc. All Rights Reserved.
+
+#include "DecoratorInterfaces/ISmoothBlend.h"
+
+#include "DecoratorBase/ExecutionContext.h"
+
+namespace UE::AnimNext
+{
+	float ISmoothBlend::GetBlendTime(FExecutionContext& Context, const TDecoratorBinding<ISmoothBlend>& Binding, int32 ChildIndex) const
+	{
+		TDecoratorBinding<ISmoothBlend> SuperBinding;
+		if (Context.GetInterfaceSuper(Binding, SuperBinding))
+		{
+			return SuperBinding.GetBlendTime(Context, ChildIndex);
+		}
+
+		return 0.0f;
+	}
+}
diff --git a/Engine/Plugins/Experimental/AnimNext/Source/AnimNext/Private/Decorators/BlendByBool.cpp b/Engine/Plugins/Experimental/AnimNext/Source/AnimNext/Private/Decorators/BlendByBool.cpp
new file mode 100644
--- /dev/null
+++ b/Engine/Plugins/Experimental/AnimNext/Source/AnimNext/Private/Decorators/BlendByBool.cpp
@@ -0,0 +1,172 @@
+// Copyright Epic Games, Inc. All Rights Reserved.
+
+#include "Decorators/BlendByBool.h"
+
+#include "DecoratorBase/ExecutionContext.h"
+#include "EvaluationVM/Tasks/BlendKeyframes.h"
+
+#if 0
+// TODO: Validate blend by bool and blend smoother through a cvar at runtime to select the bool value
+static TAutoConsoleVariable<int32> CVarAnimNextForceBlendBool(TEXT("a.AnimNextForceBlendBool"), -1, TEXT("If != -1, then the value [0 (false), 1 (true)] is used to control Blend By Bool and override its value."));
+#endif
+
+namespace UE::AnimNext
+{
+	AUTO_REGISTER_ANIM_DECORATOR(FBlendByBoolDecorator)
+
+	DEFINE_ANIM_DECORATOR_BEGIN(FBlendByBoolDecorator)
+		DEFINE_ANIM_DECORATOR_IMPLEMENTS_INTERFACE(IDiscreteBlend)
+		DEFINE_ANIM_DECORATOR_IMPLEMENTS_INTERFACE(IEvaluate)
+		DEFINE_ANIM_DECORATOR_IMPLEMENTS_INTERFACE(IHierarchy)
+		DEFINE_ANIM_DECORATOR_IMPLEMENTS_INTERFACE(IUpdate)
+	DEFINE_ANIM_DECORATOR_END(FBlendByBoolDecorator)
+
+	static constexpr int32 TRUE_CHILD_INDEX = 0;
+	static constexpr int32 FALSE_CHILD_INDEX = 1;
+
+	void FBlendByBoolDecorator::PostEvaluate(FExecutionContext& Context, const TDecoratorBinding<IEvaluate>& Binding) const
+	{
+		const FInstanceData* InstanceData = Binding.GetInstanceData<FInstanceData>();
+
+		if (InstanceData->TrueChild.IsValid() && InstanceData->FalseChild.IsValid())
+		{
+			// We have two children, interpolate them
+
+			TDecoratorBinding<IDiscreteBlend> DiscreteBlendDecorator;
+			Context.GetInterface(Binding, DiscreteBlendDecorator);
+
+			const float BlendWeight = DiscreteBlendDecorator.GetBlendWeight(Context, FALSE_CHILD_INDEX);
+
+			FEvaluateTraversalContext& TraversalContext = Context.GetTraversalContext<FEvaluateTraversalContext>();
+			TraversalContext.AppendTask(FAnimNextBlendTwoKeyframesTask::Make(BlendWeight));
+		}
+		else
+		{
+			// We have only one child that is active, do nothing
+		}
+	}
+
+	void FBlendByBoolDecorator::PreUpdate(FExecutionContext& Context, const TDecoratorBinding<IUpdate>& Binding) const
+	{
+		const FSharedData* SharedData = Binding.GetSharedData<FSharedData>();
+		FInstanceData* InstanceData = Binding.GetInstanceData<FInstanceData>();
+
+		TDecoratorBinding<IDiscreteBlend> DiscreteBlendDecorator;
+		Context.GetInterface(Binding, DiscreteBlendDecorator);
+
+		const float DestinationChildIndex = DiscreteBlendDecorator.GetBlendDestinationChildIndex(Context);
+		if (InstanceData->PreviousChildIndex != DestinationChildIndex)
+		{
+			DiscreteBlendDecorator.OnBlendTransition(Context, InstanceData->PreviousChildIndex, DestinationChildIndex);
+
+			InstanceData->PreviousChildIndex = DestinationChildIndex;
+		}
+	}
+
+	uint32 FBlendByBoolDecorator::GetNumChildren(FExecutionContext& Context, const TDecoratorBinding<IHierarchy>& Binding) const
+	{
+		return 2;
+	}
+
+	void FBlendByBoolDecorator::GetChildren(FExecutionContext& Context, const TDecoratorBinding<IHierarchy>& Binding, FChildrenArray& Children) const
+	{
+		const FInstanceData* InstanceData = Binding.GetInstanceData<FInstanceData>();
+
+		// Add the two children, even if the handles are empty
+		Children.Add(InstanceData->TrueChild);
+		Children.Add(InstanceData->FalseChild);
+	}
+
+	float FBlendByBoolDecorator::GetBlendWeight(FExecutionContext& Context, const TDecoratorBinding<IDiscreteBlend>& Binding, int32 ChildIndex) const
+	{
+		TDecoratorBinding<IDiscreteBlend> DiscreteBlendDecorator;
+		Context.GetInterface(Binding, DiscreteBlendDecorator);
+
+		const float DestinationChildIndex = DiscreteBlendDecorator.GetBlendDestinationChildIndex(Context);
+
+		if (ChildIndex == TRUE_CHILD_INDEX)
+		{
+			return DestinationChildIndex == TRUE_CHILD_INDEX ? 1.0f : 0.0f;
+		}
+		else if (ChildIndex == FALSE_CHILD_INDEX)
+		{
+			return DestinationChildIndex == FALSE_CHILD_INDEX ? 1.0f : 0.0f;
+		}
+		else
+		{
+			// Invalid child index
+			return -1.0f;
+		}
+	}
+
+	int32 FBlendByBoolDecorator::GetBlendDestinationChildIndex(FExecutionContext& Context, const TDecoratorBinding<IDiscreteBlend>& Binding) const
+	{
+		const FSharedData* SharedData = Binding.GetSharedData<FSharedData>();
+
+#if 0
+		{
+			const int32 BoolOverride = CVarAnimNextForceBlendBool.GetValueOnAnyThread();
+			if (BoolOverride == 0)
+			{
+				return FALSE_CHILD_INDEX;
+			}
+			else if (BoolOverride == 1)
+			{
+				return TRUE_CHILD_INDEX;
+			}
+		}
+#endif
+
+		const bool bCondition = SharedData->GetbCondition(Context, Binding);
+		return bCondition ? TRUE_CHILD_INDEX : FALSE_CHILD_INDEX;
+	}
+
+	void FBlendByBoolDecorator::OnBlendTransition(FExecutionContext& Context, const TDecoratorBinding<IDiscreteBlend>& Binding, int32 OldChildIndex, int32 NewChildIndex) const
+	{
+		TDecoratorBinding<IDiscreteBlend> DiscreteBlendDecorator;
+		Context.GetInterface(Binding, DiscreteBlendDecorator);
+
+		// We initiate immediately when we transition
+		DiscreteBlendDecorator.OnBlendInitiated(Context, NewChildIndex);
+
+		// We terminate immediately when we transition
+		DiscreteBlendDecorator.OnBlendTerminated(Context, OldChildIndex);
+	}
+
+	void FBlendByBoolDecorator::OnBlendInitiated(FExecutionContext& Context, const TDecoratorBinding<IDiscreteBlend>& Binding, int32 ChildIndex) const
+	{
+		const FSharedData* SharedData = Binding.GetSharedData<FSharedData>();
+		FInstanceData* InstanceData = Binding.GetInstanceData<FInstanceData>();
+
+		// Allocate our new child instance
+		if (ChildIndex == TRUE_CHILD_INDEX)
+		{
+			if (!InstanceData->TrueChild.IsValid())
+			{
+				InstanceData->TrueChild = Context.AllocateNodeInstance(Binding, SharedData->TrueChild);
+			}
+		}
+		else if (ChildIndex == FALSE_CHILD_INDEX)
+		{
+			if (!InstanceData->FalseChild.IsValid())
+			{
+				InstanceData->FalseChild = Context.AllocateNodeInstance(Binding, SharedData->FalseChild);
+			}
+		}
+	}
+
+	void FBlendByBoolDecorator::OnBlendTerminated(FExecutionContext& Context, const TDecoratorBinding<IDiscreteBlend>& Binding, int32 ChildIndex) const
+	{
+		FInstanceData* InstanceData = Binding.GetInstanceData<FInstanceData>();
+
+		// Deallocate our child instance
+		if (ChildIndex == TRUE_CHILD_INDEX)
+		{
+			InstanceData->TrueChild.Reset();
+		}
+		else if (ChildIndex == FALSE_CHILD_INDEX)
+		{
+			InstanceData->FalseChild.Reset();
+		}
+	}
+}
diff --git a/Engine/Plugins/Experimental/AnimNext/Source/AnimNext/Private/Decorators/BlendByBool.h b/Engine/Plugins/Experimental/AnimNext/Source/AnimNext/Private/Decorators/BlendByBool.h
new file mode 100644
--- /dev/null
+++ b/Engine/Plugins/Experimental/AnimNext/Source/AnimNext/Private/Decorators/BlendByBool.h
@@ -0,0 +1,78 @@
+// Copyright Epic Games, Inc. All Rights Reserved.
+
+#pragma once
+
+#include "CoreMinimal.h"
+
+#include "DecoratorBase/Decorator.h"
+#include "DecoratorInterfaces/IDiscreteBlend.h"
+#include "DecoratorInterfaces/IEvaluate.h"
+#include "DecoratorInterfaces/IHierarchy.h"
+#include "DecoratorInterfaces/IUpdate.h"
+
+#include "BlendByBool.generated.h"
+
+USTRUCT()
+struct FAnimNextBlendByBoolDecoratorSharedData : public FAnimNextDecoratorSharedData
+{
+	GENERATED_BODY()
+
+	/** First output to be blended. */
+	UPROPERTY()
+	FAnimNextDecoratorHandle TrueChild;
+
+	/** Second output to be blended. */
+	UPROPERTY()
+	FAnimNextDecoratorHandle FalseChild;
+
+	/** The boolean condition that decides which child is active. */
+	UPROPERTY()
+	bool bCondition = false;
+
+	// Latent pin support boilerplate
+	#define DECORATOR_LATENT_PROPERTIES_ENUMERATOR(GeneratorMacro) \
+		GeneratorMacro(bCondition) \
+
+	GENERATE_DECORATOR_LATENT_PROPERTIES(FAnimNextBlendByBoolDecoratorSharedData, DECORATOR_LATENT_PROPERTIES_ENUMERATOR)
+	#undef DECORATOR_LATENT_PROPERTIES_ENUMERATOR
+};
+
+namespace UE::AnimNext
+{
+	/**
+	 * FBlendByBoolDecorator
+	 * 
+	 * A decorator that can blend two discrete inputs through a boolean.
+	 */
+	struct FBlendByBoolDecorator : FBaseDecorator, IEvaluate, IUpdate, IHierarchy, IDiscreteBlend
+	{
+		DECLARE_ANIM_DECORATOR(FBlendByBoolDecorator, 0xc6d8c9ea, FBaseDecorator)
+
+		using FSharedData = FAnimNextBlendByBoolDecoratorSharedData;
+
+		struct FInstanceData : FDecorator::FInstanceData
+		{
+			FDecoratorPtr TrueChild;
+			FDecoratorPtr FalseChild;
+
+			int32 PreviousChildIndex = INDEX_NONE;
+		};
+
+		// IEvaluate impl
+		virtual void PostEvaluate(FExecutionContext& Context, const TDecoratorBinding<IEvaluate>& Binding) const override;
+
+		// IUpdate impl
+		virtual void PreUpdate(FExecutionContext& Context, const TDecoratorBinding<IUpdate>& Binding) const override;
+
+		// IHierarchy impl
+		virtual uint32 GetNumChildren(FExecutionContext& Context, const TDecoratorBinding<IHierarchy>& Binding) const override;
+		virtual void GetChildren(FExecutionContext& Context, const TDecoratorBinding<IHierarchy>& Binding, FChildrenArray& Children) const override;
+
+		// IDiscreteBlend impl
+		virtual float GetBlendWeight(FExecutionContext& Context, const TDecoratorBinding<IDiscreteBlend>& Binding, int32 ChildIndex) const override;
+		virtual int32 GetBlendDestinationChildIndex(FExecutionContext& Context, const TDecoratorBinding<IDiscreteBlend>& Binding) const override;
+		virtual void OnBlendTransition(FExecutionContext& Context, const TDecoratorBinding<IDiscreteBlend>& Binding, int32 OldChildIndex, int32 NewChildIndex) const override;
+		virtual void OnBlendInitiated(FExecutionContext& Context, const TDecoratorBinding<IDiscreteBlend>& Binding, int32 ChildIndex) const override;
+		virtual void OnBlendTerminated(FExecutionContext& Context, const TDecoratorBinding<IDiscreteBlend>& Binding, int32 ChildIndex) const override;
+	};
+}
diff --git a/Engine/Plugins/Experimental/AnimNext/Source/AnimNext/Private/Decorators/BlendInertializer.cpp b/Engine/Plugins/Experimental/AnimNext/Source/AnimNext/Private/Decorators/BlendInertializer.cpp
new file mode 100644
--- /dev/null
+++ b/Engine/Plugins/Experimental/AnimNext/Source/AnimNext/Private/Decorators/BlendInertializer.cpp
@@ -0,0 +1,48 @@
+// Copyright Epic Games, Inc. All Rights Reserved.
+
+#include "Decorators/BlendInertializer.h"
+
+#include "DecoratorBase/ExecutionContext.h"
+
+namespace UE::AnimNext
+{
+	AUTO_REGISTER_ANIM_DECORATOR(FBlendInertializerDecorator)
+
+	DEFINE_ANIM_DECORATOR_BEGIN(FBlendInertializerDecorator)
+		DEFINE_ANIM_DECORATOR_IMPLEMENTS_INTERFACE(IDiscreteBlend)
+		DEFINE_ANIM_DECORATOR_IMPLEMENTS_INTERFACE(ISmoothBlend)
+	DEFINE_ANIM_DECORATOR_END(FBlendInertializerDecorator)
+
+	void FBlendInertializerDecorator::OnBlendTransition(FExecutionContext& Context, const TDecoratorBinding<IDiscreteBlend>& Binding, int32 OldChildIndex, int32 NewChildIndex) const
+	{
+		// Trigger the new transition
+		IDiscreteBlend::OnBlendTransition(Context, Binding, OldChildIndex, NewChildIndex);
+
+		// TODO: Implement the inertialization request API
+#if 0
+		UE::Anim::IInertializationRequester* InertializationRequester = Context.GetMessage<UE::Anim::IInertializationRequester>();
+		if (InertializationRequester)
+		{
+			FInertializationRequest Request;
+			Request.Duration = CurrentBlendTimes[ChildIndex];		// TODO: Get from ISmoothBlend interface using super
+			Request.BlendProfile = CurrentBlendProfile;				// TODO: Get from shared data
+			Request.bUseBlendMode = true;
+			Request.BlendMode = GetBlendType();						// TODO: Add to ISmoothBlend interface
+			Request.CustomBlendCurve = GetCustomBlendCurve();		// TODO: Add to ISmoothBlend interface
+
+			InertializationRequester->RequestInertialization(Request);
+			bRequestedInertializationOnActiveChildIndexChange = true;
+		}
+		else
+		{
+			FAnimNode_Inertialization::LogRequestError(Context, BlendPose[ChildIndex]);
+		}
+#endif
+	}
+
+	float FBlendInertializerDecorator::GetBlendTime(FExecutionContext& Context, const TDecoratorBinding<ISmoothBlend>& Binding, int32 ChildIndex) const
+	{
+		// We hijack the blend time and always transition instantaneously
+		return 0.0f;
+	}
+}
diff --git a/Engine/Plugins/Experimental/AnimNext/Source/AnimNext/Private/Decorators/BlendInertializer.h b/Engine/Plugins/Experimental/AnimNext/Source/AnimNext/Private/Decorators/BlendInertializer.h
new file mode 100644
--- /dev/null
+++ b/Engine/Plugins/Experimental/AnimNext/Source/AnimNext/Private/Decorators/BlendInertializer.h
@@ -0,0 +1,46 @@
+// Copyright Epic Games, Inc. All Rights Reserved.
+
+#pragma once
+
+#include "CoreMinimal.h"
+
+#include "DecoratorBase/Decorator.h"
+#include "DecoratorInterfaces/IDiscreteBlend.h"
+#include "DecoratorInterfaces/ISmoothBlend.h"
+
+#include "BlendInertializer.generated.h"
+
+class UBlendProfile;
+
+USTRUCT()
+struct FAnimNextBlendInertializerDecoratorSharedData : public FAnimNextDecoratorSharedData
+{
+	GENERATED_BODY()
+
+	/** Blend profile that configures how fast to blend each bone. */
+	// TODO: Can't show list of blend profiles, we need to find a skeleton to perform the lookup with
+	UPROPERTY(meta = (Inline, UseAsBlendProfile = true))
+	TObjectPtr<UBlendProfile> BlendProfile = nullptr;
+};
+
+namespace UE::AnimNext
+{
+	/**
+	 * FBlendInertializerDecorator
+	 * 
+	 * A decorator that converts a normal smooth blend into an inertializing blend.
+	 */
+	struct FBlendInertializerDecorator : FAdditiveDecorator, IDiscreteBlend, ISmoothBlend
+	{
+		DECLARE_ANIM_DECORATOR(FBlendInertializerDecorator, 0xa6f83827, FAdditiveDecorator)
+
+		using FSharedData = FAnimNextBlendInertializerDecoratorSharedData;
+		using FInstanceData = FDecorator::FInstanceData;
+
+		// IDiscreteBlend impl
+		virtual void OnBlendTransition(FExecutionContext& Context, const TDecoratorBinding<IDiscreteBlend>& Binding, int32 OldChildIndex, int32 NewChildIndex) const override;
+
+		// ISmoothBlend impl
+		virtual float GetBlendTime(FExecutionContext& Context, const TDecoratorBinding<ISmoothBlend>& Binding, int32 ChildIndex) const override;
+	};
+}
diff --git a/Engine/Plugins/Experimental/AnimNext/Source/AnimNext/Private/Decorators/BlendSmoother.cpp b/Engine/Plugins/Experimental/AnimNext/Source/AnimNext/Private/Decorators/BlendSmoother.cpp
new file mode 100644
--- /dev/null
+++ b/Engine/Plugins/Experimental/AnimNext/Source/AnimNext/Private/Decorators/BlendSmoother.cpp
@@ -0,0 +1,182 @@
+// Copyright Epic Games, Inc. All Rights Reserved.
+
+#include "Decorators/BlendSmoother.h"
+
+#include "DecoratorBase/ExecutionContext.h"
+#include "DecoratorInterfaces/IHierarchy.h"
+#include "EvaluationVM/Tasks/BlendKeyframes.h"
+
+namespace UE::AnimNext
+{
+	AUTO_REGISTER_ANIM_DECORATOR(FBlendSmootherDecorator)
+
+	DEFINE_ANIM_DECORATOR_BEGIN(FBlendSmootherDecorator)
+		DEFINE_ANIM_DECORATOR_IMPLEMENTS_INTERFACE(IUpdate)
+		DEFINE_ANIM_DECORATOR_IMPLEMENTS_INTERFACE(IDiscreteBlend)
+		DEFINE_ANIM_DECORATOR_IMPLEMENTS_INTERFACE(ISmoothBlend)
+	DEFINE_ANIM_DECORATOR_END(FBlendSmootherDecorator)
+
+	void FBlendSmootherDecorator::PreUpdate(FExecutionContext& Context, const TDecoratorBinding<IUpdate>& Binding) const
+	{
+		const FSharedData* SharedData = Binding.GetSharedData<FSharedData>();
+		FInstanceData* InstanceData = Binding.GetInstanceData<FInstanceData>();
+
+		// If this is our first update, allocate our blend data
+		if (InstanceData->PerChildBlendData.IsEmpty())
+		{
+			InitializeInstanceData(Context, Binding, SharedData, InstanceData);
+		}
+
+		// Update the decorators below us, they might trigger a transition
+		IUpdate::PreUpdate(Context, Binding);
+
+		const FUpdateTraversalContext& TraversalContext = Context.GetTraversalContext<FUpdateTraversalContext>();
+		const float DeltaTime = TraversalContext.GetDeltaTime();
+
+		// Advance the weights
+		float SumWeight = 0.0f;
+		uint32 NumBlending = 0;
+
+		for (FBlendData& ChildBlendData : InstanceData->PerChildBlendData)
+		{
+			if (!ChildBlendData.bIsBlending)
+			{
+				continue;	// Skip children that aren't blending
+			}
+
+			ChildBlendData.Blend.Update(DeltaTime);
+
+			const float NewBlendWeight = ChildBlendData.Blend.GetBlendedValue();
+
+			ChildBlendData.Weight = NewBlendWeight;
+			SumWeight += NewBlendWeight;
+			NumBlending++;
+		}
+
+		if (NumBlending <= 1)
+		{
+			return;	// Nothing to do if we don't blend at least 2 children together
+		}
+
+		// Renormalize the weights if the sum isn't near 0.0 or near 1.0
+		if (SumWeight > ZERO_ANIMWEIGHT_THRESH &&
+			FMath::Abs(SumWeight - 1.0f) > ZERO_ANIMWEIGHT_THRESH)
+		{
+			const float ReciprocalSum = 1.0f / SumWeight;
+
+			for (FBlendData& ChildBlendData : InstanceData->PerChildBlendData)
+			{
+				ChildBlendData.Weight *= ReciprocalSum;
+			}
+		}
+
+		TDecoratorBinding<IDiscreteBlend> DiscreteBlendDecorator;
+		Context.GetInterface(Binding, DiscreteBlendDecorator);
+
+		// Free any newly inactive children
+		const int32 NumChildren = InstanceData->PerChildBlendData.Num();
+		for (int32 ChildIndex = 0; ChildIndex < NumChildren; ++ChildIndex)
+		{
+			FBlendData& ChildBlendData = InstanceData->PerChildBlendData[ChildIndex];
+
+			if (ChildBlendData.bIsBlending && ChildBlendData.Weight <= 0.0f)
+			{
+				// This child has finished blending out, terminate it
+				DiscreteBlendDecorator.OnBlendTerminated(Context, ChildIndex);
+
+				ChildBlendData.bIsBlending = false;
+			}
+		}
+	}
+
+	float FBlendSmootherDecorator::GetBlendWeight(FExecutionContext& Context, const TDecoratorBinding<IDiscreteBlend>& Binding, int32 ChildIndex) const
+	{
+		const FInstanceData* InstanceData = Binding.GetInstanceData<FInstanceData>();
+		return InstanceData->PerChildBlendData.IsValidIndex(ChildIndex) ? InstanceData->PerChildBlendData[ChildIndex].Weight : -1.0f;
+	}
+
+	const FAlphaBlend* FBlendSmootherDecorator::GetBlendState(FExecutionContext& Context, const TDecoratorBinding<IDiscreteBlend>& Binding, int32 ChildIndex) const
+	{
+		const FInstanceData* InstanceData = Binding.GetInstanceData<FInstanceData>();
+		return InstanceData->PerChildBlendData.IsValidIndex(ChildIndex) ? &InstanceData->PerChildBlendData[ChildIndex].Blend : nullptr;
+	}
+
+	void FBlendSmootherDecorator::OnBlendTransition(FExecutionContext& Context, const TDecoratorBinding<IDiscreteBlend>& Binding, int32 OldChildIndex, int32 NewChildIndex) const
+	{
+		const FSharedData* SharedData = Binding.GetSharedData<FSharedData>();
+		FInstanceData* InstanceData = Binding.GetInstanceData<FInstanceData>();
+
+		// scale by the weight difference since we want consistency:
+		// - if you're moving from 0 to full weight 1, it will use the normal blend time
+		// - if you're moving from 0.5 to full weight 1, it will get there in half the time
+		const float NewChildCurrentWeight = InstanceData->PerChildBlendData[NewChildIndex].Weight;
+		const float NewChildDesiredWeight = 1.0f;
+		const float WeightDifference = FMath::Clamp(FMath::Abs(NewChildDesiredWeight - NewChildCurrentWeight), 0.0f, 1.0f);
+
+		TDecoratorBinding<ISmoothBlend> SmoothBlendDecorator;
+		Context.GetInterface(Binding, SmoothBlendDecorator);
+
+		const float BlendTime = SmoothBlendDecorator.GetBlendTime(Context, NewChildIndex);
+		const float RemainingBlendTime = OldChildIndex != INDEX_NONE ? (BlendTime * WeightDifference) : 0.0f;
+
+		if (OldChildIndex != INDEX_NONE)
+		{
+			// Make sure the old child starts blending out
+			FBlendData& OldChildBlendData = InstanceData->PerChildBlendData[OldChildIndex];
+			OldChildBlendData.Blend.SetValueRange(OldChildBlendData.Weight, 0.0f);
+			check(OldChildBlendData.bIsBlending);
+		}
+
+		{
+			// Setup the new child to blend in
+			FBlendData& NewChildBlendData = InstanceData->PerChildBlendData[NewChildIndex];
+			NewChildBlendData.Blend.SetValueRange(NewChildBlendData.Weight, 1.0f);
+			NewChildBlendData.Blend.ResetAlpha();	// Reset the alpha right away in case another decorator needs it
+			NewChildBlendData.bIsBlending = true;
+		}
+
+		// We set the new blend time on all children
+		for (FBlendData& ChildBlendData : InstanceData->PerChildBlendData)
+		{
+			ChildBlendData.Blend.SetBlendTime(RemainingBlendTime);
+		}
+
+		// Don't call the super since we hijack the transition to smooth it out over time
+		// We just initiate the new blend manually
+
+		TDecoratorBinding<IDiscreteBlend> DiscreteBlendDecorator;
+		Context.GetInterface(Binding, DiscreteBlendDecorator);
+
+		DiscreteBlendDecorator.OnBlendInitiated(Context, NewChildIndex);
+	}
+
+	float FBlendSmootherDecorator::GetBlendTime(FExecutionContext& Context, const TDecoratorBinding<ISmoothBlend>& Binding, int32 ChildIndex) const
+	{
+		const FSharedData* SharedData = Binding.GetSharedData<FSharedData>();
+		return SharedData->BlendTimes.IsValidIndex(ChildIndex) ? SharedData->BlendTimes[ChildIndex] : 0.0f;
+	}
+
+	void FBlendSmootherDecorator::InitializeInstanceData(FExecutionContext& Context, const FDecoratorBinding& Binding, const FSharedData* SharedData, FInstanceData* InstanceData)
+	{
+		check(InstanceData->PerChildBlendData.IsEmpty());
+
+		uint32 NumExpectedChildren = SharedData->BlendTimes.Num();
+
+#if DO_CHECK
+		TDecoratorBinding<IHierarchy> HierarchyDecorator;
+		Context.GetInterface(Binding, HierarchyDecorator);
+
+		const uint32 NumActualChildren = HierarchyDecorator.GetNumChildren(Context);
+		ensureMsgf(NumActualChildren != 0, TEXT("BlendSmootherDecorator has %u blend times for %u children"), NumExpectedChildren, NumActualChildren);
+		NumExpectedChildren = NumActualChildren;
+#endif
+
+		InstanceData->PerChildBlendData.SetNum(NumExpectedChildren);
+
+		for (FBlendData& ChildBlendData : InstanceData->PerChildBlendData)
+		{
+			ChildBlendData.Blend.SetBlendOption(SharedData->BlendType);
+			ChildBlendData.Blend.SetCustomCurve(SharedData->CustomBlendCurve);
+		}
+	}
+}
diff --git a/Engine/Plugins/Experimental/AnimNext/Source/AnimNext/Private/Decorators/BlendSmoother.h b/Engine/Plugins/Experimental/AnimNext/Source/AnimNext/Private/Decorators/BlendSmoother.h
new file mode 100644
--- /dev/null
+++ b/Engine/Plugins/Experimental/AnimNext/Source/AnimNext/Private/Decorators/BlendSmoother.h
@@ -0,0 +1,81 @@
+// Copyright Epic Games, Inc. All Rights Reserved.
+
+#pragma once
+
+#include "CoreMinimal.h"
+
+#include "AlphaBlend.h"
+#include "DecoratorBase/Decorator.h"
+#include "DecoratorInterfaces/IDiscreteBlend.h"
+#include "DecoratorInterfaces/ISmoothBlend.h"
+#include "DecoratorInterfaces/IUpdate.h"
+
+#include "BlendSmoother.generated.h"
+
+class UCurveFloat;
+
+USTRUCT()
+struct FAnimNextBlendSmootherDecoratorSharedData : public FAnimNextDecoratorSharedData
+{
+	GENERATED_BODY()
+
+	/** How long to take when blending into each child. */
+	UPROPERTY(meta = (Inline))
+	TArray<float> BlendTimes;
+
+	/** What type of blend equation to use when converting the time elapsed into a blend weight. */
+	UPROPERTY(meta = (Inline))
+	EAlphaBlendOption BlendType = EAlphaBlendOption::Linear;
+
+	/** Custom curve to use when the Custom blend type is used. */
+	UPROPERTY(meta = (Inline))
+	TObjectPtr<UCurveFloat> CustomBlendCurve;
+};
+
+namespace UE::AnimNext
+{
+	/**
+	 * FBlendSmootherDecorator
+	 * 
+	 * A decorator that smoothly blends between discrete states over time.
+	 */
+	struct FBlendSmootherDecorator : FAdditiveDecorator, IUpdate, IDiscreteBlend, ISmoothBlend
+	{
+		DECLARE_ANIM_DECORATOR(FBlendSmootherDecorator, 0x7b6c3d2e, FAdditiveDecorator)
+
+		using FSharedData = FAnimNextBlendSmootherDecoratorSharedData;
+
+		// Struct for tracking blends for each pose
+		struct FBlendData
+		{
+			// Helper struct to update a time based weight
+			FAlphaBlend Blend;
+
+			// Current child weight (normalized with all children)
+			float Weight = 0.0f;
+
+			// Whether or not this child is actively blending
+			bool bIsBlending = false;
+		};
+
+		struct FInstanceData : FDecorator::FInstanceData
+		{
+			// Blend state per child
+			TArray<FBlendData> PerChildBlendData;
+		};
+
+		// IUpdate impl
+		virtual void PreUpdate(FExecutionContext& Context, const TDecoratorBinding<IUpdate>& Binding) const override;
+
+		// IDiscreteBlend impl
+		virtual float GetBlendWeight(FExecutionContext& Context, const TDecoratorBinding<IDiscreteBlend>& Binding, int32 ChildIndex) const override;
+		virtual const FAlphaBlend* GetBlendState(FExecutionContext& Context, const TDecoratorBinding<IDiscreteBlend>& Binding, int32 ChildIndex) const override;
+		virtual void OnBlendTransition(FExecutionContext& Context, const TDecoratorBinding<IDiscreteBlend>& Binding, int32 OldChildIndex, int32 NewChildIndex) const override;
+
+		// ISmoothBlend impl
+		virtual float GetBlendTime(FExecutionContext& Context, const TDecoratorBinding<ISmoothBlend>& Binding, int32 ChildIndex) const override;
+
+		// Internal impl
+		static void InitializeInstanceData(FExecutionContext& Context, const FDecoratorBinding& Binding, const FSharedData* SharedData, FInstanceData* InstanceData);
+	};
+}
diff --git a/Engine/Plugins/Experimental/AnimNext/Source/AnimNext/Private/Decorators/BlendTwoWay.cpp b/Engine/Plugins/Experimental/AnimNext/Source/AnimNext/Private/Decorators/BlendTwoWay.cpp
--- a/Engine/Plugins/Experimental/AnimNext/Source/AnimNext/Private/Decorators/BlendTwoWay.cpp
+++ b/Engine/Plugins/Experimental/AnimNext/Source/AnimNext/Private/Decorators/BlendTwoWay.cpp
@@ -2,6 +2,7 @@
 
 #include "Decorators/BlendTwoWay.h"
 
+#include "Animation/AnimTypes.h"
 #include "DecoratorBase/ExecutionContext.h"
 #include "EvaluationVM/Tasks/BlendKeyframes.h"
 
@@ -10,9 +11,10 @@ namespace UE::AnimNext
 	AUTO_REGISTER_ANIM_DECORATOR(FBlendTwoWayDecorator)
 
 	DEFINE_ANIM_DECORATOR_BEGIN(FBlendTwoWayDecorator)
+		DEFINE_ANIM_DECORATOR_IMPLEMENTS_INTERFACE(IContinuousBlend)
 		DEFINE_ANIM_DECORATOR_IMPLEMENTS_INTERFACE(IEvaluate)
-		DEFINE_ANIM_DECORATOR_IMPLEMENTS_INTERFACE(IUpdate)
 		DEFINE_ANIM_DECORATOR_IMPLEMENTS_INTERFACE(IHierarchy)
+		DEFINE_ANIM_DECORATOR_IMPLEMENTS_INTERFACE(IUpdate)
 	DEFINE_ANIM_DECORATOR_END(FBlendTwoWayDecorator)
 
 	void FBlendTwoWayDecorator::PostEvaluate(FExecutionContext& Context, const TDecoratorBinding<IEvaluate>& Binding) const
@@ -22,9 +24,11 @@ namespace UE::AnimNext
 		if (InstanceData->ChildA.IsValid() && InstanceData->ChildB.IsValid())
 		{
 			// We have two children, interpolate them
-			const FSharedData* SharedData = Binding.GetSharedData<FSharedData>();
 
-			const float BlendWeight = SharedData->GetBlendWeight(Context, Binding);
+			TDecoratorBinding<IContinuousBlend> ContinuousBlendDecorator;
+			Context.GetInterface(Binding, ContinuousBlendDecorator);
+
+			const float BlendWeight = ContinuousBlendDecorator.GetBlendWeight(Context, 1);
 
 			FEvaluateTraversalContext& TraversalContext = Context.GetTraversalContext<FEvaluateTraversalContext>();
 			TraversalContext.AppendTask(FAnimNextBlendTwoKeyframesTask::Make(BlendWeight));
@@ -40,8 +44,11 @@ namespace UE::AnimNext
 		const FSharedData* SharedData = Binding.GetSharedData<FSharedData>();
 		FInstanceData* InstanceData = Binding.GetInstanceData<FInstanceData>();
 
-		const float BlendWeight = SharedData->GetBlendWeight(Context, Binding);
-		if (BlendWeight < 1.0f)
+		TDecoratorBinding<IContinuousBlend> ContinuousBlendDecorator;
+		Context.GetInterface(Binding, ContinuousBlendDecorator);
+
+		const float BlendWeight = ContinuousBlendDecorator.GetBlendWeight(Context, 1);
+		if (!FAnimWeight::IsFullWeight(BlendWeight))
 		{
 			if (!InstanceData->ChildA.IsValid())
 			{
@@ -49,14 +56,14 @@ namespace UE::AnimNext
 				InstanceData->ChildA = Context.AllocateNodeInstance(Binding, SharedData->ChildA);
 			}
 
-			if (BlendWeight == 0.0f)
+			if (!FAnimWeight::IsRelevant(BlendWeight))
 			{
 				// We no longer need this child, release it
 				InstanceData->ChildB.Reset();
 			}
 		}
 
-		if (BlendWeight > 0.0f)
+		if (FAnimWeight::IsRelevant(BlendWeight))
 		{
 			if (!InstanceData->ChildB.IsValid())
 			{
@@ -64,7 +71,7 @@ namespace UE::AnimNext
 				InstanceData->ChildB = Context.AllocateNodeInstance(Binding, SharedData->ChildB);
 			}
 
-			if (BlendWeight == 1.0f)
+			if (FAnimWeight::IsFullWeight(BlendWeight))
 			{
 				// We no longer need this child, release it
 				InstanceData->ChildA.Reset();
@@ -72,12 +79,38 @@ namespace UE::AnimNext
 		}
 	}
 
+	uint32 FBlendTwoWayDecorator::GetNumChildren(FExecutionContext& Context, const TDecoratorBinding<IHierarchy>& Binding) const
+	{
+		return 2;
+	}
+
 	void FBlendTwoWayDecorator::GetChildren(FExecutionContext& Context, const TDecoratorBinding<IHierarchy>& Binding, FChildrenArray& Children) const
 	{
 		const FInstanceData* InstanceData = Binding.GetInstanceData<FInstanceData>();
 
-		// Add the two child handles, even if they are empty
+		// Add the two children, even if the handles are empty
 		Children.Add(InstanceData->ChildA);
 		Children.Add(InstanceData->ChildB);
 	}
+
+	float FBlendTwoWayDecorator::GetBlendWeight(FExecutionContext& Context, const TDecoratorBinding<IContinuousBlend>& Binding, int32 ChildIndex) const
+	{
+		const FSharedData* SharedData = Binding.GetSharedData<FSharedData>();
+
+		const float BlendWeight = SharedData->GetBlendWeight(Context, Binding);
+
+		if (ChildIndex == 0)
+		{
+			return 1.0f - BlendWeight;
+		}
+		else if (ChildIndex == 1)
+		{
+			return BlendWeight;
+		}
+		else
+		{
+			// Invalid child index
+			return -1.0f;
+		}
+	}
 }
diff --git a/Engine/Plugins/Experimental/AnimNext/Source/AnimNext/Private/Decorators/BlendTwoWay.h b/Engine/Plugins/Experimental/AnimNext/Source/AnimNext/Private/Decorators/BlendTwoWay.h
--- a/Engine/Plugins/Experimental/AnimNext/Source/AnimNext/Private/Decorators/BlendTwoWay.h
+++ b/Engine/Plugins/Experimental/AnimNext/Source/AnimNext/Private/Decorators/BlendTwoWay.h
@@ -5,6 +5,7 @@
 #include "CoreMinimal.h"
 
 #include "DecoratorBase/Decorator.h"
+#include "DecoratorInterfaces/IContinuousBlend.h"
 #include "DecoratorInterfaces/IEvaluate.h"
 #include "DecoratorInterfaces/IHierarchy.h"
 #include "DecoratorInterfaces/IUpdate.h"
@@ -43,7 +44,7 @@ namespace UE::AnimNext
 	 * 
 	 * A decorator that can blend two inputs.
 	 */
-	struct FBlendTwoWayDecorator : FBaseDecorator, IEvaluate, IUpdate, IHierarchy
+	struct FBlendTwoWayDecorator : FBaseDecorator, IEvaluate, IUpdate, IHierarchy, IContinuousBlend
 	{
 		DECLARE_ANIM_DECORATOR(FBlendTwoWayDecorator, 0x96a81d1e, FBaseDecorator)
 
@@ -62,6 +63,10 @@ namespace UE::AnimNext
 		virtual void PreUpdate(FExecutionContext& Context, const TDecoratorBinding<IUpdate>& Binding) const override;
 
 		// IHierarchy impl
+		virtual uint32 GetNumChildren(FExecutionContext& Context, const TDecoratorBinding<IHierarchy>& Binding) const override;
 		virtual void GetChildren(FExecutionContext& Context, const TDecoratorBinding<IHierarchy>& Binding, FChildrenArray& Children) const override;
+
+		// IContinuousBlend impl
+		virtual float GetBlendWeight(FExecutionContext& Context, const TDecoratorBinding<IContinuousBlend>& Binding, int32 ChildIndex) const override;
 	};
 }
diff --git a/Engine/Plugins/Experimental/AnimNext/Source/AnimNext/Private/Decorators/SequencePlayer.cpp b/Engine/Plugins/Experimental/AnimNext/Source/AnimNext/Private/Decorators/SequencePlayer.cpp
--- a/Engine/Plugins/Experimental/AnimNext/Source/AnimNext/Private/Decorators/SequencePlayer.cpp
+++ b/Engine/Plugins/Experimental/AnimNext/Source/AnimNext/Private/Decorators/SequencePlayer.cpp
@@ -53,7 +53,7 @@ namespace UE::AnimNext
 
 		if (SharedData->AnimSequence != nullptr)
 		{
-			FUpdateTraversalContext& TraversalContext = Context.GetTraversalContext<FUpdateTraversalContext>();
+			const FUpdateTraversalContext& TraversalContext = Context.GetTraversalContext<FUpdateTraversalContext>();
 
 			TDecoratorBinding<ITimeline> TimelineDecorator;
 			Context.GetInterface(Binding, TimelineDecorator);
diff --git a/Engine/Plugins/Experimental/AnimNext/Source/AnimNext/Private/Module.cpp b/Engine/Plugins/Experimental/AnimNext/Source/AnimNext/Private/Module.cpp
--- a/Engine/Plugins/Experimental/AnimNext/Source/AnimNext/Private/Module.cpp
+++ b/Engine/Plugins/Experimental/AnimNext/Source/AnimNext/Private/Module.cpp
@@ -2,6 +2,8 @@
 
 #include "CoreMinimal.h"
 #include "AnimNextConfig.h"
+#include "Animation/BlendProfile.h"
+#include "Curves/CurveFloat.h"
 #include "Modules/ModuleInterface.h"
 #include "Modules/ModuleManager.h"
 #include "Misc/CoreDelegates.h"
@@ -30,7 +32,9 @@ public:
 		static TPair<UClass*, FRigVMRegistry::ERegisterObjectOperation> const AllowedObjectTypes[] =
 		{
 			{ UAnimSequence::StaticClass(), FRigVMRegistry::ERegisterObjectOperation::Class },
-			{ UScriptStruct::StaticClass(), FRigVMRegistry::ERegisterObjectOperation::Class }
+			{ UScriptStruct::StaticClass(), FRigVMRegistry::ERegisterObjectOperation::Class },
+			{ UBlendProfile::StaticClass(), FRigVMRegistry::ERegisterObjectOperation::Class },
+			{ UCurveFloat::StaticClass(), FRigVMRegistry::ERegisterObjectOperation::Class },
 		};
 
 		FRigVMRegistry::Get().RegisterObjectTypes(AllowedObjectTypes);
diff --git a/Engine/Plugins/Experimental/AnimNext/Source/AnimNext/Public/DecoratorInterfaces/IContinuousBlend.h b/Engine/Plugins/Experimental/AnimNext/Source/AnimNext/Public/DecoratorInterfaces/IContinuousBlend.h
new file mode 100644
--- /dev/null
+++ b/Engine/Plugins/Experimental/AnimNext/Source/AnimNext/Public/DecoratorInterfaces/IContinuousBlend.h
@@ -0,0 +1,40 @@
+// Copyright Epic Games, Inc. All Rights Reserved.
+
+#pragma once
+
+#include "CoreMinimal.h"
+#include "DecoratorBase/IDecoratorInterface.h"
+
+namespace UE::AnimNext
+{
+	/**
+	 * IContinuousBlend
+	 *
+	 * This interface exposes continuous blend related information.
+	 */
+	struct ANIMNEXT_API IContinuousBlend : IDecoratorInterface
+	{
+		DECLARE_ANIM_DECORATOR_INTERFACE(IContinuousBlend, 0xe7d79186)
+
+		// Returns the blend weight for the specified child
+		// Multiple children can have non-zero weight but their sum must be 1.0
+		// Returns -1.0 if the child index is invalid
+		virtual float GetBlendWeight(FExecutionContext& Context, const TDecoratorBinding<IContinuousBlend>& Binding, int32 ChildIndex) const;
+	};
+
+	/**
+	 * Specialization for decorator binding.
+	 */
+	template<>
+	struct TDecoratorBinding<IContinuousBlend> : FDecoratorBinding
+	{
+		// @see IContinuousBlend::GetBlendWeight
+		float GetBlendWeight(FExecutionContext& Context, int32 ChildIndex) const
+		{
+			return GetInterface()->GetBlendWeight(Context, *this, ChildIndex);
+		}
+
+	protected:
+		const IContinuousBlend* GetInterface() const { return GetInterfaceTyped<IContinuousBlend>(); }
+	};
+}
diff --git a/Engine/Plugins/Experimental/AnimNext/Source/AnimNext/Public/DecoratorInterfaces/IDiscreteBlend.h b/Engine/Plugins/Experimental/AnimNext/Source/AnimNext/Public/DecoratorInterfaces/IDiscreteBlend.h
new file mode 100644
--- /dev/null
+++ b/Engine/Plugins/Experimental/AnimNext/Source/AnimNext/Public/DecoratorInterfaces/IDiscreteBlend.h
@@ -0,0 +1,90 @@
+// Copyright Epic Games, Inc. All Rights Reserved.
+
+#pragma once
+
+#include "CoreMinimal.h"
+#include "DecoratorBase/IDecoratorInterface.h"
+
+struct FAlphaBlend;
+
+namespace UE::AnimNext
+{
+	/**
+	 * IDiscreteBlend
+	 *
+	 * This interface exposes discrete blend related information.
+	 */
+	struct ANIMNEXT_API IDiscreteBlend : IDecoratorInterface
+	{
+		DECLARE_ANIM_DECORATOR_INTERFACE(IDiscreteBlend, 0x2d395d56)
+
+		// Returns the blend weight for the specified child
+		// Multiple children can have non-zero weight but their sum must be 1.0
+		// Returns -1.0 if the child index is invalid
+		virtual float GetBlendWeight(FExecutionContext& Context, const TDecoratorBinding<IDiscreteBlend>& Binding, int32 ChildIndex) const;
+
+		// Returns the blend alpha for the specified child
+		// Returns nullptr if the child index is invalid
+		// Allows additive decorators to query the internal state of the base decorator
+		virtual const FAlphaBlend* GetBlendState(FExecutionContext& Context, const TDecoratorBinding<IDiscreteBlend>& Binding, int32 ChildIndex) const;
+
+		// Returns the blend destination child index (aka the active child index)
+		// Returns INDEX_NONE if no child is active
+		virtual int32 GetBlendDestinationChildIndex(FExecutionContext& Context, const TDecoratorBinding<IDiscreteBlend>& Binding) const;
+
+		// Called when a blend transition between children occurs
+		virtual void OnBlendTransition(FExecutionContext& Context, const TDecoratorBinding<IDiscreteBlend>& Binding, int32 OldChildIndex, int32 NewChildIndex) const;
+
+		// Called when the blend for specified child is initiated
+		virtual void OnBlendInitiated(FExecutionContext& Context, const TDecoratorBinding<IDiscreteBlend>& Binding, int32 ChildIndex) const;
+
+		// Called when the blend for specified child terminates
+		virtual void OnBlendTerminated(FExecutionContext& Context, const TDecoratorBinding<IDiscreteBlend>& Binding, int32 ChildIndex) const;
+	};
+
+	/**
+	 * Specialization for decorator binding.
+	 */
+	template<>
+	struct TDecoratorBinding<IDiscreteBlend> : FDecoratorBinding
+	{
+		// @see IDiscreteBlend::GetBlendWeight
+		float GetBlendWeight(FExecutionContext& Context, int32 ChildIndex) const
+		{
+			return GetInterface()->GetBlendWeight(Context, *this, ChildIndex);
+		}
+
+		// @see IDiscreteBlend::GetBlendState
+		const FAlphaBlend* GetBlendState(FExecutionContext& Context, int32 ChildIndex) const
+		{
+			return GetInterface()->GetBlendState(Context, *this, ChildIndex);
+		}
+
+		// @see IDiscreteBlend::GetBlendDestinationChildIndex
+		int32 GetBlendDestinationChildIndex(FExecutionContext& Context) const
+		{
+			return GetInterface()->GetBlendDestinationChildIndex(Context, *this);
+		}
+
+		// @see IDiscreteBlend::OnBlendTransition
+		void OnBlendTransition(FExecutionContext& Context, int32 OldChildIndex, int32 NewChildIndex) const
+		{
+			GetInterface()->OnBlendTransition(Context, *this, OldChildIndex, NewChildIndex);
+		}
+
+		// @see IDiscreteBlend::OnBlendInitiated
+		void OnBlendInitiated(FExecutionContext& Context, int32 ChildIndex) const
+		{
+			GetInterface()->OnBlendInitiated(Context, *this, ChildIndex);
+		}
+
+		// @see IDiscreteBlend::OnBlendTerminated
+		void OnBlendTerminated(FExecutionContext& Context, int32 ChildIndex) const
+		{
+			GetInterface()->OnBlendTerminated(Context, *this, ChildIndex);
+		}
+
+	protected:
+		const IDiscreteBlend* GetInterface() const { return GetInterfaceTyped<IDiscreteBlend>(); }
+	};
+}
diff --git a/Engine/Plugins/Experimental/AnimNext/Source/AnimNext/Public/DecoratorInterfaces/IHierarchy.h b/Engine/Plugins/Experimental/AnimNext/Source/AnimNext/Public/DecoratorInterfaces/IHierarchy.h
--- a/Engine/Plugins/Experimental/AnimNext/Source/AnimNext/Public/DecoratorInterfaces/IHierarchy.h
+++ b/Engine/Plugins/Experimental/AnimNext/Source/AnimNext/Public/DecoratorInterfaces/IHierarchy.h
@@ -19,8 +19,13 @@ namespace UE::AnimNext
 	{
 		DECLARE_ANIM_DECORATOR_INTERFACE(IHierarchy, 0x846d8a37)
 
+		// Returns the number of children
+		// Includes inactive children
+		virtual uint32 GetNumChildren(FExecutionContext& Context, const TDecoratorBinding<IHierarchy>& Binding) const;
+
 		// Appends weak handles to any children we wish to traverse.
 		// Decorators are responsible for allocating and releasing child instance data.
+		// Empty handles and duplicates can be appended.
 		virtual void GetChildren(FExecutionContext& Context, const TDecoratorBinding<IHierarchy>& Binding, FChildrenArray& Children) const;
 	};
 
@@ -30,6 +35,12 @@ namespace UE::AnimNext
 	template<>
 	struct TDecoratorBinding<IHierarchy> : FDecoratorBinding
 	{
+		// @see IHierarchy::GetNumChildren
+		uint32 GetNumChildren(FExecutionContext& Context) const
+		{
+			return GetInterface()->GetNumChildren(Context, *this);
+		}
+
 		// @see IHierarchy::GetChildren
 		void GetChildren(FExecutionContext& Context, FChildrenArray& Children) const
 		{
diff --git a/Engine/Plugins/Experimental/AnimNext/Source/AnimNext/Public/DecoratorInterfaces/ISmoothBlend.h b/Engine/Plugins/Experimental/AnimNext/Source/AnimNext/Public/DecoratorInterfaces/ISmoothBlend.h
new file mode 100644
--- /dev/null
+++ b/Engine/Plugins/Experimental/AnimNext/Source/AnimNext/Public/DecoratorInterfaces/ISmoothBlend.h
@@ -0,0 +1,38 @@
+// Copyright Epic Games, Inc. All Rights Reserved.
+
+#pragma once
+
+#include "CoreMinimal.h"
+#include "DecoratorBase/IDecoratorInterface.h"
+
+namespace UE::AnimNext
+{
+	/**
+	 * ISmoothBlend
+	 *
+	 * This interface exposes blend smoothing related information.
+	 */
+	struct ANIMNEXT_API ISmoothBlend : IDecoratorInterface
+	{
+		DECLARE_ANIM_DECORATOR_INTERFACE(ISmoothBlend, 0x1c2c1739)
+
+		// Returns the desired blend time for the specified child
+		virtual float GetBlendTime(FExecutionContext& Context, const TDecoratorBinding<ISmoothBlend>& Binding, int32 ChildIndex) const;
+	};
+
+	/**
+	 * Specialization for decorator binding.
+	 */
+	template<>
+	struct TDecoratorBinding<ISmoothBlend> : FDecoratorBinding
+	{
+		// @see ISmoothBlend::GetBlendTime
+		float GetBlendTime(FExecutionContext& Context, int32 ChildIndex) const
+		{
+			return GetInterface()->GetBlendTime(Context, *this, ChildIndex);
+		}
+
+	protected:
+		const ISmoothBlend* GetInterface() const { return GetInterfaceTyped<ISmoothBlend>(); }
+	};
+}
