Implement blend smoother per bone decorator

diff --git a/Engine/Plugins/Experimental/AnimNext/Source/AnimNext/Private/Decorators/BlendSmootherPerBone.cpp b/Engine/Plugins/Experimental/AnimNext/Source/AnimNext/Private/Decorators/BlendSmootherPerBone.cpp
new file mode 100644
--- /dev/null
+++ b/Engine/Plugins/Experimental/AnimNext/Source/AnimNext/Private/Decorators/BlendSmootherPerBone.cpp
@@ -0,0 +1,165 @@
+// Copyright Epic Games, Inc. All Rights Reserved.
+
+#include "Decorators/BlendSmootherPerBone.h"
+
+#include "AlphaBlend.h"
+#include "Animation/BlendProfile.h"
+#include "DecoratorBase/ExecutionContext.h"
+#include "DecoratorInterfaces/IHierarchy.h"
+#include "EvaluationVM/Tasks/BlendKeyframesPerBone.h"
+#include "EvaluationVM/Tasks/NormalizeRotations.h"
+
+namespace UE::AnimNext
+{
+	AUTO_REGISTER_ANIM_DECORATOR(FBlendSmootherPerBoneDecorator)
+
+	DEFINE_ANIM_DECORATOR_BEGIN(FBlendSmootherPerBoneDecorator)
+		DEFINE_ANIM_DECORATOR_IMPLEMENTS_INTERFACE(IEvaluate)
+		DEFINE_ANIM_DECORATOR_IMPLEMENTS_INTERFACE(IUpdate)
+		DEFINE_ANIM_DECORATOR_IMPLEMENTS_INTERFACE(IDiscreteBlend)
+	DEFINE_ANIM_DECORATOR_END(FBlendSmootherPerBoneDecorator)
+
+	void FBlendSmootherPerBoneDecorator::PostEvaluate(FExecutionContext& Context, const TDecoratorBinding<IEvaluate>& Binding) const
+	{
+		const FSharedData* SharedData = Binding.GetSharedData<FSharedData>();
+		const FInstanceData* InstanceData = Binding.GetInstanceData<FInstanceData>();
+
+		if (!SharedData->BlendProfile)
+		{
+			// No blend profile set, default smooth blend behavior
+			IEvaluate::PostEvaluate(Context, Binding);
+			return;
+		}
+
+		// We override the default behavior since we need to blend per bone
+
+		const int32 NumChildren = InstanceData->PerChildBlendData.Num();
+		if (NumChildren < 2)
+		{
+			return;	// If we don't have at least 2 children, there is nothing to do
+		}
+
+		FEvaluateTraversalContext& TraversalContext = Context.GetTraversalContext<FEvaluateTraversalContext>();
+
+		// Children are visited depth first, in the order returned
+		// As such, when we evaluate the task program, the keyframe of the last child will be
+		// on top of the keyframe stack
+		// We thus process children in reverse order
+
+		// The last child override the top keyframe and scales it
+		{
+			const int32 ChildIndex = NumChildren - 1;
+			const FBlendSampleData& PoseSampleData = InstanceData->PerBoneSampleData[ChildIndex];
+
+			TraversalContext.AppendTask(FAnimNextBlendOverwriteKeyframePerBoneWithScaleTask::Make(SharedData->BlendProfile, PoseSampleData, PoseSampleData.TotalWeight));
+		}
+
+		// Other children accumulate with scale
+		for (int32 ChildIndex = NumChildren - 2; ChildIndex >= 0; --ChildIndex)
+		{
+			const FBlendSampleData& PoseSampleDataA = InstanceData->PerBoneSampleData[ChildIndex];
+			const FBlendSampleData& PoseSampleDataB = InstanceData->PerBoneSampleData[ChildIndex + 1];	// Above on the keyframe stack
+
+			TraversalContext.AppendTask(FAnimNextBlendAddKeyframePerBoneWithScaleTask::Make(SharedData->BlendProfile, PoseSampleDataA, PoseSampleDataB, PoseSampleDataA.TotalWeight));
+		}
+
+		// Once we are done, we normalize rotations
+		TraversalContext.AppendTask(FAnimNextNormalizeKeyframeRotationsTask());
+	}
+
+	void FBlendSmootherPerBoneDecorator::PreUpdate(FExecutionContext& Context, const TDecoratorBinding<IUpdate>& Binding) const
+	{
+		const FSharedData* SharedData = Binding.GetSharedData<FSharedData>();
+		FInstanceData* InstanceData = Binding.GetInstanceData<FInstanceData>();
+
+		// If this is our first update, allocate our blend data
+		if (InstanceData->PerChildBlendData.IsEmpty())
+		{
+			InitializeInstanceData(Context, Binding, SharedData, InstanceData);
+		}
+
+		// Update the decorators below us, they might trigger a transition
+		IUpdate::PreUpdate(Context, Binding);
+
+		if (!SharedData->BlendProfile)
+		{
+			return;	// No blend profile set, nothing to do
+		}
+
+		TDecoratorBinding<IDiscreteBlend> DiscreteBlendDecorator;
+		Context.GetInterface(Binding, DiscreteBlendDecorator);
+
+		const int32 DestinationChildIndex = DiscreteBlendDecorator.GetBlendDestinationChildIndex(Context);
+
+		// If we're using a blend profile, extract the scales and build blend sample data
+		const int32 NumChildren = InstanceData->PerChildBlendData.Num();
+		for (int32 ChildIndex = 0; ChildIndex < NumChildren; ++ChildIndex)
+		{
+			const float BlendWeight = DiscreteBlendDecorator.GetBlendWeight(Context, ChildIndex);
+			const FAlphaBlend* BlendState = DiscreteBlendDecorator.GetBlendState(Context, ChildIndex);
+
+			FBlendSampleData& PoseSampleData = InstanceData->PerBoneSampleData[ChildIndex];
+			PoseSampleData.TotalWeight = BlendWeight;
+
+			const FBlendData& BlendData = InstanceData->PerChildBlendData[ChildIndex];
+			const bool bInverse = SharedData->BlendProfile->Mode == EBlendProfileMode::WeightFactor ? (DestinationChildIndex != ChildIndex) : false;
+			SharedData->BlendProfile->UpdateBoneWeights(PoseSampleData, *BlendState, BlendData.StartAlpha, BlendWeight, bInverse);
+		}
+
+		FBlendSampleData::NormalizeDataWeight(InstanceData->PerBoneSampleData);
+	}
+
+	void FBlendSmootherPerBoneDecorator::OnBlendTransition(FExecutionContext& Context, const TDecoratorBinding<IDiscreteBlend>& Binding, int32 OldChildIndex, int32 NewChildIndex) const
+	{
+		const FSharedData* SharedData = Binding.GetSharedData<FSharedData>();
+		FInstanceData* InstanceData = Binding.GetInstanceData<FInstanceData>();
+
+		// Trigger the new transition
+		IDiscreteBlend::OnBlendTransition(Context, Binding, OldChildIndex, NewChildIndex);
+
+		if (!SharedData->BlendProfile)
+		{
+			return;	// No blend profile set, nothing to do
+		}
+
+		TDecoratorBinding<IDiscreteBlend> DiscreteBlendDecorator;
+		Context.GetInterface(Binding, DiscreteBlendDecorator);
+
+		const int32 NumChildren = InstanceData->PerChildBlendData.Num();
+		for (int32 ChildIndex = 0; ChildIndex < NumChildren; ++ChildIndex)
+		{
+			FBlendData& ChildBlendData = InstanceData->PerChildBlendData[ChildIndex];
+
+			const FAlphaBlend* BlendState = DiscreteBlendDecorator.GetBlendState(Context, ChildIndex);
+			ChildBlendData.StartAlpha = BlendState->GetAlpha();
+		}
+	}
+
+	void FBlendSmootherPerBoneDecorator::InitializeInstanceData(FExecutionContext& Context, const FDecoratorBinding& Binding, const FSharedData* SharedData, FInstanceData* InstanceData)
+	{
+		check(InstanceData->PerChildBlendData.IsEmpty());
+
+		if (!SharedData->BlendProfile)
+		{
+			return;	// No blend profile set, nothing to do
+		}
+
+		TDecoratorBinding<IHierarchy> HierarchyDecorator;
+		Context.GetInterface(Binding, HierarchyDecorator);
+
+		const uint32 NumChildren = HierarchyDecorator.GetNumChildren(Context);
+
+		InstanceData->PerChildBlendData.SetNum(NumChildren);
+
+		// Initialise per-bone data
+		InstanceData->PerBoneSampleData.SetNum(NumChildren);
+
+		const uint32 NumBlendEntries = SharedData->BlendProfile->GetNumBlendEntries();
+		for (uint32 Idx = 0; Idx < NumChildren; ++Idx)
+		{
+			FBlendSampleData& SampleData = InstanceData->PerBoneSampleData[Idx];
+			SampleData.SampleDataIndex = Idx;
+			SampleData.PerBoneBlendData.AddZeroed(NumBlendEntries);
+		}
+	}
+}
diff --git a/Engine/Plugins/Experimental/AnimNext/Source/AnimNext/Private/Decorators/BlendSmootherPerBone.h b/Engine/Plugins/Experimental/AnimNext/Source/AnimNext/Private/Decorators/BlendSmootherPerBone.h
new file mode 100644
--- /dev/null
+++ b/Engine/Plugins/Experimental/AnimNext/Source/AnimNext/Private/Decorators/BlendSmootherPerBone.h
@@ -0,0 +1,70 @@
+// Copyright Epic Games, Inc. All Rights Reserved.
+
+#pragma once
+
+#include "CoreMinimal.h"
+
+#include "Animation/AnimationAsset.h"
+#include "Curves/CurveFloat.h"
+#include "DecoratorBase/Decorator.h"
+#include "DecoratorInterfaces/IDiscreteBlend.h"
+#include "DecoratorInterfaces/IEvaluate.h"
+#include "DecoratorInterfaces/IUpdate.h"
+
+#include "BlendSmootherPerBone.generated.h"
+
+class UBlendProfile;
+
+USTRUCT()
+struct FAnimNextBlendSmootherPerBoneDecoratorSharedData : public FAnimNextDecoratorSharedData
+{
+	GENERATED_BODY()
+
+	/** Blend profile that configures how fast to blend each bone. */
+	// TODO: Can't show list of blend profiles, we need to find a skeleton to perform the lookup with
+	UPROPERTY(meta = (Inline, UseAsBlendProfile=true))
+	TObjectPtr<UBlendProfile> BlendProfile = nullptr;
+};
+
+namespace UE::AnimNext
+{
+	/**
+	 * FBlendSmootherDecorator
+	 * 
+	 * A decorator that smoothly blends between discrete states over time.
+	 */
+	struct FBlendSmootherPerBoneDecorator : FAdditiveDecorator, IEvaluate, IUpdate, IDiscreteBlend
+	{
+		DECLARE_ANIM_DECORATOR(FBlendSmootherPerBoneDecorator, 0x3d4e213f, FAdditiveDecorator)
+
+		using FSharedData = FAnimNextBlendSmootherPerBoneDecoratorSharedData;
+
+		// Struct for tracking blends for each pose
+		struct FBlendData
+		{
+			// Which blend alpha we started the blend with
+			float StartAlpha = 0.0f;
+		};
+
+		struct FInstanceData : FDecorator::FInstanceData
+		{
+			// Blend state per child
+			TArray<FBlendData> PerChildBlendData;
+
+			// Per-bone blending data for each child
+			TArray<FBlendSampleData> PerBoneSampleData;
+		};
+
+		// IEvaluate impl
+		virtual void PostEvaluate(FExecutionContext& Context, const TDecoratorBinding<IEvaluate>& Binding) const override;
+
+		// IUpdate impl
+		virtual void PreUpdate(FExecutionContext& Context, const TDecoratorBinding<IUpdate>& Binding) const override;
+
+		// IDiscreteBlend impl
+		virtual void OnBlendTransition(FExecutionContext& Context, const TDecoratorBinding<IDiscreteBlend>& Binding, int32 OldChildIndex, int32 NewChildIndex) const override;
+
+		// Internal impl
+		static void InitializeInstanceData(FExecutionContext& Context, const FDecoratorBinding& Binding, const FSharedData* SharedData, FInstanceData* InstanceData);
+	};
+}
diff --git a/Engine/Plugins/Experimental/AnimNext/Source/AnimNext/Private/EvaluationVM/Tasks/BlendKeyframesPerBone.cpp b/Engine/Plugins/Experimental/AnimNext/Source/AnimNext/Private/EvaluationVM/Tasks/BlendKeyframesPerBone.cpp
new file mode 100644
--- /dev/null
+++ b/Engine/Plugins/Experimental/AnimNext/Source/AnimNext/Private/EvaluationVM/Tasks/BlendKeyframesPerBone.cpp
@@ -0,0 +1,196 @@
+// Copyright Epic Games, Inc. All Rights Reserved.
+
+#include "EvaluationVM/Tasks/BlendKeyframesPerBone.h"
+
+#include "AnimationRuntime.h"
+#include "Animation/BlendProfile.h"
+#include "Animation/SkeletonRemapping.h"
+#include "Animation/SkeletonRemappingRegistry.h"
+#include "EvaluationVM/EvaluationVM.h"
+#include "EvaluationVM/KeyframeState.h"
+#include "TransformArrayOperations.h"
+
+FAnimNextBlendOverwriteKeyframePerBoneWithScaleTask FAnimNextBlendOverwriteKeyframePerBoneWithScaleTask::Make(TWeakObjectPtr<const UBlendProfile> BlendProfile, const FBlendSampleData& BlendData, float ScaleFactor)
+{
+	FAnimNextBlendOverwriteKeyframePerBoneWithScaleTask Task;
+	Task.BlendProfile = BlendProfile;
+	Task.BlendData = &BlendData;
+	Task.ScaleFactor = ScaleFactor;
+
+	return Task;
+}
+
+void FAnimNextBlendOverwriteKeyframePerBoneWithScaleTask::Execute(UE::AnimNext::FEvaluationVM& VM) const
+{
+	using namespace UE::AnimNext;
+
+	const UBlendProfile* BlendProfilePtr = BlendProfile.Get();
+	if (BlendProfilePtr == nullptr)
+	{
+		// If we don't have a blend profile, we blend the whole pose
+		Super::Execute(VM);
+		return;
+	}
+
+	TUniquePtr<FKeyframeState> Keyframe;
+	if (!VM.PopValue(KEYFRAME_STACK_NAME, Keyframe))
+	{
+		// We have no inputs, nothing to do
+		return;
+	}
+
+	const USkeleton* TargetSkeleton = Keyframe->Pose.GetSkeletonAsset();
+	TSharedPtr<IInterpolationIndexProvider::FPerBoneInterpolationData> Data = BlendProfilePtr->GetPerBoneInterpolationData(TargetSkeleton);
+
+	const USkeleton* SourceSkeleton = BlendProfilePtr->OwningSkeleton;
+	const FSkeletonRemapping& SkeletonRemapping = UE::Anim::FSkeletonRemappingRegistry::Get().GetRemapping(SourceSkeleton, TargetSkeleton);
+
+	const TArrayView<const FBoneIndexType> LODBoneIndexToSkeletonBoneIndexMap = Keyframe->Pose.GetLODBoneIndexToSkeletonBoneIndexMap();
+	const int32 NumLODBones = LODBoneIndexToSkeletonBoneIndexMap.Num();
+
+	TArray<int32> LODBoneIndexToWeightIndexMap;
+	LODBoneIndexToWeightIndexMap.AddUninitialized(NumLODBones);
+
+	if (SkeletonRemapping.IsValid())
+	{
+		for (int32 LODBoneIndex = 0; LODBoneIndex < NumLODBones; ++LODBoneIndex)
+		{
+			const FSkeletonPoseBoneIndex TargetSkeletonBoneIndex(LODBoneIndexToSkeletonBoneIndexMap[LODBoneIndex]);
+			const FSkeletonPoseBoneIndex SourceSkeletonBoneIndex(SkeletonRemapping.GetSourceSkeletonBoneIndex(TargetSkeletonBoneIndex.GetInt()));
+			LODBoneIndexToWeightIndexMap[LODBoneIndex] = SourceSkeletonBoneIndex.IsValid() ? BlendProfilePtr->GetPerBoneInterpolationIndex(SourceSkeletonBoneIndex, TargetSkeleton, Data.Get()) : INDEX_NONE;
+		}
+	}
+	else
+	{
+		for (int32 LODBoneIndex = 0; LODBoneIndex < NumLODBones; ++LODBoneIndex)
+		{
+			const FSkeletonPoseBoneIndex SkeletonBoneIndex(LODBoneIndexToSkeletonBoneIndexMap[LODBoneIndex]);
+			LODBoneIndexToWeightIndexMap[LODBoneIndex] = BlendProfilePtr->GetPerBoneInterpolationIndex(SkeletonBoneIndex, TargetSkeleton, Data.Get());
+		}
+	}
+
+	const float BlendWeight = BlendData->GetClampedWeight();
+
+	if (EnumHasAnyFlags(VM.GetFlags(), EEvaluationFlags::Bones))
+	{
+		BlendOverwritePerBoneWithScale(
+			Keyframe->Pose.LocalTransforms.GetView(), Keyframe->Pose.LocalTransforms.GetConstView(),
+			LODBoneIndexToWeightIndexMap, BlendData->PerBoneBlendData, BlendWeight);
+	}
+
+	if (EnumHasAnyFlags(VM.GetFlags(), EEvaluationFlags::Curves))
+	{
+		// Curves cannot override in place
+		FBlendedCurve Result;
+		Result.Override(Keyframe->Curves, BlendWeight);
+
+		Keyframe->Curves = MoveTemp(Result);
+	}
+
+	if (EnumHasAnyFlags(VM.GetFlags(), EEvaluationFlags::Attributes))
+	{
+		UE::Anim::Attributes::BlendAttributesPerBone({ Keyframe->Attributes }, { LODBoneIndexToWeightIndexMap }, { *BlendData }, { 0 }, { Keyframe->Attributes });
+	}
+
+	VM.PushValue(KEYFRAME_STACK_NAME, MoveTemp(Keyframe));
+}
+
+FAnimNextBlendAddKeyframePerBoneWithScaleTask FAnimNextBlendAddKeyframePerBoneWithScaleTask::Make(TWeakObjectPtr<const UBlendProfile> BlendProfile, const FBlendSampleData& BlendDataA, const FBlendSampleData& BlendDataB, float ScaleFactor)
+{
+	FAnimNextBlendAddKeyframePerBoneWithScaleTask Task;
+	Task.BlendProfile = BlendProfile;
+	Task.BlendDataA = &BlendDataA;
+	Task.BlendDataB = &BlendDataB;
+	Task.ScaleFactor = ScaleFactor;
+
+	return Task;
+}
+
+void FAnimNextBlendAddKeyframePerBoneWithScaleTask::Execute(UE::AnimNext::FEvaluationVM& VM) const
+{
+	using namespace UE::AnimNext;
+
+	const UBlendProfile* BlendProfilePtr = BlendProfile.Get();
+	if (BlendProfilePtr == nullptr)
+	{
+		// If we don't have a blend profile, we blend the whole pose
+		Super::Execute(VM);
+		return;
+	}
+
+	// Pop our top two poses, we'll re-use the top keyframe for our result
+
+	TUniquePtr<FKeyframeState> KeyframeB;
+	if (!VM.PopValue(KEYFRAME_STACK_NAME, KeyframeB))
+	{
+		// We have no inputs, nothing to do
+		return;
+	}
+
+	TUniquePtr<FKeyframeState> KeyframeA;
+	if (!VM.PopValue(KEYFRAME_STACK_NAME, KeyframeA))
+	{
+		// We have a single input, leave it on top of the stack
+		VM.PushValue(KEYFRAME_STACK_NAME, MoveTemp(KeyframeB));
+		return;
+	}
+
+	const USkeleton* TargetSkeleton = KeyframeB->Pose.GetSkeletonAsset();
+	TSharedPtr<IInterpolationIndexProvider::FPerBoneInterpolationData> Data = BlendProfilePtr->GetPerBoneInterpolationData(TargetSkeleton);
+
+	const USkeleton* SourceSkeleton = BlendProfilePtr->OwningSkeleton;
+	const FSkeletonRemapping& SkeletonRemapping = UE::Anim::FSkeletonRemappingRegistry::Get().GetRemapping(SourceSkeleton, TargetSkeleton);
+
+	const TArrayView<const FBoneIndexType> LODBoneIndexToSkeletonBoneIndexMap = KeyframeB->Pose.GetLODBoneIndexToSkeletonBoneIndexMap();
+	const int32 NumLODBones = LODBoneIndexToSkeletonBoneIndexMap.Num();
+
+	TArray<int32> LODBoneIndexToWeightIndexMap;
+	LODBoneIndexToWeightIndexMap.AddUninitialized(NumLODBones);
+
+	if (SkeletonRemapping.IsValid())
+	{
+		for (int32 LODBoneIndex = 0; LODBoneIndex < NumLODBones; ++LODBoneIndex)
+		{
+			const FSkeletonPoseBoneIndex TargetSkeletonBoneIndex(LODBoneIndexToSkeletonBoneIndexMap[LODBoneIndex]);
+			const FSkeletonPoseBoneIndex SourceSkeletonBoneIndex(SkeletonRemapping.GetSourceSkeletonBoneIndex(TargetSkeletonBoneIndex.GetInt()));
+			LODBoneIndexToWeightIndexMap[LODBoneIndex] = SourceSkeletonBoneIndex.IsValid() ? BlendProfilePtr->GetPerBoneInterpolationIndex(SourceSkeletonBoneIndex, TargetSkeleton, Data.Get()) : INDEX_NONE;
+		}
+	}
+	else
+	{
+		for (int32 LODBoneIndex = 0; LODBoneIndex < NumLODBones; ++LODBoneIndex)
+		{
+			const FSkeletonPoseBoneIndex SkeletonBoneIndex(LODBoneIndexToSkeletonBoneIndexMap[LODBoneIndex]);
+			LODBoneIndexToWeightIndexMap[LODBoneIndex] = BlendProfilePtr->GetPerBoneInterpolationIndex(SkeletonBoneIndex, TargetSkeleton, Data.Get());
+		}
+	}
+
+	const float BlendWeightA = BlendDataA->GetClampedWeight();
+
+	if (EnumHasAnyFlags(VM.GetFlags(), EEvaluationFlags::Bones))
+	{
+		check(KeyframeA->Pose.GetNumBones() == KeyframeB->Pose.GetNumBones());
+
+		BlendAddPerBoneWithScale(
+			KeyframeB->Pose.LocalTransforms.GetView(), KeyframeA->Pose.LocalTransforms.GetConstView(),
+			LODBoneIndexToWeightIndexMap, BlendDataA->PerBoneBlendData, BlendWeightA);
+	}
+
+	if (EnumHasAnyFlags(VM.GetFlags(), EEvaluationFlags::Curves))
+	{
+		KeyframeB->Curves.Accumulate(KeyframeA->Curves, BlendWeightA);
+	}
+
+	if (EnumHasAnyFlags(VM.GetFlags(), EEvaluationFlags::Attributes))
+	{
+		// TODO: Might need to be revisited, might not work as intended
+		UE::Anim::Attributes::BlendAttributesPerBone(
+			{ KeyframeA->Attributes, KeyframeB->Attributes },
+			{ LODBoneIndexToWeightIndexMap },
+			{ *BlendDataA, *BlendDataB },
+			{ 0, 1 },
+			{ KeyframeA->Attributes });
+	}
+
+	VM.PushValue(KEYFRAME_STACK_NAME, MoveTemp(KeyframeB));
+}
diff --git a/Engine/Plugins/Experimental/AnimNext/Source/AnimNext/Private/TransformArrayOperations_AoS.cpp b/Engine/Plugins/Experimental/AnimNext/Source/AnimNext/Private/TransformArrayOperations_AoS.cpp
--- a/Engine/Plugins/Experimental/AnimNext/Source/AnimNext/Private/TransformArrayOperations_AoS.cpp
+++ b/Engine/Plugins/Experimental/AnimNext/Source/AnimNext/Private/TransformArrayOperations_AoS.cpp
@@ -79,4 +79,38 @@ namespace UE::AnimNext
 			Dest[TransformIndex].AccumulateWithShortestRotation(Source[TransformIndex], VScaleWeight);
 		}
 	}
+
+	void BlendOverwritePerBoneWithScale(
+		const FTransformArrayAoSView& Dest, const FTransformArrayAoSConstView& Source,
+		const TArrayView<const int32>& LODBoneIndexToWeightIndexMap, const TArrayView<const float>& BoneWeights, const float DefaultScaleWeight)
+	{
+		const ScalarRegister VDefaultScaleWeight(DefaultScaleWeight);
+		const int32 NumTransforms = Source.Num();
+
+		check(Dest.Num() >= NumTransforms);
+
+		for (int32 LODBoneIndex = 0; LODBoneIndex < NumTransforms; ++LODBoneIndex)
+		{
+			const int32 PerBoneIndex = LODBoneIndexToWeightIndexMap[LODBoneIndex];
+			const ScalarRegister VScaleWeight = BoneWeights.IsValidIndex(PerBoneIndex) ? ScalarRegister(BoneWeights[PerBoneIndex]) : VDefaultScaleWeight;
+			Dest[LODBoneIndex] = Source[LODBoneIndex] * VScaleWeight;
+		}
+	}
+
+	void BlendAddPerBoneWithScale(
+		const FTransformArrayAoSView& Dest, const FTransformArrayAoSConstView& Source,
+		const TArrayView<const int32>& LODBoneIndexToWeightIndexMap, const TArrayView<const float>& BoneWeights, const float DefaultScaleWeight)
+	{
+		const ScalarRegister VDefaultScaleWeight(DefaultScaleWeight);
+		const int32 NumTransforms = Source.Num();
+
+		check(Dest.Num() >= NumTransforms);
+
+		for (int32 LODBoneIndex = 0; LODBoneIndex < NumTransforms; ++LODBoneIndex)
+		{
+			const int32 PerBoneIndex = LODBoneIndexToWeightIndexMap[LODBoneIndex];
+			const ScalarRegister VScaleWeight = BoneWeights.IsValidIndex(PerBoneIndex) ? ScalarRegister(BoneWeights[PerBoneIndex]) : VDefaultScaleWeight;
+			Dest[LODBoneIndex].AccumulateWithShortestRotation(Source[LODBoneIndex], VScaleWeight);
+		}
+	}
 }
diff --git a/Engine/Plugins/Experimental/AnimNext/Source/AnimNext/Private/TransformArrayOperations_SoA.cpp b/Engine/Plugins/Experimental/AnimNext/Source/AnimNext/Private/TransformArrayOperations_SoA.cpp
--- a/Engine/Plugins/Experimental/AnimNext/Source/AnimNext/Private/TransformArrayOperations_SoA.cpp
+++ b/Engine/Plugins/Experimental/AnimNext/Source/AnimNext/Private/TransformArrayOperations_SoA.cpp
@@ -177,4 +177,62 @@ namespace UE::AnimNext
 			Dest.Scales3D[TransformIndex] += Source.Scales3D[TransformIndex] * ScaleWeight;
 		}
 	}
+
+	void BlendOverwritePerBoneWithScale(
+		const FTransformArraySoAView& Dest, const FTransformArraySoAConstView& Source,
+		const TArrayView<const int32>& LODBoneIndexToWeightIndexMap, const TArrayView<const float>& BoneWeights, const float DefaultScaleWeight)
+	{
+		const int32 NumTransforms = Source.Num();
+
+		check(Dest.Num() >= NumTransforms);
+
+		for (int32 LODBoneIndex = 0; LODBoneIndex < NumTransforms; ++LODBoneIndex)
+		{
+			const int32 PerBoneIndex = LODBoneIndexToWeightIndexMap[LODBoneIndex];
+			const float ScaleWeight = BoneWeights.IsValidIndex(PerBoneIndex) ? BoneWeights[PerBoneIndex] : DefaultScaleWeight;
+
+			Dest.Translations[LODBoneIndex] = Source.Translations[LODBoneIndex] * ScaleWeight;
+			Dest.Rotations[LODBoneIndex] = Source.Rotations[LODBoneIndex] * ScaleWeight;
+			Dest.Scales3D[LODBoneIndex] = Source.Scales3D[LODBoneIndex] * ScaleWeight;
+		}
+	}
+
+	void BlendAddPerBoneWithScale(
+		const FTransformArraySoAView& Dest, const FTransformArraySoAConstView& Source,
+		const TArrayView<const int32>& LODBoneIndexToWeightIndexMap, const TArrayView<const float>& BoneWeights, const float DefaultScaleWeight)
+	{
+		const int32 NumTransforms = Source.Num();
+
+		check(Dest.Num() >= NumTransforms);
+
+		using QuatVectorRegister = FQuat::QuatVectorRegister;
+		const QuatVectorRegister Zero = VectorZero();
+
+		for (int32 LODBoneIndex = 0; LODBoneIndex < NumTransforms; ++LODBoneIndex)
+		{
+			const int32 PerBoneIndex = LODBoneIndexToWeightIndexMap[LODBoneIndex];
+			const float ScaleWeight = BoneWeights.IsValidIndex(PerBoneIndex) ? BoneWeights[PerBoneIndex] : DefaultScaleWeight;
+			const ScalarRegister VScaleWeight(ScaleWeight);
+
+			const QuatVectorRegister SourceRotation = VectorLoadAligned(&Source.Rotations[LODBoneIndex]);
+			const QuatVectorRegister DestRotation = VectorLoadAligned(&Dest.Rotations[LODBoneIndex]);
+
+			const QuatVectorRegister BlendedRotation = VectorMultiply(SourceRotation, VScaleWeight);
+
+			// Blend rotation
+			//     To ensure the 'shortest route', we make sure the dot product between the both rotations is positive.
+			//     const float Bias = (|A.B| >= 0 ? 1 : -1)
+			//     return A + B * Bias;
+
+			const QuatVectorRegister RotationDot = VectorDot4(DestRotation, BlendedRotation);
+			const QuatVectorRegister QuatRotationDirMask = VectorCompareGE(RotationDot, Zero);
+			const QuatVectorRegister NegativeB = VectorSubtract(Zero, BlendedRotation);
+			const QuatVectorRegister BiasTimesB = VectorSelect(QuatRotationDirMask, BlendedRotation, NegativeB);
+
+			Dest.Rotations[LODBoneIndex] = FQuat::MakeFromVectorRegister(VectorAdd(DestRotation, BiasTimesB));
+
+			Dest.Translations[LODBoneIndex] += Source.Translations[LODBoneIndex] * ScaleWeight;
+			Dest.Scales3D[LODBoneIndex] += Source.Scales3D[LODBoneIndex] * ScaleWeight;
+		}
+	}
 }
diff --git a/Engine/Plugins/Experimental/AnimNext/Source/AnimNext/Public/EvaluationVM/Tasks/BlendKeyframesPerBone.h b/Engine/Plugins/Experimental/AnimNext/Source/AnimNext/Public/EvaluationVM/Tasks/BlendKeyframesPerBone.h
new file mode 100644
--- /dev/null
+++ b/Engine/Plugins/Experimental/AnimNext/Source/AnimNext/Public/EvaluationVM/Tasks/BlendKeyframesPerBone.h
@@ -0,0 +1,84 @@
+// Copyright Epic Games, Inc. All Rights Reserved.
+
+#pragma once
+
+#include "CoreMinimal.h"
+#include "EvaluationVM/EvaluationTask.h"
+#include "EvaluationVM/Tasks/BlendKeyframes.h"
+
+#include "BlendKeyframesPerBone.generated.h"
+
+struct FBlendSampleData;
+class UBlendProfile;
+
+/*
+ * Blend Overwrite Keyframe Per Bone With Scale Task
+ *
+ * This pops the top keyframe from the VM keyframe stack, it scales each bone by a factor, and pushes
+ * back the result onto the stack.
+ * Top = Top * ScaleFactor
+ * 
+ * If no blend profile is provided, this task behaves like FAnimNextBlendOverwriteKeyframeWithScaleTask
+ * Note that rotations will not be normalized after this task.
+ */
+USTRUCT()
+struct ANIMNEXT_API FAnimNextBlendOverwriteKeyframePerBoneWithScaleTask : public FAnimNextBlendOverwriteKeyframeWithScaleTask
+{
+	GENERATED_BODY()
+
+	DECLARE_ANIM_EVALUATION_TASK(FAnimNextBlendOverwriteKeyframePerBoneWithScaleTask)
+
+	static FAnimNextBlendOverwriteKeyframePerBoneWithScaleTask Make(TWeakObjectPtr<const UBlendProfile> BlendProfile, const FBlendSampleData& BlendData, float ScaleFactor);
+
+	// Task entry point
+	virtual void Execute(UE::AnimNext::FEvaluationVM& VM) const override;
+
+	// The blend data associated with the keyframe to overwrite
+	// Cannot use a UPROPERTY with a pointer to a struct
+	//UPROPERTY()
+	const FBlendSampleData* BlendData = nullptr;
+
+	// The blend profile to use
+	// An optional blend profile provides a blend weight per bone
+	UPROPERTY()
+	TWeakObjectPtr<const UBlendProfile> BlendProfile;
+};
+
+/*
+ * Blend Add Keyframe Per Bone With Scale Task
+ *
+ * This pops the top two keyframes (A and B) from the VM keyframe stack (let B be at the top).
+ * B is our intermediary result that we add on top of; while A is the keyframe we scale.
+ * The result is pushed back onto the stack.
+ * Top = Top + (Top-1 * ScaleFactor)
+ * 
+ * If no blend profile is provided, this task behaves like FAnimNextBlendAddKeyframeWithScaleTask
+ * Note that rotations will not be normalized after this task.
+ */
+USTRUCT()
+struct ANIMNEXT_API FAnimNextBlendAddKeyframePerBoneWithScaleTask : public FAnimNextBlendAddKeyframeWithScaleTask
+{
+	GENERATED_BODY()
+
+	DECLARE_ANIM_EVALUATION_TASK(FAnimNextBlendAddKeyframePerBoneWithScaleTask)
+
+	static FAnimNextBlendAddKeyframePerBoneWithScaleTask Make(TWeakObjectPtr<const UBlendProfile> BlendProfile, const FBlendSampleData& BlendDataA, const FBlendSampleData& BlendDataB, float ScaleFactor);
+
+	// Task entry point
+	virtual void Execute(UE::AnimNext::FEvaluationVM& VM) const override;
+
+	// The blend data associated with the keyframe A
+	// Cannot use a UPROPERTY with a pointer to a struct
+	//UPROPERTY()
+	const FBlendSampleData* BlendDataA = nullptr;
+
+	// The blend data associated with the keyframe B
+	// Cannot use a UPROPERTY with a pointer to a struct
+	//UPROPERTY()
+	const FBlendSampleData* BlendDataB = nullptr;
+
+	// The blend profile to use
+	// An optional blend profile provides a blend weight per bone
+	UPROPERTY()
+	TWeakObjectPtr<const UBlendProfile> BlendProfile;
+};
diff --git a/Engine/Plugins/Experimental/AnimNext/Source/AnimNext/Public/TransformArrayOperations.h b/Engine/Plugins/Experimental/AnimNext/Source/AnimNext/Public/TransformArrayOperations.h
--- a/Engine/Plugins/Experimental/AnimNext/Source/AnimNext/Public/TransformArrayOperations.h
+++ b/Engine/Plugins/Experimental/AnimNext/Source/AnimNext/Public/TransformArrayOperations.h
@@ -36,4 +36,22 @@ namespace UE::AnimNext
 	// Dest = Dest + (Source * ScaleWeight)
 	ANIMNEXT_API void BlendAddWithScale(const FTransformArrayAoSView& Dest, const FTransformArrayAoSConstView& Source, const float ScaleWeight);
 	ANIMNEXT_API void BlendAddWithScale(const FTransformArraySoAView& Dest, const FTransformArraySoAConstView& Source, const float ScaleWeight);
+
+	// The source transforms are scaled by the provided per bone weight and the result is written in the destination
+	// Dest = Source * (WeightIndex != INDEX_NONE ? Weights[WeightIndex] : DefaultScaleWeight)
+	ANIMNEXT_API void BlendOverwritePerBoneWithScale(
+		const FTransformArrayAoSView& Dest, const FTransformArrayAoSConstView& Source,
+		const TArrayView<const int32>& LODBoneIndexToWeightIndexMap, const TArrayView<const float>& BoneWeights, const float DefaultScaleWeight);
+	ANIMNEXT_API void BlendOverwritePerBoneWithScale(
+		const FTransformArraySoAView& Dest, const FTransformArraySoAConstView& Source,
+		const TArrayView<const int32>& LODBoneIndexToWeightIndexMap, const TArrayView<const float>& BoneWeights, const float DefaultScaleWeight);
+
+	// The source transforms are scaled by the provided per bone weight and the result is added to the destination
+	// Dest = Dest + (Source * (WeightIndex != INDEX_NONE ? Weights[WeightIndex] : DefaultScaleWeight))
+	ANIMNEXT_API void BlendAddPerBoneWithScale(
+		const FTransformArrayAoSView& Dest, const FTransformArrayAoSConstView& Source,
+		const TArrayView<const int32>& LODBoneIndexToWeightIndexMap, const TArrayView<const float>& BoneWeights, const float DefaultScaleWeight);
+	ANIMNEXT_API void BlendAddPerBoneWithScale(
+		const FTransformArraySoAView& Dest, const FTransformArraySoAConstView& Source,
+		const TArrayView<const int32>& LODBoneIndexToWeightIndexMap, const TArrayView<const float>& BoneWeights, const float DefaultScaleWeight);
 }
diff --git a/Engine/Source/Runtime/Engine/Classes/Animation/BlendProfile.h b/Engine/Source/Runtime/Engine/Classes/Animation/BlendProfile.h
--- a/Engine/Source/Runtime/Engine/Classes/Animation/BlendProfile.h
+++ b/Engine/Source/Runtime/Engine/Classes/Animation/BlendProfile.h
@@ -177,6 +177,7 @@ public:
 
 	// IInterpolationIndexProvider
 	ENGINE_API virtual int32 GetPerBoneInterpolationIndex(const FCompactPoseBoneIndex& InCompactPoseBoneIndex, const FBoneContainer& BoneContainer, const IInterpolationIndexProvider::FPerBoneInterpolationData* Data) const override;
+	ENGINE_API virtual int32 GetPerBoneInterpolationIndex(const FSkeletonPoseBoneIndex InSkeletonBoneIndex, const USkeleton* TargetSkeleton, const IInterpolationIndexProvider::FPerBoneInterpolationData* Data) const override;
 	// End IInterpolationIndexProvider
 
 	// UObject
diff --git a/Engine/Source/Runtime/Engine/Classes/Animation/BlendSpace.h b/Engine/Source/Runtime/Engine/Classes/Animation/BlendSpace.h
--- a/Engine/Source/Runtime/Engine/Classes/Animation/BlendSpace.h
+++ b/Engine/Source/Runtime/Engine/Classes/Animation/BlendSpace.h
@@ -495,6 +495,12 @@ public:
 	* If nothing found, return INDEX_NONE
 	*/
 	virtual int32 GetPerBoneInterpolationIndex(const FCompactPoseBoneIndex& InCompactPoseBoneIndex, const FBoneContainer& RequiredBones, const IInterpolationIndexProvider::FPerBoneInterpolationData* Data) const override;
+
+	/**
+	* Get PerBoneInterpolationIndex for the input BoneIndex
+	* If nothing found, return INDEX_NONE
+	*/
+	virtual int32 GetPerBoneInterpolationIndex(const FSkeletonPoseBoneIndex InSkeletonBoneIndex, const USkeleton* TargetSkeleton, const IInterpolationIndexProvider::FPerBoneInterpolationData* Data) const override;
 	// End IInterpolationIndexProvider Overrides
 
 	/** Returns whether or not the given additive animation type is compatible with the blendspace type */
diff --git a/Engine/Source/Runtime/Engine/Private/Animation/BlendProfile.cpp b/Engine/Source/Runtime/Engine/Private/Animation/BlendProfile.cpp
--- a/Engine/Source/Runtime/Engine/Private/Animation/BlendProfile.cpp
+++ b/Engine/Source/Runtime/Engine/Private/Animation/BlendProfile.cpp
@@ -207,6 +207,11 @@ int32 UBlendProfile::GetPerBoneInterpolationIndex(const FCompactPoseBoneIndex& I
 	return GetEntryIndex(BoneContainer.GetSkeletonPoseIndexFromCompactPoseIndex(InCompactPoseBoneIndex));
 }
 
+int32 UBlendProfile::GetPerBoneInterpolationIndex(const FSkeletonPoseBoneIndex InSkeletonBoneIndex, const USkeleton* TargetSkeleton, const IInterpolationIndexProvider::FPerBoneInterpolationData* Data) const
+{
+	return GetEntryIndex(InSkeletonBoneIndex);
+}
+
 void UBlendProfile::SetSingleBoneBlendScale(int32 InBoneIdx, float InScale, bool bCreate /*= false*/)
 {
 	FBlendProfileBoneEntry* Entry = ProfileEntries.FindByPredicate([InBoneIdx](const FBlendProfileBoneEntry& InEntry)
diff --git a/Engine/Source/Runtime/Engine/Private/Animation/BlendSpace.cpp b/Engine/Source/Runtime/Engine/Private/Animation/BlendSpace.cpp
--- a/Engine/Source/Runtime/Engine/Private/Animation/BlendSpace.cpp
+++ b/Engine/Source/Runtime/Engine/Private/Animation/BlendSpace.cpp
@@ -835,6 +835,55 @@ int32 UBlendSpace::GetPerBoneInterpolationIndex(
 	return INDEX_NONE;
 }
 
+namespace UE::Anim::Private
+{
+	static FBoneContainer DummyContainer;
+}
+
+int32 UBlendSpace::GetPerBoneInterpolationIndex(const FSkeletonPoseBoneIndex InSkeletonBoneIndex, const USkeleton* TargetSkeleton, const IInterpolationIndexProvider::FPerBoneInterpolationData* Data) const
+{
+	if (!ensure(Data) || !InSkeletonBoneIndex.IsValid())
+	{
+		return INDEX_NONE;
+	}
+
+	const TArray<FSortedPerBoneInterpolation>& SortedData = static_cast<const FSortedPerBoneInterpolationData*>(Data)->Data;
+
+	const USkeleton* SourceSkeleton = GetSkeleton();
+	const FReferenceSkeleton& TargetReferenceSkeleton = TargetSkeleton->GetReferenceSkeleton();
+	const FSkeletonRemapping& SkeletonRemapping = UE::Anim::FSkeletonRemappingRegistry().Get().GetRemapping(SourceSkeleton, TargetSkeleton);
+
+	for (const FSortedPerBoneInterpolation& SortedBoneData : SortedData)
+	{
+		const FPerBoneInterpolation& PerBoneInterpolation = SortedBoneData.PerBoneBlend;
+		const FBoneReference& SmoothedBone = PerBoneInterpolation.BoneReference;
+
+		// Bone references are stored as skeleton bones, we can use a dummy bone container as it isn't required
+		FSkeletonPoseBoneIndex SmoothedSkelBoneIndex = SmoothedBone.GetSkeletonPoseIndex(UE::Anim::Private::DummyContainer);
+
+		// Remap to the target skeleton, using skeleton remapping, as we might be applying this blend space onto another skeleton than the asset was created for.
+		if (SkeletonRemapping.IsValid())
+		{
+			const int32 RemappedSkelBoneIndex = SkeletonRemapping.GetTargetSkeletonBoneIndex(SmoothedSkelBoneIndex.GetInt());
+			SmoothedSkelBoneIndex = FSkeletonPoseBoneIndex(RemappedSkelBoneIndex);
+		}
+
+		if (SmoothedSkelBoneIndex == InSkeletonBoneIndex)
+		{
+			return SortedBoneData.OriginalIndex;
+		}
+
+		// BoneIsChildOf returns true if InSkeletonBoneIndex is a child of SmoothedSkelBoneIndex
+		if (SmoothedSkelBoneIndex.IsValid() &&
+			TargetReferenceSkeleton.BoneIsChildOf(InSkeletonBoneIndex.GetInt(), SmoothedSkelBoneIndex.GetInt()))
+		{
+			return SortedBoneData.OriginalIndex;
+		}
+	}
+
+	return INDEX_NONE;
+}
+
 bool UBlendSpace::IsValidAdditiveType(EAdditiveAnimationType AdditiveType) const
 {
 	return (AdditiveType == AAT_LocalSpaceBase || AdditiveType == AAT_RotationOffsetMeshSpace || AdditiveType == AAT_None);
diff --git a/Engine/Source/Runtime/Engine/Public/AnimationRuntime.h b/Engine/Source/Runtime/Engine/Public/AnimationRuntime.h
--- a/Engine/Source/Runtime/Engine/Public/AnimationRuntime.h
+++ b/Engine/Source/Runtime/Engine/Public/AnimationRuntime.h
@@ -91,6 +91,11 @@ public:
 	// up/blending BoneIndex. This call will be passed the results of GetPerBoneInterpolationData, so the two functions
 	// should be matched.
 	virtual int32 GetPerBoneInterpolationIndex(const FCompactPoseBoneIndex& InCompactPoseBoneIndex, const FBoneContainer& RequiredBones, const FPerBoneInterpolationData* Data) const = 0;
+
+	// Implementation should return the index into the PerBoneBlendData array that would be required when looking
+	// up/blending BoneIndex. This call will be passed the results of GetPerBoneInterpolationData, so the two functions
+	// should be matched.
+	virtual int32 GetPerBoneInterpolationIndex(const FSkeletonPoseBoneIndex InSkeletonBoneIndex, const USkeleton* TargetSkeleton, const IInterpolationIndexProvider::FPerBoneInterpolationData* Data) const = 0;
 };
 
 /** In AnimationRunTime Library, we extract animation data based on Skeleton hierarchy, not ref pose hierarchy. 
