Fix AnimNext graphs to be able to load in a cooked package

diff --git a/Engine/Plugins/Experimental/AnimNext/Source/AnimNext/Internal/DecoratorBase/DecoratorReader.h b/Engine/Plugins/Experimental/AnimNext/Source/AnimNext/Internal/DecoratorBase/DecoratorReader.h
--- a/Engine/Plugins/Experimental/AnimNext/Source/AnimNext/Internal/DecoratorBase/DecoratorReader.h
+++ b/Engine/Plugins/Experimental/AnimNext/Source/AnimNext/Internal/DecoratorBase/DecoratorReader.h
@@ -35,9 +35,9 @@ namespace UE::AnimNext
 			NodeInstanceDataTooLarge,	// Exceeded the maximum node instance data size, @see FNodeInstance::MAXIMUM_NODE_INSTANCE_DATA_SIZE
 		};
 
-		explicit FDecoratorReader(FArchive& Ar);
+		explicit FDecoratorReader(const TArray<TObjectPtr<UObject>>& InGraphReferencedObjects, FArchive& Ar);
 
-		[[nodiscard]] EErrorState ReadGraph(TArray<uint8>& GraphSharedData, TArray<TObjectPtr<UObject>>& TrackedObjectsForGC);
+		[[nodiscard]] EErrorState ReadGraph(TArray<uint8>& GraphSharedData);
 
 		// Takes a node handle representing a node index and resolves it into a node handle representing a shared data offset
 		// Must be called after ReadGraphSharedData as it populates the necessary data
@@ -59,8 +59,8 @@ namespace UE::AnimNext
 		// Call first to read the graph shared data
 		[[nodiscard]] EErrorState ReadGraphSharedData(TArray<uint8>& GraphSharedData);
 
-		// Call second to read the tracked UObjects for GC
-		[[nodiscard]] EErrorState ReadTrackedObjectsForGC(TArray<TObjectPtr<UObject>>& TrackedObjectsForGC);
+		// A list of UObject references within the graph
+		const TArray<TObjectPtr<UObject>>& GraphReferencedObjects;
 
 		// A list of node handles for each node within the archive
 		TArray<FNodeHandle> NodeHandles;
diff --git a/Engine/Plugins/Experimental/AnimNext/Source/AnimNext/Internal/DecoratorBase/DecoratorWriter.h b/Engine/Plugins/Experimental/AnimNext/Source/AnimNext/Internal/DecoratorBase/DecoratorWriter.h
--- a/Engine/Plugins/Experimental/AnimNext/Source/AnimNext/Internal/DecoratorBase/DecoratorWriter.h
+++ b/Engine/Plugins/Experimental/AnimNext/Source/AnimNext/Internal/DecoratorBase/DecoratorWriter.h
@@ -60,6 +60,9 @@ namespace UE::AnimNext
 		// Returns the populated raw graph shared data buffer
 		[[nodiscard]] const TArray<uint8>& GetGraphSharedData() const;
 
+		// Returns the list of referenced UObjects in this graph
+		[[nodiscard]] const TArray<UObject*>& GetGraphReferencedObjects() const;
+
 		// FArchive implementation
 		virtual FArchive& operator<<(UObject*& Obj) override;
 		virtual FArchive& operator<<(FObjectPtr& Obj) override;
@@ -84,7 +87,7 @@ namespace UE::AnimNext
 		FNodeID NextNodeID;
 
 		// To track node writing
-		TArray<UObject*> TrackedObjectsForGC;
+		TArray<UObject*> GraphReferencedObjects;
 		uint32 NumNodesWritten;
 		FLatentPropertyHandle CurrentLatentPropertyHandle;
 		bool bIsNodeWriting;
diff --git a/Engine/Plugins/Experimental/AnimNext/Source/AnimNext/Internal/Graph/AnimNextGraph.h b/Engine/Plugins/Experimental/AnimNext/Source/AnimNext/Internal/Graph/AnimNextGraph.h
--- a/Engine/Plugins/Experimental/AnimNext/Source/AnimNext/Internal/Graph/AnimNextGraph.h
+++ b/Engine/Plugins/Experimental/AnimNext/Source/AnimNext/Internal/Graph/AnimNextGraph.h
@@ -7,6 +7,8 @@
 #include "DecoratorBase/DecoratorPtr.h"
 #include "DecoratorBase/DecoratorHandle.h"
 #include "DecoratorBase/EntryPointHandle.h"
+#include "Graph/RigUnit_AnimNextGraphEvaluator.h"
+
 #include "AnimNextGraph.generated.h"
 
 class UEdGraph;
@@ -84,7 +86,6 @@ public:
 	virtual void PostRename(UObject* OldOuter, const FName OldName) override;
 	virtual void GetPreloadDependencies(TArray<UObject*>& OutDeps) override;
 	virtual void Serialize(FArchive& Ar) override;
-	static void AddReferencedObjects(UObject* InThis, FReferenceCollector& Collector);
 
 	// Allocates an instance of the graph, retain the handle and use it with the Run() function to evaluate it
 	void AllocateInstance(FAnimNextGraphInstance& Instance) const;
@@ -106,6 +107,10 @@ protected:
 	friend struct FAnimNextGraphInstance;
 	friend class UAnimGraphNode_AnimNextGraph;
 
+	// This is the execute method definition used by this graph
+	UPROPERTY()
+	FAnimNextGraphEvaluatorExecuteDefinition ExecuteDefinition;
+
 	// This is a handle to the root decorator in our graph
 	UPROPERTY()
 	FAnimNextEntryPointHandle RootDecoratorHandle;
@@ -117,10 +122,12 @@ protected:
 	// We de-serialize manually into this buffer from the archive buffer, this is never saved on disk
 	TArray<uint8> SharedDataBuffer;
 
-	// This is the list of all referenced UObjects in the graph shared data
+	// This is a list of all referenced UObjects in the graph shared data
 	// We collect all the references here to make it quick and easy for the GC to query them
 	// It means that object references in the graph shared data are not visited at runtime by the GC (they are immutable)
-	TArray<TObjectPtr<UObject>> TrackedObjectsForGC;
+	// The shared data serialization archive stores indices to these to perform UObject serialization
+	UPROPERTY()
+	TArray<TObjectPtr<UObject>> GraphReferencedObjects;
 
 	// The RigVM object holds the bytecode, literals, etc used by the RigVM internals, there is a single instance along with the UAnimNextGraph (1:1 mapping)
 	UPROPERTY()
diff --git a/Engine/Plugins/Experimental/AnimNext/Source/AnimNext/Internal/Graph/RigUnit_AnimNextGraphEvaluator.h b/Engine/Plugins/Experimental/AnimNext/Source/AnimNext/Internal/Graph/RigUnit_AnimNextGraphEvaluator.h
--- a/Engine/Plugins/Experimental/AnimNext/Source/AnimNext/Internal/Graph/RigUnit_AnimNextGraphEvaluator.h
+++ b/Engine/Plugins/Experimental/AnimNext/Source/AnimNext/Internal/Graph/RigUnit_AnimNextGraphEvaluator.h
@@ -6,6 +6,35 @@
 
 #include "RigUnit_AnimNextGraphEvaluator.generated.h"
 
+// Represents an argument to a FRigUnit_AnimNextGraphEvaluator entry point
+USTRUCT()
+struct FAnimNextGraphEvaluatorExecuteArgument
+{
+	GENERATED_BODY();
+
+	UPROPERTY()
+	FString Name;
+
+	UPROPERTY()
+	FString CPPType;
+};
+
+// Represents an entry point of FRigUnit_AnimNextGraphEvaluator
+USTRUCT()
+struct FAnimNextGraphEvaluatorExecuteDefinition
+{
+	GENERATED_BODY();
+
+	UPROPERTY()
+	uint32 Hash;
+
+	UPROPERTY()
+	FString MethodName;
+
+	UPROPERTY()
+	TArray<FAnimNextGraphEvaluatorExecuteArgument> Arguments;
+};
+
 /**
  * Animation graph evaluator
  * This node is only used at runtime.
@@ -21,4 +50,11 @@ struct ANIMNEXT_API FRigUnit_AnimNextGraphEvaluator : public FRigUnit_AnimNextBa
 	FAnimNextExecuteContext ExecuteContext;
 
 	static void StaticExecute(FRigVMExtendedExecuteContext& RigVMExecuteContext, FRigVMMemoryHandleArray RigVMMemoryHandles, FRigVMPredicateBranchArray RigVMBranches);
+
+	// The graph evaluator entry point is dynamically registered because it can have any number of latent pins with their own name/type
+	// RigVM requires a mapping between argument list and method name and so we register things manually
+	static void RegisterExecuteMethod(const FAnimNextGraphEvaluatorExecuteDefinition& ExecuteDefinition);
+
+	// Returns an existing execute method, if found. Otherwise, nullptr is returned.
+	static const FAnimNextGraphEvaluatorExecuteDefinition* FindExecuteMethod(uint32 ExecuteMethodHash);
 };
diff --git a/Engine/Plugins/Experimental/AnimNext/Source/AnimNext/Private/DecoratorBase/DecoratorReader.cpp b/Engine/Plugins/Experimental/AnimNext/Source/AnimNext/Private/DecoratorBase/DecoratorReader.cpp
--- a/Engine/Plugins/Experimental/AnimNext/Source/AnimNext/Private/DecoratorBase/DecoratorReader.cpp
+++ b/Engine/Plugins/Experimental/AnimNext/Source/AnimNext/Private/DecoratorBase/DecoratorReader.cpp
@@ -11,15 +11,15 @@
 
 namespace UE::AnimNext
 {
-	FDecoratorReader::FDecoratorReader(FArchive& Ar)
+	FDecoratorReader::FDecoratorReader(const TArray<TObjectPtr<UObject>>& InGraphReferencedObjects, FArchive& Ar)
 		: FArchiveProxy(Ar)
+		, GraphReferencedObjects(InGraphReferencedObjects)
 	{
 	}
 
-	FDecoratorReader::EErrorState FDecoratorReader::ReadGraph(TArray<uint8>& GraphSharedData, TArray<TObjectPtr<UObject>>& TrackedObjectsForGC)
+	FDecoratorReader::EErrorState FDecoratorReader::ReadGraph(TArray<uint8>& GraphSharedData)
 	{
 		GraphSharedData.Empty(0);
-		TrackedObjectsForGC.Empty(0);
 
 		EErrorState ErrorState = ReadGraphSharedData(GraphSharedData);
 		if (ErrorState != EErrorState::None)
@@ -27,8 +27,6 @@ namespace UE::AnimNext
 			return ErrorState;
 		}
 
-		ErrorState = ReadTrackedObjectsForGC(TrackedObjectsForGC);
-
 		return ErrorState;
 	}
 
@@ -140,40 +138,20 @@ namespace UE::AnimNext
 		return EErrorState::None;
 	}
 
-	FDecoratorReader::EErrorState FDecoratorReader::ReadTrackedObjectsForGC(TArray<TObjectPtr<UObject>>& TrackedObjectsForGC)
+	FArchive& FDecoratorReader::operator<<(UObject*& Obj)
 	{
-		int32 NumTrackedObjectsForGC = 0;
-		*this << NumTrackedObjectsForGC;
-
-		TrackedObjectsForGC.SetNumZeroed(NumTrackedObjectsForGC);
+		// Load the object index
+		int32 ObjectIndex = INDEX_NONE;
+		*this << ObjectIndex;
 
-		for (int32 TrackedIndex = 0; TrackedIndex < NumTrackedObjectsForGC; ++TrackedIndex)
+		if (ensure(GraphReferencedObjects.IsValidIndex(ObjectIndex)))
 		{
-			// Load the path name to the object
-			FString LoadedString;
-			*this << LoadedString;
-
-			// Look up the object by fully qualified pathname
-			UObject* Obj = FindObject<UObject>(nullptr, *LoadedString, false);
-			TrackedObjectsForGC[TrackedIndex] = Obj;
+			Obj = GraphReferencedObjects[ObjectIndex];
 		}
-
-		return EErrorState::None;
-	}
-
-	FArchive& FDecoratorReader::operator<<(UObject*& Obj)
-	{
-		// Load the path name to the object
-		FString LoadedString;
-		*this << LoadedString;
-
-		// Look up the object by fully qualified pathname
-		Obj = FindObject<UObject>(nullptr, *LoadedString, false);
-
-		// If we couldn't find it, load it
-		if (Obj == nullptr)
+		else
 		{
-			Obj = LoadObject<UObject>(nullptr, *LoadedString);
+			// Something went wrong, the reference list must have gotten out of sync which shouldn't happen
+			Obj = nullptr;
 		}
 
 		return *this;
diff --git a/Engine/Plugins/Experimental/AnimNext/Source/AnimNext/Private/DecoratorBase/DecoratorWriter.cpp b/Engine/Plugins/Experimental/AnimNext/Source/AnimNext/Private/DecoratorBase/DecoratorWriter.cpp
--- a/Engine/Plugins/Experimental/AnimNext/Source/AnimNext/Private/DecoratorBase/DecoratorWriter.cpp
+++ b/Engine/Plugins/Experimental/AnimNext/Source/AnimNext/Private/DecoratorBase/DecoratorWriter.cpp
@@ -74,7 +74,7 @@ namespace UE::AnimNext
 		}
 
 		bIsNodeWriting = true;
-		TrackedObjectsForGC.Reset();
+		GraphReferencedObjects.Reset();
 		CurrentLatentPropertyHandle = FLatentPropertyHandle::GetFirstHandle();
 
 		// Serialize the node templates
@@ -120,17 +120,6 @@ namespace UE::AnimNext
 		}
 
 		ensure(NumNodesWritten == NodeMappings.Num());
-
-		// Now write all the object reference we found, we'll use them for the GC callbacks at runtime to avoid travering the graph for the references
-		int32 NumTrackedObjectsForGC = TrackedObjectsForGC.Num();
-		*this << NumTrackedObjectsForGC;
-
-		for (UObject* Obj : TrackedObjectsForGC)
-		{
-			// Save out the fully qualified object name
-			FString SavedString(Obj->GetPathName());
-			*this << SavedString;
-		}
 	}
 
 	void FDecoratorWriter::WriteNode(
@@ -226,19 +215,18 @@ namespace UE::AnimNext
 		return GraphSharedDataArchiveBuffer;
 	}
 
-	FArchive& FDecoratorWriter::operator<<(UObject*& Obj)
+	const TArray<UObject*>& FDecoratorWriter::GetGraphReferencedObjects() const
 	{
-		// Add our object for GC tracking
-		TrackedObjectsForGC.AddUnique(Obj);
+		return GraphReferencedObjects;
+	}
 
-		// TODO: Instead of saving our object path, we could save the object index
-		//       Indices can be re-used if an object is referenced more than once (is it common?)
-		//       Allows all objects to be batch found/loaded before nodes are de-serialized
-		//       Requires seeking to read the objects first or an intermediary archive to write objects first
+	FArchive& FDecoratorWriter::operator<<(UObject*& Obj)
+	{
+		// Add our object for tracking
+		int32 ObjectIndex = GraphReferencedObjects.AddUnique(Obj);
 
-		// Save out the fully qualified object name
-		FString SavedString(Obj->GetPathName());
-		*this << SavedString;
+		// Save our index, we'll use it to resolve the object on load
+		*this << ObjectIndex;
 
 		return *this;
 	}
diff --git a/Engine/Plugins/Experimental/AnimNext/Source/AnimNext/Private/Graph/AnimNextGraph.cpp b/Engine/Plugins/Experimental/AnimNext/Source/AnimNext/Private/Graph/AnimNextGraph.cpp
--- a/Engine/Plugins/Experimental/AnimNext/Source/AnimNext/Private/Graph/AnimNextGraph.cpp
+++ b/Engine/Plugins/Experimental/AnimNext/Source/AnimNext/Private/Graph/AnimNextGraph.cpp
@@ -106,6 +106,13 @@ void UAnimNextGraph::PostLoad()
 	Super::PostLoad();
 
 	ExtendedExecuteContext.InvalidateCachedMemory();
+
+	// In packaged builds, initialize the VM
+	// In editor, the VM will be recompiled and initialized at RecompileVM
+#if !WITH_EDITOR
+	RigVM->ClearExternalVariables(ExtendedExecuteContext);
+	RigVM->Initialize(ExtendedExecuteContext);
+#endif
 }
 
 void UAnimNextGraph::PostRename(UObject* OldOuter, const FName OldName)
@@ -162,7 +169,12 @@ void UAnimNextGraph::Serialize(FArchive& Ar)
 		SharedDataArchiveBuffer.SetNumUninitialized(SharedDataArchiveBufferSize);
 		Ar.Serialize(SharedDataArchiveBuffer.GetData(), SharedDataArchiveBufferSize);
 
-		LoadFromArchiveBuffer(SharedDataArchiveBuffer);
+		if (Ar.IsLoadingFromCookedPackage())
+		{
+			// If we are cooked, we populate our graph shared data otherwise in the editor we'll compile on load
+			// and re-populate everything then to account for changes in code/content
+			LoadFromArchiveBuffer(SharedDataArchiveBuffer);
+		}
 	}
 	else if (Ar.IsSaving())
 	{
@@ -179,7 +191,6 @@ void UAnimNextGraph::Serialize(FArchive& Ar)
 	{
 		// Counting, etc
 		Ar << SharedDataBuffer;
-		Ar << TrackedObjectsForGC;
 
 #if WITH_EDITORONLY_DATA
 		Ar << SharedDataArchiveBuffer;
@@ -187,32 +198,26 @@ void UAnimNextGraph::Serialize(FArchive& Ar)
 	}
 }
 
-void UAnimNextGraph::AddReferencedObjects(UObject* InThis, FReferenceCollector& Collector)
-{
-	UAnimNextGraph* This = CastChecked<UAnimNextGraph>(InThis);
-	Collector.AddStableReferenceArray(&This->TrackedObjectsForGC);
-
-	Super::AddReferencedObjects(InThis, Collector);
-}
-
 bool UAnimNextGraph::LoadFromArchiveBuffer(const TArray<uint8>& InSharedDataArchiveBuffer)
 {
 	using namespace UE::AnimNext;
 
 	// Reconstruct our graph shared data
 	FMemoryReader GraphSharedDataArchive(InSharedDataArchiveBuffer);
-	FDecoratorReader DecoratorReader(GraphSharedDataArchive);
+	FDecoratorReader DecoratorReader(GraphReferencedObjects, GraphSharedDataArchive);
 
-	const FDecoratorReader::EErrorState ErrorState = DecoratorReader.ReadGraph(SharedDataBuffer, TrackedObjectsForGC);
+	const FDecoratorReader::EErrorState ErrorState = DecoratorReader.ReadGraph(SharedDataBuffer);
 	if (ErrorState == FDecoratorReader::EErrorState::None)
 	{
 		ResolvedRootDecoratorHandle = DecoratorReader.ResolveEntryPointHandle(RootDecoratorHandle);
+
+		// Make sure our execute method is registered
+		FRigUnit_AnimNextGraphEvaluator::RegisterExecuteMethod(ExecuteDefinition);
 		return true;
 	}
 	else
 	{
 		SharedDataBuffer.Empty(0);
-		TrackedObjectsForGC.Empty(0);
 		ResolvedRootDecoratorHandle = FAnimNextDecoratorHandle();
 		return false;
 	}
diff --git a/Engine/Plugins/Experimental/AnimNext/Source/AnimNext/Private/Graph/RigUnit_AnimNextGraphEvaluator.cpp b/Engine/Plugins/Experimental/AnimNext/Source/AnimNext/Private/Graph/RigUnit_AnimNextGraphEvaluator.cpp
--- a/Engine/Plugins/Experimental/AnimNext/Source/AnimNext/Private/Graph/RigUnit_AnimNextGraphEvaluator.cpp
+++ b/Engine/Plugins/Experimental/AnimNext/Source/AnimNext/Private/Graph/RigUnit_AnimNextGraphEvaluator.cpp
@@ -11,6 +11,24 @@
 
 #include UE_INLINE_GENERATED_CPP_BY_NAME(RigUnit_AnimNextGraphEvaluator)
 
+namespace UE::AnimNext::Private
+{
+	static TMap<uint32, FAnimNextGraphEvaluatorExecuteDefinition> GRegisteredGraphEvaluatorMethods;
+
+	TArray<FRigVMFunctionArgument> GetGraphEvaluatorFunctionArguments(const FAnimNextGraphEvaluatorExecuteDefinition& ExecuteDefinition)
+	{
+		TArray<FRigVMFunctionArgument> Arguments;
+		Arguments.Reserve(ExecuteDefinition.Arguments.Num());
+
+		for (const FAnimNextGraphEvaluatorExecuteArgument& Argument : ExecuteDefinition.Arguments)
+		{
+			Arguments.Add(FRigVMFunctionArgument(Argument.Name, Argument.CPPType, ERigVMFunctionArgumentDirection::Input));
+		}
+
+		return Arguments;
+	}
+}
+
 void FRigUnit_AnimNextGraphEvaluator::StaticExecute(FRigVMExtendedExecuteContext& RigVMExecuteContext, FRigVMMemoryHandleArray RigVMMemoryHandles, FRigVMPredicateBranchArray RigVMBranches)
 {
 	using namespace UE::AnimNext;
@@ -63,3 +81,27 @@ void FRigUnit_AnimNextGraphEvaluator::StaticExecute(FRigVMExtendedExecuteContext
 		}
 	}
 }
+
+void FRigUnit_AnimNextGraphEvaluator::RegisterExecuteMethod(const FAnimNextGraphEvaluatorExecuteDefinition& ExecuteDefinition)
+{
+	using namespace UE::AnimNext::Private;
+
+	if (GRegisteredGraphEvaluatorMethods.Contains(ExecuteDefinition.Hash))
+	{
+		return;	// Already registered
+	}
+
+	GRegisteredGraphEvaluatorMethods.Add(ExecuteDefinition.Hash, ExecuteDefinition);
+
+	const FString FullExecuteMethodName = FString::Printf(TEXT("FRigUnit_AnimNextGraphEvaluator::%s"), *ExecuteDefinition.MethodName);
+
+	const TArray<FRigVMFunctionArgument> GraphEvaluatorArguments = GetGraphEvaluatorFunctionArguments(ExecuteDefinition);
+	FRigVMRegistry::Get().Register(*FullExecuteMethodName, &FRigUnit_AnimNextGraphEvaluator::StaticExecute, FRigUnit_AnimNextGraphEvaluator::StaticStruct(), GraphEvaluatorArguments);
+}
+
+const FAnimNextGraphEvaluatorExecuteDefinition* FRigUnit_AnimNextGraphEvaluator::FindExecuteMethod(uint32 ExecuteMethodHash)
+{
+	using namespace UE::AnimNext::Private;
+
+	return GRegisteredGraphEvaluatorMethods.Find(ExecuteMethodHash);
+}
diff --git a/Engine/Plugins/Experimental/AnimNext/Source/AnimNextTestSuite/Private/AnimNextDecoratorBaseTest.cpp b/Engine/Plugins/Experimental/AnimNext/Source/AnimNextTestSuite/Private/AnimNextDecoratorBaseTest.cpp
--- a/Engine/Plugins/Experimental/AnimNext/Source/AnimNextTestSuite/Private/AnimNextDecoratorBaseTest.cpp
+++ b/Engine/Plugins/Experimental/AnimNext/Source/AnimNextTestSuite/Private/AnimNextDecoratorBaseTest.cpp
@@ -522,6 +522,7 @@ bool FAnimationAnimNextRuntimeTest_NodeLifetime::RunTest(const FString& InParame
 
 	// Write our graph
 	TArray<uint8> GraphSharedDataArchiveBuffer;
+	TArray<TObjectPtr<UObject>> GraphReferencedObjects;
 	{
 		FDecoratorWriter DecoratorWriter;
 
@@ -553,16 +554,16 @@ bool FAnimationAnimNextRuntimeTest_NodeLifetime::RunTest(const FString& InParame
 
 		AddErrorIfFalse(DecoratorWriter.GetErrorState() == FDecoratorWriter::EErrorState::None, "FAnimationAnimNextRuntimeTest_NodeLifetime -> Failed to write decorators");
 		GraphSharedDataArchiveBuffer = DecoratorWriter.GetGraphSharedData();
+		GraphReferencedObjects = DecoratorWriter.GetGraphReferencedObjects();
 	}
 
 	// Read our graph
 	TArray<uint8> GraphSharedDataBuffer;
-	TArray<TObjectPtr<UObject>> TrackedObjectsForGC;
 	{
 		FMemoryReader GraphSharedDataArchive(GraphSharedDataArchiveBuffer);
-		FDecoratorReader DecoratorReader(GraphSharedDataArchive);
+		FDecoratorReader DecoratorReader(GraphReferencedObjects, GraphSharedDataArchive);
 
-		FDecoratorReader::EErrorState ErrorState = DecoratorReader.ReadGraph(GraphSharedDataBuffer, TrackedObjectsForGC);
+		FDecoratorReader::EErrorState ErrorState = DecoratorReader.ReadGraph(GraphSharedDataBuffer);
 		AddErrorIfFalse(ErrorState == FDecoratorReader::EErrorState::None, "FAnimationAnimNextRuntimeTest_NodeLifetime -> Failed to read graph shared data");
 
 		Node0 = DecoratorReader.ResolveNodeHandle(Node0);
@@ -733,6 +734,7 @@ bool FAnimationAnimNextRuntimeTest_GetDecoratorInterface::RunTest(const FString&
 
 	// Write our graph
 	TArray<uint8> GraphSharedDataArchiveBuffer;
+	TArray<TObjectPtr<UObject>> GraphReferencedObjects;
 	{
 		FDecoratorWriter DecoratorWriter;
 
@@ -754,16 +756,16 @@ bool FAnimationAnimNextRuntimeTest_GetDecoratorInterface::RunTest(const FString&
 
 		AddErrorIfFalse(DecoratorWriter.GetErrorState() == FDecoratorWriter::EErrorState::None, "FAnimationAnimNextRuntimeTest_GetDecoratorInterface -> Failed to write decorators");
 		GraphSharedDataArchiveBuffer = DecoratorWriter.GetGraphSharedData();
+		GraphReferencedObjects = DecoratorWriter.GetGraphReferencedObjects();
 	}
 
 	// Read our graph
 	TArray<uint8> GraphSharedDataBuffer;
-	TArray<TObjectPtr<UObject>> TrackedObjectsForGC;
 	{
 		FMemoryReader GraphSharedDataArchive(GraphSharedDataArchiveBuffer);
-		FDecoratorReader DecoratorReader(GraphSharedDataArchive);
+		FDecoratorReader DecoratorReader(GraphReferencedObjects, GraphSharedDataArchive);
 
-		FDecoratorReader::EErrorState ErrorState = DecoratorReader.ReadGraph(GraphSharedDataBuffer, TrackedObjectsForGC);
+		FDecoratorReader::EErrorState ErrorState = DecoratorReader.ReadGraph(GraphSharedDataBuffer);
 		AddErrorIfFalse(ErrorState == FDecoratorReader::EErrorState::None, "FAnimationAnimNextRuntimeTest_GetDecoratorInterface -> Failed to read graph shared data");
 
 		Node0 = DecoratorReader.ResolveNodeHandle(Node0);
@@ -1003,6 +1005,7 @@ bool FAnimationAnimNextRuntimeTest_GetDecoratorInterfaceSuper::RunTest(const FSt
 
 	// Write our graph
 	TArray<uint8> GraphSharedDataArchiveBuffer;
+	TArray<TObjectPtr<UObject>> GraphReferencedObjects;
 	{
 		FDecoratorWriter DecoratorWriter;
 
@@ -1024,16 +1027,16 @@ bool FAnimationAnimNextRuntimeTest_GetDecoratorInterfaceSuper::RunTest(const FSt
 
 		AddErrorIfFalse(DecoratorWriter.GetErrorState() == FDecoratorWriter::EErrorState::None, "FAnimationAnimNextRuntimeTest_GetDecoratorInterfaceSuper -> Failed to write decorators");
 		GraphSharedDataArchiveBuffer = DecoratorWriter.GetGraphSharedData();
+		GraphReferencedObjects = DecoratorWriter.GetGraphReferencedObjects();
 	}
 
 	// Read our graph
 	TArray<uint8> GraphSharedDataBuffer;
-	TArray<TObjectPtr<UObject>> TrackedObjectsForGC;
 	{
 		FMemoryReader GraphSharedDataArchive(GraphSharedDataArchiveBuffer);
-		FDecoratorReader DecoratorReader(GraphSharedDataArchive);
+		FDecoratorReader DecoratorReader(GraphReferencedObjects, GraphSharedDataArchive);
 
-		FDecoratorReader::EErrorState ErrorState = DecoratorReader.ReadGraph(GraphSharedDataBuffer, TrackedObjectsForGC);
+		FDecoratorReader::EErrorState ErrorState = DecoratorReader.ReadGraph(GraphSharedDataBuffer);
 		AddErrorIfFalse(ErrorState == FDecoratorReader::EErrorState::None, "FAnimationAnimNextRuntimeTest_GetDecoratorInterfaceSuper -> Failed to read graph shared data");
 
 		Node0 = DecoratorReader.ResolveNodeHandle(Node0);
@@ -1285,6 +1288,7 @@ bool FAnimationAnimNextRuntimeTest_DecoratorSerialization::RunTest(const FString
 
 	// Write our graph
 	TArray<uint8> GraphSharedDataArchiveBuffer;
+	TArray<TObjectPtr<UObject>> GraphReferencedObjects;
 	{
 		FDecoratorWriter DecoratorWriter;
 
@@ -1375,6 +1379,7 @@ bool FAnimationAnimNextRuntimeTest_DecoratorSerialization::RunTest(const FString
 
 		AddErrorIfFalse(DecoratorWriter.GetErrorState() == FDecoratorWriter::EErrorState::None, "FAnimationAnimNextRuntimeTest_DecoratorSerialization -> Failed to write decorators");
 		GraphSharedDataArchiveBuffer = DecoratorWriter.GetGraphSharedData();
+		GraphReferencedObjects = DecoratorWriter.GetGraphReferencedObjects();
 	}
 
 	// Clear out the node template registry to test registration on load
@@ -1385,12 +1390,11 @@ bool FAnimationAnimNextRuntimeTest_DecoratorSerialization::RunTest(const FString
 
 		// Read our graph
 		TArray<uint8> GraphSharedDataBuffer;
-		TArray<TObjectPtr<UObject>> TrackedObjectsForGC;
 		{
 			FMemoryReader GraphSharedDataArchive(GraphSharedDataArchiveBuffer);
-			FDecoratorReader DecoratorReader(GraphSharedDataArchive);
+			FDecoratorReader DecoratorReader(GraphReferencedObjects, GraphSharedDataArchive);
 
-			FDecoratorReader::EErrorState ErrorState = DecoratorReader.ReadGraph(GraphSharedDataBuffer, TrackedObjectsForGC);
+			FDecoratorReader::EErrorState ErrorState = DecoratorReader.ReadGraph(GraphSharedDataBuffer);
 			AddErrorIfFalse(ErrorState == FDecoratorReader::EErrorState::None, "FAnimationAnimNextRuntimeTest_DecoratorSerialization -> Failed to read graph shared data");
 
 			Node0 = DecoratorReader.ResolveNodeHandle(Node0);
diff --git a/Engine/Plugins/Experimental/AnimNext/Source/AnimNextTestSuite/Private/AnimNextDecoratorInterfacesTest.cpp b/Engine/Plugins/Experimental/AnimNext/Source/AnimNextTestSuite/Private/AnimNextDecoratorInterfacesTest.cpp
--- a/Engine/Plugins/Experimental/AnimNext/Source/AnimNextTestSuite/Private/AnimNextDecoratorInterfacesTest.cpp
+++ b/Engine/Plugins/Experimental/AnimNext/Source/AnimNextTestSuite/Private/AnimNextDecoratorInterfacesTest.cpp
@@ -254,6 +254,7 @@ bool FAnimationAnimNextRuntimeTest_IHierarchy::RunTest(const FString& InParamete
 
 	// Write our graph
 	TArray<uint8> GraphSharedDataArchiveBuffer;
+	TArray<TObjectPtr<UObject>> GraphReferencedObjects;
 	{
 		FDecoratorWriter DecoratorWriter;
 
@@ -321,16 +322,16 @@ bool FAnimationAnimNextRuntimeTest_IHierarchy::RunTest(const FString& InParamete
 
 		AddErrorIfFalse(DecoratorWriter.GetErrorState() == FDecoratorWriter::EErrorState::None, "FAnimationAnimNextRuntimeTest_IHierarchy -> Failed to write decorators");
 		GraphSharedDataArchiveBuffer = DecoratorWriter.GetGraphSharedData();
+		GraphReferencedObjects = DecoratorWriter.GetGraphReferencedObjects();
 	}
 
 	// Read our graph
 	TArray<uint8> GraphSharedDataBuffer;
-	TArray<TObjectPtr<UObject>> TrackedObjectsForGC;
 	{
 		FMemoryReader GraphSharedDataArchive(GraphSharedDataArchiveBuffer);
-		FDecoratorReader DecoratorReader(GraphSharedDataArchive);
+		FDecoratorReader DecoratorReader(GraphReferencedObjects, GraphSharedDataArchive);
 
-		FDecoratorReader::EErrorState ErrorState = DecoratorReader.ReadGraph(GraphSharedDataBuffer, TrackedObjectsForGC);
+		FDecoratorReader::EErrorState ErrorState = DecoratorReader.ReadGraph(GraphSharedDataBuffer);
 		AddErrorIfFalse(ErrorState == FDecoratorReader::EErrorState::None, "FAnimationAnimNextRuntimeTest_IHierarchy -> Failed to read graph shared data");
 
 		NodeA = DecoratorReader.ResolveNodeHandle(NodeA);
@@ -456,6 +457,7 @@ bool FAnimationAnimNextRuntimeTest_IUpdate::RunTest(const FString& InParameters)
 
 	// Write our graph
 	TArray<uint8> GraphSharedDataArchiveBuffer;
+	TArray<TObjectPtr<UObject>> GraphReferencedObjects;
 	{
 		FDecoratorWriter DecoratorWriter;
 
@@ -508,16 +510,16 @@ bool FAnimationAnimNextRuntimeTest_IUpdate::RunTest(const FString& InParameters)
 
 		AddErrorIfFalse(DecoratorWriter.GetErrorState() == FDecoratorWriter::EErrorState::None, "FAnimationAnimNextRuntimeTest_IUpdate -> Failed to write decorators");
 		GraphSharedDataArchiveBuffer = DecoratorWriter.GetGraphSharedData();
+		GraphReferencedObjects = DecoratorWriter.GetGraphReferencedObjects();
 	}
 
 	// Read our graph
 	TArray<uint8> GraphSharedDataBuffer;
-	TArray<TObjectPtr<UObject>> TrackedObjectsForGC;
 	{
 		FMemoryReader GraphSharedDataArchive(GraphSharedDataArchiveBuffer);
-		FDecoratorReader DecoratorReader(GraphSharedDataArchive);
+		FDecoratorReader DecoratorReader(GraphReferencedObjects, GraphSharedDataArchive);
 
-		FDecoratorReader::EErrorState ErrorState = DecoratorReader.ReadGraph(GraphSharedDataBuffer, TrackedObjectsForGC);
+		FDecoratorReader::EErrorState ErrorState = DecoratorReader.ReadGraph(GraphSharedDataBuffer);
 		AddErrorIfFalse(ErrorState == FDecoratorReader::EErrorState::None, "FAnimationAnimNextRuntimeTest_IUpdate -> Failed to read graph shared data");
 
 		NodeA = DecoratorReader.ResolveNodeHandle(NodeA);
@@ -605,6 +607,7 @@ bool FAnimationAnimNextRuntimeTest_IEvaluate::RunTest(const FString& InParameter
 
 	// Write our graph
 	TArray<uint8> GraphSharedDataArchiveBuffer;
+	TArray<TObjectPtr<UObject>> GraphReferencedObjects;
 	{
 		FDecoratorWriter DecoratorWriter;
 
@@ -658,16 +661,16 @@ bool FAnimationAnimNextRuntimeTest_IEvaluate::RunTest(const FString& InParameter
 
 		AddErrorIfFalse(DecoratorWriter.GetErrorState() == FDecoratorWriter::EErrorState::None, "FAnimationAnimNextRuntimeTest_IEvaluate -> Failed to write decorators");
 		GraphSharedDataArchiveBuffer = DecoratorWriter.GetGraphSharedData();
+		GraphReferencedObjects = DecoratorWriter.GetGraphReferencedObjects();
 	}
 
 	// Read our graph
 	TArray<uint8> GraphSharedDataBuffer;
-	TArray<TObjectPtr<UObject>> TrackedObjectsForGC;
 	{
 		FMemoryReader GraphSharedDataArchive(GraphSharedDataArchiveBuffer);
-		FDecoratorReader DecoratorReader(GraphSharedDataArchive);
+		FDecoratorReader DecoratorReader(GraphReferencedObjects, GraphSharedDataArchive);
 
-		FDecoratorReader::EErrorState ErrorState = DecoratorReader.ReadGraph(GraphSharedDataBuffer, TrackedObjectsForGC);
+		FDecoratorReader::EErrorState ErrorState = DecoratorReader.ReadGraph(GraphSharedDataBuffer);
 		AddErrorIfFalse(ErrorState == FDecoratorReader::EErrorState::None, "FAnimationAnimNextRuntimeTest_IEvaluate -> Failed to read graph shared data");
 
 		NodeA = DecoratorReader.ResolveNodeHandle(NodeA);
diff --git a/Engine/Plugins/Experimental/AnimNext/Source/AnimNextUncookedOnly/Private/UncookedOnlyUtils.cpp b/Engine/Plugins/Experimental/AnimNext/Source/AnimNextUncookedOnly/Private/UncookedOnlyUtils.cpp
--- a/Engine/Plugins/Experimental/AnimNext/Source/AnimNextUncookedOnly/Private/UncookedOnlyUtils.cpp
+++ b/Engine/Plugins/Experimental/AnimNext/Source/AnimNextUncookedOnly/Private/UncookedOnlyUtils.cpp
@@ -324,27 +324,36 @@ namespace Private
 		return Arguments;
 	}
 
-	FString GetGraphEvaluatorMethodName(const FRigVMPinInfoArray& LatentPins)
+	FAnimNextGraphEvaluatorExecuteDefinition GetGraphEvaluatorExecuteMethod(const FRigVMPinInfoArray& LatentPins)
 	{
-		static TMap<uint32, FString> GraphEvaluatorMethodNameCache;
-
 		const uint32 LatentPinListHash = GetTypeHash(LatentPins);
-		if (const FString* MethodName = GraphEvaluatorMethodNameCache.Find(LatentPinListHash))
+		if (const FAnimNextGraphEvaluatorExecuteDefinition* ExecuteDefinition = FRigUnit_AnimNextGraphEvaluator::FindExecuteMethod(LatentPinListHash))
 		{
-			return *MethodName;
+			return *ExecuteDefinition;
 		}
 
+		const FRigVMRegistry& Registry = FRigVMRegistry::Get();
+
 		// Generate a new method for this argument list
-		const FString MethodName = FString::Printf(TEXT("Execute_%X"), LatentPinListHash);
-		const FString FullExecuteMethodName = FString::Printf(TEXT("FRigUnit_AnimNextGraphEvaluator::%s"), *MethodName);
+		FAnimNextGraphEvaluatorExecuteDefinition ExecuteDefinition;
+		ExecuteDefinition.Hash = LatentPinListHash;
+		ExecuteDefinition.MethodName = FString::Printf(TEXT("Execute_%X"), LatentPinListHash);
+		ExecuteDefinition.Arguments.Reserve(LatentPins.Num());
 
-		const TArray<FRigVMFunctionArgument> GraphEvaluatorArguments = GetGraphEvaluatorFunctionArguments(LatentPins);
-		FRigVMRegistry::Get().Register(*FullExecuteMethodName, &FRigUnit_AnimNextGraphEvaluator::StaticExecute, FRigUnit_AnimNextGraphEvaluator::StaticStruct(), GraphEvaluatorArguments);
+		for (const FRigVMPinInfo& Pin : LatentPins)
+		{
+			const FRigVMTemplateArgumentType& TypeArg = Registry.GetType(Pin.TypeIndex);
+
+			FAnimNextGraphEvaluatorExecuteArgument Argument;
+			Argument.Name = Pin.Name.ToString();
+			Argument.CPPType = TypeArg.CPPType.ToString();
+
+			ExecuteDefinition.Arguments.Add(Argument);
+		}
 
-		// Cache our result
-		GraphEvaluatorMethodNameCache.Add(LatentPinListHash, MethodName);
+		FRigUnit_AnimNextGraphEvaluator::RegisterExecuteMethod(ExecuteDefinition);
 
-		return MethodName;
+		return ExecuteDefinition;
 	}
 }
 
@@ -394,9 +403,9 @@ void FUtils::Compile(UAnimNextGraph* InGraph)
 	const FRigVMPinInfoArray LatentPins = Private::CollectLatentPins(DecoratorStackNodes, LatentPinMapping);
 
 	// We need a unique method name to match our unique argument list
-	const FString ExecuteMethodName = Private::GetGraphEvaluatorMethodName(LatentPins);
+	const FAnimNextGraphEvaluatorExecuteDefinition ExecuteDefinition = Private::GetGraphEvaluatorExecuteMethod(LatentPins);
 
-	URigVMUnitNode* GraphEvaluatorNode = TempController->AddUnitNodeWithPins(FRigUnit_AnimNextGraphEvaluator::StaticStruct(), LatentPins, *ExecuteMethodName, FVector2D::ZeroVector, FString(), false);
+	URigVMUnitNode* GraphEvaluatorNode = TempController->AddUnitNodeWithPins(FRigUnit_AnimNextGraphEvaluator::StaticStruct(), LatentPins, *ExecuteDefinition.MethodName, FVector2D::ZeroVector, FString(), false);
 
 	// Link our shim and evaluator nodes together using the execution context
 	TempController->AddLink(
@@ -437,7 +446,9 @@ void FUtils::Compile(UAnimNextGraph* InGraph)
 	}
 
 	// Cache our compiled metadata
+	InGraph->ExecuteDefinition = ExecuteDefinition;
 	InGraph->SharedDataArchiveBuffer = DecoratorWriter.GetGraphSharedData();
+	InGraph->GraphReferencedObjects = DecoratorWriter.GetGraphReferencedObjects();
 	InGraph->RootDecoratorHandle = RootDecoratorHandle;
 
 	// Populate our runtime metadata
