Fix CIS uninitialized var issues

#rnx

diff --git a/Engine/Plugins/Experimental/AnimNext/Source/AnimNext/Internal/Param/AnimNextObjectAdapterConfig.h b/Engine/Plugins/Experimental/AnimNext/Source/AnimNext/Internal/Param/AnimNextObjectAdapterConfig.h
--- a/Engine/Plugins/Experimental/AnimNext/Source/AnimNext/Internal/Param/AnimNextObjectAdapterConfig.h
+++ b/Engine/Plugins/Experimental/AnimNext/Source/AnimNext/Internal/Param/AnimNextObjectAdapterConfig.h
@@ -1,4 +1,4 @@
-﻿// Copyright Epic Games, Inc. All Rights Reserved.
+// Copyright Epic Games, Inc. All Rights Reserved.
 
 #pragma once
 
@@ -52,5 +52,5 @@ struct FAnimNextObjectAdapterConfig
 
 	/** Whether to expose a tick function (e.g. for a component or actor) along with functions and properties */
 	UPROPERTY(EditAnywhere, Config, Category="Adapter")
-	bool bRegisterTickFunction;
+	bool bRegisterTickFunction = false;
 };
diff --git a/Engine/Plugins/Experimental/AnimNext/Source/AnimNext/Internal/Scheduler/AnimNextSchedule.h b/Engine/Plugins/Experimental/AnimNext/Source/AnimNext/Internal/Scheduler/AnimNextSchedule.h
--- a/Engine/Plugins/Experimental/AnimNext/Source/AnimNext/Internal/Scheduler/AnimNextSchedule.h
+++ b/Engine/Plugins/Experimental/AnimNext/Source/AnimNext/Internal/Scheduler/AnimNextSchedule.h
@@ -130,10 +130,10 @@ struct FAnimNextScheduleInstruction
 	GENERATED_BODY()
 
 	UPROPERTY()
-	EAnimNextScheduleScheduleOpcode Opcode;
+	EAnimNextScheduleScheduleOpcode Opcode = EAnimNextScheduleScheduleOpcode::None;
 
 	UPROPERTY()
-	int32 Operand;
+	int32 Operand = INDEX_NONE;
 };
 
 UCLASS()
