AnimNext: Can now run standalone without anim blueprints for support.

Additions:
Basic port implementation - A simple class to allow native ports to be written and extended. Ideally ports should eventually be RigVM assets that allow arbitrary composition of operations.
IO 'terms' for schedule entries - Per-schedule allocated intermediate values that are mapped into schedule task's IO terms, allowing tasks to communicate via those temporaries. Ideally this should all be handled at the RigVM level when we change schedules to be graphs, but this works to get the system stood up for now.
Anim Next Mesh Component - A non-ticking mesh component that acts as a container for component-space transforms and performs no game-thread work.
Added profiling scopes to get a better picture of system performance.
Added some more built-in parameter handle types. This has highlighted the sheer amount of code that needs to be updated to add new types, so this needs a rework.
Added the ability for a parameter layer to be mapped from another parameter layer, to support task/term mappings.
Expanded FParamStackLayerHandle API to wrap more of FParamStackLayer.
Changed FScheduleTickFunction to run a slice of schedule instructions, and used this in the initialization path to ensure schedules run initially correctly in non-ticking editor worlds (this fixes editor thumbnails).
Included parent bone indices in TReferencePose to aid in component-space pose generation. Also added some helper functions in FGenerationTools to aid with this (should these be in another library?)

Fixes/changes:
Fixed AnimNext graph editing so that the correct graph is both edited by the graph editor AND compiled by the compiler.
Changed FOnLODRequiredBonesUpdateMulticast to be a thread-safe delegate as it is now called on worker threads.
Sanitised UObject-ptr value parameter access. Now the type has to request the T<UObject*> type (or T<TObjectPtr<>>) to work. Previously it allowed T<UObject> which was incorrect.
Fixed race condition in FDataRegistry by adding appropriate lock.
Added extra checking to higher-level APIs to deal with various scenarios where FScheduler calls could be made when worlds or world subsystems might be invalid (e.g. changing editor worlds).
Added a resolved hard object ptr to schedule entries that is only valid when schedules are running - this both provides GC safety and improves performance as the weak object ptr was being accessed for every adapter call.
Various extra null checks added to places where content can provide null values (parameter blocks, sequence player decorators etc.)
Fixed schedule editing to allow undo/redo.
Fixed FEvaluationVM to work if there is not a mesh present (as it tries to make a bone container, which can assert).
Removed the guard around not displaying object types in SRigVMGraphChangePinType so that wildcard pins can be set to object type values.
Changed a number of CVars to GetValueOnAnyThread as their calling functions are now called on worker threads for the first time.
Updated FRootMotionReset::ResetRootBoneForRootMotion to not be dependent on a bone container, allowing better root motion handling.
Moved some utility functions from editor to uncookedonly.

diff --git a/Engine/Plugins/Experimental/AnimNext/Config/BaseAnimNext.ini b/Engine/Plugins/Experimental/AnimNext/Config/BaseAnimNext.ini
--- a/Engine/Plugins/Experimental/AnimNext/Config/BaseAnimNext.ini
+++ b/Engine/Plugins/Experimental/AnimNext/Config/BaseAnimNext.ini
@@ -1,4 +1,4 @@
 [/Script/AnimNext.AnimNextConfig]
 +ExposedClasses=(Class="/Script/CoreUObject.Class'/Script/Engine.CharacterMovementComponent'",RootParameter="UE_CharacterMovementComponent",FilterType=AllowOnlyFunctions,FilteredFields=,bRegisterTickFunction=True)
-+ExposedClasses=(Class="/Script/CoreUObject.Class'/Script/Engine.SkeletalMeshComponent'",RootParameter="UE_SkeletalMeshComponent",FilterType=AllowOnlyFunctions,FilteredFields=,bRegisterTickFunction=True)
++ExposedClasses=(Class="/Script/CoreUObject.Class'/Script/AnimNext.AnimNextMeshComponent'",RootParameter="UE_AnimNextMeshComponent",FilterType=AllowOnlyFunctions,FilteredFields=,bRegisterTickFunction=True)
 
diff --git a/Engine/Plugins/Experimental/AnimNext/Source/AnimNext/Internal/Graph/AnimNextExecuteContext.h b/Engine/Plugins/Experimental/AnimNext/Source/AnimNext/Internal/Graph/AnimNextExecuteContext.h
--- a/Engine/Plugins/Experimental/AnimNext/Source/AnimNext/Internal/Graph/AnimNextExecuteContext.h
+++ b/Engine/Plugins/Experimental/AnimNext/Source/AnimNext/Internal/Graph/AnimNextExecuteContext.h
@@ -6,9 +6,10 @@
 #include "RigVMCore/RigVMExecuteContext.h"
 #include "Units/RigUnit.h"
 #include "DecoratorBase/DecoratorPtr.h"
-
 #include "AnimNextExecuteContext.generated.h"
 
+class UAnimNextGraph;
+
 namespace UE::AnimNext
 {
 	struct FContext;
@@ -39,12 +40,7 @@ struct FAnimNextExecuteContext : public FRigVMExecuteContext
 {
 	GENERATED_BODY()
 
-	FAnimNextExecuteContext()
-		: FRigVMExecuteContext()
-		, Context(nullptr)
-		, SimulationSteps(EAnimNextGraphSimulationSteps::None)
-	{
-	}
+	FAnimNextExecuteContext() = default;
 
 	const UE::AnimNext::FContext& GetContext() const
 	{
@@ -52,6 +48,8 @@ struct FAnimNextExecuteContext : public FRigVMExecuteContext
 		return *Context;
 	}
 
+	const UAnimNextGraph* GetGraph() const { return Graph; }
+
 	const TArrayView<const uint8>& GetSharedDataBuffer() const { return SharedDataBuffer; }
 
 	UE::AnimNext::FWeakDecoratorPtr GetGraphInstancePtr() const { return GraphInstancePtr; }
@@ -64,6 +62,7 @@ struct FAnimNextExecuteContext : public FRigVMExecuteContext
 
 		const FAnimNextExecuteContext* OtherContext = (const FAnimNextExecuteContext*)InOtherContext;
 		Context = OtherContext->Context;
+		Graph = OtherContext->Graph; 
 		SharedDataBuffer = OtherContext->SharedDataBuffer;
 		GraphInstancePtr = OtherContext->GraphInstancePtr;
 		SimulationSteps = OtherContext->SimulationSteps;
@@ -75,8 +74,9 @@ private:
 		Context = &InContext;
 	}
 
-	void InitializeWithGraph(TArrayView<const uint8> InSharedDataBuffer, UE::AnimNext::FWeakDecoratorPtr InGraphInstancePtr)
+	void InitializeWithGraph(const UAnimNextGraph* InGraph, TArrayView<const uint8> InSharedDataBuffer, UE::AnimNext::FWeakDecoratorPtr InGraphInstancePtr)
 	{
+		Graph = InGraph;
 		SharedDataBuffer = InSharedDataBuffer;
 		GraphInstancePtr = InGraphInstancePtr;
 	}
@@ -90,16 +90,17 @@ private:
 	void DebugReset()
 	{
 		Context = nullptr;
+		Graph = nullptr;
 		SharedDataBuffer = TArrayView<const uint8>();
 		GraphInstancePtr.Reset();
 		SimulationSteps = EAnimNextGraphSimulationSteps::None;
 	}
 
-	const UE::AnimNext::FContext* Context;
-
+	const UE::AnimNext::FContext* Context = nullptr;
+	const UAnimNextGraph* Graph = nullptr;
 	TArrayView<const uint8> SharedDataBuffer;
 	UE::AnimNext::FWeakDecoratorPtr GraphInstancePtr;
-	EAnimNextGraphSimulationSteps SimulationSteps;
+	EAnimNextGraphSimulationSteps SimulationSteps = EAnimNextGraphSimulationSteps::None;
 
 	friend class UAnimNextGraph;
 	friend class UAnimNextParameterBlock;
diff --git a/Engine/Plugins/Experimental/AnimNext/Source/AnimNext/Internal/Graph/AnimNextGraph.h b/Engine/Plugins/Experimental/AnimNext/Source/AnimNext/Internal/Graph/AnimNextGraph.h
--- a/Engine/Plugins/Experimental/AnimNext/Source/AnimNext/Internal/Graph/AnimNextGraph.h
+++ b/Engine/Plugins/Experimental/AnimNext/Source/AnimNext/Internal/Graph/AnimNextGraph.h
@@ -8,7 +8,8 @@
 #include "DecoratorBase/DecoratorHandle.h"
 #include "DecoratorBase/EntryPointHandle.h"
 #include "Graph/RigUnit_AnimNextGraphEvaluator.h"
-
+#include "Param/ParamId.h"
+#include "Scheduler/IAnimNextScheduleTermInterface.h"
 #include "AnimNextGraph.generated.h"
 
 class UEdGraph;
@@ -76,7 +77,7 @@ private:
 
 // A user-created graph of logic used to supply data
 UCLASS(BlueprintType)
-class ANIMNEXT_API UAnimNextGraph : public UObject
+class ANIMNEXT_API UAnimNextGraph : public UObject, public IAnimNextScheduleTermInterface
 {
 	GENERATED_BODY()
 
@@ -87,12 +88,24 @@ public:
 	virtual void GetPreloadDependencies(TArray<UObject*>& OutDeps) override;
 	virtual void Serialize(FArchive& Ar) override;
 
+	// IAnimNextScheduleTermInterface interface
+	virtual TConstArrayView<UE::AnimNext::FScheduleTerm> GetTerms() const override;
+
 	// Allocates an instance of the graph, retain the handle and use it with the Run() function to evaluate it
 	void AllocateInstance(FAnimNextGraphInstance& Instance) const;
 
 	// Run the specified simulation steps on the provided graph with the given context
 	void Run(const UE::AnimNext::FContext& Context, FAnimNextGraphInstance& GraphInstance, EAnimNextGraphSimulationSteps SimulationSteps) const;
 
+	// Get the parameter to use to access the reference pose
+	UE::AnimNext::FParamId GetReferencePoseParam() const { return ReferencePoseId; }
+
+	// Get the parameter to use to access the current LOD
+	UE::AnimNext::FParamId GetCurrentLODParam() const { return CurrentLODId; }
+
+	// Get the parameter to use to access the delta time
+	UE::AnimNext::FParamId GetDeltaTimeParam() const { return DeltaTimeId; }
+	
 protected:
 	// Support rig VM execution
 	TArray<FRigVMExternalVariable> GetRigVMExternalVariables();
@@ -142,6 +155,22 @@ protected:
 	UPROPERTY()
 	FRigVMRuntimeSettings VMRuntimeSettings;
 
+	// The parameter to use to access the reference pose
+	UPROPERTY(EditAnywhere, Category = "Graph", meta=(CustomWidget = "ParamName", AllowedParamType = "FAnimNextGraphReferencePose"))
+	FName ReferencePose = TEXT("UE_AnimNextMeshComponent_ReferencePose");
+
+	// The parameter to use to access the current LOD
+	UPROPERTY(EditAnywhere, Category = "Graph", meta=(CustomWidget = "ParamName", AllowedParamType = "int32"))
+	FName CurrentLOD = TEXT("UE_AnimNextMeshComponent_PredictedLODLevel");
+
+	// The parameter to use to access the current delta time
+	UPROPERTY(EditAnywhere, Category = "Graph", meta=(CustomWidget = "ParamName", AllowedParamType = "float"))
+	FName DeltaTime = TEXT("UE_AnimNextSchedulerWorldSubsystem_DeltaTime");
+
+	UE::AnimNext::FParamId ReferencePoseId = UE::AnimNext::FParamId(ReferencePose);
+	UE::AnimNext::FParamId CurrentLODId = UE::AnimNext::FParamId(CurrentLOD);
+	UE::AnimNext::FParamId DeltaTimeId = UE::AnimNext::FParamId(DeltaTime);
+
 #if WITH_EDITORONLY_DATA
 	UPROPERTY(VisibleAnywhere, Instanced, Category = "Graph", meta = (ShowInnerProperties))
 	TObjectPtr<UObject> EditorData;
diff --git a/Engine/Plugins/Experimental/AnimNext/Source/AnimNext/Internal/Param/ParamDefinition.h b/Engine/Plugins/Experimental/AnimNext/Source/AnimNext/Internal/Param/ParamDefinition.h
--- a/Engine/Plugins/Experimental/AnimNext/Source/AnimNext/Internal/Param/ParamDefinition.h
+++ b/Engine/Plugins/Experimental/AnimNext/Source/AnimNext/Internal/Param/ParamDefinition.h
@@ -13,7 +13,7 @@ namespace UE::AnimNext
 	struct FParamAdapter;
 }
 
-namespace UE::AnimNext::Editor
+namespace UE::AnimNext::UncookedOnly
 {
 	struct FUtils;
 }
@@ -30,7 +30,7 @@ enum class EParamDefinitionFlags : uint8
 
 ENUM_CLASS_FLAGS(EParamDefinitionFlags);
 
-// Definition of a parameter for reserved or internal use
+// Definition of a parameter for reserved for internal use
 struct FParamDefinition
 {
 	FParamDefinition() = default;
@@ -53,6 +53,11 @@ struct FParamDefinition
 		// Periods are only a display concern and we want parameters to be expressible as members of objects & structures. 
 		check(!Name.ToString().Contains(TEXT(".")));
 	}
+
+	FParamDefinition(FName InName, const FAnimNextParamType& InType, EParamDefinitionFlags InFlags = EParamDefinitionFlags::None)
+		: FParamDefinition(InName, InType, FText::GetEmpty(), InFlags)
+	{
+	}
 	
 	FParamDefinition(FParamId InId, const FAnimNextParamType& InType, const FText& InTooltip, EParamDefinitionFlags InFlags = EParamDefinitionFlags::None)
 		: Id(InId)
@@ -68,11 +73,21 @@ struct FParamDefinition
 	{
 	}
 
+	FParamDefinition(FParamId InId, const FAnimNextParamType& InType, EParamDefinitionFlags InFlags = EParamDefinitionFlags::None)
+		: FParamDefinition(InId, InType, FText::GetEmpty(), InFlags)
+	{
+	}
+
 	FName GetName() const
 	{
 		return Name;
 	}
 
+	FParamId GetId() const
+	{
+		return Id;
+	}
+	
 	FParamTypeHandle GetTypeHandle() const
 	{
 		return TypeHandle;
@@ -87,7 +102,7 @@ private:
 	friend struct FParamId;
 	friend struct FParams;
 	friend struct FParamAdapter;
-	friend struct Editor::FUtils;
+	friend struct UncookedOnly::FUtils;
 
 	FParamDefinition(int32 InIdIndex, FName InName, const FAnimNextParamType& InType, const FText& InTooltip, EParamDefinitionFlags InFlags = EParamDefinitionFlags::None)
 		: Id(InIdIndex)
diff --git a/Engine/Plugins/Experimental/AnimNext/Source/AnimNext/Internal/Scheduler/AnimNextSchedule.h b/Engine/Plugins/Experimental/AnimNext/Source/AnimNext/Internal/Scheduler/AnimNextSchedule.h
--- a/Engine/Plugins/Experimental/AnimNext/Source/AnimNext/Internal/Scheduler/AnimNextSchedule.h
+++ b/Engine/Plugins/Experimental/AnimNext/Source/AnimNext/Internal/Scheduler/AnimNextSchedule.h
@@ -44,17 +44,17 @@ class UAnimNextScheduleEntry_AnimNextGraph : public UAnimNextScheduleEntry
 private:
 	friend class UAnimNextSchedule;
 
-	UPROPERTY(EditAnywhere, Category = "Graph", meta = (CustomWidget = "ParamName", AllowedParamType = "FAnimNextScope"))
-	FName Name;
+	// The graph to run
+	UPROPERTY(EditAnywhere, Category = "Graph")
+	TObjectPtr<UAnimNextGraph> Graph = nullptr;
 
+	// An optional entry point to use when running the supplied graph
 	UPROPERTY(EditAnywhere, Category = "Graph", meta = (CustomWidget = "ParamName", AllowedParamType = "FName"))
 	FName EntryPoint;
 
+	// The output intermediate terms used by the graph
 	UPROPERTY(EditAnywhere, Category = "Graph")
-	TObjectPtr<UAnimNextGraph> Graph;
-
-	UPROPERTY(EditAnywhere, Category = "Graph")
-	TArray<TObjectPtr<UAnimNextParameterBlock>> ParameterBlocks;
+	TArray<FName> Terms;
 };
 
 UCLASS()
@@ -65,8 +65,13 @@ class UAnimNextScheduleEntry_Port : public UAnimNextScheduleEntry
 private:
 	friend class UAnimNextSchedule;
 
-	UPROPERTY(EditAnywhere, Category = "Port", meta = (CustomWidget = "ParamName", AllowedParamType = "FAnimNextPort"))
-	FName Name;
+	// The input intermediate terms used by this port
+	UPROPERTY(EditAnywhere, Category = "Port")
+	TArray<FName> Terms;
+
+	// The type of the port to use
+	UPROPERTY(EditAnywhere, Category = "Port")
+	TSubclassOf<UAnimNextSchedulePort> Port;
 };
 
 UCLASS()
@@ -77,11 +82,13 @@ class UAnimNextScheduleEntry_ExternalTask : public UAnimNextScheduleEntry
 private:
 	friend class UAnimNextSchedule;
 
+	// The tick function that this external task should wrap
 	UPROPERTY(EditAnywhere, Category = "External Task", meta = (CustomWidget = "ParamName", AllowedParamType = "FTickFunction"))
-	FName Name;
+	FName TickFunction;
 
-	UPROPERTY(EditAnywhere, Category = "External Task", meta = (CustomWidget = "ParamName", AllowedParamType = "FName"))
-	FName ObjectName;
+	// The object that the tick function is present on
+	UPROPERTY(EditAnywhere, Category = "External Task", meta = (CustomWidget = "ParamName", AllowedParamType = "TObjectPtr<UObject>"))
+	FName Object;
 };
 
 UCLASS()
@@ -92,8 +99,9 @@ class UAnimNextScheduleEntry_ParamScope : public UAnimNextScheduleEntry
 private:
 	friend class UAnimNextSchedule;
 
+	// The scope to use
 	UPROPERTY(EditAnywhere, Category = "Parameters", meta = (CustomWidget = "ParamName", AllowedParamType = "FAnimNextScope"))
-	FName Name;
+	FName Scope;
 
 	// Parameters to apply in this scope
 	UPROPERTY(EditAnywhere, Category = "Parameters")
@@ -136,6 +144,22 @@ struct FAnimNextScheduleInstruction
 	int32 Operand = INDEX_NONE;
 };
 
+UENUM()
+enum class EAnimNextScheduleInitMethod : uint8
+{
+	// Do not perform any initial update, set up data structures only
+	None,
+
+	// Set up data structures, perform an initial update and then pause
+	InitializeAndPause,
+
+	// Set up data structures, perform an initial update and then pause in editor only, otherwise act like InitializeAndRun
+	InitializeAndPauseInEditor,
+
+	// Set up data structures then continue updating
+	InitializeAndRun
+};
+
 UCLASS()
 class ANIMNEXT_API UAnimNextSchedule : public UObject
 {
@@ -155,6 +179,7 @@ private:
 	virtual void PostLoad() override;
 #if WITH_EDITOR
 	virtual void PostEditChangeProperty(FPropertyChangedEvent& PropertyChangedEvent) override;
+	virtual void PostEditUndo() override;
 
 	// Compile the editor data into a compact runtime representation
 	void CompileSchedule();
@@ -169,38 +194,38 @@ private:
 #endif
 
 	// TEMP: Instructions derived from the entries above
-	UPROPERTY()
+	UPROPERTY(NonTransactional)
 	TArray<FAnimNextScheduleInstruction> Instructions;
 
 	// TEMP: Tasks derived from the entries above
-	UPROPERTY()
+	UPROPERTY(NonTransactional)
 	TArray<FAnimNextScheduleGraphTask> Tasks;
 
 	// TEMP: Ports derived from the entries above
-	UPROPERTY()
+	UPROPERTY(NonTransactional)
 	TArray<FAnimNextSchedulePortTask> Ports;
 
 	// TEMP: External tasks derived from the entries above
-	UPROPERTY()
+	UPROPERTY(NonTransactional)
 	TArray<FAnimNextScheduleExternalTask> ExternalTasks;
 
 	// TEMP: Parameter scope entry tasks derived from the entries above
-	UPROPERTY()
+	UPROPERTY(NonTransactional)
 	TArray<FAnimNextScheduleParamScopeEntryTask> ParamScopeEntryTasks;
 
 	// TEMP: Parameter scope exit tasks derived from the entries above
-	UPROPERTY()
+	UPROPERTY(NonTransactional)
 	TArray<FAnimNextScheduleParamScopeExitTask> ParamScopeExitTasks;
 
-	// TEMP: Map from port name to port index in Ports array
-	UPROPERTY()
-	TMap<FName, int32> PortNameIndexMap;
+	// TEMP: Data for intermediates, defined as a property bag
+	UPROPERTY(NonTransactional)
+	FInstancedPropertyBag IntermediatesData;
 
 	// TEMP: Count of total number of parameter scopes that this schedule needs to execute
-	UPROPERTY()
+	UPROPERTY(NonTransactional)
 	uint32 NumParameterScopes = 0;
 
 	// TEMP: Count of total number of tick functions that this schedule needs to execute
-	UPROPERTY()
+	UPROPERTY(NonTransactional)
 	uint32 NumTickFunctions = 0;
 };
\ No newline at end of file
diff --git a/Engine/Plugins/Experimental/AnimNext/Source/AnimNext/Internal/Scheduler/AnimNextScheduleExternalTask.h b/Engine/Plugins/Experimental/AnimNext/Source/AnimNext/Internal/Scheduler/AnimNextScheduleExternalTask.h
--- a/Engine/Plugins/Experimental/AnimNext/Source/AnimNext/Internal/Scheduler/AnimNextScheduleExternalTask.h
+++ b/Engine/Plugins/Experimental/AnimNext/Source/AnimNext/Internal/Scheduler/AnimNextScheduleExternalTask.h
@@ -40,9 +40,9 @@ private:
 
 	/** The name of the external task parameter */
 	UPROPERTY()
-	FName Name;
+	FName TickFunction;
 
 	/** The name of the external task's object parameter */
 	UPROPERTY()
-	FName ObjectName;
+	FName Object;
 };
diff --git a/Engine/Plugins/Experimental/AnimNext/Source/AnimNext/Internal/Scheduler/AnimNextScheduleGraphTask.h b/Engine/Plugins/Experimental/AnimNext/Source/AnimNext/Internal/Scheduler/AnimNextScheduleGraphTask.h
--- a/Engine/Plugins/Experimental/AnimNext/Source/AnimNext/Internal/Scheduler/AnimNextScheduleGraphTask.h
+++ b/Engine/Plugins/Experimental/AnimNext/Source/AnimNext/Internal/Scheduler/AnimNextScheduleGraphTask.h
@@ -3,6 +3,7 @@
 #pragma once
 
 #include "CoreMinimal.h"
+#include "Graph/AnimNext_LODPose.h"
 #include "AnimNextScheduleGraphTask.generated.h"
 
 class UAnimNextGraph;
@@ -20,8 +21,7 @@ struct FAnimNextScheduleGraphTask
 {
 	GENERATED_BODY()
 
-	FAnimNextScheduleGraphTask()
-	{}
+	FAnimNextScheduleGraphTask() = default;
 
 private:
 	friend class UAnimNextComponent;
@@ -29,7 +29,7 @@ private:
 	friend struct UE::AnimNext::FScheduleInstanceData;
 	friend struct UE::AnimNext::FScheduleTickFunction;
 
-	void RunGraph(const UE::AnimNext::FScheduleContext& InScheduleContext) const;
+	void RunGraph(const UE::AnimNext::FScheduleContext& InContext) const;
 
 private:
 	UPROPERTY()
@@ -41,15 +41,13 @@ private:
 	UPROPERTY()
 	uint32 ParamParentScopeIndex = MAX_uint32;
 
-	UPROPERTY()
-	FName Name;
-
 	UPROPERTY()
 	FName EntryPoint;
 
 	UPROPERTY()
-	TObjectPtr<UAnimNextGraph> Graph;
+	TObjectPtr<UAnimNextGraph> Graph = nullptr;
 
+	// Index of each term in the schedule intermediates
 	UPROPERTY()
-	TArray<TObjectPtr<UAnimNextParameterBlock>> ParameterBlocks;
+	TArray<uint32> Terms;
 };
diff --git a/Engine/Plugins/Experimental/AnimNext/Source/AnimNext/Internal/Scheduler/AnimNextScheduleParamScopeTask.h b/Engine/Plugins/Experimental/AnimNext/Source/AnimNext/Internal/Scheduler/AnimNextScheduleParamScopeTask.h
--- a/Engine/Plugins/Experimental/AnimNext/Source/AnimNext/Internal/Scheduler/AnimNextScheduleParamScopeTask.h
+++ b/Engine/Plugins/Experimental/AnimNext/Source/AnimNext/Internal/Scheduler/AnimNextScheduleParamScopeTask.h
@@ -46,7 +46,7 @@ private:
 
 	// The name of the scope
 	UPROPERTY()
-	FName Name;
+	FName Scope;
 
 	/** Parameters to apply in this scope */
 	UPROPERTY()
@@ -75,5 +75,5 @@ private:
 
 	// The name of the scope
 	UPROPERTY()
-	FName Name;
+	FName Scope;
 };
diff --git a/Engine/Plugins/Experimental/AnimNext/Source/AnimNext/Internal/Scheduler/AnimNextSchedulePortTask.h b/Engine/Plugins/Experimental/AnimNext/Source/AnimNext/Internal/Scheduler/AnimNextSchedulePortTask.h
--- a/Engine/Plugins/Experimental/AnimNext/Source/AnimNext/Internal/Scheduler/AnimNextSchedulePortTask.h
+++ b/Engine/Plugins/Experimental/AnimNext/Source/AnimNext/Internal/Scheduler/AnimNextSchedulePortTask.h
@@ -3,15 +3,19 @@
 #pragma once
 
 #include "CoreMinimal.h"
+#include "Param/ParamType.h"
+#include "Templates/SubclassOf.h"
 #include "AnimNextSchedulePortTask.generated.h"
 
 class UAnimNextSchedule;
+class UAnimNextSchedulePort;
 
 namespace UE::AnimNext
 {
 	struct FScheduleContext;
 	struct FScheduleTask;
 	struct FScheduleTickFunction;
+	struct FScheduleInstanceData;
 }
 
 USTRUCT()
@@ -22,6 +26,7 @@ struct FAnimNextSchedulePortTask
 	friend struct UE::AnimNext::FScheduleTask;
 	friend struct UE::AnimNext::FScheduleTickFunction;
 	friend class UAnimNextSchedule;
+	friend struct UE::AnimNext::FScheduleInstanceData;
 
 	FAnimNextSchedulePortTask() = default;
 
@@ -35,11 +40,11 @@ private:
 	UPROPERTY()
 	uint32 ParamScopeIndex = MAX_uint32;
 
-	// The name of the port
+	// The type of the port
 	UPROPERTY()
-	FName Name;
+	TSubclassOf<UAnimNextSchedulePort> Port;
 
-	// The name of the ports input parameter
+	// Index of each term in the schedule intermediates
 	UPROPERTY()
-	FName InputParameterName;
+	TArray<uint32> Terms;
 };
diff --git a/Engine/Plugins/Experimental/AnimNext/Source/AnimNext/Private/AnimNextConfig.cpp b/Engine/Plugins/Experimental/AnimNext/Source/AnimNext/Private/AnimNextConfig.cpp
--- a/Engine/Plugins/Experimental/AnimNext/Source/AnimNext/Private/AnimNextConfig.cpp
+++ b/Engine/Plugins/Experimental/AnimNext/Source/AnimNext/Private/AnimNextConfig.cpp
@@ -11,7 +11,7 @@ void UAnimNextConfig::PostEditChangeProperty(FPropertyChangedEvent& PropertyChan
 	
 	if(PropertyChangedEvent.Property->GetFName() == GET_MEMBER_NAME_CHECKED(UAnimNextConfig, ExposedClasses))
 	{
-		FParamId::RefreshConfigAdapters();
+		FParamId::RefreshAdapters();
 	}
 
 	SaveConfig();
diff --git a/Engine/Plugins/Experimental/AnimNext/Source/AnimNext/Private/Component/AnimNextComponent.cpp b/Engine/Plugins/Experimental/AnimNext/Source/AnimNext/Private/Component/AnimNextComponent.cpp
--- a/Engine/Plugins/Experimental/AnimNext/Source/AnimNext/Private/Component/AnimNextComponent.cpp
+++ b/Engine/Plugins/Experimental/AnimNext/Source/AnimNext/Private/Component/AnimNextComponent.cpp
@@ -15,7 +15,7 @@ void UAnimNextComponent::OnRegister()
 	if (Schedule)
 	{
 		check(!SchedulerHandle.IsValid());
-		SchedulerHandle = FScheduler::AcquireHandle(this, Schedule, Parameters);
+		SchedulerHandle = FScheduler::AcquireHandle(this, Schedule, Parameters, InitMethod);
 	}
 }
 
diff --git a/Engine/Plugins/Experimental/AnimNext/Source/AnimNext/Private/Component/AnimNextComponent.h b/Engine/Plugins/Experimental/AnimNext/Source/AnimNext/Private/Component/AnimNextComponent.h
--- a/Engine/Plugins/Experimental/AnimNext/Source/AnimNext/Private/Component/AnimNextComponent.h
+++ b/Engine/Plugins/Experimental/AnimNext/Source/AnimNext/Private/Component/AnimNextComponent.h
@@ -5,6 +5,7 @@
 #include "Components/ActorComponent.h"
 #include "Param/AnimNextParameterCollection.h"
 #include "Param/IAnimNextParameterSourceInterface.h"
+#include "Scheduler/AnimNextSchedule.h"
 #include "Scheduler/ScheduleHandle.h"
 #include "AnimNextComponent.generated.h"
 
@@ -54,6 +55,10 @@ private:
 	UPROPERTY(EditAnywhere, Category = "Parameters", meta = (CustomWidget = "ParamName", AllowedParamType = "FAnimNextScope"))
 	TMap<FName, FAnimNextParameterCollection> Parameters;
 
+	// How to initialize the schedule
+	UPROPERTY(EditAnywhere, Category="Schedule")
+	EAnimNextScheduleInitMethod InitMethod = EAnimNextScheduleInitMethod::InitializeAndPauseInEditor;
+
 	// Handle to the registered results/schedule
 	UE::AnimNext::FScheduleHandle SchedulerHandle;
 
diff --git a/Engine/Plugins/Experimental/AnimNext/Source/AnimNext/Private/Component/AnimNextMeshComponent.cpp b/Engine/Plugins/Experimental/AnimNext/Source/AnimNext/Private/Component/AnimNextMeshComponent.cpp
new file mode 100644
--- /dev/null
+++ b/Engine/Plugins/Experimental/AnimNext/Source/AnimNext/Private/Component/AnimNextMeshComponent.cpp
@@ -0,0 +1,34 @@
+﻿// Copyright Epic Games, Inc. All Rights Reserved.
+
+#include "AnimNextMeshComponent.h"
+#include "GenerationTools.h"
+
+UAnimNextMeshComponent::UAnimNextMeshComponent()
+{
+	// Disable regular skeletal mesh ticking
+	PrimaryComponentTick.bCanEverTick = false;
+	PrimaryComponentTick.bStartWithTickEnabled = false;
+	PrimaryComponentTick.bAllowTickOnDedicatedServer = false;
+}
+
+void UAnimNextMeshComponent::CompleteAndDispatch(TConstArrayView<FBoneIndexType> InParentIndices, TConstArrayView<FBoneIndexType> InRequiredBoneIndices, TConstArrayView<FTransform> InLocalSpaceTransforms)
+{
+	// Fill the component space transform buffer
+	TArrayView<FTransform> ComponentSpaceTransforms = GetEditableComponentSpaceTransforms();
+	if(ComponentSpaceTransforms.Num() > 0)
+	{
+		UE::AnimNext::FGenerationTools::ConvertLocalSpaceToComponentSpace(InParentIndices, InLocalSpaceTransforms, InRequiredBoneIndices, ComponentSpaceTransforms);
+
+		// Flag buffer for flip
+		bNeedToFlipSpaceBaseBuffers = true;
+		bHasValidBoneTransform = false;
+		FlipEditableSpaceBases();
+		bHasValidBoneTransform = true;
+
+		InvalidateCachedBounds();
+		UpdateBounds();
+
+		// Send updated transforms to the renderer
+		SendRenderDynamicData_Concurrent();
+	}
+}
diff --git a/Engine/Plugins/Experimental/AnimNext/Source/AnimNext/Private/Component/AnimNextMeshComponent.h b/Engine/Plugins/Experimental/AnimNext/Source/AnimNext/Private/Component/AnimNextMeshComponent.h
new file mode 100644
--- /dev/null
+++ b/Engine/Plugins/Experimental/AnimNext/Source/AnimNext/Private/Component/AnimNextMeshComponent.h
@@ -0,0 +1,22 @@
+﻿// Copyright Epic Games, Inc. All Rights Reserved.
+
+#pragma once
+
+#include "Components/SkeletalMeshComponent.h"
+#include "AnimNextMeshComponent.generated.h"
+
+class UAnimNextSchedulePort_AnimNextMeshComponentPose;
+
+// Mesh component for use with AnimNext
+UCLASS(MinimalAPI, Blueprintable, meta = (BlueprintSpawnableComponent))
+class UAnimNextMeshComponent : public USkeletalMeshComponent
+{
+	GENERATED_BODY()
+
+	friend class UAnimNextSchedulePort_AnimNextMeshComponentPose;
+
+	UAnimNextMeshComponent();
+
+	// Converts local to component space, flips buffers, updates bounds and dispatches to renderer
+	void CompleteAndDispatch(TConstArrayView<FBoneIndexType> InParentIndices, TConstArrayView<FBoneIndexType> InRequiredBoneIndices, TConstArrayView<FTransform> InLocalSpaceTransforms);
+};
\ No newline at end of file
diff --git a/Engine/Plugins/Experimental/AnimNext/Source/AnimNext/Private/DecompressionTools.cpp b/Engine/Plugins/Experimental/AnimNext/Source/AnimNext/Private/DecompressionTools.cpp
--- a/Engine/Plugins/Experimental/AnimNext/Source/AnimNext/Private/DecompressionTools.cpp
+++ b/Engine/Plugins/Experimental/AnimNext/Source/AnimNext/Private/DecompressionTools.cpp
@@ -456,7 +456,9 @@ void FDecompressionTools::DecompressPose(FLODPose& OutAnimationPoseData,
 	// Once pose has been extracted, snap root bone back to first frame if we are extracting root motion.
 	if ((ExtractionContext.bExtractRootMotion && RootMotionReset.bEnableRootMotion) || RootMotionReset.bForceRootLock)
 	{
-		//RootMotionReset.ResetRootBoneForRootMotion(OutAnimationPoseData.LocalTransforms[FCompactPoseBoneIndex(0)], RequiredBones);
+		FTransform RootTransform = OutAnimationPoseData.LocalTransformsView[0];
+		RootMotionReset.ResetRootBoneForRootMotion(RootTransform, ReferencePose.GetRefPoseTransform(0));
+		OutAnimationPoseData.LocalTransformsView[0] = RootTransform;
 	}
 
 	// Anim Scale Retargeting
diff --git a/Engine/Plugins/Experimental/AnimNext/Source/AnimNext/Private/DecoratorInterfaces/IEvaluate.cpp b/Engine/Plugins/Experimental/AnimNext/Source/AnimNext/Private/DecoratorInterfaces/IEvaluate.cpp
--- a/Engine/Plugins/Experimental/AnimNext/Source/AnimNext/Private/DecoratorInterfaces/IEvaluate.cpp
+++ b/Engine/Plugins/Experimental/AnimNext/Source/AnimNext/Private/DecoratorInterfaces/IEvaluate.cpp
@@ -4,6 +4,9 @@
 
 #include "DecoratorBase/ExecutionContext.h"
 #include "DecoratorInterfaces/IHierarchy.h"
+#include "AnimNextStats.h"
+
+DEFINE_STAT(STAT_AnimNext_EvaluateGraph);
 
 namespace UE::AnimNext
 {
@@ -70,6 +73,8 @@ namespace UE::AnimNext
 
 	FEvaluationProgram EvaluateGraph(FExecutionContext& Context, FEvaluateTraversalContext& TraversalContext, FWeakDecoratorPtr GraphRootPtr)
 	{
+		SCOPE_CYCLE_COUNTER(STAT_AnimNext_EvaluateGraph);
+		
 		FEvaluationProgram EvaluationProgram;
 
 		if (!GraphRootPtr.IsValid())
diff --git a/Engine/Plugins/Experimental/AnimNext/Source/AnimNext/Private/DecoratorInterfaces/IUpdate.cpp b/Engine/Plugins/Experimental/AnimNext/Source/AnimNext/Private/DecoratorInterfaces/IUpdate.cpp
--- a/Engine/Plugins/Experimental/AnimNext/Source/AnimNext/Private/DecoratorInterfaces/IUpdate.cpp
+++ b/Engine/Plugins/Experimental/AnimNext/Source/AnimNext/Private/DecoratorInterfaces/IUpdate.cpp
@@ -4,6 +4,9 @@
 
 #include "DecoratorBase/ExecutionContext.h"
 #include "DecoratorInterfaces/IHierarchy.h"
+#include "AnimNextStats.h"
+
+DEFINE_STAT(STAT_AnimNext_UpdateGraph);
 
 namespace UE::AnimNext
 {
@@ -44,6 +47,8 @@ namespace UE::AnimNext
 
 	void UpdateGraph(FExecutionContext& Context, FUpdateTraversalContext& TraversalContext, FWeakDecoratorPtr GraphRootPtr)
 	{
+		SCOPE_CYCLE_COUNTER(STAT_AnimNext_UpdateGraph);
+		
 		if (!GraphRootPtr.IsValid())
 		{
 			return;	// Nothing to update
diff --git a/Engine/Plugins/Experimental/AnimNext/Source/AnimNext/Private/Decorators/SequencePlayer.cpp b/Engine/Plugins/Experimental/AnimNext/Source/AnimNext/Private/Decorators/SequencePlayer.cpp
--- a/Engine/Plugins/Experimental/AnimNext/Source/AnimNext/Private/Decorators/SequencePlayer.cpp
+++ b/Engine/Plugins/Experimental/AnimNext/Source/AnimNext/Private/Decorators/SequencePlayer.cpp
@@ -29,12 +29,12 @@ namespace UE::AnimNext
 	void FSequencePlayerDecorator::PreEvaluate(FExecutionContext& Context, const TDecoratorBinding<IEvaluate>& Binding) const
 	{
 		const FSharedData* SharedData = Binding.GetSharedData<FSharedData>();
-		const FInstanceData* InstanceData = Binding.GetInstanceData<FInstanceData>();
 
+		const FInstanceData* InstanceData = Binding.GetInstanceData<FInstanceData>();
 		const bool bInterpolate = true;
 
 		FAnimNextAnimSequenceKeyframeTask Task = FAnimNextAnimSequenceKeyframeTask::MakeFromSampleTime(SharedData->AnimSequence, InstanceData->InternalTimeAccumulator, bInterpolate);
-		Task.bExtractTrajectory = false;	/*Output.AnimInstanceProxy->ShouldExtractRootMotion()*/
+		Task.bExtractTrajectory = true;	/*Output.AnimInstanceProxy->ShouldExtractRootMotion()*/
 
 		FEvaluateTraversalContext& TraversalContext = Context.GetTraversalContext<FEvaluateTraversalContext>();
 		TraversalContext.AppendTask(Task);
@@ -49,10 +49,9 @@ namespace UE::AnimNext
 	void FSequencePlayerDecorator::PreUpdate(FExecutionContext& Context, const TDecoratorBinding<IUpdate>& Binding) const
 	{
 		const FSharedData* SharedData = Binding.GetSharedData<FSharedData>();
-		FInstanceData* InstanceData = Binding.GetInstanceData<FInstanceData>();
-
 		if (SharedData->AnimSequence != nullptr)
 		{
+			FInstanceData* InstanceData = Binding.GetInstanceData<FInstanceData>();
 			const FUpdateTraversalContext& TraversalContext = Context.GetTraversalContext<FUpdateTraversalContext>();
 
 			TDecoratorBinding<ITimeline> TimelineDecorator;
diff --git a/Engine/Plugins/Experimental/AnimNext/Source/AnimNext/Private/EvaluationVM/EvaluationProgram.cpp b/Engine/Plugins/Experimental/AnimNext/Source/AnimNext/Private/EvaluationVM/EvaluationProgram.cpp
--- a/Engine/Plugins/Experimental/AnimNext/Source/AnimNext/Private/EvaluationVM/EvaluationProgram.cpp
+++ b/Engine/Plugins/Experimental/AnimNext/Source/AnimNext/Private/EvaluationVM/EvaluationProgram.cpp
@@ -4,9 +4,11 @@
 
 #include "EvaluationVM/EvaluationTask.h"
 #include "EvaluationVM/EvaluationVM.h"
-
+#include "AnimNextStats.h"
 #include "Misc/StringBuilder.h"
 
+DEFINE_STAT(STAT_AnimNext_EvaluationProgram_Execute);
+
 namespace UE::AnimNext
 {
 	bool FEvaluationProgram::IsEmpty() const
@@ -16,7 +18,9 @@ namespace UE::AnimNext
 
 	void FEvaluationProgram::Execute(FEvaluationVM& VM) const
 	{
-		if (ensure(VM.IsValid()))
+		SCOPE_CYCLE_COUNTER(STAT_AnimNext_EvaluationProgram_Execute);
+
+		if (VM.IsValid())
 		{
 			for (const TUniquePtr<FAnimNextEvaluationTask>& Task : Tasks)
 			{
diff --git a/Engine/Plugins/Experimental/AnimNext/Source/AnimNext/Private/EvaluationVM/EvaluationVM.cpp b/Engine/Plugins/Experimental/AnimNext/Source/AnimNext/Private/EvaluationVM/EvaluationVM.cpp
--- a/Engine/Plugins/Experimental/AnimNext/Source/AnimNext/Private/EvaluationVM/EvaluationVM.cpp
+++ b/Engine/Plugins/Experimental/AnimNext/Source/AnimNext/Private/EvaluationVM/EvaluationVM.cpp
@@ -34,30 +34,32 @@ namespace UE::AnimNext
 		, EvaluationFlags(InEvaluationFlags)
 	{
 		const UE::Anim::FCurveFilterSettings CurveFilterSettings(UE::Anim::ECurveFilterMode::DisallowAll);
-		USkeleton* Skeleton = const_cast<USkeleton*>(InReferencePose.Skeleton.Get());	// const_cast because the bone container takes a mutable reference
-		BoneContainer.InitializeTo(InReferencePose.GetLODBoneIndexToMeshBoneIndexMap(InCurrentLOD), CurveFilterSettings, *Skeleton);
+		if(USkeleton* Skeleton = const_cast<USkeleton*>(InReferencePose.Skeleton.Get())) // const_cast because the bone container takes a mutable reference
+		{
+			BoneContainer.InitializeTo(InReferencePose.GetLODBoneIndexToMeshBoneIndexMap(InCurrentLOD), CurveFilterSettings, *Skeleton);
 
-		// TODO: In AnimInstanceProxy this is how we initialize the bone container, we need to get the component somehow or we
-		// gotta figure out how to support ref pose overrides
+			// TODO: In AnimInstanceProxy this is how we initialize the bone container, we need to get the component somehow or we
+			// gotta figure out how to support ref pose overrides
 #if 0
-		// Use the shared bone container
-		RequiredBones = Component->GetSharedRequiredBones();
+			// Use the shared bone container
+			RequiredBones = Component->GetSharedRequiredBones();
 
-		// The first anim instance will initialize the required bones, all others will re-use it
-		if (!RequiredBones->IsValid())
-		{
-			RequiredBones->InitializeTo(Component->RequiredBones, Component->GetCurveFilterSettings(), *Asset);
+			// The first anim instance will initialize the required bones, all others will re-use it
+			if (!RequiredBones->IsValid())
+			{
+				RequiredBones->InitializeTo(Component->RequiredBones, Component->GetCurveFilterSettings(), *Asset);
 
-			// If there is a ref pose override, we want to replace ref pose in RequiredBones
-			// Update ref pose in required bones structure (either set it, or clear it, depending on if one is set on the Component)
-			RequiredBones->SetRefPoseOverride(Component->GetRefPoseOverride());
-		}
+				// If there is a ref pose override, we want to replace ref pose in RequiredBones
+				// Update ref pose in required bones structure (either set it, or clear it, depending on if one is set on the Component)
+				RequiredBones->SetRefPoseOverride(Component->GetRefPoseOverride());
+			}
 #endif
+		}
 	}
 
 	bool FEvaluationVM::IsValid() const
 	{
-		return ReferencePose != nullptr;
+		return ReferencePose != nullptr && BoneContainer.IsValid();
 	}
 
 	void FEvaluationVM::Shrink()
diff --git a/Engine/Plugins/Experimental/AnimNext/Source/AnimNext/Private/EvaluationVM/Tasks/PushAnimSequenceKeyframe.cpp b/Engine/Plugins/Experimental/AnimNext/Source/AnimNext/Private/EvaluationVM/Tasks/PushAnimSequenceKeyframe.cpp
--- a/Engine/Plugins/Experimental/AnimNext/Source/AnimNext/Private/EvaluationVM/Tasks/PushAnimSequenceKeyframe.cpp
+++ b/Engine/Plugins/Experimental/AnimNext/Source/AnimNext/Private/EvaluationVM/Tasks/PushAnimSequenceKeyframe.cpp
@@ -30,38 +30,46 @@ FAnimNextAnimSequenceKeyframeTask FAnimNextAnimSequenceKeyframeTask::MakeFromKey
 void FAnimNextAnimSequenceKeyframeTask::Execute(UE::AnimNext::FEvaluationVM& VM) const
 {
 	using namespace UE::AnimNext;
+	
+	if(const UAnimSequence* AnimSequencePtr = AnimSequence.Get())
+	{
+		const bool bIsAdditive = AnimSequencePtr->IsValidAdditive();
 
-	const UAnimSequence* AnimSequencePtr = AnimSequence.Get();
-	const bool bIsAdditive = AnimSequencePtr->IsValidAdditive();
+		FDeltaTimeRecord DeltaTimeRecord;	// Not needed
+		const bool bExtractRootMotion = bExtractTrajectory;
+		const bool bLooping = false;		// Not needed
+		const bool bUseRawData = false;
 
-	FDeltaTimeRecord DeltaTimeRecord;	// Not needed
-	const bool bExtractRootMotion = bExtractTrajectory;
-	const bool bLooping = false;		// Not needed
-	const bool bUseRawData = false;
+		const FAnimExtractContext ExtractionContext(SampleTime, bExtractRootMotion, DeltaTimeRecord, bLooping);
 
-	const FAnimExtractContext ExtractionContext(SampleTime, bExtractRootMotion, DeltaTimeRecord, bLooping);
+		FKeyframeState Keyframe = VM.MakeUninitializedKeyframe(bIsAdditive);
 
-	FKeyframeState Keyframe = VM.MakeUninitializedKeyframe(bIsAdditive);
+		if (EnumHasAnyFlags(VM.GetFlags(), EEvaluationFlags::Bones))
+		{
+			FDecompressionTools::GetAnimationPose(AnimSequencePtr, Keyframe.Pose, ExtractionContext);
+		}
 
-	if (EnumHasAnyFlags(VM.GetFlags(), EEvaluationFlags::Bones))
-	{
-		FDecompressionTools::GetAnimationPose(AnimSequencePtr, Keyframe.Pose, ExtractionContext);
-	}
+		if (EnumHasAnyFlags(VM.GetFlags(), EEvaluationFlags::Curves))
+		{
+			AnimSequencePtr->EvaluateCurveData(Keyframe.Curves, static_cast<float>(SampleTime), bUseRawData);
+		}
 
-	if (EnumHasAnyFlags(VM.GetFlags(), EEvaluationFlags::Curves))
-	{
-		AnimSequencePtr->EvaluateCurveData(Keyframe.Curves, static_cast<float>(SampleTime), bUseRawData);
-	}
+		if (EnumHasAnyFlags(VM.GetFlags(), EEvaluationFlags::Attributes))
+		{
+			FCompactPose Pose;		// Dummy but we need the bone container
+			Pose.SetBoneContainer(&VM.GetBoneContainer());
 
-	if (EnumHasAnyFlags(VM.GetFlags(), EEvaluationFlags::Attributes))
-	{
-		FCompactPose Pose;		// Dummy but we need the bone container
-		Pose.SetBoneContainer(&VM.GetBoneContainer());
+			FAnimationPoseData PoseData(Pose, Keyframe.Curves, Keyframe.Attributes);
 
-		FAnimationPoseData PoseData(Pose, Keyframe.Curves, Keyframe.Attributes);
+			AnimSequencePtr->EvaluateAttributes(PoseData, ExtractionContext, bUseRawData);
+		}
 
-		AnimSequencePtr->EvaluateAttributes(PoseData, ExtractionContext, bUseRawData);
+		VM.PushValue(KEYFRAME_STACK_NAME, MakeUnique<FKeyframeState>(MoveTemp(Keyframe)));
+	}
+	else
+	{
+		constexpr bool bIsAdditive = false;
+		FKeyframeState Keyframe = VM.MakeReferenceKeyframe(bIsAdditive);
+		VM.PushValue(KEYFRAME_STACK_NAME, MakeUnique<FKeyframeState>(MoveTemp(Keyframe)));
 	}
-
-	VM.PushValue(KEYFRAME_STACK_NAME, MakeUnique<FKeyframeState>(MoveTemp(Keyframe)));
 }
diff --git a/Engine/Plugins/Experimental/AnimNext/Source/AnimNext/Private/GenerationTools.cpp b/Engine/Plugins/Experimental/AnimNext/Source/AnimNext/Private/GenerationTools.cpp
--- a/Engine/Plugins/Experimental/AnimNext/Source/AnimNext/Private/GenerationTools.cpp
+++ b/Engine/Plugins/Experimental/AnimNext/Source/AnimNext/Private/GenerationTools.cpp
@@ -15,11 +15,13 @@
 #include "PhysicsEngine/PhysicsAsset.h"
 #include "Rendering/SkeletalMeshModel.h"
 #include "Rendering/SkeletalMeshRenderData.h"
-
-// auto command
 #include "AssetRegistry/AssetData.h"
 #include "AssetRegistry/AssetRegistryModule.h"
+#include "AnimNextStats.h"
 
+DEFINE_STAT(STAT_AnimNext_GenerateReferencePose);
+DEFINE_STAT(STAT_AnimNext_RemapPose);
+DEFINE_STAT(STAT_AnimNext_ConvertLocalSpaceToComponentSpace);
 
 DEFINE_LOG_CATEGORY_STATIC(LogAnimGenerationTools, Log, All)
 
@@ -40,6 +42,8 @@ struct FCompareBoneIndexType
 	, const USkeletalMesh* SkeletalMesh
 	, FReferencePose& OutAnimationReferencePose)
 {
+	SCOPE_CYCLE_COUNTER(STAT_AnimNext_GenerateReferencePose);
+	
 	using namespace UE::AnimNext;
 
 	bool ReferencePoseGenerated = false;
@@ -408,6 +412,8 @@ struct FCompareBoneIndexType
 // The target pose should be assigned to the correct reference pose prior to this call
 /*static*/ void FGenerationTools::RemapPose(const FPoseContext& SourcePose, FLODPose& TargetPose)
 {
+	SCOPE_CYCLE_COUNTER(STAT_AnimNext_RemapPose);
+	
 	const FBoneContainer& BoneContainer = SourcePose.Pose.GetBoneContainer();
 	const FReferencePose& RefPose = TargetPose.GetRefPose();
 	const TArrayView<const FBoneIndexType> LODBoneIndexes = RefPose.GetLODBoneIndexToMeshBoneIndexMap(TargetPose.LODLevel);
@@ -443,6 +449,8 @@ struct FCompareBoneIndexType
 // The target pose should be assigned to the correct reference pose prior to this call
 /*static*/ void FGenerationTools::RemapPose(const FLODPose& SourcePose, FPoseContext& TargetPose)
 {
+	SCOPE_CYCLE_COUNTER(STAT_AnimNext_RemapPose);
+	
 	const FBoneContainer& BoneContainer = TargetPose.Pose.GetBoneContainer();
 	const FReferencePose& RefPose = SourcePose.GetRefPose();
 	const TArrayView<const FBoneIndexType> LODBoneIndexes = RefPose.GetLODBoneIndexToMeshBoneIndexMap(SourcePose.LODLevel);
@@ -473,6 +481,68 @@ struct FCompareBoneIndexType
 	}
 }
 
+// Converts AnimNext pose to local space transform array
+// This function expects the output pose to have the same or a greater number of bones (as it may be being calculated
+// for a lower LOD
+// The target pose should be assigned to the correct reference pose prior to this call, as transforms will not be filled
+// in by this call if they are not affected by the current LOD.
+/*static*/ void FGenerationTools::RemapPose(const FLODPose& SourcePose, TArrayView<FTransform> TargetTransforms)
+{
+	SCOPE_CYCLE_COUNTER(STAT_AnimNext_RemapPose);
+	
+	const FReferencePose& RefPose = SourcePose.GetRefPose();
+	const TArrayView<const FBoneIndexType> LODBoneIndexes = RefPose.GetLODBoneIndexToMeshBoneIndexMap(SourcePose.LODLevel);
+	const int32 NumLODBones = LODBoneIndexes.Num();
+
+	check(SourcePose.GetNumBones() == NumLODBones);
+
+	for (int32 LODBoneIndex = 0; LODBoneIndex < NumLODBones; ++LODBoneIndex)
+	{
+		TargetTransforms[LODBoneIndexes[LODBoneIndex]] = SourcePose.LocalTransformsView[LODBoneIndex];
+	}
+}
+
+void FGenerationTools::ConvertLocalSpaceToComponentSpace(
+												 TConstArrayView<FBoneIndexType> InParentIndices,
+												 TConstArrayView<FTransform> InBoneSpaceTransforms,
+												 TConstArrayView<FBoneIndexType> InRequiredBoneIndices, 
+												 TArrayView<FTransform> OutComponentSpaceTransforms)
+{
+	SCOPE_CYCLE_COUNTER(STAT_AnimNext_ConvertLocalSpaceToComponentSpace);
+
+	const int32 NumBones = InParentIndices.Num();
+	checkf(NumBones == InBoneSpaceTransforms.Num(), TEXT("Buffer mismatch: %d:%d"), NumBones, InBoneSpaceTransforms.Num());
+	checkf(NumBones == OutComponentSpaceTransforms.Num(), TEXT("Buffer mismatch: %d:%d"), NumBones, OutComponentSpaceTransforms.Num());
+
+	const FBoneIndexType* RESTRICT RequiredBoneIndicesData = InRequiredBoneIndices.GetData();
+	const FBoneIndexType* RESTRICT ParentIndicesData = InParentIndices.GetData();
+	const FTransform* RESTRICT LocalTransformsData = InBoneSpaceTransforms.GetData();
+	FTransform* RESTRICT ComponentSpaceData = OutComponentSpaceTransforms.GetData();
+
+	// First bone (if we have one) is always root bone, and it doesn't have a parent.
+	{
+		check(InRequiredBoneIndices.Num() == 0 || InRequiredBoneIndices[0] == 0);
+		OutComponentSpaceTransforms[0] = InBoneSpaceTransforms[0];
+	}
+
+	const int32 NumRequiredBones = InRequiredBoneIndices.Num();
+	for (int32 RequiredBoneIndex = 1; RequiredBoneIndex < NumRequiredBones; ++RequiredBoneIndex)
+	{
+		const FBoneIndexType* BoneIndex = RequiredBoneIndicesData + RequiredBoneIndex;
+		FTransform* RESTRICT ComponentSpaceTransform = ComponentSpaceData + *BoneIndex;
+		const FBoneIndexType* RESTRICT ParentIndex = ParentIndicesData + *BoneIndex;
+		const FTransform* RESTRICT ParentComponentSpaceTransform = ComponentSpaceData + *ParentIndex;
+		const FTransform* RESTRICT LocalSpaceTransform = LocalTransformsData + *BoneIndex;
+
+		FTransform::Multiply(ComponentSpaceTransform, LocalSpaceTransform, ParentComponentSpaceTransform);
+
+		ComponentSpaceTransform->NormalizeRotation();
+
+		checkSlow(ComponentSpaceTransform->IsRotationNormalized());
+		checkSlow(!ComponentSpaceTransform->ContainsNaN());
+	}
+}
+
 } // end namespace UE::AnimNext
 
 
diff --git a/Engine/Plugins/Experimental/AnimNext/Source/AnimNext/Private/Graph/AnimNextGraph.cpp b/Engine/Plugins/Experimental/AnimNext/Source/AnimNext/Private/Graph/AnimNextGraph.cpp
--- a/Engine/Plugins/Experimental/AnimNext/Source/AnimNext/Private/Graph/AnimNextGraph.cpp
+++ b/Engine/Plugins/Experimental/AnimNext/Source/AnimNext/Private/Graph/AnimNextGraph.cpp
@@ -9,14 +9,18 @@
 #include "Param/ParamTypeHandle.h"
 #include "DecoratorBase/DecoratorReader.h"
 #include "DecoratorBase/ExecutionContext.h"
+#include "Graph/AnimNext_LODPose.h"
 #include "Serialization/MemoryReader.h"
+#include "AnimNextStats.h"
+
+DEFINE_STAT(STAT_AnimNext_Graph);
+DEFINE_STAT(STAT_AnimNext_Graph_AllocateInstance);
 
 #include UE_INLINE_GENERATED_CPP_BY_NAME(AnimNextGraph)
 
 namespace UE::AnimNext::Graph
 {
-const FName EntryPointName("GetData");
-const FName ResultName("Result");
+const UE::AnimNext::FParamId OutputPoseId("UE_Internal_Graph_OutputPose");
 }
 
 FAnimNextGraphInstance::~FAnimNextGraphInstance()
@@ -46,6 +50,8 @@ bool FAnimNextGraphInstance::IsValid() const
 
 void UAnimNextGraph::AllocateInstance(FAnimNextGraphInstance& Instance) const
 {
+	SCOPE_CYCLE_COUNTER(STAT_AnimNext_Graph_AllocateInstance);
+
 	Instance.Release();
 
 	if (!ResolvedRootDecoratorHandle.IsValid())
@@ -65,11 +71,13 @@ void UAnimNextGraph::AllocateInstance(FAnimNextGraphInstance& Instance) const
 
 void UAnimNextGraph::Run(const UE::AnimNext::FContext& Context, FAnimNextGraphInstance& GraphInstance, EAnimNextGraphSimulationSteps SimulationSteps) const
 {
+	SCOPE_CYCLE_COUNTER(STAT_AnimNext_Graph);
+
 	if (RigVM && GraphInstance.IsValid())
 	{
 		FAnimNextExecuteContext& AnimNextContext = GraphInstance.ExtendedExecuteContext.GetPublicDataSafe<FAnimNextExecuteContext>();
 		AnimNextContext.SetContextData(Context);
-		AnimNextContext.InitializeWithGraph(SharedDataBuffer, GraphInstance.GraphInstancePtr);
+		AnimNextContext.InitializeWithGraph(this, SharedDataBuffer, GraphInstance.GraphInstancePtr);
 		AnimNextContext.SetSimulationSteps(SimulationSteps);
 
 		RigVM->ExecuteVM(GraphInstance.ExtendedExecuteContext, FRigUnit_AnimNextShimRoot::EventName);
@@ -103,6 +111,8 @@ static TArray<UClass*> GetClassObjectsInPackage(UPackage* InPackage)
 
 void UAnimNextGraph::PostLoad()
 {
+	using namespace UE::AnimNext;
+	
 	Super::PostLoad();
 
 	ExtendedExecuteContext.InvalidateCachedMemory();
@@ -113,6 +123,10 @@ void UAnimNextGraph::PostLoad()
 	RigVM->ClearExternalVariables(ExtendedExecuteContext);
 	RigVM->Initialize(ExtendedExecuteContext);
 #endif
+
+	ReferencePoseId = FParamId(ReferencePose);
+	CurrentLODId = FParamId(CurrentLOD);
+	DeltaTimeId = FParamId(DeltaTime);
 }
 
 void UAnimNextGraph::PostRename(UObject* OldOuter, const FName OldName)
@@ -198,6 +212,18 @@ void UAnimNextGraph::Serialize(FArchive& Ar)
 	}
 }
 
+TConstArrayView<UE::AnimNext::FScheduleTerm> UAnimNextGraph::GetTerms() const
+{
+	using namespace UE::AnimNext;
+
+	static const FScheduleTerm Terms[] =
+	{
+		FScheduleTerm(Graph::OutputPoseId, FAnimNextParamType::GetType<FAnimNextGraphLODPose>(), EScheduleTermDirection::Output)
+	};
+
+	return Terms;
+}
+
 bool UAnimNextGraph::LoadFromArchiveBuffer(const TArray<uint8>& InSharedDataArchiveBuffer)
 {
 	using namespace UE::AnimNext;
diff --git a/Engine/Plugins/Experimental/AnimNext/Source/AnimNext/Private/Graph/RigUnit_AnimNextGraphEvaluator.cpp b/Engine/Plugins/Experimental/AnimNext/Source/AnimNext/Private/Graph/RigUnit_AnimNextGraphEvaluator.cpp
--- a/Engine/Plugins/Experimental/AnimNext/Source/AnimNext/Private/Graph/RigUnit_AnimNextGraphEvaluator.cpp
+++ b/Engine/Plugins/Experimental/AnimNext/Source/AnimNext/Private/Graph/RigUnit_AnimNextGraphEvaluator.cpp
@@ -6,8 +6,10 @@
 #include "DecoratorInterfaces/IEvaluate.h"
 #include "EvaluationVM/EvaluationVM.h"
 #include "Context.h"
+#include "Graph/AnimNextGraph.h"
 #include "Param/ParamStack.h"
 #include "Graph/AnimNext_LODPose.h"
+#include "AnimNextStats.h"
 
 #include UE_INLINE_GENERATED_CPP_BY_NAME(RigUnit_AnimNextGraphEvaluator)
 
@@ -59,24 +61,31 @@ void FRigUnit_AnimNextGraphEvaluator::StaticExecute(FRigVMExtendedExecuteContext
 		{
 			FParamStack& ParamStack = FParamStack::Get();
 
-			const FAnimNextGraphReferencePose& GraphReferencePose = ParamStack.GetParam<FAnimNextGraphReferencePose>("GraphReferencePose");
-			const int32 GraphLODLevel = ParamStack.GetParam<int32>("GraphLODLevel");
+			const UAnimNextGraph* Graph = VMExecuteContext.GetGraph();
+			const FAnimNextGraphReferencePose* GraphReferencePosePtr = ParamStack.GetParamPtr<FAnimNextGraphReferencePose>(Graph->GetReferencePoseParam());
+			const int32* GraphLODLevelPtr = ParamStack.GetParamPtr<int32>(Graph->GetCurrentLODParam());
+			static FParamId ResultId("UE_Internal_ResultPose");
+			FAnimNextGraphLODPose* ResultPosePtr = ParamStack.GetMutableParamPtr<FAnimNextGraphLODPose>(ResultId);
+			static FParamId ExpectsAdditiveId("UE_Internal_GraphExpectsAdditive");
+			const bool* bExpectsAdditivePtr = ParamStack.GetParamPtr<bool>(ExpectsAdditiveId);
 
-			FEvaluationVM EvaluationVM(EEvaluationFlags::All, *GraphReferencePose.ReferencePose, GraphLODLevel);
-			EvaluationProgram.Execute(EvaluationVM);
-
-			FAnimNextGraphLODPose& ResultPose = ParamStack.GetMutableParam<FAnimNextGraphLODPose>("ResultPose");
-
-			TUniquePtr<FKeyframeState> EvaluatedKeyframe;
-			if (EvaluationVM.PopValue(KEYFRAME_STACK_NAME, EvaluatedKeyframe))
-			{
-				ResultPose.LODPose.CopyFrom(EvaluatedKeyframe->Pose);
-			}
-			else
+			if(GraphReferencePosePtr && GraphLODLevelPtr && ResultPosePtr && bExpectsAdditivePtr)
 			{
-				// We need to output a valid pose, generate one
-				FKeyframeState ReferenceKeyframe = EvaluationVM.MakeReferenceKeyframe(ParamStack.GetParam<bool>("GraphExpectsAdditive"));
-				ResultPose.LODPose.CopyFrom(ReferenceKeyframe.Pose);
+				FEvaluationVM EvaluationVM(EEvaluationFlags::All, *GraphReferencePosePtr->ReferencePose, *GraphLODLevelPtr);
+				EvaluationProgram.Execute(EvaluationVM);
+
+				TUniquePtr<FKeyframeState> EvaluatedKeyframe;
+				if (EvaluationVM.PopValue(KEYFRAME_STACK_NAME, EvaluatedKeyframe))
+				{
+					ResultPosePtr->LODPose.CopyFrom(EvaluatedKeyframe->Pose);
+				}
+				else
+				{
+					// We need to output a valid pose, generate one
+					
+					FKeyframeState ReferenceKeyframe = EvaluationVM.MakeReferenceKeyframe(*bExpectsAdditivePtr);
+					ResultPosePtr->LODPose.CopyFrom(ReferenceKeyframe.Pose);
+				}
 			}
 		}
 	}
diff --git a/Engine/Plugins/Experimental/AnimNext/Source/AnimNext/Private/Param/AnimNextParameterBlock.cpp b/Engine/Plugins/Experimental/AnimNext/Source/AnimNext/Private/Param/AnimNextParameterBlock.cpp
--- a/Engine/Plugins/Experimental/AnimNext/Source/AnimNext/Private/Param/AnimNextParameterBlock.cpp
+++ b/Engine/Plugins/Experimental/AnimNext/Source/AnimNext/Private/Param/AnimNextParameterBlock.cpp
@@ -7,6 +7,9 @@
 #include "UObject/Package.h"
 #include "Graph/RigUnit_AnimNextBeginExecution.h"
 #include "Param/AnimNextParameterExecuteContext.h"
+#include "AnimNextStats.h"
+
+DEFINE_STAT(STAT_AnimNext_ParamBlock_UpdateLayer);
 
 #include UE_INLINE_GENERATED_CPP_BY_NAME(AnimNextParameterBlock)
 
@@ -40,6 +43,8 @@ UAnimNextParameterBlock::UAnimNextParameterBlock(const FObjectInitializer& Objec
 
 void UAnimNextParameterBlock::UpdateLayer(UE::AnimNext::FParamStackLayerHandle& InHandle) const
 {
+	SCOPE_CYCLE_COUNTER(STAT_AnimNext_ParamBlock_UpdateLayer);
+	
 	if (VM)
 	{
 		UE::AnimNext::FRigVMRuntimeData* RuntimeData = UE::AnimNext::FRigVMRuntimeDataRegistry::FindRuntimeData(VM);
diff --git a/Engine/Plugins/Experimental/AnimNext/Source/AnimNext/Private/Param/ParamAdapter.cpp b/Engine/Plugins/Experimental/AnimNext/Source/AnimNext/Private/Param/ParamAdapter.cpp
--- a/Engine/Plugins/Experimental/AnimNext/Source/AnimNext/Private/Param/ParamAdapter.cpp
+++ b/Engine/Plugins/Experimental/AnimNext/Source/AnimNext/Private/Param/ParamAdapter.cpp
@@ -19,7 +19,7 @@ FParamResult FParamAdapter::GetParamData(FParamTypeHandle InTypeHandle, TConstAr
 		return EParamResult::TypeError;
 	}
 
-	UObject* ContextObject = FScheduleContext::Get().Entry->Object.Get();
+	UObject* ContextObject = FScheduleContext::Get().Entry->ResolvedObject;
 	if(const uint8* Data = Function(ContextObject, Definition.Id))
 	{
 		OutParamData = TConstArrayView<uint8>(Data, Definition.TypeHandle.GetSize());
@@ -52,7 +52,7 @@ FParamResult FParamAdapter::GetMutableParamData(FParamTypeHandle InTypeHandle, T
 		return EParamResult::TypeError;
 	}
 
-	UObject* ContextObject = FScheduleContext::Get().Entry->Object.Get();
+	UObject* ContextObject = FScheduleContext::Get().Entry->ResolvedObject;
 	if(uint8* Data = Function(ContextObject, Definition.Id))
 	{
 		OutParamData = TArrayView<uint8>(Data, Definition.TypeHandle.GetSize());
diff --git a/Engine/Plugins/Experimental/AnimNext/Source/AnimNext/Private/Param/ParamHelpers.cpp b/Engine/Plugins/Experimental/AnimNext/Source/AnimNext/Private/Param/ParamHelpers.cpp
--- a/Engine/Plugins/Experimental/AnimNext/Source/AnimNext/Private/Param/ParamHelpers.cpp
+++ b/Engine/Plugins/Experimental/AnimNext/Source/AnimNext/Private/Param/ParamHelpers.cpp
@@ -4,6 +4,7 @@
 #include "Param/ParamType.h"
 #include "Param/ParamTypeHandle.h"
 #include "Templates/SubclassOf.h"
+#include "Graph/AnimNext_LODPose.h"
 
 namespace UE::AnimNext
 {
@@ -314,6 +315,18 @@ void FParamHelpers::Copy(const FParamTypeHandle& InTypeHandle, TConstArrayView<u
 	case FParamTypeHandle::EParamType::Transform:
 		SimpleCopy(sizeof(FTransform), alignof(FTransform));
 		break;
+	case FParamTypeHandle::EParamType::Object:
+	case FParamTypeHandle::EParamType::CharacterMovementComponent:
+	case FParamTypeHandle::EParamType::AnimNextMeshComponent:
+	case FParamTypeHandle::EParamType::AnimSequence:
+		SimpleCopy(sizeof(UObject*), alignof(UObject*));
+		break;
+	case FParamTypeHandle::EParamType::AnimNextGraphLODPose:
+		SimpleCopy(sizeof(FAnimNextGraphLODPose), alignof(FAnimNextGraphLODPose));
+		break;
+	case FParamTypeHandle::EParamType::AnimNextGraphReferencePose:
+		SimpleCopy(sizeof(FAnimNextGraphReferencePose), alignof(FAnimNextGraphReferencePose));
+		break;
 	case FParamTypeHandle::EParamType::Custom:
 		{
 			FAnimNextParamType Type = InTypeHandle.GetType();
@@ -498,6 +511,12 @@ void FParamHelpers::Destroy(const FParamTypeHandle& InTypeHandle, TArrayView<uin
 	case FParamTypeHandle::EParamType::Vector4:
 	case FParamTypeHandle::EParamType::Quat:
 	case FParamTypeHandle::EParamType::Transform:
+	case FParamTypeHandle::EParamType::Object:
+	case FParamTypeHandle::EParamType::CharacterMovementComponent:
+	case FParamTypeHandle::EParamType::AnimNextMeshComponent:
+	case FParamTypeHandle::EParamType::AnimSequence:
+	case FParamTypeHandle::EParamType::AnimNextGraphLODPose:
+	case FParamTypeHandle::EParamType::AnimNextGraphReferencePose:
 		break;
 	case FParamTypeHandle::EParamType::Custom:
 		{
diff --git a/Engine/Plugins/Experimental/AnimNext/Source/AnimNext/Private/Param/ParamId.cpp b/Engine/Plugins/Experimental/AnimNext/Source/AnimNext/Private/Param/ParamId.cpp
--- a/Engine/Plugins/Experimental/AnimNext/Source/AnimNext/Private/Param/ParamId.cpp
+++ b/Engine/Plugins/Experimental/AnimNext/Source/AnimNext/Private/Param/ParamId.cpp
@@ -3,6 +3,8 @@
 #include "Param/ParamId.h"
 
 #include "AnimNextConfig.h"
+#include "DataRegistry.h"
+#include "ReferencePose.h"
 #include "Misc/ScopeRWLock.h"
 #include "HAL/ThreadSingleton.h"
 #include "Param/AnimNextObjectAdapterConfig.h"
@@ -14,6 +16,11 @@
 #include "Param/ParamStack.h"
 #include "RigVMCore/RigVMRegistry.h"
 #include "Engine/World.h"
+#include "Graph/AnimNext_LODPose.h"
+#include "Component/AnimNextMeshComponent.h"
+#include "AnimNextStats.h"
+
+DEFINE_STAT(STAT_AnimNext_ParamIdLock)
 
 #define LOCTEXT_NAMESPACE "ParamId"
 
@@ -22,6 +29,16 @@ namespace UE::AnimNext
 
 static FRWLock GParamIdLock;
 
+// TODO: Expand scratch buffer system to become a general cache of adapter values.
+// The assumption being that if we are grabbing gameplay-driven values, they do not vary over 
+// the course of a schedule unless the 'tick' of the values (e.g. a component's tick function) runs
+// during that time, at which point the values are reset and lazily updated
+// We should cache these values in a FInstancedPropertyBag, per adapted object type (e.g. all 
+// adapters for skel mesh components), per thread. This also allows us better handling of complex 
+// types, object lifetimes and GC.
+
+// TODO: Investigate allowing function libraries as adapters
+
 struct FParamData
 {
 	FParamData(FName InName)
@@ -69,8 +86,7 @@ static FParamIdGlobalData& GetParamIdData()
 
 void FParamId::Init()
 {
-	RegisterBuiltInAdapters();
-	RefreshConfigAdapters();
+	RefreshAdapters();
 }
 
 void FParamId::Destroy()
@@ -82,6 +98,31 @@ void FParamId::Destroy()
 	ParamIdData.NameToParamId.Empty();
 }
 
+void FParamId::RefreshAdapters()
+{
+	ResetAdapters();
+	RegisterBuiltInAdapters();
+	RefreshConfigAdapters();
+}
+
+void FParamId::ResetAdapters()
+{
+	FParamIdGlobalData& ParamIdGlobalData = GetParamIdData();
+	
+	ParamIdGlobalData.ResultBufferSize = 0;
+	ParamIdGlobalData.ParamAdapters.Empty();
+
+	// Reset existing IDs adapter data
+	for(FParamData& ParamData : ParamIdGlobalData.ParamData)
+	{
+		ParamData.AdapterIndex = INDEX_NONE;
+		ParamData.BufferOffset = MAX_uint32;
+
+		// Unregister this param as built-in
+		FParams::UnregisterBuiltInParameter(ParamData.Name);
+	}
+}
+
 void FParamId::RegisterBuiltInAdapters()
 {
 	FRWScopeLock Lock(GParamIdLock, SLT_Write);
@@ -105,6 +146,37 @@ void FParamId::RegisterBuiltInAdapters()
 		FParamAdapter& NewAdapter = ParamIdGlobalData.ParamAdapters.Emplace_GetRef(FParamDefinition(PhysicsTickFunctionIndex, PhysicsTickName, FAnimNextParamType::GetType<FTickFunction>(), TooltipText), MoveTemp(PhysicsTickFunctionFunction));
 		FParams::RegisterBuiltInParameter(NewAdapter.Definition);
 	}
+
+	// Register ref pose accessor
+	{
+		FObjectAdapterFunction ReferencePoseFunction = [](UObject* InContextObject, FParamId InId) -> uint8*
+		{
+			static FParamId ComponentParam("UE_AnimNextMeshComponent");
+			if (const TObjectPtr<UAnimNextMeshComponent>* SkeletalMeshComponent = FParamStack::Get().GetParamPtr<TObjectPtr<UAnimNextMeshComponent>>(ComponentParam))
+			{
+				if(*SkeletalMeshComponent)
+				{
+					FDataHandle RefPoseHandle = FDataRegistry::Get()->GetOrGenerateReferencePose(SkeletalMeshComponent->Get());
+					const FReferencePose& RefPose = RefPoseHandle.GetRef<FReferencePose>();
+					FAnimNextGraphReferencePose* ReturnValue = GetAdapterReturnValue<FAnimNextGraphReferencePose>(InId);
+					*ReturnValue = FAnimNextGraphReferencePose(&RefPose);
+					return reinterpret_cast<uint8*>(ReturnValue);
+				}
+			}
+			return nullptr;
+		};
+
+		FName ReferencePoseName("UE_AnimNextMeshComponent_ReferencePose");
+		const uint32 ReferencePoseIndex = MakeParamId_NoLock(ReferencePoseName);
+		ParamIdGlobalData.ParamData[ReferencePoseIndex].AdapterIndex = ParamIdGlobalData.ParamAdapters.Num();
+		FText TooltipText = LOCTEXT("ReferencePoseTooltip", "The reference pose of the skeletal mesh component");
+		FParamAdapter& NewAdapter = ParamIdGlobalData.ParamAdapters.Emplace_GetRef(FParamDefinition(ReferencePoseIndex, ReferencePoseName, FAnimNextParamType::GetType<FAnimNextGraphReferencePose>(), TooltipText), MoveTemp(ReferencePoseFunction));
+		FParams::RegisterBuiltInParameter(NewAdapter.Definition);
+
+		// Calc buffer size for the result
+		ParamIdGlobalData.ParamData[ReferencePoseIndex].BufferOffset = Align(ParamIdGlobalData.ResultBufferSize, alignof(FAnimNextGraphReferencePose));
+		ParamIdGlobalData.ResultBufferSize += sizeof(FAnimNextGraphReferencePose);
+	}
 }
 
 void FParamId::RefreshConfigAdapters()
@@ -112,19 +184,6 @@ void FParamId::RefreshConfigAdapters()
 	FRWScopeLock Lock(GParamIdLock, SLT_Write);
 	FParamIdGlobalData& ParamIdGlobalData = GetParamIdData();
 
-	ParamIdGlobalData.ResultBufferSize = 0;
-	ParamIdGlobalData.ParamAdapters.Empty();
-
-	// Reset existing IDs adapter data
-	for(FParamData& ParamData : ParamIdGlobalData.ParamData)
-	{
-		ParamData.AdapterIndex = INDEX_NONE;
-		ParamData.BufferOffset = MAX_uint32;
-
-		// Unregister this param as built-in
-		FParams::UnregisterBuiltInParameter(ParamData.Name);
-	}
-
 	for(const FAnimNextObjectAdapterConfig& AdapterConfig : GetDefault<UAnimNextConfig>()->ExposedClasses)
 	{
 		if(AdapterConfig.Class != nullptr && AdapterConfig.RootParameter != NAME_None)
@@ -149,25 +208,10 @@ void FParamId::RefreshConfigAdapters()
 					{
 						if (AActor* Actor = ContextComponent->GetOwner())
 						{
-							// Components can be redirected via their name
-							UActorComponent* ActorComponent = nullptr;
-							if(const FName* ComponentName = FParamStack::Get().GetParamPtr<FName>(FParamId(NameParamIndex)))
-							{
-								const FObjectPropertyBase* ComponentProperty = FindFProperty<FObjectPropertyBase>(Actor->GetClass(), *ComponentName);
-								if(ComponentProperty != nullptr)
-								{
-									ActorComponent = Cast<UActorComponent>(ComponentProperty->GetObjectPropertyValue_InContainer(Actor));
-								}
-							}
-							else
-							{
-								// No name, just find the first component
-								ActorComponent = Actor->FindComponentByClass(Class.Get());
-							}
-
+							UActorComponent* ActorComponent = Actor->FindComponentByClass(Class.Get());
 							UActorComponent** ReturnValue = GetAdapterReturnValue<UActorComponent*>(InId);
 							*ReturnValue = ActorComponent;
-							return reinterpret_cast<uint8*>(*ReturnValue);
+							return reinterpret_cast<uint8*>(ReturnValue);
 						}
 					}
 
@@ -187,7 +231,7 @@ void FParamId::RefreshConfigAdapters()
 							{
 								AActor** ReturnValue = GetAdapterReturnValue<AActor*>(InId);
 								*ReturnValue = Actor;
-								return reinterpret_cast<uint8*>(*ReturnValue);
+								return reinterpret_cast<uint8*>(ReturnValue);
 							}
 						}
 					}
@@ -196,6 +240,11 @@ void FParamId::RefreshConfigAdapters()
 				};
 			}
 
+			// TODO: Add runtime mechanism to prevent accessing object data:
+			// - When an object's tick function is not bound in a schedule (and when we have non-linear schedules,
+			//   cannot be concurrent)
+			// - When an object's tick function is not concurrently running (if necessary - should be mitigated by the above)
+			
 			// Add root param
 			ParamIdGlobalData.ParamData[RootParamIndex].AdapterIndex = ParamIdGlobalData.ParamAdapters.Num(); 
 			FParamAdapter& NewRootAdapter = ParamIdGlobalData.ParamAdapters.Emplace_GetRef(FParamDefinition(RootParamIndex, RootParamName, AdapterConfig.Class->GetDefaultObject()), MoveTemp(RootFunction));
@@ -235,10 +284,10 @@ void FParamId::RefreshConfigAdapters()
 							FObjectAdapterFunction AdapterFunction = [RootParamIndex, Property, Class = AdapterConfig.Class.Get()](UObject* InContextObject, FParamId InId) -> uint8*
 							{
 								// Grab object container
-								if(const UObject* Object = FParamStack::Get().GetParamPtr<UObject>(FParamId(RootParamIndex)))
+								if(const TObjectPtr<UObject>* Object = FParamStack::Get().GetParamPtr<TObjectPtr<UObject>>(FParamId(RootParamIndex)))
 								{
 									check(Class->IsChildOf(Property->GetOwnerClass()));
-									return const_cast<uint8*>(Property->ContainerPtrToValuePtr<uint8>(Object));
+									return const_cast<uint8*>(Property->ContainerPtrToValuePtr<uint8>(*Object));
 								}
 								return nullptr;
 							};
@@ -263,7 +312,7 @@ void FParamId::RefreshConfigAdapters()
 							(AdapterConfig.FilterType == EAnimNextObjectAdapterFilterType::AllowList && AdapterConfig.FilteredFields.Contains(InName)) ||
 							(AdapterConfig.FilterType == EAnimNextObjectAdapterFilterType::DenyList && !AdapterConfig.FilteredFields.Contains(InName));
 				};
-				
+
 				for(TFieldIterator<UFunction> It(AdapterConfig.Class, EFieldIterationFlags::IncludeSuper | EFieldIterationFlags::IncludeInterfaces); It; ++It)
 				{
 					UFunction* Function = *It;
@@ -284,22 +333,26 @@ void FParamId::RefreshConfigAdapters()
 								bNameParamAdded = true;
 							}
 
-							FObjectAdapterFunction AdapterFunction = [RootParamIndex, WeakFunction = TWeakObjectPtr<UFunction>(Function), Class = AdapterConfig.Class.Get()](UObject* InContextObject, FParamId InId)
+							FObjectAdapterFunction AdapterFunction = [RootParamIndex, WeakFunction = TWeakObjectPtr<UFunction>(Function), Class = AdapterConfig.Class.Get()](UObject* InContextObject, FParamId InId) -> uint8*
 							{
 								// Grab object to call the function with
 								if(UFunction* Function = WeakFunction.Get())
 								{
-									if(UObject* Object = FParamStack::Get().GetMutableParamPtr<UObject>(FParamId(RootParamIndex)))
+									if(const TObjectPtr<UObject>* ObjectPtr = FParamStack::Get().GetParamPtr<TObjectPtr<UObject>>(FParamId(RootParamIndex)))
 									{
-										check(Object->GetClass()->IsChildOf(Class));
-										check(Function->GetOuterUClass()->IsChildOf(Class));
-										const FProperty* ReturnProperty = Function->GetReturnProperty();
-										check(ReturnProperty);
-
-										uint8* ReturnValuePtr = GetScratchAreaForParamIdAdapter(InId);
-										ReturnProperty->InitializeValue(ReturnValuePtr);
-										FFrame Stack(Object, Function, nullptr, nullptr, Function->ChildProperties);
-										Function->Invoke(Object, Stack, ReturnValuePtr);
+										if(UObject* Object = *ObjectPtr)
+										{
+											check(Object->GetClass()->IsChildOf(Class));
+											check(Class->IsChildOf(Function->GetOuterUClass()));
+											const FProperty* ReturnProperty = Function->GetReturnProperty();
+											check(ReturnProperty);
+
+											uint8* ReturnValuePtr = GetScratchAreaForParamIdAdapter(InId);
+											ReturnProperty->InitializeValue(ReturnValuePtr);
+											FFrame Stack(Object, Function, nullptr, nullptr, Function->ChildProperties);
+											Function->Invoke(Object, Stack, ReturnValuePtr);
+											return ReturnValuePtr;
+										}
 									}
 								}
 								return nullptr;
@@ -341,9 +394,12 @@ void FParamId::RefreshConfigAdapters()
 					ParamType = FAnimNextParamType::GetType<FActorComponentTickFunction>();
 					TickFunctionFunction = [RootParamIndex](UObject* InContextObject, FParamId InId) -> uint8*
 					{
-						if (const UActorComponent* Component = FParamStack::Get().GetParamPtr<UActorComponent>(FParamId(RootParamIndex)))
+						if (const TObjectPtr<UActorComponent>* Component = FParamStack::Get().GetParamPtr<TObjectPtr<UActorComponent>>(FParamId(RootParamIndex)))
 						{
-							return reinterpret_cast<uint8*>(const_cast<FActorComponentTickFunction*>(&Component->PrimaryComponentTick));
+							if(Component->Get())
+							{
+								return reinterpret_cast<uint8*>(&Component->Get()->PrimaryComponentTick);
+							}
 						}
 						return nullptr;
 					};
@@ -353,9 +409,12 @@ void FParamId::RefreshConfigAdapters()
 					ParamType = FAnimNextParamType::GetType<FActorTickFunction>();
 					TickFunctionFunction = [RootParamIndex](UObject* InContextObject, FParamId InId) -> uint8*
 					{
-						if (const AActor* Actor = FParamStack::Get().GetParamPtr<AActor>(FParamId(RootParamIndex)))
+						if (const TObjectPtr<AActor>* Actor = FParamStack::Get().GetParamPtr<TObjectPtr<AActor>>(FParamId(RootParamIndex)))
 						{
-							return reinterpret_cast<uint8*>(const_cast<FActorTickFunction*>(&Actor->PrimaryActorTick));
+							if(Actor->Get())
+							{
+								return reinterpret_cast<uint8*>(&Actor->Get()->PrimaryActorTick);
+							}
 						}
 						return nullptr;
 					};
@@ -387,6 +446,8 @@ void FParamId::RefreshConfigAdapters()
 
 FParamId::FParamId(FName InName)
 {
+	SCOPE_CYCLE_COUNTER(STAT_AnimNext_ParamIdLock);
+
 	FRWScopeLock Lock(GParamIdLock, SLT_Write);
 	ParameterIndex = MakeParamId_NoLock(InName);
 }
diff --git a/Engine/Plugins/Experimental/AnimNext/Source/AnimNext/Private/Param/ParamStack.cpp b/Engine/Plugins/Experimental/AnimNext/Source/AnimNext/Private/Param/ParamStack.cpp
--- a/Engine/Plugins/Experimental/AnimNext/Source/AnimNext/Private/Param/ParamStack.cpp
+++ b/Engine/Plugins/Experimental/AnimNext/Source/AnimNext/Private/Param/ParamStack.cpp
@@ -1,6 +1,7 @@
 // Copyright Epic Games, Inc. All Rights Reserved.
 
 #include "Param/ParamStack.h"
+
 #include "Param/ParamHelpers.h"
 #include "PropertyBag.h"
 #include "EngineLogs.h"
@@ -10,6 +11,9 @@
 #include "Param/ParamStackLayer.h"
 #include "UObject/ObjectKey.h"
 
+DEFINE_STAT(STAT_AnimNext_ParamStack_GetParam);
+DEFINE_STAT(STAT_AnimNext_ParamStack_Adapter);
+
 #define LOCTEXT_NAMESPACE "AnimNextParamStack"
 
 namespace UE::AnimNext
@@ -166,6 +170,49 @@ struct FInstancedPropertyBagReferenceLayer : FInstancedPropertyBagLayer
 	FInstancedPropertyBag& PropertyBag;
 };
 
+// Stack layer that remaps params from another layer
+struct FRemappedLayer : FParamStackLayer
+{
+	FRemappedLayer() = delete;
+
+	explicit FRemappedLayer(const FParamStackLayer& InLayer, const TMap<FName, FName>& InMapping)
+	{
+		// Determine param ID range for this layer
+		TArray<TPair<FParamId, Private::FParamEntry>, TInlineAllocator<128>> CachedParams;
+		CachedParams.Reserve(InMapping.Num());
+		MinParamId = MAX_uint32;
+		uint32 MaxParamId = 0;
+		for (uint32 LocalParamIndex = 0; LocalParamIndex < static_cast<uint32>(InLayer.Params.Num()); ++LocalParamIndex)
+		{
+			const Private::FParamEntry& ParamEntry = InLayer.Params[LocalParamIndex];
+			if(ParamEntry.IsValid())
+			{
+				const uint32 GlobalParamIndex = InLayer.MinParamId + LocalParamIndex;
+				const FParamId SourceParamId(GlobalParamIndex);
+				if(const FName* RemappedName = InMapping.Find(SourceParamId.ToName()))
+				{
+					const FParamId RemappedParamId(*RemappedName);
+					MinParamId = FMath::Min(RemappedParamId.ToInt(), MinParamId);
+					MaxParamId = FMath::Max(RemappedParamId.ToInt(), MaxParamId);
+					CachedParams.Add({ RemappedParamId, ParamEntry });
+				}
+			}
+		}
+
+		if (MinParamId <= MaxParamId)
+		{
+			const uint32 ParamRangeSize = (MaxParamId - MinParamId) + 1;
+			Params.SetNumZeroed(ParamRangeSize);
+			for (uint32 ParamIndex = 0; ParamIndex < static_cast<uint32>(CachedParams.Num()); ++ParamIndex)
+			{
+				const TPair<FParamId, Private::FParamEntry>& ParamIdPair = CachedParams[ParamIndex];
+				const uint32 LocalParamIndex = ParamIdPair.Key.ToInt() - MinParamId;
+				Params[LocalParamIndex] = ParamIdPair.Value;
+			}
+		}
+	}
+};
+
 FParamStack::FPushedLayer::FPushedLayer(FParamStackLayer& InLayer, FParamStack& InStack)
 	: Layer(InLayer)
 {
@@ -222,6 +269,7 @@ FParamStack::FParamStack()
 	Layers.Reserve(8);
 	LayerIndices.Reserve(FParamId::GetMaxParamId().ToInt());
 	PreviousLayerIndices.Reserve(FParamId::GetMaxParamId().ToInt());
+	OwnedStackLayers.Reserve(8);
 }
 
 FParamStack::~FParamStack()
@@ -303,7 +351,7 @@ FParamStack::FPushedLayerHandle FParamStack::PushLayer(TConstArrayView<TPair<FPa
 {
 	if(Layers.Num() < MAX_uint16)
 	{
-		FParamStackLayer& OwnedLayer = OwnedStackLayers.Add_GetRef(FParamStackLayer(InParams));
+		FParamStackLayer& OwnedLayer = OwnedStackLayers.Emplace_GetRef(InParams);
 		OwnedLayer.OwnedStorageOffset = AllocAndCopyOwnedParamStorage(OwnedLayer.Params);
 		return PushLayerInternal(OwnedLayer);
 	}
@@ -420,6 +468,12 @@ FParamStackLayerHandle FParamStack::MakeReferenceLayer(FInstancedPropertyBag& In
 	return FParamStackLayerHandle(MoveTemp(Layer));
 }
 
+FParamStackLayerHandle FParamStack::MakeRemappedLayer(const FParamStackLayerHandle& InLayer, const TMap<FName, FName>& InMapping)
+{
+	TUniquePtr<FParamStackLayer> Layer = MakeUnique<FRemappedLayer>(*InLayer.Layer.Get(), InMapping);
+	return FParamStackLayerHandle(MoveTemp(Layer));
+}
+
 FParamStackLayerHandle FParamStack::MakeLayer(TConstArrayView<TPair<FParamId, Private::FParamEntry>> InParams)
 {
 	TUniquePtr<FParamStackLayer> Layer = MakeUnique<FParamStackLayer>(InParams);
diff --git a/Engine/Plugins/Experimental/AnimNext/Source/AnimNext/Private/Param/ParamStackLayer.cpp b/Engine/Plugins/Experimental/AnimNext/Source/AnimNext/Private/Param/ParamStackLayer.cpp
--- a/Engine/Plugins/Experimental/AnimNext/Source/AnimNext/Private/Param/ParamStackLayer.cpp
+++ b/Engine/Plugins/Experimental/AnimNext/Source/AnimNext/Private/Param/ParamStackLayer.cpp
@@ -56,6 +56,11 @@ FParamResult FParamStackLayer::GetParamData(FParamId InId, FParamTypeHandle InTy
 {
 	const uint32 LocalParamIndex = InId.ToInt() - MinParamId;
 
+	if(!Params.IsValidIndex(LocalParamIndex))
+	{
+		return EParamResult::NotInScope;
+	}
+
 	const Private::FParamEntry& Param = Params[LocalParamIndex];
 	if (!Param.IsValid())
 	{
diff --git a/Engine/Plugins/Experimental/AnimNext/Source/AnimNext/Private/Param/ParamStackLayerHandle.cpp b/Engine/Plugins/Experimental/AnimNext/Source/AnimNext/Private/Param/ParamStackLayerHandle.cpp
--- a/Engine/Plugins/Experimental/AnimNext/Source/AnimNext/Private/Param/ParamStackLayerHandle.cpp
+++ b/Engine/Plugins/Experimental/AnimNext/Source/AnimNext/Private/Param/ParamStackLayerHandle.cpp
@@ -98,4 +98,14 @@ FInstancedPropertyBag* FParamStackLayerHandle::GetInstancedPropertyBagFromLayer(
 	return Layer->AsInstancedPropertyBag();
 }
 
+FParamResult FParamStackLayerHandle::GetParamData(FParamId InId, FParamTypeHandle InTypeHandle, TConstArrayView<uint8>& OutParamData, FParamTypeHandle& OutParamTypeHandle, FParamCompatibility InRequiredCompatibility) const
+{
+	return Layer->GetParamData(InId, InTypeHandle, OutParamData, OutParamTypeHandle, InRequiredCompatibility);
+}
+
+FParamResult FParamStackLayerHandle::GetMutableParamData(FParamId InId, FParamTypeHandle InTypeHandle, TArrayView<uint8>& OutParamData, FParamTypeHandle& OutParamTypeHandle, FParamCompatibility InRequiredCompatibility) const
+{
+	return Layer->GetMutableParamData(InId, InTypeHandle, OutParamData, OutParamTypeHandle, InRequiredCompatibility);
+}
+
 }
\ No newline at end of file
diff --git a/Engine/Plugins/Experimental/AnimNext/Source/AnimNext/Private/Param/ParamType.cpp b/Engine/Plugins/Experimental/AnimNext/Source/AnimNext/Private/Param/ParamType.cpp
--- a/Engine/Plugins/Experimental/AnimNext/Source/AnimNext/Private/Param/ParamType.cpp
+++ b/Engine/Plugins/Experimental/AnimNext/Source/AnimNext/Private/Param/ParamType.cpp
@@ -1,10 +1,16 @@
 // Copyright Epic Games, Inc. All Rights Reserved.
 
 #include "Param/ParamType.h"
+
+#include "Component/AnimNextMeshComponent.h"
+#include "Graph/AnimNext_LODPose.h"
 #include "Param/ParamTypeHandle.h"
 #include "Misc/StringBuilder.h"
 #include "UObject/Class.h"
 #include "Templates/SubclassOf.h"
+#include "GameFramework/CharacterMovementComponent.h"
+#include "Components/SkeletalMeshComponent.h"
+#include "Animation/AnimSequence.h"
 
 FAnimNextParamType::FAnimNextParamType(EValueType InValueType, EContainerType InContainerType, const UObject* InValueTypeObject)
 	: ValueTypeObject(InValueTypeObject)
@@ -111,6 +117,14 @@ UE::AnimNext::FParamTypeHandle FAnimNextParamType::GetHandle() const
 				{
 					Handle.SetParameterType(FParamTypeHandle::EParamType::Transform);
 				}
+				else if(ScriptStruct == FAnimNextGraphLODPose::StaticStruct())
+				{
+					Handle.SetParameterType(FParamTypeHandle::EParamType::AnimNextGraphLODPose);
+				}
+				else if(ScriptStruct == FAnimNextGraphReferencePose::StaticStruct())
+				{
+					Handle.SetParameterType(FParamTypeHandle::EParamType::AnimNextGraphReferencePose);
+				}
 				else
 				{
 					Handle.SetParameterType(FParamTypeHandle::EParamType::Custom);
@@ -134,6 +148,30 @@ UE::AnimNext::FParamTypeHandle FAnimNextParamType::GetHandle() const
 			}
 			break;
 		case EValueType::Object:
+			if(const UClass* Class = Cast<UClass>(ValueTypeObject.Get()))
+			{
+				if (Class == UObject::StaticClass())
+				{
+					Handle.SetParameterType(FParamTypeHandle::EParamType::Object);
+					break;
+				}
+				else if (Class == UCharacterMovementComponent::StaticClass())
+				{
+					Handle.SetParameterType(FParamTypeHandle::EParamType::CharacterMovementComponent);
+					break;
+				}
+				else if (Class == UAnimNextMeshComponent::StaticClass())
+				{
+					Handle.SetParameterType(FParamTypeHandle::EParamType::AnimNextMeshComponent);
+					break;
+				}
+				else if (Class == UAnimSequence::StaticClass())
+				{
+					Handle.SetParameterType(FParamTypeHandle::EParamType::AnimSequence);
+					break;
+				}
+			}
+			// fall through
 		case EValueType::SoftObject:
 		case EValueType::Class:
 		case EValueType::SoftClass:
diff --git a/Engine/Plugins/Experimental/AnimNext/Source/AnimNext/Private/Param/ParamTypeHandle.cpp b/Engine/Plugins/Experimental/AnimNext/Source/AnimNext/Private/Param/ParamTypeHandle.cpp
--- a/Engine/Plugins/Experimental/AnimNext/Source/AnimNext/Private/Param/ParamTypeHandle.cpp
+++ b/Engine/Plugins/Experimental/AnimNext/Source/AnimNext/Private/Param/ParamTypeHandle.cpp
@@ -5,6 +5,14 @@
 #include "Misc/ScopeRWLock.h"
 #include "UObject/Class.h"
 #include "UObject/TextProperty.h"
+#include "AnimNextStats.h"
+#include "Component/AnimNextMeshComponent.h"
+#include "GameFramework/CharacterMovementComponent.h"
+#include "Components/SkeletalMeshComponent.h"
+#include "Animation/AnimSequence.h"
+#include "Graph/AnimNext_LODPose.h"
+
+DEFINE_STAT(STAT_AnimNext_ParamTypeHandle_Lock);
 
 namespace UE::AnimNext
 {
@@ -71,24 +79,35 @@ uint32 FParamTypeHandle::GetOrAllocateCustomTypeIndex(FAnimNextParamType::EValue
 
 	const uint32 Hash = GetTypeHash(ParameterType);
 
-	// NOTE: It is tempting to try to acquire a only read lock here during the FindByHash, but doing so means that the
-	// lazy-insert operation is no longer atomic as there is no way to upgrade a read lock into a write lock with
-	// FRWScopeLock (@see FRWScopeLock comments), hence the write lock around the map/array access here.
 	{
-		// See if the type already exists in the map
-		FRWScopeLock ScopeLock(GParamTypeHandleLock, SLT_Write);
-
-		if(const uint32* IndexPtr = GetTypeHandleData().TypeToIndexMap.FindByHash(Hash, ParameterType))
 		{
-			return *IndexPtr + 1;
+			// See if the type already exists in the map
+			FRWScopeLock ScopeLock(GParamTypeHandleLock, SLT_ReadOnly);
+
+			if(const uint32* IndexPtr = GetTypeHandleData().TypeToIndexMap.FindByHash(Hash, ParameterType))
+			{
+				return *IndexPtr + 1;
+			}
 		}
 
-		// Add a new custom type
-		const uint32 Index = GetTypeHandleData().CustomTypes.Add(ParameterType);
-		GetTypeHandleData().TypeToIndexMap.AddByHash(Hash, ParameterType, Index);
+		{
+			SCOPE_CYCLE_COUNTER(STAT_AnimNext_ParamTypeHandle_Lock);
+
+			FRWScopeLock ScopeLock(GParamTypeHandleLock, SLT_Write);
+
+			// See if the type already exists in the map AGAIN, in case another thread grabbed an index outside of the locks
+			if(const uint32* IndexPtr = GetTypeHandleData().TypeToIndexMap.FindByHash(Hash, ParameterType))
+			{
+				return *IndexPtr + 1;
+			}
+
+			// Add a new custom type
+			const uint32 Index = GetTypeHandleData().CustomTypes.Add(ParameterType);
+			GetTypeHandleData().TypeToIndexMap.AddByHash(Hash, ParameterType, Index);
 
-		checkf((Index + 1) < (1 << 24), TEXT("FParamTypeHandle::GetCustomTypeIndex: Type index overflowed"));
-		return Index + 1;
+			checkf((Index + 1) < (1 << 24), TEXT("FParamTypeHandle::GetCustomTypeIndex: Type index overflowed"));
+			return Index + 1;
+		}
 	}
 }
 
@@ -172,6 +191,14 @@ FParamTypeHandle FParamTypeHandle::FromPropertyBagPropertyDesc(const FPropertyBa
 					{
 						Handle.SetParameterType(EParamType::Transform);
 					}
+					else if(ScriptStruct == FAnimNextGraphLODPose::StaticStruct())
+					{
+						Handle.SetParameterType(EParamType::AnimNextGraphLODPose);
+					}
+					else if(ScriptStruct == FAnimNextGraphReferencePose::StaticStruct())
+					{
+						Handle.SetParameterType(EParamType::AnimNextGraphReferencePose);
+					}
 					else
 					{
 						Handle.SetParameterType(EParamType::Custom);
@@ -195,6 +222,30 @@ FParamTypeHandle FParamTypeHandle::FromPropertyBagPropertyDesc(const FPropertyBa
 				}
 				break;
 			case EPropertyBagPropertyType::Object:
+				if(const UClass* Class = Cast<UClass>(Desc.ValueTypeObject.Get()))
+				{
+					if (Class == UObject::StaticClass())
+					{
+						Handle.SetParameterType(EParamType::Object);
+						break;
+					}
+					else if (Class == UCharacterMovementComponent::StaticClass())
+					{
+						Handle.SetParameterType(EParamType::CharacterMovementComponent);
+						break;
+					}
+					else if (Class == UAnimNextMeshComponent::StaticClass())
+					{
+						Handle.SetParameterType(EParamType::AnimNextMeshComponent);
+						break;
+					}
+					else if (Class == UAnimSequence::StaticClass())
+					{
+						Handle.SetParameterType(EParamType::AnimSequence);
+						break;
+					}
+				}
+				// fall through
 			case EPropertyBagPropertyType::SoftObject:
 			case EPropertyBagPropertyType::Class:
 			case EPropertyBagPropertyType::SoftClass:
@@ -317,7 +368,7 @@ FParamTypeHandle FParamTypeHandle::FromProperty(const FProperty* InProperty)
 	}
 	else if (InProperty->IsA<FStructProperty>())
 	{
-		UScriptStruct* ScriptStruct = CastField<FStructProperty>(InProperty)->Struct;
+		const UScriptStruct* ScriptStruct = CastField<FStructProperty>(InProperty)->Struct;
 		if (ScriptStruct == TBaseStructure<FVector>::Get())
 		{
 			Handle.SetParameterType(EParamType::Vector);
@@ -334,6 +385,30 @@ FParamTypeHandle FParamTypeHandle::FromProperty(const FProperty* InProperty)
 		{
 			Handle.SetParameterType(EParamType::Transform);
 		}
+		else if(ScriptStruct == FAnimNextGraphLODPose::StaticStruct())
+		{
+			Handle.SetParameterType(EParamType::AnimNextGraphLODPose);
+		}
+		else if(ScriptStruct == FAnimNextGraphReferencePose::StaticStruct())
+		{
+			Handle.SetParameterType(EParamType::AnimNextGraphReferencePose);
+		}
+	}
+	else if(InProperty->IsA<FObjectProperty>() || InProperty->IsA<FObjectPtrProperty>())
+	{
+		const UClass* Class = CastField<FObjectPropertyBase>(InProperty)->PropertyClass;
+		if (Class == UObject::StaticClass())
+		{
+			Handle.SetParameterType(EParamType::Object);
+		}
+		else if (Class == UCharacterMovementComponent::StaticClass())
+		{
+			Handle.SetParameterType(EParamType::CharacterMovementComponent);
+		}
+		else if (Class == UAnimNextMeshComponent::StaticClass())
+		{
+			Handle.SetParameterType(EParamType::AnimNextMeshComponent);
+		}
 	}
 
 	// Not found - custom type
@@ -459,13 +534,56 @@ FParamTypeHandle FParamTypeHandle::FromObject(const UObject* InObject)
 		{
 			Handle.SetCustomTypeIndex(GetOrAllocateCustomTypeIndex(FAnimNextParamType::EValueType::Class, FAnimNextParamType::EContainerType::None, Class));
 		}
-		else if(const UScriptStruct* Struct = Cast<UScriptStruct>(InObject))
+		else if(const UScriptStruct* ScriptStruct = Cast<UScriptStruct>(InObject))
 		{
-			Handle.SetCustomTypeIndex(GetOrAllocateCustomTypeIndex(FAnimNextParamType::EValueType::Struct, FAnimNextParamType::EContainerType::None, Struct));
+			if (ScriptStruct == TBaseStructure<FVector>::Get())
+			{
+				Handle.SetParameterType(EParamType::Vector);
+			}
+			else if (ScriptStruct == TBaseStructure<FVector4>::Get())
+			{
+				Handle.SetParameterType(EParamType::Vector4);
+			}
+			else if (ScriptStruct == TBaseStructure<FQuat>::Get())
+			{
+				Handle.SetParameterType(EParamType::Quat);
+			}
+			else if (ScriptStruct == TBaseStructure<FTransform>::Get())
+			{
+				Handle.SetParameterType(EParamType::Transform);
+			}
+			else if(ScriptStruct == FAnimNextGraphLODPose::StaticStruct())
+			{
+				Handle.SetParameterType(EParamType::AnimNextGraphLODPose);
+			}
+			else if(ScriptStruct == FAnimNextGraphReferencePose::StaticStruct())
+			{
+				Handle.SetParameterType(EParamType::AnimNextGraphReferencePose);
+			}
+			else
+			{
+				Handle.SetCustomTypeIndex(GetOrAllocateCustomTypeIndex(FAnimNextParamType::EValueType::Struct, FAnimNextParamType::EContainerType::None, ScriptStruct));
+			}
 		}
 		else
 		{
-			Handle.SetCustomTypeIndex(GetOrAllocateCustomTypeIndex(FAnimNextParamType::EValueType::Object, FAnimNextParamType::EContainerType::None, InObject->GetClass()));
+			const UClass* ObjectClass = InObject->GetClass();
+			if (ObjectClass == UObject::StaticClass())
+			{
+				Handle.SetParameterType(EParamType::Object);
+			}
+			else if (ObjectClass == UCharacterMovementComponent::StaticClass())
+			{
+				Handle.SetParameterType(EParamType::CharacterMovementComponent);
+			}
+			else if (ObjectClass == UAnimNextMeshComponent::StaticClass())
+			{
+				Handle.SetParameterType(EParamType::AnimNextMeshComponent);
+			}
+			else
+			{
+				Handle.SetCustomTypeIndex(GetOrAllocateCustomTypeIndex(FAnimNextParamType::EValueType::Object, FAnimNextParamType::EContainerType::None, ObjectClass));
+			}
 		}
 	}
 
@@ -537,8 +655,40 @@ FAnimNextParamType FParamTypeHandle::GetType() const
 		ParameterType.ValueType = FAnimNextParamType::EValueType::Struct;
 		ParameterType.ContainerType = FAnimNextParamType::EContainerType::None;
 		break;
+	case EParamType::Object:
+		ParameterType.ValueTypeObject = UObject::StaticClass();
+		ParameterType.ValueType = FAnimNextParamType::EValueType::Object;
+		ParameterType.ContainerType = FAnimNextParamType::EContainerType::None;
+		break;
+	case EParamType::CharacterMovementComponent:
+		ParameterType.ValueTypeObject = UCharacterMovementComponent::StaticClass();
+		ParameterType.ValueType = FAnimNextParamType::EValueType::Object;
+		ParameterType.ContainerType = FAnimNextParamType::EContainerType::None;
+		break;
+	case EParamType::AnimNextMeshComponent:
+		ParameterType.ValueTypeObject = UAnimNextMeshComponent::StaticClass();
+		ParameterType.ValueType = FAnimNextParamType::EValueType::Object;
+		ParameterType.ContainerType = FAnimNextParamType::EContainerType::None;
+		break;
+	case EParamType::AnimSequence:
+		ParameterType.ValueTypeObject = UAnimSequence::StaticClass();
+		ParameterType.ValueType = FAnimNextParamType::EValueType::Object;
+		ParameterType.ContainerType = FAnimNextParamType::EContainerType::None;
+		break;
+	case EParamType::AnimNextGraphLODPose:
+		ParameterType.ValueTypeObject = FAnimNextGraphLODPose::StaticStruct();
+		ParameterType.ValueType = FAnimNextParamType::EValueType::Struct;
+		ParameterType.ContainerType = FAnimNextParamType::EContainerType::None;
+		break;
+	case EParamType::AnimNextGraphReferencePose:
+		ParameterType.ValueTypeObject = FAnimNextGraphReferencePose::StaticStruct();
+		ParameterType.ValueType = FAnimNextParamType::EValueType::Struct;
+		ParameterType.ContainerType = FAnimNextParamType::EContainerType::None;
+		break;
 	case EParamType::Custom:
 		{
+			SCOPE_CYCLE_COUNTER(STAT_AnimNext_ParamTypeHandle_Lock);
+
 			FRWScopeLock ScopeLock(GParamTypeHandleLock, SLT_ReadOnly);
 			ParameterType = GetTypeHandleData().CustomTypes[GetCustomTypeIndex() - 1];
 		}
@@ -582,6 +732,15 @@ size_t FParamTypeHandle::GetSize() const
 		return sizeof(FQuat);
 	case EParamType::Transform:
 		return sizeof(FTransform);
+	case EParamType::Object:
+	case EParamType::CharacterMovementComponent:
+	case EParamType::AnimNextMeshComponent:
+	case EParamType::AnimSequence:
+		return sizeof(UObject*);
+	case EParamType::AnimNextGraphLODPose:
+		return sizeof(FAnimNextGraphLODPose);
+	case EParamType::AnimNextGraphReferencePose:
+		return sizeof(FAnimNextGraphReferencePose);
 	case EParamType::Custom:
 		return GetType().GetSize();
 	}
@@ -627,6 +786,15 @@ size_t FParamTypeHandle::GetAlignment() const
 		return alignof(FQuat);
 	case EParamType::Transform:
 		return alignof(FTransform);
+	case EParamType::Object:
+	case EParamType::CharacterMovementComponent:
+	case EParamType::AnimNextMeshComponent:
+	case EParamType::AnimSequence:
+		return alignof(UObject*);
+	case EParamType::AnimNextGraphLODPose:
+		return alignof(FAnimNextGraphLODPose);
+	case EParamType::AnimNextGraphReferencePose:
+		return alignof(FAnimNextGraphReferencePose);
 	case EParamType::Custom:
 		return GetType().GetAlignment();
 	}
diff --git a/Engine/Plugins/Experimental/AnimNext/Source/AnimNext/Private/Param/ParamUtils.cpp b/Engine/Plugins/Experimental/AnimNext/Source/AnimNext/Private/Param/ParamUtils.cpp
--- a/Engine/Plugins/Experimental/AnimNext/Source/AnimNext/Private/Param/ParamUtils.cpp
+++ b/Engine/Plugins/Experimental/AnimNext/Source/AnimNext/Private/Param/ParamUtils.cpp
@@ -1,6 +1,10 @@
 // Copyright Epic Games, Inc. All Rights Reserved.
 
 #include "Param/ParamUtils.h"
+
+#include "Component/AnimNextMeshComponent.h"
+#include "GameFramework/CharacterMovementComponent.h"
+#include "Animation/AnimSequence.h"
 #include "Param/ParamType.h"
 #include "Param/ParamTypeHandle.h"
 #include "Param/ParamCompatibility.h"
@@ -10,6 +14,26 @@ namespace UE::AnimNext
 
 FParamCompatibility FParamUtils::GetCompatibility(const FParamTypeHandle& InLHS, const FParamTypeHandle& InRHS)
 {
+	auto CheckClassCast = [](const FParamTypeHandle& InLHS, const UClass* InRHSClass)
+	{
+		FAnimNextParamType::EValueType ValueTypeLHS;
+		FAnimNextParamType::EContainerType ContainerTypeLHS;
+		const UObject* ValueTypeObjectLHS;
+
+		InLHS.GetCustomTypeInfo(ValueTypeLHS, ContainerTypeLHS, ValueTypeObjectLHS);
+		if(ContainerTypeLHS == FAnimNextParamType::EContainerType::None && ValueTypeLHS == FAnimNextParamType::EValueType::Object)
+		{
+			if(const UClass* Class = Cast<UClass>(ValueTypeObjectLHS))
+			{
+				if(InRHSClass->IsChildOf(Class))
+				{
+					return true;
+				}
+			}
+		}
+		return false;
+	};
+	
 	switch (InRHS.GetParameterType())
 	{
 	case FParamTypeHandle::EParamType::Bool:
@@ -143,6 +167,69 @@ FParamCompatibility FParamUtils::GetCompatibility(const FParamTypeHandle& InLHS,
 			return EParamCompatibility::Compatible_Equal;
 		}
 		break;
+	case FParamTypeHandle::EParamType::Object:
+		switch (InLHS.GetParameterType())
+		{
+		case FParamTypeHandle::EParamType::Object:
+			return EParamCompatibility::Compatible_Equal;
+		}
+		break;
+	case FParamTypeHandle::EParamType::CharacterMovementComponent:
+		switch (InLHS.GetParameterType())
+		{
+		case FParamTypeHandle::EParamType::Object:
+			return EParamCompatibility::Compatible_Cast;
+		case FParamTypeHandle::EParamType::CharacterMovementComponent:
+			return EParamCompatibility::Compatible_Equal;
+		case FParamTypeHandle::EParamType::Custom:
+			if(CheckClassCast(InLHS, UCharacterMovementComponent::StaticClass()))
+			{
+				return EParamCompatibility::Compatible_Cast;
+			}
+		}
+		break;
+	case FParamTypeHandle::EParamType::AnimNextMeshComponent:
+		switch (InLHS.GetParameterType())
+		{
+		case FParamTypeHandle::EParamType::Object:
+			return EParamCompatibility::Compatible_Cast;
+		case FParamTypeHandle::EParamType::AnimNextMeshComponent:
+			return EParamCompatibility::Compatible_Equal;
+		case FParamTypeHandle::EParamType::Custom:
+			if(CheckClassCast(InLHS, UAnimNextMeshComponent::StaticClass()))
+			{
+				return EParamCompatibility::Compatible_Cast;
+			}
+		}
+		break;
+	case FParamTypeHandle::EParamType::AnimSequence:
+		switch (InLHS.GetParameterType())
+		{
+		case FParamTypeHandle::EParamType::Object:
+			return EParamCompatibility::Compatible_Cast;
+		case FParamTypeHandle::EParamType::AnimSequence:
+			return EParamCompatibility::Compatible_Equal;
+		case FParamTypeHandle::EParamType::Custom:
+			if(CheckClassCast(InLHS, UAnimSequence::StaticClass()))
+			{
+				return EParamCompatibility::Compatible_Cast;
+			}
+		}
+		break;
+	case FParamTypeHandle::EParamType::AnimNextGraphLODPose:
+		switch (InLHS.GetParameterType())
+		{
+		case FParamTypeHandle::EParamType::AnimNextGraphLODPose:
+			return EParamCompatibility::Compatible_Equal;
+		}
+		break;
+	case FParamTypeHandle::EParamType::AnimNextGraphReferencePose:
+		switch (InLHS.GetParameterType())
+		{
+		case FParamTypeHandle::EParamType::AnimNextGraphReferencePose:
+			return EParamCompatibility::Compatible_Equal;
+		}
+		break;
 	case FParamTypeHandle::EParamType::Custom:
 		switch (InLHS.GetParameterType())
 		{
diff --git a/Engine/Plugins/Experimental/AnimNext/Source/AnimNext/Private/Scheduler/AnimNextSchedule.cpp b/Engine/Plugins/Experimental/AnimNext/Source/AnimNext/Private/Scheduler/AnimNextSchedule.cpp
--- a/Engine/Plugins/Experimental/AnimNext/Source/AnimNext/Private/Scheduler/AnimNextSchedule.cpp
+++ b/Engine/Plugins/Experimental/AnimNext/Source/AnimNext/Private/Scheduler/AnimNextSchedule.cpp
@@ -3,8 +3,9 @@
 #include "Scheduler/AnimNextSchedule.h"
 #include "Tasks/Task.h"
 #include "Async/TaskGraphInterfaces.h"
-#include "UObject/ObjectSaveContext.h"
 #include "EngineLogs.h"
+#include "Scheduler/AnimNextSchedulePort.h"
+#include "Graph/AnimNextGraph.h"
 
 void UAnimNextSchedule::PostLoad()
 {
@@ -24,15 +25,24 @@ void UAnimNextSchedule::PostEditChangeProperty(FPropertyChangedEvent& PropertyCh
 	CompileSchedule();
 }
 
+void UAnimNextSchedule::PostEditUndo()
+{
+	Super::PostEditUndo();
+
+	CompileSchedule();
+}
+
 void UAnimNextSchedule::CompileSchedule()
 {
+	using namespace UE::AnimNext;
+
 	Instructions.Empty();
 	Tasks.Empty();
 	Ports.Empty();
 	ExternalTasks.Empty();
 	ParamScopeEntryTasks.Empty();
 	ParamScopeExitTasks.Empty();
-	PortNameIndexMap.Empty();
+	IntermediatesData.Reset();
 	NumParameterScopes = 0;
 	NumTickFunctions = 0;
 
@@ -75,10 +85,13 @@ void UAnimNextSchedule::CompileSchedule()
 	// MAX_uint32 means 'global scope' in this context
 	uint32 ParentScopeIndex = MAX_uint32;
 
+	TArray<FAnimNextParamType> IntermediateTypes;
+	TMap<FName, uint32> IntermediateMap;
+
 	TFunction<void(const TArray<TObjectPtr<UAnimNextScheduleEntry>>&)> EmitEntries;
 
 	// Iterate over all entries, recursing into scopes
-	EmitEntries = [this, &EmitEntries, &Emit, &EmitPrerequisite, &ParentScopeIndex](const TArray<TObjectPtr<UAnimNextScheduleEntry>>& InEntries)
+	EmitEntries = [this, &EmitEntries, &Emit, &EmitPrerequisite, &ParentScopeIndex, &IntermediateTypes, &IntermediateMap](const TArray<TObjectPtr<UAnimNextScheduleEntry>>& InEntries)
 	{
 		for (int32 EntryIndex = 0; EntryIndex < InEntries.Num(); ++EntryIndex)
 		{
@@ -86,36 +99,145 @@ void UAnimNextSchedule::CompileSchedule()
 
 			if (UAnimNextScheduleEntry_Port* PortEntry = Cast<UAnimNextScheduleEntry_Port>(Entry))
 			{
-				EmitPrerequisite();
-
-				FAnimNextSchedulePortTask PortTask;
-				PortTask.TaskIndex = Ports.Num();
-				PortTask.ParamScopeIndex = ParentScopeIndex;
-				PortTask.Name = PortEntry->Name;
-				int32 PortIndex = Ports.Add(PortTask);
-				PortNameIndexMap.Add(PortEntry->Name, PortIndex);
-
-				NumTickFunctions++;
-
-				Emit(EAnimNextScheduleScheduleOpcode::RunPort, PortIndex);
+				bool bValid = true;
+
+				if(PortEntry->Port == nullptr)
+				{
+					UE_LOG(LogAnimation, Error, TEXT("AnimNext: Invalid port class found"));
+					bValid = false;
+				}
+				else
+				{
+					UAnimNextSchedulePort* CDO = PortEntry->Port->GetDefaultObject<UAnimNextSchedulePort>();
+					check(CDO);
+
+					TConstArrayView<FScheduleTerm> Terms = CDO->GetTerms();
+					if(PortEntry->Terms.Num() != Terms.Num())
+					{
+						UE_LOG(LogAnimation, Error, TEXT("AnimNext: Incorrect term count for port: %d"), PortEntry->Terms.Num());
+						bValid = false;
+					}
+
+					for(int32 TermIndex = 0; TermIndex < PortEntry->Terms.Num(); ++TermIndex)
+					{
+						FName Term = PortEntry->Terms[TermIndex];
+						const uint32* ExistingIntermediateIndexPtr = IntermediateMap.Find(Term);
+						if(ExistingIntermediateIndexPtr != nullptr)
+						{
+							const FAnimNextParamType& IntermediateType = IntermediateTypes[*ExistingIntermediateIndexPtr];
+							if(IntermediateType != Terms[TermIndex].GetType())
+							{
+								UE_LOG(LogAnimation, Error, TEXT("AnimNext: Mismatched types when processing port term, ignored: '%s'"), *Term.ToString());
+								bValid = false;
+							}
+						}
+					}
+				}
+				
+				if(bValid)
+				{
+					EmitPrerequisite();
+
+					FAnimNextSchedulePortTask PortTask;
+					PortTask.TaskIndex = Ports.Num();
+					PortTask.ParamScopeIndex = ParentScopeIndex;
+					PortTask.Port = PortEntry->Port;
+
+					UAnimNextSchedulePort* CDO = PortEntry->Port->GetDefaultObject<UAnimNextSchedulePort>();
+					check(CDO);
+					TConstArrayView<FScheduleTerm> Terms = CDO->GetTerms();
+
+					for(int32 TermIndex = 0; TermIndex < PortEntry->Terms.Num(); ++TermIndex)
+					{
+						FName Term = PortEntry->Terms[TermIndex];
+						const uint32* ExistingIntermediateIndexPtr = IntermediateMap.Find(Term);
+						if(ExistingIntermediateIndexPtr == nullptr)
+						{
+							uint32 IntermediateIndex = IntermediateTypes.Add(Terms[TermIndex].GetType());
+							IntermediateMap.Add(Term, IntermediateIndex);
+							PortTask.Terms.Add(IntermediateIndex);
+						}
+						else
+						{
+							PortTask.Terms.Add(*ExistingIntermediateIndexPtr);
+						}
+					}
+
+					int32 PortIndex = Ports.Add(PortTask);
+
+					NumTickFunctions++;
+
+					Emit(EAnimNextScheduleScheduleOpcode::RunPort, PortIndex);
+				}
 			}
 			else if (UAnimNextScheduleEntry_AnimNextGraph* GraphEntry = Cast<UAnimNextScheduleEntry_AnimNextGraph>(Entry))
 			{
-				EmitPrerequisite();
-
-				FAnimNextScheduleGraphTask Task;
-				Task.TaskIndex = Tasks.Num();
-				Task.ParamScopeIndex = NumParameterScopes++;
-				Task.ParamParentScopeIndex = ParentScopeIndex;
-				Task.Name = GraphEntry->Name;
-				Task.EntryPoint = GraphEntry->EntryPoint;
-				Task.Graph = GraphEntry->Graph;
-				Task.ParameterBlocks = GraphEntry->ParameterBlocks;
-				int32 TaskIndex = Tasks.Add(Task);
-
-				NumTickFunctions++;
-
-				Emit(EAnimNextScheduleScheduleOpcode::RunTask, TaskIndex);
+				bool bValid = true;
+
+				if(GraphEntry->Graph == nullptr)
+				{
+					UE_LOG(LogAnimation, Warning, TEXT("AnimNext: Invalid graph supplied"));
+					bValid = false;
+				}
+				else
+				{
+					TConstArrayView<FScheduleTerm> Terms = GraphEntry->Graph->GetTerms();
+					if(GraphEntry->Terms.Num() != Terms.Num())
+					{
+						UE_LOG(LogAnimation, Error, TEXT("AnimNext: Incorrect term count for graph: %d"), GraphEntry->Terms.Num());
+						bValid = false;
+					}
+
+					for(int32 TermIndex = 0; TermIndex < GraphEntry->Terms.Num(); ++TermIndex)
+					{
+						FName Term = GraphEntry->Terms[TermIndex];
+						const uint32* ExistingIntermediateIndexPtr = IntermediateMap.Find(Term);
+						if(ExistingIntermediateIndexPtr != nullptr)
+						{
+							const FAnimNextParamType& IntermediateType = IntermediateTypes[*ExistingIntermediateIndexPtr];
+							if(IntermediateType != Terms[TermIndex].GetType())
+							{
+								UE_LOG(LogAnimation, Error, TEXT("AnimNext: Mismatched types when processing graph term, ignored: '%s'"), *Term.ToString());
+								bValid = false;
+							}
+						}
+					}
+				}
+
+				if(bValid)
+				{
+					EmitPrerequisite();
+
+					FAnimNextScheduleGraphTask Task;
+					Task.TaskIndex = Tasks.Num();
+					Task.ParamScopeIndex = NumParameterScopes++;
+					Task.ParamParentScopeIndex = ParentScopeIndex;
+					Task.EntryPoint = GraphEntry->EntryPoint;
+					Task.Graph = GraphEntry->Graph;
+
+					TConstArrayView<FScheduleTerm> Terms = GraphEntry->Graph->GetTerms();
+					for(int32 TermIndex = 0; TermIndex < GraphEntry->Terms.Num(); ++TermIndex)
+					{
+						FName Term = GraphEntry->Terms[TermIndex];
+						const uint32* ExistingIntermediateIndexPtr = IntermediateMap.Find(Term);
+						if(ExistingIntermediateIndexPtr == nullptr)
+						{
+							uint32 IntermediateIndex = IntermediateTypes.Add(Terms[TermIndex].GetType());
+							IntermediateMap.Add(Term, IntermediateIndex);
+							Task.Terms.Add(IntermediateIndex);
+						}
+						else
+						{
+							Task.Terms.Add(*ExistingIntermediateIndexPtr);
+						}
+					}
+
+					int32 TaskIndex = Tasks.Add(Task);
+
+					NumTickFunctions++;
+
+					Emit(EAnimNextScheduleScheduleOpcode::RunTask, TaskIndex);
+				}
 			}
 			else if (UAnimNextScheduleEntry_ExternalTask* ExternalTaskEntry = Cast<UAnimNextScheduleEntry_ExternalTask>(Entry))
 			{
@@ -125,8 +247,8 @@ void UAnimNextSchedule::CompileSchedule()
 				ExternalTask.TaskIndex = ExternalTasks.Num();
 				ExternalTask.ParamScopeIndex = NumParameterScopes++;
 				ExternalTask.ParamParentScopeIndex = ParentScopeIndex;
-				ExternalTask.Name = ExternalTaskEntry->Name;
-				ExternalTask.ObjectName = ExternalTaskEntry->ObjectName;
+				ExternalTask.TickFunction = ExternalTaskEntry->TickFunction;
+				ExternalTask.Object = ExternalTaskEntry->Object;
 				int32 ExternalTaskIndex = ExternalTasks.Add(ExternalTask);
 
 				// Emit the external task
@@ -148,7 +270,7 @@ void UAnimNextSchedule::CompileSchedule()
 				ParamScopeEntryTask.ParamScopeIndex = ParamScopeIndex;
 				ParamScopeEntryTask.ParamParentScopeIndex = ParentScopeIndex;
 				ParamScopeEntryTask.TickFunctionIndex = NumTickFunctions;
-				ParamScopeEntryTask.Name = ParamScopeTaskEntry->Name;
+				ParamScopeEntryTask.Scope = ParamScopeTaskEntry->Scope;
 				ParamScopeEntryTask.ParameterBlocks = ParamScopeTaskEntry->ParameterBlocks;
 				int32 ParamScopeTaskEntryIndex = ParamScopeEntryTasks.Add(ParamScopeEntryTask);
 
@@ -170,7 +292,7 @@ void UAnimNextSchedule::CompileSchedule()
 				FAnimNextScheduleParamScopeExitTask ParamScopeExitTask;
 				ParamScopeExitTask.TaskIndex = ParamScopeExitTasks.Num();
 				ParamScopeExitTask.ParamScopeIndex = ParamScopeIndex;
-				ParamScopeExitTask.Name = ParamScopeTaskEntry->Name;
+				ParamScopeExitTask.Scope = ParamScopeTaskEntry->Scope;
 				int32 ParamScopeExitTaskIndex = ParamScopeExitTasks.Add(ParamScopeExitTask);
 
 				Emit(EAnimNextScheduleScheduleOpcode::RunParamScopeExit, ParamScopeExitTaskIndex);
@@ -182,6 +304,24 @@ void UAnimNextSchedule::CompileSchedule()
 	EmitEntries(Entries);
 
 	Emit(EAnimNextScheduleScheduleOpcode::Exit);
+
+	// Process intermediates
+	if(IntermediateMap.Num() > 0)
+	{
+		check(IntermediateMap.Num() == IntermediateTypes.Num());
+		
+		TArray<FPropertyBagPropertyDesc> PropertyDescs;
+		PropertyDescs.Reserve(IntermediateTypes.Num());
+		 
+		for(const TPair<FName, uint32>& IntermediatePair : IntermediateMap)
+		{
+			const FAnimNextParamType& IntermediateType = IntermediateTypes[IntermediatePair.Value];
+			check(IntermediateType.IsValid());
+			PropertyDescs.Emplace(IntermediatePair.Key, IntermediateType.GetContainerType(), IntermediateType.GetValueType(), IntermediateType.GetValueTypeObject());
+		}
+
+		IntermediatesData.AddProperties(PropertyDescs);
+	}
 }
 
 #endif // #if WITH_EDITOR
\ No newline at end of file
diff --git a/Engine/Plugins/Experimental/AnimNext/Source/AnimNext/Private/Scheduler/AnimNextScheduleGraphTask.cpp b/Engine/Plugins/Experimental/AnimNext/Source/AnimNext/Private/Scheduler/AnimNextScheduleGraphTask.cpp
--- a/Engine/Plugins/Experimental/AnimNext/Source/AnimNext/Private/Scheduler/AnimNextScheduleGraphTask.cpp
+++ b/Engine/Plugins/Experimental/AnimNext/Source/AnimNext/Private/Scheduler/AnimNextScheduleGraphTask.cpp
@@ -1,30 +1,97 @@
 // Copyright Epic Games, Inc. All Rights Reserved.
 
 #include "Scheduler/AnimNextScheduleGraphTask.h"
+
 #include "Scheduler/ScheduleContext.h"
 #include "Graph/AnimNextGraph.h"
 #include "Context.h"
 #include "Graph/AnimNextExecuteContext.h"
+#include "Graph/AnimNext_LODPose.h"
+#include "AnimNextStats.h"
+
+DEFINE_STAT(STAT_AnimNext_Task_Graph);
 
-void FAnimNextScheduleGraphTask::RunGraph(const UE::AnimNext::FScheduleContext& InScheduleContext) const
+void FAnimNextScheduleGraphTask::RunGraph(const UE::AnimNext::FScheduleContext& InContext) const
 {
+	SCOPE_CYCLE_COUNTER(STAT_AnimNext_Task_Graph);
+
 	using namespace UE::AnimNext;
 
-	if (Graph)
+	if (Graph == nullptr)
+	{
+		return;
+	}
+
+	FScheduleInstanceData& InstanceData = InContext.GetInstanceData();
+
+	// Allocate our graph instance data
+	if (!InstanceData.GraphInstanceData[TaskIndex].IsValid())
+	{
+		Graph->AllocateInstance(InstanceData.GraphInstanceData[TaskIndex]);
+	}
+
+	FParamStack& ParamStack = FParamStack::Get();
+	const FAnimNextGraphReferencePose* GraphReferencePose = ParamStack.GetParamPtr<FAnimNextGraphReferencePose>(Graph->GetReferencePoseParam());
+	if(GraphReferencePose == nullptr)
 	{
-		FScheduleInstanceData& InstanceData = InScheduleContext.GetInstanceData();
+		return;
+	}
 
-		// Allocate our graph instance data
-		if (!InstanceData.GraphInstanceData[TaskIndex].IsValid())
-		{
-			 Graph->AllocateInstance(InstanceData.GraphInstanceData[TaskIndex]);
-		}
+	const int32* GraphLODLevel = ParamStack.GetParamPtr<int32>(Graph->GetCurrentLODParam());
+	if(GraphLODLevel == nullptr)
+	{
+		return;
+	}
 
-		// FIXME: Cant run concurrently so disabling for now.
-		if(0)
-		{
-			const FContext Context;
-			Graph->Run(Context, InstanceData.GraphInstanceData[TaskIndex], EAnimNextGraphSimulationSteps::All);
-		}
+	// TODO: This should not be fixed at arg 0, we should define this in the graph asset
+	const FParamStackLayerHandle& TermLayerHandle = InstanceData.GraphTermLayers[TaskIndex];
+	FAnimNextGraphLODPose* OutputPose = TermLayerHandle.GetMutableParamPtr<FAnimNextGraphLODPose>(Graph->GetTerms()[0].GetId());
+	if(OutputPose == nullptr)
+	{
+		return;
 	}
+
+	float DeltaTime = 1/60.0f;
+	const float* ExternalDeltaTime = ParamStack.GetParamPtr<float>(Graph->GetDeltaTimeParam());
+	if(ExternalDeltaTime != nullptr)
+	{
+		DeltaTime = *ExternalDeltaTime;
+	}
+
+	// Create or update our result pose
+	// TODO: Currently forcing additive flag to false here
+	if (OutputPose->LODPose.ShouldPrepareForLOD(*GraphReferencePose->ReferencePose, *GraphLODLevel, false))
+	{
+		OutputPose->LODPose.PrepareForLOD(*GraphReferencePose->ReferencePose, *GraphLODLevel, true, false);
+	}
+
+	check(OutputPose->LODPose.LODLevel == *GraphLODLevel);
+
+	// Push parameter layers that translate schedule data to graph inputs
+	// TODO: This should probably be reworked, its overly complex for what it does!
+	static FParamId ResultId("UE_Internal_ResultPose");
+	static FParamId ExpectsAdditiveId("UE_Internal_GraphExpectsAdditive");
+
+	if(!InstanceData.GraphInputLayers[TaskIndex].IsValid())
+	{
+		InstanceData.GraphInputLayers[TaskIndex] = FParamStack::MakeValuesLayer(
+			ResultId, *OutputPose,
+			ExpectsAdditiveId, false);
+	}
+	else
+	{
+		InstanceData.GraphInputLayers[TaskIndex].SetValues(
+			ResultId, *OutputPose,
+			ExpectsAdditiveId, false);
+	}
+
+	FParamStack::FPushedLayerHandle LayerHandle = ParamStack.PushLayer(InstanceData.GraphInputLayers[TaskIndex]);
+
+	// Internally we use memstack allocation, so we need a mark here
+	FMemMark MemMark(FMemStack::Get());
+
+	const FContext Context(DeltaTime);
+	Graph->Run(Context, InstanceData.GraphInstanceData[TaskIndex], EAnimNextGraphSimulationSteps::All);
+
+	ParamStack.PopLayer(LayerHandle);
 }
diff --git a/Engine/Plugins/Experimental/AnimNext/Source/AnimNext/Private/Scheduler/AnimNextScheduleParamScopeTask.cpp b/Engine/Plugins/Experimental/AnimNext/Source/AnimNext/Private/Scheduler/AnimNextScheduleParamScopeTask.cpp
--- a/Engine/Plugins/Experimental/AnimNext/Source/AnimNext/Private/Scheduler/AnimNextScheduleParamScopeTask.cpp
+++ b/Engine/Plugins/Experimental/AnimNext/Source/AnimNext/Private/Scheduler/AnimNextScheduleParamScopeTask.cpp
@@ -5,9 +5,15 @@
 #include "Scheduler/ScheduleContext.h"
 #include "Scheduler/AnimNextSchedulerEntry.h"
 #include "Param/AnimNextParameterSourceRef.h"
+#include "AnimNextStats.h"
+
+DEFINE_STAT(STAT_AnimNext_Task_ScopeEntry);
+DEFINE_STAT(STAT_AnimNext_Task_ScopeExit);
 
 void FAnimNextScheduleParamScopeEntryTask::RunParamScopeEntry(const UE::AnimNext::FScheduleContext& InScheduleContext) const
 {
+	SCOPE_CYCLE_COUNTER(STAT_AnimNext_Task_ScopeEntry);
+
 	using namespace UE::AnimNext;
 
 	FParamStack& ParamStack = FParamStack::Get();
@@ -16,7 +22,7 @@ void FAnimNextScheduleParamScopeEntryTask::RunParamScopeEntry(const UE::AnimNext
 	FScheduleInstanceData::FScopeCache& ScopeCache = InstanceData.ScopeCaches[TaskIndex];
 	ScopeCache.PushedLayers.Reset();
 
-	const FAnimNextParameterCollection* FoundUserScope = InstanceData.UserScopes.Find(Name);
+	const FAnimNextParameterCollection* FoundUserScope = InstanceData.UserScopes.Find(Scope);
 	if(FoundUserScope)
 	{
 		ScopeCache.PushedLayers.Reserve(ParameterBlocks.Num() + FoundUserScope->Parameters.Num());
@@ -32,13 +38,16 @@ void FAnimNextScheduleParamScopeEntryTask::RunParamScopeEntry(const UE::AnimNext
 		ScopeCache.StaticHandles.SetNum(ParameterBlocks.Num());
 		for (int32 LayerHandleIndex = 0; LayerHandleIndex < ScopeCache.StaticHandles.Num(); ++LayerHandleIndex)
 		{
-			ScopeCache.StaticHandles[LayerHandleIndex] = ParameterBlocks[LayerHandleIndex]->CacheLayer();
+			if(ParameterBlocks[LayerHandleIndex])
+			{
+				ScopeCache.StaticHandles[LayerHandleIndex] = ParameterBlocks[LayerHandleIndex]->CacheLayer();
+			}
 		}
 	}
 
 	if (FoundUserScope)
 	{
-		const UObject* ObjectContext = InstanceData.Entry->Object.Get();
+		const UObject* ObjectContext = InstanceData.Entry->ResolvedObject;
 
 		// Resize user scope data appropriately
 		if (ScopeCache.UserHandles.Num() != FoundUserScope->Parameters.Num())
@@ -74,16 +83,20 @@ void FAnimNextScheduleParamScopeEntryTask::RunParamScopeEntry(const UE::AnimNext
 	// Update & push static params
 	for (int32 ParamBlockIndex = 0; ParamBlockIndex < ParameterBlocks.Num(); ++ParamBlockIndex)
 	{
-		const UAnimNextParameterBlock* ParameterBlock = ParameterBlocks[ParamBlockIndex];
-		FParamStackLayerHandle& LayerHandle = ScopeCache.StaticHandles[ParamBlockIndex];
+		if(const UAnimNextParameterBlock* ParameterBlock = ParameterBlocks[ParamBlockIndex])
+		{
+			FParamStackLayerHandle& LayerHandle = ScopeCache.StaticHandles[ParamBlockIndex];
 
-		ParameterBlock->UpdateLayer(LayerHandle);
-		ScopeCache.PushedLayers.Add(ParamStack.PushLayer(LayerHandle));
+			ParameterBlock->UpdateLayer(LayerHandle);
+			ScopeCache.PushedLayers.Add(ParamStack.PushLayer(LayerHandle));
+		}
 	}
 }
 
 void FAnimNextScheduleParamScopeExitTask::RunParamScopeExit(const UE::AnimNext::FScheduleContext& InScheduleContext) const
 {
+	SCOPE_CYCLE_COUNTER(STAT_AnimNext_Task_ScopeExit);
+
 	using namespace UE::AnimNext;
 
 	FParamStack& ParamStack = FParamStack::Get();
diff --git a/Engine/Plugins/Experimental/AnimNext/Source/AnimNext/Private/Scheduler/AnimNextSchedulePort.h b/Engine/Plugins/Experimental/AnimNext/Source/AnimNext/Private/Scheduler/AnimNextSchedulePort.h
deleted file mode 100644
diff --git a/Engine/Plugins/Experimental/AnimNext/Source/AnimNext/Private/Scheduler/AnimNextSchedulePortTask.cpp b/Engine/Plugins/Experimental/AnimNext/Source/AnimNext/Private/Scheduler/AnimNextSchedulePortTask.cpp
--- a/Engine/Plugins/Experimental/AnimNext/Source/AnimNext/Private/Scheduler/AnimNextSchedulePortTask.cpp
+++ b/Engine/Plugins/Experimental/AnimNext/Source/AnimNext/Private/Scheduler/AnimNextSchedulePortTask.cpp
@@ -2,26 +2,27 @@
 
 #include "Scheduler/AnimNextSchedulePortTask.h"
 #include "Scheduler/AnimNextSchedulePort.h"
-#include "Scheduler/SchedulePortDefinition.h"
-#include "Scheduler/Scheduler.h"
-#include "Scheduler/ScheduleContext.h"
-#include "Scheduler/SchedulePortAdapterContext.h"
+#include "ScheduleContext.h"
+#include "Scheduler/ScheduleTermContext.h"
+#include "AnimNextStats.h"
+
+DEFINE_STAT(STAT_AnimNext_Task_Port);
 
 void FAnimNextSchedulePortTask::RunPort(const UE::AnimNext::FScheduleContext& InScheduleContext) const
 {
+	SCOPE_CYCLE_COUNTER(STAT_AnimNext_Task_Port);
+
 	using namespace UE::AnimNext;
 
-	if (const FSchedulePortDefinition* PortDefinition = FScheduler::FindPortDefinition(Name))
+	if(Port)
 	{
-		if (PortDefinition->Adapter)
+		if(const UAnimNextSchedulePort* CDO = Port->GetDefaultObject<UAnimNextSchedulePort>())
 		{
-			if (const FAnimNextSchedulePort* PortParam = FParamStack::Get().GetParamPtr<FAnimNextSchedulePort>(Name))
-			{
-				FSchedulePortAdapterContext Context;
-				Context.OutputObject = PortParam->Object.Get();
-				Context.OutputData = PortParam->Data;
-				PortDefinition->Adapter(Context);
-			}
+			FScheduleInstanceData& InstanceData = InScheduleContext.GetInstanceData();
+
+			// Supply the I/O terms that the port needs
+			FScheduleTermContext TermContext(InScheduleContext, InstanceData.PortTermLayers[TaskIndex]);
+			CDO->Run(TermContext);
 		}
 	}
 }
\ No newline at end of file
diff --git a/Engine/Plugins/Experimental/AnimNext/Source/AnimNext/Private/Scheduler/AnimNextSchedulerEntry.cpp b/Engine/Plugins/Experimental/AnimNext/Source/AnimNext/Private/Scheduler/AnimNextSchedulerEntry.cpp
--- a/Engine/Plugins/Experimental/AnimNext/Source/AnimNext/Private/Scheduler/AnimNextSchedulerEntry.cpp
+++ b/Engine/Plugins/Experimental/AnimNext/Source/AnimNext/Private/Scheduler/AnimNextSchedulerEntry.cpp
@@ -1,92 +1,138 @@
 // Copyright Epic Games, Inc. All Rights Reserved.
 
 #include "AnimNextSchedulerEntry.h"
+
 #include "ScheduleInstanceData.h"
 #include "Scheduler/AnimNextSchedule.h"
 #include "Engine/World.h"
+#include "AnimNextStats.h"
+
+DEFINE_STAT(STAT_AnimNext_InitializeEntry);
 
-FAnimNextSchedulerEntry::FAnimNextSchedulerEntry(const UAnimNextSchedule* InSchedule, UObject* InObject, UE::AnimNext::FScheduleHandle InHandle, const TMap<FName, FAnimNextParameterCollection>& InUserScopes)
+FAnimNextSchedulerEntry::FAnimNextSchedulerEntry(const UAnimNextSchedule* InSchedule, UObject* InObject, UE::AnimNext::FScheduleHandle InHandle, const TMap<FName, FAnimNextParameterCollection>& InUserScopes, EAnimNextScheduleInitMethod InInitMethod)
 	: Schedule(InSchedule)
 	, UserScopes(InUserScopes)
-	, Object(InObject)
+	, WeakObject(InObject)
 	, Handle(InHandle)
 	, Context(InSchedule, this)
+	, RunState(ERunState::None)
+	, InitMethod(InInitMethod)
+{
+}
+
+FAnimNextSchedulerEntry::~FAnimNextSchedulerEntry()
+{
+	Invalidate();
+}
+
+void FAnimNextSchedulerEntry::Initialize()
 {
+	SCOPE_CYCLE_COUNTER(STAT_AnimNext_InitializeEntry);
+	
 	using namespace UE::AnimNext;
 
-	check(InObject)
-	check(InSchedule);
-	check(InHandle.IsValid());
+	check(IsInGameThread());
 
-	RootParamStack = MakeShared<FParamStack>();
+	check(WeakObject.IsValid())
+	check(Schedule);
+	check(Handle.IsValid());
 
-	FParamStack::AttachToCurrentThread(RootParamStack);
-	FScheduleContext::AttachToCurrentThread(Context);
+	ResolvedObject = WeakObject.Get();
+	bIsEditor = ResolvedObject && ResolvedObject->GetWorld()->WorldType == EWorldType::Editor;
 
 	// Setup tick function graph
-	if (InSchedule->Instructions.Num() > 0)
+	if (Schedule->Instructions.Num() > 0)
 	{
+		// TODO: split allocation and registration into separate async tasks to reduce GT overheads of spawning
+		// This will require correctly dealing with ReleaseHandle in the schedule as the two could then overlap
+		check(RunState == ERunState::None);
+		RunState = ERunState::CreatingTasks;
+
+		RootParamStack = MakeShared<FParamStack>();
+
+		FParamStack::AttachToCurrentThread(RootParamStack);
+		FScheduleContext::AttachToCurrentThread(Context);
+
+		FParamStack& ParamStack = FParamStack::Get();
+
+		TargetObjects.SetNum(Schedule->Instructions.Num());
+		
 		BeginTickFunction = MakeUnique<FScheduleBeginTickFunction>(*this);
 		EndTickFunction = MakeUnique<FScheduleEndTickFunction>(*this);
 
 		TArray<TUniquePtr<FScheduleTickFunction>*> Prerequisites;
 		int32 InstructionIndex = 0;
-		while (InstructionIndex < InSchedule->Instructions.Num())
+		while (InstructionIndex < Schedule->Instructions.Num())
 		{
-			auto AddTickFunctionAndPrerequisites = [this, &Prerequisites, InObject](TUniquePtr<FScheduleTickFunction>&& InTickFunction)
+			auto AddTickFunctionAndPrerequisites = [this, &Prerequisites](TUniquePtr<FScheduleTickFunction>&& InTickFunction)
 			{
 				// Add a prereq on the previous tick function we added, if any (we only support linear chains right now so this is OK)
 				if (TickFunctions.Num() == 0)
 				{
-					BeginTickFunction->Subsequent = FTickPrerequisite(InObject, *InTickFunction.Get());
-					InTickFunction->AddPrerequisite(InObject, *BeginTickFunction.Get());
+					BeginTickFunction->Subsequent = FTickPrerequisite(ResolvedObject, *InTickFunction.Get());
+					InTickFunction->AddPrerequisite(ResolvedObject, *BeginTickFunction.Get());
 				}
 
 				for (TUniquePtr<FScheduleTickFunction>* Prerequisite : Prerequisites)
 				{
-					Prerequisite->Get()->Subsequents.Emplace(InObject, *InTickFunction.Get());
-					InTickFunction->AddPrerequisite(InObject, *Prerequisite->Get());
+					Prerequisite->Get()->Subsequents.Emplace(ResolvedObject, *InTickFunction.Get());
+					InTickFunction->AddPrerequisite(ResolvedObject, *Prerequisite->Get());
 				}
 				Prerequisites.Reset();
 				TickFunctions.Add(MoveTemp(InTickFunction));
 			};
 
-			const FAnimNextScheduleInstruction& Instruction = InSchedule->Instructions[InstructionIndex++];
+			TWeakObjectPtr<UObject>& TargetObject = TargetObjects[InstructionIndex];
+			const FAnimNextScheduleInstruction& Instruction = Schedule->Instructions[InstructionIndex++];
 			switch (Instruction.Opcode)
 			{
 			case EAnimNextScheduleScheduleOpcode::BeginRunExternalTask:
 				{
-					AddTickFunctionAndPrerequisites(MakeUnique<FScheduleTickFunction>(Context, Instruction));
+					AddTickFunctionAndPrerequisites(MakeUnique<FScheduleTickFunction>(
+						Context,
+						TConstArrayView<FAnimNextScheduleInstruction>(&Instruction, 1),
+						TConstArrayView<TWeakObjectPtr<UObject>>(&TargetObject, 1)));
 
-					const FName TickFunctionName = InSchedule->ExternalTasks[Instruction.Operand].Name;
-					const FName ObjectName = InSchedule->ExternalTasks[Instruction.Operand].ObjectName;
+					const FName TickFunctionName = Schedule->ExternalTasks[Instruction.Operand].TickFunction;
+					const FName ObjectName = Schedule->ExternalTasks[Instruction.Operand].Object;
 					if(TickFunctionName != NAME_None && ObjectName != NAME_None)
 					{
-						FTickFunction* FoundFunction = FParamStack::Get().GetMutableParamPtr<FTickFunction>(TickFunctionName);
-						const UObject* FoundObject = FParamStack::Get().GetParamPtr<UObject>(ObjectName);
-						if(FoundFunction && FoundObject)
+						FTickFunction* FoundFunction = ParamStack.GetMutableParamPtr<FTickFunction>(TickFunctionName);
+						const TObjectPtr<UObject>* FoundObject = ParamStack.GetParamPtr<TObjectPtr<UObject>>(ObjectName);
+						if(FoundFunction && FoundObject && FoundObject->Get())
+						{
+							FoundFunction->AddPrerequisite(ResolvedObject, *TickFunctions.Last().Get());
+							TargetObject = FoundObject->Get();
+							TickFunctions.Last()->Subsequents.Emplace(FoundObject->Get(), *FoundFunction);
+						}
+						else
 						{
-							FoundFunction->AddPrerequisite(InObject, *TickFunctions.Last().Get());
-							TickFunctions.Last()->TargetObject = const_cast<UObject*>(FoundObject);
-							TickFunctions.Last()->Subsequents.Emplace(const_cast<UObject*>(FoundObject), *FoundFunction);
+							UE_LOG(LogAnimation, Warning, TEXT("AnimNext: Could not bind to external tick function using binding parameters (TickFunctionName=%s, ObjectName=%s)"), *TickFunctionName.ToString(), *ObjectName.ToString());
 						}
 					}
+					else
+					{
+						UE_LOG(LogAnimation, Warning, TEXT("AnimNext: Invalid external tick function binding parameters (TickFunctionName=%s, ObjectName=%s)"), *TickFunctionName.ToString(), *ObjectName.ToString());
+					}
 					break;
 				}
 			case EAnimNextScheduleScheduleOpcode::EndRunExternalTask:
 				{
-					AddTickFunctionAndPrerequisites(MakeUnique<FScheduleTickFunction>(Context, Instruction));
+					AddTickFunctionAndPrerequisites(MakeUnique<FScheduleTickFunction>(
+						Context,
+						TConstArrayView<FAnimNextScheduleInstruction>(&Instruction, 1),
+						TConstArrayView<TWeakObjectPtr<UObject>>(&TargetObject, 1)));
 
-					const FName TickFunctionName = InSchedule->ExternalTasks[Instruction.Operand].Name;
-					const FName ObjectName = InSchedule->ExternalTasks[Instruction.Operand].ObjectName;
+					const FName TickFunctionName = Schedule->ExternalTasks[Instruction.Operand].TickFunction;
+					const FName ObjectName = Schedule->ExternalTasks[Instruction.Operand].Object;
 					if(TickFunctionName != NAME_None && ObjectName != NAME_None)
 					{
-						FTickFunction* FoundFunction = FParamStack::Get().GetMutableParamPtr<FTickFunction>(TickFunctionName);
-						const UObject* FoundObject = FParamStack::Get().GetParamPtr<UObject>(ObjectName);
-						if(FoundFunction && FoundObject)
+						FTickFunction* FoundFunction = ParamStack.GetMutableParamPtr<FTickFunction>(TickFunctionName);
+						const TObjectPtr<UObject>* FoundObject = ParamStack.GetParamPtr<TObjectPtr<UObject>>(ObjectName);
+						if(FoundFunction && FoundObject && FoundObject->Get())
 						{
-							TickFunctions.Last()->TargetObject = const_cast<UObject*>(FoundObject);
-							TickFunctions.Last()->AddPrerequisite(const_cast<UObject*>(FoundObject), *FoundFunction);
+							TargetObject = FoundObject->Get();
+							TickFunctions.Last()->AddPrerequisite(FoundObject->Get(), *FoundFunction);
 						}
 					}
 					break;
@@ -95,7 +141,10 @@ FAnimNextSchedulerEntry::FAnimNextSchedulerEntry(const UAnimNextSchedule* InSche
 			case EAnimNextScheduleScheduleOpcode::RunPort:
 			case EAnimNextScheduleScheduleOpcode::RunParamScopeEntry:
 			case EAnimNextScheduleScheduleOpcode::RunParamScopeExit:
-				AddTickFunctionAndPrerequisites(MakeUnique<FScheduleTickFunction>(Context, Instruction));
+				AddTickFunctionAndPrerequisites(MakeUnique<FScheduleTickFunction>(
+					Context,
+					TConstArrayView<FAnimNextScheduleInstruction>(&Instruction, 1),
+					TConstArrayView<TWeakObjectPtr<UObject>>(&TargetObject, 1)));
 				break;
 			case EAnimNextScheduleScheduleOpcode::PrerequisiteTask:
 			case EAnimNextScheduleScheduleOpcode::PrerequisiteBeginExternalTask:
@@ -107,8 +156,8 @@ FAnimNextSchedulerEntry::FAnimNextSchedulerEntry(const UAnimNextSchedule* InSche
 				break;
 			case EAnimNextScheduleScheduleOpcode::Exit:
 				check(Prerequisites.Num() == 0);
-				EndTickFunction->AddPrerequisite(InObject, *TickFunctions.Last().Get());
-				TickFunctions.Last()->Subsequents.Emplace(InObject, *EndTickFunction.Get());
+				EndTickFunction->AddPrerequisite(ResolvedObject, *TickFunctions.Last().Get());
+				TickFunctions.Last()->Subsequents.Emplace(ResolvedObject, *EndTickFunction.Get());
 				break;
 			default:
 				check(false);
@@ -116,39 +165,62 @@ FAnimNextSchedulerEntry::FAnimNextSchedulerEntry(const UAnimNextSchedule* InSche
 			}
 		}
 
+		FScheduleContext::DetachFromCurrentThread();
+		FParamStack::DetachFromCurrentThread();
+
+		check(RunState == ERunState::CreatingTasks);
+		RunState = ERunState::BindingTasks;
+
 		// Register our tick functions
-		ULevel* Level = InObject->GetWorld()->PersistentLevel;
+		UWorld* World = ResolvedObject->GetWorld();
+		ULevel* Level = World->PersistentLevel;
 		BeginTickFunction->RegisterTickFunction(Level);
 		EndTickFunction->RegisterTickFunction(Level);
 		for (TUniquePtr<FScheduleTickFunction>& TickFunction : TickFunctions)
 		{
 			TickFunction->RegisterTickFunction(Level);
 		}
-	}
 
-	FScheduleContext::DetachFromCurrentThread();
-	FParamStack::DetachFromCurrentThread();
-}
+		RunState = ERunState::RunningInitialUpdate;
 
-FAnimNextSchedulerEntry::~FAnimNextSchedulerEntry()
-{
-	Invalidate();
+		// Just pause now if we arent needing an initial update
+		if(InitMethod == EAnimNextScheduleInitMethod::None)
+		{
+			Enable(false);
+		}
+		else
+		{
+			// In editor preview worlds (ideally just thumbnail scenes) we run a linearized 'initial tick' to ensure we
+			// generate an output pose, as these worlds never tick
+			if(World->WorldType == EWorldType::EditorPreview)
+			{
+				LazyAllocateInstanceData();
+				FScheduleContext::AttachToCurrentThread(Context);
+				FScheduleTickFunction::RunSchedule(Schedule->Instructions);
+				FScheduleContext::DetachFromCurrentThread();
+			}
+		}
+	}
+
+	ResolvedObject = nullptr;
 }
 
 void FAnimNextSchedulerEntry::Invalidate()
 {
 	using namespace UE::AnimNext;
 
+	check(IsInGameThread());
+
 	if(BeginTickFunction)
 	{
-		BeginTickFunction->Subsequent.PrerequisiteTickFunction->RemovePrerequisite(Object.Get(), *BeginTickFunction.Get());
+		BeginTickFunction->Subsequent.PrerequisiteTickFunction->RemovePrerequisite(WeakObject.Get(), *BeginTickFunction.Get());
 	}
 	
 	for (TUniquePtr<FScheduleTickFunction>& TickFunction : TickFunctions)
 	{
 		for (FTickPrerequisite& Subsequent : TickFunction->Subsequents)
 		{
-			Subsequent.PrerequisiteTickFunction->RemovePrerequisite(Object.Get(), *TickFunction.Get());
+			Subsequent.PrerequisiteTickFunction->RemovePrerequisite(WeakObject.Get(), *TickFunction.Get());
 		}
 		TickFunction->UnRegisterTickFunction();
 	}
@@ -167,20 +239,50 @@ void FAnimNextSchedulerEntry::Invalidate()
 	TickFunctions.Reset();
 
 	Schedule = nullptr;
-	Object = nullptr;
+	WeakObject = nullptr;
+	TargetObjects.Reset();
 	UserScopes.Empty();
 	Handle.Invalidate();
 	Context.InstanceData.Reset();
 }
 
-void FAnimNextSchedulerEntry::Enable(bool bInEnabled)
+void FAnimNextSchedulerEntry::ClearTickFunctionPauseFlags()
 {
 	using namespace UE::AnimNext;
+
+	check(IsInGameThread());
 	
+	BeginTickFunction->bTickEvenWhenPaused = false;
+	for (TUniquePtr<FScheduleTickFunction>& TickFunction : TickFunctions)
+	{
+		TickFunction->bTickEvenWhenPaused = false;
+	}
+	EndTickFunction->bTickEvenWhenPaused = false;
+}
+
+void FAnimNextSchedulerEntry::Enable(bool bInEnabled)
+{
+	using namespace UE::AnimNext;
+
+	check(IsInGameThread());
+
 	BeginTickFunction->SetTickFunctionEnable(bInEnabled);
 	for (TUniquePtr<FScheduleTickFunction>& TickFunction : TickFunctions)
 	{
 		TickFunction->SetTickFunctionEnable(bInEnabled);
 	}
 	EndTickFunction->SetTickFunctionEnable(bInEnabled);
+
+	RunState = bInEnabled ? ERunState::Running : ERunState::Paused;
 }
+
+void FAnimNextSchedulerEntry::LazyAllocateInstanceData()
+{
+	using namespace UE::AnimNext;
+
+	// Allocate instance data if required
+	if (!Context.InstanceData.IsValid())
+	{
+		Context.InstanceData = MakeUnique<FScheduleInstanceData>(Context, Schedule, Handle, this, MoveTemp(UserScopes));
+	}
+}
\ No newline at end of file
diff --git a/Engine/Plugins/Experimental/AnimNext/Source/AnimNext/Private/Scheduler/AnimNextSchedulerEntry.h b/Engine/Plugins/Experimental/AnimNext/Source/AnimNext/Private/Scheduler/AnimNextSchedulerEntry.h
--- a/Engine/Plugins/Experimental/AnimNext/Source/AnimNext/Private/Scheduler/AnimNextSchedulerEntry.h
+++ b/Engine/Plugins/Experimental/AnimNext/Source/AnimNext/Private/Scheduler/AnimNextSchedulerEntry.h
@@ -3,6 +3,7 @@
 #pragma once
 
 #include "CoreMinimal.h"
+#include "Scheduler/AnimNextSchedule.h"
 #include "Scheduler/ScheduleTickFunction.h"
 #include "Scheduler/ScheduleContext.h"
 #include "AnimNextSchedulerEntry.generated.h"
@@ -22,25 +23,41 @@ struct FAnimNextSchedulerEntry
 	GENERATED_BODY()
 
 	FAnimNextSchedulerEntry() = default;
-	FAnimNextSchedulerEntry(const UAnimNextSchedule* InSchedule, UObject* InObject, UE::AnimNext::FScheduleHandle InHandle, const TMap<FName, FAnimNextParameterCollection>& InUserScopes);
+	FAnimNextSchedulerEntry(const UAnimNextSchedule* InSchedule, UObject* InObject, UE::AnimNext::FScheduleHandle InHandle, const TMap<FName, FAnimNextParameterCollection>& InUserScopes, EAnimNextScheduleInitMethod InInitMethod);
 	~FAnimNextSchedulerEntry();
 
+	// Setup the entry
+	void Initialize();
+
 	// Used for pooling
 	void Invalidate();
 
 	// Enables/disables the ticking of this entry
 	void Enable(bool bInEnabled);
-	
-	UPROPERTY()
+
+	// Clears the bTickEvenWhenPaused flags of the entries' tick functions
+	void ClearTickFunctionPauseFlags();
+
+	// Allocate instance data if required
+	void LazyAllocateInstanceData();
+
+	UPROPERTY(Transient)
 	TObjectPtr<const UAnimNextSchedule> Schedule = nullptr;
 
 	// User scopes are copied into this entry on construction, but moved out later into instance data
 	// So will be invalid here after first run
-	UPROPERTY()
+	UPROPERTY(Transient)
 	TMap<FName, FAnimNextParameterCollection> UserScopes;
 
+	// Object this entry is bound to, valid only during schedule execution
+	UPROPERTY(Transient)
+	TObjectPtr<UObject> ResolvedObject = nullptr;
+
 	// Object this entry is bound to
-	TWeakObjectPtr<UObject> Object;
+	TWeakObjectPtr<UObject> WeakObject;
+
+	// Objects each of the schedule entries are bound to, if any
+	TArray<TWeakObjectPtr<UObject>> TargetObjects;
 
 	// Copy of the handle that represents this entry to client systems
 	UE::AnimNext::FScheduleHandle Handle;
@@ -57,6 +74,30 @@ struct FAnimNextSchedulerEntry
 
 	// Pre-allocated graph of tick functions
 	TArray<TUniquePtr<UE::AnimNext::FScheduleTickFunction>> TickFunctions;
+
+	enum class ERunState
+	{
+		None,
+
+		CreatingTasks,
+
+		BindingTasks,
+
+		RunningInitialUpdate,
+
+		Running,
+
+		Paused,
+	};
+
+	// Current running state
+	ERunState RunState = ERunState::None;
+
+	// How this entry initializes
+	EAnimNextScheduleInitMethod InitMethod = EAnimNextScheduleInitMethod::InitializeAndPauseInEditor;
+
+	// Whether this represents an editor object 
+	bool bIsEditor = false;
 };
 
 template<>
diff --git a/Engine/Plugins/Experimental/AnimNext/Source/AnimNext/Private/Scheduler/AnimNextSchedulerWorldSubsystem.cpp b/Engine/Plugins/Experimental/AnimNext/Source/AnimNext/Private/Scheduler/AnimNextSchedulerWorldSubsystem.cpp
--- a/Engine/Plugins/Experimental/AnimNext/Source/AnimNext/Private/Scheduler/AnimNextSchedulerWorldSubsystem.cpp
+++ b/Engine/Plugins/Experimental/AnimNext/Source/AnimNext/Private/Scheduler/AnimNextSchedulerWorldSubsystem.cpp
@@ -97,7 +97,7 @@ void UAnimNextSchedulerWorldSubsystem::FlushPendingActions()
 	}
 }
 
-UE::AnimNext::FScheduleHandle UAnimNextSchedulerWorldSubsystem::AcquireHandle(UObject* InObject, UAnimNextSchedule* InSchedule, const TMap<FName, FAnimNextParameterCollection>& InUserScopes)
+UE::AnimNext::FScheduleHandle UAnimNextSchedulerWorldSubsystem::AcquireHandle(UObject* InObject, UAnimNextSchedule* InSchedule, const TMap<FName, FAnimNextParameterCollection>& InUserScopes, EAnimNextScheduleInitMethod InInitMethod)
 {
 	using namespace UE::AnimNext;
 
@@ -111,16 +111,18 @@ UE::AnimNext::FScheduleHandle UAnimNextSchedulerWorldSubsystem::AcquireHandle(UO
 		Handle.Index = FreeEntryIndices.Last();
 		Handle.SerialNumber = ++GEntrySerialNumber;
 		FreeEntryIndices.Pop(false);
-		new (Entries[Handle.Index].Get()) FAnimNextSchedulerEntry(InSchedule, InObject, Handle, InUserScopes);
+		new (Entries[Handle.Index].Get()) FAnimNextSchedulerEntry(InSchedule, InObject, Handle, InUserScopes, InInitMethod);
 	}
 	// Otherwise append a new entry
 	else
 	{
 		Handle.Index = Entries.Num();
 		Handle.SerialNumber = ++GEntrySerialNumber;
-		Entries.Emplace(MakeUnique<FAnimNextSchedulerEntry>(InSchedule, InObject, Handle, InUserScopes));
+		Entries.Emplace(MakeUnique<FAnimNextSchedulerEntry>(InSchedule, InObject, Handle, InUserScopes, InInitMethod));
 	}
 
+	Entries[Handle.Index]->Initialize();
+
 	// Skip 'invalid' 0 serial number
 	if (GEntrySerialNumber == 0)
 	{
@@ -177,7 +179,7 @@ void UAnimNextSchedulerWorldSubsystem::QueueTask(UE::AnimNext::FScheduleHandle I
 		// TODO: Only supporting scope tasks for now
 		const FAnimNextScheduleParamScopeEntryTask* FoundScope = Entry->Schedule->ParamScopeEntryTasks.FindByPredicate([&InScheduleTaskName](const FAnimNextScheduleParamScopeEntryTask& InTask)
 		{
-			return InTask.Name == InScheduleTaskName;
+			return InTask.Scope == InScheduleTaskName;
 		});
 
 		if (FoundScope)
diff --git a/Engine/Plugins/Experimental/AnimNext/Source/AnimNext/Private/Scheduler/AnimNextSchedulerWorldSubsystem.h b/Engine/Plugins/Experimental/AnimNext/Source/AnimNext/Private/Scheduler/AnimNextSchedulerWorldSubsystem.h
--- a/Engine/Plugins/Experimental/AnimNext/Source/AnimNext/Private/Scheduler/AnimNextSchedulerWorldSubsystem.h
+++ b/Engine/Plugins/Experimental/AnimNext/Source/AnimNext/Private/Scheduler/AnimNextSchedulerWorldSubsystem.h
@@ -64,7 +64,7 @@ class UAnimNextSchedulerWorldSubsystem : public UWorldSubsystem
 	void FlushPendingActions();
 
 	// Acquire a handle that binds a schedule with the supplied parameters
-	UE::AnimNext::FScheduleHandle AcquireHandle(UObject* InObject, UAnimNextSchedule* InSchedule, const TMap<FName, FAnimNextParameterCollection>& InUserScopes);
+	UE::AnimNext::FScheduleHandle AcquireHandle(UObject* InObject, UAnimNextSchedule* InSchedule, const TMap<FName, FAnimNextParameterCollection>& InUserScopes, EAnimNextScheduleInitMethod InInitMethod);
 
 	// Release an already acquired handle
 	// The full release of the binding referenced by the handle will be deferred after this call is made
diff --git a/Engine/Plugins/Experimental/AnimNext/Source/AnimNext/Private/Scheduler/Ports/AnimNextSchedulePort_AnimNextMeshComponentPose.cpp b/Engine/Plugins/Experimental/AnimNext/Source/AnimNext/Private/Scheduler/Ports/AnimNextSchedulePort_AnimNextMeshComponentPose.cpp
new file mode 100644
--- /dev/null
+++ b/Engine/Plugins/Experimental/AnimNext/Source/AnimNext/Private/Scheduler/Ports/AnimNextSchedulePort_AnimNextMeshComponentPose.cpp
@@ -0,0 +1,84 @@
+﻿// Copyright Epic Games, Inc. All Rights Reserved.
+
+#include "AnimNextSchedulePort_AnimNextMeshComponentPose.h"
+
+#include "GenerationTools.h"
+#include "Graph/AnimNext_LODPose.h"
+#include "Param/ParamStack.h"
+#include "Scheduler/ScheduleTermContext.h"
+#include "Components/SkeletalMeshComponent.h"
+#include "Engine/SkeletalMesh.h"
+#include "AnimNextStats.h"
+#include "Component/AnimNextMeshComponent.h"
+
+DEFINE_STAT(STAT_AnimNext_Port_SkeletalMeshComponent);
+
+const UE::AnimNext::FParamId UAnimNextSchedulePort_AnimNextMeshComponentPose::InputId("UE_Internal_AnimNextMeshComponentPose_Input");
+
+void UAnimNextSchedulePort_AnimNextMeshComponentPose::Run(const UE::AnimNext::FScheduleTermContext& InContext) const
+{
+	SCOPE_CYCLE_COUNTER(STAT_AnimNext_Port_SkeletalMeshComponent);
+	
+	using namespace UE::AnimNext;
+
+	const FParamStack& ParamStack = FParamStack::Get();
+
+	static FParamId ComponentParamId("UE_AnimNextMeshComponent");
+	const TObjectPtr<UAnimNextMeshComponent>* ComponentPtr = ParamStack.GetParamPtr<TObjectPtr<UAnimNextMeshComponent>>(ComponentParamId);
+	if(ComponentPtr == nullptr)
+	{
+		return;
+	}
+
+	const FAnimNextGraphLODPose* InputPose = InContext.GetLayerHandle().GetParamPtr<FAnimNextGraphLODPose>(InputId);
+	if(InputPose == nullptr)
+	{
+		return;
+	}
+
+	if(!InputPose->LODPose.IsValid())
+	{
+		return;
+	}
+
+	static const FParamId ReferencePoseParamId("UE_AnimNextMeshComponent_ReferencePose");
+	const FAnimNextGraphReferencePose* GraphReferencePose = ParamStack.GetParamPtr<FAnimNextGraphReferencePose>(ReferencePoseParamId);
+	if(GraphReferencePose == nullptr)
+	{
+		return;
+	}
+
+	UAnimNextMeshComponent* Component = *ComponentPtr;
+
+	USkeletalMesh* SkeletalMesh = Component->GetSkeletalMeshAsset();
+	if(SkeletalMesh == nullptr)
+	{
+		return;
+	}
+
+	FMemMark MemMark(FMemStack::Get());
+
+	TArray<FTransform, TMemStackAllocator<>> LocalSpaceTransforms;
+	LocalSpaceTransforms.SetNumUninitialized(SkeletalMesh->GetRefSkeleton().GetNum());
+
+	// Map LOD pose into local-space scratch buffer
+	FGenerationTools::RemapPose(InputPose->LODPose, LocalSpaceTransforms);
+
+	// Convert and dispatch to renderer
+	Component->CompleteAndDispatch(
+		GraphReferencePose->ReferencePose->GetParentIndices(),
+		GraphReferencePose->ReferencePose->GetLODBoneIndexToMeshBoneIndexMap(InputPose->LODPose.LODLevel),
+		LocalSpaceTransforms);
+}
+
+TConstArrayView<UE::AnimNext::FScheduleTerm> UAnimNextSchedulePort_AnimNextMeshComponentPose::GetTerms() const
+{
+	using namespace UE::AnimNext;
+
+	static const FScheduleTerm Terms[] =
+	{
+		FScheduleTerm(InputId, FAnimNextParamType::GetType<FAnimNextGraphLODPose>(), EScheduleTermDirection::Input)
+	};
+
+	return Terms;
+}
diff --git a/Engine/Plugins/Experimental/AnimNext/Source/AnimNext/Private/Scheduler/Ports/AnimNextSchedulePort_AnimNextMeshComponentPose.h b/Engine/Plugins/Experimental/AnimNext/Source/AnimNext/Private/Scheduler/Ports/AnimNextSchedulePort_AnimNextMeshComponentPose.h
new file mode 100644
--- /dev/null
+++ b/Engine/Plugins/Experimental/AnimNext/Source/AnimNext/Private/Scheduler/Ports/AnimNextSchedulePort_AnimNextMeshComponentPose.h
@@ -0,0 +1,19 @@
+﻿// Copyright Epic Games, Inc. All Rights Reserved.
+
+#pragma once
+
+#include "Scheduler/AnimNextSchedulePort.h"
+#include "AnimNextSchedulePort_AnimNextMeshComponentPose.generated.h"
+
+UCLASS(DisplayName = "AnimNext Mesh Component Pose")
+class UAnimNextSchedulePort_AnimNextMeshComponentPose : public UAnimNextSchedulePort
+{
+	GENERATED_BODY()
+
+	// UAnimNextSchedulePort interface
+	virtual void Run(const UE::AnimNext::FScheduleTermContext& InContext) const override;
+	virtual TConstArrayView<UE::AnimNext::FScheduleTerm> GetTerms() const override;
+
+	// ID of our input
+	static const UE::AnimNext::FParamId InputId;
+};
\ No newline at end of file
diff --git a/Engine/Plugins/Experimental/AnimNext/Source/AnimNext/Private/Scheduler/ScheduleInstanceData.cpp b/Engine/Plugins/Experimental/AnimNext/Source/AnimNext/Private/Scheduler/ScheduleInstanceData.cpp
--- a/Engine/Plugins/Experimental/AnimNext/Source/AnimNext/Private/Scheduler/ScheduleInstanceData.cpp
+++ b/Engine/Plugins/Experimental/AnimNext/Source/AnimNext/Private/Scheduler/ScheduleInstanceData.cpp
@@ -5,6 +5,10 @@
 #include "ScheduleHandle.h"
 #include "Scheduler/AnimNextSchedulerWorldSubsystem.h"
 #include "Scheduler/AnimNextSchedulerEntry.h"
+#include "Scheduler/AnimNextSchedulePort.h"
+#include "AnimNextStats.h"
+
+DEFINE_STAT(STAT_AnimNext_CreateInstanceData);
 
 namespace UE::AnimNext
 {
@@ -14,9 +18,12 @@ FScheduleInstanceData::FScheduleInstanceData(const FScheduleContext& InScheduleC
 	, Entry(InCurrentEntry) 
 	, UserScopes(MoveTemp(InUserScopes))
 {
+	SCOPE_CYCLE_COUNTER(STAT_AnimNext_CreateInstanceData);
+
 	// Preallocate data for all scopes & graphs in the schedule
 	ScopeCaches.SetNum(InSchedule->ParamScopeEntryTasks.Num());
 	GraphInstanceData.SetNum(InSchedule->Tasks.Num());
+	GraphInputLayers.SetNum(InSchedule->Tasks.Num());
 
 	// Setup param stack graph
 	RootParamStack = InCurrentEntry->RootParamStack;
@@ -43,6 +50,50 @@ FScheduleInstanceData::FScheduleInstanceData(const FScheduleContext& InScheduleC
 		const TSharedPtr<FParamStack> ParentStack = ScopeEntryTask.ParamParentScopeIndex != MAX_uint32 ? ParamStacks[ScopeEntryTask.ParamParentScopeIndex] : RootParamStack;
 		ParamStacks[ScopeEntryTask.ParamScopeIndex]->SetParent(ParentStack);
 	}
+
+	// Duplicate intermediate data area
+	IntermediatesData = InSchedule->IntermediatesData;
+
+	// Make a hosting layer for the intermediates
+	IntermediatesLayer = FParamStack::MakeReferenceLayer(IntermediatesData);
+
+	// Setup remapped intermediate data layers for graph and port tasks
+	GraphTermLayers.Reserve(InSchedule->Tasks.Num());
+	for (const FAnimNextScheduleGraphTask& Task : InSchedule->Tasks)
+	{
+		TConstArrayView<FScheduleTerm> Terms = Task.Graph->GetTerms();
+		check(Task.Terms.Num() == Terms.Num());
+
+		TMap<FName, FName> Mapping;
+		Mapping.Reserve(Task.Terms.Num());
+		for(int32 TermIndex = 0; TermIndex < Task.Terms.Num(); ++TermIndex)
+		{
+			uint32 IntermediateTermIndex = Task.Terms[TermIndex];
+			const FPropertyBagPropertyDesc& PropertyDesc = IntermediatesData.GetPropertyBagStruct()->GetPropertyDescs()[IntermediateTermIndex];
+			Mapping.Add(PropertyDesc.Name, Terms[TermIndex].GetName());
+		}
+
+		GraphTermLayers.Add(FParamStack::MakeRemappedLayer(IntermediatesLayer, Mapping));
+	}
+
+	PortTermLayers.Reserve(InSchedule->Ports.Num());
+	for (const FAnimNextSchedulePortTask& PortTask : InSchedule->Ports)
+	{
+		UAnimNextSchedulePort* CDO = PortTask.Port->GetDefaultObject<UAnimNextSchedulePort>();
+		TConstArrayView<FScheduleTerm> Terms = CDO->GetTerms();
+		check(PortTask.Terms.Num() == Terms.Num());
+
+		TMap<FName, FName> Mapping;
+		Mapping.Reserve(PortTask.Terms.Num());
+		for(int32 TermIndex = 0; TermIndex < PortTask.Terms.Num(); ++TermIndex)
+		{
+			uint32 IntermediateTermIndex = PortTask.Terms[TermIndex];
+			const FPropertyBagPropertyDesc& PropertyDesc = IntermediatesData.GetPropertyBagStruct()->GetPropertyDescs()[IntermediateTermIndex];
+			Mapping.Add(PropertyDesc.Name, Terms[TermIndex].GetName());
+		}
+
+		PortTermLayers.Add(FParamStack::MakeRemappedLayer(IntermediatesLayer, Mapping));
+	}
 }
 
 void FScheduleInstanceData::AddReferencedObjects(FReferenceCollector& Collector)
@@ -58,4 +109,16 @@ FString FScheduleInstanceData::GetReferencerName() const
 	return TEXT("AnimNextInstanceData");
 }
 
+TSharedPtr<FParamStack> FScheduleInstanceData::GetParamStack(uint32 InIndex) const
+{
+	if(InIndex == MAX_uint32)
+	{
+		return RootParamStack;
+	}
+	else
+	{
+		return ParamStacks[InIndex];
+	}
+}
+
 }
\ No newline at end of file
diff --git a/Engine/Plugins/Experimental/AnimNext/Source/AnimNext/Private/Scheduler/ScheduleInstanceData.h b/Engine/Plugins/Experimental/AnimNext/Source/AnimNext/Private/Scheduler/ScheduleInstanceData.h
--- a/Engine/Plugins/Experimental/AnimNext/Source/AnimNext/Private/Scheduler/ScheduleInstanceData.h
+++ b/Engine/Plugins/Experimental/AnimNext/Source/AnimNext/Private/Scheduler/ScheduleInstanceData.h
@@ -3,6 +3,7 @@
 #pragma once
 
 #include "CoreMinimal.h"
+#include "LODPose.h"
 #include "ScheduleHandle.h"
 #include "Param/AnimNextParameterCollection.h"
 #include "Graph/AnimNextGraph.h"
@@ -29,6 +30,9 @@ struct FScheduleInstanceData : public FGCObject
 	virtual void AddReferencedObjects(FReferenceCollector& Collector) override;
 	virtual FString GetReferencerName() const override;
 
+	// Get the appropriate params stack given the ID
+	TSharedPtr<FParamStack> GetParamStack(uint32 InIndex) const;
+	
 	// Handle to the currently executing entry in the schedule
 	FScheduleHandle Handle;
 
@@ -62,6 +66,21 @@ struct FScheduleInstanceData : public FGCObject
 
 	// Graph instance data for each graph task
 	TArray<FAnimNextGraphInstance> GraphInstanceData;
+
+	// Layer handles for translating schedule terms to graph inputs
+	TArray<FParamStackLayerHandle> GraphInputLayers;
+
+	// Intermediate data area
+	FInstancedPropertyBag IntermediatesData;
+
+	// Layer for intermediates data
+	FParamStackLayerHandle IntermediatesLayer;
+
+	// Remapped data layers for each graph
+	TArray<FParamStackLayerHandle> GraphTermLayers;
+
+	// Remapped data layers for each port
+	TArray<FParamStackLayerHandle> PortTermLayers;
 };
 
 }
\ No newline at end of file
diff --git a/Engine/Plugins/Experimental/AnimNext/Source/AnimNext/Private/Scheduler/SchedulePortContext.cpp b/Engine/Plugins/Experimental/AnimNext/Source/AnimNext/Private/Scheduler/SchedulePortContext.cpp
new file mode 100644
--- /dev/null
+++ b/Engine/Plugins/Experimental/AnimNext/Source/AnimNext/Private/Scheduler/SchedulePortContext.cpp
@@ -0,0 +1,15 @@
+﻿// Copyright Epic Games, Inc. All Rights Reserved.
+
+#include "Scheduler/SchedulePortContext.h"
+
+#include "Param/ParamStackLayer.h"
+
+namespace UE::AnimNext
+{
+
+FParamResult FSchedulePortContext::GetParamData(FParamId InId, FParamTypeHandle InTypeHandle, TConstArrayView<uint8>& OutParamData, FParamTypeHandle& OutParamTypeHandle, FParamCompatibility InRequiredCompatibility) const
+{
+	return ResultLayer->GetParamData(InId, InTypeHandle, OutParamData, OutParamTypeHandle, InRequiredCompatibility);
+}
+
+}
diff --git a/Engine/Plugins/Experimental/AnimNext/Source/AnimNext/Private/Scheduler/ScheduleTickFunction.cpp b/Engine/Plugins/Experimental/AnimNext/Source/AnimNext/Private/Scheduler/ScheduleTickFunction.cpp
--- a/Engine/Plugins/Experimental/AnimNext/Source/AnimNext/Private/Scheduler/ScheduleTickFunction.cpp
+++ b/Engine/Plugins/Experimental/AnimNext/Source/AnimNext/Private/Scheduler/ScheduleTickFunction.cpp
@@ -15,11 +15,9 @@ namespace UE::AnimNext
 
 void FScheduleBeginTickFunction::ExecuteTick(float DeltaTime, ELevelTick TickType, ENamedThreads::Type CurrentThread, const FGraphEventRef& MyCompletionGraphEvent)
 {
-	// Allocate instance data if required
-	if (!Entry.Context.InstanceData.IsValid())
-	{
-		Entry.Context.InstanceData = MakeUnique<FScheduleInstanceData>(Entry.Context, Entry.Schedule, Entry.Handle, &Entry, MoveTemp(Entry.UserScopes));
-	}
+	Entry.LazyAllocateInstanceData();
+
+	Entry.ResolvedObject = Entry.WeakObject.Get();
 }
 
 FString FScheduleBeginTickFunction::DiagnosticMessage()
@@ -29,6 +27,43 @@ FString FScheduleBeginTickFunction::DiagnosticMessage()
 
 void FScheduleEndTickFunction::ExecuteTick(float DeltaTime, ELevelTick TickType, ENamedThreads::Type CurrentThread, const FGraphEventRef& MyCompletionGraphEvent)
 {
+	if(Entry.RunState == FAnimNextSchedulerEntry::ERunState::RunningInitialUpdate)
+	{
+		if( Entry.InitMethod == EAnimNextScheduleInitMethod::InitializeAndPause
+#if WITH_EDITOR
+			|| (Entry.InitMethod == EAnimNextScheduleInitMethod::InitializeAndPauseInEditor && Entry.bIsEditor)
+#endif
+			)
+		{
+			// Queue task to disable our tick functions now we have performed our initial update
+			FFunctionGraphTask::CreateAndDispatchWhenReady(
+			[this]()
+			{
+				check(IsInGameThread());
+
+				Entry.Enable(false);
+				Entry.RunState = FAnimNextSchedulerEntry::ERunState::Paused;
+			},
+			TStatId(),
+			nullptr,
+			ENamedThreads::GameThread);
+		}
+	}
+	else
+	{
+		FFunctionGraphTask::CreateAndDispatchWhenReady(
+		[this]()
+		{
+			check(IsInGameThread());
+
+			Entry.RunState = FAnimNextSchedulerEntry::ERunState::Running;
+		},
+		TStatId(),
+		nullptr,
+		ENamedThreads::GameThread);
+	}
+
+	Entry.ResolvedObject = nullptr;
 }
 
 FString FScheduleEndTickFunction::DiagnosticMessage()
@@ -40,94 +75,127 @@ void FScheduleTickFunction::ExecuteTick(float DeltaTime, ELevelTick TickType, EN
 {
 	FScheduleContext::AttachToCurrentThread(ScheduleContext);
 
-	while (!PreExecuteTasks.IsEmpty())
-	{
-		TOptional<TUniqueFunction<void(const UE::AnimNext::FScheduleContext&)>> Function = PreExecuteTasks.Dequeue();
-		check(Function.IsSet());
-		Function.GetValue()(ScheduleContext);
-	}
+	RunSchedule(Instructions, TargetObjects,
+		[this]()
+		{
+			while (!PreExecuteTasks.IsEmpty())
+			{
+				TOptional<TUniqueFunction<void(const UE::AnimNext::FScheduleContext&)>> Function = PreExecuteTasks.Dequeue();
+				check(Function.IsSet());
+				Function.GetValue()(ScheduleContext);
+			}
+		}, 
+		[this]()
+		{
+			while (!PostExecuteTasks.IsEmpty())
+			{
+				TOptional<TUniqueFunction<void(const UE::AnimNext::FScheduleContext&)>> Function = PostExecuteTasks.Dequeue();
+				check(Function.IsSet());
+				Function.GetValue()(ScheduleContext);
+			}
+		});
+
+	FScheduleContext::DetachFromCurrentThread();
+}
+
+void FScheduleTickFunction::RunSchedule(TConstArrayView<FAnimNextScheduleInstruction> InInstructions)
+{
+	RunSchedule(InInstructions, {}, [](){}, [](){});
+}
+
+void FScheduleTickFunction::RunSchedule(TConstArrayView<FAnimNextScheduleInstruction> InInstructions, TConstArrayView<TWeakObjectPtr<UObject>> InTargetObjects, TFunctionRef<void(void)> InPreExecuteScope, TFunctionRef<void(void)> InPostExecuteScope)
+{
+	const FScheduleContext& ScheduleContext = FScheduleContext::Get();
 
 	const UAnimNextSchedule* Schedule = ScheduleContext.Schedule;
 
-	switch (Instruction.Opcode)
+	int32 InstructionIndex = 0;
+	while(InstructionIndex < InInstructions.Num())
 	{
-	case EAnimNextScheduleScheduleOpcode::RunTask:
+		const FAnimNextScheduleInstruction& Instruction = InInstructions[InstructionIndex];
+		switch (Instruction.Opcode)
 		{
-			uint32 TaskIndex = Instruction.Operand;
-			FScheduleInstanceData& InstanceData = ScheduleContext.GetInstanceData();
-			FParamStack::AttachToCurrentThread(InstanceData.ParamStacks[Schedule->Tasks[TaskIndex].ParamScopeIndex]);
+		case EAnimNextScheduleScheduleOpcode::RunTask:
+			{
+				uint32 TaskIndex = Instruction.Operand;
+				FScheduleInstanceData& InstanceData = ScheduleContext.GetInstanceData();
+				FParamStack::AttachToCurrentThread(InstanceData.GetParamStack(Schedule->Tasks[TaskIndex].ParamScopeIndex));
 
-			Schedule->Tasks[TaskIndex].RunGraph(ScheduleContext);
+				Schedule->Tasks[TaskIndex].RunGraph(ScheduleContext);
 
-			FParamStack::DetachFromCurrentThread();
-			break;
-		}
-	case EAnimNextScheduleScheduleOpcode::BeginRunExternalTask:
-		{
-			if (const UObject* Object = TargetObject.Get())
+				FParamStack::DetachFromCurrentThread();
+				break;
+			}
+		case EAnimNextScheduleScheduleOpcode::BeginRunExternalTask:
 			{
-				FScheduleInstanceData& InstanceData = ScheduleContext.GetInstanceData();
-				uint32 ExternalTaskIndex = Instruction.Operand;
-				FParamStack::AddForPendingObject(Object, InstanceData.ParamStacks[Schedule->ExternalTasks[ExternalTaskIndex].ParamScopeIndex]);
+				if(InTargetObjects.Num() > 0)
+				{
+					if (const UObject* Object = InTargetObjects[InstructionIndex].Get())
+					{
+						FScheduleInstanceData& InstanceData = ScheduleContext.GetInstanceData();
+						uint32 ExternalTaskIndex = Instruction.Operand;
+						FParamStack::AddForPendingObject(Object, InstanceData.GetParamStack(Schedule->ExternalTasks[ExternalTaskIndex].ParamScopeIndex));
+					}
+				}
+				break;
 			}
-			break;
-		}
-	case EAnimNextScheduleScheduleOpcode::EndRunExternalTask:
-		{
-			if (const UObject* Object = TargetObject.Get())
+		case EAnimNextScheduleScheduleOpcode::EndRunExternalTask:
+			{
+				if(InTargetObjects.Num() > 0)
+				{
+					if (const UObject* Object = InTargetObjects[InstructionIndex].Get())
+					{
+						FScheduleInstanceData& InstanceData = ScheduleContext.GetInstanceData();
+						uint32 ExternalTaskIndex = Instruction.Operand;
+						FParamStack::RemoveForPendingObject(Object);
+					}
+				}
+				break;
+			}
+		case EAnimNextScheduleScheduleOpcode::RunPort:
 			{
+				uint32 PortIndex = Instruction.Operand;
 				FScheduleInstanceData& InstanceData = ScheduleContext.GetInstanceData();
-				uint32 ExternalTaskIndex = Instruction.Operand;
-				FParamStack::RemoveForPendingObject(Object);
+				FParamStack::AttachToCurrentThread(InstanceData.GetParamStack(Schedule->Ports[PortIndex].ParamScopeIndex));
+
+				Schedule->Ports[PortIndex].RunPort(ScheduleContext);
+
+				FParamStack::DetachFromCurrentThread();
+				break;
 			}
-			break;
-		}
-	case EAnimNextScheduleScheduleOpcode::RunPort:
-		{
-			uint32 PortIndex = Instruction.Operand;
-			FScheduleInstanceData& InstanceData = ScheduleContext.GetInstanceData();
-			FParamStack::AttachToCurrentThread(InstanceData.ParamStacks[Schedule->Ports[PortIndex].ParamScopeIndex]);
+		case EAnimNextScheduleScheduleOpcode::RunParamScopeEntry:
+			{
+				InPreExecuteScope();
 
-			Schedule->Ports[PortIndex].RunPort(ScheduleContext);
+				uint32 ScopeEntryIndex = Instruction.Operand;
+				FScheduleInstanceData& InstanceData = ScheduleContext.GetInstanceData();
+				FParamStack::AttachToCurrentThread(InstanceData.GetParamStack(Schedule->ParamScopeEntryTasks[ScopeEntryIndex].ParamScopeIndex));
 
-			FParamStack::DetachFromCurrentThread();
-			break;
-		}
-	case EAnimNextScheduleScheduleOpcode::RunParamScopeEntry:
-		{
-			uint32 ScopeEntryIndex = Instruction.Operand;
-			FScheduleInstanceData& InstanceData = ScheduleContext.GetInstanceData();
-			FParamStack::AttachToCurrentThread(InstanceData.ParamStacks[Schedule->ParamScopeEntryTasks[ScopeEntryIndex].ParamScopeIndex]);
+				Schedule->ParamScopeEntryTasks[ScopeEntryIndex].RunParamScopeEntry(ScheduleContext);
 
-			Schedule->ParamScopeEntryTasks[ScopeEntryIndex].RunParamScopeEntry(ScheduleContext);
+				FParamStack::DetachFromCurrentThread();
 
-			FParamStack::DetachFromCurrentThread();
-			break;
-		}
-	case EAnimNextScheduleScheduleOpcode::RunParamScopeExit:
-		{
-			uint32 ScopeExitIndex = Instruction.Operand;
-			FScheduleInstanceData& InstanceData = ScheduleContext.GetInstanceData();
-			FParamStack::AttachToCurrentThread(InstanceData.ParamStacks[Schedule->ParamScopeExitTasks[ScopeExitIndex].ParamScopeIndex]);
+				InPostExecuteScope();
 
-			Schedule->ParamScopeExitTasks[ScopeExitIndex].RunParamScopeExit(ScheduleContext);
+				break;
+			}
+		case EAnimNextScheduleScheduleOpcode::RunParamScopeExit:
+			{
+				uint32 ScopeExitIndex = Instruction.Operand;
+				FScheduleInstanceData& InstanceData = ScheduleContext.GetInstanceData();
+				FParamStack::AttachToCurrentThread(InstanceData.GetParamStack(Schedule->ParamScopeExitTasks[ScopeExitIndex].ParamScopeIndex));
+
+				Schedule->ParamScopeExitTasks[ScopeExitIndex].RunParamScopeExit(ScheduleContext);
 
-			FParamStack::DetachFromCurrentThread();
+				FParamStack::DetachFromCurrentThread();
+				break;
+			}
+		default:
 			break;
 		}
-	default:
-		check(false);	// Unimplemented instruction
-		break;
-	}
 
-	while (!PostExecuteTasks.IsEmpty())
-	{
-		TOptional<TUniqueFunction<void(const UE::AnimNext::FScheduleContext&)>> Function = PostExecuteTasks.Dequeue();
-		check(Function.IsSet());
-		Function.GetValue()(ScheduleContext);
+		InstructionIndex++;
 	}
-
-	FScheduleContext::DetachFromCurrentThread();
 }
 
 FString FScheduleTickFunction::DiagnosticMessage()
diff --git a/Engine/Plugins/Experimental/AnimNext/Source/AnimNext/Private/Scheduler/ScheduleTickFunction.h b/Engine/Plugins/Experimental/AnimNext/Source/AnimNext/Private/Scheduler/ScheduleTickFunction.h
--- a/Engine/Plugins/Experimental/AnimNext/Source/AnimNext/Private/Scheduler/ScheduleTickFunction.h
+++ b/Engine/Plugins/Experimental/AnimNext/Source/AnimNext/Private/Scheduler/ScheduleTickFunction.h
@@ -53,12 +53,12 @@ struct FScheduleEndTickFunction : public FTickFunction
 	FAnimNextSchedulerEntry& Entry;
 };
 
-// TODO: Run a linear slice of instructions rather than have each instruction a separate tick function
 struct FScheduleTickFunction : public FTickFunction
 {
-	FScheduleTickFunction(const FScheduleContext& InScheduleContext, const FAnimNextScheduleInstruction& InInstruction)
+	FScheduleTickFunction(const FScheduleContext& InScheduleContext, TConstArrayView<FAnimNextScheduleInstruction> InInstructions, TConstArrayView<TWeakObjectPtr<UObject>> InTargetObjects)
 		: ScheduleContext(InScheduleContext)
-		, Instruction(InInstruction)
+		, Instructions(InInstructions)
+		, TargetObjects(InTargetObjects)
 	{
 		bCanEverTick = true;
 		bStartWithTickEnabled = true;
@@ -69,10 +69,15 @@ struct FScheduleTickFunction : public FTickFunction
 	virtual void ExecuteTick(float DeltaTime, ELevelTick TickType, ENamedThreads::Type CurrentThread, const FGraphEventRef& MyCompletionGraphEvent) override;
 	virtual FString DiagnosticMessage() override;
 
+	// Static helpers for running a slice of instructions
+	// This can be called standalone to run a whole schedule in one call
+	static void RunSchedule(TConstArrayView<FAnimNextScheduleInstruction> InInstructions);
+	static void RunSchedule(TConstArrayView<FAnimNextScheduleInstruction> InInstructions, TConstArrayView<TWeakObjectPtr<UObject>> InTargetObjects, TFunctionRef<void(void)> InPreExecuteScope, TFunctionRef<void(void)> InPostExecuteScope);
+
 	const FScheduleContext& ScheduleContext;
-	const FAnimNextScheduleInstruction& Instruction;
+	TConstArrayView<FAnimNextScheduleInstruction> Instructions;
 	TArray<FTickPrerequisite> Subsequents;
-	TWeakObjectPtr<UObject> TargetObject;
+	TConstArrayView<TWeakObjectPtr<UObject>> TargetObjects;
 	TSpscQueue<TUniqueFunction<void(const UE::AnimNext::FScheduleContext&)>> PreExecuteTasks;
 	TSpscQueue<TUniqueFunction<void(const UE::AnimNext::FScheduleContext&)>> PostExecuteTasks;
 };
diff --git a/Engine/Plugins/Experimental/AnimNext/Source/AnimNext/Private/Scheduler/Scheduler.cpp b/Engine/Plugins/Experimental/AnimNext/Source/AnimNext/Private/Scheduler/Scheduler.cpp
--- a/Engine/Plugins/Experimental/AnimNext/Source/AnimNext/Private/Scheduler/Scheduler.cpp
+++ b/Engine/Plugins/Experimental/AnimNext/Source/AnimNext/Private/Scheduler/Scheduler.cpp
@@ -19,25 +19,6 @@ namespace UE::AnimNext
 
 struct FSchedulerImpl
 {
-	void RegisterBuiltInPortDefinitions()
-	{
-/*		RegisteredPortDefinitions.Add(
-			Ports::SkeletalMeshComponentPose.Name,
-			FSchedulePortDefinition(
-				Ports::SkeletalMeshComponentPose.Name,
-				FParamTypeHandle::GetHandle<TArray<FAnimNextLODPose>>(),
-				FParamTypeHandle::GetHandle<TArray<FTransform>>(),
-				[](const FSchedulePortAdapterContext& InContext)
-				{
-					check(InContext.InputType == FParamTypeHandle::GetHandle<FAnimNextLODPose>());
-					check(InContext.OutputType == FParamTypeHandle::GetHandle<TArray<FTransform>>());
-
-					USkeletalMeshComponent* OutputObject = CastChecked<USkeletalMeshComponent>(InContext.OutputObject);
-				
-					// TODO: Perform pose transpose
-				}));*/
-	}
-
 	FDelegateHandle OnWorldPreActorTickHandle;
 
 	TMap<FName, FSchedulePortDefinition> RegisteredPortDefinitions;
@@ -49,8 +30,6 @@ static FSchedulerImpl Impl;
 
 void FScheduler::Init()
 {
-	Impl.RegisterBuiltInPortDefinitions();
-
 	// Kick off root task at the start of each world tick
 	Impl.OnWorldPreActorTickHandle = FWorldDelegates::OnWorldPreActorTick.AddLambda([](UWorld* InWorld, ELevelTick InTickType, float InDeltaSeconds)
 	{
@@ -71,7 +50,7 @@ void FScheduler::Destroy()
 	FWorldDelegates::OnWorldPreActorTick.Remove(Impl.OnWorldPreActorTickHandle);
 }
 
-FScheduleHandle FScheduler::AcquireHandle(UObject* InObject, UAnimNextSchedule* InSchedule, const TMap<FName, FAnimNextParameterCollection>& InUserScopes)
+FScheduleHandle FScheduler::AcquireHandle(UObject* InObject, UAnimNextSchedule* InSchedule, const TMap<FName, FAnimNextParameterCollection>& InUserScopes, EAnimNextScheduleInitMethod InInitMethod)
 {
 	FScheduleHandle Handle;
 
@@ -89,54 +68,85 @@ FScheduleHandle FScheduler::AcquireHandle(UObject* InObject, UAnimNextSchedule*
 	}
 
 	UWorld* World = InObject->GetWorld();
+	if(World == nullptr)
+	{
+		return FScheduleHandle();
+	}
+
 	UAnimNextSchedulerWorldSubsystem* Subsystem = World->GetSubsystem<UAnimNextSchedulerWorldSubsystem>();
-	return Subsystem->AcquireHandle(InObject, InSchedule, InUserScopes);
+	if(Subsystem == nullptr)
+	{
+		return FScheduleHandle();
+	}
+
+	return Subsystem->AcquireHandle(InObject, InSchedule, InUserScopes, InInitMethod);
 }
 
 void FScheduler::ReleaseHandle(UObject* InObject, FScheduleHandle& InHandle)
 {
-	if(InHandle.IsValid())
+	if(!InHandle.IsValid())
+	{
+		return;
+	}
+	
+	UWorld* World = InObject->GetWorld();
+	if(World == nullptr)
+	{
+		return;
+	}
+
+	UAnimNextSchedulerWorldSubsystem* Subsystem = World->GetSubsystem<UAnimNextSchedulerWorldSubsystem>();
+	if(Subsystem == nullptr)
 	{
-		UWorld* World = InObject->GetWorld();
-		UAnimNextSchedulerWorldSubsystem* Subsystem = World->GetSubsystem<UAnimNextSchedulerWorldSubsystem>();
-		Subsystem->ReleaseHandle(InHandle);
-		InHandle.Invalidate();
+		return;
 	}
+
+	Subsystem->ReleaseHandle(InHandle);
+	InHandle.Invalidate();
 }
 
 void FScheduler::EnableHandle(UObject* InObject, FScheduleHandle InHandle, bool bInEnabled)
 {
-	if (InHandle.IsValid())
+	if(!InHandle.IsValid())
 	{
-		UWorld* World = InObject->GetWorld();
-		UAnimNextSchedulerWorldSubsystem* Subsystem = World->GetSubsystem<UAnimNextSchedulerWorldSubsystem>();
-		Subsystem->EnableHandle(InHandle, bInEnabled);
+		return;
+	}
+	
+	UWorld* World = InObject->GetWorld();
+	if(World == nullptr)
+	{
+		return;
 	}
-}
 
-void FScheduler::QueueTask(UObject* InObject, FScheduleHandle InHandle, FName InScheduleTaskName, TUniqueFunction<void(const FScheduleContext&)>&& InTaskFunction, ETaskRunLocation InLocation)
-{
-	if (InHandle.IsValid())
+	UAnimNextSchedulerWorldSubsystem* Subsystem = World->GetSubsystem<UAnimNextSchedulerWorldSubsystem>();
+	if(Subsystem == nullptr)
 	{
-		UWorld* World = InObject->GetWorld();
-		UAnimNextSchedulerWorldSubsystem* Subsystem = World->GetSubsystem<UAnimNextSchedulerWorldSubsystem>();
-		Subsystem->QueueTask(InHandle, InScheduleTaskName, MoveTemp(InTaskFunction), InLocation);
+		return;
 	}
-}
 
-void FScheduler::RegisterPortDefinition(FSchedulePortDefinition&& InPortDefinition)
-{
-	Impl.RegisteredPortDefinitions.Add(InPortDefinition.Name, MoveTemp(InPortDefinition));
+	Subsystem->EnableHandle(InHandle, bInEnabled);
 }
 
-void FScheduler::UnregisterPortDefinition(FName InDefinitionName)
+void FScheduler::QueueTask(UObject* InObject, FScheduleHandle InHandle, FName InScheduleTaskName, TUniqueFunction<void(const FScheduleContext&)>&& InTaskFunction, ETaskRunLocation InLocation)
 {
-	Impl.RegisteredPortDefinitions.Remove(InDefinitionName);
-}
+	if(!InHandle.IsValid())
+	{
+		return;
+	}
+	
+	UWorld* World = InObject->GetWorld();
+	if(World == nullptr)
+	{
+		return;
+	}
 
-const FSchedulePortDefinition* FScheduler::FindPortDefinition(FName InDefinitionName)
-{
-	return Impl.RegisteredPortDefinitions.Find(InDefinitionName);
+	UAnimNextSchedulerWorldSubsystem* Subsystem = World->GetSubsystem<UAnimNextSchedulerWorldSubsystem>();
+	if(Subsystem == nullptr)
+	{
+		return;
+	}
+
+	Subsystem->QueueTask(InHandle, InScheduleTaskName, MoveTemp(InTaskFunction), InLocation);
 }
 
 }
\ No newline at end of file
diff --git a/Engine/Plugins/Experimental/AnimNext/Source/AnimNext/Private/Scheduler/Scheduler.h b/Engine/Plugins/Experimental/AnimNext/Source/AnimNext/Private/Scheduler/Scheduler.h
--- a/Engine/Plugins/Experimental/AnimNext/Source/AnimNext/Private/Scheduler/Scheduler.h
+++ b/Engine/Plugins/Experimental/AnimNext/Source/AnimNext/Private/Scheduler/Scheduler.h
@@ -3,9 +3,9 @@
 #pragma once
 
 #include "CoreMinimal.h"
-#include "Param/ParamTypeHandle.h"
 #include "ScheduleHandle.h"
 
+enum class EAnimNextScheduleInitMethod : uint8;
 class UAnimNextGraph;
 class UAnimNextParameterBlock;
 class UAnimNextSchedule;
@@ -33,7 +33,7 @@ struct FScheduler
 	static void Destroy();
 
 	// Acquire a handle that binds a schedule with the supplied parameters
-	static FScheduleHandle AcquireHandle(UObject* InObject, UAnimNextSchedule* InSchedule, const TMap<FName, FAnimNextParameterCollection>& InUserScopes);
+	static FScheduleHandle AcquireHandle(UObject* InObject, UAnimNextSchedule* InSchedule, const TMap<FName, FAnimNextParameterCollection>& InUserScopes, EAnimNextScheduleInitMethod InInitMethod);
 
 	// Release an already acquired handle
 	// The full release of the binding referenced by the handle map be deferred after this call is made
@@ -58,15 +58,6 @@ struct FScheduler
 	// @param	InTaskFunction	The function to run
 	// @param	InLocation		Where to run the task, before or after
 	static void QueueTask(UObject* InObject, FScheduleHandle InHandle, FName InScheduleTaskName, TUniqueFunction<void(const FScheduleContext&)>&& InTaskFunction, ETaskRunLocation InLocation = ETaskRunLocation::Before);
-
-	// Register a port definition
-	static void RegisterPortDefinition(FSchedulePortDefinition&& InPortDefinition);
-
-	// Unregister a port definition
-	static void UnregisterPortDefinition(FName InDefinitionName);
-
-	// Find a registered port definition
-	static const FSchedulePortDefinition* FindPortDefinition(FName InDefinitionName);
 };
 
 }
\ No newline at end of file
diff --git a/Engine/Plugins/Experimental/AnimNext/Source/AnimNext/Private/TransformArrayOperations_SoA.cpp b/Engine/Plugins/Experimental/AnimNext/Source/AnimNext/Private/TransformArrayOperations_SoA.cpp
--- a/Engine/Plugins/Experimental/AnimNext/Source/AnimNext/Private/TransformArrayOperations_SoA.cpp
+++ b/Engine/Plugins/Experimental/AnimNext/Source/AnimNext/Private/TransformArrayOperations_SoA.cpp
@@ -2,6 +2,9 @@
 
 #include "TransformArrayOperations.h"
 #include "TransformArray.h"
+#include "AnimNextStats.h"
+
+DEFINE_STAT(STAT_AnimNext_CopyTransforms_SoA);
 
 namespace UE::AnimNext
 {
@@ -55,11 +58,13 @@ namespace UE::AnimNext
 
 	void CopyTransforms(const FTransformArraySoAView& Dest, const FTransformArraySoAConstView& Source, int32 StartIndex, int32 NumToCopy)
 	{
+		SCOPE_CYCLE_COUNTER(STAT_AnimNext_CopyTransforms_SoA);
+		
 		const int32 NumTransforms = Dest.Num();
 		const int32 EndIndex = NumToCopy >= 0 ? (StartIndex + NumToCopy) : NumTransforms;
 
 		check(Source.Num() >= NumTransforms);
-		check(StartIndex >= 0 && StartIndex < NumTransforms);
+		check(StartIndex >= 0 && StartIndex <= NumTransforms);
 		check(EndIndex <= NumTransforms);
 
 		for (int32 TransformIndex = StartIndex; TransformIndex < EndIndex; ++TransformIndex)
diff --git a/Engine/Plugins/Experimental/AnimNext/Source/AnimNext/Public/AnimNextStats.h b/Engine/Plugins/Experimental/AnimNext/Source/AnimNext/Public/AnimNextStats.h
new file mode 100644
--- /dev/null
+++ b/Engine/Plugins/Experimental/AnimNext/Source/AnimNext/Public/AnimNextStats.h
@@ -0,0 +1,34 @@
+﻿// Copyright Epic Games, Inc. All Rights Reserved.
+
+#pragma once
+
+#include "Stats/Stats2.h"
+
+DECLARE_STATS_GROUP(TEXT("AnimNext"), STATGROUP_AnimNext, STATCAT_Advanced);
+
+DECLARE_CYCLE_STAT_EXTERN(TEXT("AnimNext: Initialize Entry"), STAT_AnimNext_InitializeEntry, STATGROUP_AnimNext, ANIMNEXT_API);
+DECLARE_CYCLE_STAT_EXTERN(TEXT("AnimNext: Create Instance Data"), STAT_AnimNext_CreateInstanceData, STATGROUP_AnimNext, ANIMNEXT_API);
+
+DECLARE_CYCLE_STAT_EXTERN(TEXT("AnimNext Task: Graph"), STAT_AnimNext_Task_Graph, STATGROUP_AnimNext, ANIMNEXT_API);
+DECLARE_CYCLE_STAT_EXTERN(TEXT("AnimNext Task: Port"), STAT_AnimNext_Task_Port, STATGROUP_AnimNext, ANIMNEXT_API);
+DECLARE_CYCLE_STAT_EXTERN(TEXT("AnimNext Task: Scope Entry"), STAT_AnimNext_Task_ScopeEntry, STATGROUP_AnimNext, ANIMNEXT_API);
+DECLARE_CYCLE_STAT_EXTERN(TEXT("AnimNext Task: Scope Exit"), STAT_AnimNext_Task_ScopeExit, STATGROUP_AnimNext, ANIMNEXT_API);
+
+DECLARE_CYCLE_STAT_EXTERN(TEXT("AnimNext: Allocate Graph Instance"), STAT_AnimNext_Graph_AllocateInstance, STATGROUP_AnimNext, ANIMNEXT_API);
+DECLARE_CYCLE_STAT_EXTERN(TEXT("AnimNext: Run Graph"), STAT_AnimNext_Graph, STATGROUP_AnimNext, ANIMNEXT_API);
+DECLARE_CYCLE_STAT_EXTERN(TEXT("AnimNext: Update Graph"), STAT_AnimNext_UpdateGraph, STATGROUP_AnimNext, ANIMNEXT_API);
+DECLARE_CYCLE_STAT_EXTERN(TEXT("AnimNext: Evaluate Graph"), STAT_AnimNext_EvaluateGraph, STATGROUP_AnimNext, ANIMNEXT_API);
+DECLARE_CYCLE_STAT_EXTERN(TEXT("AnimNext: Execute Evaluation Program"), STAT_AnimNext_EvaluationProgram_Execute, STATGROUP_AnimNext, ANIMNEXT_API);
+DECLARE_CYCLE_STAT_EXTERN(TEXT("AnimNext: Copy Transforms (SoA)"), STAT_AnimNext_CopyTransforms_SoA, STATGROUP_AnimNext, ANIMNEXT_API);
+
+DECLARE_CYCLE_STAT_EXTERN(TEXT("AnimNext: Param Type Handle Lock"), STAT_AnimNext_ParamTypeHandle_Lock, STATGROUP_AnimNext, ANIMNEXT_API);
+DECLARE_CYCLE_STAT_EXTERN(TEXT("AnimNext: Param ID Lock"), STAT_AnimNext_ParamIdLock, STATGROUP_AnimNext, ANIMNEXT_API);
+DECLARE_CYCLE_STAT_EXTERN(TEXT("AnimNext: Param Stack Get Param"), STAT_AnimNext_ParamStack_GetParam, STATGROUP_AnimNext, ANIMNEXT_API);
+DECLARE_CYCLE_STAT_EXTERN(TEXT("AnimNext: Param Stack Adapter"), STAT_AnimNext_ParamStack_Adapter, STATGROUP_AnimNext, ANIMNEXT_API);
+DECLARE_CYCLE_STAT_EXTERN(TEXT("AnimNext: Run Param Block"), STAT_AnimNext_ParamBlock_UpdateLayer, STATGROUP_AnimNext, ANIMNEXT_API);
+
+DECLARE_CYCLE_STAT_EXTERN(TEXT("AnimNext: Generate Reference Pose"), STAT_AnimNext_GenerateReferencePose, STATGROUP_AnimNext, ANIMNEXT_API);
+DECLARE_CYCLE_STAT_EXTERN(TEXT("AnimNext: Remap Pose"), STAT_AnimNext_RemapPose, STATGROUP_AnimNext, ANIMNEXT_API);
+DECLARE_CYCLE_STAT_EXTERN(TEXT("AnimNext: Convert Local Space To Component Space"), STAT_AnimNext_ConvertLocalSpaceToComponentSpace, STATGROUP_AnimNext, ANIMNEXT_API);
+
+DECLARE_CYCLE_STAT_EXTERN(TEXT("AnimNext: Skeletal Mesh Component Port"), STAT_AnimNext_Port_SkeletalMeshComponent, STATGROUP_AnimNext, ANIMNEXT_API);
diff --git a/Engine/Plugins/Experimental/AnimNext/Source/AnimNext/Public/DataRegistry.h b/Engine/Plugins/Experimental/AnimNext/Source/AnimNext/Public/DataRegistry.h
--- a/Engine/Plugins/Experimental/AnimNext/Source/AnimNext/Public/DataRegistry.h
+++ b/Engine/Plugins/Experimental/AnimNext/Source/AnimNext/Public/DataRegistry.h
@@ -100,6 +100,8 @@ public:
 
 			Private::FAllocatedBlock* AllocatedBlock = new Private::FAllocatedBlock(Memory, NumElements, ParamTypeHandle); // TODO : avoid memory fragmentation
 			AllocatedBlock->AddRef();
+
+			FRWScopeLock Lock(AllocatedBlocksLock, SLT_Write);
 			AllocatedBlocks.Add(AllocatedBlock);
 			return FDataHandle(AllocatedBlock);
 		}
diff --git a/Engine/Plugins/Experimental/AnimNext/Source/AnimNext/Public/GenerationTools.h b/Engine/Plugins/Experimental/AnimNext/Source/AnimNext/Public/GenerationTools.h
--- a/Engine/Plugins/Experimental/AnimNext/Source/AnimNext/Public/GenerationTools.h
+++ b/Engine/Plugins/Experimental/AnimNext/Source/AnimNext/Public/GenerationTools.h
@@ -91,6 +91,16 @@ public:
 	// This function expects both poses to have the same LOD (number of bones and indexes)
 	// The target pose should be assigned to the correct reference pose prior to this call
 	static void RemapPose(const FLODPose& SourcePose, FPoseContext& TargetPose);
+
+	// Converts AnimNext pose to local space transform array
+	// This function expects the output pose to have the same or a greater number of bones (as it may be being calculated
+	// for a lower LOD)
+	// The target pose should be assigned to the correct reference pose prior to this call, as transforms will not be filled
+	// in by this call if they are not affected by the current LOD.
+	static void RemapPose(const FLODPose& SourcePose, TArrayView<FTransform> TargetTransforms);
+
+	// Converts a local space to component space buffer given a number of required bones
+	static void ConvertLocalSpaceToComponentSpace(TConstArrayView<FBoneIndexType> InParentIndices, TConstArrayView<FTransform> InBoneSpaceTransforms, TConstArrayView<FBoneIndexType> InRequiredBoneIndices, TArrayView<FTransform> OutComponentSpaceTransforms);
 };
 
 } // namespace UE::AnimNext
diff --git a/Engine/Plugins/Experimental/AnimNext/Source/AnimNext/Public/LODPose.h b/Engine/Plugins/Experimental/AnimNext/Source/AnimNext/Public/LODPose.h
--- a/Engine/Plugins/Experimental/AnimNext/Source/AnimNext/Public/LODPose.h
+++ b/Engine/Plugins/Experimental/AnimNext/Source/AnimNext/Public/LODPose.h
@@ -181,6 +181,14 @@ struct TLODPose : public FLODPose
 		PrepareForLOD(InRefPose, InLODLevel, bSetRefPose, bAdditive);
 	}
 
+	bool ShouldPrepareForLOD(const FReferencePose& InRefPose, int32 InLODLevel, bool bAdditive = false) const
+	{
+		return
+			LODLevel != InLODLevel ||
+			RefPose != &InRefPose ||
+			bAdditive != EnumHasAnyFlags(Flags, ELODPoseFlags::Additive);
+	}
+	
 	void PrepareForLOD(const FReferencePose& InRefPose, int32 InLODLevel, bool bSetRefPose = true, bool bAdditive = false)
 	{
 		LODLevel = InLODLevel;
diff --git a/Engine/Plugins/Experimental/AnimNext/Source/AnimNext/Public/Param/ParamEntry.h b/Engine/Plugins/Experimental/AnimNext/Source/AnimNext/Public/Param/ParamEntry.h
--- a/Engine/Plugins/Experimental/AnimNext/Source/AnimNext/Public/Param/ParamEntry.h
+++ b/Engine/Plugins/Experimental/AnimNext/Source/AnimNext/Public/Param/ParamEntry.h
@@ -11,6 +11,7 @@ namespace UE::AnimNext
 	struct FUObjectLayer;
 	struct FInstancedPropertyBagLayer;
 	struct FParamStackLayerHandle;
+	struct FRemappedLayer;
 }
 
 namespace UE::AnimNext
@@ -47,6 +48,7 @@ private:
 	friend struct UE::AnimNext::FUObjectLayer;
 	friend struct UE::AnimNext::FInstancedPropertyBagLayer;
 	friend struct UE::AnimNext::FParamStackLayerHandle;
+	friend struct UE::AnimNext::FRemappedLayer;
 
 	ANIMNEXT_API FParamEntry(const FParamTypeHandle& InTypeHandle, TArrayView<uint8> InData, bool bInIsReference, bool bInIsMutable);
 
diff --git a/Engine/Plugins/Experimental/AnimNext/Source/AnimNext/Public/Param/ParamId.h b/Engine/Plugins/Experimental/AnimNext/Source/AnimNext/Public/Param/ParamId.h
--- a/Engine/Plugins/Experimental/AnimNext/Source/AnimNext/Public/Param/ParamId.h
+++ b/Engine/Plugins/Experimental/AnimNext/Source/AnimNext/Public/Param/ParamId.h
@@ -14,6 +14,7 @@ namespace UE::AnimNext
 	struct FParamAdapter;
 	struct FParamDefinition;
 	class FModule;
+	struct FRemappedLayer;
 }
 
 namespace UE::AnimNext::Tests
@@ -36,6 +37,7 @@ struct ANIMNEXT_API FParamId
 	friend class ::UAnimNextConfig;
 	friend struct FParamDefinition;
 	friend class FModule;
+	friend struct FRemappedLayer;
 
 	static constexpr uint32 InvalidIndex = MAX_uint32;
 
@@ -77,6 +79,12 @@ private:
 	// Shut down the parameter ID system
 	static void Destroy();
 
+	// Refresh all adapters, built in and per-config
+	static void RefreshAdapters();
+
+	// Resets/clears all adapters
+	static void ResetAdapters();
+
 	// Register all built-in adapters
 	static void RegisterBuiltInAdapters();
 	
diff --git a/Engine/Plugins/Experimental/AnimNext/Source/AnimNext/Public/Param/ParamStack.h b/Engine/Plugins/Experimental/AnimNext/Source/AnimNext/Public/Param/ParamStack.h
--- a/Engine/Plugins/Experimental/AnimNext/Source/AnimNext/Public/Param/ParamStack.h
+++ b/Engine/Plugins/Experimental/AnimNext/Source/AnimNext/Public/Param/ParamStack.h
@@ -14,6 +14,7 @@
 #include "Param/ParamStackLayerHandle.h"
 #include "ParamEntry.h"
 #include "Param/ParamResult.h"
+#include "AnimNextStats.h"
 
 class UAnimNextSchedule;
 class UAnimNextSchedulerWorldSubsystem;
@@ -209,6 +210,10 @@ public:
 	// Create a cached parameter layer from an instanced property bag. This layer will reference the suppled propery bag and does not transfer ownership.
 	static ANIMNEXT_API FParamStackLayerHandle MakeReferenceLayer(FInstancedPropertyBag& InInstancedPropertyBag);
 
+	// Create a cached parameter layer by remapping the entries from another layer
+	// Original layer continues to own the memory
+	static ANIMNEXT_API FParamStackLayerHandle MakeRemappedLayer(const FParamStackLayerHandle& InLayer, const TMap<FName, FName>& InMapping);
+
 	// Make a parameter layer from a value
 	// @param	InParamId			Parameter ID for the parameter
 	// @param	InValue				Value to push
diff --git a/Engine/Plugins/Experimental/AnimNext/Source/AnimNext/Public/Param/ParamStackLayerHandle.h b/Engine/Plugins/Experimental/AnimNext/Source/AnimNext/Public/Param/ParamStackLayerHandle.h
--- a/Engine/Plugins/Experimental/AnimNext/Source/AnimNext/Public/Param/ParamStackLayerHandle.h
+++ b/Engine/Plugins/Experimental/AnimNext/Source/AnimNext/Public/Param/ParamStackLayerHandle.h
@@ -3,6 +3,7 @@
 #pragma once
 
 #include "CoreMinimal.h"
+#include "ParamCompatibility.h"
 #include "Param/ParamResult.h"
 #include "Param/ParamId.h"
 #include "Param/ParamTypeHandle.h"
@@ -68,12 +69,6 @@ public:
 	// @param	InData				Value to set
 	ANIMNEXT_API FParamResult SetValueRaw(FParamId InParamId, FParamTypeHandle InTypeHandle, TConstArrayView<uint8> InData);
 
-	// Get a parameter's value from this layer
-	// @param	InParamId			Parameter ID for the parameter
-	// @param	InTypeHandle		The type of the parameter
-	// @param	InData				Value to get
-	ANIMNEXT_API FParamResult GetValueRaw(FParamId InParamId, FParamTypeHandle InTypeHandle, TConstArrayView<uint8>& OutData);
-
 	// Set multiple parameter values in this layer
 	// @param	InValues	Names & values to set, interleaved
 	template<typename... Args>
@@ -85,6 +80,74 @@ public:
 		return SetValuesHelper(ParamIdValues, Forward<Args>(InValues)...);
 	}
 
+	// Get a parameter's value from this layer
+	// @param	InParamId			Parameter ID for the parameter
+	// @param	InTypeHandle		The type of the parameter
+	// @param	InData				Value to get
+	ANIMNEXT_API FParamResult GetValueRaw(FParamId InParamId, FParamTypeHandle InTypeHandle, TConstArrayView<uint8>& OutData);
+
+	// Get a pointer to a parameter's value given a FParamId.
+	// @param	InParamId			Parameter ID to find the currently-pushed value for
+	// @param	OutResult			Optional ptr to aresult that will be filled if an error occurs
+	// @return nullptr if the parameter is not mapped
+	template<typename ValueType>
+	const ValueType* GetParamPtr(FParamId InParamId, FParamResult* OutResult = nullptr) const
+	{
+		TConstArrayView<uint8> Data;
+		FParamTypeHandle TypeHandle;
+		FParamResult Result = GetParamData(InParamId, FParamTypeHandle::GetHandle<ValueType>(), Data, TypeHandle);
+		if (OutResult)
+		{
+			*OutResult = Result;
+		}
+		return reinterpret_cast<const ValueType*>(Data.GetData());
+	}
+
+	// Get a const reference to the value of a parameter given a FParamId.
+	// @param	InParamId			Parameter ID to find the currently-pushed value for
+	// @param	OutResult			Optional ptr to a result that will be filled if an error occurs
+	// @return a const reference to the parameter's value. Function asserts if the parameter is not present
+	template<typename ValueType>
+	const ValueType& GetParam(FParamId InParamId, FParamResult* OutResult = nullptr) const
+	{
+		const ValueType* ParamPtr = GetParamPtr<ValueType>(InParamId, OutResult);
+		check(ParamPtr);
+		return *ParamPtr;
+	}
+
+	ANIMNEXT_API FParamResult GetParamData(FParamId InId, FParamTypeHandle InTypeHandle, TConstArrayView<uint8>& OutParamData, FParamTypeHandle& OutParamTypeHandle, FParamCompatibility InRequiredCompatibility = FParamCompatibility::Equal()) const;
+
+	// Get a pointer to a parameter's value given a FParamId.
+	// @param	InParamId			Parameter ID to find the currently-pushed value for
+	// @param	OutResult			Optional ptr to aresult that will be filled if an error occurs
+	// @return nullptr if the parameter is not mapped
+	template<typename ValueType>
+	ValueType* GetMutableParamPtr(FParamId InParamId, FParamResult* OutResult = nullptr) const
+	{
+		TArrayView<uint8> Data;
+		FParamTypeHandle TypeHandle;
+		FParamResult Result = GetMutableParamData(InParamId, FParamTypeHandle::GetHandle<ValueType>(), Data, TypeHandle);
+		if (OutResult)
+		{
+			*OutResult = Result;
+		}
+		return reinterpret_cast<ValueType*>(Data.GetData());
+	}
+
+	// Get a reference to the value of a parameter given a FParamId.
+	// @param	InParamId			Parameter ID to find the currently-pushed value for
+	// @param	OutResult			Optional ptr to a result that will be filled if an error occurs
+	// @return a reference to the parameter's value. Function asserts if the parameter is not present
+	template<typename ValueType>
+	ValueType& GetMutableParam(FParamId InParamId, FParamResult* OutResult = nullptr) const
+	{
+		ValueType* ParamPtr = GetMutableParamPtr<ValueType>(InParamId, OutResult);
+		check(ParamPtr);
+		return *ParamPtr;
+	}
+	
+	ANIMNEXT_API FParamResult GetMutableParamData(FParamId InId, FParamTypeHandle InTypeHandle, TArrayView<uint8>& OutParamData, FParamTypeHandle& OutParamTypeHandle, FParamCompatibility InRequiredCompatibility = FParamCompatibility::Equal()) const;
+	
 	// Get the underlying storage for this layer.
 	// @return nullptr if the layer does not use the supplied underlying storage type
 	template<typename WrappedType>
diff --git a/Engine/Plugins/Experimental/AnimNext/Source/AnimNext/Public/Param/ParamType.h b/Engine/Plugins/Experimental/AnimNext/Source/AnimNext/Public/Param/ParamType.h
--- a/Engine/Plugins/Experimental/AnimNext/Source/AnimNext/Public/Param/ParamType.h
+++ b/Engine/Plugins/Experimental/AnimNext/Source/AnimNext/Public/Param/ParamType.h
@@ -88,6 +88,8 @@ private:
 	template<typename ParamType>
 	static constexpr void GetTypeInner(FAnimNextParamType& ParameterType)
 	{
+		using NonPtrParamType = std::remove_pointer_t<ParamType>;
+
 		if constexpr (std::is_same_v<ParamType, bool>)
 		{
 			ParameterType.ValueType = EValueType::Bool;
@@ -139,17 +141,17 @@ private:
 			ParameterType.ValueType = EValueType::Struct;
 			ParameterType.ValueTypeObject = TBaseStructure<ParamType>::Get();
 		}
-		else if constexpr (TModels<CStaticClassProvider, ParamType>::Value)
+		else if constexpr (TModels<CStaticClassProvider, NonPtrParamType>::Value)
 		{
-			if constexpr (std::is_same_v<ParamType, UClass>)
+			if constexpr (std::is_same_v<NonPtrParamType, UClass>)
 			{
 				ParameterType.ValueType = EValueType::Class;
-				ParameterType.ValueTypeObject = ParamType::StaticClass();
+				ParameterType.ValueTypeObject = NonPtrParamType::StaticClass();
 			}
 			else
 			{
 				ParameterType.ValueType = EValueType::Object;
-				ParameterType.ValueTypeObject = ParamType::StaticClass();
+				ParameterType.ValueTypeObject = NonPtrParamType::StaticClass();
 			}
 		}
 		else if constexpr (TIsTObjectPtr<ParamType>::Value)
diff --git a/Engine/Plugins/Experimental/AnimNext/Source/AnimNext/Public/Param/ParamTypeHandle.h b/Engine/Plugins/Experimental/AnimNext/Source/AnimNext/Public/Param/ParamTypeHandle.h
--- a/Engine/Plugins/Experimental/AnimNext/Source/AnimNext/Public/Param/ParamTypeHandle.h
+++ b/Engine/Plugins/Experimental/AnimNext/Source/AnimNext/Public/Param/ParamTypeHandle.h
@@ -6,6 +6,12 @@
 #include "Concepts/BaseStructureProvider.h"
 #include "Param/ParamType.h"
 
+class UAnimSequence;
+class UCharacterMovementComponent;
+class UAnimNextMeshComponent;
+struct FAnimNextGraphLODPose;
+struct FAnimNextGraphReferencePose;
+
 namespace UE::AnimNext::Tests
 {
 	class FParamTypesTest;
@@ -43,7 +49,17 @@ struct ANIMNEXT_API FParamTypeHandle
 		Quat,
 		Transform,
 
-		MaxBuiltIn = Transform,
+		// Common object types
+		Object,
+		CharacterMovementComponent,
+		AnimNextMeshComponent,
+		AnimSequence,
+
+		// Common struct types
+		AnimNextGraphLODPose,
+		AnimNextGraphReferencePose,
+
+		MaxBuiltIn = AnimNextGraphReferencePose,
 
 		Custom = 0xff,
 
@@ -103,6 +119,8 @@ private:
 	template<typename ParamType>
 	static constexpr void GetHandleInner(FAnimNextParamType::EValueType& OutValueType, UObject*& OutValueTypeObject)
 	{
+		using NonPtrParamType = std::remove_pointer_t<ParamType>;
+
 		if constexpr (std::is_same_v<ParamType, bool>)
 		{
 			OutValueType = FAnimNextParamType::EValueType::Bool;
@@ -154,17 +172,17 @@ private:
 			OutValueType = FAnimNextParamType::EValueType::Struct;
 			OutValueTypeObject = TBaseStructure<ParamType>::Get();
 		}
-		else if constexpr (TModels<CStaticClassProvider, ParamType>::Value)
+		else if constexpr (TModels<CStaticClassProvider, NonPtrParamType>::Value)
 		{
-			if constexpr (std::is_same_v<ParamType, UClass>)
+			if constexpr (std::is_same_v<NonPtrParamType, UClass>)
 			{
 				OutValueType = FAnimNextParamType::EValueType::Class;
-				OutValueTypeObject = ParamType::StaticClass();
+				OutValueTypeObject = NonPtrParamType::StaticClass();
 			}
 			else
 			{
 				OutValueType = FAnimNextParamType::EValueType::Object;
-				OutValueTypeObject = ParamType::StaticClass();
+				OutValueTypeObject = NonPtrParamType::StaticClass();
 			}
 		}
 		else if constexpr (TIsTObjectPtr<ParamType>::Value)
@@ -283,6 +301,30 @@ public:
 		{
 			TypeHandle.SetParameterType(EParamType::Transform);
 		}
+		else if constexpr (std::is_same_v<NonConstType, TObjectPtr<UObject>> || std::is_same_v<NonConstType, UObject*>)
+		{
+			TypeHandle.SetParameterType(EParamType::Object);
+		}
+		else if constexpr (std::is_same_v<NonConstType, TObjectPtr<UCharacterMovementComponent>> || std::is_same_v<NonConstType, UCharacterMovementComponent*>)
+		{
+			TypeHandle.SetParameterType(EParamType::CharacterMovementComponent);
+		}
+		else if constexpr (std::is_same_v<NonConstType, TObjectPtr<UAnimNextMeshComponent>> || std::is_same_v<NonConstType, UAnimNextMeshComponent*>)
+		{
+			TypeHandle.SetParameterType(EParamType::AnimNextMeshComponent);
+		}
+		else if constexpr (std::is_same_v<NonConstType, TObjectPtr<UAnimSequence>> || std::is_same_v<NonConstType, UAnimSequence*>)
+		{
+			TypeHandle.SetParameterType(EParamType::AnimSequence);
+		}
+		else if constexpr (std::is_same_v<NonConstType, FAnimNextGraphLODPose>)
+		{
+			TypeHandle.SetParameterType(EParamType::AnimNextGraphLODPose);
+		}
+		else if constexpr (std::is_same_v<NonConstType, FAnimNextGraphReferencePose>)
+		{
+			TypeHandle.SetParameterType(EParamType::AnimNextGraphReferencePose);
+		}
 		else
 		{
 			// Not a built-in-type, so we need to do some work
@@ -307,7 +349,7 @@ public:
 			}
 		}
 
-		//check(TypeHandle.IsValid());
+		check(TypeHandle.IsValid());
 
 		return TypeHandle;
 	}
diff --git a/Engine/Plugins/Experimental/AnimNext/Source/AnimNext/Public/ReferencePose.h b/Engine/Plugins/Experimental/AnimNext/Source/AnimNext/Public/ReferencePose.h
--- a/Engine/Plugins/Experimental/AnimNext/Source/AnimNext/Public/ReferencePose.h
+++ b/Engine/Plugins/Experimental/AnimNext/Source/AnimNext/Public/ReferencePose.h
@@ -52,6 +52,9 @@ struct TReferencePose
 	// Number of bones for each LOD
 	TArray<int32, AllocatorType> LODNumBones;
 
+	// Mapping of mesh bone indices to mesh parent indices for each bone
+	TArray<FBoneIndexType, AllocatorType> ParentIndices;
+
 	TWeakObjectPtr<const USkeletalMesh> SkeletalMesh = nullptr;
 	TWeakObjectPtr<const USkeleton> Skeleton = nullptr;
 	EReferencePoseGenerationFlags GenerationFlags = EReferencePoseGenerationFlags::None;
@@ -83,15 +86,18 @@ struct TReferencePose
 		, bool bFastPath = false)
 	{
 		const int32 NumBonesLOD0 = !InLODNumBones.IsEmpty() ? InLODNumBones[0] : 0;
-
+		const int32 NumBonesMesh = RefSkeleton.GetRefBoneInfo().Num();
+		
+		ParentIndices.SetNum(NumBonesMesh);
 		ReferenceLocalTransforms.SetNum(NumBonesLOD0);
 		LODBoneIndexToMeshBoneIndexMapPerLOD = InLODBoneIndexToMeshBoneIndexMapPerLOD;
 		LODBoneIndexToSkeletonBoneIndexMapPerLOD = InLODBoneIndexToSkeletonBoneIndexMapPerLOD;
 		SkeletonBoneIndexToLODBoneIndexMapPerLOD = InSkeletonBoneIndexToLODBoneIndexMapPerLOD;
 		LODNumBones = InLODNumBones;
-		
+
 		const TArray<FTransform>& RefBonePose = RefSkeleton.GetRefBonePose();
 		const TArray<FBoneIndexType>& BoneLODIndexToMeshIndexMap0 = InLODBoneIndexToMeshBoneIndexMapPerLOD[0]; // Fill the transforms with the LOD0 indexes
+		const TArray<FMeshBoneInfo>& RefBoneInfo = RefSkeleton.GetRefBoneInfo();
 
 		for (int32 LODBoneIndex = 0; LODBoneIndex < NumBonesLOD0; ++LODBoneIndex)
 		{
@@ -99,6 +105,11 @@ struct TReferencePose
 			ReferenceLocalTransforms[LODBoneIndex] = RefBonePose[BoneLODIndexToMeshIndexMap0[LODBoneIndex]];
 		}
 
+		for(int32 BoneIndex = 0; BoneIndex < NumBonesMesh; ++BoneIndex)
+		{
+			ParentIndices[BoneIndex] = RefBoneInfo[BoneIndex].ParentIndex;
+		}
+
 		GenerationFlags = bFastPath ? EReferencePoseGenerationFlags::FastPath : EReferencePoseGenerationFlags::None;
 	}
 
@@ -157,6 +168,12 @@ struct TReferencePose
 		return LODBoneIndexToMeshBoneIndexMapPerLOD[0][LODBoneIndex];
 	}
 
+	// Returns a mapping of mesh bone indices to mesh parent indices for each bone
+	TConstArrayView<FBoneIndexType> GetParentIndices() const
+	{
+		return ParentIndices;
+	}
+
 	int32 GetSkeletonBoneIndexFromLODBoneIndex(int32 LODBoneIndex) const
 	{
 		const int32 NumBonesLOD0 = LODBoneIndexToSkeletonBoneIndexMapPerLOD[0].Num();
diff --git a/Engine/Plugins/Experimental/AnimNext/Source/AnimNext/Public/Scheduler/AnimNextSchedulePort.h b/Engine/Plugins/Experimental/AnimNext/Source/AnimNext/Public/Scheduler/AnimNextSchedulePort.h
new file mode 100644
--- /dev/null
+++ b/Engine/Plugins/Experimental/AnimNext/Source/AnimNext/Public/Scheduler/AnimNextSchedulePort.h
@@ -0,0 +1,26 @@
+// Copyright Epic Games, Inc. All Rights Reserved.
+
+#pragma once
+
+#include "CoreMinimal.h"
+#include "IAnimNextScheduleTermInterface.h"
+#include "AnimNextSchedulePort.generated.h"
+
+struct FAnimNextSchedulePortTask;
+
+namespace UE::AnimNext
+{
+	struct FScheduleTermContext;
+}
+
+UCLASS(Abstract)
+class ANIMNEXT_API UAnimNextSchedulePort : public UObject, public IAnimNextScheduleTermInterface
+{
+	GENERATED_BODY()
+
+private:
+	friend struct FAnimNextSchedulePortTask;
+
+	// Run the port's logic
+	virtual void Run(const UE::AnimNext::FScheduleTermContext& InContext) const PURE_VIRTUAL(UAnimNextSchedulePort::Run, )
+};
diff --git a/Engine/Plugins/Experimental/AnimNext/Source/AnimNext/Public/Scheduler/IAnimNextScheduleTermInterface.h b/Engine/Plugins/Experimental/AnimNext/Source/AnimNext/Public/Scheduler/IAnimNextScheduleTermInterface.h
new file mode 100644
--- /dev/null
+++ b/Engine/Plugins/Experimental/AnimNext/Source/AnimNext/Public/Scheduler/IAnimNextScheduleTermInterface.h
@@ -0,0 +1,60 @@
+﻿// Copyright Epic Games, Inc. All Rights Reserved.
+
+#pragma once
+
+#include "CoreMinimal.h"
+#include "Param/ParamDefinition.h"
+#include "UObject/Interface.h"
+#include "IAnimNextScheduleTermInterface.generated.h"
+
+class UAnimNextSchedule;
+
+namespace UE::AnimNext
+{
+	struct FScheduleInstanceData;
+}
+
+namespace UE::AnimNext
+{
+
+enum class EScheduleTermDirection
+{
+	Input,
+	Output,
+	IO,
+};
+
+struct FScheduleTerm : FParamDefinition
+{
+	FScheduleTerm(FName InName, const FAnimNextParamType& InType, EScheduleTermDirection InDirection)
+		: FParamDefinition(InName, InType)
+		, Direction(InDirection)
+	{}
+
+	FScheduleTerm(FParamId InId, const FAnimNextParamType& InType, EScheduleTermDirection InDirection)
+		: FParamDefinition(InId, InType)
+		, Direction(InDirection)
+	{}
+
+	EScheduleTermDirection Direction;
+};
+
+}
+
+UINTERFACE()
+class UAnimNextScheduleTermInterface : public UInterface
+{
+	GENERATED_BODY()
+};
+
+class IAnimNextScheduleTermInterface
+{
+	GENERATED_BODY()
+
+private:
+	friend class UAnimNextSchedule;
+	friend struct UE::AnimNext::FScheduleInstanceData;
+
+	// Get the terms that define the data dependencies of a schedule element
+	virtual TConstArrayView<UE::AnimNext::FScheduleTerm> GetTerms() const PURE_VIRTUAL(IAnimNextScheduleTermInterface::GetTerms, return {};)
+};
\ No newline at end of file
diff --git a/Engine/Plugins/Experimental/AnimNext/Source/AnimNext/Public/Scheduler/SchedulePortContext.h b/Engine/Plugins/Experimental/AnimNext/Source/AnimNext/Public/Scheduler/SchedulePortContext.h
new file mode 100644
--- /dev/null
+++ b/Engine/Plugins/Experimental/AnimNext/Source/AnimNext/Public/Scheduler/SchedulePortContext.h
@@ -0,0 +1,55 @@
+﻿// Copyright Epic Games, Inc. All Rights Reserved.
+
+#pragma once
+
+#include "Param/ParamCompatibility.h"
+#include "Param/ParamStackLayerHandle.h"
+
+namespace UE::AnimNext
+{
+
+struct FSchedulePortContext
+{
+	FSchedulePortContext() = delete;
+
+	FSchedulePortContext(const TUniquePtr<FParamStackLayer>& InResultLayer)
+		: ResultLayer(InResultLayer)
+	{}
+
+	// Get a pointer to a parameter's value given a FParamId.
+	// @param	InParamId			Parameter ID to find the currently-pushed value for
+	// @param	OutResult			Optional ptr to aresult that will be filled if an error occurs
+	// @return nullptr if the parameter is not mapped
+	template<typename ValueType>
+	const ValueType* GetParamPtr(FParamId InParamId, FParamResult* OutResult = nullptr) const
+	{
+		TConstArrayView<uint8> Data;
+		FParamTypeHandle TypeHandle;
+		FParamResult Result = GetParamData(InParamId, FParamTypeHandle::GetHandle<ValueType>(), Data, TypeHandle);
+		if (OutResult)
+		{
+			*OutResult = Result;
+		}
+		return reinterpret_cast<const ValueType*>(Data.GetData());
+	}
+
+	// Get a const reference to the value of a parameter given a FParamId.
+	// @param	InParamId			Parameter ID to find the currently-pushed value for
+	// @param	OutResult			Optional ptr to a result that will be filled if an error occurs
+	// @return a const reference to the parameter's value. Function asserts if the parameter is not present
+	template<typename ValueType>
+	const ValueType& GetParam(FParamId InParamId, FParamResult* OutResult = nullptr) const
+	{
+		const ValueType* ParamPtr = GetParamPtr<ValueType>(InParamId, OutResult);
+		check(ParamPtr);
+		return *ParamPtr;
+	}
+
+	ANIMNEXT_API FParamResult GetParamData(FParamId InId, FParamTypeHandle InTypeHandle, TConstArrayView<uint8>& OutParamData, FParamTypeHandle& OutParamTypeHandle, FParamCompatibility InRequiredCompatibility = FParamCompatibility::Equal()) const;
+
+private:
+	// The result layer that the port operates on
+	const TUniquePtr<FParamStackLayer>& ResultLayer;
+};
+
+}
diff --git a/Engine/Plugins/Experimental/AnimNext/Source/AnimNext/Public/Scheduler/ScheduleTermContext.h b/Engine/Plugins/Experimental/AnimNext/Source/AnimNext/Public/Scheduler/ScheduleTermContext.h
new file mode 100644
--- /dev/null
+++ b/Engine/Plugins/Experimental/AnimNext/Source/AnimNext/Public/Scheduler/ScheduleTermContext.h
@@ -0,0 +1,37 @@
+﻿// Copyright Epic Games, Inc. All Rights Reserved.
+
+#pragma once
+
+namespace UE::AnimNext
+{
+	struct FParamStackLayerHandle;
+	struct FScheduleContext;
+}
+
+namespace UE::AnimNext
+{
+
+struct FScheduleTermContext
+{
+	FScheduleTermContext() = delete;
+
+	FScheduleTermContext(const FScheduleContext& InScheduleContext, const FParamStackLayerHandle& InResultLayerHandle)
+		: ScheduleContext(InScheduleContext)
+		, ResultLayerHandle(InResultLayerHandle)
+	{}
+
+	// Get the schedule context we wrap
+	const FScheduleContext& GetScheduleContext() const { return ScheduleContext; }
+	
+	// Get the layer handle for the terms that the port operates on
+	const FParamStackLayerHandle& GetLayerHandle() const { return ResultLayerHandle; }
+
+private:
+	// The schedule context we wrap
+	const FScheduleContext& ScheduleContext;
+	
+	// The layer handle for the terms that the port operates on
+	const FParamStackLayerHandle& ResultLayerHandle;
+};
+
+}
diff --git a/Engine/Plugins/Experimental/AnimNext/Source/AnimNextEditor/Private/EditorUtils.cpp b/Engine/Plugins/Experimental/AnimNext/Source/AnimNextEditor/Private/EditorUtils.cpp
--- a/Engine/Plugins/Experimental/AnimNext/Source/AnimNextEditor/Private/EditorUtils.cpp
+++ b/Engine/Plugins/Experimental/AnimNext/Source/AnimNextEditor/Private/EditorUtils.cpp
@@ -15,6 +15,7 @@
 #include "Param/IAnimNextParameterBlockParameterInterface.h"
 #include "Param/Params.h"
 #include "Workspace/AnimNextWorkspace.h"
+#include "UncookedOnlyUtils.h"
 
 namespace UE::AnimNext::Editor
 {
@@ -434,16 +435,10 @@ void FUtils::GetFilteredVariableTypeTree(TArray<TSharedPtr<UEdGraphSchema_K2::FP
 	}
 };
 
-bool FUtils::GetExportedParametersForLibrary(const FAssetData& InLibraryAsset, FAnimNextParameterLibraryAssetRegistryExports& OutExports)
-{
-	const FString TagValue = InLibraryAsset.GetTagValueRef<FString>(UAnimNextParameterLibrary::ExportsAssetRegistryTag);
-	return FAnimNextParameterLibraryAssetRegistryExports::StaticStruct()->ImportText(*TagValue, &OutExports, nullptr, PPF_None, nullptr, FAnimNextParameterLibraryAssetRegistryExports::StaticStruct()->GetName()) != nullptr;
-}
-
 FName FUtils::GetNewParameterNameInLibrary(const FAssetData& InLibraryAsset, const TCHAR* InBaseName, TArrayView<FName> InAdditionalExistingNames)
 {
 	FAnimNextParameterLibraryAssetRegistryExports Exports;
-	GetExportedParametersForLibrary(InLibraryAsset, Exports);
+	UncookedOnly::FUtils::GetExportedParametersForLibrary(InLibraryAsset, Exports);
 
 	auto NameExists = [&Exports, &InAdditionalExistingNames](const TCHAR* InName)
 	{
@@ -490,7 +485,7 @@ FName FUtils::GetNewParameterNameInLibrary(const FAssetData& InLibraryAsset, con
 bool FUtils::DoesParameterExistInLibrary(const FAssetData& InLibraryAsset, const FName InParameterName)
 {
 	FAnimNextParameterLibraryAssetRegistryExports Exports;
-	GetExportedParametersForLibrary(InLibraryAsset, Exports);
+	UncookedOnly::FUtils::GetExportedParametersForLibrary(InLibraryAsset, Exports);
 
 	for(const FAnimNextParameterLibraryAssetRegistryExportEntry& Parameter : Exports.Parameters)
 	{
@@ -521,44 +516,10 @@ bool FUtils::GetExportedBindingsForBlock(const FAssetData& InLibraryAsset, FAnim
 	return FAnimNextParameterBlockAssetRegistryExports::StaticStruct()->ImportText(*TagValue, &OutExports, nullptr, PPF_None, nullptr, FAnimNextParameterBlockAssetRegistryExports::StaticStruct()->GetName()) != nullptr;
 }
 
-FAnimNextParamType FUtils::GetParameterTypeFromName(FName InName)
-{
-	// Check built-in params first as they are cheaper
-	if(const FParamDefinition* FoundDefinition = FParams::FindBuiltInParameter(InName))
-	{
-		return FoundDefinition->Type;
-	}
-
-	// Query the asset registry for other params
-	IAssetRegistry& AssetRegistry = FModuleManager::Get().LoadModuleChecked<FAssetRegistryModule>(TEXT("AssetRegistry")).Get();
-
-	FARFilter ARFilter;
-	ARFilter.ClassPaths = { UAnimNextParameterLibrary::StaticClass()->GetClassPathName() };
-
-	TArray<FAssetData> LibraryAssets;
-	AssetRegistry.GetAssets(ARFilter, LibraryAssets);
-
-	for(const FAssetData& LibraryAsset : LibraryAssets)
-	{
-		FAnimNextParameterLibraryAssetRegistryExports Exports;
-		GetExportedParametersForLibrary(LibraryAsset, Exports);
-		for(const FAnimNextParameterLibraryAssetRegistryExportEntry& Export : Exports.Parameters)
-		{
-			if(Export.Name == InName)
-			{
-				return Export.Type;
-			}
-		}
-	}
-
-	return FAnimNextParamType();
-}
-
 bool FUtils::GetExportedAssetsForWorkspace(const FAssetData& InWorkspaceAsset, FAnimNextWorkspaceAssetRegistryExports& OutExports)
 {
 	const FString TagValue = InWorkspaceAsset.GetTagValueRef<FString>(UAnimNextWorkspace::ExportsAssetRegistryTag);
 	return FAnimNextWorkspaceAssetRegistryExports::StaticStruct()->ImportText(*TagValue, &OutExports, nullptr, PPF_None, nullptr, FAnimNextWorkspaceAssetRegistryExports::StaticStruct()->GetName()) != nullptr;
 }
 
-
 }
diff --git a/Engine/Plugins/Experimental/AnimNext/Source/AnimNextEditor/Private/EditorUtils.h b/Engine/Plugins/Experimental/AnimNext/Source/AnimNextEditor/Private/EditorUtils.h
--- a/Engine/Plugins/Experimental/AnimNext/Source/AnimNextEditor/Private/EditorUtils.h
+++ b/Engine/Plugins/Experimental/AnimNext/Source/AnimNextEditor/Private/EditorUtils.h
@@ -36,8 +36,6 @@ struct FUtils
 
 	static void GetFilteredVariableTypeTree(TArray<TSharedPtr<UEdGraphSchema_K2::FPinTypeTreeInfo>>& TypeTree, ETypeTreeFilter TypeTreeFilter);
 
-	static bool GetExportedParametersForLibrary(const FAssetData& InLibraryAsset, FAnimNextParameterLibraryAssetRegistryExports& OutExports);
-
 	static FName GetNewParameterNameInLibrary(const FAssetData& InLibraryAsset, const TCHAR* InBaseName, TArrayView<FName> InAdditionalExistingNames);
 
 	static bool DoesParameterExistInLibrary(const FAssetData& InLibraryAsset, const FName InParameterName);
@@ -46,10 +44,6 @@ struct FUtils
 
 	static bool GetExportedBindingsForBlock(const FAssetData& InBlockAsset, FAnimNextParameterBlockAssetRegistryExports& OutExports);
 
-	// Attempts to determine the type from a parameter name
-	// If the name cannot be found, the returned type will be invalid
-	// Note that this is expensive and can query the asset registry
-	static FAnimNextParamType GetParameterTypeFromName(FName InName);
 
 	static bool GetExportedAssetsForWorkspace(const FAssetData& InWorkspaceAsset, FAnimNextWorkspaceAssetRegistryExports& OutExports);
 };
diff --git a/Engine/Plugins/Experimental/AnimNext/Source/AnimNextEditor/Private/Param/ParamNamePropertyCustomization.cpp b/Engine/Plugins/Experimental/AnimNext/Source/AnimNextEditor/Private/Param/ParamNamePropertyCustomization.cpp
--- a/Engine/Plugins/Experimental/AnimNext/Source/AnimNextEditor/Private/Param/ParamNamePropertyCustomization.cpp
+++ b/Engine/Plugins/Experimental/AnimNext/Source/AnimNextEditor/Private/Param/ParamNamePropertyCustomization.cpp
@@ -107,7 +107,7 @@ void FParamNamePropertyTypeCustomization::Refresh()
 	{
 		if(PropertyHandle->GetValue(CachedName) == FPropertyAccess::Success)
 		{
-			CachedType = FUtils::GetParameterTypeFromName(CachedName);
+			CachedType = UncookedOnly::FUtils::GetParameterTypeFromName(CachedName);
 		}
 	}
 }
diff --git a/Engine/Plugins/Experimental/AnimNext/Source/AnimNextEditor/Private/Param/SGraphPinParamName.cpp b/Engine/Plugins/Experimental/AnimNext/Source/AnimNextEditor/Private/Param/SGraphPinParamName.cpp
--- a/Engine/Plugins/Experimental/AnimNext/Source/AnimNextEditor/Private/Param/SGraphPinParamName.cpp
+++ b/Engine/Plugins/Experimental/AnimNext/Source/AnimNextEditor/Private/Param/SGraphPinParamName.cpp
@@ -4,6 +4,7 @@
 #include "EditorUtils.h"
 #include "SParameterPickerCombo.h"
 #include "ScopedTransaction.h"
+#include "UncookedOnlyUtils.h"
 
 #define LOCTEXT_NAMESPACE "SGraphPinParamName"
 
@@ -41,7 +42,7 @@ TSharedRef<SWidget> SGraphPinParamName::GetDefaultValueWidget()
 		})
 		.OnGetParameterType_Lambda([this]()
 		{
-			return FUtils::GetParameterTypeFromName(FName(*GraphPinObj->DefaultValue));
+			return UncookedOnly::FUtils::GetParameterTypeFromName(FName(*GraphPinObj->DefaultValue));
 		});
 }
 
diff --git a/Engine/Plugins/Experimental/AnimNext/Source/AnimNextEditor/Private/Param/SParameterPicker.cpp b/Engine/Plugins/Experimental/AnimNext/Source/AnimNextEditor/Private/Param/SParameterPicker.cpp
--- a/Engine/Plugins/Experimental/AnimNext/Source/AnimNextEditor/Private/Param/SParameterPicker.cpp
+++ b/Engine/Plugins/Experimental/AnimNext/Source/AnimNextEditor/Private/Param/SParameterPicker.cpp
@@ -283,7 +283,7 @@ void SParameterPicker::RefreshEntries()
 		}
 
 		FAnimNextParameterLibraryAssetRegistryExports& NewExports = LibraryExportMap.Add(InLibrary);
-		FUtils::GetExportedParametersForLibrary(InLibrary, NewExports);
+		UncookedOnly::FUtils::GetExportedParametersForLibrary(InLibrary, NewExports);
 		return NewExports;
 	};
 	
diff --git a/Engine/Plugins/Experimental/AnimNext/Source/AnimNextTestSuite/Private/AnimNextParametersTest.cpp b/Engine/Plugins/Experimental/AnimNext/Source/AnimNextTestSuite/Private/AnimNextParametersTest.cpp
--- a/Engine/Plugins/Experimental/AnimNext/Source/AnimNextTestSuite/Private/AnimNextParametersTest.cpp
+++ b/Engine/Plugins/Experimental/AnimNext/Source/AnimNextTestSuite/Private/AnimNextParametersTest.cpp
@@ -133,9 +133,9 @@ bool FParamTypesTest::RunTest(const FString& InParameters)
 	AddErrorIfFalse(FAnimNextParamType::GetType<FVector>().IsValid(), TEXT("Struct parameter is invalid."));
 	AddErrorIfFalse(FAnimNextParamType::GetType<FTransform>().IsValid(), TEXT("Struct parameter is invalid."));
 	AddErrorIfFalse(FAnimNextParamType::GetType<FQuat>().IsValid(), TEXT("Struct parameter is invalid."));
-	AddErrorIfFalse(FAnimNextParamType::GetType<UObject>().IsValid(), TEXT("UObject parameter is invalid."));
+	AddErrorIfFalse(FAnimNextParamType::GetType<UObject*>().IsValid(), TEXT("UObject parameter is invalid."));
 	AddErrorIfFalse(FAnimNextParamType::GetType<TObjectPtr<UObject>>().IsValid(), TEXT("TObjectPtr<UObject> parameter is invalid."));
-	AddErrorIfFalse(FAnimNextParamType::GetType<UClass>().IsValid(), TEXT("UClass parameter is invalid."));
+	AddErrorIfFalse(FAnimNextParamType::GetType<UClass*>().IsValid(), TEXT("UClass parameter is invalid."));
 	AddErrorIfFalse(FAnimNextParamType::GetType<TSubclassOf<UObject>>().IsValid(), TEXT("TSubclassOff<UObject> parameter is invalid."));
 	AddErrorIfFalse(FAnimNextParamType::GetType<TSoftObjectPtr<UObject>>().IsValid(), TEXT("TSoftObjectPtr<UObject> parameter is invalid."));
 	AddErrorIfFalse(FAnimNextParamType::GetType<TSoftClassPtr<UObject>>().IsValid(), TEXT("TSoftClassPtr<UObject> parameter is invalid."));
@@ -154,9 +154,9 @@ bool FParamTypesTest::RunTest(const FString& InParameters)
 	AddErrorIfFalse(FAnimNextParamType::GetType<TArray<FVector>>().IsValid(), TEXT("Struct array parameter is invalid."));
 	AddErrorIfFalse(FAnimNextParamType::GetType<TArray<FTransform>>().IsValid(), TEXT("Struct array parameter is invalid."));
 	AddErrorIfFalse(FAnimNextParamType::GetType<TArray<FQuat>>().IsValid(), TEXT("Struct array parameter is invalid."));
-	AddErrorIfFalse(FAnimNextParamType::GetType<TArray<UObject>>().IsValid(), TEXT("UObject array parameter is invalid."));
+	AddErrorIfFalse(FAnimNextParamType::GetType<TArray<UObject*>>().IsValid(), TEXT("UObject array parameter is invalid."));
 	AddErrorIfFalse(FAnimNextParamType::GetType<TArray<TObjectPtr<UObject>>>().IsValid(), TEXT("TObjectPtr<UObject> array parameter is invalid."));
-	AddErrorIfFalse(FAnimNextParamType::GetType<TArray<UClass>>().IsValid(), TEXT("UClass array parameter is invalid."));
+	AddErrorIfFalse(FAnimNextParamType::GetType<TArray<UClass*>>().IsValid(), TEXT("UClass array parameter is invalid."));
 	AddErrorIfFalse(FAnimNextParamType::GetType<TArray<TSubclassOf<UObject>>>().IsValid(), TEXT("TSubclassOff<UObject> array parameter is invalid."));
 	AddErrorIfFalse(FAnimNextParamType::GetType<TArray<TSoftObjectPtr<UObject>>>().IsValid(), TEXT("TSoftObjectPtr<UObject> array parameter is invalid."));
 	AddErrorIfFalse(FAnimNextParamType::GetType<TArray<TSoftClassPtr<UObject>>>().IsValid(), TEXT("TSoftClassPtr<UObject> array parameter is invalid."));
@@ -176,9 +176,9 @@ bool FParamTypesTest::RunTest(const FString& InParameters)
 	AddErrorIfFalse(FAnimNextParamType::GetType<FVector>().GetHandle().IsValid(), TEXT("Struct parameter handle is invalid."));
 	AddErrorIfFalse(FAnimNextParamType::GetType<FTransform>().GetHandle().IsValid(), TEXT("Struct parameter handle is invalid."));
 	AddErrorIfFalse(FAnimNextParamType::GetType<FQuat>().GetHandle().IsValid(), TEXT("Struct parameter handle is invalid."));
-	AddErrorIfFalse(FAnimNextParamType::GetType<UObject>().GetHandle().IsValid(), TEXT("UObject parameter handle is invalid."));
+	AddErrorIfFalse(FAnimNextParamType::GetType<UObject*>().GetHandle().IsValid(), TEXT("UObject parameter handle is invalid."));
 	AddErrorIfFalse(FAnimNextParamType::GetType<TObjectPtr<UObject>>().GetHandle().IsValid(), TEXT("TObjectPtr<UObject> parameter handle is invalid."));
-	AddErrorIfFalse(FAnimNextParamType::GetType<UClass>().GetHandle().IsValid(), TEXT("UClass parameter handle is invalid."));
+	AddErrorIfFalse(FAnimNextParamType::GetType<UClass*>().GetHandle().IsValid(), TEXT("UClass parameter handle is invalid."));
 	AddErrorIfFalse(FAnimNextParamType::GetType<TSubclassOf<UObject>>().GetHandle().IsValid(), TEXT("TSubclassOff<UObject> parameter handle is invalid."));
 	AddErrorIfFalse(FAnimNextParamType::GetType<TSoftObjectPtr<UObject>>().GetHandle().IsValid(), TEXT("TSoftObjectPtr<UObject> parameter handle is invalid."));
 	AddErrorIfFalse(FAnimNextParamType::GetType<TSoftClassPtr<UObject>>().GetHandle().IsValid(), TEXT("TSoftClassPtr<UObject> parameter handle is invalid."));
@@ -197,9 +197,9 @@ bool FParamTypesTest::RunTest(const FString& InParameters)
 	AddErrorIfFalse(FAnimNextParamType::GetType<TArray<FVector>>().GetHandle().IsValid(), TEXT("Struct array parameter handle is invalid."));
 	AddErrorIfFalse(FAnimNextParamType::GetType<TArray<FTransform>>().GetHandle().IsValid(), TEXT("Struct array parameter handle is invalid."));
 	AddErrorIfFalse(FAnimNextParamType::GetType<TArray<FQuat>>().GetHandle().IsValid(), TEXT("Struct array parameter handle is invalid."));
-	AddErrorIfFalse(FAnimNextParamType::GetType<TArray<UObject>>().GetHandle().IsValid(), TEXT("UObject array parameter handle is invalid."));
+	AddErrorIfFalse(FAnimNextParamType::GetType<TArray<UObject*>>().GetHandle().IsValid(), TEXT("UObject array parameter handle is invalid."));
 	AddErrorIfFalse(FAnimNextParamType::GetType<TArray<TObjectPtr<UObject>>>().GetHandle().IsValid(), TEXT("TObjectPtr<UObject> array parameter handle is invalid."));
-	AddErrorIfFalse(FAnimNextParamType::GetType<TArray<UClass>>().GetHandle().IsValid(), TEXT("UClass array parameter handle is invalid."));
+	AddErrorIfFalse(FAnimNextParamType::GetType<TArray<UClass*>>().GetHandle().IsValid(), TEXT("UClass array parameter handle is invalid."));
 	AddErrorIfFalse(FAnimNextParamType::GetType<TArray<TSubclassOf<UObject>>>().GetHandle().IsValid(), TEXT("TSubclassOff<UObject> array parameter handle is invalid."));
 	AddErrorIfFalse(FAnimNextParamType::GetType<TArray<TSoftObjectPtr<UObject>>>().GetHandle().IsValid(), TEXT("TSoftObjectPtr<UObject> array parameter handle is invalid."));
 	AddErrorIfFalse(FAnimNextParamType::GetType<TArray<TSoftClassPtr<UObject>>>().GetHandle().IsValid(), TEXT("TSoftClassPtr<UObject> array parameter handle is invalid."));
@@ -219,9 +219,9 @@ bool FParamTypesTest::RunTest(const FString& InParameters)
 	AddErrorIfFalse(FParamTypeHandle::GetHandle<FVector>().IsValid(), TEXT("Struct parameter handle is invalid."));
 	AddErrorIfFalse(FParamTypeHandle::GetHandle<FTransform>().IsValid(), TEXT("Struct parameter handle is invalid."));
 	AddErrorIfFalse(FParamTypeHandle::GetHandle<FQuat>().IsValid(), TEXT("Struct parameter handle is invalid."));
-	AddErrorIfFalse(FParamTypeHandle::GetHandle<UObject>().IsValid(), TEXT("UObject parameter handle is invalid."));
+	AddErrorIfFalse(FParamTypeHandle::GetHandle<UObject*>().IsValid(), TEXT("UObject parameter handle is invalid."));
 	AddErrorIfFalse(FParamTypeHandle::GetHandle<TObjectPtr<UObject>>().IsValid(), TEXT("TObjectPtr<UObject> parameter handle is invalid."));
-	AddErrorIfFalse(FParamTypeHandle::GetHandle<UClass>().IsValid(), TEXT("UClass parameter handle is invalid."));
+	AddErrorIfFalse(FParamTypeHandle::GetHandle<UClass*>().IsValid(), TEXT("UClass parameter handle is invalid."));
 	AddErrorIfFalse(FParamTypeHandle::GetHandle<TSubclassOf<UObject>>().IsValid(), TEXT("TSubclassOf<UObject> parameter handle is invalid."));
 	AddErrorIfFalse(FParamTypeHandle::GetHandle<TSoftObjectPtr<UObject>>().IsValid(), TEXT("TSoftObjectPtr<UObject> parameter handle is invalid."));
 	AddErrorIfFalse(FParamTypeHandle::GetHandle<TSoftClassPtr<UObject>>().IsValid(), TEXT("TSoftClassPtr<UObject> parameter handle is invalid."));
@@ -240,9 +240,9 @@ bool FParamTypesTest::RunTest(const FString& InParameters)
 	AddErrorIfFalse(FParamTypeHandle::GetHandle<TArray<FVector>>().IsValid(), TEXT("Struct array parameter handle is invalid."));
 	AddErrorIfFalse(FParamTypeHandle::GetHandle<TArray<FTransform>>().IsValid(), TEXT("Struct array parameter handle is invalid."));
 	AddErrorIfFalse(FParamTypeHandle::GetHandle<TArray<FQuat>>().IsValid(), TEXT("Struct array parameter handle is invalid."));
-	AddErrorIfFalse(FParamTypeHandle::GetHandle<TArray<UObject>>().IsValid(), TEXT("UObject array parameter handle is invalid."));
+	AddErrorIfFalse(FParamTypeHandle::GetHandle<TArray<UObject*>>().IsValid(), TEXT("UObject array parameter handle is invalid."));
 	AddErrorIfFalse(FParamTypeHandle::GetHandle<TArray<TObjectPtr<UObject>>>().IsValid(), TEXT("TObjectPtr<UObject> array parameter handle is invalid."));
-	AddErrorIfFalse(FParamTypeHandle::GetHandle<TArray<UClass>>().IsValid(), TEXT("UClass array parameter handle is invalid."));
+	AddErrorIfFalse(FParamTypeHandle::GetHandle<TArray<UClass*>>().IsValid(), TEXT("UClass array parameter handle is invalid."));
 	AddErrorIfFalse(FParamTypeHandle::GetHandle<TArray<TSubclassOf<UObject>>>().IsValid(), TEXT("TSubclassOf<UObject> array parameter handle is invalid."));
 	AddErrorIfFalse(FParamTypeHandle::GetHandle<TArray<TSoftObjectPtr<UObject>>>().IsValid(), TEXT("TSoftObjectPtr<UObject> array parameter handle is invalid."));
 	AddErrorIfFalse(FParamTypeHandle::GetHandle<TArray<TSoftClassPtr<UObject>>>().IsValid(), TEXT("TSoftClassPtr<UObject> array parameter handle is invalid."));
@@ -266,9 +266,9 @@ bool FParamTypesTest::RunTest(const FString& InParameters)
 	// Should not be built in:
 	AddErrorIfFalse(!FParamTypeHandle::GetHandle<EPropertyBagContainerType>().IsBuiltInType(), TEXT("Enum parameter handle is built in."));
 	AddErrorIfFalse(!FParamTypeHandle::GetHandle<FAnimNextParamType>().IsBuiltInType(), TEXT("Struct parameter handle is built in."));
-	AddErrorIfFalse(!FParamTypeHandle::GetHandle<UObject>().IsBuiltInType(), TEXT("UObject parameter handle is built in."));
+	AddErrorIfFalse(!FParamTypeHandle::GetHandle<UObject*>().IsBuiltInType(), TEXT("UObject parameter handle is built in."));
 	AddErrorIfFalse(!FParamTypeHandle::GetHandle<TObjectPtr<UObject>>().IsBuiltInType(), TEXT("TObjectPtr<UObject> parameter handle is built in."));
-	AddErrorIfFalse(!FParamTypeHandle::GetHandle<UClass>().IsBuiltInType(), TEXT("UClass parameter handle is built in."));
+	AddErrorIfFalse(!FParamTypeHandle::GetHandle<UClass*>().IsBuiltInType(), TEXT("UClass parameter handle is built in."));
 	AddErrorIfFalse(!FParamTypeHandle::GetHandle<TSubclassOf<UObject>>().IsBuiltInType(), TEXT("TSubclassOff<UObject> parameter handle is built in."));
 	AddErrorIfFalse(!FParamTypeHandle::GetHandle<TSoftObjectPtr<UObject>>().IsBuiltInType(), TEXT("TSoftObjectPtr<UObject> parameter handle is built in."));
 	AddErrorIfFalse(!FParamTypeHandle::GetHandle<TSoftClassPtr<UObject>>().IsBuiltInType(), TEXT("TSoftClassPtr<UObject> parameter handle is built in."));
diff --git a/Engine/Plugins/Experimental/AnimNext/Source/AnimNextUncookedOnly/Internal/Param/AnimNextParameterLibrary.h b/Engine/Plugins/Experimental/AnimNext/Source/AnimNextUncookedOnly/Internal/Param/AnimNextParameterLibrary.h
--- a/Engine/Plugins/Experimental/AnimNext/Source/AnimNextUncookedOnly/Internal/Param/AnimNextParameterLibrary.h
+++ b/Engine/Plugins/Experimental/AnimNext/Source/AnimNextUncookedOnly/Internal/Param/AnimNextParameterLibrary.h
@@ -23,6 +23,11 @@ namespace UE::AnimNext::Editor
 	class SParameterPicker;
 }
 
+namespace UE::AnimNext::UncookedOnly
+{
+	struct FUtils;
+}
+
 namespace UE::AnimNext::UncookedOnly
 {
 	// A delegate for subscribing / reacting to parameter library modifications.
@@ -69,6 +74,7 @@ class UAnimNextParameterLibrary : public UObject
 	ANIMNEXTUNCOOKEDONLY_API static const FName ExportsAssetRegistryTag;
 
 	friend struct UE::AnimNext::Editor::FUtils;
+	friend struct UE::AnimNext::UncookedOnly::FUtils;
 	friend class UE::AnimNext::Editor::SParameterBlockView;
 	friend class UE::AnimNext::Editor::SParameterLibraryView;
 	friend class UE::AnimNext::Editor::SParameterPicker;
diff --git a/Engine/Plugins/Experimental/AnimNext/Source/AnimNextUncookedOnly/Internal/UncookedOnlyUtils.h b/Engine/Plugins/Experimental/AnimNext/Source/AnimNextUncookedOnly/Internal/UncookedOnlyUtils.h
--- a/Engine/Plugins/Experimental/AnimNext/Source/AnimNextUncookedOnly/Internal/UncookedOnlyUtils.h
+++ b/Engine/Plugins/Experimental/AnimNext/Source/AnimNextUncookedOnly/Internal/UncookedOnlyUtils.h
@@ -8,6 +8,7 @@
 #include "Param/ParamTypeHandle.h"
 #include "RigVMCore/RigVMTemplate.h"
 
+struct FAnimNextParameterLibraryAssetRegistryExports;
 class UAnimNextGraph;
 class UAnimNextGraph_EditorData;
 class UAnimNextGraph_EdGraph;
@@ -111,6 +112,13 @@ struct ANIMNEXTUNCOOKEDONLY_API FUtils
 			OutNodes.Append(GraphNodes);
 		}
 	}
+
+	static bool GetExportedParametersForLibrary(const FAssetData& InLibraryAsset, FAnimNextParameterLibraryAssetRegistryExports& OutExports);
+	
+	// Attempts to determine the type from a parameter name
+	// If the name cannot be found, the returned type will be invalid
+	// Note that this is expensive and can query the asset registry
+	static FAnimNextParamType GetParameterTypeFromName(FName InName);
 };
 
 }
\ No newline at end of file
diff --git a/Engine/Plugins/Experimental/AnimNext/Source/AnimNextUncookedOnly/Private/Graph/AnimNextGraph_EditorData.cpp b/Engine/Plugins/Experimental/AnimNext/Source/AnimNextUncookedOnly/Private/Graph/AnimNextGraph_EditorData.cpp
--- a/Engine/Plugins/Experimental/AnimNext/Source/AnimNextUncookedOnly/Private/Graph/AnimNextGraph_EditorData.cpp
+++ b/Engine/Plugins/Experimental/AnimNext/Source/AnimNextUncookedOnly/Private/Graph/AnimNextGraph_EditorData.cpp
@@ -50,7 +50,6 @@ UAnimNextGraph_EditorData::UAnimNextGraph_EditorData(const FObjectInitializer& O
 	RigVMClient.SetOuterClientHost(this, GET_MEMBER_NAME_CHECKED(UAnimNextGraph_EditorData, RigVMClient));
 	{
 		TGuardValue<bool> DisableClientNotifs(RigVMClient.bSuspendNotifications, true);
-		RigVMClient.AddModel(TEXT("RigVMGraph"), false, &ObjectInitializer);
 		RigVMClient.GetOrCreateFunctionLibrary(false, &ObjectInitializer);
 	}
 	RigVMClient.SetExecuteContextStruct(FAnimNextExecuteContext::StaticStruct());
diff --git a/Engine/Plugins/Experimental/AnimNext/Source/AnimNextUncookedOnly/Private/UncookedOnlyUtils.cpp b/Engine/Plugins/Experimental/AnimNext/Source/AnimNextUncookedOnly/Private/UncookedOnlyUtils.cpp
--- a/Engine/Plugins/Experimental/AnimNext/Source/AnimNextUncookedOnly/Private/UncookedOnlyUtils.cpp
+++ b/Engine/Plugins/Experimental/AnimNext/Source/AnimNextUncookedOnly/Private/UncookedOnlyUtils.cpp
@@ -1,6 +1,9 @@
 ﻿// Copyright Epic Games, Inc. All Rights Reserved.
 
 #include "UncookedOnlyUtils.h"
+
+#include "AssetRegistry/AssetRegistryModule.h"
+#include "AssetRegistry/IAssetRegistry.h"
 #include "Graph/AnimNextGraph.h"
 #include "Graph/AnimNextGraph_Controller.h"
 #include "Graph/AnimNextGraph_EditorData.h"
@@ -29,6 +32,8 @@
 #include "DecoratorBase/DecoratorRegistry.h"
 #include "DecoratorBase/Decorator.h"
 #include "Graph/RigUnit_AnimNextBeginExecution.h"
+#include "Param/AnimNextParameterLibrary.h"
+#include "Param/Params.h"
 #include "Serialization/MemoryReader.h"
 #include "RigVMRuntimeDataRegistry.h"
 
@@ -1084,6 +1089,43 @@ FText FUtils::GetParameterDisplayNameText(FName InParameterName)
 	return FText::FromString(NameAsString);
 }
 
+bool FUtils::GetExportedParametersForLibrary(const FAssetData& InLibraryAsset, FAnimNextParameterLibraryAssetRegistryExports& OutExports)
+{
+	const FString TagValue = InLibraryAsset.GetTagValueRef<FString>(UAnimNextParameterLibrary::ExportsAssetRegistryTag);
+	return FAnimNextParameterLibraryAssetRegistryExports::StaticStruct()->ImportText(*TagValue, &OutExports, nullptr, PPF_None, nullptr, FAnimNextParameterLibraryAssetRegistryExports::StaticStruct()->GetName()) != nullptr;
+}
+
+FAnimNextParamType FUtils::GetParameterTypeFromName(FName InName)
+{
+	// Check built-in params first as they are cheaper
+	if(const FParamDefinition* FoundDefinition = FParams::FindBuiltInParameter(InName))
+	{
+		return FoundDefinition->Type;
+	}
+
+	// Query the asset registry for other params
+	IAssetRegistry& AssetRegistry = FModuleManager::Get().LoadModuleChecked<FAssetRegistryModule>(TEXT("AssetRegistry")).Get();
+
+	FARFilter ARFilter;
+	ARFilter.ClassPaths = { UAnimNextParameterLibrary::StaticClass()->GetClassPathName() };
 
+	TArray<FAssetData> LibraryAssets;
+	AssetRegistry.GetAssets(ARFilter, LibraryAssets);
+
+	for(const FAssetData& LibraryAsset : LibraryAssets)
+	{
+		FAnimNextParameterLibraryAssetRegistryExports Exports;
+		GetExportedParametersForLibrary(LibraryAsset, Exports);
+		for(const FAnimNextParameterLibraryAssetRegistryExportEntry& Export : Exports.Parameters)
+		{
+			if(Export.Name == InName)
+			{
+				return Export.Type;
+			}
+		}
+	}
+
+	return FAnimNextParamType();
+}
 
 }
diff --git a/Engine/Plugins/Runtime/RigVM/Source/RigVMEditor/Private/Widgets/SRigVMGraphChangePinType.cpp b/Engine/Plugins/Runtime/RigVM/Source/RigVMEditor/Private/Widgets/SRigVMGraphChangePinType.cpp
--- a/Engine/Plugins/Runtime/RigVM/Source/RigVMEditor/Private/Widgets/SRigVMGraphChangePinType.cpp
+++ b/Engine/Plugins/Runtime/RigVM/Source/RigVMEditor/Private/Widgets/SRigVMGraphChangePinType.cpp
@@ -221,7 +221,7 @@ void SRigVMGraphChangePinType::FillPinTypeMenu(FMenuBuilder& MenuBuilder)
 
 			if(bHasAllTypes && Type.CPPTypeObject)
 			{
-				if(Type.CPPTypeObject->IsA<UEnum>() || Type.CPPTypeObject->IsA<UClass>())
+				if(Type.CPPTypeObject->IsA<UEnum>())
 				{
 					continue;
 				}
diff --git a/Engine/Source/Runtime/Engine/Classes/Components/SkeletalMeshComponent.h b/Engine/Source/Runtime/Engine/Classes/Components/SkeletalMeshComponent.h
--- a/Engine/Source/Runtime/Engine/Classes/Components/SkeletalMeshComponent.h
+++ b/Engine/Source/Runtime/Engine/Classes/Components/SkeletalMeshComponent.h
@@ -65,7 +65,7 @@ DECLARE_DYNAMIC_MULTICAST_DELEGATE(FOnBoneTransformsFinalized);  // Deprecated,
 
 DECLARE_MULTICAST_DELEGATE(FOnBoneTransformsFinalizedMultiCast);
 
-DECLARE_MULTICAST_DELEGATE_ThreeParams(FOnLODRequiredBonesUpdateMulticast, USkeletalMeshComponent*, int32, const TArray<FBoneIndexType>&);
+DECLARE_TS_MULTICAST_DELEGATE_ThreeParams(FOnLODRequiredBonesUpdateMulticast, USkeletalMeshComponent*, int32, const TArray<FBoneIndexType>&);
 typedef FOnLODRequiredBonesUpdateMulticast::FDelegate FOnLODRequiredBonesUpdate;
 
 /** Method used when retrieving a attribute value*/
diff --git a/Engine/Source/Runtime/Engine/Private/Components/SkeletalMeshComponent.cpp b/Engine/Source/Runtime/Engine/Private/Components/SkeletalMeshComponent.cpp
--- a/Engine/Source/Runtime/Engine/Private/Components/SkeletalMeshComponent.cpp
+++ b/Engine/Source/Runtime/Engine/Private/Components/SkeletalMeshComponent.cpp
@@ -2835,7 +2835,7 @@ FBoxSphereBounds USkeletalMeshComponent::CalcBounds(const FTransform& LocalToWor
 			}
 		}
 
-		const bool bCacheLocalSpaceBounds = CVarCacheLocalSpaceBounds.GetValueOnGameThread() != 0;
+		const bool bCacheLocalSpaceBounds = CVarCacheLocalSpaceBounds.GetValueOnAnyThread() != 0;
 		
 		const FTransform CachedBoundsTransform = bCacheLocalSpaceBounds ? FTransform::Identity : LocalToWorld;
 
diff --git a/Engine/Source/Runtime/Engine/Private/SkeletalRenderGPUSkin.cpp b/Engine/Source/Runtime/Engine/Private/SkeletalRenderGPUSkin.cpp
--- a/Engine/Source/Runtime/Engine/Private/SkeletalRenderGPUSkin.cpp
+++ b/Engine/Source/Runtime/Engine/Private/SkeletalRenderGPUSkin.cpp
@@ -77,7 +77,7 @@ static int32 GetRayTracingSkeletalMeshGlobalLODBias()
 		IConsoleManager::Get().FindTConsoleVariableDataInt(TEXT("r.RayTracing.Geometry.SkeletalMeshes.LODBias"));
 
 	return !RayTracingSkeletalMeshLODBiasVar ? 0 :
-		FMath::Max(0, RayTracingSkeletalMeshLODBiasVar->GetValueOnGameThread());  // Only allows positive bias to narrow cloth mapping requirements
+		FMath::Max(0, RayTracingSkeletalMeshLODBiasVar->GetValueOnAnyThread());  // Only allows positive bias to narrow cloth mapping requirements
 }
 
 /*-----------------------------------------------------------------------------
diff --git a/Engine/Source/Runtime/Engine/Public/Animation/AnimCompressionTypes.h b/Engine/Source/Runtime/Engine/Public/Animation/AnimCompressionTypes.h
--- a/Engine/Source/Runtime/Engine/Public/Animation/AnimCompressionTypes.h
+++ b/Engine/Source/Runtime/Engine/Public/Animation/AnimCompressionTypes.h
@@ -872,14 +872,14 @@ struct FRootMotionReset
 
 	bool bIsValidAdditive;
 
-	void ResetRootBoneForRootMotion(FTransform& BoneTransform, const FBoneContainer& RequiredBones) const
+	void ResetRootBoneForRootMotion(FTransform& BoneTransform, const FTransform& RefPoseRootTransform) const
 	{
 		switch (RootMotionRootLock)
 		{
 		case ERootMotionRootLock::AnimFirstFrame: BoneTransform = AnimFirstFrame; break;
 		case ERootMotionRootLock::Zero: BoneTransform = FTransform::Identity; break;
 		default:
-		case ERootMotionRootLock::RefPose: BoneTransform = RequiredBones.GetRefPoseArray()[0]; break;
+		case ERootMotionRootLock::RefPose: BoneTransform = RefPoseRootTransform; break;
 		}
 
 		if (bIsValidAdditive && RootMotionRootLock != ERootMotionRootLock::AnimFirstFrame)
@@ -888,6 +888,11 @@ struct FRootMotionReset
 			BoneTransform.SetScale3D(BoneTransform.GetScale3D() - FVector(1.f));
 		}
 	}
+
+	void ResetRootBoneForRootMotion(FTransform& BoneTransform, const FBoneContainer& RequiredBones) const
+	{
+		ResetRootBoneForRootMotion(BoneTransform, RequiredBones.GetRefPoseArray()[0]);
+	}
 };
 
 #if WITH_EDITOR
