Rebase of AnimNext_Graph on top of RigVMHost

diff --git a/Engine/Plugins/Experimental/AnimNext/Source/AnimNext/Internal/Graph/AnimNextGraph.h b/Engine/Plugins/Experimental/AnimNext/Source/AnimNext/Internal/Graph/AnimNextGraph.h
--- a/Engine/Plugins/Experimental/AnimNext/Source/AnimNext/Internal/Graph/AnimNextGraph.h
+++ b/Engine/Plugins/Experimental/AnimNext/Source/AnimNext/Internal/Graph/AnimNextGraph.h
@@ -10,6 +10,8 @@
 #include "Graph/RigUnit_AnimNextGraphEvaluator.h"
 #include "Param/ParamId.h"
 #include "Scheduler/IAnimNextScheduleTermInterface.h"
+#include "RigVMHost.h"
+
 #include "AnimNextGraph.generated.h"
 
 class UEdGraph;
@@ -77,11 +79,13 @@ private:
 
 // A user-created graph of logic used to supply data
 UCLASS(BlueprintType)
-class ANIMNEXT_API UAnimNextGraph : public UObject, public IAnimNextScheduleTermInterface
+class ANIMNEXT_API UAnimNextGraph :  public URigVMHost, public IAnimNextScheduleTermInterface
 {
 	GENERATED_BODY()
 
 public:
+	UAnimNextGraph(const FObjectInitializer& ObjectInitializer);
+
 	// UObject interface
 	virtual void PostLoad() override;
 	virtual void PostRename(UObject* OldOuter, const FName OldName) override;
@@ -146,15 +150,13 @@ protected:
 	UPROPERTY()
 	TObjectPtr<URigVM> RigVM;
 
-	// The ExtendedExecuteContext object holds instance/work data used by the RigVM internals. It is populated during compilation same as the URigVM object above.
-	// We also have a 1:1 mapping with the UAnimNextGraph but each instance of the anim graph also needs its own exetended execute context object. This one is used
-	// as a reference we copy from.
-	UPROPERTY()
+	// The ExtendedExecuteContext object holds the common work data used by the RigVM internals. It is populated during the initial VM initialization.
+	// Each instance of an AnimGraph requires a copy of this context and a call to initialize the VM instance with the context copy, 
+	// so the cached memory handles are updated to the correct memory addresses.
+	// This context is used as a reference to copy the common data for all instances created.
+	UPROPERTY(Transient)
 	FRigVMExtendedExecuteContext ExtendedExecuteContext;
 
-	UPROPERTY()
-	FRigVMRuntimeSettings VMRuntimeSettings;
-
 	// The parameter to use to access the reference pose
 	UPROPERTY(EditAnywhere, Category = "Graph", meta=(CustomWidget = "ParamName", AllowedParamType = "FAnimNextGraphReferencePose"))
 	FName ReferencePose = TEXT("UE_AnimNextMeshComponent_ReferencePose");
diff --git a/Engine/Plugins/Experimental/AnimNext/Source/AnimNext/Private/Graph/AnimNextGraph.cpp b/Engine/Plugins/Experimental/AnimNext/Source/AnimNext/Private/Graph/AnimNextGraph.cpp
--- a/Engine/Plugins/Experimental/AnimNext/Source/AnimNext/Private/Graph/AnimNextGraph.cpp
+++ b/Engine/Plugins/Experimental/AnimNext/Source/AnimNext/Private/Graph/AnimNextGraph.cpp
@@ -48,6 +48,12 @@ bool FAnimNextGraphInstance::IsValid() const
 	return GraphInstancePtr.IsValid();
 }
 
+UAnimNextGraph::UAnimNextGraph(const FObjectInitializer& ObjectInitializer)
+	: Super(ObjectInitializer)
+{
+	SetRigVMExtendedExecuteContext(&ExtendedExecuteContext);
+}
+
 void UAnimNextGraph::AllocateInstance(FAnimNextGraphInstance& Instance) const
 {
 	SCOPE_CYCLE_COUNTER(STAT_AnimNext_Graph_AllocateInstance);
@@ -66,21 +72,21 @@ void UAnimNextGraph::AllocateInstance(FAnimNextGraphInstance& Instance) const
 
 	Instance.ExtendedExecuteContext.CopyMemoryStorage(ExtendedExecuteContext);
 
-	RigVM->InitializeInstance(Instance.ExtendedExecuteContext);
+	VM->InitializeInstance(Instance.ExtendedExecuteContext);
 }
 
 void UAnimNextGraph::Run(const UE::AnimNext::FContext& Context, FAnimNextGraphInstance& GraphInstance, EAnimNextGraphSimulationSteps SimulationSteps) const
 {
 	SCOPE_CYCLE_COUNTER(STAT_AnimNext_Graph);
 
-	if (RigVM && GraphInstance.IsValid())
+	if (VM && GraphInstance.IsValid())
 	{
 		FAnimNextExecuteContext& AnimNextContext = GraphInstance.ExtendedExecuteContext.GetPublicDataSafe<FAnimNextExecuteContext>();
 		AnimNextContext.SetContextData(Context);
 		AnimNextContext.InitializeWithGraph(this, SharedDataBuffer, GraphInstance.GraphInstancePtr);
 		AnimNextContext.SetSimulationSteps(SimulationSteps);
 
-		RigVM->ExecuteVM(GraphInstance.ExtendedExecuteContext, FRigUnit_AnimNextShimRoot::EventName);
+		VM->ExecuteVM(GraphInstance.ExtendedExecuteContext, FRigUnit_AnimNextShimRoot::EventName);
 
 		// Reset the context to avoid issues if we forget to reset it the next time we use it
 		AnimNextContext.DebugReset();
@@ -115,13 +121,18 @@ void UAnimNextGraph::PostLoad()
 	
 	Super::PostLoad();
 
+	VM = RigVM;
+
 	ExtendedExecuteContext.InvalidateCachedMemory();
 
 	// In packaged builds, initialize the VM
 	// In editor, the VM will be recompiled and initialized at RecompileVM
 #if !WITH_EDITOR
-	RigVM->ClearExternalVariables(ExtendedExecuteContext);
-	RigVM->Initialize(ExtendedExecuteContext);
+	if (VM)
+	{
+		VM->ClearExternalVariables(ExtendedExecuteContext);
+		VM->Initialize(ExtendedExecuteContext);
+	}
 #endif
 
 	ReferencePoseId = FParamId(ReferencePose);
diff --git a/Engine/Plugins/Experimental/AnimNext/Source/AnimNextUncookedOnly/Private/Graph/AnimNextGraph_EditorData.cpp b/Engine/Plugins/Experimental/AnimNext/Source/AnimNextUncookedOnly/Private/Graph/AnimNextGraph_EditorData.cpp
--- a/Engine/Plugins/Experimental/AnimNext/Source/AnimNextUncookedOnly/Private/Graph/AnimNextGraph_EditorData.cpp
+++ b/Engine/Plugins/Experimental/AnimNext/Source/AnimNextUncookedOnly/Private/Graph/AnimNextGraph_EditorData.cpp
@@ -387,9 +387,9 @@ void UAnimNextGraph_EditorData::HandleConfigureRigVMController(const FRigVMClien
 		{
 			if(UAnimNextGraph* Graph = WeakThis->GetTypedOuter<UAnimNextGraph>())
 			{
-				if (Graph->RigVM)
+				if (Graph->VM)
 				{
-					return &Graph->RigVM->GetByteCode();
+					return &Graph->VM->GetByteCode();
 				}
 			}
 		}
diff --git a/Engine/Plugins/Experimental/AnimNext/Source/AnimNextUncookedOnly/Private/UncookedOnlyUtils.cpp b/Engine/Plugins/Experimental/AnimNext/Source/AnimNextUncookedOnly/Private/UncookedOnlyUtils.cpp
--- a/Engine/Plugins/Experimental/AnimNext/Source/AnimNextUncookedOnly/Private/UncookedOnlyUtils.cpp
+++ b/Engine/Plugins/Experimental/AnimNext/Source/AnimNextUncookedOnly/Private/UncookedOnlyUtils.cpp
@@ -477,10 +477,10 @@ void FUtils::Compile(UAnimNextGraph* InGraph)
 	URigVMCompiler* Compiler = URigVMCompiler::StaticClass()->GetDefaultObject<URigVMCompiler>();
 	EditorData->VMCompileSettings.SetExecuteContextStruct(EditorData->RigVMClient.GetExecuteContextStruct());
 	const FRigVMCompileSettings Settings = (EditorData->bCompileInDebugMode) ? FRigVMCompileSettings::Fast(EditorData->VMCompileSettings.GetExecuteContextStruct()) : EditorData->VMCompileSettings;
-	Compiler->Compile(Settings, { VMTempGraph }, TempController, InGraph->RigVM, InGraph->ExtendedExecuteContext, InGraph->GetRigVMExternalVariables(), & EditorData->PinToOperandMap);
+	Compiler->Compile(Settings, { VMTempGraph }, TempController, InGraph->VM, InGraph->ExtendedExecuteContext, InGraph->GetRigVMExternalVariables(), & EditorData->PinToOperandMap);
 
 	// Initialize right away, in packaged builds we initialize during PostLoad
-	InGraph->RigVM->Initialize(InGraph->ExtendedExecuteContext);
+	InGraph->VM->Initialize(InGraph->ExtendedExecuteContext);
 
 	if (EditorData->bErrorsDuringCompilation)
 	{
@@ -491,9 +491,9 @@ void FUtils::Compile(UAnimNextGraph* InGraph)
 	}
 
 	EditorData->bVMRecompilationRequired = false;
-	if(InGraph->RigVM)
+	if(InGraph->VM)
 	{
-		EditorData->RigVMCompiledEvent.Broadcast(InGraph, InGraph->RigVM, InGraph->ExtendedExecuteContext);
+		EditorData->RigVMCompiledEvent.Broadcast(InGraph, InGraph->VM, InGraph->ExtendedExecuteContext);
 	}
 
 	VMClient->RemoveController(VMTempGraph);
@@ -505,8 +505,9 @@ void FUtils::Compile(UAnimNextGraph* InGraph)
 
 void FUtils::RecreateVM(UAnimNextGraph* InGraph)
 {
-	InGraph->RigVM = NewObject<URigVM>(InGraph, TEXT("VM"), RF_NoFlags);
-	InGraph->RigVM->Reset(InGraph->ExtendedExecuteContext);
+	InGraph->VM = NewObject<URigVM>(InGraph, TEXT("VM"), RF_NoFlags);
+	InGraph->VM->Reset(InGraph->ExtendedExecuteContext);
+	InGraph->RigVM = InGraph->VM; // Local serialization
 }
 
 UAnimNextGraph_EditorData* FUtils::GetEditorData(const UAnimNextGraph* InAnimNextGraph)
