AnimNext: Cooked standalone and editor -game fixes

Added a minimal worker-thread tick for anim next's mesh component. This ensures that the recently rendered flag and the LOD get updated correctly.
Removed GEditor check around HandlePackageDone callback, which fixes editor -game builds.
Copy hash to prevent parameter block VMs being reinitialized each time.

Additionally:
Fixed uninitialized member.
Removed some dead code around memory storage classes.
Trimmed "K2_" prefixes from getters in adapter code.

diff --git a/Engine/Plugins/Experimental/AnimNext/Source/AnimNext/Internal/Graph/AnimNextGraph.h b/Engine/Plugins/Experimental/AnimNext/Source/AnimNext/Internal/Graph/AnimNextGraph.h
--- a/Engine/Plugins/Experimental/AnimNext/Source/AnimNext/Internal/Graph/AnimNextGraph.h
+++ b/Engine/Plugins/Experimental/AnimNext/Source/AnimNext/Internal/Graph/AnimNextGraph.h
@@ -88,8 +88,6 @@ public:
 
 	// UObject interface
 	virtual void PostLoad() override;
-	virtual void PostRename(UObject* OldOuter, const FName OldName) override;
-	virtual void GetPreloadDependencies(TArray<UObject*>& OutDeps) override;
 	virtual void Serialize(FArchive& Ar) override;
 
 	// IAnimNextScheduleTermInterface interface
diff --git a/Engine/Plugins/Experimental/AnimNext/Source/AnimNext/Internal/Graph/RigUnit_AnimNextGraphEvaluator.h b/Engine/Plugins/Experimental/AnimNext/Source/AnimNext/Internal/Graph/RigUnit_AnimNextGraphEvaluator.h
--- a/Engine/Plugins/Experimental/AnimNext/Source/AnimNext/Internal/Graph/RigUnit_AnimNextGraphEvaluator.h
+++ b/Engine/Plugins/Experimental/AnimNext/Source/AnimNext/Internal/Graph/RigUnit_AnimNextGraphEvaluator.h
@@ -26,7 +26,7 @@ struct FAnimNextGraphEvaluatorExecuteDefinition
 	GENERATED_BODY();
 
 	UPROPERTY()
-	uint32 Hash;
+	uint32 Hash = 0;
 
 	UPROPERTY()
 	FString MethodName;
diff --git a/Engine/Plugins/Experimental/AnimNext/Source/AnimNext/Internal/Param/AnimNextParameterBlock.h b/Engine/Plugins/Experimental/AnimNext/Source/AnimNext/Internal/Param/AnimNextParameterBlock.h
--- a/Engine/Plugins/Experimental/AnimNext/Source/AnimNext/Internal/Param/AnimNextParameterBlock.h
+++ b/Engine/Plugins/Experimental/AnimNext/Source/AnimNext/Internal/Param/AnimNextParameterBlock.h
@@ -85,8 +85,6 @@ class UAnimNextParameterBlock : public URigVMHost, public IAnimNextParameterSour
 
 	// UObject interface
 	virtual void PostLoad() override;
-	virtual void PostRename(UObject* OldOuter, const FName OldName) override;
-	virtual void GetPreloadDependencies(TArray<UObject*>& OutDeps) override;
 	virtual void GetAssetRegistryTags(TArray<FAssetRegistryTag>& OutTags) const override;
 	
 	FRigVMExtendedExecuteContext BaseRigVMContext;
diff --git a/Engine/Plugins/Experimental/AnimNext/Source/AnimNext/Private/Component/AnimNextMeshComponent.cpp b/Engine/Plugins/Experimental/AnimNext/Source/AnimNext/Private/Component/AnimNextMeshComponent.cpp
--- a/Engine/Plugins/Experimental/AnimNext/Source/AnimNext/Private/Component/AnimNextMeshComponent.cpp
+++ b/Engine/Plugins/Experimental/AnimNext/Source/AnimNext/Private/Component/AnimNextMeshComponent.cpp
@@ -2,13 +2,18 @@
 
 #include "AnimNextMeshComponent.h"
 #include "GenerationTools.h"
+#include "Engine/World.h"
 
 UAnimNextMeshComponent::UAnimNextMeshComponent()
 {
-	// Disable regular skeletal mesh ticking
-	PrimaryComponentTick.bCanEverTick = false;
-	PrimaryComponentTick.bStartWithTickEnabled = false;
-	PrimaryComponentTick.bAllowTickOnDedicatedServer = false;
+	PrimaryComponentTick.bRunOnAnyThread = true;
+}
+
+void UAnimNextMeshComponent::TickComponent(float DeltaTime, enum ELevelTick TickType, FActorComponentTickFunction *ThisTickFunction)
+{
+	// Simple tick mimicking a minimal tick of SkinnedMeshComponent
+	bRecentlyRendered = (GetLastRenderTime() > GetWorld()->TimeSeconds - 1.0f);
+	UpdateLODStatus();
 }
 
 void UAnimNextMeshComponent::CompleteAndDispatch(TConstArrayView<FBoneIndexType> InParentIndices, TConstArrayView<FBoneIndexType> InRequiredBoneIndices, TConstArrayView<FTransform> InLocalSpaceTransforms)
diff --git a/Engine/Plugins/Experimental/AnimNext/Source/AnimNext/Private/Component/AnimNextMeshComponent.h b/Engine/Plugins/Experimental/AnimNext/Source/AnimNext/Private/Component/AnimNextMeshComponent.h
--- a/Engine/Plugins/Experimental/AnimNext/Source/AnimNext/Private/Component/AnimNextMeshComponent.h
+++ b/Engine/Plugins/Experimental/AnimNext/Source/AnimNext/Private/Component/AnimNextMeshComponent.h
@@ -17,6 +17,8 @@ class UAnimNextMeshComponent : public USkeletalMeshComponent
 
 	UAnimNextMeshComponent();
 
+	virtual void TickComponent(float DeltaTime, enum ELevelTick TickType, FActorComponentTickFunction *ThisTickFunction) override;
+
 	// Converts local to component space, flips buffers, updates bounds and dispatches to renderer
 	void CompleteAndDispatch(TConstArrayView<FBoneIndexType> InParentIndices, TConstArrayView<FBoneIndexType> InRequiredBoneIndices, TConstArrayView<FTransform> InLocalSpaceTransforms);
 };
\ No newline at end of file
diff --git a/Engine/Plugins/Experimental/AnimNext/Source/AnimNext/Private/Graph/AnimNextGraph.cpp b/Engine/Plugins/Experimental/AnimNext/Source/AnimNext/Private/Graph/AnimNextGraph.cpp
--- a/Engine/Plugins/Experimental/AnimNext/Source/AnimNext/Private/Graph/AnimNextGraph.cpp
+++ b/Engine/Plugins/Experimental/AnimNext/Source/AnimNext/Private/Graph/AnimNextGraph.cpp
@@ -98,23 +98,6 @@ TArray<FRigVMExternalVariable> UAnimNextGraph::GetRigVMExternalVariables()
 	return TArray<FRigVMExternalVariable>(); 
 }
 
-static TArray<UClass*> GetClassObjectsInPackage(UPackage* InPackage)
-{
-	TArray<UObject*> Objects;
-	GetObjectsWithOuter(InPackage, Objects, false);
-
-	TArray<UClass*> ClassObjects;
-	for (UObject* Object : Objects)
-	{
-		if (UClass* Class = Cast<UClass>(Object))
-		{
-			ClassObjects.Add(Class);
-		}
-	}
-
-	return ClassObjects;
-}
-
 void UAnimNextGraph::PostLoad()
 {
 	using namespace UE::AnimNext;
@@ -140,42 +123,6 @@ void UAnimNextGraph::PostLoad()
 	DeltaTimeId = FParamId(DeltaTime);
 }
 
-void UAnimNextGraph::PostRename(UObject* OldOuter, const FName OldName)
-{
-	Super::PostRename(OldOuter, OldName);
-
-	// Whenever the asset is renamed/moved, generated classes parented to the old package
-	// are not moved to the new package automatically (see FAssetRenameManager), so we
-	// have to manually perform the move/rename, to avoid invalid reference to the old package
-
-	// Note: while asset duplication doesn't duplicate the classes either, it is not a problem there
-	// because we always recompile in post duplicate.
-	TArray<UClass*> ClassObjects = GetClassObjectsInPackage(OldOuter->GetPackage());
-
-	for (UClass* ClassObject : ClassObjects)
-	{
-		if (URigVMMemoryStorageGeneratorClass* MemoryClass = Cast<URigVMMemoryStorageGeneratorClass>(ClassObject))
-		{
-			MemoryClass->Rename(nullptr, GetPackage(), REN_ForceNoResetLoaders | REN_DoNotDirty | REN_DontCreateRedirectors | REN_NonTransactional);
-		}
-	}
-}
-
-void UAnimNextGraph::GetPreloadDependencies(TArray<UObject*>& OutDeps)
-{
-	Super::GetPreloadDependencies(OutDeps);
-
-	TArray<UClass*> ClassObjects = GetClassObjectsInPackage(GetPackage());
-
-	for (UClass* ClassObject : ClassObjects)
-	{
-		if (URigVMMemoryStorageGeneratorClass* MemoryClass = Cast<URigVMMemoryStorageGeneratorClass>(ClassObject))
-		{
-			OutDeps.Add(MemoryClass);
-		}
-	}
-}
-
 void UAnimNextGraph::Serialize(FArchive& Ar)
 {
 	Super::Serialize(Ar);
diff --git a/Engine/Plugins/Experimental/AnimNext/Source/AnimNext/Private/Param/AnimNextParameterBlock.cpp b/Engine/Plugins/Experimental/AnimNext/Source/AnimNext/Private/Param/AnimNextParameterBlock.cpp
--- a/Engine/Plugins/Experimental/AnimNext/Source/AnimNext/Private/Param/AnimNextParameterBlock.cpp
+++ b/Engine/Plugins/Experimental/AnimNext/Source/AnimNext/Private/Param/AnimNextParameterBlock.cpp
@@ -13,28 +13,6 @@ DEFINE_STAT(STAT_AnimNext_ParamBlock_UpdateLayer);
 
 #include UE_INLINE_GENERATED_CPP_BY_NAME(AnimNextParameterBlock)
 
-namespace UE::AnimNext::Private
-{
-
-static TArray<UClass*> GetClassObjectsInPackage(UPackage* InPackage)
-{
-	TArray<UObject*> Objects;
-	GetObjectsWithOuter(InPackage, Objects, false);
-
-	TArray<UClass*> ClassObjects;
-	for (UObject* Object : Objects)
-	{
-		if (UClass* Class = Cast<UClass>(Object))
-		{
-			ClassObjects.Add(Class);
-		}
-	}
-
-	return ClassObjects;
-}
-
-}
-
 UAnimNextParameterBlock::UAnimNextParameterBlock(const FObjectInitializer& ObjectInitializer)
 	: Super(ObjectInitializer)
 {
@@ -52,6 +30,7 @@ void UAnimNextParameterBlock::UpdateLayer(UE::AnimNext::FParamStackLayerHandle&
 		{
 			RuntimeData = UE::AnimNext::FRigVMRuntimeDataRegistry::AddRuntimeData(VM, GetRigVMExtendedExecuteContext());
 			VM->InitializeInstance(RuntimeData->Context, false); // TODO zzz : Temp until VM supports WorkData cached handles using offsets (UE-197067)
+			RuntimeData->Context.VMHash = VM->GetVMHash();
 		}
 
 		FRigVMExtendedExecuteContext& Context = RuntimeData->Context;
@@ -84,42 +63,6 @@ bool UAnimNextParameterBlock::ShouldCacheLayer(const UE::AnimNext::FParamStackLa
 	return false;
 }
 
-void UAnimNextParameterBlock::PostRename(UObject* OldOuter, const FName OldName)
-{
-	Super::PostRename(OldOuter, OldName);
-
-	// Whenever the asset is renamed/moved, generated classes parented to the old package
-	// are not moved to the new package automatically (see FAssetRenameManager), so we
-	// have to manually perform the move/rename, to avoid invalid reference to the old package
-
-	// Note: while asset duplication doesn't duplicate the classes either, it is not a problem there
-	// because we always recompile in post duplicate.
-	TArray<UClass*> ClassObjects = UE::AnimNext::Private::GetClassObjectsInPackage(OldOuter->GetPackage());
-
-	for (UClass* ClassObject : ClassObjects)
-	{
-		if (URigVMMemoryStorageGeneratorClass* MemoryClass = Cast<URigVMMemoryStorageGeneratorClass>(ClassObject))
-		{
-			MemoryClass->Rename(nullptr, GetPackage(), REN_ForceNoResetLoaders | REN_DoNotDirty | REN_DontCreateRedirectors | REN_NonTransactional);
-		}
-	}
-}
-
-void UAnimNextParameterBlock::GetPreloadDependencies(TArray<UObject*>& OutDeps)
-{
-	Super::GetPreloadDependencies(OutDeps);
-
-	TArray<UClass*> ClassObjects = UE::AnimNext::Private::GetClassObjectsInPackage(GetPackage());
-
-	for (UClass* ClassObject : ClassObjects)
-	{
-		if (URigVMMemoryStorageGeneratorClass* MemoryClass = Cast<URigVMMemoryStorageGeneratorClass>(ClassObject))
-		{
-			OutDeps.Add(MemoryClass);
-		}
-	}
-}
-
 void UAnimNextParameterBlock::PostLoad()
 {
 	Super::PostLoad();
@@ -131,6 +74,7 @@ void UAnimNextParameterBlock::PostLoad()
 #if !WITH_EDITOR
 	if(VM != nullptr)
 	{
+		VM->Initialize(BaseRigVMContext);
 		InitializeVM(FRigUnit_AnimNextBeginExecution::EventName);
 	}
 #endif
diff --git a/Engine/Plugins/Experimental/AnimNext/Source/AnimNext/Private/Param/ParamId.cpp b/Engine/Plugins/Experimental/AnimNext/Source/AnimNext/Private/Param/ParamId.cpp
--- a/Engine/Plugins/Experimental/AnimNext/Source/AnimNext/Private/Param/ParamId.cpp
+++ b/Engine/Plugins/Experimental/AnimNext/Source/AnimNext/Private/Param/ParamId.cpp
@@ -324,6 +324,7 @@ void FParamId::RefreshConfigAdapters()
 						if(FParamTypeHandle::FromProperty(ReturnProperty).IsValid())
 						{
 							FString FunctionName = Function->GetName();
+							FunctionName.RemoveFromStart(TEXT("K2_"));
 							FunctionName.RemoveFromStart(TEXT("Get"));
 							
 							TStringBuilder<128> ParameterNameBuilder;
diff --git a/Engine/Plugins/Experimental/AnimNext/Source/AnimNextUncookedOnly/Private/Graph/AnimNextGraph_EditorData.cpp b/Engine/Plugins/Experimental/AnimNext/Source/AnimNextUncookedOnly/Private/Graph/AnimNextGraph_EditorData.cpp
--- a/Engine/Plugins/Experimental/AnimNext/Source/AnimNextUncookedOnly/Private/Graph/AnimNextGraph_EditorData.cpp
+++ b/Engine/Plugins/Experimental/AnimNext/Source/AnimNextUncookedOnly/Private/Graph/AnimNextGraph_EditorData.cpp
@@ -158,11 +158,8 @@ void UAnimNextGraph_EditorData::PostLoad()
 		Entries.Add(InLoadedEntry);
 	});
 	
-	if (GIsEditor)
-	{
-		// delay compilation until the package has been loaded
-		FCoreUObjectDelegates::OnEndLoadPackage.AddUObject(this, &UAnimNextGraph_EditorData::HandlePackageDone);
-	}
+	// delay compilation until the package has been loaded
+	FCoreUObjectDelegates::OnEndLoadPackage.AddUObject(this, &UAnimNextGraph_EditorData::HandlePackageDone);
 #else // !WITH_EDITOR
 	RecompileVMIfRequired();
 #endif // WITH_EDITOR
diff --git a/Engine/Plugins/Experimental/AnimNext/Source/AnimNextUncookedOnly/Private/Param/AnimNextParameterBlock_EditorData.cpp b/Engine/Plugins/Experimental/AnimNext/Source/AnimNextUncookedOnly/Private/Param/AnimNextParameterBlock_EditorData.cpp
--- a/Engine/Plugins/Experimental/AnimNext/Source/AnimNextUncookedOnly/Private/Param/AnimNextParameterBlock_EditorData.cpp
+++ b/Engine/Plugins/Experimental/AnimNext/Source/AnimNextUncookedOnly/Private/Param/AnimNextParameterBlock_EditorData.cpp
@@ -587,12 +587,9 @@ void UAnimNextParameterBlock_EditorData::PostLoad()
 		check(IsValid(InLoadedEntry));
 		Entries.Add(InLoadedEntry);
 	});
-
-	if (GIsEditor)
-	{
-		// delay compilation until the package has been loaded
-		FCoreUObjectDelegates::OnEndLoadPackage.AddUObject(this, &UAnimNextParameterBlock_EditorData::HandlePackageDone);
-	}
+	
+	// delay compilation until the package has been loaded
+	FCoreUObjectDelegates::OnEndLoadPackage.AddUObject(this, &UAnimNextParameterBlock_EditorData::HandlePackageDone);
 #else // !WITH_EDITOR
 	RecompileVMIfRequired();
 #endif // WITH_EDITOR
