AnimNext: Workflow improvements

Opening workspace-hosted assets now correctly opens asset tabs in already-opened workspaces. Graphs assets also open their (root) graph by default.
Fixed context menu padding and added icons to units/dispatch factories. Consolidated graph and param block action menus into a common class.
Fixed pin auto-wiring when spawning new nodes from the context menu.
Removed unneeded empty root graph from parameter blocks

diff --git a/Engine/Plugins/Experimental/AnimNext/Source/AnimNextEditor/Private/AnimNextEditorModule.cpp b/Engine/Plugins/Experimental/AnimNext/Source/AnimNextEditor/Private/AnimNextEditorModule.cpp
--- a/Engine/Plugins/Experimental/AnimNext/Source/AnimNextEditor/Private/AnimNextEditorModule.cpp
+++ b/Engine/Plugins/Experimental/AnimNext/Source/AnimNextEditor/Private/AnimNextEditorModule.cpp
@@ -138,6 +138,19 @@ class FModule : public IModule
 		{
 			UAnimNextGraph* Graph = CastChecked<UAnimNextGraph>(InAsset);
 			UAnimNextGraph_EditorData* EditorData = UncookedOnly::FUtils::GetEditorData(Graph);
+
+			// Open default model for this graph asset as well as the graph-list tab (deferred so we dont modify opened document records during iteration)
+			ExecuteOnGameThread(UE_SOURCE_LOCATION, [WeakEditor = TWeakPtr<FWorkspaceEditor>(InEditor), WeakEditorData = TWeakObjectPtr<UAnimNextGraph_EditorData>(EditorData)]()
+			{
+				if(WeakEditorData.Get() && WeakEditor.IsValid())
+				{
+					if(UObject* EditorObject = WeakEditorData->GetEditorObjectForRigVMGraph(WeakEditorData->GetRigVMClient()->GetDefaultModel()))
+					{
+						WeakEditor.Pin()->OpenDocument(EditorObject, FDocumentTracker::EOpenDocumentCause::OpenNewDocument);
+					}
+				}
+			});
+
 			return SNew(SAnimNextGraphView, EditorData)
 				.OnOpenGraph_Lambda([WeakEditor = TWeakPtr<FWorkspaceEditor>(InEditor)](URigVMGraph* InGraph)
 				{
diff --git a/Engine/Plugins/Experimental/AnimNext/Source/AnimNextEditor/Private/Common/GraphEditorSchemaActions.cpp b/Engine/Plugins/Experimental/AnimNext/Source/AnimNextEditor/Private/Common/GraphEditorSchemaActions.cpp
new file mode 100644
--- /dev/null
+++ b/Engine/Plugins/Experimental/AnimNext/Source/AnimNextEditor/Private/Common/GraphEditorSchemaActions.cpp
@@ -0,0 +1,119 @@
+﻿// Copyright Epic Games, Inc. All Rights Reserved.
+
+#include "GraphEditorSchemaActions.h"
+#include "EditorUtils.h"
+#include "Graph/AnimNextGraph_EdGraph.h"
+#include "Graph/AnimNextGraph_EdGraphNode.h"
+#include "Graph/AnimNextGraph_EditorData.h"
+#include "RigVMModel/Nodes/RigVMUnitNode.h"
+#include "Settings/ControlRigSettings.h"
+#include "Units/RigUnit.h"
+
+UEdGraphNode* FAnimNextSchemaAction_RigUnit::PerformAction(UEdGraph* ParentGraph, TArray<UEdGraphPin*>& FromPins, const FVector2D Location, bool bSelectNewNode)
+{
+	UAnimNextGraph_EditorData* EditorData = ParentGraph->GetTypedOuter<UAnimNextGraph_EditorData>();
+	UAnimNextGraph_EdGraphNode* NewNode = nullptr;
+	UAnimNextGraph_EdGraph* EdGraph = Cast<UAnimNextGraph_EdGraph>(ParentGraph);
+
+	UEdGraphPin* FromPin = nullptr;
+	if (FromPins.Num() > 0)
+	{
+		FromPin = FromPins[0];
+	}
+	
+	if (EditorData != nullptr && EdGraph != nullptr)
+	{
+		FName Name = UE::AnimNext::Editor::FUtils::ValidateName(EditorData, StructTemplate->GetFName().ToString());
+		URigVMController* Controller = EditorData->GetRigVMClient()->GetController(ParentGraph);
+
+		Controller->OpenUndoBracket(FString::Printf(TEXT("Add '%s' Node"), *Name.ToString()));
+
+		FRigVMUnitNodeCreatedContext& UnitNodeCreatedContext = Controller->GetUnitNodeCreatedContext();
+		FRigVMUnitNodeCreatedContext::FScope ReasonScope(UnitNodeCreatedContext, ERigVMNodeCreatedReason::NodeSpawner);
+
+		if (URigVMUnitNode* ModelNode = Controller->AddUnitNode(StructTemplate, FRigVMStruct::ExecuteName, Location, Name.ToString(), true, false))
+		{
+			NewNode = Cast<UAnimNextGraph_EdGraphNode>(EdGraph->FindNodeForModelNodeName(ModelNode->GetFName()));
+			check(NewNode);
+
+			if (NewNode)
+			{
+				NewNode->AutowireNewNode(FromPin);
+
+				Controller->ClearNodeSelection(true);
+				Controller->SelectNode(ModelNode, true, true);
+			}
+
+#if WITH_EDITORONLY_DATA
+
+			// @TODO: switch over to custom settings here
+			const FControlRigSettingsPerPinBool* ExpansionMapPtr = UControlRigEditorSettings::Get()->RigUnitPinExpansion.Find(ModelNode->GetScriptStruct()->GetName());
+			if (ExpansionMapPtr)
+			{
+				const FControlRigSettingsPerPinBool& ExpansionMap = *ExpansionMapPtr;
+
+				for (const TPair<FString, bool>& Pair : ExpansionMap.Values)
+				{
+					FString PinPath = FString::Printf(TEXT("%s.%s"), *ModelNode->GetName(), *Pair.Key);
+					Controller->SetPinExpansion(PinPath, Pair.Value, true);
+				}
+			}
+		}
+#endif
+		
+		Controller->CloseUndoBracket();
+	}
+
+	return NewNode;
+}
+
+
+UEdGraphNode* FAnimNextSchemaAction_DispatchFactory::PerformAction(UEdGraph* ParentGraph, TArray<UEdGraphPin*>& FromPins, const FVector2D Location, bool bSelectNewNode)
+{
+	UAnimNextGraph_EditorData* EditorData = ParentGraph->GetTypedOuter<UAnimNextGraph_EditorData>();
+	UAnimNextGraph_EdGraphNode* NewNode = nullptr;
+	UAnimNextGraph_EdGraph* EdGraph = Cast<UAnimNextGraph_EdGraph>(ParentGraph);
+
+	UEdGraphPin* FromPin = nullptr;
+	if (FromPins.Num() > 0)
+	{
+		FromPin = FromPins[0];
+	}
+	
+	if (EditorData != nullptr && EdGraph != nullptr)
+	{
+		const FRigVMTemplate* Template = FRigVMRegistry::Get().FindTemplate(Notation);
+		if (Template == nullptr)
+		{
+			return nullptr;
+		}
+
+		const int32 NotationHash = (int32)GetTypeHash(Notation);
+		const FString TemplateName = TEXT("RigVMTemplate_") + FString::FromInt(NotationHash);
+
+		FName Name = UE::AnimNext::Editor::FUtils::ValidateName(EditorData, Template->GetName().ToString());
+		URigVMController* Controller = EditorData->GetRigVMClient()->GetController(ParentGraph);
+
+		Controller->OpenUndoBracket(FString::Printf(TEXT("Add '%s' Node"), *Name.ToString()));
+
+		if (URigVMTemplateNode* ModelNode = Controller->AddTemplateNode(Notation, Location, Name.ToString(), true, false))
+		{
+			NewNode = Cast<UAnimNextGraph_EdGraphNode>(EdGraph->FindNodeForModelNodeName(ModelNode->GetFName()));
+
+			if (NewNode)
+			{
+				NewNode->AutowireNewNode(FromPin);
+
+				Controller->ClearNodeSelection(true);
+				Controller->SelectNode(ModelNode, true, true);
+			}
+
+			Controller->CloseUndoBracket();
+		}
+		else
+		{
+			Controller->CancelUndoBracket();
+		}
+	}
+	return NewNode;
+}
\ No newline at end of file
diff --git a/Engine/Plugins/Experimental/AnimNext/Source/AnimNextEditor/Private/Common/GraphEditorSchemaActions.h b/Engine/Plugins/Experimental/AnimNext/Source/AnimNextEditor/Private/Common/GraphEditorSchemaActions.h
new file mode 100644
--- /dev/null
+++ b/Engine/Plugins/Experimental/AnimNext/Source/AnimNextEditor/Private/Common/GraphEditorSchemaActions.h
@@ -0,0 +1,82 @@
+﻿// Copyright Epic Games, Inc. All Rights Reserved.
+
+#pragma once
+
+#include "CoreMinimal.h"
+#include "EdGraph/EdGraphSchema.h"
+#include "Editor/RigVMEditorStyle.h"
+#include "Styling/AppStyle.h"
+#include "GraphEditorSchemaActions.generated.h"
+
+struct FSlateBrush;
+
+USTRUCT()
+struct FAnimNextSchemaAction : public FEdGraphSchemaAction
+{
+	GENERATED_BODY()
+
+	FAnimNextSchemaAction() = default;
+	
+	FAnimNextSchemaAction(FText InNodeCategory, FText InMenuDesc, FText InToolTip, FText InKeywords = FText::GetEmpty())
+		: FEdGraphSchemaAction(MoveTemp(InNodeCategory), MoveTemp(InMenuDesc), MoveTemp(InToolTip), 0, MoveTemp(InKeywords))
+	{
+	}
+
+	virtual const FSlateBrush* GetIconBrush() const
+	{
+		return FRigVMEditorStyle::Get().GetBrush("RigVM.Unit");
+	}
+
+	virtual const FLinearColor& GetIconColor() const
+	{
+		return FLinearColor::White;
+	}
+};
+
+USTRUCT()
+struct FAnimNextSchemaAction_RigUnit : public FAnimNextSchemaAction
+{
+	GENERATED_BODY()
+
+	FAnimNextSchemaAction_RigUnit() = default;
+	
+	FAnimNextSchemaAction_RigUnit(UScriptStruct* InStructTemplate, FText InNodeCategory, FText InMenuDesc, FText InToolTip, FText InKeywords = FText::GetEmpty())
+		: FAnimNextSchemaAction(InNodeCategory, InMenuDesc, InToolTip, InKeywords)
+		, StructTemplate(InStructTemplate)
+	{}
+
+	// FEdGraphSchemaAction Interface
+	virtual UEdGraphNode* PerformAction(UEdGraph* ParentGraph, TArray<UEdGraphPin*>& FromPins, const FVector2D Location, bool bSelectNewNode = true) override;
+	virtual UEdGraphNode* PerformAction(UEdGraph* ParentGraph, UEdGraphPin* FromPin, const FVector2D Location, bool bSelectNewNode = true) { return nullptr; }
+
+private:
+	// The script struct for our rig unit
+	UScriptStruct* StructTemplate = nullptr;
+};
+
+
+USTRUCT()
+struct FAnimNextSchemaAction_DispatchFactory : public FAnimNextSchemaAction
+{
+	GENERATED_BODY()
+
+	FAnimNextSchemaAction_DispatchFactory() = default;
+
+	FAnimNextSchemaAction_DispatchFactory(FName InNotation, FText InNodeCategory, FText InMenuDesc, FText InToolTip, FText InKeywords = FText::GetEmpty())
+		: FAnimNextSchemaAction(InNodeCategory, InMenuDesc, InToolTip, InKeywords)
+		, Notation(InNotation)
+	{}
+
+	virtual const FSlateBrush* GetIconBrush() const
+	{
+		return FRigVMEditorStyle::Get().GetBrush("RigVM.Template");
+	}
+	
+	// FEdGraphSchemaAction Interface
+	virtual UEdGraphNode* PerformAction(UEdGraph* ParentGraph, TArray<UEdGraphPin*>& FromPins, const FVector2D Location, bool bSelectNewNode = true) override;
+	virtual UEdGraphNode* PerformAction(UEdGraph* ParentGraph, UEdGraphPin* FromPin, const FVector2D Location, bool bSelectNewNode = true) { return nullptr; }
+
+private:
+	// Notation for dispatch factory
+	FName Notation;
+};
diff --git a/Engine/Plugins/Experimental/AnimNext/Source/AnimNextEditor/Private/Common/SActionMenu.cpp b/Engine/Plugins/Experimental/AnimNext/Source/AnimNextEditor/Private/Common/SActionMenu.cpp
new file mode 100644
--- /dev/null
+++ b/Engine/Plugins/Experimental/AnimNext/Source/AnimNextEditor/Private/Common/SActionMenu.cpp
@@ -0,0 +1,234 @@
+﻿// Copyright Epic Games, Inc. All Rights Reserved.
+
+#include "SActionMenu.h"
+
+#include "GraphEditorSchemaActions.h"
+#include "Framework/Application/SlateApplication.h"
+#include "IDocumentation.h"
+#include "SSubobjectEditor.h"
+#include "Graph/AnimNextGraph_EdGraphSchema.h"
+#include "Widgets/Images/SImage.h"
+#include "Widgets/Layout/SBox.h"
+#include "Widgets/SBoxPanel.h"
+#include "Widgets/Text/STextBlock.h"
+#include "SGraphPalette.h"
+#include "RigVMCore/RigVMRegistry.h"
+#include "Units/RigUnit.h"
+#include "Widgets/SToolTip.h"
+#include "Graph/AnimNextExecuteContext.h"
+
+#define LOCTEXT_NAMESPACE "AnimNextEditor"
+
+namespace UE::AnimNext::Editor
+{
+
+void SActionMenu::CollectAllAnimNextGraphActions(FGraphContextMenuBuilder& MenuBuilder) const
+{
+	for(const FRigVMFunction& Function : FRigVMRegistry::Get().GetFunctions())
+	{
+		const UScriptStruct* FunctionContext = Function.GetExecuteContextStruct();
+		if (FunctionContext == nullptr || !AllowedExecuteContexts.Contains(FunctionContext))
+		{
+			continue;
+		}
+
+		UScriptStruct* Struct = Function.Struct;
+		if (Struct == nullptr)
+		{
+			continue;
+		}
+
+		FString CategoryMetadata, DisplayNameMetadata, MenuDescSuffixMetadata;
+		Struct->GetStringMetaDataHierarchical(FRigVMStruct::CategoryMetaName, &CategoryMetadata);
+		Struct->GetStringMetaDataHierarchical(FRigVMStruct::DisplayNameMetaName, &DisplayNameMetadata);
+		Struct->GetStringMetaDataHierarchical(FRigVMStruct::MenuDescSuffixMetaName, &MenuDescSuffixMetadata);
+		if (!MenuDescSuffixMetadata.IsEmpty())
+		{
+			MenuDescSuffixMetadata = TEXT(" ") + MenuDescSuffixMetadata;
+		}
+		const FText NodeCategory = FText::FromString(CategoryMetadata);
+		const FText MenuDesc = FText::FromString(DisplayNameMetadata + MenuDescSuffixMetadata);
+		const FText ToolTip = Struct->GetToolTipText();
+
+		if (MenuDesc.IsEmpty())
+		{
+			continue;
+		}
+
+		MenuBuilder.AddAction(MakeShared<FAnimNextSchemaAction_RigUnit>(Struct, NodeCategory, MenuDesc, ToolTip));
+	}
+
+	for (const FRigVMDispatchFactory* Factory : FRigVMRegistry::Get().GetFactories())
+	{
+		// See if the factory allows our execute contexts
+		bool bAllowed = true;
+		for(const UScriptStruct* AllowedExecuteContext : AllowedExecuteContexts)
+		{
+			if (!Factory->SupportsExecuteContextStruct(AllowedExecuteContext))
+			{
+				bAllowed = false;
+				break;
+			}
+		}
+
+		if(bAllowed)
+		{
+			continue;
+		}
+
+		const FRigVMTemplate* Template = Factory->GetTemplate();
+		if (Template == nullptr)
+		{
+			continue;
+		}
+
+		FText NodeCategory = FText::FromString(Factory->GetCategory());
+		FText MenuDesc = FText::FromString(Factory->GetNodeTitle(FRigVMTemplateTypeMap()));
+		FText ToolTip = Factory->GetNodeTooltip(FRigVMTemplateTypeMap());
+
+		MenuBuilder.AddAction(MakeShared<FAnimNextSchemaAction_DispatchFactory>(Template->GetNotation(), NodeCategory, MenuDesc, ToolTip));
+	};
+}
+
+SActionMenu::~SActionMenu()
+{
+	OnClosedCallback.ExecuteIfBound();
+	OnCloseReasonCallback.ExecuteIfBound(bActionExecuted, false, !DraggedFromPins.IsEmpty());
+}
+
+void SActionMenu::Construct(const FArguments& InArgs)
+{
+	Graph = InArgs._Graph;
+	DraggedFromPins = InArgs._DraggedFromPins;
+	NewNodePosition = InArgs._NewNodePosition;
+	OnClosedCallback = InArgs._OnClosedCallback;
+	bAutoExpandActionMenu = InArgs._AutoExpandActionMenu;
+	OnCloseReasonCallback = InArgs._OnCloseReason;
+	AllowedExecuteContexts = InArgs._AllowedExecuteContexts;
+
+	SBorder::Construct(SBorder::FArguments()
+		.BorderImage(FAppStyle::Get().GetBrush("Menu.Background"))
+		.Padding(5)
+		[
+			SNew(SBox)
+			.WidthOverride(400)
+			.HeightOverride(400)
+			[
+				SNew(SVerticalBox)
+				+SVerticalBox::Slot()
+				[
+					SAssignNew(GraphActionMenu, SGraphActionMenu)
+						.OnActionSelected(this, &SActionMenu::OnActionSelected)
+						.OnCreateWidgetForAction(SGraphActionMenu::FOnCreateWidgetForAction::CreateSP(this, &SActionMenu::OnCreateWidgetForAction))
+						.OnCollectAllActions(this, &SActionMenu::CollectAllActions)
+						.OnCreateCustomRowExpander_Lambda([](const FCustomExpanderData& InActionMenuData)
+						{
+							// Default table row doesnt indent correctly
+							return SNew(SExpanderArrow, InActionMenuData.TableRow);
+						})
+						.DraggedFromPins(DraggedFromPins)
+						.GraphObj(Graph)
+						.AlphaSortItems(true)
+						.bAllowPreselectedItemActivation(true)
+				]
+			]
+		]
+	);
+}
+
+void SActionMenu::CollectAllActions(FGraphActionListBuilderBase& OutAllActions)
+{
+	if (!Graph)
+	{
+		return;
+	}
+
+	FGraphContextMenuBuilder MenuBuilder(Graph);
+	if (!DraggedFromPins.IsEmpty())
+	{
+		MenuBuilder.FromPin = DraggedFromPins[0];
+	}
+
+	// Cannot call GetGraphContextActions() during serialization and GC due to its use of FindObject()
+	if(!GIsSavingPackage && !IsGarbageCollecting())
+	{
+		CollectAllAnimNextGraphActions(MenuBuilder);
+	}
+
+	OutAllActions.Append(MenuBuilder);
+}
+
+TSharedRef<SEditableTextBox> SActionMenu::GetFilterTextBox()
+{
+	return GraphActionMenu->GetFilterTextBox();
+}
+
+TSharedRef<SWidget> SActionMenu::OnCreateWidgetForAction(FCreateWidgetForActionData* const InCreateData)
+{
+	check(InCreateData);
+	InCreateData->bHandleMouseButtonDown = false;
+
+	const FSlateBrush* IconBrush = nullptr;
+	FLinearColor IconColor;
+	TSharedPtr<FAnimNextSchemaAction> Action = StaticCastSharedPtr<FAnimNextSchemaAction>(InCreateData->Action);
+	if (Action.IsValid())
+	{
+		IconBrush = Action->GetIconBrush();
+		IconColor = Action->GetIconColor();
+	}
+
+	TSharedPtr<SHorizontalBox> WidgetBox = SNew(SHorizontalBox);
+	if (IconBrush)
+	{
+		WidgetBox->AddSlot()
+			.AutoWidth()
+			.VAlign(VAlign_Center)
+			.Padding(0, 0, 0, 0)
+			[
+				SNew(SImage)
+				.ColorAndOpacity(IconColor)
+				.Image(IconBrush)
+			];
+	}
+
+	WidgetBox->AddSlot()
+		.AutoWidth()
+		.VAlign(VAlign_Center)
+		.Padding(IconBrush ? 4.0f : 0.0f, 0, 0, 0)
+		[
+			SNew(SGraphPaletteItem, InCreateData)
+		];
+
+	return WidgetBox->AsShared();
+}
+
+void SActionMenu::OnActionSelected(const TArray<TSharedPtr<FEdGraphSchemaAction>>& SelectedAction, ESelectInfo::Type InSelectionType)
+{
+	if (!Graph)
+	{
+		return;
+	}
+
+	if (InSelectionType != ESelectInfo::OnMouseClick  && InSelectionType != ESelectInfo::OnKeyPress && !SelectedAction.IsEmpty())
+	{
+		return;
+	}
+
+	for (const TSharedPtr<FEdGraphSchemaAction>& Action : SelectedAction)
+	{
+		if (Action.IsValid() && Graph)
+		{
+			if (!bActionExecuted && (Action->GetTypeId() != FEdGraphSchemaAction_Dummy::StaticGetTypeId()))
+			{
+				FSlateApplication::Get().DismissAllMenus();
+				bActionExecuted = true;
+			}
+
+			Action->PerformAction(Graph, DraggedFromPins, NewNodePosition);
+		}
+	}
+}
+
+} 
+
+#undef LOCTEXT_NAMESPACE
diff --git a/Engine/Plugins/Experimental/AnimNext/Source/AnimNextEditor/Private/Common/SActionMenu.h b/Engine/Plugins/Experimental/AnimNext/Source/AnimNextEditor/Private/Common/SActionMenu.h
new file mode 100644
--- /dev/null
+++ b/Engine/Plugins/Experimental/AnimNext/Source/AnimNextEditor/Private/Common/SActionMenu.h
@@ -0,0 +1,77 @@
+﻿// Copyright Epic Games, Inc. All Rights Reserved.
+#pragma once
+
+#include "GraphEditor.h"
+#include "EdGraph/EdGraph.h"
+#include "EdGraph/EdGraphPin.h"
+#include "EdGraph/EdGraphSchema.h"
+#include "Layout/Margin.h"
+#include "Layout/Visibility.h"
+#include "SGraphActionMenu.h"
+#include "Widgets/DeclarativeSyntaxSupport.h"
+#include "Widgets/Input/SEditableTextBox.h"
+#include "Widgets/Layout/SBorder.h"
+#include "Widgets/SWidget.h"
+#include "Widgets/Views/SExpanderArrow.h"
+
+class SEditableTextBox;
+class SGraphActionMenu;
+class UEdGraph;
+
+namespace UE::AnimNext::Editor
+{
+
+class SActionMenu : public SBorder
+{
+public:
+	/** Delegate for the OnCloseReason event which is always raised when the SActionMenu closes */
+	DECLARE_DELEGATE_ThreeParams(FClosedReason, bool /*bActionExecuted*/, bool /*bContextSensitiveChecked*/, bool /*bGraphPinContext*/);
+
+	SLATE_BEGIN_ARGS(SActionMenu)
+		: _Graph(static_cast<UEdGraph*>(nullptr))
+		, _NewNodePosition(FVector2D::ZeroVector)
+		, _AutoExpandActionMenu(false)
+	{
+	}
+		SLATE_ARGUMENT(UEdGraph*, Graph)
+		SLATE_ARGUMENT(FVector2D, NewNodePosition)
+		SLATE_ARGUMENT(TArray<UEdGraphPin*>, DraggedFromPins)
+		SLATE_ARGUMENT(SGraphEditor::FActionMenuClosed, OnClosedCallback)
+		SLATE_ARGUMENT(bool, AutoExpandActionMenu)
+		SLATE_EVENT(FClosedReason, OnCloseReason)
+		SLATE_ARGUMENT(TArray<UScriptStruct*>, AllowedExecuteContexts)
+
+	SLATE_END_ARGS()
+
+	void Construct(const FArguments& InArgs);
+
+	~SActionMenu();
+
+	TSharedRef<SEditableTextBox> GetFilterTextBox();
+
+protected:
+	void OnActionSelected(const TArray<TSharedPtr<FEdGraphSchemaAction>>& SelectedAction, ESelectInfo::Type InSelectionType);
+
+	TSharedRef<SWidget> OnCreateWidgetForAction(struct FCreateWidgetForActionData* const InCreateData);
+
+	/** Callback used to populate all actions list in SGraphActionMenu */
+	void CollectAllActions(FGraphActionListBuilderBase& OutAllActions);
+
+	void CollectAllAnimNextGraphActions(FGraphContextMenuBuilder& MenuBuilder) const;
+	
+private:
+	UEdGraph* Graph = nullptr;
+	bool bAutoExpandActionMenu = false;
+	bool bActionExecuted = false;
+
+	TArray<UEdGraphPin*> DraggedFromPins;
+	FVector2D NewNodePosition = FVector2D::ZeroVector;
+
+	SGraphEditor::FActionMenuClosed OnClosedCallback;
+	FClosedReason OnCloseReasonCallback;
+
+	TSharedPtr<SGraphActionMenu> GraphActionMenu;
+	TArray<UScriptStruct*> AllowedExecuteContexts;
+};
+
+}
diff --git a/Engine/Plugins/Experimental/AnimNext/Source/AnimNextEditor/Private/Graph/AnimNextGraphDocumentSummoner.cpp b/Engine/Plugins/Experimental/AnimNext/Source/AnimNextEditor/Private/Graph/AnimNextGraphDocumentSummoner.cpp
--- a/Engine/Plugins/Experimental/AnimNext/Source/AnimNextEditor/Private/Graph/AnimNextGraphDocumentSummoner.cpp
+++ b/Engine/Plugins/Experimental/AnimNext/Source/AnimNextEditor/Private/Graph/AnimNextGraphDocumentSummoner.cpp
@@ -1,7 +1,8 @@
 ﻿// Copyright Epic Games, Inc. All Rights Reserved.
 
 #include "AnimNextGraphDocumentSummoner.h"
-#include "SActionMenu.h"
+#include "Common/SActionMenu.h"
+#include "Graph/AnimNextExecuteContext.h"
 #include "Graph/AnimNextGraph_EdGraph.h"
 #include "Workspace/AnimNextWorkspaceEditor.h"
 
@@ -20,7 +21,8 @@ FActionMenuContent FAnimNextGraphDocumentSummoner::OnCreateGraphActionMenu(UEdGr
 		.Graph(InGraph)
 		.NewNodePosition(InNodePosition)
 		.DraggedFromPins(InDraggedPins)
-		.OnClosedCallback(InOnMenuClosed);
+		.OnClosedCallback(InOnMenuClosed)
+		.AllowedExecuteContexts( { FRigVMExecuteContext::StaticStruct(), FAnimNextExecuteContext::StaticStruct() });
 
 	TSharedPtr<SWidget> FilterTextBox = StaticCastSharedRef<SWidget>(ActionMenu->GetFilterTextBox());
 	return FActionMenuContent(StaticCastSharedRef<SWidget>(ActionMenu), FilterTextBox);
diff --git a/Engine/Plugins/Experimental/AnimNext/Source/AnimNextEditor/Private/Graph/AnimNextGraphEditor.cpp b/Engine/Plugins/Experimental/AnimNext/Source/AnimNextEditor/Private/Graph/AnimNextGraphEditor.cpp
--- a/Engine/Plugins/Experimental/AnimNext/Source/AnimNextEditor/Private/Graph/AnimNextGraphEditor.cpp
+++ b/Engine/Plugins/Experimental/AnimNext/Source/AnimNextEditor/Private/Graph/AnimNextGraphEditor.cpp
@@ -6,7 +6,7 @@
 #include "Graph/AnimNextGraph_EdGraphNode.h"
 #include "Graph/AnimNextGraph_EditorData.h"
 #include "EdGraphNode_Comment.h"
-#include "SActionMenu.h"
+#include "Common/SActionMenu.h"
 #include "UncookedOnlyUtils.h"
 #include "Framework/Commands/GenericCommands.h"
 #include "RigVMModel/RigVMController.h"
@@ -182,7 +182,8 @@ FActionMenuContent FGraphEditor::OnCreateGraphActionMenu(UEdGraph* InGraph, cons
 		.Graph(InGraph)
 		.NewNodePosition(InNodePosition)
 		.DraggedFromPins(InDraggedPins)
-		.OnClosedCallback(InOnMenuClosed);
+		.OnClosedCallback(InOnMenuClosed)
+		.AllowedExecuteContexts( { FRigVMExecuteContext::StaticStruct(), FAnimNextExecuteContext::StaticStruct() });
 
 	TSharedPtr<SWidget> FilterTextBox = StaticCastSharedRef<SWidget>(ActionMenu->GetFilterTextBox());
 	return FActionMenuContent(StaticCastSharedRef<SWidget>(ActionMenu), FilterTextBox);
diff --git a/Engine/Plugins/Experimental/AnimNext/Source/AnimNextEditor/Private/Graph/GraphEditorSchemaActions.cpp b/Engine/Plugins/Experimental/AnimNext/Source/AnimNextEditor/Private/Graph/GraphEditorSchemaActions.cpp
deleted file mode 100644
diff --git a/Engine/Plugins/Experimental/AnimNext/Source/AnimNextEditor/Private/Graph/GraphEditorSchemaActions.h b/Engine/Plugins/Experimental/AnimNext/Source/AnimNextEditor/Private/Graph/GraphEditorSchemaActions.h
deleted file mode 100644
diff --git a/Engine/Plugins/Experimental/AnimNext/Source/AnimNextEditor/Private/Graph/SActionMenu.cpp b/Engine/Plugins/Experimental/AnimNext/Source/AnimNextEditor/Private/Graph/SActionMenu.cpp
deleted file mode 100644
diff --git a/Engine/Plugins/Experimental/AnimNext/Source/AnimNextEditor/Private/Graph/SActionMenu.h b/Engine/Plugins/Experimental/AnimNext/Source/AnimNextEditor/Private/Graph/SActionMenu.h
deleted file mode 100644
diff --git a/Engine/Plugins/Experimental/AnimNext/Source/AnimNextEditor/Private/Param/AnimNextParameterBlockEditor.cpp b/Engine/Plugins/Experimental/AnimNext/Source/AnimNextEditor/Private/Param/AnimNextParameterBlockEditor.cpp
--- a/Engine/Plugins/Experimental/AnimNext/Source/AnimNextEditor/Private/Param/AnimNextParameterBlockEditor.cpp
+++ b/Engine/Plugins/Experimental/AnimNext/Source/AnimNextEditor/Private/Param/AnimNextParameterBlockEditor.cpp
@@ -6,7 +6,7 @@
 #include "Param/AnimNextParameterBlock_EditorData.h"
 #include "EdGraphNode_Comment.h"
 #include "ExternalPackageHelper.h"
-#include "Param/SParametersActionMenu.h"
+#include "Common/SActionMenu.h"
 #include "UncookedOnlyUtils.h"
 #include "Framework/Commands/GenericCommands.h"
 #include "RigVMModel/RigVMController.h"
@@ -15,6 +15,7 @@
 #include "Param/AnimNextParameterBlockEntry.h"
 #include "Param/AnimNextParameterSettings.h"
 #include "ParameterBlockEditorMode.h"
+#include "Param/AnimNextParameterExecuteContext.h"
 #include "Widgets/Docking/SDockTab.h"
 
 #define LOCTEXT_NAMESPACE "AnimNextParameterBlockEditor"
@@ -224,12 +225,13 @@ TSharedRef<SGraphEditor> FParameterBlockEditor::CreateGraphEditorWidget(TSharedR
 
 FActionMenuContent FParameterBlockEditor::OnCreateGraphActionMenu(UEdGraph* InGraph, const FVector2D& InNodePosition, const TArray<UEdGraphPin*>& InDraggedPins, bool bAutoExpand, SGraphEditor::FActionMenuClosed InOnMenuClosed)
 {
-	TSharedRef<SParametersActionMenu> ActionMenu = SNew(SParametersActionMenu)
+	TSharedRef<SActionMenu> ActionMenu = SNew(SActionMenu)
 		.AutoExpandActionMenu(bAutoExpand)
 		.Graph(InGraph)
 		.NewNodePosition(InNodePosition)
 		.DraggedFromPins(InDraggedPins)
-		.OnClosedCallback(InOnMenuClosed);
+		.OnClosedCallback(InOnMenuClosed)
+		.AllowedExecuteContexts( { FRigVMExecuteContext::StaticStruct(), FAnimNextParameterExecuteContext::StaticStruct() });
 
 	TSharedPtr<SWidget> FilterTextBox = StaticCastSharedRef<SWidget>(ActionMenu->GetFilterTextBox());
 	return FActionMenuContent(StaticCastSharedRef<SWidget>(ActionMenu), FilterTextBox);
diff --git a/Engine/Plugins/Experimental/AnimNext/Source/AnimNextEditor/Private/Param/ParameterBlockGraphDocumentSummoner.cpp b/Engine/Plugins/Experimental/AnimNext/Source/AnimNextEditor/Private/Param/ParameterBlockGraphDocumentSummoner.cpp
--- a/Engine/Plugins/Experimental/AnimNext/Source/AnimNextEditor/Private/Param/ParameterBlockGraphDocumentSummoner.cpp
+++ b/Engine/Plugins/Experimental/AnimNext/Source/AnimNextEditor/Private/Param/ParameterBlockGraphDocumentSummoner.cpp
@@ -1,9 +1,10 @@
 ﻿// Copyright Epic Games, Inc. All Rights Reserved.
 
 #include "ParameterBlockGraphDocumentSummoner.h"
-#include "SParametersActionMenu.h"
+#include "Common/SActionMenu.h"
 #include "Graph/AnimNextGraph.h"
 #include "Param/AnimNextParameterBlock_EdGraph.h"
+#include "Param/AnimNextParameterExecuteContext.h"
 #include "Workspace/AnimNextWorkspaceEditor.h"
 
 namespace UE::AnimNext::Editor
@@ -16,12 +17,13 @@ FParameterBlockGraphDocumentSummoner::FParameterBlockGraphDocumentSummoner(FName
 
 FActionMenuContent FParameterBlockGraphDocumentSummoner::OnCreateGraphActionMenu(UEdGraph* InGraph, const FVector2D& InNodePosition, const TArray<UEdGraphPin*>& InDraggedPins, bool bAutoExpand, SGraphEditor::FActionMenuClosed InOnMenuClosed) const
 {
-	TSharedRef<SParametersActionMenu> ActionMenu = SNew(SParametersActionMenu)
+	TSharedRef<SActionMenu> ActionMenu = SNew(SActionMenu)
 		.AutoExpandActionMenu(bAutoExpand)
 		.Graph(InGraph)
 		.NewNodePosition(InNodePosition)
 		.DraggedFromPins(InDraggedPins)
-		.OnClosedCallback(InOnMenuClosed);
+		.OnClosedCallback(InOnMenuClosed)
+		.AllowedExecuteContexts( { FRigVMExecuteContext::StaticStruct(), FAnimNextParameterExecuteContext::StaticStruct() });
 
 	TSharedPtr<SWidget> FilterTextBox = StaticCastSharedRef<SWidget>(ActionMenu->GetFilterTextBox());
 	return FActionMenuContent(StaticCastSharedRef<SWidget>(ActionMenu), FilterTextBox);
diff --git a/Engine/Plugins/Experimental/AnimNext/Source/AnimNextEditor/Private/Param/ParametersEditorSchemaActions.cpp b/Engine/Plugins/Experimental/AnimNext/Source/AnimNextEditor/Private/Param/ParametersEditorSchemaActions.cpp
deleted file mode 100644
diff --git a/Engine/Plugins/Experimental/AnimNext/Source/AnimNextEditor/Private/Param/ParametersEditorSchemaActions.h b/Engine/Plugins/Experimental/AnimNext/Source/AnimNextEditor/Private/Param/ParametersEditorSchemaActions.h
deleted file mode 100644
diff --git a/Engine/Plugins/Experimental/AnimNext/Source/AnimNextEditor/Private/Param/SParametersActionMenu.cpp b/Engine/Plugins/Experimental/AnimNext/Source/AnimNextEditor/Private/Param/SParametersActionMenu.cpp
deleted file mode 100644
diff --git a/Engine/Plugins/Experimental/AnimNext/Source/AnimNextEditor/Private/Param/SParametersActionMenu.h b/Engine/Plugins/Experimental/AnimNext/Source/AnimNextEditor/Private/Param/SParametersActionMenu.h
deleted file mode 100644
diff --git a/Engine/Plugins/Experimental/AnimNext/Source/AnimNextEditor/Private/Workspace/AnimNextWorkspaceEditor.cpp b/Engine/Plugins/Experimental/AnimNext/Source/AnimNextEditor/Private/Workspace/AnimNextWorkspaceEditor.cpp
--- a/Engine/Plugins/Experimental/AnimNext/Source/AnimNextEditor/Private/Workspace/AnimNextWorkspaceEditor.cpp
+++ b/Engine/Plugins/Experimental/AnimNext/Source/AnimNextEditor/Private/Workspace/AnimNextWorkspaceEditor.cpp
@@ -149,7 +149,9 @@ void FWorkspaceEditor::OpenWorkspaceForAsset(UObject* InAsset, EOpenWorkspaceMet
 		}
 	}
 
-	auto HandleNewWorkspace = [InAsset]()
+	FWorkspaceEditor* WorkspaceEditor = nullptr;
+
+	auto HandleNewWorkspace = [InAsset, &WorkspaceEditor]()
 	{
 		UAnimNextWorkspaceFactory* Factory = NewObject<UAnimNextWorkspaceFactory>();
 		UPackage* Package = CreatePackage(nullptr);
@@ -159,16 +161,20 @@ void FWorkspaceEditor::OpenWorkspaceForAsset(UObject* InAsset, EOpenWorkspaceMet
 		NewWorkspace->MarkPackageDirty();
 		TSharedRef<FWorkspaceEditor> Editor = MakeShared<FWorkspaceEditor>();
 		Editor->InitEditor(EToolkitMode::Standalone, nullptr, NewWorkspace);
+
+		WorkspaceEditor = &Editor.Get();
 	};
 
-	auto HandleExistingWorkspace = [](const FAssetData& InAssetData)
+	auto HandleExistingWorkspace = [InAsset, &WorkspaceEditor](const FAssetData& InAssetData)
 	{
 		if(UAnimNextWorkspace* ExistingWorkspace = Cast<UAnimNextWorkspace>(InAssetData.GetAsset()))
 		{
 			GEditor->GetEditorSubsystem<UAssetEditorSubsystem>()->OpenEditorForAsset(ExistingWorkspace);
+
+			WorkspaceEditor = static_cast<FWorkspaceEditor*>(GEditor->GetEditorSubsystem<UAssetEditorSubsystem>()->FindEditorForAsset(ExistingWorkspace, true));
 		}
 	};
-	
+
 	if(InOpenMethod == EOpenWorkspaceMethod::AlwaysOpenNewWorkspace || RelevantWorkspaceAssets.Num() == 0)
 	{
 		// No relevant workspaces, so open a new one and add the asset
@@ -189,6 +195,11 @@ void FWorkspaceEditor::OpenWorkspaceForAsset(UObject* InAsset, EOpenWorkspaceMet
 
 		WorkspacePicker->ShowModal();
 	}
+	
+	if(WorkspaceEditor)
+	{
+		WorkspaceEditor->OpenDocument(InAsset, FDocumentTracker::OpenNewDocument);
+	}
 }
 
 void FWorkspaceEditor::RestoreEditedObjectState()
