AnimNext: Fix adapter config

This allows correct access to test components

diff --git a/Engine/Plugins/Experimental/AnimNext/Config/BaseAnimNext.ini b/Engine/Plugins/Experimental/AnimNext/Config/BaseAnimNext.ini
--- a/Engine/Plugins/Experimental/AnimNext/Config/BaseAnimNext.ini
+++ b/Engine/Plugins/Experimental/AnimNext/Config/BaseAnimNext.ini
@@ -1,4 +1,4 @@
 [/Script/AnimNext.AnimNextConfig]
-+ExposedClasses=(Class="/Script/CoreUObject.Class'/Script/Engine.CharacterMovementComponent'",RootParameter="UE_CharacterMovementComponent",FilterType=AllowOnlyFunctions,FilteredFields=,bRegisterTickFunction=True)
-+ExposedClasses=(Class="/Script/CoreUObject.Class'/Script/AnimNext.AnimNextMeshComponent'",RootParameter="UE_AnimNextMeshComponent",FilterType=AllowOnlyFunctions,FilteredFields=,bRegisterTickFunction=True)
++ExposedClasses=(Class="/Script/CoreUObject.Class'/Script/Engine.CharacterMovementComponent'",RootParameter="UE_CharacterMovementComponent",FilterType=AllowAllPropertiesAndFunctions,FilteredFields=,bRegisterTickFunction=True)
++ExposedClasses=(Class="/Script/CoreUObject.Class'/Script/AnimNext.AnimNextMeshComponent'",RootParameter="UE_AnimNextMeshComponent",FilterType=AllowAllPropertiesAndFunctions,FilteredFields=,bRegisterTickFunction=True)
 
