Fix AnimNext live editing by recreating live instances during re-compilation
Add GC references for AnimNext graphs in schedules

diff --git a/Engine/Plugins/Experimental/AnimNext/Source/AnimNext/Internal/Graph/AnimNextGraph.h b/Engine/Plugins/Experimental/AnimNext/Source/AnimNext/Internal/Graph/AnimNextGraph.h
--- a/Engine/Plugins/Experimental/AnimNext/Source/AnimNext/Internal/Graph/AnimNextGraph.h
+++ b/Engine/Plugins/Experimental/AnimNext/Source/AnimNext/Internal/Graph/AnimNextGraph.h
@@ -52,6 +52,19 @@ struct ANIMNEXT_API FAnimNextGraphInstance
 	// Creates an empty graph instance that doesn't reference anything
 	FAnimNextGraphInstance() = default;
 
+#if WITH_EDITORONLY_DATA
+	// In editor, we need custom copy/move semantics to update the live instance tracking
+	FAnimNextGraphInstance(const FAnimNextGraphInstance& Other);
+	FAnimNextGraphInstance(FAnimNextGraphInstance&& Other);
+	FAnimNextGraphInstance& operator=(const FAnimNextGraphInstance& Other);
+	FAnimNextGraphInstance& operator=(FAnimNextGraphInstance&& Other);
+#else
+	FAnimNextGraphInstance(const FAnimNextGraphInstance&) = default;
+	FAnimNextGraphInstance& operator=(const FAnimNextGraphInstance&) = default;
+	FAnimNextGraphInstance(FAnimNextGraphInstance&&) = default;
+	FAnimNextGraphInstance& operator=(FAnimNextGraphInstance&&) = default;
+#endif
+
 	// If the graph instance is allocated, we release it during destruction
 	~FAnimNextGraphInstance();
 
@@ -114,7 +127,12 @@ protected:
 
 	// Loads the graph data from the provided archive buffer and returns true on success, false otherwise
 	bool LoadFromArchiveBuffer(const TArray<uint8>& SharedDataArchiveBuffer);
-	
+
+#if WITH_EDITORONLY_DATA
+	// Releases all live graph instances and returns a list of the instances that were released
+	TArray<FAnimNextGraphInstance*> ReleaseAllInstances();
+#endif
+
 	friend class UAnimNextGraphFactory;
 	friend class UAnimNextGraph_EditorData;
 	friend struct UE::AnimNext::UncookedOnly::FUtils;
@@ -122,6 +140,13 @@ protected:
 	friend struct FAnimNextGraphInstance;
 	friend class UAnimGraphNode_AnimNextGraph;
 
+#if WITH_EDITORONLY_DATA
+	mutable FRWLock GraphInstancesLock;
+
+	// This is a list of live graph instances that have been allocated, used in the editor to reset instances when we re-compile/live edit
+	mutable TSet<FAnimNextGraphInstance*> GraphInstances;
+#endif
+
 	// This is the execute method definition used by this graph
 	UPROPERTY()
 	FAnimNextGraphEvaluatorExecuteDefinition ExecuteDefinition;
diff --git a/Engine/Plugins/Experimental/AnimNext/Source/AnimNext/Private/Decorators/ReferencePoseDecorator.cpp b/Engine/Plugins/Experimental/AnimNext/Source/AnimNext/Private/Decorators/ReferencePoseDecorator.cpp
new file mode 100644
--- /dev/null
+++ b/Engine/Plugins/Experimental/AnimNext/Source/AnimNext/Private/Decorators/ReferencePoseDecorator.cpp
@@ -0,0 +1,26 @@
+// Copyright Epic Games, Inc. All Rights Reserved.
+
+#include "Decorators/ReferencePoseDecorator.h"
+
+#include "DecoratorBase/ExecutionContext.h"
+#include "EvaluationVM/Tasks/PushReferenceKeyframe.h"
+
+namespace UE::AnimNext
+{
+	AUTO_REGISTER_ANIM_DECORATOR(FReferencePoseDecorator)
+
+	DEFINE_ANIM_DECORATOR_BEGIN(FReferencePoseDecorator)
+		DEFINE_ANIM_DECORATOR_IMPLEMENTS_INTERFACE(IEvaluate)
+	DEFINE_ANIM_DECORATOR_END(FReferencePoseDecorator)
+
+	void FReferencePoseDecorator::PreEvaluate(FExecutionContext& Context, const TDecoratorBinding<IEvaluate>& Binding) const
+	{
+		const FSharedData* SharedData = Binding.GetSharedData<FSharedData>();
+
+		FAnimNextPushReferenceKeyframeTask Task;
+		Task.bIsAdditive = SharedData->ReferencePoseType == EAnimNextReferencePoseType::AdditiveIdentity;
+
+		FEvaluateTraversalContext& TraversalContext = Context.GetTraversalContext<FEvaluateTraversalContext>();
+		TraversalContext.AppendTask(Task);
+	}
+}
diff --git a/Engine/Plugins/Experimental/AnimNext/Source/AnimNext/Private/Decorators/ReferencePoseDecorator.h b/Engine/Plugins/Experimental/AnimNext/Source/AnimNext/Private/Decorators/ReferencePoseDecorator.h
new file mode 100644
--- /dev/null
+++ b/Engine/Plugins/Experimental/AnimNext/Source/AnimNext/Private/Decorators/ReferencePoseDecorator.h
@@ -0,0 +1,50 @@
+// Copyright Epic Games, Inc. All Rights Reserved.
+
+#pragma once
+
+#include "CoreMinimal.h"
+#include "DecoratorBase/Decorator.h"
+#include "DecoratorInterfaces/IEvaluate.h"
+
+#include "ReferencePoseDecorator.generated.h"
+
+// TODO: Ideally, the reference pose we output should be a tag and the task that consumes the reference pose should
+// be the one to determine whether it should be in local space or the additive identity (or something else).
+// This way, we have a single option in the graph, and the system can figure out what to use removing the room the user mistakes.
+// It will also be more efficient since we don't need to manipulate the reference pose until the point of use where we might
+// be able to avoid the copy.
+
+UENUM()
+enum class EAnimNextReferencePoseType : int32
+{
+	MeshLocalSpace,
+	AdditiveIdentity,
+};
+
+USTRUCT()
+struct FAnimNextReferencePoseDecoratorSharedData : public FAnimNextDecoratorSharedData
+{
+	GENERATED_BODY()
+
+	/** The type of the reference pose. */
+	UPROPERTY(meta = (Inline))
+	EAnimNextReferencePoseType ReferencePoseType = EAnimNextReferencePoseType::MeshLocalSpace;
+};
+
+namespace UE::AnimNext
+{
+	/**
+	 * FReferencePoseDecorator
+	 * 
+	 * A decorator that outputs a reference pose.
+	 */
+	struct FReferencePoseDecorator : FBaseDecorator, IEvaluate
+	{
+		DECLARE_ANIM_DECORATOR(FReferencePoseDecorator, 0xc03d6afc, FBaseDecorator)
+
+		using FSharedData = FAnimNextReferencePoseDecoratorSharedData;
+
+		// IEvaluate impl
+		virtual void PreEvaluate(FExecutionContext& Context, const TDecoratorBinding<IEvaluate>& Binding) const override;
+	};
+}
diff --git a/Engine/Plugins/Experimental/AnimNext/Source/AnimNext/Private/Graph/AnimNextGraph.cpp b/Engine/Plugins/Experimental/AnimNext/Source/AnimNext/Private/Graph/AnimNextGraph.cpp
--- a/Engine/Plugins/Experimental/AnimNext/Source/AnimNext/Private/Graph/AnimNextGraph.cpp
+++ b/Engine/Plugins/Experimental/AnimNext/Source/AnimNext/Private/Graph/AnimNextGraph.cpp
@@ -23,6 +23,110 @@ namespace UE::AnimNext::Graph
 const UE::AnimNext::FParamId OutputPoseId("UE_Internal_Graph_OutputPose");
 }
 
+#if WITH_EDITORONLY_DATA
+FAnimNextGraphInstance::FAnimNextGraphInstance(const FAnimNextGraphInstance& Other)
+{
+	*this = Other;
+}
+
+FAnimNextGraphInstance::FAnimNextGraphInstance(FAnimNextGraphInstance&& Other)
+{
+	*this = MoveTemp(Other);
+}
+
+FAnimNextGraphInstance& FAnimNextGraphInstance::operator=(const FAnimNextGraphInstance& Other)
+{
+	if (GraphInstancePtr.IsValid())
+	{
+		if (!Other.GraphInstancePtr.IsValid())
+		{
+			// We were valid but we won't be anymore, unregister ourself
+
+			FRWScopeLock Lock(Graph->GraphInstancesLock, SLT_Write);
+
+			check(Graph->GraphInstances.Contains(this));
+			Graph->GraphInstances.Remove(this);
+		}
+		else
+		{
+			// Both instances remain valid, nothing to do
+		}
+	}
+	else
+	{
+		if (Other.GraphInstancePtr.IsValid())
+		{
+			// We were invalid but we will become valid, register ourself
+
+			FRWScopeLock Lock(Graph->GraphInstancesLock, SLT_Write);
+
+			check(!Graph->GraphInstances.Contains(this));
+			Graph->GraphInstances.Add(this);
+		}
+		else
+		{
+			// Both instances were invalid, nothing to do
+		}
+	}
+
+	Graph = Other.Graph;
+	GraphInstancePtr =  Other.GraphInstancePtr;
+	ExtendedExecuteContext = Other.ExtendedExecuteContext;
+
+	return *this;
+}
+
+FAnimNextGraphInstance& FAnimNextGraphInstance::operator=(FAnimNextGraphInstance&& Other)
+{
+	if (GraphInstancePtr.IsValid())
+	{
+		if (!Other.GraphInstancePtr.IsValid())
+		{
+			// We were valid but we won't be anymore, unregister ourself
+			// Other will become valid, register it
+
+			FRWScopeLock Lock(Graph->GraphInstancesLock, SLT_Write);
+
+			check(Graph->GraphInstances.Contains(this));
+			Graph->GraphInstances.Remove(this);
+
+			check(!Graph->GraphInstances.Contains(&Other));
+			Graph->GraphInstances.Add(&Other);
+		}
+		else
+		{
+			// Both instances remain valid, nothing to do
+		}
+	}
+	else
+	{
+		if (Other.GraphInstancePtr.IsValid())
+		{
+			// We were invalid but we will become valid, register ourself
+			// Other will become invalid, unregister it
+
+			FRWScopeLock Lock(Graph->GraphInstancesLock, SLT_Write);
+
+			check(!Graph->GraphInstances.Contains(this));
+			Graph->GraphInstances.Add(this);
+
+			check(Graph->GraphInstances.Contains(&Other));
+			Graph->GraphInstances.Remove(&Other);
+		}
+		else
+		{
+			// Both instances were invalid, nothing to do
+		}
+	}
+
+	Swap(Graph, Other.Graph);
+	Swap(GraphInstancePtr, Other.GraphInstancePtr);
+	Swap(ExtendedExecuteContext, Other.ExtendedExecuteContext);
+
+	return *this;
+}
+#endif
+
 FAnimNextGraphInstance::~FAnimNextGraphInstance()
 {
 	Release();
@@ -39,8 +143,15 @@ void FAnimNextGraphInstance::Release()
 	UE::AnimNext::FExecutionContext Context(Graph->SharedDataBuffer);
 
 	GraphInstancePtr.Reset();
-	Graph = nullptr;
 	ExtendedExecuteContext.Reset();
+
+#if WITH_EDITORONLY_DATA
+	FRWScopeLock Lock(Graph->GraphInstancesLock, SLT_Write);
+	check(Graph->GraphInstances.Contains(this));
+	Graph->GraphInstances.Remove(this);
+#endif
+
+	Graph = nullptr;
 }
 
 bool FAnimNextGraphInstance::IsValid() const
@@ -73,6 +184,12 @@ void UAnimNextGraph::AllocateInstance(FAnimNextGraphInstance& Instance) const
 	Instance.ExtendedExecuteContext.CopyMemoryStorage(ExtendedExecuteContext);
 
 	VM->InitializeInstance(Instance.ExtendedExecuteContext);
+
+#if WITH_EDITORONLY_DATA
+	FRWScopeLock Lock(GraphInstancesLock, SLT_Write);
+	check(!GraphInstances.Contains(&Instance));
+	GraphInstances.Add(&Instance);
+#endif
 }
 
 void UAnimNextGraph::Run(const UE::AnimNext::FContext& Context, FAnimNextGraphInstance& GraphInstance, EAnimNextGraphSimulationSteps SimulationSteps) const
@@ -206,3 +323,25 @@ bool UAnimNextGraph::LoadFromArchiveBuffer(const TArray<uint8>& InSharedDataArch
 		return false;
 	}
 }
+
+#if WITH_EDITORONLY_DATA
+TArray<FAnimNextGraphInstance*> UAnimNextGraph::ReleaseAllInstances()
+{
+	// Make a copy of our live instances since we'll release them all
+	TArray<FAnimNextGraphInstance*> TempGraphInstances;
+	{
+		// Lock shouldn't be necessary here since this should be called before graph compilation
+		// If a graph is executing and we attempt to release it, things are likely to crash
+		// Compilation must happen at a point in the frame where animation isn't executing
+		FRWScopeLock Lock(GraphInstancesLock, SLT_Write);
+		TempGraphInstances = GraphInstances.Array();
+	}
+
+	for (FAnimNextGraphInstance* GraphInstance : TempGraphInstances)
+	{
+		GraphInstance->Release();
+	}
+
+	return TempGraphInstances;
+}
+#endif
diff --git a/Engine/Plugins/Experimental/AnimNext/Source/AnimNext/Private/Scheduler/ScheduleInstanceData.cpp b/Engine/Plugins/Experimental/AnimNext/Source/AnimNext/Private/Scheduler/ScheduleInstanceData.cpp
--- a/Engine/Plugins/Experimental/AnimNext/Source/AnimNext/Private/Scheduler/ScheduleInstanceData.cpp
+++ b/Engine/Plugins/Experimental/AnimNext/Source/AnimNext/Private/Scheduler/ScheduleInstanceData.cpp
@@ -102,6 +102,11 @@ void FScheduleInstanceData::AddReferencedObjects(FReferenceCollector& Collector)
 	{
 		Collector.AddPropertyReferencesWithStructARO(FAnimNextParameterCollection::StaticStruct(), &ParamPair.Value);
 	}
+
+	for (FAnimNextGraphInstance& GraphInstance : GraphInstanceData)
+	{
+		Collector.AddPropertyReferencesWithStructARO(FAnimNextGraphInstance::StaticStruct(), &GraphInstance);
+	}
 }
 
 FString FScheduleInstanceData::GetReferencerName() const
diff --git a/Engine/Plugins/Experimental/AnimNext/Source/AnimNextUncookedOnly/Private/UncookedOnlyUtils.cpp b/Engine/Plugins/Experimental/AnimNext/Source/AnimNextUncookedOnly/Private/UncookedOnlyUtils.cpp
--- a/Engine/Plugins/Experimental/AnimNext/Source/AnimNextUncookedOnly/Private/UncookedOnlyUtils.cpp
+++ b/Engine/Plugins/Experimental/AnimNext/Source/AnimNextUncookedOnly/Private/UncookedOnlyUtils.cpp
@@ -230,16 +230,17 @@ namespace Private
 		return nullptr;
 	}
 
-	TArray<FDecoratorStackMapping> CollectDecoratorStacks(const URigVMGraph* VMGraph)
+	TArray<FDecoratorStackMapping> CollectDecoratorStacks(const URigVMGraph* VMGraph, UAnimNextGraph_Controller* VMController)
 	{
 		const TArray<URigVMNode*>& VMNodes = VMGraph->GetNodes();
 		const URigVMUnitNode* VMRootNode = FindRootNode(VMNodes);
 
 		TArray<FDecoratorStackMapping> DecoratorStackNodes;
 
-		if (VMRootNode == nullptr)
+		if (!ensure(VMRootNode != nullptr))
 		{
-			return DecoratorStackNodes;
+			// Root node wasn't found, add it, we'll need it to compile
+			VMRootNode = VMController->AddUnitNode(FRigUnit_AnimNextGraphRoot::StaticStruct(), FRigUnit_AnimNextGraphRoot::EventName, FVector2D(0.0f, 0.0f), FString(), false);
 		}
 
 		TArray<const URigVMNode*> NodesToVisit;
@@ -270,6 +271,41 @@ namespace Private
 			NodesToVisit.Append(SourceNodes);
 		}
 
+		if (DecoratorStackNodes.IsEmpty())
+		{
+			// If the graph is empty, add a dummy node that just pushes a reference pose
+			URigVMUnitNode* VMNode = VMController->AddUnitNode(FRigUnit_AnimNextDecoratorStack::StaticStruct(), FRigVMStruct::ExecuteName, FVector2D(0.0f, 0.0f), FString(), false);
+
+			const UScriptStruct* CppDecoratorStruct = FRigDecorator_AnimNextCppDecorator::StaticStruct();
+
+			FString DefaultValue;
+			{
+				const UE::AnimNext::FDecoratorUID ReferencePoseDecoratorUID(0xc03d6afc);	// Decorator header is private, reference by UID directly
+				const FDecorator* Decorator = FDecoratorRegistry::Get().Find(ReferencePoseDecoratorUID);
+				check(Decorator != nullptr);
+
+				const FRigDecorator_AnimNextCppDecorator DefaultCppDecoratorStructInstance;
+				FRigDecorator_AnimNextCppDecorator CppDecoratorStructInstance;
+				CppDecoratorStructInstance.DecoratorSharedDataStruct = Decorator->GetDecoratorSharedDataStruct();
+
+				const FProperty* Prop = FAnimNextCppDecoratorWrapper::StaticStruct()->FindPropertyByName(GET_MEMBER_NAME_STRING_CHECKED(FAnimNextCppDecoratorWrapper, CppDecorator));
+				check(Prop != nullptr);
+
+				Prop->ExportText_Direct(DefaultValue, &CppDecoratorStructInstance, &DefaultCppDecoratorStructInstance, nullptr, PPF_None);
+			}
+
+			VMController->AddDecorator(VMNode->GetFName(), *CppDecoratorStruct->GetPathName(), TEXT("ReferencePose"), DefaultValue, INDEX_NONE, false, false);
+
+			FDecoratorStackMapping Mapping(VMNode);
+			ForEachDecoratorInStack(VMNode,
+				[&Mapping](const URigVMNode* DecoratorStackNode, const URigVMPin* DecoratorPin, const FDecorator* Decorator)
+				{
+					Mapping.DecoratorEntries.Add(FDecoratorEntryMapping(DecoratorStackNode, DecoratorPin, Decorator));
+				});
+
+			DecoratorStackNodes.Add(MoveTemp(Mapping));
+		}
+
 		return DecoratorStackNodes;
 	}
 
@@ -312,23 +348,6 @@ namespace Private
 		return LatentPins;
 	}
 
-	TArray<FRigVMFunctionArgument> GetGraphEvaluatorFunctionArguments(const FRigVMPinInfoArray& LatentPins)
-	{
-		const FRigVMRegistry& Registry = FRigVMRegistry::Get();
-
-		TArray<FRigVMFunctionArgument> Arguments;
-		Arguments.Reserve(LatentPins.Num());
-
-		for (const FRigVMPinInfo& Pin : LatentPins)
-		{
-			const FRigVMTemplateArgumentType& TypeArg = Registry.GetType(Pin.TypeIndex);
-
-			Arguments.Add(FRigVMFunctionArgument(Pin.Name.ToString(), TypeArg.CPPType.ToString(), ERigVMFunctionArgumentDirection::Input));
-		}
-
-		return Arguments;
-	}
-
 	FAnimNextGraphEvaluatorExecuteDefinition GetGraphEvaluatorExecuteMethod(const FRigVMPinInfoArray& LatentPins)
 	{
 		const uint32 LatentPinListHash = GetTypeHash(LatentPins);
@@ -365,21 +384,25 @@ namespace Private
 void FUtils::Compile(UAnimNextGraph* InGraph)
 {
 	check(InGraph);
-	
+
 	UAnimNextGraph_EditorData* EditorData = GetEditorData(InGraph);
-	
-	if(EditorData->bIsCompiling)
+
+	if (EditorData->bIsCompiling)
 	{
 		return;
 	}
-	
+
 	TGuardValue<bool> CompilingGuard(EditorData->bIsCompiling, true);
-	
+
+	// Before we re-compile a graph, we need to release and live instances since we need the metadata we are about to replace
+	// to call decorator destructors etc
+	const TArray<FAnimNextGraphInstance*> PreviousLiveGraphInstances = InGraph->ReleaseAllInstances();
+
 	EditorData->bErrorsDuringCompilation = false;
 
 	EditorData->RigGraphDisplaySettings.MinMicroSeconds = EditorData->RigGraphDisplaySettings.LastMinMicroSeconds = DBL_MAX;
 	EditorData->RigGraphDisplaySettings.MaxMicroSeconds = EditorData->RigGraphDisplaySettings.LastMaxMicroSeconds = (double)INDEX_NONE;
-	
+
 	TGuardValue<bool> ReentrantGuardSelf(EditorData->bSuspendModelNotificationsForSelf, true);
 	TGuardValue<bool> ReentrantGuardOthers(EditorData->bSuspendModelNotificationsForOthers, true);
 
@@ -398,7 +421,8 @@ void FUtils::Compile(UAnimNextGraph* InGraph)
 	UAnimNextGraph_Controller* TempController = CastChecked<UAnimNextGraph_Controller>(VMClient->GetOrCreateController(VMTempGraph));
 
 	// Gather our decorator stacks
-	TArray<Private::FDecoratorStackMapping> DecoratorStackNodes = Private::CollectDecoratorStacks(VMTempGraph);
+	TArray<Private::FDecoratorStackMapping> DecoratorStackNodes = Private::CollectDecoratorStacks(VMTempGraph, TempController);
+	check(!DecoratorStackNodes.IsEmpty());
 
 	// Add our runtime shim root node
 	URigVMUnitNode* TempShimRootNode = TempController->AddUnitNode(FRigUnit_AnimNextShimRoot::StaticStruct(), FRigVMStruct::ExecuteName, FVector2D::ZeroVector, FString(), false);
@@ -498,6 +522,12 @@ void FUtils::Compile(UAnimNextGraph* InGraph)
 
 	VMClient->RemoveController(VMTempGraph);
 
+	// Now that the graph has been re-compiled, re-allocate the previous live instances
+	for (FAnimNextGraphInstance* GraphInstance : PreviousLiveGraphInstances)
+	{
+		InGraph->AllocateInstance(*GraphInstance);
+	}
+
 #if WITH_EDITOR
 //	RefreshBreakpoints(EditorData);
 #endif
