AnimNext: Fix crash placing new nodes in graphs

Skip autowire if pin is null

diff --git a/Engine/Plugins/Experimental/AnimNext/Source/AnimNextEditor/Private/Common/GraphEditorSchemaActions.cpp b/Engine/Plugins/Experimental/AnimNext/Source/AnimNextEditor/Private/Common/GraphEditorSchemaActions.cpp
--- a/Engine/Plugins/Experimental/AnimNext/Source/AnimNextEditor/Private/Common/GraphEditorSchemaActions.cpp
+++ b/Engine/Plugins/Experimental/AnimNext/Source/AnimNextEditor/Private/Common/GraphEditorSchemaActions.cpp
@@ -102,7 +102,10 @@ UEdGraphNode* FAnimNextSchemaAction_DispatchFactory::PerformAction(UEdGraph* Par
 
 			if (NewNode)
 			{
-				NewNode->AutowireNewNode(FromPin);
+				if(FromPin)
+				{
+					NewNode->AutowireNewNode(FromPin);
+				}
 
 				Controller->ClearNodeSelection(true);
 				Controller->SelectNode(ModelNode, true, true);
