Fix animnext unit tests following recent refactoring

diff --git a/Engine/Plugins/Experimental/AnimNext/Source/AnimNextEditor/Internal/Graph/GraphFactory.h b/Engine/Plugins/Experimental/AnimNext/Source/AnimNextEditor/Internal/Graph/GraphFactory.h
new file mode 100644
--- /dev/null
+++ b/Engine/Plugins/Experimental/AnimNext/Source/AnimNextEditor/Internal/Graph/GraphFactory.h
@@ -0,0 +1,22 @@
+﻿// Copyright Epic Games, Inc. All Rights Reserved.
+
+#pragma once
+
+#include "CoreMinimal.h"
+#include "Factories/Factory.h"
+#include "GraphFactory.generated.h"
+
+UCLASS(MinimalAPI)
+class UAnimNextGraphFactory : public UFactory
+{
+	GENERATED_BODY()
+
+	UAnimNextGraphFactory();
+	
+	virtual bool ConfigureProperties() override;
+	virtual UObject* FactoryCreateNew(UClass* Class, UObject* InParent, FName Name, EObjectFlags Flags, UObject* Context, FFeedbackContext* Warn, FName CallingContext) override;
+	virtual UObject* FactoryCreateNew(UClass* Class, UObject* InParent, FName Name, EObjectFlags Flags, UObject* Context, FFeedbackContext* Warn) override
+	{
+		return FactoryCreateNew(Class, InParent, Name, Flags, Context, Warn, NAME_None);
+	}
+};
\ No newline at end of file
diff --git a/Engine/Plugins/Experimental/AnimNext/Source/AnimNextEditor/Private/Graph/GraphFactory.cpp b/Engine/Plugins/Experimental/AnimNext/Source/AnimNextEditor/Private/Graph/GraphFactory.cpp
--- a/Engine/Plugins/Experimental/AnimNext/Source/AnimNextEditor/Private/Graph/GraphFactory.cpp
+++ b/Engine/Plugins/Experimental/AnimNext/Source/AnimNextEditor/Private/Graph/GraphFactory.cpp
@@ -1,6 +1,6 @@
 ﻿// Copyright Epic Games, Inc. All Rights Reserved.
 
-#include "GraphFactory.h"
+#include "Graph/GraphFactory.h"
 #include "Graph/AnimNextGraph.h"
 #include "Graph/AnimNextGraph_EditorData.h"
 #include "Graph/AnimNextGraph_EdGraph.h"
diff --git a/Engine/Plugins/Experimental/AnimNext/Source/AnimNextEditor/Private/Graph/GraphFactory.h b/Engine/Plugins/Experimental/AnimNext/Source/AnimNextEditor/Private/Graph/GraphFactory.h
deleted file mode 100644
diff --git a/Engine/Plugins/Experimental/AnimNext/Source/AnimNextTestSuite/Private/AnimNextDecoratorGraphTest.cpp b/Engine/Plugins/Experimental/AnimNext/Source/AnimNextTestSuite/Private/AnimNextDecoratorGraphTest.cpp
--- a/Engine/Plugins/Experimental/AnimNext/Source/AnimNextTestSuite/Private/AnimNextDecoratorGraphTest.cpp
+++ b/Engine/Plugins/Experimental/AnimNext/Source/AnimNextTestSuite/Private/AnimNextDecoratorGraphTest.cpp
@@ -12,6 +12,7 @@
 #include "Graph/AnimNextExecuteContext.h"
 #include "Graph/AnimNextGraph.h"
 #include "Graph/AnimNextGraph_EditorData.h"
+#include "Graph/GraphFactory.h"
 #include "Graph/RigDecorator_AnimNextCppDecorator.h"
 #include "Graph/RigUnit_AnimNextGraphRoot.h"
 #include "Graph/RigUnit_AnimNextDecoratorStack.h"
@@ -23,11 +24,6 @@
 
 #if WITH_DEV_AUTOMATION_TESTS
 
-void UAnimNextGraphTest::SetEditorData(UAnimNextGraph_EditorData* InEditorData)
-{
-	EditorData = InEditorData;
-}
-
 //****************************************************************************
 // AnimNext Runtime Decorator Graph Tests
 //****************************************************************************
@@ -86,13 +82,12 @@ bool FAnimationAnimNextRuntimeTest_GraphAddDecorator::RunTest(const FString& InP
 
 	FScopedClearNodeTemplateRegistry ScopedClearNodeTemplateRegistry;
 
-	UAnimNextGraphTest* AnimNextGraph = NewObject<UAnimNextGraphTest>();
+	UFactory* GraphFactory = NewObject<UAnimNextGraphFactory>();
+	UAnimNextGraph* AnimNextGraph = CastChecked<UAnimNextGraph>(GraphFactory->FactoryCreateNew(UAnimNextGraph::StaticClass(), GetTransientPackage(), TEXT("TestAnimNextGraph"), RF_Transient, nullptr, nullptr, NAME_None));
+	UE_RETURN_ON_ERROR(AnimNextGraph != nullptr, "FAnimationAnimNextRuntimeTest_GraphAddDecorator -> Failed to create animation graph");
 
-	UAnimNextGraph_EditorData* EditorData = NewObject<UAnimNextGraph_EditorData>(AnimNextGraph, TEXT("EditorData"));
-	AnimNextGraph->SetEditorData(EditorData);
-
-	EditorData->Initialize(/*bRecompileVM*/false);
-	EditorData->GetRigVMClient()->SetExecuteContextStruct(FAnimNextExecuteContext::StaticStruct());
+	UAnimNextGraph_EditorData* EditorData = UncookedOnly::FUtils::GetEditorData(AnimNextGraph);
+	UE_RETURN_ON_ERROR(EditorData != nullptr, "FAnimationAnimNextRuntimeTest_GraphAddDecorator -> Failed to find animation graph editor data");
 
 	URigVMController* Controller = EditorData->GetRigVMClient()->GetController(EditorData->GetRigVMClient()->GetDefaultModel());
 	UE_RETURN_ON_ERROR(Controller != nullptr, "FAnimationAnimNextRuntimeTest_GraphAddDecorator -> Failed to get RigVM controller");
@@ -185,19 +180,18 @@ bool FAnimationAnimNextRuntimeTest_GraphExecute::RunTest(const FString& InParame
 
 	FScopedClearNodeTemplateRegistry ScopedClearNodeTemplateRegistry;
 
-	UAnimNextGraphTest* AnimNextGraph = NewObject<UAnimNextGraphTest>();
-
-	UAnimNextGraph_EditorData* EditorData = NewObject<UAnimNextGraph_EditorData>(AnimNextGraph, TEXT("EditorData"));
-	AnimNextGraph->SetEditorData(EditorData);
+	UFactory* GraphFactory = NewObject<UAnimNextGraphFactory>();
+	UAnimNextGraph* AnimNextGraph = CastChecked<UAnimNextGraph>(GraphFactory->FactoryCreateNew(UAnimNextGraph::StaticClass(), GetTransientPackage(), TEXT("TestAnimNextGraph"), RF_Transient, nullptr, nullptr, NAME_None));
+	UE_RETURN_ON_ERROR(AnimNextGraph != nullptr, "FAnimationAnimNextRuntimeTest_GraphExecute -> Failed to create animation graph");
 
-	EditorData->Initialize(/*bRecompileVM*/false);
-	EditorData->GetRigVMClient()->SetExecuteContextStruct(FAnimNextExecuteContext::StaticStruct());
+	UAnimNextGraph_EditorData* EditorData = UncookedOnly::FUtils::GetEditorData(AnimNextGraph);
+	UE_RETURN_ON_ERROR(EditorData != nullptr, "FAnimationAnimNextRuntimeTest_GraphExecute -> Failed to find animation graph editor data");
 
 	URigVMController* Controller = EditorData->GetRigVMClient()->GetController(EditorData->GetRigVMClient()->GetDefaultModel());
 	UE_RETURN_ON_ERROR(Controller != nullptr, "FAnimationAnimNextRuntimeTest_GraphExecute -> Failed to get RigVM controller");
 
-	// Add graph entry point
-	URigVMUnitNode* MainEntryPointNode = Controller->AddUnitNode(FRigUnit_AnimNextGraphRoot::StaticStruct(), FRigUnit_AnimNextGraphRoot::EventName, FVector2D(0.0f, 0.0f), FString(), false);
+	// Find graph entry point
+	URigVMNode* MainEntryPointNode = Controller->GetGraph()->FindNodeByName(FRigUnit_AnimNextGraphRoot::StaticStruct()->GetFName());
 	UE_RETURN_ON_ERROR(MainEntryPointNode != nullptr, "FAnimationAnimNextRuntimeTest_GraphExecute -> Failed to find main entry point node");
 
 	URigVMPin* BeginExecutePin = MainEntryPointNode->FindPin(GET_MEMBER_NAME_STRING_CHECKED(FRigUnit_AnimNextGraphRoot, Result));
@@ -297,19 +291,18 @@ bool FAnimationAnimNextRuntimeTest_GraphExecuteLatent::RunTest(const FString& In
 
 	FScopedClearNodeTemplateRegistry ScopedClearNodeTemplateRegistry;
 
-	UAnimNextGraphTest* AnimNextGraph = NewObject<UAnimNextGraphTest>();
-
-	UAnimNextGraph_EditorData* EditorData = NewObject<UAnimNextGraph_EditorData>(AnimNextGraph, TEXT("EditorData"));
-	AnimNextGraph->SetEditorData(EditorData);
+	UFactory* GraphFactory = NewObject<UAnimNextGraphFactory>();
+	UAnimNextGraph* AnimNextGraph = CastChecked<UAnimNextGraph>(GraphFactory->FactoryCreateNew(UAnimNextGraph::StaticClass(), GetTransientPackage(), TEXT("TestAnimNextGraph"), RF_Transient, nullptr, nullptr, NAME_None));
+	UE_RETURN_ON_ERROR(AnimNextGraph != nullptr, "FAnimationAnimNextRuntimeTest_GraphExecuteLatent -> Failed to create animation graph");
 
-	EditorData->Initialize(/*bRecompileVM*/false);
-	EditorData->GetRigVMClient()->SetExecuteContextStruct(FAnimNextExecuteContext::StaticStruct());
+	UAnimNextGraph_EditorData* EditorData = UncookedOnly::FUtils::GetEditorData(AnimNextGraph);
+	UE_RETURN_ON_ERROR(EditorData != nullptr, "FAnimationAnimNextRuntimeTest_GraphExecuteLatent -> Failed to find animation graph editor data");
 
 	URigVMController* Controller = EditorData->GetRigVMClient()->GetController(EditorData->GetRigVMClient()->GetDefaultModel());
 	UE_RETURN_ON_ERROR(Controller != nullptr, "FAnimationAnimNextRuntimeTest_GraphExecuteLatent -> Failed to get RigVM controller");
 
-	// Add graph entry point
-	URigVMUnitNode* MainEntryPointNode = Controller->AddUnitNode(FRigUnit_AnimNextGraphRoot::StaticStruct(), FRigUnit_AnimNextGraphRoot::EventName, FVector2D(0.0f, 0.0f), FString(), false);
+	// Find graph entry point
+	URigVMNode* MainEntryPointNode = Controller->GetGraph()->FindNodeByName(FRigUnit_AnimNextGraphRoot::StaticStruct()->GetFName());
 	UE_RETURN_ON_ERROR(MainEntryPointNode != nullptr, "FAnimationAnimNextRuntimeTest_GraphExecuteLatent -> Failed to find main entry point node");
 
 	URigVMPin* BeginExecutePin = MainEntryPointNode->FindPin(GET_MEMBER_NAME_STRING_CHECKED(FRigUnit_AnimNextGraphRoot, Result));
diff --git a/Engine/Plugins/Experimental/AnimNext/Source/AnimNextTestSuite/Private/AnimNextDecoratorGraphTest.h b/Engine/Plugins/Experimental/AnimNext/Source/AnimNextTestSuite/Private/AnimNextDecoratorGraphTest.h
--- a/Engine/Plugins/Experimental/AnimNext/Source/AnimNextTestSuite/Private/AnimNextDecoratorGraphTest.h
+++ b/Engine/Plugins/Experimental/AnimNext/Source/AnimNextTestSuite/Private/AnimNextDecoratorGraphTest.h
@@ -9,15 +9,6 @@
 
 #include "AnimNextDecoratorGraphTest.generated.h"
 
-UCLASS(BlueprintType)
-class UAnimNextGraphTest : public UAnimNextGraph
-{
-	GENERATED_BODY()
-
-public:
-	void SetEditorData(class UAnimNextGraph_EditorData* InEditorData);
-};
-
 USTRUCT()
 struct FTestDecoratorSharedData : public FAnimNextDecoratorSharedData
 {
diff --git a/Engine/Plugins/Experimental/AnimNext/Source/AnimNextUncookedOnly/Private/UncookedOnlyUtils.cpp b/Engine/Plugins/Experimental/AnimNext/Source/AnimNextUncookedOnly/Private/UncookedOnlyUtils.cpp
--- a/Engine/Plugins/Experimental/AnimNext/Source/AnimNextUncookedOnly/Private/UncookedOnlyUtils.cpp
+++ b/Engine/Plugins/Experimental/AnimNext/Source/AnimNextUncookedOnly/Private/UncookedOnlyUtils.cpp
@@ -237,7 +237,7 @@ namespace Private
 
 		TArray<FDecoratorStackMapping> DecoratorStackNodes;
 
-		if (!ensure(VMRootNode != nullptr))
+		if (VMRootNode == nullptr)
 		{
 			// Root node wasn't found, add it, we'll need it to compile
 			VMRootNode = VMController->AddUnitNode(FRigUnit_AnimNextGraphRoot::StaticStruct(), FRigUnit_AnimNextGraphRoot::EventName, FVector2D(0.0f, 0.0f), FString(), false);
