AnimNext: Fix failing parameter test

diff --git a/Engine/Plugins/Experimental/AnimNext/Source/AnimNextTestSuite/Private/AnimNextParametersTest.cpp b/Engine/Plugins/Experimental/AnimNext/Source/AnimNextTestSuite/Private/AnimNextParametersTest.cpp
--- a/Engine/Plugins/Experimental/AnimNext/Source/AnimNextTestSuite/Private/AnimNextParametersTest.cpp
+++ b/Engine/Plugins/Experimental/AnimNext/Source/AnimNextTestSuite/Private/AnimNextParametersTest.cpp
@@ -262,12 +262,12 @@ bool FParamTypesTest::RunTest(const FString& InParameters)
 	AddErrorIfFalse(FParamTypeHandle::GetHandle<FVector4>().IsBuiltInType(), TEXT("FVector4 parameter handle is not built in."));
 	AddErrorIfFalse(FParamTypeHandle::GetHandle<FQuat>().IsBuiltInType(), TEXT("FQuat parameter handle is not built in."));
 	AddErrorIfFalse(FParamTypeHandle::GetHandle<FTransform>().IsBuiltInType(), TEXT("FTransform parameter handle is not built in."));
-
+	AddErrorIfFalse(FParamTypeHandle::GetHandle<UObject*>().IsBuiltInType(), TEXT("UObject parameter handle is not built in."));
+	AddErrorIfFalse(FParamTypeHandle::GetHandle<TObjectPtr<UObject>>().IsBuiltInType(), TEXT("TObjectPtr<UObject> parameter handle is not built in."));
+	
 	// Should not be built in:
 	AddErrorIfFalse(!FParamTypeHandle::GetHandle<EPropertyBagContainerType>().IsBuiltInType(), TEXT("Enum parameter handle is built in."));
 	AddErrorIfFalse(!FParamTypeHandle::GetHandle<FAnimNextParamType>().IsBuiltInType(), TEXT("Struct parameter handle is built in."));
-	AddErrorIfFalse(!FParamTypeHandle::GetHandle<UObject*>().IsBuiltInType(), TEXT("UObject parameter handle is built in."));
-	AddErrorIfFalse(!FParamTypeHandle::GetHandle<TObjectPtr<UObject>>().IsBuiltInType(), TEXT("TObjectPtr<UObject> parameter handle is built in."));
 	AddErrorIfFalse(!FParamTypeHandle::GetHandle<UClass*>().IsBuiltInType(), TEXT("UClass parameter handle is built in."));
 	AddErrorIfFalse(!FParamTypeHandle::GetHandle<TSubclassOf<UObject>>().IsBuiltInType(), TEXT("TSubclassOff<UObject> parameter handle is built in."));
 	AddErrorIfFalse(!FParamTypeHandle::GetHandle<TSoftObjectPtr<UObject>>().IsBuiltInType(), TEXT("TSoftObjectPtr<UObject> parameter handle is built in."));
