AnimNext: Fix cooking crash and anim BP execution path

Ensure a minimal way of accessing context object for parameter adapters

diff --git a/Engine/Plugins/Experimental/AnimNext/Source/AnimNext/Private/Graph/AnimGraph/AnimNode_AnimNextGraph.cpp b/Engine/Plugins/Experimental/AnimNext/Source/AnimNext/Private/Graph/AnimGraph/AnimNode_AnimNextGraph.cpp
--- a/Engine/Plugins/Experimental/AnimNext/Source/AnimNext/Private/Graph/AnimGraph/AnimNode_AnimNextGraph.cpp
+++ b/Engine/Plugins/Experimental/AnimNext/Source/AnimNext/Private/Graph/AnimGraph/AnimNode_AnimNextGraph.cpp
@@ -16,6 +16,7 @@
 #include "BoneContainer.h"
 #include "Param/ParamStack.h"
 #include "AnimGraphParamStackScope.h"
+#include "Scheduler/ScheduleContext.h"
 
 #include UE_INLINE_GENERATED_CPP_BY_NAME(AnimNode_AnimNextGraph)
 
@@ -143,13 +144,16 @@ void FAnimNode_AnimNextGraph::Evaluate_AnyThread(FPoseContext& Output)
 
 		FAnimNextGraphLODPose ResultPose(FLODPoseHeap(RefPose, LODLevel, true, Output.ExpectsAdditivePose()));
 
-		UE::AnimNext::FAnimGraphParamStackScope Scope(Output);
+		FScheduleContext ScheduleContext(SkeletalMeshComponent);
+		FScheduleContext::AttachToCurrentThread(ScheduleContext);
+
+		FAnimGraphParamStackScope Scope(Output);
 		FParamStack& ParamStack = FParamStack::Get();
 		FParamStack::FPushedLayerHandle LayerHandle = ParamStack.PushValues(
-			"GraphReferencePose", GraphReferencePose,
-			"ResultPose", ResultPose,
-			"GraphLODLevel", LODLevel,
-			"GraphExpectsAdditive", Output.ExpectsAdditivePose()
+			AnimNextGraph->GetReferencePoseParam(), GraphReferencePose,
+			"UE_Internal_ResultPose", ResultPose,
+			AnimNextGraph->GetCurrentLODParam(), LODLevel,
+			"UE_Internal_GraphExpectsAdditive", Output.ExpectsAdditivePose()
 		);
 
 		AnimNextGraph->Run(Context, GraphInstance, EAnimNextGraphSimulationSteps::Evaluate);
@@ -157,6 +161,8 @@ void FAnimNode_AnimNextGraph::Evaluate_AnyThread(FPoseContext& Output)
 		FGenerationTools::RemapPose(ResultPose.LODPose, Output);
 
 		ParamStack.PopLayer(LayerHandle);
+
+		FScheduleContext::DetachFromCurrentThread();
 	}
 	else
 	{
diff --git a/Engine/Plugins/Experimental/AnimNext/Source/AnimNext/Private/Param/ParamAdapter.cpp b/Engine/Plugins/Experimental/AnimNext/Source/AnimNext/Private/Param/ParamAdapter.cpp
--- a/Engine/Plugins/Experimental/AnimNext/Source/AnimNext/Private/Param/ParamAdapter.cpp
+++ b/Engine/Plugins/Experimental/AnimNext/Source/AnimNext/Private/Param/ParamAdapter.cpp
@@ -19,7 +19,7 @@ FParamResult FParamAdapter::GetParamData(FParamTypeHandle InTypeHandle, TConstAr
 		return EParamResult::TypeError;
 	}
 
-	UObject* ContextObject = FScheduleContext::Get().Entry->ResolvedObject;
+	UObject* ContextObject = FScheduleContext::Get().GetContextObject();
 	if(const uint8* Data = Function(ContextObject, Definition.Id))
 	{
 		OutParamData = TConstArrayView<uint8>(Data, Definition.TypeHandle.GetSize());
@@ -52,7 +52,7 @@ FParamResult FParamAdapter::GetMutableParamData(FParamTypeHandle InTypeHandle, T
 		return EParamResult::TypeError;
 	}
 
-	UObject* ContextObject = FScheduleContext::Get().Entry->ResolvedObject;
+	UObject* ContextObject = FScheduleContext::Get().GetContextObject();
 	if(uint8* Data = Function(ContextObject, Definition.Id))
 	{
 		OutParamData = TArrayView<uint8>(Data, Definition.TypeHandle.GetSize());
diff --git a/Engine/Plugins/Experimental/AnimNext/Source/AnimNext/Private/Scheduler/ScheduleContext.cpp b/Engine/Plugins/Experimental/AnimNext/Source/AnimNext/Private/Scheduler/ScheduleContext.cpp
--- a/Engine/Plugins/Experimental/AnimNext/Source/AnimNext/Private/Scheduler/ScheduleContext.cpp
+++ b/Engine/Plugins/Experimental/AnimNext/Source/AnimNext/Private/Scheduler/ScheduleContext.cpp
@@ -1,6 +1,7 @@
 // Copyright Epic Games, Inc. All Rights Reserved.
 
 #include "ScheduleContext.h"
+#include "AnimNextSchedulerEntry.h"
 #include "Scheduler/AnimNextSchedule.h"
 
 namespace UE::AnimNext
@@ -27,4 +28,16 @@ const FScheduleContext& FScheduleContext::Get()
 	return *GScheduleContext;
 }
 
+UObject* FScheduleContext::GetContextObject() const
+{
+	if(Entry)
+	{
+		return Entry->ResolvedObject;
+	}
+	else
+	{
+		return ContextObject;
+	}
+}
+
 }
diff --git a/Engine/Plugins/Experimental/AnimNext/Source/AnimNext/Private/Scheduler/ScheduleContext.h b/Engine/Plugins/Experimental/AnimNext/Source/AnimNext/Private/Scheduler/ScheduleContext.h
--- a/Engine/Plugins/Experimental/AnimNext/Source/AnimNext/Private/Scheduler/ScheduleContext.h
+++ b/Engine/Plugins/Experimental/AnimNext/Source/AnimNext/Private/Scheduler/ScheduleContext.h
@@ -36,6 +36,10 @@ public:
 		, Entry(InEntry)
 	{}
 
+	explicit FScheduleContext(UObject* InContextObject)
+		: ContextObject(InContextObject)
+	{}
+	
 	// Get a string that can identify this schedule for debugging
 	FString GetDebugString() const 
 	{ 
@@ -53,6 +57,9 @@ public:
 		return *InstanceData.Get(); 
 	}
 
+	// Get the current object context that the schedule is running in 
+	UObject* GetContextObject() const;
+
 	// Current schedule
 	const UAnimNextSchedule* Schedule = nullptr;
 
@@ -61,6 +68,9 @@ public:
 
 	// Instance data for the schedule
 	TUniquePtr<FScheduleInstanceData> InstanceData;
+
+	// Context object - only valid when we dont have a running schedule
+	UObject* ContextObject = nullptr;
 };
 
 }
