AnimNext: Fix uninitialized mutex access during static initialization

Move static FParamIds from global to local scope. To fix this generally we will need to rework FParamId.

diff --git a/Engine/Plugins/Experimental/AnimNext/Source/AnimNext/Private/Graph/AnimNextGraph.cpp b/Engine/Plugins/Experimental/AnimNext/Source/AnimNext/Private/Graph/AnimNextGraph.cpp
--- a/Engine/Plugins/Experimental/AnimNext/Source/AnimNext/Private/Graph/AnimNextGraph.cpp
+++ b/Engine/Plugins/Experimental/AnimNext/Source/AnimNext/Private/Graph/AnimNextGraph.cpp
@@ -18,11 +18,6 @@ DEFINE_STAT(STAT_AnimNext_Graph_AllocateInstance);
 
 #include UE_INLINE_GENERATED_CPP_BY_NAME(AnimNextGraph)
 
-namespace UE::AnimNext::Graph
-{
-const UE::AnimNext::FParamId OutputPoseId("UE_Internal_Graph_OutputPose");
-}
-
 #if WITH_EDITORONLY_DATA
 FAnimNextGraphInstance::FAnimNextGraphInstance(const FAnimNextGraphInstance& Other)
 {
@@ -291,9 +286,11 @@ TConstArrayView<UE::AnimNext::FScheduleTerm> UAnimNextGraph::GetTerms() const
 {
 	using namespace UE::AnimNext;
 
+	static const FParamId OutputPoseId("UE_Internal_Graph_OutputPose");
+
 	static const FScheduleTerm Terms[] =
 	{
-		FScheduleTerm(Graph::OutputPoseId, FAnimNextParamType::GetType<FAnimNextGraphLODPose>(), EScheduleTermDirection::Output)
+		FScheduleTerm(OutputPoseId, FAnimNextParamType::GetType<FAnimNextGraphLODPose>(), EScheduleTermDirection::Output)
 	};
 
 	return Terms;
diff --git a/Engine/Plugins/Experimental/AnimNext/Source/AnimNext/Private/Scheduler/Ports/AnimNextSchedulePort_AnimNextMeshComponentPose.cpp b/Engine/Plugins/Experimental/AnimNext/Source/AnimNext/Private/Scheduler/Ports/AnimNextSchedulePort_AnimNextMeshComponentPose.cpp
--- a/Engine/Plugins/Experimental/AnimNext/Source/AnimNext/Private/Scheduler/Ports/AnimNextSchedulePort_AnimNextMeshComponentPose.cpp
+++ b/Engine/Plugins/Experimental/AnimNext/Source/AnimNext/Private/Scheduler/Ports/AnimNextSchedulePort_AnimNextMeshComponentPose.cpp
@@ -13,8 +13,6 @@
 
 DEFINE_STAT(STAT_AnimNext_Port_SkeletalMeshComponent);
 
-const UE::AnimNext::FParamId UAnimNextSchedulePort_AnimNextMeshComponentPose::InputId("UE_Internal_AnimNextMeshComponentPose_Input");
-
 void UAnimNextSchedulePort_AnimNextMeshComponentPose::Run(const UE::AnimNext::FScheduleTermContext& InContext) const
 {
 	SCOPE_CYCLE_COUNTER(STAT_AnimNext_Port_SkeletalMeshComponent);
@@ -30,7 +28,7 @@ void UAnimNextSchedulePort_AnimNextMeshComponentPose::Run(const UE::AnimNext::FS
 		return;
 	}
 
-	const FAnimNextGraphLODPose* InputPose = InContext.GetLayerHandle().GetParamPtr<FAnimNextGraphLODPose>(InputId);
+	const FAnimNextGraphLODPose* InputPose = InContext.GetLayerHandle().GetParamPtr<FAnimNextGraphLODPose>(GetTerms()[0].GetId());
 	if(InputPose == nullptr)
 	{
 		return;
@@ -75,6 +73,8 @@ TConstArrayView<UE::AnimNext::FScheduleTerm> UAnimNextSchedulePort_AnimNextMeshC
 {
 	using namespace UE::AnimNext;
 
+	static const FParamId InputId("UE_Internal_AnimNextMeshComponentPose_Input");
+
 	static const FScheduleTerm Terms[] =
 	{
 		FScheduleTerm(InputId, FAnimNextParamType::GetType<FAnimNextGraphLODPose>(), EScheduleTermDirection::Input)
diff --git a/Engine/Plugins/Experimental/AnimNext/Source/AnimNext/Private/Scheduler/Ports/AnimNextSchedulePort_AnimNextMeshComponentPose.h b/Engine/Plugins/Experimental/AnimNext/Source/AnimNext/Private/Scheduler/Ports/AnimNextSchedulePort_AnimNextMeshComponentPose.h
--- a/Engine/Plugins/Experimental/AnimNext/Source/AnimNext/Private/Scheduler/Ports/AnimNextSchedulePort_AnimNextMeshComponentPose.h
+++ b/Engine/Plugins/Experimental/AnimNext/Source/AnimNext/Private/Scheduler/Ports/AnimNextSchedulePort_AnimNextMeshComponentPose.h
@@ -13,7 +13,4 @@ class UAnimNextSchedulePort_AnimNextMeshComponentPose : public UAnimNextSchedule
 	// UAnimNextSchedulePort interface
 	virtual void Run(const UE::AnimNext::FScheduleTermContext& InContext) const override;
 	virtual TConstArrayView<UE::AnimNext::FScheduleTerm> GetTerms() const override;
-
-	// ID of our input
-	static const UE::AnimNext::FParamId InputId;
 };
\ No newline at end of file
