AnimNext: Misc fixes

Fix corner case for being able unable access various combinations of object parameter types (this needs to change!)
Fix being unable to spawn new units in AnimNext graphs
Added missing ini

diff --git a/Engine/Plugins/Experimental/AnimNext/Source/AnimNext/Private/Param/ParamUtils.cpp b/Engine/Plugins/Experimental/AnimNext/Source/AnimNext/Private/Param/ParamUtils.cpp
--- a/Engine/Plugins/Experimental/AnimNext/Source/AnimNext/Private/Param/ParamUtils.cpp
+++ b/Engine/Plugins/Experimental/AnimNext/Source/AnimNext/Private/Param/ParamUtils.cpp
@@ -14,7 +14,7 @@ namespace UE::AnimNext
 
 FParamCompatibility FParamUtils::GetCompatibility(const FParamTypeHandle& InLHS, const FParamTypeHandle& InRHS)
 {
-	auto CheckClassCast = [](const FParamTypeHandle& InLHS, const UClass* InRHSClass)
+	auto CheckClassCastToCustom = [](const FParamTypeHandle& InLHS, const UClass* InRHSClass)
 	{
 		FAnimNextParamType::EValueType ValueTypeLHS;
 		FAnimNextParamType::EContainerType ContainerTypeLHS;
@@ -33,6 +33,26 @@ FParamCompatibility FParamUtils::GetCompatibility(const FParamTypeHandle& InLHS,
 		}
 		return false;
 	};
+
+	auto CheckClassCastFromCustom = [](const UClass* InLHSClass, const FParamTypeHandle& InRHS)
+	{
+		FAnimNextParamType::EValueType ValueTypeRHS;
+		FAnimNextParamType::EContainerType ContainerTypeRHS;
+		const UObject* ValueTypeObjectRHS;
+
+		InRHS.GetCustomTypeInfo(ValueTypeRHS, ContainerTypeRHS, ValueTypeObjectRHS);
+		if(ContainerTypeRHS == FAnimNextParamType::EContainerType::None && ValueTypeRHS == FAnimNextParamType::EValueType::Object)
+		{
+			if(const UClass* Class = Cast<UClass>(ValueTypeObjectRHS))
+			{
+				if(Class->IsChildOf(InLHSClass))
+				{
+					return true;
+				}
+			}
+		}
+		return false;
+	};
 	
 	switch (InRHS.GetParameterType())
 	{
@@ -172,6 +192,11 @@ FParamCompatibility FParamUtils::GetCompatibility(const FParamTypeHandle& InLHS,
 		{
 		case FParamTypeHandle::EParamType::Object:
 			return EParamCompatibility::Compatible_Equal;
+		case FParamTypeHandle::EParamType::Custom:
+			if(CheckClassCastToCustom(InLHS, UObject::StaticClass()))
+			{
+				return EParamCompatibility::Compatible_Cast;
+			}
 		}
 		break;
 	case FParamTypeHandle::EParamType::CharacterMovementComponent:
@@ -182,7 +207,7 @@ FParamCompatibility FParamUtils::GetCompatibility(const FParamTypeHandle& InLHS,
 		case FParamTypeHandle::EParamType::CharacterMovementComponent:
 			return EParamCompatibility::Compatible_Equal;
 		case FParamTypeHandle::EParamType::Custom:
-			if(CheckClassCast(InLHS, UCharacterMovementComponent::StaticClass()))
+			if(CheckClassCastToCustom(InLHS, UCharacterMovementComponent::StaticClass()))
 			{
 				return EParamCompatibility::Compatible_Cast;
 			}
@@ -196,7 +221,7 @@ FParamCompatibility FParamUtils::GetCompatibility(const FParamTypeHandle& InLHS,
 		case FParamTypeHandle::EParamType::AnimNextMeshComponent:
 			return EParamCompatibility::Compatible_Equal;
 		case FParamTypeHandle::EParamType::Custom:
-			if(CheckClassCast(InLHS, UAnimNextMeshComponent::StaticClass()))
+			if(CheckClassCastToCustom(InLHS, UAnimNextMeshComponent::StaticClass()))
 			{
 				return EParamCompatibility::Compatible_Cast;
 			}
@@ -210,7 +235,7 @@ FParamCompatibility FParamUtils::GetCompatibility(const FParamTypeHandle& InLHS,
 		case FParamTypeHandle::EParamType::AnimSequence:
 			return EParamCompatibility::Compatible_Equal;
 		case FParamTypeHandle::EParamType::Custom:
-			if(CheckClassCast(InLHS, UAnimSequence::StaticClass()))
+			if(CheckClassCastToCustom(InLHS, UAnimSequence::StaticClass()))
 			{
 				return EParamCompatibility::Compatible_Cast;
 			}
@@ -233,6 +258,30 @@ FParamCompatibility FParamUtils::GetCompatibility(const FParamTypeHandle& InLHS,
 	case FParamTypeHandle::EParamType::Custom:
 		switch (InLHS.GetParameterType())
 		{
+		case FParamTypeHandle::EParamType::Object:
+			if(CheckClassCastFromCustom(UObject::StaticClass(), InRHS))
+			{
+				return EParamCompatibility::Compatible_Cast;
+			}
+			break;
+		case FParamTypeHandle::EParamType::CharacterMovementComponent:
+			if(CheckClassCastFromCustom(UCharacterMovementComponent::StaticClass(), InRHS))
+			{
+				return EParamCompatibility::Compatible_Cast;
+			}
+			break;
+		case FParamTypeHandle::EParamType::AnimNextMeshComponent:
+			if(CheckClassCastFromCustom(UAnimNextMeshComponent::StaticClass(), InRHS))
+			{
+				return EParamCompatibility::Compatible_Cast;
+			}
+			break;
+		case FParamTypeHandle::EParamType::AnimSequence:
+			if(CheckClassCastFromCustom(UAnimSequence::StaticClass(), InRHS))
+			{
+				return EParamCompatibility::Compatible_Cast;
+			}
+			break;
 		case FParamTypeHandle::EParamType::Custom:
 			{
 				FAnimNextParamType::EValueType ValueTypeLHS, ValueTypeRHS;
diff --git a/Engine/Plugins/Experimental/AnimNext/Source/AnimNextEditor/Private/Common/GraphEditorSchemaActions.cpp b/Engine/Plugins/Experimental/AnimNext/Source/AnimNextEditor/Private/Common/GraphEditorSchemaActions.cpp
--- a/Engine/Plugins/Experimental/AnimNext/Source/AnimNextEditor/Private/Common/GraphEditorSchemaActions.cpp
+++ b/Engine/Plugins/Experimental/AnimNext/Source/AnimNextEditor/Private/Common/GraphEditorSchemaActions.cpp
@@ -11,9 +11,9 @@
 
 UEdGraphNode* FAnimNextSchemaAction_RigUnit::PerformAction(UEdGraph* ParentGraph, TArray<UEdGraphPin*>& FromPins, const FVector2D Location, bool bSelectNewNode)
 {
-	UAnimNextGraph_EditorData* EditorData = ParentGraph->GetTypedOuter<UAnimNextGraph_EditorData>();
-	UAnimNextGraph_EdGraphNode* NewNode = nullptr;
-	UAnimNextGraph_EdGraph* EdGraph = Cast<UAnimNextGraph_EdGraph>(ParentGraph);
+	IRigVMClientHost* Host = ParentGraph->GetImplementingOuter<IRigVMClientHost>();
+	URigVMEdGraphNode* NewNode = nullptr;
+	URigVMEdGraph* EdGraph = Cast<URigVMEdGraph>(ParentGraph);
 
 	UEdGraphPin* FromPin = nullptr;
 	if (FromPins.Num() > 0)
@@ -21,10 +21,10 @@ UEdGraphNode* FAnimNextSchemaAction_RigUnit::PerformAction(UEdGraph* ParentGraph
 		FromPin = FromPins[0];
 	}
 	
-	if (EditorData != nullptr && EdGraph != nullptr)
+	if (Host != nullptr && EdGraph != nullptr)
 	{
-		FName Name = UE::AnimNext::Editor::FUtils::ValidateName(EditorData, StructTemplate->GetFName().ToString());
-		URigVMController* Controller = EditorData->GetRigVMClient()->GetController(ParentGraph);
+		FName Name = UE::AnimNext::Editor::FUtils::ValidateName(Cast<UObject>(Host), StructTemplate->GetFName().ToString());
+		URigVMController* Controller = Host->GetRigVMClient()->GetController(ParentGraph);
 
 		Controller->OpenUndoBracket(FString::Printf(TEXT("Add '%s' Node"), *Name.ToString()));
 
@@ -33,7 +33,7 @@ UEdGraphNode* FAnimNextSchemaAction_RigUnit::PerformAction(UEdGraph* ParentGraph
 
 		if (URigVMUnitNode* ModelNode = Controller->AddUnitNode(StructTemplate, FRigVMStruct::ExecuteName, Location, Name.ToString(), true, false))
 		{
-			NewNode = Cast<UAnimNextGraph_EdGraphNode>(EdGraph->FindNodeForModelNodeName(ModelNode->GetFName()));
+			NewNode = Cast<URigVMEdGraphNode>(EdGraph->FindNodeForModelNodeName(ModelNode->GetFName()));
 			check(NewNode);
 
 			if (NewNode)
@@ -73,9 +73,9 @@ UEdGraphNode* FAnimNextSchemaAction_RigUnit::PerformAction(UEdGraph* ParentGraph
 
 UEdGraphNode* FAnimNextSchemaAction_DispatchFactory::PerformAction(UEdGraph* ParentGraph, TArray<UEdGraphPin*>& FromPins, const FVector2D Location, bool bSelectNewNode)
 {
-	UAnimNextGraph_EditorData* EditorData = ParentGraph->GetTypedOuter<UAnimNextGraph_EditorData>();
-	UAnimNextGraph_EdGraphNode* NewNode = nullptr;
-	UAnimNextGraph_EdGraph* EdGraph = Cast<UAnimNextGraph_EdGraph>(ParentGraph);
+	IRigVMClientHost* Host = ParentGraph->GetImplementingOuter<IRigVMClientHost>();
+	URigVMEdGraphNode* NewNode = nullptr;
+	URigVMEdGraph* EdGraph = Cast<URigVMEdGraph>(ParentGraph);
 
 	UEdGraphPin* FromPin = nullptr;
 	if (FromPins.Num() > 0)
@@ -83,7 +83,7 @@ UEdGraphNode* FAnimNextSchemaAction_DispatchFactory::PerformAction(UEdGraph* Par
 		FromPin = FromPins[0];
 	}
 	
-	if (EditorData != nullptr && EdGraph != nullptr)
+	if (Host != nullptr && EdGraph != nullptr)
 	{
 		const FRigVMTemplate* Template = FRigVMRegistry::Get().FindTemplate(Notation);
 		if (Template == nullptr)
@@ -94,14 +94,14 @@ UEdGraphNode* FAnimNextSchemaAction_DispatchFactory::PerformAction(UEdGraph* Par
 		const int32 NotationHash = (int32)GetTypeHash(Notation);
 		const FString TemplateName = TEXT("RigVMTemplate_") + FString::FromInt(NotationHash);
 
-		FName Name = UE::AnimNext::Editor::FUtils::ValidateName(EditorData, Template->GetName().ToString());
-		URigVMController* Controller = EditorData->GetRigVMClient()->GetController(ParentGraph);
+		FName Name = UE::AnimNext::Editor::FUtils::ValidateName(Cast<UObject>(Host), Template->GetName().ToString());
+		URigVMController* Controller = Host->GetRigVMClient()->GetController(ParentGraph);
 
 		Controller->OpenUndoBracket(FString::Printf(TEXT("Add '%s' Node"), *Name.ToString()));
 
 		if (URigVMTemplateNode* ModelNode = Controller->AddTemplateNode(Notation, Location, Name.ToString(), true, false))
 		{
-			NewNode = Cast<UAnimNextGraph_EdGraphNode>(EdGraph->FindNodeForModelNodeName(ModelNode->GetFName()));
+			NewNode = Cast<URigVMEdGraphNode>(EdGraph->FindNodeForModelNodeName(ModelNode->GetFName()));
 
 			if (NewNode)
 			{
diff --git a/Engine/Plugins/Experimental/AnimNext/Source/AnimNextEditor/Private/EditorUtils.cpp b/Engine/Plugins/Experimental/AnimNext/Source/AnimNextEditor/Private/EditorUtils.cpp
--- a/Engine/Plugins/Experimental/AnimNext/Source/AnimNextEditor/Private/EditorUtils.cpp
+++ b/Engine/Plugins/Experimental/AnimNext/Source/AnimNextEditor/Private/EditorUtils.cpp
@@ -30,14 +30,13 @@ void FUtils::GetAllGraphNames(const UAnimNextGraph_EditorData* InEditorData, TSe
 
 static const int32 MaxNameLength = 100;
 
-FName FUtils::ValidateName(const UAnimNextGraph_EditorData* InEditorData, const FString& InName)
+FName FUtils::ValidateName(const UObject* InObject, const FString& InName)
 {
 	struct FNameValidator : public INameValidatorInterface
 	{
-		explicit FNameValidator(const UAnimNextGraph_EditorData* InEditorData)
-			: EditorData(InEditorData)
+		explicit FNameValidator(const UObject* InObject)
+			: Object(InObject)
 		{
-			GetAllGraphNames(EditorData, Names);
 		}
 		
 		virtual EValidatorResult IsValid (const FName& Name, bool bOriginal = false) override
@@ -60,7 +59,7 @@ FName FUtils::ValidateName(const UAnimNextGraph_EditorData* InEditorData, const
 					ValidatorResult = EValidatorResult::Ok;
 
 					// Check for collision with an existing object.
-					if (UObject* ExistingObject = StaticFindObject(/*Class=*/ nullptr, const_cast<UAnimNextGraph_EditorData*>(EditorData), *Name.ToString(), true))
+					if (UObject* ExistingObject = StaticFindObject(/*Class=*/ nullptr, const_cast<UObject*>(Object), *Name.ToString(), true))
 					{
 						ValidatorResult = EValidatorResult::AlreadyInUse;
 					}
@@ -89,7 +88,7 @@ FName FUtils::ValidateName(const UAnimNextGraph_EditorData* InEditorData, const
 		/** Name set to validate */
 		TSet<FName> Names;
 		/** The editor data to check for validity within */
-		const UAnimNextGraph_EditorData* EditorData;
+		const UObject* Object;
 	};
 	
 	FString Name = InName;
@@ -98,7 +97,7 @@ FName FUtils::ValidateName(const UAnimNextGraph_EditorData* InEditorData, const
 		Name.RightChopInline(8, false);
 	}
 
-	FNameValidator NameValidator(InEditorData);
+	FNameValidator NameValidator(InObject);
 
 	// Clean up BaseName to not contain any invalid characters, which will mean we can never find a legal name no matter how many numbers we add
 	if (NameValidator.IsValid(Name) == EValidatorResult::ContainsInvalidCharacters)
diff --git a/Engine/Plugins/Experimental/AnimNext/Source/AnimNextEditor/Private/EditorUtils.h b/Engine/Plugins/Experimental/AnimNext/Source/AnimNextEditor/Private/EditorUtils.h
--- a/Engine/Plugins/Experimental/AnimNext/Source/AnimNextEditor/Private/EditorUtils.h
+++ b/Engine/Plugins/Experimental/AnimNext/Source/AnimNextEditor/Private/EditorUtils.h
@@ -24,7 +24,7 @@ namespace UE::AnimNext::Editor
 
 struct FUtils
 {
-	static FName ValidateName(const UAnimNextGraph_EditorData* InEditorData, const FString& InName);
+	static FName ValidateName(const UObject* InObject, const FString& InName);
 
 	static void GetAllGraphNames(const UAnimNextGraph_EditorData* InEditorData, TSet<FName>& OutNames);
 
