Added inline value edit to parameters at SParameterBlockView
Fixed a deadlock at UAnimNextParameterBlock_EditorData AddGraph, when adding more than one graph parameter
Fixed rename not renaming the Library element at AnimNextParameterBlockParameter, which caused a desync
Created a AnimNextParameterBlockParameter class details customization, so we can edit default values
Separated graphs from other parameters in the ParameterView table

diff --git a/Engine/Plugins/Experimental/AnimNext/Source/AnimNext/Internal/Param/AnimNextParameterBlock.h b/Engine/Plugins/Experimental/AnimNext/Source/AnimNext/Internal/Param/AnimNextParameterBlock.h
--- a/Engine/Plugins/Experimental/AnimNext/Source/AnimNext/Internal/Param/AnimNextParameterBlock.h
+++ b/Engine/Plugins/Experimental/AnimNext/Source/AnimNext/Internal/Param/AnimNextParameterBlock.h
@@ -28,6 +28,8 @@ namespace UE::AnimNext::Editor
 {
 	class FParametersEditor;
 	struct FUtils;
+	class SParameterBlockViewRow;
+	class FParameterBlockParameterCustomization;
 }
 
 // Library entry used to export to asset registry
@@ -77,6 +79,8 @@ class UAnimNextParameterBlock : public URigVMHost, public IAnimNextParameterSour
 	friend struct FAnimNode_AnimNextParameters;
 	friend struct FAnimNextScheduleGraphTask;
 	friend struct FAnimNextScheduleParamScopeEntryTask;
+	friend class UE::AnimNext::Editor::SParameterBlockViewRow;
+	friend class UE::AnimNext::Editor::FParameterBlockParameterCustomization;
 
 	// IAnimNextParameterSourceInterface interface
 	virtual void UpdateLayer(UE::AnimNext::FParamStackLayerHandle& InHandle) const override;
@@ -86,7 +90,10 @@ class UAnimNextParameterBlock : public URigVMHost, public IAnimNextParameterSour
 	// UObject interface
 	virtual void PostLoad() override;
 	virtual void GetAssetRegistryTags(TArray<FAssetRegistryTag>& OutTags) const override;
-	
+
+	FInstancedPropertyBag& GetPropertyBag() { return PropertyBag; }
+
+
 	FRigVMExtendedExecuteContext BaseRigVMContext;
 
 	UPROPERTY()
diff --git a/Engine/Plugins/Experimental/AnimNext/Source/AnimNextEditor/AnimNextEditor.Build.cs b/Engine/Plugins/Experimental/AnimNext/Source/AnimNextEditor/AnimNextEditor.Build.cs
--- a/Engine/Plugins/Experimental/AnimNext/Source/AnimNextEditor/AnimNextEditor.Build.cs
+++ b/Engine/Plugins/Experimental/AnimNext/Source/AnimNextEditor/AnimNextEditor.Build.cs
@@ -33,6 +33,7 @@ namespace UnrealBuildTool.Rules
 					"AssetDefinition",
 					"SourceControl", 
 					"KismetWidgets",
+					"StructUtils",
 					"StructUtilsEditor",
 					"BlueprintGraph",	// For K2 Schema
 					"DesktopWidgets",
diff --git a/Engine/Plugins/Experimental/AnimNext/Source/AnimNextEditor/Private/AnimNextEditorModule.cpp b/Engine/Plugins/Experimental/AnimNext/Source/AnimNextEditor/Private/AnimNextEditorModule.cpp
--- a/Engine/Plugins/Experimental/AnimNext/Source/AnimNextEditor/Private/AnimNextEditorModule.cpp
+++ b/Engine/Plugins/Experimental/AnimNext/Source/AnimNextEditor/Private/AnimNextEditorModule.cpp
@@ -21,7 +21,8 @@
 #include "Graph/SAnimNextGraphView.h"
 #include "Scheduler/AnimNextSchedule.h"
 #include "Workspace/AnimNextWorkspaceEditor.h"
-
+#include "Param/ParameterBlockParameterCustomization.h"
+#include "Param/AnimNextParameterBlockParameter.h"
 
 #define LOCTEXT_NAMESPACE "AnimNextEditorModule"
 
@@ -52,6 +53,9 @@ class FModule : public IModule
 			FOnGetPropertyTypeCustomizationInstance::CreateLambda([] { return MakeShared<FParamNamePropertyTypeCustomization>(); }),
 			Identifier);
 
+		PropertyModule.RegisterCustomClassLayout(UAnimNextParameterBlockParameter::StaticClass()->GetFName(), 
+			FOnGetDetailCustomizationInstance::CreateLambda([] { return MakeShared<FParameterBlockParameterCustomization>(); }));
+
 		AnimNextGraphPanelNodeFactory = MakeShared<FAnimNextGraphPanelNodeFactory>();
 		FEdGraphUtilities::RegisterVisualNodeFactory(AnimNextGraphPanelNodeFactory);
 
diff --git a/Engine/Plugins/Experimental/AnimNext/Source/AnimNextEditor/Private/Param/ParameterBlockParameterCustomization.cpp b/Engine/Plugins/Experimental/AnimNext/Source/AnimNextEditor/Private/Param/ParameterBlockParameterCustomization.cpp
new file mode 100644
--- /dev/null
+++ b/Engine/Plugins/Experimental/AnimNext/Source/AnimNextEditor/Private/Param/ParameterBlockParameterCustomization.cpp
@@ -0,0 +1,99 @@
+// Copyright Epic Games, Inc. All Rights Reserved.
+
+#include "Param/ParameterBlockParameterCustomization.h"
+
+#include "DetailLayoutBuilder.h"
+#include "DetailWidgetRow.h"
+#include "PropertyBagDetails.h"
+#include "EdGraphSchema_K2.h"
+#include "EditorUtils.h"
+#include "UncookedOnlyUtils.h"
+#include "Param/ParamType.h"
+#include "PropertyHandle.h"
+#include "Param/AnimNextParameterBlockParameter.h"
+#include "Param/AnimNextParameterBlock.h"
+#include "ISinglePropertyView.h"
+#include "IStructureDataProvider.h"
+#include "UncookedOnlyUtils.h"
+
+#define LOCTEXT_NAMESPACE "ParamTypePropertyCustomization"
+
+namespace UE::AnimNext::Editor
+{
+
+void FParameterBlockParameterCustomization::CustomizeDetails(IDetailLayoutBuilder& DetailBuilder)
+{
+	TArray<TWeakObjectPtr<UObject>> Objects;
+	DetailBuilder.GetObjectsBeingCustomized(Objects);
+
+	if (Objects.IsEmpty() || Objects.Num() > 1)
+	{
+		return;
+	}
+
+	if (UAnimNextParameterBlockParameter* BlockParam = Cast<UAnimNextParameterBlockParameter>(Objects[0].Get()))
+	{
+		IDetailCategoryBuilder& ParameterCategory = DetailBuilder.EditCategory(TEXT("Parameter"), FText::GetEmpty(), ECategoryPriority::Important);
+
+
+		IDetailCategoryBuilder& DefaultValueCategory = DetailBuilder.EditCategory(TEXT("DefaultValue"), FText::GetEmpty(), ECategoryPriority::Default);
+
+		TSharedRef< SWidget > ColumnWidget = SNullWidget::NullWidget;
+
+		if (UAnimNextParameterBlock_EditorData* EditorData = Cast<UAnimNextParameterBlock_EditorData>(BlockParam->GetOuter()))
+		{
+			if (UAnimNextParameterBlock* ReferencedBlock = UE::AnimNext::UncookedOnly::FUtils::GetBlock(EditorData))
+			{
+				TSharedPtr<FStructOnScope> StructOnScope = MakeShareable(new FStructOnScope(ReferencedBlock->GetPropertyBag().GetPropertyBagStruct(), (uint8*)ReferencedBlock->GetPropertyBag().GetValue().GetMemory()));
+				FSinglePropertyParams SinglePropertyArgs;
+				const FName ParameterName = BlockParam->GetParameterName();
+				IDetailPropertyRow* DetailPropertyRow = DefaultValueCategory.AddExternalStructureProperty(StructOnScope, ParameterName);
+				if (DetailPropertyRow)
+				{
+					if (TSharedPtr<IPropertyHandle> Handle = DetailPropertyRow->GetPropertyHandle(); Handle.IsValid())
+					{
+						Handle->SetOnChildPropertyValuePreChange(FSimpleDelegate::CreateLambda([this, ReferencedBlock]()
+							{
+								ReferencedBlock->Modify(); // needed to enable the transaction when we modify the PropertyBag
+							}));
+						Handle->SetOnPropertyValueChanged(FSimpleDelegate::CreateLambda([this, ReferencedBlock, ParameterName]()
+							{
+								if (UAnimNextParameterBlock_EditorData* EditorData = Cast<UAnimNextParameterBlock_EditorData>(ReferencedBlock->EditorData))
+								{
+									if (UAnimNextParameterBlockEntry* BlockEntry = EditorData->FindBinding(ParameterName))
+									{
+										BlockEntry->MarkPackageDirty();
+										ReferencedBlock->MarkPackageDirty(); // show the ParameterBlock has changed
+									}
+								}
+							}));
+					}
+				}
+			}
+		}
+	}
+}
+
+void FParameterBlockParameterCustomization::CustomizeDetails(const TSharedPtr<IDetailLayoutBuilder>& DetailBuilder) 
+{
+	CustomizeDetails(*DetailBuilder);
+}
+
+FText FParameterBlockParameterCustomization::GetName() const
+{
+	return FText();
+}
+
+void FParameterBlockParameterCustomization::SetName(const FText& InNewText, ETextCommit::Type InCommitType)
+{
+}
+
+bool FParameterBlockParameterCustomization::OnVerifyNameChanged(const FText& InText, FText& OutErrorMessage)
+{
+	return true;
+}
+
+
+}
+
+#undef LOCTEXT_NAMESPACE
diff --git a/Engine/Plugins/Experimental/AnimNext/Source/AnimNextEditor/Private/Param/ParameterBlockParameterCustomization.h b/Engine/Plugins/Experimental/AnimNext/Source/AnimNextEditor/Private/Param/ParameterBlockParameterCustomization.h
new file mode 100644
--- /dev/null
+++ b/Engine/Plugins/Experimental/AnimNext/Source/AnimNextEditor/Private/Param/ParameterBlockParameterCustomization.h
@@ -0,0 +1,30 @@
+// Copyright Epic Games, Inc. All Rights Reserved.
+
+#pragma once
+
+#include "IDetailCustomization.h"
+#include "IPropertyTypeCustomization.h"
+#include "PropertyEditorModule.h"
+
+namespace UE::AnimNext::Editor
+{
+
+class FParameterBlockParameterCustomization : public IDetailCustomization
+{
+private:
+
+	/** Called when details should be customized */
+	virtual void CustomizeDetails(IDetailLayoutBuilder& DetailBuilder) override;
+
+	/**
+	 * Called when details should be customized, this version allows the customization to store a weak ptr to the layout builder.
+	 * This allows property changes to trigger a force refresh of the detail panel.
+	 */
+	virtual void CustomizeDetails(const TSharedPtr<IDetailLayoutBuilder>& DetailBuilder) override;
+
+	FText GetName() const;
+	void SetName(const FText& InNewText, ETextCommit::Type InCommitType);
+	bool OnVerifyNameChanged(const FText& InText, FText& OutErrorMessage);
+};
+
+}
diff --git a/Engine/Plugins/Experimental/AnimNext/Source/AnimNextEditor/Private/Param/SParameterBlockView.cpp b/Engine/Plugins/Experimental/AnimNext/Source/AnimNextEditor/Private/Param/SParameterBlockView.cpp
--- a/Engine/Plugins/Experimental/AnimNext/Source/AnimNextEditor/Private/Param/SParameterBlockView.cpp
+++ b/Engine/Plugins/Experimental/AnimNext/Source/AnimNextEditor/Private/Param/SParameterBlockView.cpp
@@ -32,6 +32,10 @@
 #include "Param/AnimNextParameterBlockParameter.h"
 #include "Param/AnimNextParameterBlockGraph.h"
 #include "Framework/Application/SlateApplication.h"
+#include "PropertyEditorModule.h"
+#include "ISinglePropertyView.h"
+#include "IStructureDataProvider.h"
+#include "Misc/NotifyHook.h"
 
 #define LOCTEXT_NAMESPACE "AnimNextParameterBlockView"
 
@@ -46,19 +50,113 @@ static FName Column_RevisionControl(TEXT("RevisionControl"));
 static FName Column_ModifiedStatus(TEXT("ModifiedStatus"));
 static FName Column_Name(TEXT("Name"));
 static FName Column_Type(TEXT("Type"));
+static FName Column_Container(TEXT("Container"));
 static FName Column_Value(TEXT("Value"));
 
 }
 
+// An entry category displayed in the parameters table
+template<typename ItemType>
+class SCategoryHeaderTableRow : public STableRow<ItemType>
+{
+public:
+	SLATE_BEGIN_ARGS(SCategoryHeaderTableRow)
+		{}
+		SLATE_DEFAULT_SLOT(typename SCategoryHeaderTableRow::FArguments, Content)
+	SLATE_END_ARGS()
+
+	void Construct(const FArguments& InArgs, const TSharedRef<STableViewBase>& InOwnerTableView)
+	{
+		STableRow<ItemType>::ChildSlot
+			.Padding(0.0f, 2.0f, .0f, 0.0f)
+			[
+				SAssignNew(ContentBorder, SBorder)
+					.BorderImage(this, &SCategoryHeaderTableRow::GetBackgroundImage)
+					.Padding(FMargin(3.0f, 5.0f))
+					[
+						SNew(SHorizontalBox)
+							+ SHorizontalBox::Slot()
+							.VAlign(VAlign_Center)
+							.Padding(5.0f)
+							.AutoWidth()
+							[
+								SNew(SExpanderArrow, STableRow< ItemType >::SharedThis(this))
+							]
+							+ SHorizontalBox::Slot()
+							.VAlign(VAlign_Center)
+							.AutoWidth()
+							[
+								InArgs._Content.Widget
+							]
+					]
+			];
+
+		STableRow < ItemType >::ConstructInternal(
+			typename STableRow< ItemType >::FArguments()
+			.Style(FAppStyle::Get(), "DetailsView.TreeView.TableRow")
+			.ShowSelection(false),
+			InOwnerTableView
+		);
+	}
+
+	const FSlateBrush* GetBackgroundImage() const
+	{
+		if (STableRow<ItemType>::IsHovered())
+		{
+			return FAppStyle::Get().GetBrush("Brushes.Secondary");
+		}
+		else
+		{
+			return FAppStyle::Get().GetBrush("Brushes.Header");
+		}
+	}
+
+	virtual void SetContent(TSharedRef< SWidget > InContent) override
+	{
+		ContentBorder->SetContent(InContent);
+	}
+
+	virtual void SetRowContent(TSharedRef< SWidget > InContent) override
+	{
+		ContentBorder->SetContent(InContent);
+	}
+
+	virtual const FSlateBrush* GetBorder() const
+	{
+		return nullptr;
+	}
+
+	FReply OnMouseButtonDown(const FGeometry& MyGeometry, const FPointerEvent& MouseEvent)
+	{
+		if (MouseEvent.GetEffectingButton() == EKeys::LeftMouseButton)
+		{
+			STableRow<ItemType>::ToggleExpansion();
+			return FReply::Handled();
+		}
+		else
+		{
+			return FReply::Unhandled();
+		}
+	}
+private:
+	TSharedPtr<SBorder> ContentBorder;
+};
+
 // An entry displayed in the parameters view
 struct FParameterBlockViewEntry
 {
+
 	FParameterBlockViewEntry() = default;
 
 	FParameterBlockViewEntry(UAnimNextParameterBlockEntry* InEntry)
 		: WeakEntry(InEntry)
 	{}
 
+	FParameterBlockViewEntry(EParameterBlockCategoryType InCategoryType, const FText& InCategoryName)
+		: CategoryType(InCategoryType)
+		, CategoryName(InCategoryName)
+	{}
+
 	bool PassesFilter(const FString& InFilterText) const
 	{
 		if(UAnimNextParameterBlockEntry* Entry = WeakEntry.Get())
@@ -69,12 +167,32 @@ struct FParameterBlockViewEntry
 		return false;
 	}
 
+	void ResetChildren(int32 NewSize = 0)
+	{
+		Children.Reset(NewSize);
+	}
+
+	void GetChildrenRecursive(TArray< TSharedRef<FParameterBlockViewEntry> >& OutChildren)
+	{
+		for (TSharedRef<FParameterBlockViewEntry>& Entry : Children)
+		{
+			OutChildren.Add(Entry);
+			Entry->GetChildrenRecursive(OutChildren);
+		}
+	}
+
 	// Ptr to the underlying entry
 	TWeakObjectPtr<UAnimNextParameterBlockEntry> WeakEntry;
 
 	// Widget used to rename items
 	TWeakPtr<SInlineEditableTextBlock> NameWidget;
 
+	EParameterBlockCategoryType CategoryType = EParameterBlockCategoryType::Invalid;
+	FText CategoryName;
+
+	// Children when entry is a category
+	TArray< TSharedRef<FParameterBlockViewEntry> > Children;
+
 	// Flag to indicate a rename was requested
 	bool bRenameWhenScrolledIntoView = false;
 };
@@ -87,6 +205,12 @@ void SParameterBlockView::Construct(const FArguments& InArgs, UAnimNextParameter
 
 	EditorData = InEditorData;
 
+	Categories = {
+		MakeShared<FParameterBlockViewEntry>(EParameterBlockCategoryType::Parameter, LOCTEXT("ParameterBlockParameterCategory", "Parameters")),
+		MakeShared<FParameterBlockViewEntry>(EParameterBlockCategoryType::BuiltIn, LOCTEXT("ParameterBlockBindingCategory", "Built-In")),
+		MakeShared<FParameterBlockViewEntry>(EParameterBlockCategoryType::Graph, LOCTEXT("ParameterBlockGraphCategory", "Graphs"))
+	};
+
 	// Cache asset data for block for comparisons/filtering
 	BlockAssetData = FAssetData(UncookedOnly::FUtils::GetBlock(EditorData));
 
@@ -103,88 +227,7 @@ void SParameterBlockView::Construct(const FArguments& InArgs, UAnimNextParameter
 		.AutoHeight()
 		[
 			SNew(SHorizontalBox)
-			+SHorizontalBox::Slot()
-			.AutoWidth()
-			.Padding(2.0f)
-			[
-				SNew(SSimpleComboButton)
-				.Text(LOCTEXT("AddParameterButton", "Add Parameter"))
-				.HasDownArrow(true)
-				.OnGetMenuContent_Lambda([this]()
-				{
-					FParameterPickerArgs Args;
-					Args.bMultiSelect = false;
-					Args.bShowLibraries = false;
-					Args.bShowBlocks = false;
-					Args.OnParameterPicked = FOnParameterPicked::CreateLambda([this](const FParameterBindingReference& InParameterBinding)
-					{
-						FSlateApplication::Get().DismissAllMenus();
-
-						FScopedTransaction Transaction(LOCTEXT("AddParameter", "Add parameter"));
-						PendingSelection.Empty();
-
-						// Create a new entry for the parameter
-						UAnimNextParameterLibrary* Library = Cast<UAnimNextParameterLibrary>(InParameterBinding.Library.GetAsset());
-						UAnimNextParameterBlockParameter* Parameter = EditorData->AddParameter(InParameterBinding.Parameter, Library);
-						PendingSelection.Add(Parameter);
-					});
-					return SNew(SParameterPicker)
-						.Args(Args);
-				})
-			]
-			+SHorizontalBox::Slot()
-			.AutoWidth()
-			.Padding(2.0f)
-			[
-				SNew(SSimpleButton)
-				.Text(LOCTEXT("AddNewParameterButton", "Add New Parameter"))
-				.Icon(FAppStyle::Get().GetBrush("Icons.Plus"))
-				.OnClicked_Lambda([this]()
-				{
-					TSharedRef<SAddParametersDialog> AddParametersDialog = SNew(SAddParametersDialog);
-					TArray<FParameterToAdd> ParametersToAdd;
-					if(AddParametersDialog->ShowModal(ParametersToAdd))
-					{
-						FScopedTransaction Transaction(FText::FormatOrdered(LOCTEXT("AddParameters", "Add {0}|plural(one=parameter,other=parameters)"), ParametersToAdd.Num()));
-
-						PendingSelection.Empty();
-
-						for(const FParameterToAdd& ParameterToAdd : ParametersToAdd)
-						{
-							// Create a new parameter in the supplied library
-							UAnimNextParameterLibrary* Library = Cast<UAnimNextParameterLibrary>(ParameterToAdd.Library.GetAsset());
-							UAnimNextParameter* NewParameter = Library->AddParameter(ParameterToAdd.Name, ParameterToAdd.Type);
-
-							// Create a new entry for the parameter
-							UAnimNextParameterBlockParameter* Parameter = EditorData->AddParameter(NewParameter->GetFName(), Library);
-							PendingSelection.Add(Parameter);
-						}
-					}
-
-					return FReply::Handled();
-				})
-			]
-			+SHorizontalBox::Slot()
-			.AutoWidth()
-			.Padding(2.0f)
-			[
-				SNew(SSimpleButton)
-				.Text(LOCTEXT("AddNewGraphButton", "Graph"))
-				.Icon(FAppStyle::Get().GetBrush("Icons.Plus"))
-				.OnClicked_Lambda([this]()
-				{
-					FScopedTransaction Transaction(LOCTEXT("AddGraph", "Add Graph"));
-
-					PendingSelection.Empty();
-
-					// Create a new entry for the graph
-					UAnimNextParameterBlockGraph* Graph = EditorData->AddGraph(TEXT("NewGraph"));
-					PendingSelection.Add(Graph);
-
-					return FReply::Handled();
-				})
-			]
-			+SHorizontalBox::Slot()
+			+ SHorizontalBox::Slot()
 			.FillWidth(1.0f)
 			.Padding(2.0f)
 			[
@@ -258,6 +301,11 @@ void SParameterBlockView::Construct(const FArguments& InArgs, UAnimNextParameter
 				.DefaultLabel(LOCTEXT("ValueColumnHeader", "Value"))
 				.ToolTipText(LOCTEXT("ValueColumnHeaderTooltip", "The value or binding to the parameter"))
 				.FillWidth(10.0f)
+
+				+ SHeaderRow::Column(Column_Container)
+				.DefaultLabel(LOCTEXT("ContainerColumnHeader", "Owner Block"))
+				.ToolTipText(LOCTEXT("ContainerColumnHeaderTooltip", "The block containing the parameter"))
+				.FillWidth(10.0f)
 			)
 		]
 	];
@@ -285,7 +333,12 @@ void SParameterBlockView::Construct(const FArguments& InArgs, UAnimNextParameter
 			return EActiveTimerReturnType::Continue;
 		}));
 
-	RequestRefresh();
+	RefreshEntries();
+
+	for (TSharedRef<FParameterBlockViewEntry>& Category : Categories)
+	{
+		EntriesList->SetItemExpansion(Category, true);
+	}
 }
 
 FReply SParameterBlockView::OnKeyDown(const FGeometry& MyGeometry, const FKeyEvent& InKeyEvent)
@@ -302,23 +355,78 @@ void SParameterBlockView::RequestRefresh()
 	bRefreshRequested = true;
 }
 
+template<typename ItemType, typename ComparisonType>
+void RestoreExpansionState(TSharedPtr< STreeView<ItemType> > InTree, const TArray<ItemType>& ItemSource, const TSet<ItemType>& OldExpansionState, ComparisonType ComparisonFunction)
+{
+	check(InTree.IsValid());
+
+	// Iterate over new tree items
+	for (int32 ItemIdx = 0; ItemIdx < ItemSource.Num(); ItemIdx++)
+	{
+		ItemType NewItem = ItemSource[ItemIdx];
+
+		// Look through old expansion state
+		for (typename TSet<ItemType>::TConstIterator OldExpansionIter(OldExpansionState); OldExpansionIter; ++OldExpansionIter)
+		{
+			const ItemType OldItem = *OldExpansionIter;
+			// See if this matches this new item
+			if (ComparisonFunction(OldItem, NewItem))
+			{
+				// It does, so expand it
+				InTree->SetItemExpansion(NewItem, true);
+			}
+		}
+	}
+}
+
+static bool CompareGraphActionNode(TSharedRef<FParameterBlockViewEntry> A, TSharedRef<FParameterBlockViewEntry> B)
+{
+	if (A->CategoryType != B->CategoryType)
+	{
+		return false;
+	}
+
+	// First check grouping is the same
+	if (A->CategoryName.ToString() != B->CategoryName.ToString())
+	{
+		return false;
+	}
+
+	return true;
+}
+
 void SParameterBlockView::RefreshEntries()
 {
-	Entries.Empty();
+	// First, save off current expansion state
+	TSet< TSharedRef<FParameterBlockViewEntry> > OldExpansionState;
+	EntriesList->GetExpandedItems(OldExpansionState);
+
+	// Add empty categories, the filtering will add the elements to the correct one
+	Entries.Reset();
 
 	TArray<TSharedRef<FParameterBlockViewEntry>> EntriesToSelect;
 
-	for(UAnimNextParameterBlockEntry* Entry : EditorData->Entries)
+	for (UAnimNextParameterBlockEntry* Entry : EditorData->Entries)
 	{
 		TSharedRef<FParameterBlockViewEntry> NewEntry = MakeShared<FParameterBlockViewEntry>(Entry);
+
 		Entries.Add(NewEntry);
 
-		if(PendingSelection.Contains(Entry))
+		if (PendingSelection.Contains(Entry))
 		{
 			EntriesToSelect.Add(NewEntry);
 		}
 	}
 
+	// Restore the expanded items
+	TArray< TSharedRef<FParameterBlockViewEntry> > AllEntries;
+	for (const TSharedRef<FParameterBlockViewEntry> & Entry : Entries)
+	{
+		AllEntries.Add(Entry);
+		Entry->GetChildrenRecursive(AllEntries);
+	}
+	RestoreExpansionState< TSharedRef<FParameterBlockViewEntry> >(EntriesList, AllEntries, OldExpansionState, CompareGraphActionNode);
+
 	PendingSelection.Empty();
 
 	RefreshFilter();
@@ -331,14 +439,47 @@ void SParameterBlockView::RefreshEntries()
 
 void SParameterBlockView::RefreshFilter()
 {
-	FilteredEntries.Empty();
+	FilteredEntries = Categories;
+
+	for (auto& Category : Categories)
+	{
+		Category->ResetChildren();
+	}
+
+	TSharedRef<FParameterBlockViewEntry> ParametersCategory = GetCategory(EParameterBlockCategoryType::Parameter);
+	TSharedRef<FParameterBlockViewEntry> BindingsCategory = GetCategory(EParameterBlockCategoryType::BuiltIn);
+	TSharedRef<FParameterBlockViewEntry> GraphsCategory = GetCategory(EParameterBlockCategoryType::Graph);
 
+	check(ParametersCategory->CategoryType == EParameterBlockCategoryType::Parameter);
+	check(BindingsCategory->CategoryType == EParameterBlockCategoryType::BuiltIn);
+	check(GraphsCategory->CategoryType == EParameterBlockCategoryType::Graph);
+
+	// add the entries as categories children
 	const FString FilterTextAsString = FilterText.ToString();
-	for(const TSharedRef<FParameterBlockViewEntry>& Entry : Entries)
+	for (const TSharedRef<FParameterBlockViewEntry>& Entry : Entries)
 	{
-		if(Entry->PassesFilter(FilterTextAsString))
+		if (Entry->PassesFilter(FilterTextAsString))
 		{
-			FilteredEntries.Add(Entry);
+			const UAnimNextParameterBlockEntry* BlockEntry = Entry->WeakEntry.Get();
+			if (const IAnimNextParameterBlockParameterInterface* Parameter = Cast<IAnimNextParameterBlockParameterInterface>(BlockEntry))
+			{
+				UAnimNextParameterBlock* ReferencedBlock = UE::AnimNext::UncookedOnly::FUtils::GetBlock(EditorData);
+				FInstancedPropertyBag* PropertyBag = UE::AnimNext::UncookedOnly::FUtils::GetPropertyBag(ReferencedBlock);
+
+				const FName ParameterName = Parameter->GetParameterName();
+				if (PropertyBag->FindPropertyDescByName(ParameterName))
+				{
+					ParametersCategory->Children.Add(Entry);
+				}
+				else
+				{
+					BindingsCategory->Children.Add(Entry);
+				}
+			}
+			else if (Cast<IAnimNextParameterBlockGraphInterface>(BlockEntry))
+			{
+				GraphsCategory->Children.Add(Entry);
+			}
 		}
 	}
 
@@ -387,9 +528,8 @@ TSharedRef<SWidget> SParameterBlockView::HandleGetContextContent()
 
 void SParameterBlockView::HandleDelete()
 {
-	if(EntriesList->GetNumItemsSelected() > 0)
+	auto DeleteSelected = [this](TArray<TSharedRef<FParameterBlockViewEntry>> SelectedItems)
 	{
-		TArray<TSharedRef<FParameterBlockViewEntry>> SelectedItems = EntriesList->GetSelectedItems();
 		TArray<UAnimNextParameterBlockEntry*> EntriesToRemove;
 		Algo::Transform(SelectedItems, EntriesToRemove, [](const TSharedRef<FParameterBlockViewEntry>& InEntry) { return InEntry->WeakEntry.Get(); });
 
@@ -399,16 +539,27 @@ void SParameterBlockView::HandleDelete()
 			FScopedTransaction Transaction(FText::FormatOrdered(LOCTEXT("DeleteParameterBlockEntry", "Delete parameter block {0}|plural(one=entry,other=entries)"), EntriesList->GetNumItemsSelected()));
 			EditorData->RemoveEntries(EntriesToRemove);
 		}
+	};
+
+	if (EntriesList->GetNumItemsSelected() > 0)
+	{
+		TArray<TSharedRef<FParameterBlockViewEntry>> SelectedItems = EntriesList->GetSelectedItems();
+		DeleteSelected(SelectedItems);
 	}
 }
 
 void SParameterBlockView::HandleRename()
 {
+	auto RenameSelected = [this](TSharedPtr<STreeView<TSharedRef<FParameterBlockViewEntry>>>& EntriesList, const TSharedRef<FParameterBlockViewEntry>& SelectedItem)
+	{
+		SelectedItem->bRenameWhenScrolledIntoView = true;
+		EntriesList->RequestScrollIntoView(SelectedItem);
+	};
+
 	if(EntriesList->GetNumItemsSelected() == 1)
 	{
 		TArray<TSharedRef<FParameterBlockViewEntry>> SelectedItems = EntriesList->GetSelectedItems();
-		SelectedItems[0]->bRenameWhenScrolledIntoView = true;
-		EntriesList->RequestScrollIntoView(SelectedItems[0]);
+		RenameSelected(EntriesList, SelectedItems[0]);
 	}
 }
 
@@ -435,7 +586,7 @@ void SParameterBlockView::HandleItemScrolledIntoView(TSharedRef<FParameterBlockV
 	}
 }
 
-class SParameterBlockViewRow : public SMultiColumnTableRow<TSharedRef<FParameterBlockViewEntry>>
+class SParameterBlockViewRow : public SMultiColumnTableRow<TSharedRef<FParameterBlockViewEntry>>, private FNotifyHook
 {
 	SLATE_BEGIN_ARGS(SParameterBlockViewRow) {}
 	SLATE_END_ARGS()
@@ -462,6 +613,37 @@ class SParameterBlockViewRow : public SMultiColumnTableRow<TSharedRef<FParameter
 		return SMultiColumnTableRow<TSharedRef<FParameterBlockViewEntry>>::OnMouseButtonDoubleClick(InMyGeometry, InMouseEvent);
 	}
 	
+	//~ Begin FNotifyHook Interface
+	virtual void NotifyPreChange(FProperty* PropertyAboutToChange) override
+	{
+		if (TSharedPtr<SParameterBlockView> View = WeakView.Pin())
+		{
+			if (UAnimNextParameterBlock_EditorData* EditorData = View->EditorData)
+			{
+				if (UAnimNextParameterBlockEntry * BlockEntry = EditorData->FindBinding(PropertyAboutToChange->GetFName()))
+				{
+					UAnimNextParameterBlock* ReferencedBlock = UE::AnimNext::UncookedOnly::FUtils::GetBlock(EditorData);
+					ReferencedBlock->Modify(); // needed to enable the transaction when we modify the PropertyBag
+				}
+			}
+		}
+	}
+
+	virtual void NotifyPostChange(const FPropertyChangedEvent& PropertyChangedEvent, FProperty* PropertyThatChanged) override
+	{
+		if (TSharedPtr<SParameterBlockView> View = WeakView.Pin())
+		{
+			if (UAnimNextParameterBlock_EditorData* EditorData = View->EditorData)
+			{
+				if (UAnimNextParameterBlockEntry* BlockEntry = EditorData->FindBinding(PropertyChangedEvent.GetMemberPropertyName()))
+				{
+					BlockEntry->MarkPackageDirty(); // needed to show the changed sign a the table when we modify the PropertyBag
+				}
+			}
+		}
+	}
+	// End of FNotifyHook
+
 	virtual TSharedRef<SWidget> GenerateWidgetForColumn(const FName& InColumnName) override
 	{
 		using namespace ParameterBlockView;
@@ -621,13 +803,24 @@ class SParameterBlockViewRow : public SMultiColumnTableRow<TSharedRef<FParameter
 					.IsSelected(this, &SParameterBlockViewRow::IsSelectedExclusively)
 					.IsReadOnly_Lambda([this]()
 					{
-						return Cast<IAnimNextParameterBlockGraphInterface>(Entry->WeakEntry.Get()) == nullptr;
+						if (Cast<IAnimNextParameterBlockParameterInterface>(Entry->WeakEntry.Get()) 
+							|| Cast<IAnimNextParameterBlockGraphInterface>(Entry->WeakEntry.Get()))
+						{
+							return false;
+						}
+						return true;
 					})
 					.OnTextCommitted_Lambda([this](const FText& InNewText, ETextCommit::Type InCommitType)
 					{
 						if(InCommitType == ETextCommit::OnEnter)
 						{
-							if(IAnimNextParameterBlockGraphInterface* Graph = Cast<IAnimNextParameterBlockGraphInterface>(Entry->WeakEntry.Get()))
+							if (IAnimNextParameterBlockParameterInterface* Parameter = Cast<IAnimNextParameterBlockParameterInterface>(Entry->WeakEntry.Get()))
+							{
+								FScopedTransaction Transaction(LOCTEXT("SetParameterName", "Set Parameter name"));
+
+								Parameter->SetParameterName(*InNewText.ToString());
+							}
+							else if(IAnimNextParameterBlockGraphInterface* Graph = Cast<IAnimNextParameterBlockGraphInterface>(Entry->WeakEntry.Get()))
 							{
 								FScopedTransaction Transaction(LOCTEXT("SetGraphName", "Set graph name"));
 
@@ -684,6 +877,48 @@ class SParameterBlockViewRow : public SMultiColumnTableRow<TSharedRef<FParameter
 				];
 		}
 		else if(InColumnName == Column_Value)
+		{
+			TSharedRef< SWidget > ColumnWidget = SNullWidget::NullWidget;
+
+			if (TSharedPtr<SParameterBlockView> View = WeakView.Pin())
+			{
+				const IAnimNextParameterBlockParameterInterface* Binding = Cast<IAnimNextParameterBlockParameterInterface>(Entry->WeakEntry.Get());
+				if (Binding != nullptr)
+				{
+					if (UAnimNextParameterBlock_EditorData* EditorData = View->EditorData)
+					{
+						UAnimNextParameterBlock* ReferencedBlock = UE::AnimNext::UncookedOnly::FUtils::GetBlock(EditorData);
+
+						FInstancedPropertyBag &PropertyBag = ReferencedBlock->GetPropertyBag();
+						const FName ParameterName = Binding->GetParameterName();
+
+						if (PropertyBag.FindPropertyDescByName(ParameterName))
+						{
+							TSharedPtr<FStructOnScope> StructOnScope = MakeShareable(new FStructOnScope(PropertyBag.GetPropertyBagStruct(), (uint8*)PropertyBag.GetValue().GetMemory()));
+							FSinglePropertyParams SinglePropertyArgs;
+							SinglePropertyArgs.NamePlacement = EPropertyNamePlacement::Hidden;
+							SinglePropertyArgs.NotifyHook = this;
+
+							FPropertyEditorModule& PropertyEditorModule = FModuleManager::Get().LoadModuleChecked<FPropertyEditorModule>("PropertyEditor");
+							const TSharedPtr<ISinglePropertyView> SingleStructPropertyView = PropertyEditorModule.CreateSingleProperty(MakeShared<FStructOnScopeStructureDataProvider>(StructOnScope), ParameterName, SinglePropertyArgs);
+							if (SingleStructPropertyView.IsValid())
+							{
+								ColumnWidget = SingleStructPropertyView.ToSharedRef();
+							}
+						}
+					}
+				}
+			}
+
+			return SNew(SHorizontalBox)
+				+ SHorizontalBox::Slot()
+				.HAlign(HAlign_Center)
+				.AutoWidth()
+				[
+					ColumnWidget
+				];
+		}
+		else if (InColumnName == Column_Container)
 		{
 			TSharedPtr<SHorizontalBox> HorizontalBox;
 
@@ -695,34 +930,34 @@ class SParameterBlockViewRow : public SMultiColumnTableRow<TSharedRef<FParameter
 					SAssignNew(HorizontalBox, SHorizontalBox)
 				];
 
-			if(UAnimNextParameterBlockEntry* BlockEntry = Entry->WeakEntry.Get())
+			if (TSharedPtr<SParameterBlockView> View = WeakView.Pin())
 			{
-				if(IAnimNextParameterBlockReferenceInterface* ReferenceInterface = Cast<IAnimNextParameterBlockReferenceInterface>(BlockEntry))
+				if (UAnimNextParameterBlock_EditorData* EditorData = View->EditorData)
 				{
-					const UAnimNextParameterBlock* ReferencedBlock = ReferenceInterface->GetBlock();
+					const UAnimNextParameterBlock* ReferencedBlock = UE::AnimNext::UncookedOnly::FUtils::GetBlock(EditorData);
+
 					HorizontalBox->AddSlot()
 					[
 						SNew(STextBlock)
-						.Font(IDetailLayoutBuilder::GetDetailFont())
-						.IsEnabled(false)
-						.Text_Lambda([WeakBlock = TWeakObjectPtr<const UAnimNextParameterBlock>(ReferencedBlock)]()
-						{
-							if(const UAnimNextParameterBlock* Block = WeakBlock.Get())
-							{
-								return FText::FromString(Block->GetName());
-							}
+							.Font(IDetailLayoutBuilder::GetDetailFont())
+							.Text_Lambda([WeakBlock = TWeakObjectPtr<const UAnimNextParameterBlock>(ReferencedBlock)]()
+								{
+									if (const UAnimNextParameterBlock* Block = WeakBlock.Get())
+									{
+										return FText::FromString(Block->GetName());
+									}
 
-							return LOCTEXT("MissingParameterBlock", "Missing Parameter Block");
-						})
-						.ToolTipText_Lambda([WeakBlock = TWeakObjectPtr<const UAnimNextParameterBlock>(ReferencedBlock)]()
-						{
-							if(const UAnimNextParameterBlock* Block = WeakBlock.Get())
-							{
-								return FText::FromString(Block->GetPathName());
-							}
+									return LOCTEXT("MissingParameterBlock", "Missing Parameter Block");
+								})
+							.ToolTipText_Lambda([WeakBlock = TWeakObjectPtr<const UAnimNextParameterBlock>(ReferencedBlock)]()
+								{
+									if (const UAnimNextParameterBlock* Block = WeakBlock.Get())
+									{
+										return FText::FromString(Block->GetPathName());
+									}
 
-							return LOCTEXT("MissingParameterBlock", "Missing Parameter Block");
-						})
+									return LOCTEXT("MissingParameterBlock", "Missing Parameter Block");
+								})
 					];
 				}
 			}
@@ -739,12 +974,156 @@ class SParameterBlockViewRow : public SMultiColumnTableRow<TSharedRef<FParameter
 
 TSharedRef<ITableRow> SParameterBlockView::HandleGenerateRow(TSharedRef<FParameterBlockViewEntry> InEntry, const TSharedRef<STableViewBase>& InOwnerTable)
 {
-	return SNew(SParameterBlockViewRow, InOwnerTable, SharedThis(this), InEntry);
+	if (InEntry->CategoryType != EParameterBlockCategoryType::Invalid)
+	{
+		TSharedPtr< STableRow< TSharedPtr<FParameterBlockViewEntry> > > TableRow;
+		TableRow = SNew(SCategoryHeaderTableRow< TSharedPtr<FParameterBlockViewEntry> >, InOwnerTable);
+
+		TSharedPtr<SHorizontalBox> RowContainer;
+		TableRow->SetRowContent
+		(
+			SAssignNew(RowContainer, SHorizontalBox)
+		);
+
+		const FMargin RowPadding = FMargin(0, 2);
+		TSharedPtr<SHorizontalBox> RowContent = SNew(SHorizontalBox)
+			+ SHorizontalBox::Slot()
+			.VAlign(VAlign_Center)
+			[
+				SNew(STextBlock)
+				.Font(FCoreStyle::GetDefaultFontStyle("Bold", 9))
+				.Text(InEntry->CategoryName)
+			];
+
+		RowContainer->AddSlot()
+			.AutoWidth()
+			.VAlign(VAlign_Fill)
+			.HAlign(HAlign_Right)
+			[
+				SNew(SExpanderArrow, TableRow)
+					.BaseIndentLevel(1)
+			];
+
+		RowContainer->AddSlot()
+			.FillWidth(1.0)
+			.Padding(RowPadding)
+			[
+				RowContent.ToSharedRef()
+			];
+
+
+		if (InEntry->CategoryType == EParameterBlockCategoryType::Parameter)
+		{
+			RowContainer->AddSlot()
+			.AutoWidth()
+			.VAlign(VAlign_Fill)
+			.HAlign(HAlign_Right)
+			[
+				SNew(SSimpleComboButton)
+				.Text(LOCTEXT("AddParameterButton", "Add Parameter"))
+				.Icon(FAppStyle::Get().GetBrush("Icons.Plus"))
+				.HasDownArrow(true)
+				.OnGetMenuContent_Lambda([this]()
+				{
+					FParameterPickerArgs Args;
+					Args.bMultiSelect = false;
+					Args.bShowLibraries = false;
+					Args.bShowBlocks = false;
+					Args.bShowBoundParameters = false;
+					Args.bShowBuiltInParameters = false;
+					Args.OnParameterPicked = FOnParameterPicked::CreateLambda([this](const FParameterBindingReference& InParameterBinding)
+					{
+						FSlateApplication::Get().DismissAllMenus();
+
+						FScopedTransaction Transaction(LOCTEXT("AddParameter", "Add parameter"));
+						PendingSelection.Empty();
+
+						// Create a new entry for the parameter
+						UAnimNextParameterLibrary* Library = Cast<UAnimNextParameterLibrary>(InParameterBinding.Library.GetAsset());
+						UAnimNextParameterBlockParameter* Parameter = EditorData->AddParameter(InParameterBinding.Parameter, Library);
+						PendingSelection.Add(Parameter);
+					});
+					return SNew(SParameterPicker)
+						.Args(Args);
+				})
+			];
+		}
+		else if (InEntry->CategoryType == EParameterBlockCategoryType::BuiltIn)
+		{
+			RowContainer->AddSlot()
+			.AutoWidth()
+			.VAlign(VAlign_Fill)
+			.HAlign(HAlign_Right)
+			[
+				SNew(SSimpleComboButton)
+				.Text(LOCTEXT("AddBuiltInParameterButton", "Add Built-In Parameter"))
+				.Icon(FAppStyle::Get().GetBrush("Icons.Plus"))
+				.HasDownArrow(true)
+				.OnGetMenuContent_Lambda([this]()
+				{
+					FParameterPickerArgs Args;
+					Args.bMultiSelect = false;
+					Args.bShowLibraries = false;
+					Args.bShowBlocks = false;
+					Args.bShowBoundParameters = false;
+					Args.bShowUnboundParameters = false;
+					Args.bAllowNew = false;
+					Args.OnParameterPicked = FOnParameterPicked::CreateLambda([this](const FParameterBindingReference& InParameterBinding)
+					{
+						FSlateApplication::Get().DismissAllMenus();
+
+						FScopedTransaction Transaction(LOCTEXT("AddBuiltInParameter", "Add builtin parameter"));
+						PendingSelection.Empty();
+
+						// Create a new entry for the parameter
+						UAnimNextParameterLibrary* Library = Cast<UAnimNextParameterLibrary>(InParameterBinding.Library.GetAsset());
+						UAnimNextParameterBlockParameter* Parameter = EditorData->AddParameter(InParameterBinding.Parameter, Library);
+						PendingSelection.Add(Parameter);
+					});
+					return SNew(SParameterPicker)
+					.Args(Args);
+				})
+			];
+		}
+		else if (InEntry->CategoryType == EParameterBlockCategoryType::Graph)
+		{
+			RowContainer->AddSlot()
+			.AutoWidth()
+			.VAlign(VAlign_Fill)
+			.HAlign(HAlign_Right)
+			[
+				SNew(SSimpleButton)
+				.Text(LOCTEXT("AddGraphButton", "Add Graph"))
+				.Icon(FAppStyle::Get().GetBrush("Icons.Plus"))
+				.OnClicked_Lambda([this]()
+				{
+					FScopedTransaction Transaction(LOCTEXT("AddGraph", "Add Graph"));
+
+					PendingSelection.Empty();
+
+					// Create a new entry for the graph
+					UAnimNextParameterBlockGraph* Graph = EditorData->AddGraph(TEXT("NewGraph"));
+					PendingSelection.Add(Graph);
+
+					return FReply::Handled();
+				})
+			];
+		}
+
+		return TableRow.ToSharedRef();
+	}
+	else
+	{
+		return SNew(SParameterBlockViewRow, InOwnerTable, SharedThis(this), InEntry);
+	}
 }
 
 void SParameterBlockView::HandleGetChildren(TSharedRef<FParameterBlockViewEntry> InEntry, TArray<TSharedRef<FParameterBlockViewEntry>>& OutChildren)
 {
-
+	if (InEntry->CategoryType != EParameterBlockCategoryType::Invalid)
+	{
+		OutChildren = InEntry->Children;
+	}
 }
 
 void SParameterBlockView::HandleSelectionChanged(TSharedPtr<FParameterBlockViewEntry> InEntry, ESelectInfo::Type InSelectionType)
@@ -760,7 +1139,6 @@ void SParameterBlockView::HandleSelectionChanged(TSharedPtr<FParameterBlockViewE
 				SelectedItems.Add(Entry);
 			}
 		}
-
 		OnSelectionChangedDelegate.Execute(SelectedItems);
 	}
 }
@@ -779,6 +1157,21 @@ EFilterParameterResult SParameterBlockView::HandleFilterLinkedParameter(const FP
 	return EFilterParameterResult::Include;
 }
 
+TSharedRef<FParameterBlockViewEntry> SParameterBlockView::GetCategory(EParameterBlockCategoryType CategoryType)
+{
+	static TSharedRef<FParameterBlockViewEntry> InvalidCategory = MakeShared<FParameterBlockViewEntry>(EParameterBlockCategoryType::Invalid, LOCTEXT("ParameterBlockInvalidCategory", "Parameters"));
+	
+	for (TSharedRef<FParameterBlockViewEntry>& Category : Categories)
+	{
+		if (Category->CategoryType == CategoryType)
+		{
+			return Category;
+		}
+	}
+
+	return TSharedRef<FParameterBlockViewEntry>();
+}
+
 }
 
 #undef LOCTEXT_NAMESPACE
\ No newline at end of file
diff --git a/Engine/Plugins/Experimental/AnimNext/Source/AnimNextEditor/Private/Param/SParameterBlockView.h b/Engine/Plugins/Experimental/AnimNext/Source/AnimNextEditor/Private/Param/SParameterBlockView.h
--- a/Engine/Plugins/Experimental/AnimNext/Source/AnimNextEditor/Private/Param/SParameterBlockView.h
+++ b/Engine/Plugins/Experimental/AnimNext/Source/AnimNextEditor/Private/Param/SParameterBlockView.h
@@ -3,8 +3,9 @@
 #pragma once
 
 #include "Widgets/SCompoundWidget.h"
-#include "Widgets/Views/SListView.h"
+#include "Widgets/Views/STreeView.h"
 #include "AssetRegistry/AssetData.h"
+#include "Templates/SharedPointer.h"
 
 class UAnimNextParameterBlockBinding;
 class UAnimNextParameter;
@@ -21,6 +22,15 @@ struct FParameterBindingReference;
 
 struct FParameterBlockViewEntry; 
 
+enum class EParameterBlockCategoryType : uint8
+{
+	Parameter,
+	BuiltIn,
+	Graph,
+	// --- ---
+	Invalid
+};
+
 class SParameterBlockView : public SCompoundWidget
 {
 public:
@@ -80,11 +90,15 @@ private:
 	void HandleSelectionChanged(TSharedPtr<FParameterBlockViewEntry> InEntry, ESelectInfo::Type InSelectionType);
 
 	EFilterParameterResult HandleFilterLinkedParameter(const FParameterBindingReference& InParameterBinding);
+
+	TSharedRef<FParameterBlockViewEntry> GetCategory(EParameterBlockCategoryType CategoryType);
 	
 private:
 	friend class SParameterBlockViewRow;
 	
-	TSharedPtr<SListView<TSharedRef<FParameterBlockViewEntry>>> EntriesList;
+	TArray<TSharedRef<FParameterBlockViewEntry>> Categories;
+
+	TSharedPtr<STreeView<TSharedRef<FParameterBlockViewEntry>>> EntriesList;
 
 	TArray<TSharedRef<FParameterBlockViewEntry>> Entries;
 
diff --git a/Engine/Plugins/Experimental/AnimNext/Source/AnimNextEditor/Private/Param/SParameterPicker.cpp b/Engine/Plugins/Experimental/AnimNext/Source/AnimNextEditor/Private/Param/SParameterPicker.cpp
--- a/Engine/Plugins/Experimental/AnimNext/Source/AnimNextEditor/Private/Param/SParameterPicker.cpp
+++ b/Engine/Plugins/Experimental/AnimNext/Source/AnimNextEditor/Private/Param/SParameterPicker.cpp
@@ -22,6 +22,7 @@
 #include "Widgets/Input/SButton.h"
 #include "Framework/Application/SlateApplication.h"
 #include "ScopedTransaction.h"
+#include "SSimpleButton.h"
 
 #define LOCTEXT_NAMESPACE "SParameterPicker"
 
@@ -135,6 +136,41 @@ void SParameterPicker::Construct(const FArguments& InArgs)
 	}
 
 	TSharedPtr<SHeaderRow> HeaderRow;
+	TSharedRef<SWidget> AddNewParameterWidget = SNullWidget::NullWidget;
+
+	if(Args.bAllowNew)
+	{
+		SAssignNew(AddNewParameterWidget, SSimpleButton)
+		.Text(LOCTEXT("AddNewParameterButton", "New Parameter"))
+		.ToolTipText(LOCTEXT("AddColumnHeaderTooltip", "Add a new parameter at global scope"))
+		.Icon(FAppStyle::Get().GetBrush("Icons.Plus"))
+		.OnClicked_Lambda([this]()
+		{
+			FSlateApplication::Get().DismissAllMenus();
+			TSharedRef<SAddParametersDialog> AddParametersDialog =
+				SNew(SAddParametersDialog)
+				.AllowMultiple(false);
+			TArray<FParameterToAdd> ParametersToAdd;
+			if(AddParametersDialog->ShowModal(ParametersToAdd))
+			{
+				if(ParametersToAdd.Num() > 0)
+				{
+					FScopedTransaction Transaction(LOCTEXT("AddParameter", "Add parameter"));
+
+					// Create a new parameter in the supplied library
+					UAnimNextParameterLibrary* Library = Cast<UAnimNextParameterLibrary>(ParametersToAdd[0].Library.GetAsset());
+					UAnimNextParameter* NewParameter = Library->AddParameter(ParametersToAdd[0].Name, ParametersToAdd[0].Type);
+							
+					FParameterBindingReference Reference;
+					Reference.Parameter = ParametersToAdd[0].Name;
+					Reference.Library = ParametersToAdd[0].Library;
+					Args.OnParameterPicked.ExecuteIfBound(Reference);
+				}
+			}
+			return FReply::Handled();
+		});
+	}
+
 
 	ChildSlot
 	[
@@ -143,12 +179,23 @@ void SParameterPicker::Construct(const FArguments& InArgs)
 		.AutoHeight()
 		.Padding(2.0f)
 		[
-			SAssignNew(SearchBox, SSearchBox)
-			.OnTextChanged_Lambda([this](FText InText)
-			{
-				FilterText = InText;
-				RefreshFilter();
-			})
+			SNew(SHorizontalBox)
+			+ SHorizontalBox::Slot()
+			.MaxWidth(250.0f)
+			[
+				SAssignNew(SearchBox, SSearchBox)
+				.OnTextChanged_Lambda([this](FText InText)
+				{
+					FilterText = InText;
+					RefreshFilter();
+				})
+			]
+			+ SHorizontalBox::Slot()
+			.HAlign(HAlign_Center)
+			.AutoWidth()
+			[
+				AddNewParameterWidget
+			]
 		]
 		+SVerticalBox::Slot()
 		.FillHeight(1.0f)
@@ -208,57 +255,6 @@ void SParameterPicker::Construct(const FArguments& InArgs)
 			.FillWidth(0.33f));
 	}
 
-	if(Args.bAllowNew)
-	{
-		HeaderRow->AddColumn(
-			SHeaderRow::Column(Column_New)
-			.DefaultLabel(FText::GetEmpty())
-			.HeaderContentPadding(FMargin(0.0f))
-			.FixedWidth(24.0f)
-			.HeaderContent()
-			[
-				SNew(SButton)
-				.ToolTipText(LOCTEXT("AddColumnHeaderTooltip", "Add a new parameter at global scope"))
-				.ButtonStyle(FAppStyle::Get(), "HoverHintOnly")
-				.OnClicked_Lambda([this]()
-				{
-					FSlateApplication::Get().DismissAllMenus();
-					TSharedRef<SAddParametersDialog> AddParametersDialog =
-						SNew(SAddParametersDialog)
-						.AllowMultiple(false);
-					TArray<FParameterToAdd> ParametersToAdd;
-					if(AddParametersDialog->ShowModal(ParametersToAdd))
-					{
-						if(ParametersToAdd.Num() > 0)
-						{
-							FScopedTransaction Transaction(LOCTEXT("AddParameter", "Add parameter"));
-
-							// Create a new parameter in the supplied library
-							UAnimNextParameterLibrary* Library = Cast<UAnimNextParameterLibrary>(ParametersToAdd[0].Library.GetAsset());
-							UAnimNextParameter* NewParameter = Library->AddParameter(ParametersToAdd[0].Name, ParametersToAdd[0].Type);
-							
-							FParameterBindingReference Reference;
-							Reference.Parameter = ParametersToAdd[0].Name;
-							Reference.Library = ParametersToAdd[0].Library;
-							Args.OnParameterPicked.ExecuteIfBound(Reference);
-						}
-					}
-					return FReply::Handled();
-				})
-				[
-					SNew(SBox)
-					.WidthOverride(16.0f)
-					.HeightOverride(16.0f)
-					.VAlign(VAlign_Center)
-					.HAlign(HAlign_Center)
-					[
-						SNew(SImage)
-						.ColorAndOpacity(FSlateColor::UseForeground())
-						.Image(FAppStyle::GetBrush("Icons.Plus"))
-					]
-				]
-			]);
-	}
 	
 	RefreshEntries();
 }
@@ -288,7 +284,6 @@ void SParameterPicker::RefreshEntries()
 	};
 	
 	// Find all blocks and their bound parameters
-	if(Args.bShowBoundParameters)
 	{
 		ARFilter.ClassPaths = { UAnimNextParameterBlock::StaticClass()->GetClassPathName() };
 		
@@ -315,8 +310,11 @@ void SParameterPicker::RefreshEntries()
 							FAnimNextParamType ParamType = FUtils::GetParameterTypeFromLibraryExports(Export.Name, LibraryExports);
 							if(!Args.OnFilterParameterType.IsBound() || Args.OnFilterParameterType.Execute(ParamType) == EFilterParameterResult::Include)
 							{
-								TSharedRef<FParameterPickerEntry> NewEntry = MakeShared<FParameterPickerEntry>(NewReference, ParamType);
-								Entries.Add(NewEntry);
+								if (Args.bShowBoundParameters)
+								{
+									TSharedRef<FParameterPickerEntry> NewEntry = MakeShared<FParameterPickerEntry>(NewReference, ParamType);
+									Entries.Add(NewEntry);
+								}
 							}
 						}
 					}
diff --git a/Engine/Plugins/Experimental/AnimNext/Source/AnimNextUncookedOnly/Internal/Param/AnimNextParameterBlockParameter.h b/Engine/Plugins/Experimental/AnimNext/Source/AnimNextUncookedOnly/Internal/Param/AnimNextParameterBlockParameter.h
--- a/Engine/Plugins/Experimental/AnimNext/Source/AnimNextUncookedOnly/Internal/Param/AnimNextParameterBlockParameter.h
+++ b/Engine/Plugins/Experimental/AnimNext/Source/AnimNextUncookedOnly/Internal/Param/AnimNextParameterBlockParameter.h
@@ -1,4 +1,4 @@
-﻿// Copyright Epic Games, Inc. All Rights Reserved.
+// Copyright Epic Games, Inc. All Rights Reserved.
 
 #pragma once
 
@@ -10,12 +10,18 @@
 class UAnimNextParameterLibrary;
 class UAnimNextParameterBlock_EditorData;
 
-UCLASS()
+namespace UE::AnimNext::Editor
+{
+	class FParameterBlockParameterCustomization;
+}
+
+UCLASS(MinimalAPI)
 class UAnimNextParameterBlockParameter : public UAnimNextParameterBlockEntry, public IAnimNextParameterBlockParameterInterface
 {
 	GENERATED_BODY()
 
 	friend class UAnimNextParameterBlock_EditorData;
+	friend class UE::AnimNext::Editor::FParameterBlockParameterCustomization;
 	
 	// UAnimNextParameterBlockEntry interface
 	virtual FText GetDisplayName() const override;
diff --git a/Engine/Plugins/Experimental/AnimNext/Source/AnimNextUncookedOnly/Internal/Param/AnimNextParameterBlock_EditorData.h b/Engine/Plugins/Experimental/AnimNext/Source/AnimNextUncookedOnly/Internal/Param/AnimNextParameterBlock_EditorData.h
--- a/Engine/Plugins/Experimental/AnimNext/Source/AnimNextUncookedOnly/Internal/Param/AnimNextParameterBlock_EditorData.h
+++ b/Engine/Plugins/Experimental/AnimNext/Source/AnimNextUncookedOnly/Internal/Param/AnimNextParameterBlock_EditorData.h
@@ -31,6 +31,8 @@ namespace UE::AnimNext::Editor
 	class SParameterBlockView;
 	class SParameterPicker;
 	class FParameterBlockTabSummoner;
+	class SParameterBlockViewRow;
+	class FParameterBlockParameterCustomization;
 }
 
 namespace UE::AnimNext::UncookedOnly
@@ -107,6 +109,8 @@ class UAnimNextParameterBlock_EditorData : public UObject, public IRigVMClientHo
 	friend class UE::AnimNext::Editor::SParameterPicker;
 	friend class UE::AnimNext::Editor::FParameterBlockTabSummoner;
 	friend class FAnimationAnimNextParametersEditorTest_Block;
+	friend class UE::AnimNext::Editor::SParameterBlockViewRow;
+	friend class UE::AnimNext::Editor::FParameterBlockParameterCustomization;
 
 	ANIMNEXTUNCOOKEDONLY_API UAnimNextParameterBlockParameter* AddParameter(FName InName, UAnimNextParameterLibrary* InLibrary, bool bSetupUndoRedo = true, bool bPrintPythonCommand = true);
 
@@ -178,7 +182,7 @@ class UAnimNextParameterBlock_EditorData : public UObject, public IRigVMClientHo
 
 	void CreateEdGraphForCollapseNode(URigVMCollapseNode* InNode);
 
-	void BroadcastModified();
+	ANIMNEXTUNCOOKEDONLY_API void BroadcastModified();
 
 	void ReportError(const TCHAR* InMessage) const;
 
diff --git a/Engine/Plugins/Experimental/AnimNext/Source/AnimNextUncookedOnly/Internal/UncookedOnlyUtils.h b/Engine/Plugins/Experimental/AnimNext/Source/AnimNextUncookedOnly/Internal/UncookedOnlyUtils.h
--- a/Engine/Plugins/Experimental/AnimNext/Source/AnimNextUncookedOnly/Internal/UncookedOnlyUtils.h
+++ b/Engine/Plugins/Experimental/AnimNext/Source/AnimNextUncookedOnly/Internal/UncookedOnlyUtils.h
@@ -49,6 +49,8 @@ struct ANIMNEXTUNCOOKEDONLY_API FUtils
 
 	static UAnimNextParameterBlock* GetBlock(const UAnimNextParameterBlock_EditorData* InEditorData);
 
+	static FInstancedPropertyBag* GetPropertyBag(UAnimNextParameterBlock* ReferencedBlock);
+
 	static void RecreateVM(UAnimNextParameterBlock* InParameterBlock);
 
 	/**
diff --git a/Engine/Plugins/Experimental/AnimNext/Source/AnimNextUncookedOnly/Private/Param/AnimNextParameterBlockParameter.cpp b/Engine/Plugins/Experimental/AnimNext/Source/AnimNextUncookedOnly/Private/Param/AnimNextParameterBlockParameter.cpp
--- a/Engine/Plugins/Experimental/AnimNext/Source/AnimNextUncookedOnly/Private/Param/AnimNextParameterBlockParameter.cpp
+++ b/Engine/Plugins/Experimental/AnimNext/Source/AnimNextUncookedOnly/Private/Param/AnimNextParameterBlockParameter.cpp
@@ -1,4 +1,4 @@
-﻿// Copyright Epic Games, Inc. All Rights Reserved.
+// Copyright Epic Games, Inc. All Rights Reserved.
 
 #include "Param/AnimNextParameterBlockParameter.h"
 #include "Param/AnimNextParameter.h"
@@ -40,7 +40,17 @@ void UAnimNextParameterBlockParameter::SetParameterName(FName InName, bool bSetu
 		Modify();
 	}
 
-	ParameterName = InName;
+	if (Library)
+	{
+		UAnimNextParameter* Parameter = Library->FindParameter(ParameterName);
+		if (Parameter)
+		{
+			if (Parameter->Rename(*InName.ToString(), Parameter->GetOuter()))
+			{
+				ParameterName = InName;
+			}
+		}
+	}
 
 	BroadcastModified();
 }
diff --git a/Engine/Plugins/Experimental/AnimNext/Source/AnimNextUncookedOnly/Private/Param/AnimNextParameterBlock_EditorData.cpp b/Engine/Plugins/Experimental/AnimNext/Source/AnimNextUncookedOnly/Private/Param/AnimNextParameterBlock_EditorData.cpp
--- a/Engine/Plugins/Experimental/AnimNext/Source/AnimNextUncookedOnly/Private/Param/AnimNextParameterBlock_EditorData.cpp
+++ b/Engine/Plugins/Experimental/AnimNext/Source/AnimNextUncookedOnly/Private/Param/AnimNextParameterBlock_EditorData.cpp
@@ -157,7 +157,7 @@ UAnimNextParameterBlockGraph* UAnimNextParameterBlock_EditorData::AddGraph(FName
 	int32 NameNumber = InName.GetNumber() + 1;
 	while(bAlreadyExists)
 	{
-		NewGraphName = FName(InName, NameNumber);
+		NewGraphName = FName(InName, NameNumber++);
 		bAlreadyExists =  Entries.ContainsByPredicate(DuplicateNamePredicate);
 	}
 
diff --git a/Engine/Plugins/Experimental/AnimNext/Source/AnimNextUncookedOnly/Private/UncookedOnlyUtils.cpp b/Engine/Plugins/Experimental/AnimNext/Source/AnimNextUncookedOnly/Private/UncookedOnlyUtils.cpp
--- a/Engine/Plugins/Experimental/AnimNext/Source/AnimNextUncookedOnly/Private/UncookedOnlyUtils.cpp
+++ b/Engine/Plugins/Experimental/AnimNext/Source/AnimNextUncookedOnly/Private/UncookedOnlyUtils.cpp
@@ -788,6 +788,13 @@ UAnimNextParameterBlock* FUtils::GetBlock(const UAnimNextParameterBlock_EditorDa
 	return CastChecked<UAnimNextParameterBlock>(InEditorData->GetOuter());
 }
 
+FInstancedPropertyBag* FUtils::GetPropertyBag(UAnimNextParameterBlock* ReferencedBlock)
+{
+	FInstancedPropertyBag* InstancedPropertyBag =&ReferencedBlock->PropertyBag;
+
+	return InstancedPropertyBag;
+}
+
 FParamTypeHandle FUtils::GetParamTypeHandleFromPinType(const FEdGraphPinType& InPinType)
 {
 	return GetParamTypeFromPinType(InPinType).GetHandle();
