Fix CIS Mac compile issue

#rnx

diff --git a/Engine/Plugins/Experimental/AnimNext/Source/AnimNextEditor/Private/Param/SParameterBlockView.cpp b/Engine/Plugins/Experimental/AnimNext/Source/AnimNextEditor/Private/Param/SParameterBlockView.cpp
--- a/Engine/Plugins/Experimental/AnimNext/Source/AnimNextEditor/Private/Param/SParameterBlockView.cpp
+++ b/Engine/Plugins/Experimental/AnimNext/Source/AnimNextEditor/Private/Param/SParameterBlockView.cpp
@@ -550,10 +550,10 @@ void SParameterBlockView::HandleDelete()
 
 void SParameterBlockView::HandleRename()
 {
-	auto RenameSelected = [this](TSharedPtr<STreeView<TSharedRef<FParameterBlockViewEntry>>>& EntriesList, const TSharedRef<FParameterBlockViewEntry>& SelectedItem)
+	auto RenameSelected = [this](TSharedPtr<STreeView<TSharedRef<FParameterBlockViewEntry>>>& InEntriesList, const TSharedRef<FParameterBlockViewEntry>& SelectedItem)
 	{
 		SelectedItem->bRenameWhenScrolledIntoView = true;
-		EntriesList->RequestScrollIntoView(SelectedItem);
+		InEntriesList->RequestScrollIntoView(SelectedItem);
 	};
 
 	if(EntriesList->GetNumItemsSelected() == 1)
