AnimNext: Display correct icon on asset document tabs

diff --git a/Engine/Plugins/Experimental/AnimNext/Source/AnimNextEditor/Private/Workspace/AssetDocumentSummoner.cpp b/Engine/Plugins/Experimental/AnimNext/Source/AnimNextEditor/Private/Workspace/AssetDocumentSummoner.cpp
--- a/Engine/Plugins/Experimental/AnimNext/Source/AnimNextEditor/Private/Workspace/AssetDocumentSummoner.cpp
+++ b/Engine/Plugins/Experimental/AnimNext/Source/AnimNextEditor/Private/Workspace/AssetDocumentSummoner.cpp
@@ -4,6 +4,7 @@
 
 #include "AnimNextWorkspaceEditor.h"
 #include "AssetDefinitionRegistry.h"
+#include "ClassIconFinder.h"
 #include "Widgets/Docking/SDockTab.h"
 #include "Widgets/Layout/SSpacer.h"
 
@@ -115,7 +116,12 @@ const FSlateBrush* FAssetDocumentSummoner::GetTabIconForObject(const FWorkflowTa
 		const FAssetData AssetData(DocumentID);
 		if(const UAssetDefinition* AssetDefinition = AssetDefinitionRegistry->GetAssetDefinitionForAsset(AssetData))
 		{
-			return AssetDefinition->GetIconBrush(AssetData, AssetData.AssetClassPath.GetAssetName());
+			const FSlateBrush* ThumbnailBrush = AssetDefinition->GetThumbnailBrush(AssetData, AssetData.AssetClassPath.GetAssetName());
+			if(ThumbnailBrush == nullptr)
+			{
+				return FClassIconFinder::FindThumbnailForClass(DocumentID->GetClass(), NAME_None);
+			}
+			return ThumbnailBrush;
 		}
 	}
 	return nullptr;
