AnimNext: Dynamic graphs and component workflow improvements

Added a parameter to graph tasks that allows the task to dynamically determine what graph to run.
Added extra validation to the schedule compiler to deal with using parameters to determine the graph to run.
Switched component parameters to use instanced subobjects, of which there is only one type supported at the moment (graphs).
Also added access to correct world delta time so we run according to variable framerate.

diff --git a/Engine/Plugins/Experimental/AnimNext/Source/AnimNext/Internal/Graph/AnimNextGraph.h b/Engine/Plugins/Experimental/AnimNext/Source/AnimNext/Internal/Graph/AnimNextGraph.h
--- a/Engine/Plugins/Experimental/AnimNext/Source/AnimNext/Internal/Graph/AnimNextGraph.h
+++ b/Engine/Plugins/Experimental/AnimNext/Source/AnimNext/Internal/Graph/AnimNextGraph.h
@@ -74,6 +74,9 @@ struct ANIMNEXT_API FAnimNextGraphInstance
 	// Returns true if we have a live graph instance, false otherwise
 	bool IsValid() const;
 
+	// Check to see if this instance data matches the provided graph
+	bool UsesGraph(const UAnimNextGraph* InGraph) const;
+
 private:
 	// Hard reference to the graph used to create this instance to ensure we can release it safely
 	UPROPERTY()
@@ -118,9 +121,6 @@ public:
 	// Get the parameter to use to access the current LOD
 	UE::AnimNext::FParamId GetCurrentLODParam() const { return CurrentLODId; }
 
-	// Get the parameter to use to access the delta time
-	UE::AnimNext::FParamId GetDeltaTimeParam() const { return DeltaTimeId; }
-	
 protected:
 	// Support rig VM execution
 	TArray<FRigVMExternalVariable> GetRigVMExternalVariables();
@@ -188,13 +188,8 @@ protected:
 	UPROPERTY(EditAnywhere, Category = "Graph", meta=(CustomWidget = "ParamName", AllowedParamType = "int32"))
 	FName CurrentLOD = TEXT("UE_AnimNextMeshComponent_PredictedLODLevel");
 
-	// The parameter to use to access the current delta time
-	UPROPERTY(EditAnywhere, Category = "Graph", meta=(CustomWidget = "ParamName", AllowedParamType = "float"))
-	FName DeltaTime = TEXT("UE_AnimNextSchedulerWorldSubsystem_DeltaTime");
-
 	UE::AnimNext::FParamId ReferencePoseId = UE::AnimNext::FParamId(ReferencePose);
 	UE::AnimNext::FParamId CurrentLODId = UE::AnimNext::FParamId(CurrentLOD);
-	UE::AnimNext::FParamId DeltaTimeId = UE::AnimNext::FParamId(DeltaTime);
 
 #if WITH_EDITORONLY_DATA
 	UPROPERTY(VisibleAnywhere, Instanced, Category = "Graph", meta = (ShowInnerProperties))
diff --git a/Engine/Plugins/Experimental/AnimNext/Source/AnimNext/Internal/Scheduler/AnimNextSchedule.h b/Engine/Plugins/Experimental/AnimNext/Source/AnimNext/Internal/Scheduler/AnimNextSchedule.h
--- a/Engine/Plugins/Experimental/AnimNext/Source/AnimNext/Internal/Scheduler/AnimNextSchedule.h
+++ b/Engine/Plugins/Experimental/AnimNext/Source/AnimNext/Internal/Scheduler/AnimNextSchedule.h
@@ -8,6 +8,7 @@
 #include "Scheduler/AnimNextScheduleExternalTask.h"
 #include "Scheduler/AnimNextScheduleParamScopeTask.h"
 #include "Tasks/Task.h"
+#include "Scheduler/IAnimNextScheduleTermInterface.h"
 #include "AnimNextSchedule.generated.h"
 
 class UAnimNextGraph;
@@ -27,6 +28,29 @@ namespace UE::AnimNext
 	struct FScheduleTickFunction;
 }
 
+USTRUCT()
+struct FAnimNextScheduleEntryTerm
+{
+	GENERATED_BODY()
+
+	FAnimNextScheduleEntryTerm() = default;
+	
+	FAnimNextScheduleEntryTerm(FName InName, const FAnimNextParamType& InType, EScheduleTermDirection InDirection)
+		: Name(InName)
+		, Type(InType)
+		, Direction(InDirection)
+	{}
+
+	UPROPERTY(EditAnywhere, Category = "Term")
+	FName Name;
+
+	UPROPERTY(EditAnywhere, Category = "Term")
+	FAnimNextParamType Type;
+
+	UPROPERTY(EditAnywhere, Category = "Term")
+	EScheduleTermDirection Direction = EScheduleTermDirection::Input;
+};
+
 UCLASS(EditInlineNew, Abstract)
 class UAnimNextScheduleEntry : public UObject
 {
@@ -36,7 +60,7 @@ private:
 	friend class UAnimNextSchedule;
 };
 
-UCLASS()
+UCLASS(DisplayName="Graph")
 class UAnimNextScheduleEntry_AnimNextGraph : public UAnimNextScheduleEntry
 {
 	GENERATED_BODY()
@@ -44,20 +68,24 @@ class UAnimNextScheduleEntry_AnimNextGraph : public UAnimNextScheduleEntry
 private:
 	friend class UAnimNextSchedule;
 
-	// The graph to run
+	// The graph to run by default
 	UPROPERTY(EditAnywhere, Category = "Graph")
 	TObjectPtr<UAnimNextGraph> Graph = nullptr;
 
+	// Parameter to get the graph from dynamically
+	UPROPERTY(EditAnywhere, Category = "Graph", meta = (CustomWidget = "ParamName", AllowedParamType = "TObjectPtr<UAnimNextGraph>"))
+	FName DynamicGraph;
+
 	// An optional entry point to use when running the supplied graph
 	UPROPERTY(EditAnywhere, Category = "Graph", meta = (CustomWidget = "ParamName", AllowedParamType = "FName"))
 	FName EntryPoint;
 
-	// The output intermediate terms used by the graph
+	// The intermediate terms used by the graph
 	UPROPERTY(EditAnywhere, Category = "Graph")
-	TArray<FName> Terms;
+	TArray<FAnimNextScheduleEntryTerm> Terms;
 };
 
-UCLASS()
+UCLASS(DisplayName="Port")
 class UAnimNextScheduleEntry_Port : public UAnimNextScheduleEntry
 {
 	GENERATED_BODY()
@@ -65,16 +93,16 @@ class UAnimNextScheduleEntry_Port : public UAnimNextScheduleEntry
 private:
 	friend class UAnimNextSchedule;
 
-	// The input intermediate terms used by this port
-	UPROPERTY(EditAnywhere, Category = "Port")
-	TArray<FName> Terms;
-
 	// The type of the port to use
-	UPROPERTY(EditAnywhere, Category = "Port")
+	UPROPERTY(EditAnywhere, Category = "Port", meta = (ShowDisplayNames))
 	TSubclassOf<UAnimNextSchedulePort> Port;
+
+	// The intermediate terms used by this port
+	UPROPERTY(EditAnywhere, Category = "Port")
+	TArray<FAnimNextScheduleEntryTerm> Terms;
 };
 
-UCLASS()
+UCLASS(DisplayName="External")
 class UAnimNextScheduleEntry_ExternalTask : public UAnimNextScheduleEntry
 {
 	GENERATED_BODY()
@@ -91,7 +119,7 @@ private:
 	FName Object;
 };
 
-UCLASS()
+UCLASS(DisplayName="Scope")
 class UAnimNextScheduleEntry_ParamScope : public UAnimNextScheduleEntry
 {
 	GENERATED_BODY()
diff --git a/Engine/Plugins/Experimental/AnimNext/Source/AnimNext/Internal/Scheduler/AnimNextScheduleGraphTask.h b/Engine/Plugins/Experimental/AnimNext/Source/AnimNext/Internal/Scheduler/AnimNextScheduleGraphTask.h
--- a/Engine/Plugins/Experimental/AnimNext/Source/AnimNext/Internal/Scheduler/AnimNextScheduleGraphTask.h
+++ b/Engine/Plugins/Experimental/AnimNext/Source/AnimNext/Internal/Scheduler/AnimNextScheduleGraphTask.h
@@ -47,6 +47,9 @@ private:
 	UPROPERTY()
 	TObjectPtr<UAnimNextGraph> Graph = nullptr;
 
+	UPROPERTY()
+	FName DynamicGraph;
+
 	// Index of each term in the schedule intermediates
 	UPROPERTY()
 	TArray<uint32> Terms;
diff --git a/Engine/Plugins/Experimental/AnimNext/Source/AnimNext/Internal/Scheduler/AnimNextScheduleParamScopeTask.h b/Engine/Plugins/Experimental/AnimNext/Source/AnimNext/Internal/Scheduler/AnimNextScheduleParamScopeTask.h
--- a/Engine/Plugins/Experimental/AnimNext/Source/AnimNext/Internal/Scheduler/AnimNextScheduleParamScopeTask.h
+++ b/Engine/Plugins/Experimental/AnimNext/Source/AnimNext/Internal/Scheduler/AnimNextScheduleParamScopeTask.h
@@ -8,6 +8,7 @@
 class UAnimNextSchedule;
 class UAnimNextParameterBlock;
 class UAnimNextSchedulerWorldSubsystem;
+class UAnimNextComponent;
 
 namespace UE::AnimNext
 {
@@ -31,6 +32,7 @@ private:
 	friend class UAnimNextSchedulerWorldSubsystem;
 	friend struct UE::AnimNext::FScheduleInstanceData;
 	friend struct UE::AnimNext::FScheduleTickFunction;
+	friend class UAnimNextComponent;
 
 	UPROPERTY()
 	uint32 TaskIndex = MAX_uint32;
diff --git a/Engine/Plugins/Experimental/AnimNext/Source/AnimNext/Private/Component/AnimNextComponent.cpp b/Engine/Plugins/Experimental/AnimNext/Source/AnimNext/Private/Component/AnimNextComponent.cpp
--- a/Engine/Plugins/Experimental/AnimNext/Source/AnimNext/Private/Component/AnimNextComponent.cpp
+++ b/Engine/Plugins/Experimental/AnimNext/Source/AnimNext/Private/Component/AnimNextComponent.cpp
@@ -14,8 +14,70 @@ void UAnimNextComponent::OnRegister()
 
 	if (Schedule)
 	{
+		// Initialization callback to set up any persistent external parameters
+		auto Initialize = [this](const FScheduleContext& InContext)
+		{
+			FScheduleInstanceData& InstanceData = InContext.GetInstanceData();
+
+			// First group params into scopes
+			TMap<FName, TArray<UAnimNextComponentParameter*, TInlineAllocator<4>>, TInlineSetAllocator<4>> ParamsByScope;
+			for(UAnimNextComponentParameter* Parameter : Parameters)
+			{
+				if(Parameter && Parameter->IsValid())
+				{
+					ParamsByScope.FindOrAdd(Parameter->Scope).Add(Parameter);
+				}
+			}
+
+			for(const TPair<FName, TArray<UAnimNextComponentParameter*, TInlineAllocator<4>>>& ParamPair : ParamsByScope)
+			{
+				FName Scope = ParamPair.Key;
+				TSharedPtr<FParamStack> StackToUse;
+
+				if(Scope == NAME_None)
+				{
+					StackToUse = InstanceData.RootParamStack;
+				}
+				else
+				{
+					const FAnimNextScheduleParamScopeEntryTask* FoundTask = InContext.Schedule->ParamScopeEntryTasks.FindByPredicate([Scope](const FAnimNextScheduleParamScopeEntryTask& InTask)
+					{
+						return InTask.Scope == Scope;
+					});
+
+					if(FoundTask)
+					{
+						StackToUse = InstanceData.ParamStacks[FoundTask->ParamScopeIndex];
+					}
+				}
+
+				if(StackToUse.IsValid())
+				{
+					TArray<TPair<FParamId, Private::FParamEntry>, TInlineAllocator<4>> Params;
+					for(UAnimNextComponentParameter* Parameter : ParamPair.Value)
+					{
+						FParamId ParamId;
+						FAnimNextParamType Type;
+						uint8* Value = nullptr;
+						Parameter->GetParamInfo(ParamId, Type, Value);
+						check(Type.IsValid() && Value != nullptr);
+
+						constexpr bool bIsReference = true;
+						constexpr bool bIsMutable = false;
+						Params.Emplace(ParamId, Private::FParamEntry(Type.GetHandle(), TArrayView<uint8>(Value, 1), bIsReference, bIsMutable));
+					}
+
+					FParamStackLayerHandle NewLayer = FParamStack::MakeLayer(Params);
+					InstanceData.StaticUserHandles.Add(MoveTemp(NewLayer));
+
+					// Layer is never popped as this is happening at the very start of execution
+					StackToUse->PushLayer(InstanceData.StaticUserHandles.Last());
+				}
+			}
+		};
+
 		check(!SchedulerHandle.IsValid());
-		SchedulerHandle = FScheduler::AcquireHandle(this, Schedule, Parameters, InitMethod);
+		SchedulerHandle = FScheduler::AcquireHandle(this, Schedule, InitMethod, MoveTemp(Initialize));
 	}
 }
 
@@ -29,29 +91,6 @@ void UAnimNextComponent::OnUnregister()
 	SchedulerHandle.Invalidate();
 }
 
-void UAnimNextComponent::UpdateLayer(UE::AnimNext::FParamStackLayerHandle& InHandle) const
-{
-	if (UAnimNextComponent* This = InHandle.As<UAnimNextComponent>())
-	{
-		This->Update();
-	}
-}
-
-UE::AnimNext::FParamStackLayerHandle UAnimNextComponent::CacheLayer() const
-{
-	return UE::AnimNext::FParamStack::MakeReferenceLayer(const_cast<UAnimNextComponent*>(this));
-}
-
-TStructOnScope<FActorComponentInstanceData> UAnimNextComponent::GetComponentInstanceData() const
-{
-	TStructOnScope<FActorComponentInstanceData> InstanceData = MakeStructOnScope<FActorComponentInstanceData, FAnimNextComponentInstanceData>(this);
-	FAnimNextComponentInstanceData* AnimNextComponentInstanceData = InstanceData.Cast<FAnimNextComponentInstanceData>();
-	AnimNextComponentInstanceData->Schedule = Schedule;
-	AnimNextComponentInstanceData->Parameters = Parameters;
-
-	return InstanceData;
-}
-
 void UAnimNextComponent::SetParameterInScope(FName Scope, FName Name, int32 Value)
 {
 	checkNoEntry();
@@ -117,12 +156,3 @@ void UAnimNextComponent::Enable(bool bEnabled)
 {
 	UE::AnimNext::FScheduler::EnableHandle(this, SchedulerHandle, bEnabled);
 }
-
-void FAnimNextComponentInstanceData::ApplyToComponent(UActorComponent* Component, const ECacheApplyPhase CacheApplyPhase) 
-{
-	Super::ApplyToComponent(Component, CacheApplyPhase);
-
-	UAnimNextComponent* AnimNextComponent = CastChecked<UAnimNextComponent>(Component);
-	AnimNextComponent->Schedule = Schedule;
-	AnimNextComponent->Parameters = Parameters;
-}
diff --git a/Engine/Plugins/Experimental/AnimNext/Source/AnimNext/Private/Component/AnimNextComponent.h b/Engine/Plugins/Experimental/AnimNext/Source/AnimNext/Private/Component/AnimNextComponent.h
--- a/Engine/Plugins/Experimental/AnimNext/Source/AnimNext/Private/Component/AnimNextComponent.h
+++ b/Engine/Plugins/Experimental/AnimNext/Source/AnimNext/Private/Component/AnimNextComponent.h
@@ -2,9 +2,8 @@
 
 #pragma once
 
+#include "AnimNextComponentParameter.h"
 #include "Components/ActorComponent.h"
-#include "Param/AnimNextParameterCollection.h"
-#include "Param/IAnimNextParameterSourceInterface.h"
 #include "Scheduler/AnimNextSchedule.h"
 #include "Scheduler/ScheduleHandle.h"
 #include "AnimNextComponent.generated.h"
@@ -12,27 +11,14 @@
 class UAnimNextSchedule;
 struct FAnimNextComponentInstanceData;
 
-UCLASS(MinimalAPI, Blueprintable, meta = (BlueprintSpawnableComponent))
-class UAnimNextComponent : public UActorComponent, public IAnimNextParameterSourceInterface
+UCLASS(MinimalAPI, meta = (BlueprintSpawnableComponent))
+class UAnimNextComponent : public UActorComponent
 {
 	GENERATED_BODY()
 
-	friend struct FAnimNextComponentInstanceData;
-
 	// UActorComponent interface
 	virtual void OnRegister() override;
 	virtual void OnUnregister() override;
-	virtual TStructOnScope<FActorComponentInstanceData> GetComponentInstanceData() const override;
-
-	// IAnimNextParameterSourceInterface interface
-	virtual void UpdateLayer(UE::AnimNext::FParamStackLayerHandle& InHandle) const override;
-	virtual UE::AnimNext::FParamStackLayerHandle CacheLayer() const override;
-
-	// Because UHT says: "A Private function cannot be a BlueprintImplementableEvent!"
-protected:
-	// Override point to allow self to be updated
-	UFUNCTION(Category = "AnimNext", BlueprintImplementableEvent, meta=(BlueprintThreadSafe, ForceAsFunction))
-	void Update();
 
 public:
 	// Sets a parameter in the scope. Scopes correspond to an existing scope in a schedule.
@@ -49,11 +35,11 @@ private:
 private:
 	// The execution schedule that this component will run
 	UPROPERTY(EditAnywhere, Category="Schedule")
-	UAnimNextSchedule* Schedule = nullptr;
+	TObjectPtr<UAnimNextSchedule> Schedule = nullptr;
 
-	// Parameters to apply at each scope
-	UPROPERTY(EditAnywhere, Category = "Parameters", meta = (CustomWidget = "ParamName", AllowedParamType = "FAnimNextScope"))
-	TMap<FName, FAnimNextParameterCollection> Parameters;
+	// Parameters to apply on schedule/component registration
+	UPROPERTY(Instanced, EditAnywhere, Category = "Parameters")
+	TArray<TObjectPtr<UAnimNextComponentParameter>> Parameters;
 
 	// How to initialize the schedule
 	UPROPERTY(EditAnywhere, Category="Schedule")
@@ -61,36 +47,4 @@ private:
 
 	// Handle to the registered results/schedule
 	UE::AnimNext::FScheduleHandle SchedulerHandle;
-
-	// Cached layer handle for 'self'
-	mutable UE::AnimNext::FParamStackLayerHandle LayerHandle;
-};
-
-/** Used to store component data during RerunConstructionScripts */
-USTRUCT()
-struct FAnimNextComponentInstanceData : public FActorComponentInstanceData
-{
-	GENERATED_BODY()
-
-	FAnimNextComponentInstanceData() = default;
-	FAnimNextComponentInstanceData(const UAnimNextComponent* InSourceComponent)
-		: FActorComponentInstanceData(InSourceComponent)
-	{}
-
-private:
-	friend class UAnimNextComponent;
-
-	virtual bool ContainsData() const override
-	{
-		return true;
-	}
-
-	virtual void ApplyToComponent(UActorComponent* Component, const ECacheApplyPhase CacheApplyPhase) override;
-
-	UPROPERTY(EditAnywhere, Category = "Schedule")
-	UAnimNextSchedule* Schedule = nullptr;
-
-	UPROPERTY(EditAnywhere, Category = "Parameters")
-	TMap<FName, FAnimNextParameterCollection> Parameters;
 };
-
diff --git a/Engine/Plugins/Experimental/AnimNext/Source/AnimNext/Private/Component/AnimNextComponentParameter.h b/Engine/Plugins/Experimental/AnimNext/Source/AnimNext/Private/Component/AnimNextComponentParameter.h
new file mode 100644
--- /dev/null
+++ b/Engine/Plugins/Experimental/AnimNext/Source/AnimNext/Private/Component/AnimNextComponentParameter.h
@@ -0,0 +1,82 @@
+﻿// Copyright Epic Games, Inc. All Rights Reserved.
+
+#pragma once
+
+#include "Param/ParamId.h"
+#include "Param/ParamType.h"
+#include "Param/ParamTypeHandle.h"
+#include "AnimNextComponentParameter.generated.h"
+
+class UAnimNextGraph;
+
+// Boilerplate macro that all derived types of UAnimNextComponentParameter should use.
+// Usage is: IMPLEMENT_ANIMNEXT_COMPONENT_PARAMETER(UAnimNextComponentParameter_DerivedType)
+#define IMPLEMENT_ANIMNEXT_COMPONENT_PARAMETER(Type) \
+	virtual void CacheParamInfo() const override \
+	{ \
+		using namespace UE::AnimNext; \
+		if(!ParameterId.IsValid() && Parameter != NAME_None) \
+		{ \
+			ParameterId = FParamId(Parameter); \
+			ValuePtr = const_cast<uint8*>(reinterpret_cast<const uint8*>((&Value))); \
+			ValueProperty = StaticClass()->FindPropertyByName(GET_MEMBER_NAME_CHECKED(Type, Value)); \
+			ValueType = FParamTypeHandle::FromProperty(ValueProperty).GetType(); \
+		} \
+	} \
+
+// Base class for static parameters that can be inserted into a schedule via a component
+// Each parameter type needs its own derived type of this object
+UCLASS(DefaultToInstanced, abstract, editinlinenew)
+class UAnimNextComponentParameter : public UObject
+{
+	GENERATED_BODY()
+
+public:
+	// The scope to apply the parameter to. If this is empty, it is applied at the root scope of the schedule.
+	UPROPERTY(EditAnywhere, Category = "Parameter", AdvancedDisplay, meta = (CustomWidget = "ParamName"))
+	FName Scope;
+
+	// Check validity
+	bool IsValid() const
+	{
+		CacheParamInfo();
+
+		return ValueType.IsValid() && ParameterId.IsValid() && ValuePtr != nullptr;
+	}
+	
+	// Get the name, type and value for this parameter
+	void GetParamInfo(UE::AnimNext::FParamId& OutParamId, FAnimNextParamType& OutType, uint8*& OutValue) const
+	{
+		CacheParamInfo();
+
+		OutParamId = ParameterId;
+		OutType = ValueType;
+		OutValue = ValuePtr;
+	}
+
+private:
+	virtual void CacheParamInfo() const PURE_VIRTUAL(UAnimNextComponentParameter::CacheParamInfo, )
+
+protected:
+	mutable UE::AnimNext::FParamId ParameterId;
+	mutable FProperty* ValueProperty = nullptr;
+	mutable FAnimNextParamType ValueType;
+	mutable	uint8* ValuePtr = nullptr;
+};
+
+// An object parameter
+UCLASS(MinimalAPI, meta = (DisplayName = "Anim Next Graph"))
+class UAnimNextComponentParameter_AnimNextGraph : public UAnimNextComponentParameter
+{
+	GENERATED_BODY()
+
+	IMPLEMENT_ANIMNEXT_COMPONENT_PARAMETER(UAnimNextComponentParameter_AnimNextGraph);
+
+	// The parameter to set the value to
+	UPROPERTY(EditAnywhere, Category = "Parameter", meta = (CustomWidget = "ParamName", AllowedParamType = "TObjectPtr<UAnimNextGraph>"))
+	FName Parameter;
+
+	// The value to set
+	UPROPERTY(EditAnywhere, Category = "Parameter")
+	TObjectPtr<UAnimNextGraph> Value;
+};
diff --git a/Engine/Plugins/Experimental/AnimNext/Source/AnimNext/Private/Graph/AnimNextGraph.cpp b/Engine/Plugins/Experimental/AnimNext/Source/AnimNext/Private/Graph/AnimNextGraph.cpp
--- a/Engine/Plugins/Experimental/AnimNext/Source/AnimNext/Private/Graph/AnimNextGraph.cpp
+++ b/Engine/Plugins/Experimental/AnimNext/Source/AnimNext/Private/Graph/AnimNextGraph.cpp
@@ -154,6 +154,11 @@ bool FAnimNextGraphInstance::IsValid() const
 	return GraphInstancePtr.IsValid();
 }
 
+bool FAnimNextGraphInstance::UsesGraph(const UAnimNextGraph* InGraph) const
+{
+	return Graph == InGraph;
+}
+
 UAnimNextGraph::UAnimNextGraph(const FObjectInitializer& ObjectInitializer)
 	: Super(ObjectInitializer)
 {
@@ -232,7 +237,6 @@ void UAnimNextGraph::PostLoad()
 
 	ReferencePoseId = FParamId(ReferencePose);
 	CurrentLODId = FParamId(CurrentLOD);
-	DeltaTimeId = FParamId(DeltaTime);
 }
 
 void UAnimNextGraph::Serialize(FArchive& Ar)
diff --git a/Engine/Plugins/Experimental/AnimNext/Source/AnimNext/Private/Param/AnimNextParameterSourceRef.cpp b/Engine/Plugins/Experimental/AnimNext/Source/AnimNext/Private/Param/AnimNextParameterSourceRef.cpp
--- a/Engine/Plugins/Experimental/AnimNext/Source/AnimNext/Private/Param/AnimNextParameterSourceRef.cpp
+++ b/Engine/Plugins/Experimental/AnimNext/Source/AnimNext/Private/Param/AnimNextParameterSourceRef.cpp
@@ -1,7 +1,7 @@
 // Copyright Epic Games, Inc. All Rights Reserved.
 
 #include "Param/AnimNextParameterSourceRef.h"
-#include "Component/AnimNextComponent.h"
+#include "Components/ActorComponent.h"
 #include "GameFramework/Actor.h"
 
 const IAnimNextParameterSourceInterface* FAnimNextParameterSourceRef::Get(const UObject* InContextObject) const
@@ -9,14 +9,14 @@ const IAnimNextParameterSourceInterface* FAnimNextParameterSourceRef::Get(const
 	switch (Type)
 	{
 	case EAnimNextParameterSourceRefType::Self:
-		return Cast<UAnimNextComponent>(InContextObject);
-	case EAnimNextParameterSourceRefType::Asset:
+		return Cast<IAnimNextParameterSourceInterface>(InContextObject);
+	case EAnimNextParameterSourceRefType::Asset: 
 		return Asset.GetInterface();
 	case EAnimNextParameterSourceRefType::Component:
 	case EAnimNextParameterSourceRefType::Actor:
-		if (const UAnimNextComponent* AnimNextComponent = Cast<UAnimNextComponent>(InContextObject))
+		if (const UActorComponent* ActorComponent = Cast<UActorComponent>(InContextObject))
 		{
-			if (AActor* Actor = AnimNextComponent->GetOwner())
+			if (AActor* Actor = ActorComponent->GetOwner())
 			{
 				if (Type == EAnimNextParameterSourceRefType::Actor)
 				{
diff --git a/Engine/Plugins/Experimental/AnimNext/Source/AnimNext/Private/Param/ParamId.cpp b/Engine/Plugins/Experimental/AnimNext/Source/AnimNext/Private/Param/ParamId.cpp
--- a/Engine/Plugins/Experimental/AnimNext/Source/AnimNext/Private/Param/ParamId.cpp
+++ b/Engine/Plugins/Experimental/AnimNext/Source/AnimNext/Private/Param/ParamId.cpp
@@ -128,6 +128,31 @@ void FParamId::RegisterBuiltInAdapters()
 	FRWScopeLock Lock(GParamIdLock, SLT_Write);
 	FParamIdGlobalData& ParamIdGlobalData = GetParamIdData();
 
+	// Register world delta time
+	{
+		FObjectAdapterFunction DeltaTimeFunction = [](UObject* InContextObject, FParamId InId) -> uint8*
+		{
+			if (UWorld* World = InContextObject->GetWorld())
+			{
+				float* ReturnValue = GetAdapterReturnValue<float>(InId);
+				*ReturnValue = World->GetDeltaSeconds();
+				return reinterpret_cast<uint8*>(ReturnValue);
+			}
+			return nullptr;
+		};
+
+		FName DeltaTimeName("UE_Frame_DeltaTime");
+		const uint32 DeltaTimeNameFunctionIndex = MakeParamId_NoLock(DeltaTimeName);
+		ParamIdGlobalData.ParamData[DeltaTimeNameFunctionIndex].AdapterIndex = ParamIdGlobalData.ParamAdapters.Num();
+		FText TooltipText = LOCTEXT("DeltaTimeTooltip", "The current delta time");
+		FParamAdapter& NewAdapter = ParamIdGlobalData.ParamAdapters.Emplace_GetRef(FParamDefinition(DeltaTimeNameFunctionIndex, DeltaTimeName, FAnimNextParamType::GetType<float>(), TooltipText), MoveTemp(DeltaTimeFunction));
+		FParams::RegisterBuiltInParameter(NewAdapter.Definition);
+
+		// Calc buffer size for the result
+		ParamIdGlobalData.ParamData[DeltaTimeNameFunctionIndex].BufferOffset = Align(ParamIdGlobalData.ResultBufferSize, alignof(float));
+		ParamIdGlobalData.ResultBufferSize += sizeof(float);
+	}
+	
 	// Register physics tick function
 	{
 		FObjectAdapterFunction PhysicsTickFunctionFunction = [](UObject* InContextObject, FParamId InId) -> uint8*
diff --git a/Engine/Plugins/Experimental/AnimNext/Source/AnimNext/Private/Scheduler/AnimNextSchedule.cpp b/Engine/Plugins/Experimental/AnimNext/Source/AnimNext/Private/Scheduler/AnimNextSchedule.cpp
--- a/Engine/Plugins/Experimental/AnimNext/Source/AnimNext/Private/Scheduler/AnimNextSchedule.cpp
+++ b/Engine/Plugins/Experimental/AnimNext/Source/AnimNext/Private/Scheduler/AnimNextSchedule.cpp
@@ -85,18 +85,33 @@ void UAnimNextSchedule::CompileSchedule()
 	// MAX_uint32 means 'global scope' in this context
 	uint32 ParentScopeIndex = MAX_uint32;
 
-	TArray<FAnimNextParamType> IntermediateTypes;
+	TArray<FAnimNextScheduleEntryTerm> IntermediateTerms;
 	TMap<FName, uint32> IntermediateMap;
 
 	TFunction<void(const TArray<TObjectPtr<UAnimNextScheduleEntry>>&)> EmitEntries;
 
 	// Iterate over all entries, recursing into scopes
-	EmitEntries = [this, &EmitEntries, &Emit, &EmitPrerequisite, &ParentScopeIndex, &IntermediateTypes, &IntermediateMap](const TArray<TObjectPtr<UAnimNextScheduleEntry>>& InEntries)
+	EmitEntries = [this, &EmitEntries, &Emit, &EmitPrerequisite, &ParentScopeIndex, &IntermediateTerms, &IntermediateMap](const TArray<TObjectPtr<UAnimNextScheduleEntry>>& InEntries)
 	{
 		for (int32 EntryIndex = 0; EntryIndex < InEntries.Num(); ++EntryIndex)
 		{
 			UAnimNextScheduleEntry* Entry = InEntries[EntryIndex];
 
+			auto CheckTermDirectionCompatibility = [](FName InName, EScheduleTermDirection InExistingDirection, EScheduleTermDirection InNewDirection)
+			{
+				switch(InExistingDirection)
+				{
+				case EScheduleTermDirection::Input:
+					// Input before output: error
+					UE_LOG(LogAnimation, Error, TEXT("Term '%s' was used as an input before it was output"), *InName.ToString());
+					return false;
+				case EScheduleTermDirection::Output:
+					return true;
+				}
+
+				return false;
+			};
+
 			if (UAnimNextScheduleEntry_Port* PortEntry = Cast<UAnimNextScheduleEntry_Port>(Entry))
 			{
 				bool bValid = true;
@@ -120,20 +135,33 @@ void UAnimNextSchedule::CompileSchedule()
 
 					for(int32 TermIndex = 0; TermIndex < PortEntry->Terms.Num(); ++TermIndex)
 					{
-						FName Term = PortEntry->Terms[TermIndex];
-						const uint32* ExistingIntermediateIndexPtr = IntermediateMap.Find(Term);
-						if(ExistingIntermediateIndexPtr != nullptr)
+						FName TermName = PortEntry->Terms[TermIndex].Name;
+						if(!PortEntry->Terms[TermIndex].Type.IsValid())
+						{
+							UE_LOG(LogAnimation, Error, TEXT("AnimNext: Invalid type when processing port term, ignored: '%s'"), *TermName.ToString());
+							bValid = false;
+						}
+						else
 						{
-							const FAnimNextParamType& IntermediateType = IntermediateTypes[*ExistingIntermediateIndexPtr];
-							if(IntermediateType != Terms[TermIndex].GetType())
+							const uint32* ExistingIntermediateIndexPtr = IntermediateMap.Find(TermName);
+							if(ExistingIntermediateIndexPtr != nullptr)
 							{
-								UE_LOG(LogAnimation, Error, TEXT("AnimNext: Mismatched types when processing port term, ignored: '%s'"), *Term.ToString());
-								bValid = false;
+								const FAnimNextScheduleEntryTerm& IntermediateTerm = IntermediateTerms[*ExistingIntermediateIndexPtr];
+								if(IntermediateTerm.Type != Terms[TermIndex].GetType())
+								{
+									UE_LOG(LogAnimation, Error, TEXT("AnimNext: Mismatched types when processing port term, ignored: '%s'"), *TermName.ToString());
+									bValid = false;
+								}
+
+								if(!CheckTermDirectionCompatibility(TermName, IntermediateTerm.Direction, Terms[TermIndex].Direction))
+								{
+									bValid = false;
+								}
 							}
 						}
 					}
 				}
-				
+
 				if(bValid)
 				{
 					EmitPrerequisite();
@@ -143,18 +171,14 @@ void UAnimNextSchedule::CompileSchedule()
 					PortTask.ParamScopeIndex = ParentScopeIndex;
 					PortTask.Port = PortEntry->Port;
 
-					UAnimNextSchedulePort* CDO = PortEntry->Port->GetDefaultObject<UAnimNextSchedulePort>();
-					check(CDO);
-					TConstArrayView<FScheduleTerm> Terms = CDO->GetTerms();
-
 					for(int32 TermIndex = 0; TermIndex < PortEntry->Terms.Num(); ++TermIndex)
 					{
-						FName Term = PortEntry->Terms[TermIndex];
-						const uint32* ExistingIntermediateIndexPtr = IntermediateMap.Find(Term);
+						FName TermName = PortEntry->Terms[TermIndex].Name;
+						const uint32* ExistingIntermediateIndexPtr = IntermediateMap.Find(TermName);
 						if(ExistingIntermediateIndexPtr == nullptr)
 						{
-							uint32 IntermediateIndex = IntermediateTypes.Add(Terms[TermIndex].GetType());
-							IntermediateMap.Add(Term, IntermediateIndex);
+							uint32 IntermediateIndex = IntermediateTerms.Emplace(TermName, PortEntry->Terms[TermIndex].Type, PortEntry->Terms[TermIndex].Direction);
+							IntermediateMap.Add(TermName, IntermediateIndex);
 							PortTask.Terms.Add(IntermediateIndex);
 						}
 						else
@@ -174,12 +198,12 @@ void UAnimNextSchedule::CompileSchedule()
 			{
 				bool bValid = true;
 
-				if(GraphEntry->Graph == nullptr)
+				if(GraphEntry->Graph == nullptr && GraphEntry->DynamicGraph == NAME_None)
 				{
-					UE_LOG(LogAnimation, Warning, TEXT("AnimNext: Invalid graph supplied"));
+					UE_LOG(LogAnimation, Error, TEXT("AnimNext: Invalid graph or no parameter supplied"));
 					bValid = false;
 				}
-				else
+				else if(GraphEntry->Graph != nullptr)
 				{
 					TConstArrayView<FScheduleTerm> Terms = GraphEntry->Graph->GetTerms();
 					if(GraphEntry->Terms.Num() != Terms.Num())
@@ -187,17 +211,50 @@ void UAnimNextSchedule::CompileSchedule()
 						UE_LOG(LogAnimation, Error, TEXT("AnimNext: Incorrect term count for graph: %d"), GraphEntry->Terms.Num());
 						bValid = false;
 					}
+					else
+					{
+						// Validate graph terms match schedule-expected terms
+						for(int32 TermIndex = 0; TermIndex < GraphEntry->Terms.Num(); ++TermIndex)
+						{
+							FName TermName = GraphEntry->Terms[TermIndex].Name;
+							if(Terms[TermIndex].Direction != GraphEntry->Terms[TermIndex].Direction)
+							{
+								UE_LOG(LogAnimation, Error, TEXT("AnimNext: Mismatched direction when processing graph term, ignored: '%s'"), *TermName.ToString());
+								bValid = false;
+							}
+							
+							if(Terms[TermIndex].GetType() != GraphEntry->Terms[TermIndex].Type)
+							{
+								UE_LOG(LogAnimation, Error, TEXT("AnimNext: Mismatched types when processing graph term, ignored: '%s'"), *TermName.ToString());
+								bValid = false;
+							}
+						}
+					}
+				}
 
-					for(int32 TermIndex = 0; TermIndex < GraphEntry->Terms.Num(); ++TermIndex)
+				// Validate terms and check against priors
+				for(int32 TermIndex = 0; TermIndex < GraphEntry->Terms.Num(); ++TermIndex)
+				{
+					FName TermName = GraphEntry->Terms[TermIndex].Name;
+					if(!GraphEntry->Terms[TermIndex].Type.IsValid())
+					{
+						UE_LOG(LogAnimation, Error, TEXT("AnimNext: Invalid type when processing graph term, ignored: '%s'"), *TermName.ToString());
+						bValid = false;
+					}
+					else
 					{
-						FName Term = GraphEntry->Terms[TermIndex];
-						const uint32* ExistingIntermediateIndexPtr = IntermediateMap.Find(Term);
+						const uint32* ExistingIntermediateIndexPtr = IntermediateMap.Find(TermName);
 						if(ExistingIntermediateIndexPtr != nullptr)
 						{
-							const FAnimNextParamType& IntermediateType = IntermediateTypes[*ExistingIntermediateIndexPtr];
-							if(IntermediateType != Terms[TermIndex].GetType())
+							const FAnimNextScheduleEntryTerm& IntermediateTerm = IntermediateTerms[*ExistingIntermediateIndexPtr];
+							if(IntermediateTerm.Type != GraphEntry->Terms[TermIndex].Type)
+							{
+								UE_LOG(LogAnimation, Error, TEXT("AnimNext: Mismatched types when processing graph term, ignored: '%s'"), *TermName.ToString());
+								bValid = false;
+							}
+								
+							if(!CheckTermDirectionCompatibility(TermName, IntermediateTerm.Direction, GraphEntry->Terms[TermIndex].Direction))
 							{
-								UE_LOG(LogAnimation, Error, TEXT("AnimNext: Mismatched types when processing graph term, ignored: '%s'"), *Term.ToString());
 								bValid = false;
 							}
 						}
@@ -214,16 +271,16 @@ void UAnimNextSchedule::CompileSchedule()
 					Task.ParamParentScopeIndex = ParentScopeIndex;
 					Task.EntryPoint = GraphEntry->EntryPoint;
 					Task.Graph = GraphEntry->Graph;
+					Task.DynamicGraph = GraphEntry->DynamicGraph;
 
-					TConstArrayView<FScheduleTerm> Terms = GraphEntry->Graph->GetTerms();
 					for(int32 TermIndex = 0; TermIndex < GraphEntry->Terms.Num(); ++TermIndex)
 					{
-						FName Term = GraphEntry->Terms[TermIndex];
-						const uint32* ExistingIntermediateIndexPtr = IntermediateMap.Find(Term);
+						FName TermName = GraphEntry->Terms[TermIndex].Name;
+						const uint32* ExistingIntermediateIndexPtr = IntermediateMap.Find(TermName);
 						if(ExistingIntermediateIndexPtr == nullptr)
 						{
-							uint32 IntermediateIndex = IntermediateTypes.Add(Terms[TermIndex].GetType());
-							IntermediateMap.Add(Term, IntermediateIndex);
+							uint32 IntermediateIndex = IntermediateTerms.Emplace(TermName, GraphEntry->Terms[TermIndex].Type, GraphEntry->Terms[TermIndex].Direction);
+							IntermediateMap.Add(TermName, IntermediateIndex);
 							Task.Terms.Add(IntermediateIndex);
 						}
 						else
@@ -308,14 +365,14 @@ void UAnimNextSchedule::CompileSchedule()
 	// Process intermediates
 	if(IntermediateMap.Num() > 0)
 	{
-		check(IntermediateMap.Num() == IntermediateTypes.Num());
+		check(IntermediateMap.Num() == IntermediateTerms.Num());
 		
 		TArray<FPropertyBagPropertyDesc> PropertyDescs;
-		PropertyDescs.Reserve(IntermediateTypes.Num());
+		PropertyDescs.Reserve(IntermediateTerms.Num());
 		 
 		for(const TPair<FName, uint32>& IntermediatePair : IntermediateMap)
 		{
-			const FAnimNextParamType& IntermediateType = IntermediateTypes[IntermediatePair.Value];
+			const FAnimNextParamType& IntermediateType = IntermediateTerms[IntermediatePair.Value].Type;
 			check(IntermediateType.IsValid());
 			PropertyDescs.Emplace(IntermediatePair.Key, IntermediateType.GetContainerType(), IntermediateType.GetValueType(), IntermediateType.GetValueTypeObject());
 		}
diff --git a/Engine/Plugins/Experimental/AnimNext/Source/AnimNext/Private/Scheduler/AnimNextScheduleGraphTask.cpp b/Engine/Plugins/Experimental/AnimNext/Source/AnimNext/Private/Scheduler/AnimNextScheduleGraphTask.cpp
--- a/Engine/Plugins/Experimental/AnimNext/Source/AnimNext/Private/Scheduler/AnimNextScheduleGraphTask.cpp
+++ b/Engine/Plugins/Experimental/AnimNext/Source/AnimNext/Private/Scheduler/AnimNextScheduleGraphTask.cpp
@@ -17,45 +17,73 @@ void FAnimNextScheduleGraphTask::RunGraph(const UE::AnimNext::FScheduleContext&
 
 	using namespace UE::AnimNext;
 
-	if (Graph == nullptr)
+	FParamStack& ParamStack = FParamStack::Get();
+
+	UAnimNextGraph* GraphToRun = Graph;
+	if (GraphToRun == nullptr && DynamicGraph != NAME_None)
+	{
+		if(const TObjectPtr<UAnimNextGraph>* FoundGraph = ParamStack.GetParamPtr<TObjectPtr<UAnimNextGraph>>(DynamicGraph))
+		{
+			GraphToRun = *FoundGraph;
+		}
+	}
+
+	if(GraphToRun == nullptr)
 	{
 		return;
 	}
 
 	FScheduleInstanceData& InstanceData = InContext.GetInstanceData();
 
+	// Check if we are running the correct graph and release (and any term mapping layers) it if not
+	if(InstanceData.GraphInstanceData[TaskIndex].IsValid() && !InstanceData.GraphInstanceData[TaskIndex].UsesGraph(GraphToRun))
+	{
+		InstanceData.GraphInstanceData[TaskIndex].Release();
+		InstanceData.GraphTermLayers[TaskIndex].Invalidate();
+	}
+
 	// Allocate our graph instance data
 	if (!InstanceData.GraphInstanceData[TaskIndex].IsValid())
 	{
-		Graph->AllocateInstance(InstanceData.GraphInstanceData[TaskIndex]);
+		GraphToRun->AllocateInstance(InstanceData.GraphInstanceData[TaskIndex]);
 	}
 
-	FParamStack& ParamStack = FParamStack::Get();
-	const FAnimNextGraphReferencePose* GraphReferencePose = ParamStack.GetParamPtr<FAnimNextGraphReferencePose>(Graph->GetReferencePoseParam());
+	const FAnimNextGraphReferencePose* GraphReferencePose = ParamStack.GetParamPtr<FAnimNextGraphReferencePose>(GraphToRun->GetReferencePoseParam());
 	if(GraphReferencePose == nullptr)
 	{
 		return;
 	}
 
-	const int32* GraphLODLevel = ParamStack.GetParamPtr<int32>(Graph->GetCurrentLODParam());
+	const int32* GraphLODLevel = ParamStack.GetParamPtr<int32>(GraphToRun->GetCurrentLODParam());
 	if(GraphLODLevel == nullptr)
 	{
 		return;
 	}
 
-	// TODO: This should not be fixed at arg 0, we should define this in the graph asset
-	const FParamStackLayerHandle& TermLayerHandle = InstanceData.GraphTermLayers[TaskIndex];
-	FAnimNextGraphLODPose* OutputPose = TermLayerHandle.GetMutableParamPtr<FAnimNextGraphLODPose>(Graph->GetTerms()[0].GetId());
-	if(OutputPose == nullptr)
+	// Check and allocate remapped term layer
+	FParamStackLayerHandle& TermLayerHandle = InstanceData.GraphTermLayers[TaskIndex];
+	if(!TermLayerHandle.IsValid())
 	{
-		return;
+		TConstArrayView<FScheduleTerm> GraphTerms = GraphToRun->GetTerms();
+		check(Terms.Num() == GraphTerms.Num());
+
+		TMap<FName, FName> Mapping;
+		Mapping.Reserve(GraphTerms.Num());
+		for(int32 TermIndex = 0; TermIndex < Terms.Num(); ++TermIndex)
+		{
+			uint32 IntermediateTermIndex = Terms[TermIndex];
+			const FPropertyBagPropertyDesc& PropertyDesc = InstanceData.IntermediatesData.GetPropertyBagStruct()->GetPropertyDescs()[IntermediateTermIndex];
+			Mapping.Add(PropertyDesc.Name, GraphTerms[TermIndex].GetName());
+		}
+
+		TermLayerHandle = FParamStack::MakeRemappedLayer(InstanceData.IntermediatesLayer, Mapping);
 	}
 
-	float DeltaTime = 1/60.0f;
-	const float* ExternalDeltaTime = ParamStack.GetParamPtr<float>(Graph->GetDeltaTimeParam());
-	if(ExternalDeltaTime != nullptr)
+	// TODO: This should not be fixed at arg 0, we should define this in the graph asset
+	FAnimNextGraphLODPose* OutputPose = TermLayerHandle.GetMutableParamPtr<FAnimNextGraphLODPose>(GraphToRun->GetTerms()[0].GetId());
+	if(OutputPose == nullptr)
 	{
-		DeltaTime = *ExternalDeltaTime;
+		return;
 	}
 
 	// Create or update our result pose
@@ -95,8 +123,9 @@ void FAnimNextScheduleGraphTask::RunGraph(const UE::AnimNext::FScheduleContext&
 	// This reduces churn internally by avoiding a chunk to be repeatedly allocated and freed as we push/pop marks
 	MemStack.Alloc(size_t(FPageAllocator::SmallPageSize) + 1, 16);
 
+	const float DeltaTime = InContext.GetDeltaTime();
 	const FContext Context(DeltaTime);
-	Graph->Run(Context, InstanceData.GraphInstanceData[TaskIndex], EAnimNextGraphSimulationSteps::All);
+	GraphToRun->Run(Context, InstanceData.GraphInstanceData[TaskIndex], EAnimNextGraphSimulationSteps::All);
 
 	ParamStack.PopLayer(LayerHandle);
 }
diff --git a/Engine/Plugins/Experimental/AnimNext/Source/AnimNext/Private/Scheduler/AnimNextSchedulePortTask.cpp b/Engine/Plugins/Experimental/AnimNext/Source/AnimNext/Private/Scheduler/AnimNextSchedulePortTask.cpp
--- a/Engine/Plugins/Experimental/AnimNext/Source/AnimNext/Private/Scheduler/AnimNextSchedulePortTask.cpp
+++ b/Engine/Plugins/Experimental/AnimNext/Source/AnimNext/Private/Scheduler/AnimNextSchedulePortTask.cpp
@@ -20,6 +20,25 @@ void FAnimNextSchedulePortTask::RunPort(const UE::AnimNext::FScheduleContext& In
 		{
 			FScheduleInstanceData& InstanceData = InScheduleContext.GetInstanceData();
 
+			// Check and allocate remapped term layer
+			FParamStackLayerHandle& TermLayerHandle = InstanceData.PortTermLayers[TaskIndex];
+			if(!TermLayerHandle.IsValid())
+			{
+				TConstArrayView<FScheduleTerm> PortTerms = CDO->GetTerms();
+				check(Terms.Num() == PortTerms.Num());
+
+				TMap<FName, FName> Mapping;
+				Mapping.Reserve(Terms.Num());
+				for(int32 TermIndex = 0; TermIndex < Terms.Num(); ++TermIndex)
+				{
+					uint32 IntermediateTermIndex = Terms[TermIndex];
+					const FPropertyBagPropertyDesc& PropertyDesc = InstanceData.IntermediatesData.GetPropertyBagStruct()->GetPropertyDescs()[IntermediateTermIndex];
+					Mapping.Add(PropertyDesc.Name, PortTerms[TermIndex].GetName());
+				}
+
+				TermLayerHandle = FParamStack::MakeRemappedLayer(InstanceData.IntermediatesLayer, Mapping);
+			}
+			
 			// Supply the I/O terms that the port needs
 			FScheduleTermContext TermContext(InScheduleContext, InstanceData.PortTermLayers[TaskIndex]);
 			CDO->Run(TermContext);
diff --git a/Engine/Plugins/Experimental/AnimNext/Source/AnimNext/Private/Scheduler/AnimNextSchedulerEntry.cpp b/Engine/Plugins/Experimental/AnimNext/Source/AnimNext/Private/Scheduler/AnimNextSchedulerEntry.cpp
--- a/Engine/Plugins/Experimental/AnimNext/Source/AnimNext/Private/Scheduler/AnimNextSchedulerEntry.cpp
+++ b/Engine/Plugins/Experimental/AnimNext/Source/AnimNext/Private/Scheduler/AnimNextSchedulerEntry.cpp
@@ -9,9 +9,8 @@
 
 DEFINE_STAT(STAT_AnimNext_InitializeEntry);
 
-FAnimNextSchedulerEntry::FAnimNextSchedulerEntry(const UAnimNextSchedule* InSchedule, UObject* InObject, UE::AnimNext::FScheduleHandle InHandle, const TMap<FName, FAnimNextParameterCollection>& InUserScopes, EAnimNextScheduleInitMethod InInitMethod)
+FAnimNextSchedulerEntry::FAnimNextSchedulerEntry(const UAnimNextSchedule* InSchedule, UObject* InObject, UE::AnimNext::FScheduleHandle InHandle, EAnimNextScheduleInitMethod InInitMethod)
 	: Schedule(InSchedule)
-	, UserScopes(InUserScopes)
 	, WeakObject(InObject)
 	, Handle(InHandle)
 	, Context(InSchedule, this)
@@ -25,7 +24,7 @@ FAnimNextSchedulerEntry::~FAnimNextSchedulerEntry()
 	Invalidate();
 }
 
-void FAnimNextSchedulerEntry::Initialize()
+void FAnimNextSchedulerEntry::Initialize(TUniqueFunction<void(const UE::AnimNext::FScheduleContext&)>&& InInitializeCallback)
 {
 	SCOPE_CYCLE_COUNTER(STAT_AnimNext_InitializeEntry);
 	
@@ -181,8 +180,15 @@ void FAnimNextSchedulerEntry::Initialize()
 			TickFunction->RegisterTickFunction(Level);
 		}
 
+		LazyAllocateInstanceData();
+
 		RunState = ERunState::RunningInitialUpdate;
 
+		if(InInitializeCallback)
+		{
+			InInitializeCallback(Context);
+		}
+
 		// Just pause now if we arent needing an initial update
 		if(InitMethod == EAnimNextScheduleInitMethod::None)
 		{
@@ -194,7 +200,6 @@ void FAnimNextSchedulerEntry::Initialize()
 			// generate an output pose, as these worlds never tick
 			if(World->WorldType == EWorldType::EditorPreview)
 			{
-				LazyAllocateInstanceData();
 				FScheduleContext::AttachToCurrentThread(Context);
 				FScheduleTickFunction::RunSchedule(Schedule->Instructions);
 				FScheduleContext::DetachFromCurrentThread();
@@ -241,7 +246,6 @@ void FAnimNextSchedulerEntry::Invalidate()
 	Schedule = nullptr;
 	WeakObject = nullptr;
 	TargetObjects.Reset();
-	UserScopes.Empty();
 	Handle.Invalidate();
 	Context.InstanceData.Reset();
 }
@@ -283,6 +287,6 @@ void FAnimNextSchedulerEntry::LazyAllocateInstanceData()
 	// Allocate instance data if required
 	if (!Context.InstanceData.IsValid())
 	{
-		Context.InstanceData = MakeUnique<FScheduleInstanceData>(Context, Schedule, Handle, this, MoveTemp(UserScopes));
+		Context.InstanceData = MakeUnique<FScheduleInstanceData>(Context, Schedule, Handle, this);
 	}
 }
\ No newline at end of file
diff --git a/Engine/Plugins/Experimental/AnimNext/Source/AnimNext/Private/Scheduler/AnimNextSchedulerEntry.h b/Engine/Plugins/Experimental/AnimNext/Source/AnimNext/Private/Scheduler/AnimNextSchedulerEntry.h
--- a/Engine/Plugins/Experimental/AnimNext/Source/AnimNext/Private/Scheduler/AnimNextSchedulerEntry.h
+++ b/Engine/Plugins/Experimental/AnimNext/Source/AnimNext/Private/Scheduler/AnimNextSchedulerEntry.h
@@ -23,11 +23,11 @@ struct FAnimNextSchedulerEntry
 	GENERATED_BODY()
 
 	FAnimNextSchedulerEntry() = default;
-	FAnimNextSchedulerEntry(const UAnimNextSchedule* InSchedule, UObject* InObject, UE::AnimNext::FScheduleHandle InHandle, const TMap<FName, FAnimNextParameterCollection>& InUserScopes, EAnimNextScheduleInitMethod InInitMethod);
+	FAnimNextSchedulerEntry(const UAnimNextSchedule* InSchedule, UObject* InObject, UE::AnimNext::FScheduleHandle InHandle, EAnimNextScheduleInitMethod InInitMethod);
 	~FAnimNextSchedulerEntry();
 
 	// Setup the entry
-	void Initialize();
+	void Initialize(TUniqueFunction<void(const UE::AnimNext::FScheduleContext&)>&& InInitializeCallback);
 
 	// Used for pooling
 	void Invalidate();
@@ -44,11 +44,6 @@ struct FAnimNextSchedulerEntry
 	UPROPERTY(Transient)
 	TObjectPtr<const UAnimNextSchedule> Schedule = nullptr;
 
-	// User scopes are copied into this entry on construction, but moved out later into instance data
-	// So will be invalid here after first run
-	UPROPERTY(Transient)
-	TMap<FName, FAnimNextParameterCollection> UserScopes;
-
 	// Object this entry is bound to, valid only during schedule execution
 	UPROPERTY(Transient)
 	TObjectPtr<UObject> ResolvedObject = nullptr;
@@ -75,6 +70,9 @@ struct FAnimNextSchedulerEntry
 	// Pre-allocated graph of tick functions
 	TArray<TUniquePtr<UE::AnimNext::FScheduleTickFunction>> TickFunctions;
 
+	// Current delta time, updated each time the schedule runs
+	float DeltaTime = 0.0f;
+
 	enum class ERunState
 	{
 		None,
diff --git a/Engine/Plugins/Experimental/AnimNext/Source/AnimNext/Private/Scheduler/AnimNextSchedulerWorldSubsystem.cpp b/Engine/Plugins/Experimental/AnimNext/Source/AnimNext/Private/Scheduler/AnimNextSchedulerWorldSubsystem.cpp
--- a/Engine/Plugins/Experimental/AnimNext/Source/AnimNext/Private/Scheduler/AnimNextSchedulerWorldSubsystem.cpp
+++ b/Engine/Plugins/Experimental/AnimNext/Source/AnimNext/Private/Scheduler/AnimNextSchedulerWorldSubsystem.cpp
@@ -97,7 +97,7 @@ void UAnimNextSchedulerWorldSubsystem::FlushPendingActions()
 	}
 }
 
-UE::AnimNext::FScheduleHandle UAnimNextSchedulerWorldSubsystem::AcquireHandle(UObject* InObject, UAnimNextSchedule* InSchedule, const TMap<FName, FAnimNextParameterCollection>& InUserScopes, EAnimNextScheduleInitMethod InInitMethod)
+UE::AnimNext::FScheduleHandle UAnimNextSchedulerWorldSubsystem::AcquireHandle(UObject* InObject, UAnimNextSchedule* InSchedule, EAnimNextScheduleInitMethod InInitMethod, TUniqueFunction<void(const UE::AnimNext::FScheduleContext&)>&& InInitializeCallback)
 {
 	using namespace UE::AnimNext;
 
@@ -111,17 +111,17 @@ UE::AnimNext::FScheduleHandle UAnimNextSchedulerWorldSubsystem::AcquireHandle(UO
 		Handle.Index = FreeEntryIndices.Last();
 		Handle.SerialNumber = ++GEntrySerialNumber;
 		FreeEntryIndices.Pop(false);
-		new (Entries[Handle.Index].Get()) FAnimNextSchedulerEntry(InSchedule, InObject, Handle, InUserScopes, InInitMethod);
+		new (Entries[Handle.Index].Get()) FAnimNextSchedulerEntry(InSchedule, InObject, Handle, InInitMethod);
 	}
 	// Otherwise append a new entry
 	else
 	{
 		Handle.Index = Entries.Num();
 		Handle.SerialNumber = ++GEntrySerialNumber;
-		Entries.Emplace(MakeUnique<FAnimNextSchedulerEntry>(InSchedule, InObject, Handle, InUserScopes, InInitMethod));
+		Entries.Emplace(MakeUnique<FAnimNextSchedulerEntry>(InSchedule, InObject, Handle, InInitMethod));
 	}
 
-	Entries[Handle.Index]->Initialize();
+	Entries[Handle.Index]->Initialize(MoveTemp(InInitializeCallback));
 
 	// Skip 'invalid' 0 serial number
 	if (GEntrySerialNumber == 0)
diff --git a/Engine/Plugins/Experimental/AnimNext/Source/AnimNext/Private/Scheduler/AnimNextSchedulerWorldSubsystem.h b/Engine/Plugins/Experimental/AnimNext/Source/AnimNext/Private/Scheduler/AnimNextSchedulerWorldSubsystem.h
--- a/Engine/Plugins/Experimental/AnimNext/Source/AnimNext/Private/Scheduler/AnimNextSchedulerWorldSubsystem.h
+++ b/Engine/Plugins/Experimental/AnimNext/Source/AnimNext/Private/Scheduler/AnimNextSchedulerWorldSubsystem.h
@@ -64,7 +64,7 @@ class UAnimNextSchedulerWorldSubsystem : public UWorldSubsystem
 	void FlushPendingActions();
 
 	// Acquire a handle that binds a schedule with the supplied parameters
-	UE::AnimNext::FScheduleHandle AcquireHandle(UObject* InObject, UAnimNextSchedule* InSchedule, const TMap<FName, FAnimNextParameterCollection>& InUserScopes, EAnimNextScheduleInitMethod InInitMethod);
+	UE::AnimNext::FScheduleHandle AcquireHandle(UObject* InObject, UAnimNextSchedule* InSchedule, EAnimNextScheduleInitMethod InInitMethod, TUniqueFunction<void(const UE::AnimNext::FScheduleContext&)>&& InInitializeCallback);
 
 	// Release an already acquired handle
 	// The full release of the binding referenced by the handle will be deferred after this call is made
diff --git a/Engine/Plugins/Experimental/AnimNext/Source/AnimNext/Private/Scheduler/ScheduleContext.cpp b/Engine/Plugins/Experimental/AnimNext/Source/AnimNext/Private/Scheduler/ScheduleContext.cpp
--- a/Engine/Plugins/Experimental/AnimNext/Source/AnimNext/Private/Scheduler/ScheduleContext.cpp
+++ b/Engine/Plugins/Experimental/AnimNext/Source/AnimNext/Private/Scheduler/ScheduleContext.cpp
@@ -40,4 +40,10 @@ UObject* FScheduleContext::GetContextObject() const
 	}
 }
 
+float FScheduleContext::GetDeltaTime() const 
+{ 
+	check(Entry != nullptr);
+	return Entry->DeltaTime; 
+}
+
 }
diff --git a/Engine/Plugins/Experimental/AnimNext/Source/AnimNext/Private/Scheduler/ScheduleContext.h b/Engine/Plugins/Experimental/AnimNext/Source/AnimNext/Private/Scheduler/ScheduleContext.h
--- a/Engine/Plugins/Experimental/AnimNext/Source/AnimNext/Private/Scheduler/ScheduleContext.h
+++ b/Engine/Plugins/Experimental/AnimNext/Source/AnimNext/Private/Scheduler/ScheduleContext.h
@@ -57,6 +57,9 @@ public:
 		return *InstanceData.Get(); 
 	}
 
+	// Get the current delta time - only valid if we are within a running schedule
+	float GetDeltaTime() const;
+
 	// Get the current object context that the schedule is running in 
 	UObject* GetContextObject() const;
 
diff --git a/Engine/Plugins/Experimental/AnimNext/Source/AnimNext/Private/Scheduler/ScheduleInstanceData.cpp b/Engine/Plugins/Experimental/AnimNext/Source/AnimNext/Private/Scheduler/ScheduleInstanceData.cpp
--- a/Engine/Plugins/Experimental/AnimNext/Source/AnimNext/Private/Scheduler/ScheduleInstanceData.cpp
+++ b/Engine/Plugins/Experimental/AnimNext/Source/AnimNext/Private/Scheduler/ScheduleInstanceData.cpp
@@ -13,10 +13,9 @@ DEFINE_STAT(STAT_AnimNext_CreateInstanceData);
 namespace UE::AnimNext
 {
 
-FScheduleInstanceData::FScheduleInstanceData(const FScheduleContext& InScheduleContext, const UAnimNextSchedule* InSchedule, FScheduleHandle InHandle, FAnimNextSchedulerEntry* InCurrentEntry, TMap<FName, FAnimNextParameterCollection>&& InUserScopes)
+FScheduleInstanceData::FScheduleInstanceData(const FScheduleContext& InScheduleContext, const UAnimNextSchedule* InSchedule, FScheduleHandle InHandle, FAnimNextSchedulerEntry* InCurrentEntry)
 	: Handle(InHandle)
 	, Entry(InCurrentEntry) 
-	, UserScopes(MoveTemp(InUserScopes))
 {
 	SCOPE_CYCLE_COUNTER(STAT_AnimNext_CreateInstanceData);
 
@@ -57,43 +56,9 @@ FScheduleInstanceData::FScheduleInstanceData(const FScheduleContext& InScheduleC
 	// Make a hosting layer for the intermediates
 	IntermediatesLayer = FParamStack::MakeReferenceLayer(IntermediatesData);
 
-	// Setup remapped intermediate data layers for graph and port tasks
-	GraphTermLayers.Reserve(InSchedule->Tasks.Num());
-	for (const FAnimNextScheduleGraphTask& Task : InSchedule->Tasks)
-	{
-		TConstArrayView<FScheduleTerm> Terms = Task.Graph->GetTerms();
-		check(Task.Terms.Num() == Terms.Num());
-
-		TMap<FName, FName> Mapping;
-		Mapping.Reserve(Task.Terms.Num());
-		for(int32 TermIndex = 0; TermIndex < Task.Terms.Num(); ++TermIndex)
-		{
-			uint32 IntermediateTermIndex = Task.Terms[TermIndex];
-			const FPropertyBagPropertyDesc& PropertyDesc = IntermediatesData.GetPropertyBagStruct()->GetPropertyDescs()[IntermediateTermIndex];
-			Mapping.Add(PropertyDesc.Name, Terms[TermIndex].GetName());
-		}
-
-		GraphTermLayers.Add(FParamStack::MakeRemappedLayer(IntermediatesLayer, Mapping));
-	}
-
-	PortTermLayers.Reserve(InSchedule->Ports.Num());
-	for (const FAnimNextSchedulePortTask& PortTask : InSchedule->Ports)
-	{
-		UAnimNextSchedulePort* CDO = PortTask.Port->GetDefaultObject<UAnimNextSchedulePort>();
-		TConstArrayView<FScheduleTerm> Terms = CDO->GetTerms();
-		check(PortTask.Terms.Num() == Terms.Num());
-
-		TMap<FName, FName> Mapping;
-		Mapping.Reserve(PortTask.Terms.Num());
-		for(int32 TermIndex = 0; TermIndex < PortTask.Terms.Num(); ++TermIndex)
-		{
-			uint32 IntermediateTermIndex = PortTask.Terms[TermIndex];
-			const FPropertyBagPropertyDesc& PropertyDesc = IntermediatesData.GetPropertyBagStruct()->GetPropertyDescs()[IntermediateTermIndex];
-			Mapping.Add(PropertyDesc.Name, Terms[TermIndex].GetName());
-		}
-
-		PortTermLayers.Add(FParamStack::MakeRemappedLayer(IntermediatesLayer, Mapping));
-	}
+	// Resize remapped intermediate data layers for graph and port tasks, they will be allocated lazily later
+	GraphTermLayers.SetNum(InSchedule->Tasks.Num());
+	PortTermLayers.SetNum(InSchedule->Ports.Num());
 }
 
 void FScheduleInstanceData::AddReferencedObjects(FReferenceCollector& Collector)
diff --git a/Engine/Plugins/Experimental/AnimNext/Source/AnimNext/Private/Scheduler/ScheduleInstanceData.h b/Engine/Plugins/Experimental/AnimNext/Source/AnimNext/Private/Scheduler/ScheduleInstanceData.h
--- a/Engine/Plugins/Experimental/AnimNext/Source/AnimNext/Private/Scheduler/ScheduleInstanceData.h
+++ b/Engine/Plugins/Experimental/AnimNext/Source/AnimNext/Private/Scheduler/ScheduleInstanceData.h
@@ -24,7 +24,7 @@ namespace UE::AnimNext
 // Host for all data needed to run a schedule instance
 struct FScheduleInstanceData : public FGCObject
 {
-	FScheduleInstanceData(const FScheduleContext& InScheduleContext, const UAnimNextSchedule* InSchedule, FScheduleHandle InHandle, FAnimNextSchedulerEntry* InCurrentEntry, TMap<FName, FAnimNextParameterCollection>&& InDynamicScopes);
+	FScheduleInstanceData(const FScheduleContext& InScheduleContext, const UAnimNextSchedule* InSchedule, FScheduleHandle InHandle, FAnimNextSchedulerEntry* InCurrentEntry);
 
 	// FGCObject interface
 	virtual void AddReferencedObjects(FReferenceCollector& Collector) override;
@@ -64,6 +64,9 @@ struct FScheduleInstanceData : public FGCObject
 	// Param stacks required to run the schedule (one per task that requires a stack)
 	TArray<TSharedPtr<FParamStack>> ParamStacks;
 
+	// User handles initialized at startup, always pushed
+	TArray<FParamStackLayerHandle> StaticUserHandles;
+
 	// Graph instance data for each graph task
 	TArray<FAnimNextGraphInstance> GraphInstanceData;
 
diff --git a/Engine/Plugins/Experimental/AnimNext/Source/AnimNext/Private/Scheduler/ScheduleTickFunction.cpp b/Engine/Plugins/Experimental/AnimNext/Source/AnimNext/Private/Scheduler/ScheduleTickFunction.cpp
--- a/Engine/Plugins/Experimental/AnimNext/Source/AnimNext/Private/Scheduler/ScheduleTickFunction.cpp
+++ b/Engine/Plugins/Experimental/AnimNext/Source/AnimNext/Private/Scheduler/ScheduleTickFunction.cpp
@@ -18,6 +18,7 @@ void FScheduleBeginTickFunction::ExecuteTick(float DeltaTime, ELevelTick TickTyp
 	Entry.LazyAllocateInstanceData();
 
 	Entry.ResolvedObject = Entry.WeakObject.Get();
+	Entry.DeltaTime = DeltaTime;
 }
 
 FString FScheduleBeginTickFunction::DiagnosticMessage()
diff --git a/Engine/Plugins/Experimental/AnimNext/Source/AnimNext/Private/Scheduler/Scheduler.cpp b/Engine/Plugins/Experimental/AnimNext/Source/AnimNext/Private/Scheduler/Scheduler.cpp
--- a/Engine/Plugins/Experimental/AnimNext/Source/AnimNext/Private/Scheduler/Scheduler.cpp
+++ b/Engine/Plugins/Experimental/AnimNext/Source/AnimNext/Private/Scheduler/Scheduler.cpp
@@ -50,7 +50,7 @@ void FScheduler::Destroy()
 	FWorldDelegates::OnWorldPreActorTick.Remove(Impl.OnWorldPreActorTickHandle);
 }
 
-FScheduleHandle FScheduler::AcquireHandle(UObject* InObject, UAnimNextSchedule* InSchedule, const TMap<FName, FAnimNextParameterCollection>& InUserScopes, EAnimNextScheduleInitMethod InInitMethod)
+FScheduleHandle FScheduler::AcquireHandle(UObject* InObject, UAnimNextSchedule* InSchedule, EAnimNextScheduleInitMethod InInitMethod, TUniqueFunction<void(const FScheduleContext&)>&& InInitializeCallback)
 {
 	FScheduleHandle Handle;
 
@@ -79,7 +79,7 @@ FScheduleHandle FScheduler::AcquireHandle(UObject* InObject, UAnimNextSchedule*
 		return FScheduleHandle();
 	}
 
-	return Subsystem->AcquireHandle(InObject, InSchedule, InUserScopes, InInitMethod);
+	return Subsystem->AcquireHandle(InObject, InSchedule, InInitMethod, MoveTemp(InInitializeCallback));
 }
 
 void FScheduler::ReleaseHandle(UObject* InObject, FScheduleHandle& InHandle)
diff --git a/Engine/Plugins/Experimental/AnimNext/Source/AnimNext/Private/Scheduler/Scheduler.h b/Engine/Plugins/Experimental/AnimNext/Source/AnimNext/Private/Scheduler/Scheduler.h
--- a/Engine/Plugins/Experimental/AnimNext/Source/AnimNext/Private/Scheduler/Scheduler.h
+++ b/Engine/Plugins/Experimental/AnimNext/Source/AnimNext/Private/Scheduler/Scheduler.h
@@ -3,6 +3,7 @@
 #pragma once
 
 #include "CoreMinimal.h"
+#include "ScheduleContext.h"
 #include "ScheduleHandle.h"
 
 enum class EAnimNextScheduleInitMethod : uint8;
@@ -32,8 +33,10 @@ struct FScheduler
 	// Shut down the scheduler system
 	static void Destroy();
 
-	// Acquire a handle that binds a schedule with the supplied parameters
-	static FScheduleHandle AcquireHandle(UObject* InObject, UAnimNextSchedule* InSchedule, const TMap<FName, FAnimNextParameterCollection>& InUserScopes, EAnimNextScheduleInitMethod InInitMethod);
+	// Acquire a handle that binds a schedule
+	// Initial parameter binding can be achieved via InitializeCallback, which is called once the schedule's data
+	// structures have been set up, but before it is first run.
+	static FScheduleHandle AcquireHandle(UObject* InObject, UAnimNextSchedule* InSchedule, EAnimNextScheduleInitMethod InInitMethod, TUniqueFunction<void(const FScheduleContext&)>&& InitializeCallback = nullptr);
 
 	// Release an already acquired handle
 	// The full release of the binding referenced by the handle map be deferred after this call is made
diff --git a/Engine/Plugins/Experimental/AnimNext/Source/AnimNext/Public/Param/ParamEntry.h b/Engine/Plugins/Experimental/AnimNext/Source/AnimNext/Public/Param/ParamEntry.h
--- a/Engine/Plugins/Experimental/AnimNext/Source/AnimNext/Public/Param/ParamEntry.h
+++ b/Engine/Plugins/Experimental/AnimNext/Source/AnimNext/Public/Param/ParamEntry.h
@@ -5,6 +5,8 @@
 #include "CoreMinimal.h"
 #include "Param/ParamTypeHandle.h"
 
+class UAnimNextComponent;
+
 namespace UE::AnimNext
 {
 	struct FParamStack;
@@ -49,6 +51,7 @@ private:
 	friend struct UE::AnimNext::FInstancedPropertyBagLayer;
 	friend struct UE::AnimNext::FParamStackLayerHandle;
 	friend struct UE::AnimNext::FRemappedLayer;
+	friend class ::UAnimNextComponent;
 
 	ANIMNEXT_API FParamEntry(const FParamTypeHandle& InTypeHandle, TArrayView<uint8> InData, bool bInIsReference, bool bInIsMutable);
 
diff --git a/Engine/Plugins/Experimental/AnimNext/Source/AnimNext/Public/Param/ParamStack.h b/Engine/Plugins/Experimental/AnimNext/Source/AnimNext/Public/Param/ParamStack.h
--- a/Engine/Plugins/Experimental/AnimNext/Source/AnimNext/Public/Param/ParamStack.h
+++ b/Engine/Plugins/Experimental/AnimNext/Source/AnimNext/Public/Param/ParamStack.h
@@ -71,6 +71,7 @@ struct FParamStack
 	friend struct ::FAnimNextSchedulerEntry;
 	friend class ::FAnimationAnimNextRuntimeTest_GraphExecute;
 	friend class ::FAnimationAnimNextRuntimeTest_GraphExecuteLatent;
+	friend class ::UAnimNextComponent;
 
 	template<typename AllocatorType> 
 	friend struct TLayerBuilder;
diff --git a/Engine/Plugins/Experimental/AnimNext/Source/AnimNext/Public/Scheduler/IAnimNextScheduleTermInterface.h b/Engine/Plugins/Experimental/AnimNext/Source/AnimNext/Public/Scheduler/IAnimNextScheduleTermInterface.h
--- a/Engine/Plugins/Experimental/AnimNext/Source/AnimNext/Public/Scheduler/IAnimNextScheduleTermInterface.h
+++ b/Engine/Plugins/Experimental/AnimNext/Source/AnimNext/Public/Scheduler/IAnimNextScheduleTermInterface.h
@@ -14,16 +14,16 @@ namespace UE::AnimNext
 	struct FScheduleInstanceData;
 }
 
-namespace UE::AnimNext
-{
-
-enum class EScheduleTermDirection
+UENUM()
+enum class EScheduleTermDirection : uint8
 {
 	Input,
 	Output,
-	IO,
 };
 
+namespace UE::AnimNext
+{
+
 struct FScheduleTerm : FParamDefinition
 {
 	FScheduleTerm(FName InName, const FAnimNextParamType& InType, EScheduleTermDirection InDirection)
@@ -54,6 +54,7 @@ class IAnimNextScheduleTermInterface
 private:
 	friend class UAnimNextSchedule;
 	friend struct UE::AnimNext::FScheduleInstanceData;
+	friend struct FAnimNextSchedulePortTask;
 
 	// Get the terms that define the data dependencies of a schedule element
 	virtual TConstArrayView<UE::AnimNext::FScheduleTerm> GetTerms() const PURE_VIRTUAL(IAnimNextScheduleTermInterface::GetTerms, return {};)
diff --git a/Engine/Plugins/Experimental/AnimNext/Source/AnimNextEditor/Private/AnimNextEditorModule.cpp b/Engine/Plugins/Experimental/AnimNext/Source/AnimNextEditor/Private/AnimNextEditorModule.cpp
--- a/Engine/Plugins/Experimental/AnimNext/Source/AnimNextEditor/Private/AnimNextEditorModule.cpp
+++ b/Engine/Plugins/Experimental/AnimNext/Source/AnimNextEditor/Private/AnimNextEditorModule.cpp
@@ -143,18 +143,6 @@ class FModule : public IModule
 			UAnimNextGraph* Graph = CastChecked<UAnimNextGraph>(InAsset);
 			UAnimNextGraph_EditorData* EditorData = UncookedOnly::FUtils::GetEditorData(Graph);
 
-			// Open default model for this graph asset as well as the graph-list tab (deferred so we dont modify opened document records during iteration)
-			ExecuteOnGameThread(UE_SOURCE_LOCATION, [WeakEditor = TWeakPtr<FWorkspaceEditor>(InEditor), WeakEditorData = TWeakObjectPtr<UAnimNextGraph_EditorData>(EditorData)]()
-			{
-				if(WeakEditorData.Get() && WeakEditor.IsValid())
-				{
-					if(UObject* EditorObject = WeakEditorData->GetEditorObjectForRigVMGraph(WeakEditorData->GetRigVMClient()->GetDefaultModel()))
-					{
-						WeakEditor.Pin()->OpenDocument(EditorObject, FDocumentTracker::EOpenDocumentCause::OpenNewDocument);
-					}
-				}
-			});
-
 			return SNew(SAnimNextGraphView, EditorData)
 				.OnOpenGraph_Lambda([WeakEditor = TWeakPtr<FWorkspaceEditor>(InEditor)](URigVMGraph* InGraph)
 				{
