AnimNext: Enable migration of default values when compiling parameter blocks

Prior to this, default values would be lost on each compile

diff --git a/Engine/Plugins/Experimental/AnimNext/Source/AnimNextUncookedOnly/Private/UncookedOnlyUtils.cpp b/Engine/Plugins/Experimental/AnimNext/Source/AnimNextUncookedOnly/Private/UncookedOnlyUtils.cpp
--- a/Engine/Plugins/Experimental/AnimNext/Source/AnimNextUncookedOnly/Private/UncookedOnlyUtils.cpp
+++ b/Engine/Plugins/Experimental/AnimNext/Source/AnimNextUncookedOnly/Private/UncookedOnlyUtils.cpp
@@ -732,9 +732,6 @@ void FUtils::CompileStruct(UAnimNextParameterBlock* InParameterBlock)
 
 	TGuardValue<bool> CompilingGuard(EditorData->bIsCompiling, true);
 
-	FInstancedPropertyBag& PropertyBag = InParameterBlock->PropertyBag;
-	PropertyBag.Reset();
-
 	TArray<FPropertyBagPropertyDesc> PropertyDescs;
 	PropertyDescs.Reserve(EditorData->Entries.Num());
 	
@@ -751,10 +748,32 @@ void FUtils::CompileStruct(UAnimNextParameterBlock* InParameterBlock)
 		}
 	}
 
-	// Bulk add to the bag
-	PropertyBag.AddProperties(PropertyDescs);
+	if(PropertyDescs.Num() > 0)
+	{
+		// find any existing IDs for old properties with name-matching
+		for(FPropertyBagPropertyDesc& NewDesc : PropertyDescs)
+		{
+			if(InParameterBlock->PropertyBag.GetPropertyBagStruct())
+			{
+				for(const FPropertyBagPropertyDesc& ExistingDesc : InParameterBlock->PropertyBag.GetPropertyBagStruct()->GetPropertyDescs())
+				{
+					if(ExistingDesc.Name == NewDesc.Name)
+					{
+						NewDesc.ID = ExistingDesc.ID;
+						break;
+					}
+				}
+			}
+		}
 
-	// TODO: Now copy over defaults for those properties that need it (literals)
+		// Create new property bag and migrate
+		const UPropertyBag* NewBagStruct = UPropertyBag::GetOrCreateFromDescs(PropertyDescs);
+		InParameterBlock->PropertyBag.MigrateToNewBagStruct(NewBagStruct);
+	}
+	else
+	{
+		InParameterBlock->PropertyBag.Reset();
+	}
 
 	EditorData->bStructRecompilationRequired = false;
 }
