AnimNext: Fix crash with structure details viewing deleted memory on compile/refresh of parameter blocks

Adds a new IStructureDataProvider that views FInstancePropertyBag directly, so details panels can reflect correct data even if the bag is invalidated.
Also set the structure property node in SSIngleProperty to use the EPropertyNodeFlags::RequiresValidation flag.

diff --git a/Engine/Plugins/Experimental/AnimNext/Source/AnimNextEditor/Private/Param/ParameterBlockParameterCustomization.cpp b/Engine/Plugins/Experimental/AnimNext/Source/AnimNextEditor/Private/Param/ParameterBlockParameterCustomization.cpp
--- a/Engine/Plugins/Experimental/AnimNext/Source/AnimNextEditor/Private/Param/ParameterBlockParameterCustomization.cpp
+++ b/Engine/Plugins/Experimental/AnimNext/Source/AnimNextEditor/Private/Param/ParameterBlockParameterCustomization.cpp
@@ -7,6 +7,7 @@
 #include "PropertyBagDetails.h"
 #include "EdGraphSchema_K2.h"
 #include "EditorUtils.h"
+#include "InstancedPropertyBagStructureDataProvider.h"
 #include "UncookedOnlyUtils.h"
 #include "Param/ParamType.h"
 #include "PropertyHandle.h"
@@ -44,19 +45,27 @@ void FParameterBlockParameterCustomization::CustomizeDetails(IDetailLayoutBuilde
 		{
 			if (UAnimNextParameterBlock* ReferencedBlock = UE::AnimNext::UncookedOnly::FUtils::GetBlock(EditorData))
 			{
-				TSharedPtr<FStructOnScope> StructOnScope = MakeShareable(new FStructOnScope(ReferencedBlock->GetPropertyBag().GetPropertyBagStruct(), (uint8*)ReferencedBlock->GetPropertyBag().GetValue().GetMemory()));
 				FSinglePropertyParams SinglePropertyArgs;
 				const FName ParameterName = BlockParam->GetParameterName();
-				IDetailPropertyRow* DetailPropertyRow = DefaultValueCategory.AddExternalStructureProperty(StructOnScope, ParameterName);
-				if (DetailPropertyRow)
+				TArray<IDetailPropertyRow*> DetailPropertyRows;
+
+				// FIXME: As AddExternalStructureProperty does not support IStructureDataProvider, we add all properties and only show the one we need. Icky.
+				DefaultValueCategory.AddAllExternalStructureProperties(MakeShared<FInstancePropertyBagStructureDataProvider>(ReferencedBlock->GetPropertyBag()), EPropertyLocation::Default, &DetailPropertyRows);
+				for (IDetailPropertyRow* DetailPropertyRow : DetailPropertyRows)
 				{
-					if (TSharedPtr<IPropertyHandle> Handle = DetailPropertyRow->GetPropertyHandle(); Handle.IsValid())
+					if(DetailPropertyRow->GetPropertyHandle()->GetProperty()->GetFName() != ParameterName)
+					{
+						DetailPropertyRow->Visibility(EVisibility::Collapsed);
+					}
+					else
 					{
-						Handle->SetOnChildPropertyValuePreChange(FSimpleDelegate::CreateLambda([this, ReferencedBlock]()
+						if (TSharedPtr<IPropertyHandle> Handle = DetailPropertyRow->GetPropertyHandle(); Handle.IsValid())
+						{
+							Handle->SetOnChildPropertyValuePreChange(FSimpleDelegate::CreateLambda([this, ReferencedBlock]()
 							{
 								ReferencedBlock->Modify(); // needed to enable the transaction when we modify the PropertyBag
 							}));
-						Handle->SetOnPropertyValueChanged(FSimpleDelegate::CreateLambda([this, ReferencedBlock, ParameterName]()
+							Handle->SetOnPropertyValueChanged(FSimpleDelegate::CreateLambda([this, ReferencedBlock, ParameterName]()
 							{
 								if (UAnimNextParameterBlock_EditorData* EditorData = Cast<UAnimNextParameterBlock_EditorData>(ReferencedBlock->EditorData))
 								{
@@ -67,6 +76,7 @@ void FParameterBlockParameterCustomization::CustomizeDetails(IDetailLayoutBuilde
 									}
 								}
 							}));
+						}
 					}
 				}
 			}
diff --git a/Engine/Plugins/Experimental/AnimNext/Source/AnimNextEditor/Private/Param/SParameterBlockView.cpp b/Engine/Plugins/Experimental/AnimNext/Source/AnimNextEditor/Private/Param/SParameterBlockView.cpp
--- a/Engine/Plugins/Experimental/AnimNext/Source/AnimNextEditor/Private/Param/SParameterBlockView.cpp
+++ b/Engine/Plugins/Experimental/AnimNext/Source/AnimNextEditor/Private/Param/SParameterBlockView.cpp
@@ -13,6 +13,7 @@
 #include "ISourceControlModule.h"
 #include "ISourceControlProvider.h"
 #include "EditorUtils.h"
+#include "InstancedPropertyBagStructureDataProvider.h"
 #include "UncookedOnlyUtils.h"
 #include "ParameterBlockViewMenuContext.h"
 #include "PropertyBagDetails.h"
@@ -792,6 +793,7 @@ class SParameterBlockViewRow : public SMultiColumnTableRow<TSharedRef<FParameter
 				.AutoWidth()
 				[
 					SNew(SExpanderArrow, SharedThis(this))
+					.Visibility_Lambda([this]() { return DoesItemHaveChildren() ? EVisibility::Visible : EVisibility::Collapsed; })
 				]
 				+SHorizontalBox::Slot()
 				.HAlign(HAlign_Left)
@@ -894,13 +896,12 @@ class SParameterBlockViewRow : public SMultiColumnTableRow<TSharedRef<FParameter
 
 						if (PropertyBag.FindPropertyDescByName(ParameterName))
 						{
-							TSharedPtr<FStructOnScope> StructOnScope = MakeShareable(new FStructOnScope(PropertyBag.GetPropertyBagStruct(), (uint8*)PropertyBag.GetValue().GetMemory()));
 							FSinglePropertyParams SinglePropertyArgs;
 							SinglePropertyArgs.NamePlacement = EPropertyNamePlacement::Hidden;
 							SinglePropertyArgs.NotifyHook = this;
 
 							FPropertyEditorModule& PropertyEditorModule = FModuleManager::Get().LoadModuleChecked<FPropertyEditorModule>("PropertyEditor");
-							const TSharedPtr<ISinglePropertyView> SingleStructPropertyView = PropertyEditorModule.CreateSingleProperty(MakeShared<FStructOnScopeStructureDataProvider>(StructOnScope), ParameterName, SinglePropertyArgs);
+							const TSharedPtr<ISinglePropertyView> SingleStructPropertyView = PropertyEditorModule.CreateSingleProperty(MakeShared<FInstancePropertyBagStructureDataProvider>(PropertyBag), ParameterName, SinglePropertyArgs);
 							if (SingleStructPropertyView.IsValid())
 							{
 								ColumnWidget = SingleStructPropertyView.ToSharedRef();
diff --git a/Engine/Plugins/Experimental/StructUtils/Source/StructUtilsEditor/Public/InstancedPropertyBagStructureDataProvider.h b/Engine/Plugins/Experimental/StructUtils/Source/StructUtilsEditor/Public/InstancedPropertyBagStructureDataProvider.h
new file mode 100644
--- /dev/null
+++ b/Engine/Plugins/Experimental/StructUtils/Source/StructUtilsEditor/Public/InstancedPropertyBagStructureDataProvider.h
@@ -0,0 +1,43 @@
+﻿// Copyright Epic Games, Inc. All Rights Reserved.
+
+#pragma once
+
+#include "IStructureDataProvider.h"
+#include "PropertyBag.h"
+
+//	FInstancePropertyBagStructureDataProvider
+//	Allows structure views to use FInstancedPropertyBag even if the bag layout changes.
+//	The caller needs to make sure that property bag outlives the property view widget.
+class FInstancePropertyBagStructureDataProvider : public IStructureDataProvider
+{
+public:
+	FInstancePropertyBagStructureDataProvider(FInstancedPropertyBag& InPropertyBag)
+		: PropertyBag(InPropertyBag)
+	{
+	}
+
+	virtual bool IsValid() const override
+	{
+		return PropertyBag.IsValid();
+	};
+	
+	virtual const UStruct* GetBaseStructure() const override
+	{
+		return PropertyBag.IsValid() ? PropertyBag.GetPropertyBagStruct() : nullptr;
+	}
+	
+	virtual void GetInstances(TArray<TSharedPtr<FStructOnScope>>& OutInstances, const UStruct* ExpectedBaseStructure) const override
+	{
+		if (PropertyBag.IsValid())
+		{
+			const UStruct* Struct = PropertyBag.GetPropertyBagStruct();
+			if (ExpectedBaseStructure && Struct && Struct == ExpectedBaseStructure)
+			{
+				OutInstances.Add(MakeShared<FStructOnScope>(Struct, PropertyBag.GetMutableValue().GetMemory()));
+			}
+		}
+	}
+
+protected:
+	FInstancedPropertyBag& PropertyBag;
+};
\ No newline at end of file
diff --git a/Engine/Source/Editor/PropertyEditor/Private/SSingleProperty.cpp b/Engine/Source/Editor/PropertyEditor/Private/SSingleProperty.cpp
--- a/Engine/Source/Editor/PropertyEditor/Private/SSingleProperty.cpp
+++ b/Engine/Source/Editor/PropertyEditor/Private/SSingleProperty.cpp
@@ -157,6 +157,7 @@ void SSingleProperty::SetStruct(const TSharedPtr<IStructureDataProvider>& InStru
 	if (!RootPropertyNode.IsValid() || RootPropertyNode->GetPropertyType() != FComplexPropertyNode::EPT_StandaloneStructure)
 	{
 		RootPropertyNode = MakeShareable(new FStructurePropertyNode);
+		RootPropertyNode->SetNodeFlags(EPropertyNodeFlags::RequiresValidation, true);
 	}
 
 	FStructurePropertyNode* RootStructPropertyNode = (FStructurePropertyNode*)RootPropertyNode.Get();
