Removed Built-In category from parameter block view

#rnx

diff --git a/Engine/Plugins/Experimental/AnimNext/Source/AnimNextEditor/Private/Param/SParameterBlockView.cpp b/Engine/Plugins/Experimental/AnimNext/Source/AnimNextEditor/Private/Param/SParameterBlockView.cpp
--- a/Engine/Plugins/Experimental/AnimNext/Source/AnimNextEditor/Private/Param/SParameterBlockView.cpp
+++ b/Engine/Plugins/Experimental/AnimNext/Source/AnimNextEditor/Private/Param/SParameterBlockView.cpp
@@ -208,7 +208,6 @@ void SParameterBlockView::Construct(const FArguments& InArgs, UAnimNextParameter
 
 	Categories = {
 		MakeShared<FParameterBlockViewEntry>(EParameterBlockCategoryType::Parameter, LOCTEXT("ParameterBlockParameterCategory", "Parameters")),
-		MakeShared<FParameterBlockViewEntry>(EParameterBlockCategoryType::BuiltIn, LOCTEXT("ParameterBlockBindingCategory", "Built-In")),
 		MakeShared<FParameterBlockViewEntry>(EParameterBlockCategoryType::Graph, LOCTEXT("ParameterBlockGraphCategory", "Graphs"))
 	};
 
@@ -448,11 +447,9 @@ void SParameterBlockView::RefreshFilter()
 	}
 
 	TSharedRef<FParameterBlockViewEntry> ParametersCategory = GetCategory(EParameterBlockCategoryType::Parameter);
-	TSharedRef<FParameterBlockViewEntry> BindingsCategory = GetCategory(EParameterBlockCategoryType::BuiltIn);
 	TSharedRef<FParameterBlockViewEntry> GraphsCategory = GetCategory(EParameterBlockCategoryType::Graph);
 
 	check(ParametersCategory->CategoryType == EParameterBlockCategoryType::Parameter);
-	check(BindingsCategory->CategoryType == EParameterBlockCategoryType::BuiltIn);
 	check(GraphsCategory->CategoryType == EParameterBlockCategoryType::Graph);
 
 	// add the entries as categories children
@@ -464,18 +461,7 @@ void SParameterBlockView::RefreshFilter()
 			const UAnimNextParameterBlockEntry* BlockEntry = Entry->WeakEntry.Get();
 			if (const IAnimNextParameterBlockParameterInterface* Parameter = Cast<IAnimNextParameterBlockParameterInterface>(BlockEntry))
 			{
-				UAnimNextParameterBlock* ReferencedBlock = UE::AnimNext::UncookedOnly::FUtils::GetBlock(EditorData);
-				FInstancedPropertyBag* PropertyBag = UE::AnimNext::UncookedOnly::FUtils::GetPropertyBag(ReferencedBlock);
-
-				const FName ParameterName = Parameter->GetParameterName();
-				if (PropertyBag->FindPropertyDescByName(ParameterName))
-				{
-					ParametersCategory->Children.Add(Entry);
-				}
-				else
-				{
-					BindingsCategory->Children.Add(Entry);
-				}
+				ParametersCategory->Children.Add(Entry);
 			}
 			else if (Cast<IAnimNextParameterBlockGraphInterface>(BlockEntry))
 			{
@@ -1031,7 +1017,7 @@ TSharedRef<ITableRow> SParameterBlockView::HandleGenerateRow(TSharedRef<FParamet
 					Args.bShowLibraries = false;
 					Args.bShowBlocks = false;
 					Args.bShowBoundParameters = false;
-					Args.bShowBuiltInParameters = false;
+					Args.bShowBuiltInParameters = false; // Built-In paameters Disabled for MVP
 					Args.OnParameterPicked = FOnParameterPicked::CreateLambda([this](const FParameterBindingReference& InParameterBinding)
 					{
 						FSlateApplication::Get().DismissAllMenus();
@@ -1049,43 +1035,6 @@ TSharedRef<ITableRow> SParameterBlockView::HandleGenerateRow(TSharedRef<FParamet
 				})
 			];
 		}
-		else if (InEntry->CategoryType == EParameterBlockCategoryType::BuiltIn)
-		{
-			RowContainer->AddSlot()
-			.AutoWidth()
-			.VAlign(VAlign_Fill)
-			.HAlign(HAlign_Right)
-			[
-				SNew(SSimpleComboButton)
-				.Text(LOCTEXT("AddBuiltInParameterButton", "Add Built-In Parameter"))
-				.Icon(FAppStyle::Get().GetBrush("Icons.Plus"))
-				.HasDownArrow(true)
-				.OnGetMenuContent_Lambda([this]()
-				{
-					FParameterPickerArgs Args;
-					Args.bMultiSelect = false;
-					Args.bShowLibraries = false;
-					Args.bShowBlocks = false;
-					Args.bShowBoundParameters = false;
-					Args.bShowUnboundParameters = false;
-					Args.bAllowNew = false;
-					Args.OnParameterPicked = FOnParameterPicked::CreateLambda([this](const FParameterBindingReference& InParameterBinding)
-					{
-						FSlateApplication::Get().DismissAllMenus();
-
-						FScopedTransaction Transaction(LOCTEXT("AddBuiltInParameter", "Add builtin parameter"));
-						PendingSelection.Empty();
-
-						// Create a new entry for the parameter
-						UAnimNextParameterLibrary* Library = Cast<UAnimNextParameterLibrary>(InParameterBinding.Library.GetAsset());
-						UAnimNextParameterBlockParameter* Parameter = EditorData->AddParameter(InParameterBinding.Parameter, Library);
-						PendingSelection.Add(Parameter);
-					});
-					return SNew(SParameterPicker)
-					.Args(Args);
-				})
-			];
-		}
 		else if (InEntry->CategoryType == EParameterBlockCategoryType::Graph)
 		{
 			RowContainer->AddSlot()
diff --git a/Engine/Plugins/Experimental/AnimNext/Source/AnimNextEditor/Private/Param/SParameterBlockView.h b/Engine/Plugins/Experimental/AnimNext/Source/AnimNextEditor/Private/Param/SParameterBlockView.h
--- a/Engine/Plugins/Experimental/AnimNext/Source/AnimNextEditor/Private/Param/SParameterBlockView.h
+++ b/Engine/Plugins/Experimental/AnimNext/Source/AnimNextEditor/Private/Param/SParameterBlockView.h
@@ -25,7 +25,6 @@ struct FParameterBlockViewEntry;
 enum class EParameterBlockCategoryType : uint8
 {
 	Parameter,
-	BuiltIn,
 	Graph,
 	// --- ---
 	Invalid
