AnimNext: Allow binding to parameters that are already bound in (other) parameter blocks

diff --git a/Engine/Plugins/Experimental/AnimNext/Source/AnimNextEditor/Private/Param/SParameterBlockView.cpp b/Engine/Plugins/Experimental/AnimNext/Source/AnimNextEditor/Private/Param/SParameterBlockView.cpp
--- a/Engine/Plugins/Experimental/AnimNext/Source/AnimNextEditor/Private/Param/SParameterBlockView.cpp
+++ b/Engine/Plugins/Experimental/AnimNext/Source/AnimNextEditor/Private/Param/SParameterBlockView.cpp
@@ -1016,8 +1016,18 @@ TSharedRef<ITableRow> SParameterBlockView::HandleGenerateRow(TSharedRef<FParamet
 					Args.bMultiSelect = false;
 					Args.bShowLibraries = false;
 					Args.bShowBlocks = false;
-					Args.bShowBoundParameters = false;
+					Args.bShowBoundParameters = true;
 					Args.bShowBuiltInParameters = false; // Built-In paameters Disabled for MVP
+					Args.OnFilterParameter = FOnFilterParameter::CreateLambda([this](const FParameterBindingReference& InParameterBinding)
+					{
+						// Skip params that are already bound in this block
+						if(InParameterBinding.Block == BlockAssetData)
+						{
+							return EFilterParameterResult::Exclude;
+						}
+						
+						return EFilterParameterResult::Include;
+					});
 					Args.OnParameterPicked = FOnParameterPicked::CreateLambda([this](const FParameterBindingReference& InParameterBinding)
 					{
 						FSlateApplication::Get().DismissAllMenus();
