Optimize GetTypeID for evaluation vm tasks by making it constexpr

diff --git a/Engine/Plugins/Experimental/AnimNext/Source/AnimNext/Public/EvaluationVM/EvaluationVM.h b/Engine/Plugins/Experimental/AnimNext/Source/AnimNext/Public/EvaluationVM/EvaluationVM.h
--- a/Engine/Plugins/Experimental/AnimNext/Source/AnimNext/Public/EvaluationVM/EvaluationVM.h
+++ b/Engine/Plugins/Experimental/AnimNext/Source/AnimNext/Public/EvaluationVM/EvaluationVM.h
@@ -15,18 +15,34 @@
 // Use this macro inside the UE::AnimNext namespace
 #define ANIM_NEXT_ENABLE_EVALUATION_STACK_USAGE(Type) \
 	template<> \
-	inline uint32 GetTypeID<Type>() \
+	constexpr uint32 GetTypeID<Type>() \
 	{ \
-		/* TODO: Make this constexpr */ \
-		return FCrc::StrCrc32(#Type); \
+		return UE::AnimNext::ConstexprStringFnv32(#Type); \
 	}
 
 namespace UE::AnimNext
 {
+	constexpr uint32 ConstexprStringFnv32(const char* StringLiteral)
+	{
+		constexpr uint32 Offset = 0x811c9dc5;
+		constexpr uint32 Prime = 0x01000193;
+
+		const char* CharPtr = StringLiteral;
+
+		uint32 Fnv = Offset;
+		while (*CharPtr != 0)
+		{
+			Fnv ^= *CharPtr++;
+			Fnv *= Prime;
+		}
+
+		return Fnv;
+	}
+
 	// Helper function that returns a UID for the specified type
 	// @see ANIM_NEXT_ENABLE_EVALUATION_STACK_USAGE
 	template<typename Type>
-	inline uint32 GetTypeID()
+	constexpr uint32 GetTypeID()
 	{
 		checkf(false, TEXT("Not implemented! See ANIM_NEXT_ENABLE_EVALUATION_STACK_USAGE for details and register your type"));
 		return 0;
