Add more AnimNext profiling scopes

diff --git a/Engine/Plugins/Experimental/AnimNext/Source/AnimNext/Private/TransformArrayOperations_SoA.cpp b/Engine/Plugins/Experimental/AnimNext/Source/AnimNext/Private/TransformArrayOperations_SoA.cpp
--- a/Engine/Plugins/Experimental/AnimNext/Source/AnimNext/Private/TransformArrayOperations_SoA.cpp
+++ b/Engine/Plugins/Experimental/AnimNext/Source/AnimNext/Private/TransformArrayOperations_SoA.cpp
@@ -5,6 +5,9 @@
 #include "AnimNextStats.h"
 
 DEFINE_STAT(STAT_AnimNext_CopyTransforms_SoA);
+DEFINE_STAT(STAT_AnimNext_NormalizeRotations_SoA);
+DEFINE_STAT(STAT_AnimNext_BlendOverwrite_SoA);
+DEFINE_STAT(STAT_AnimNext_BlendAccumulate_SoA);
 
 namespace UE::AnimNext
 {
@@ -77,6 +80,8 @@ namespace UE::AnimNext
 
 	void NormalizeRotations(const FTransformArraySoAView& Input)
 	{
+		SCOPE_CYCLE_COUNTER(STAT_AnimNext_NormalizeRotations_SoA);
+
 		const int32 NumTransforms = Input.Num();
 
 		for (int32 TransformIndex = 0; TransformIndex < NumTransforms; ++TransformIndex)
@@ -137,6 +142,8 @@ namespace UE::AnimNext
 
 	void BlendOverwriteWithScale(const FTransformArraySoAView& Dest, const FTransformArraySoAConstView& Source, const float ScaleWeight)
 	{
+		SCOPE_CYCLE_COUNTER(STAT_AnimNext_BlendOverwrite_SoA);
+
 		const int32 NumTransforms = Source.Num();
 
 		check(Dest.Num() >= NumTransforms);
@@ -151,6 +158,8 @@ namespace UE::AnimNext
 
 	void BlendAddWithScale(const FTransformArraySoAView& Dest, const FTransformArraySoAConstView& Source, const float ScaleWeight)
 	{
+		SCOPE_CYCLE_COUNTER(STAT_AnimNext_BlendAccumulate_SoA);
+
 		const int32 NumTransforms = Source.Num();
 
 		check(Dest.Num() >= NumTransforms);
diff --git a/Engine/Plugins/Experimental/AnimNext/Source/AnimNext/Public/AnimNextStats.h b/Engine/Plugins/Experimental/AnimNext/Source/AnimNext/Public/AnimNextStats.h
--- a/Engine/Plugins/Experimental/AnimNext/Source/AnimNext/Public/AnimNextStats.h
+++ b/Engine/Plugins/Experimental/AnimNext/Source/AnimNext/Public/AnimNextStats.h
@@ -1,4 +1,4 @@
-﻿// Copyright Epic Games, Inc. All Rights Reserved.
+// Copyright Epic Games, Inc. All Rights Reserved.
 
 #pragma once
 
@@ -20,6 +20,9 @@ DECLARE_CYCLE_STAT_EXTERN(TEXT("AnimNext: Update Graph"), STAT_AnimNext_UpdateGr
 DECLARE_CYCLE_STAT_EXTERN(TEXT("AnimNext: Evaluate Graph"), STAT_AnimNext_EvaluateGraph, STATGROUP_AnimNext, ANIMNEXT_API);
 DECLARE_CYCLE_STAT_EXTERN(TEXT("AnimNext: Execute Evaluation Program"), STAT_AnimNext_EvaluationProgram_Execute, STATGROUP_AnimNext, ANIMNEXT_API);
 DECLARE_CYCLE_STAT_EXTERN(TEXT("AnimNext: Copy Transforms (SoA)"), STAT_AnimNext_CopyTransforms_SoA, STATGROUP_AnimNext, ANIMNEXT_API);
+DECLARE_CYCLE_STAT_EXTERN(TEXT("AnimNext: Normalize Rotations (SoA)"), STAT_AnimNext_NormalizeRotations_SoA, STATGROUP_AnimNext, ANIMNEXT_API);
+DECLARE_CYCLE_STAT_EXTERN(TEXT("AnimNext: Blend Overwrite (SoA)"), STAT_AnimNext_BlendOverwrite_SoA, STATGROUP_AnimNext, ANIMNEXT_API);
+DECLARE_CYCLE_STAT_EXTERN(TEXT("AnimNext: Blend Accumulate (SoA)"), STAT_AnimNext_BlendAccumulate_SoA, STATGROUP_AnimNext, ANIMNEXT_API);
 
 DECLARE_CYCLE_STAT_EXTERN(TEXT("AnimNext: Param Type Handle Lock"), STAT_AnimNext_ParamTypeHandle_Lock, STATGROUP_AnimNext, ANIMNEXT_API);
 DECLARE_CYCLE_STAT_EXTERN(TEXT("AnimNext: Param ID Lock"), STAT_AnimNext_ParamIdLock, STATGROUP_AnimNext, ANIMNEXT_API);
