Rewrite AnimNext update traversal
    - Now propagates state down the graph (total weight, etc)
    - Traversal context can now derive from execution context proxy to allow seamless usage and consistent API

diff --git a/Engine/Plugins/Experimental/AnimNext/Source/AnimNext/Private/DecoratorBase/DecoratorPtr.cpp b/Engine/Plugins/Experimental/AnimNext/Source/AnimNext/Private/DecoratorBase/DecoratorPtr.cpp
--- a/Engine/Plugins/Experimental/AnimNext/Source/AnimNext/Private/DecoratorBase/DecoratorPtr.cpp
+++ b/Engine/Plugins/Experimental/AnimNext/Source/AnimNext/Private/DecoratorBase/DecoratorPtr.cpp
@@ -5,16 +5,16 @@
 
 namespace UE::AnimNext
 {
-	FDecoratorPtr::FDecoratorPtr(FNodeInstance* NodeInstance, uint32 DecoratorIndex_)
-		: PackedPointerAndFlags(NodeInstance != nullptr ? reinterpret_cast<uintptr_t>(NodeInstance) : 0)
-		, DecoratorIndex(DecoratorIndex_)
+	FDecoratorPtr::FDecoratorPtr(FNodeInstance* InNodeInstance, uint32 InDecoratorIndex)
+		: PackedPointerAndFlags(InNodeInstance != nullptr ? reinterpret_cast<uintptr_t>(InNodeInstance) : 0)
+		, DecoratorIndex(InDecoratorIndex)
 	{
-		check((reinterpret_cast<uintptr_t>(NodeInstance) & FLAGS_MASK) == 0);	// Make sure we have enough alignment
+		check((reinterpret_cast<uintptr_t>(InNodeInstance) & FLAGS_MASK) == 0);	// Make sure we have enough alignment
 		check(DecoratorIndex <= MAX_uint8);	// Make sure we don't truncate
 
-		if (NodeInstance != nullptr)
+		if (InNodeInstance != nullptr)
 		{
-			NodeInstance->AddReference();
+			InNodeInstance->AddReference();
 		}
 	}
 
@@ -41,17 +41,17 @@ namespace UE::AnimNext
 		DecoratorPtr.DecoratorIndex = 0;
 	}
 
-	FDecoratorPtr::FDecoratorPtr(FNodeInstance* NodeInstance, EFlags Flags, uint32 DecoratorIndex_)
-		: PackedPointerAndFlags(NodeInstance != nullptr ? (reinterpret_cast<uintptr_t>(NodeInstance) | Flags) : 0)
-		, DecoratorIndex(DecoratorIndex_)
+	FDecoratorPtr::FDecoratorPtr(FNodeInstance* InNodeInstance, EFlags InFlags, uint32 InDecoratorIndex)
+		: PackedPointerAndFlags(InNodeInstance != nullptr ? (reinterpret_cast<uintptr_t>(InNodeInstance) | InFlags) : 0)
+		, DecoratorIndex(InDecoratorIndex)
 	{
-		check((reinterpret_cast<uintptr_t>(NodeInstance) & FLAGS_MASK) == 0);	// Make sure we have enough alignment
+		check((reinterpret_cast<uintptr_t>(InNodeInstance) & FLAGS_MASK) == 0);	// Make sure we have enough alignment
 		check(DecoratorIndex <= MAX_uint8);	// Make sure we don't truncate
 
 		// Only increment the reference count if we aren't a weak handle
-		if (NodeInstance != nullptr && (Flags & IS_WEAK_BIT) == 0)
+		if (InNodeInstance != nullptr && (InFlags & IS_WEAK_BIT) == 0)
 		{
-			NodeInstance->AddReference();
+			InNodeInstance->AddReference();
 		}
 	}
 
@@ -112,9 +112,9 @@ namespace UE::AnimNext
 		DecoratorIndex = 0;
 	}
 
-	FWeakDecoratorPtr::FWeakDecoratorPtr(FNodeInstance* NodeInstance_, uint32 DecoratorIndex_)
-		: NodeInstance(NodeInstance_)
-		, DecoratorIndex(DecoratorIndex_)
+	FWeakDecoratorPtr::FWeakDecoratorPtr(FNodeInstance* InNodeInstance, uint32 InDecoratorIndex)
+		: NodeInstance(InNodeInstance)
+		, DecoratorIndex(InDecoratorIndex)
 	{
 		check(DecoratorIndex <= MAX_uint8);	// Make sure we don't truncate
 	}
diff --git a/Engine/Plugins/Experimental/AnimNext/Source/AnimNext/Private/DecoratorBase/ExecutionContext.cpp b/Engine/Plugins/Experimental/AnimNext/Source/AnimNext/Private/DecoratorBase/ExecutionContext.cpp
--- a/Engine/Plugins/Experimental/AnimNext/Source/AnimNext/Private/DecoratorBase/ExecutionContext.cpp
+++ b/Engine/Plugins/Experimental/AnimNext/Source/AnimNext/Private/DecoratorBase/ExecutionContext.cpp
@@ -19,7 +19,9 @@ namespace UE::AnimNext
 	}
 
 	FExecutionContext::FExecutionContext(TArrayView<const uint8> InGraphSharedData)
-		: GraphSharedData(InGraphSharedData)
+		: NodeTemplateRegistry(FNodeTemplateRegistry::Get())
+		, DecoratorRegistry(FDecoratorRegistry::Get())
+		, GraphSharedData(InGraphSharedData)
 		, RigVMLatentMemoryHandles()
 		, RigVMExecuteContext(nullptr)
 	{
@@ -29,7 +31,9 @@ namespace UE::AnimNext
 	}
 
 	FExecutionContext::FExecutionContext(TArrayView<const uint8> InGraphSharedData, FRigVMExtendedExecuteContext& InRigVMExecuteContext, FRigVMMemoryHandleArray InRigVMLatentMemoryHandles)
-		: GraphSharedData(InGraphSharedData)
+		: NodeTemplateRegistry(FNodeTemplateRegistry::Get())
+		, DecoratorRegistry(FDecoratorRegistry::Get())
+		, GraphSharedData(InGraphSharedData)
 		, RigVMLatentMemoryHandles(InRigVMLatentMemoryHandles)
 		, RigVMExecuteContext(&InRigVMExecuteContext)
 	{
@@ -45,7 +49,7 @@ namespace UE::AnimNext
 		Private::GThreadLocalExecutionContext = nullptr;
 	}
 
-	FDecoratorPtr FExecutionContext::AllocateNodeInstance(FWeakDecoratorPtr ParentBinding, FAnimNextDecoratorHandle ChildDecoratorHandle)
+	FDecoratorPtr FExecutionContext::AllocateNodeInstance(const FWeakDecoratorPtr& ParentBinding, FAnimNextDecoratorHandle ChildDecoratorHandle) const
 	{
 		ensure(ChildDecoratorHandle.IsValid());
 		if (!ChildDecoratorHandle.IsValid())
@@ -115,7 +119,7 @@ namespace UE::AnimNext
 		return FDecoratorPtr(NodeInstance, ChildDecoratorIndex);
 	}
 
-	void FExecutionContext::ReleaseNodeInstance(FNodeInstance* NodeInstance)
+	void FExecutionContext::ReleaseNodeInstance(FNodeInstance* NodeInstance) const
 	{
 		ensure(NodeInstance != nullptr && NodeInstance->IsValid());
 		if (NodeInstance == nullptr || !NodeInstance->IsValid())
@@ -160,7 +164,7 @@ namespace UE::AnimNext
 		FMemory::Free(NodeInstance);
 	}
 
-	bool FExecutionContext::GetInterfaceImpl(FDecoratorInterfaceUID InterfaceUID, FWeakDecoratorPtr DecoratorPtr, FDecoratorBinding& InterfaceBinding) const
+	bool FExecutionContext::GetInterfaceImpl(FDecoratorInterfaceUID InterfaceUID, const FWeakDecoratorPtr& DecoratorPtr, FDecoratorBinding& InterfaceBinding) const
 	{
 		if (!DecoratorPtr.IsValid())
 		{
@@ -273,7 +277,7 @@ namespace UE::AnimNext
 		return false;
 	}
 
-	bool FExecutionContext::GetInterfaceSuperImpl(FDecoratorInterfaceUID InterfaceUID, FWeakDecoratorPtr DecoratorPtr, FDecoratorBinding& SuperBinding) const
+	bool FExecutionContext::GetInterfaceSuperImpl(FDecoratorInterfaceUID InterfaceUID, const FWeakDecoratorPtr& DecoratorPtr, FDecoratorBinding& SuperBinding) const
 	{
 		if (!DecoratorPtr.IsValid())
 		{
@@ -358,16 +362,12 @@ namespace UE::AnimNext
 	const FNodeTemplate* FExecutionContext::GetNodeTemplate(const FNodeDescription& NodeDesc) const
 	{
 		check(NodeDesc.GetTemplateHandle().IsValid());
-
-		const FNodeTemplateRegistry& NodeTemplateRegistry = FNodeTemplateRegistry::Get();
 		return NodeTemplateRegistry.Find(NodeDesc.GetTemplateHandle());
 	}
 
 	const FDecorator* FExecutionContext::GetDecorator(const FDecoratorTemplate& Template) const
 	{
 		check(Template.GetRegistryHandle().IsValid());
-
-		const FDecoratorRegistry& DecoratorRegistry = FDecoratorRegistry::Get();
 		return DecoratorRegistry.Find(Template.GetRegistryHandle());
 	}
 
diff --git a/Engine/Plugins/Experimental/AnimNext/Source/AnimNext/Private/DecoratorInterfaces/IContinuousBlend.cpp b/Engine/Plugins/Experimental/AnimNext/Source/AnimNext/Private/DecoratorInterfaces/IContinuousBlend.cpp
--- a/Engine/Plugins/Experimental/AnimNext/Source/AnimNext/Private/DecoratorInterfaces/IContinuousBlend.cpp
+++ b/Engine/Plugins/Experimental/AnimNext/Source/AnimNext/Private/DecoratorInterfaces/IContinuousBlend.cpp
@@ -6,7 +6,7 @@
 
 namespace UE::AnimNext
 {
-	float IContinuousBlend::GetBlendWeight(FExecutionContext& Context, const TDecoratorBinding<IContinuousBlend>& Binding, int32 ChildIndex) const
+	float IContinuousBlend::GetBlendWeight(const FExecutionContext& Context, const TDecoratorBinding<IContinuousBlend>& Binding, int32 ChildIndex) const
 	{
 		TDecoratorBinding<IContinuousBlend> SuperBinding;
 		if (Context.GetInterfaceSuper(Binding, SuperBinding))
diff --git a/Engine/Plugins/Experimental/AnimNext/Source/AnimNext/Private/DecoratorInterfaces/IDiscreteBlend.cpp b/Engine/Plugins/Experimental/AnimNext/Source/AnimNext/Private/DecoratorInterfaces/IDiscreteBlend.cpp
--- a/Engine/Plugins/Experimental/AnimNext/Source/AnimNext/Private/DecoratorInterfaces/IDiscreteBlend.cpp
+++ b/Engine/Plugins/Experimental/AnimNext/Source/AnimNext/Private/DecoratorInterfaces/IDiscreteBlend.cpp
@@ -6,7 +6,7 @@
 
 namespace UE::AnimNext
 {
-	float IDiscreteBlend::GetBlendWeight(FExecutionContext& Context, const TDecoratorBinding<IDiscreteBlend>& Binding, int32 ChildIndex) const
+	float IDiscreteBlend::GetBlendWeight(const FExecutionContext& Context, const TDecoratorBinding<IDiscreteBlend>& Binding, int32 ChildIndex) const
 	{
 		TDecoratorBinding<IDiscreteBlend> SuperBinding;
 		if (Context.GetInterfaceSuper(Binding, SuperBinding))
@@ -17,7 +17,7 @@ namespace UE::AnimNext
 		return -1.0f;
 	}
 
-	const FAlphaBlend* IDiscreteBlend::GetBlendState(FExecutionContext& Context, const TDecoratorBinding<IDiscreteBlend>& Binding, int32 ChildIndex) const
+	const FAlphaBlend* IDiscreteBlend::GetBlendState(const FExecutionContext& Context, const TDecoratorBinding<IDiscreteBlend>& Binding, int32 ChildIndex) const
 	{
 		TDecoratorBinding<IDiscreteBlend> SuperBinding;
 		if (Context.GetInterfaceSuper(Binding, SuperBinding))
@@ -28,7 +28,7 @@ namespace UE::AnimNext
 		return nullptr;
 	}
 
-	int32 IDiscreteBlend::GetBlendDestinationChildIndex(FExecutionContext& Context, const TDecoratorBinding<IDiscreteBlend>& Binding) const
+	int32 IDiscreteBlend::GetBlendDestinationChildIndex(const FExecutionContext& Context, const TDecoratorBinding<IDiscreteBlend>& Binding) const
 	{
 		TDecoratorBinding<IDiscreteBlend> SuperBinding;
 		if (Context.GetInterfaceSuper(Binding, SuperBinding))
@@ -39,7 +39,7 @@ namespace UE::AnimNext
 		return INDEX_NONE;
 	}
 
-	void IDiscreteBlend::OnBlendTransition(FExecutionContext& Context, const TDecoratorBinding<IDiscreteBlend>& Binding, int32 OldChildIndex, int32 NewChildIndex) const
+	void IDiscreteBlend::OnBlendTransition(const FExecutionContext& Context, const TDecoratorBinding<IDiscreteBlend>& Binding, int32 OldChildIndex, int32 NewChildIndex) const
 	{
 		TDecoratorBinding<IDiscreteBlend> SuperBinding;
 		if (Context.GetInterfaceSuper(Binding, SuperBinding))
@@ -48,7 +48,7 @@ namespace UE::AnimNext
 		}
 	}
 
-	void IDiscreteBlend::OnBlendInitiated(FExecutionContext& Context, const TDecoratorBinding<IDiscreteBlend>& Binding, int32 ChildIndex) const
+	void IDiscreteBlend::OnBlendInitiated(const FExecutionContext& Context, const TDecoratorBinding<IDiscreteBlend>& Binding, int32 ChildIndex) const
 	{
 		TDecoratorBinding<IDiscreteBlend> SuperBinding;
 		if (Context.GetInterfaceSuper(Binding, SuperBinding))
@@ -57,7 +57,7 @@ namespace UE::AnimNext
 		}
 	}
 
-	void IDiscreteBlend::OnBlendTerminated(FExecutionContext& Context, const TDecoratorBinding<IDiscreteBlend>& Binding, int32 ChildIndex) const
+	void IDiscreteBlend::OnBlendTerminated(const FExecutionContext& Context, const TDecoratorBinding<IDiscreteBlend>& Binding, int32 ChildIndex) const
 	{
 		TDecoratorBinding<IDiscreteBlend> SuperBinding;
 		if (Context.GetInterfaceSuper(Binding, SuperBinding))
diff --git a/Engine/Plugins/Experimental/AnimNext/Source/AnimNext/Private/DecoratorInterfaces/IEvaluate.cpp b/Engine/Plugins/Experimental/AnimNext/Source/AnimNext/Private/DecoratorInterfaces/IEvaluate.cpp
--- a/Engine/Plugins/Experimental/AnimNext/Source/AnimNext/Private/DecoratorInterfaces/IEvaluate.cpp
+++ b/Engine/Plugins/Experimental/AnimNext/Source/AnimNext/Private/DecoratorInterfaces/IEvaluate.cpp
@@ -41,7 +41,7 @@ namespace UE::AnimNext
 		FEvaluationProgram* OldEvaluationProgram;
 	};
 
-	void IEvaluate::PreEvaluate(FExecutionContext& Context, const TDecoratorBinding<IEvaluate>& Binding) const
+	void IEvaluate::PreEvaluate(const FExecutionContext& Context, const TDecoratorBinding<IEvaluate>& Binding) const
 	{
 		TDecoratorBinding<IEvaluate> SuperBinding;
 		if (Context.GetInterfaceSuper(Binding, SuperBinding))
@@ -50,7 +50,7 @@ namespace UE::AnimNext
 		}
 	}
 
-	void IEvaluate::PostEvaluate(FExecutionContext& Context, const TDecoratorBinding<IEvaluate>& Binding) const
+	void IEvaluate::PostEvaluate(const FExecutionContext& Context, const TDecoratorBinding<IEvaluate>& Binding) const
 	{
 		TDecoratorBinding<IEvaluate> SuperBinding;
 		if (Context.GetInterfaceSuper(Binding, SuperBinding))
diff --git a/Engine/Plugins/Experimental/AnimNext/Source/AnimNext/Private/DecoratorInterfaces/IHierarchy.cpp b/Engine/Plugins/Experimental/AnimNext/Source/AnimNext/Private/DecoratorInterfaces/IHierarchy.cpp
--- a/Engine/Plugins/Experimental/AnimNext/Source/AnimNext/Private/DecoratorInterfaces/IHierarchy.cpp
+++ b/Engine/Plugins/Experimental/AnimNext/Source/AnimNext/Private/DecoratorInterfaces/IHierarchy.cpp
@@ -6,7 +6,7 @@
 
 namespace UE::AnimNext
 {
-	uint32 IHierarchy::GetNumChildren(FExecutionContext& Context, const TDecoratorBinding<IHierarchy>& Binding) const
+	uint32 IHierarchy::GetNumChildren(const FExecutionContext& Context, const TDecoratorBinding<IHierarchy>& Binding) const
 	{
 		TDecoratorBinding<IHierarchy> SuperBinding;
 		if (Context.GetInterfaceSuper(Binding, SuperBinding))
@@ -17,7 +17,7 @@ namespace UE::AnimNext
 		return 0;
 	}
 
-	void IHierarchy::GetChildren(FExecutionContext& Context, const TDecoratorBinding<IHierarchy>& Binding, FChildrenArray& Children) const
+	void IHierarchy::GetChildren(const FExecutionContext& Context, const TDecoratorBinding<IHierarchy>& Binding, FChildrenArray& Children) const
 	{
 		TDecoratorBinding<IHierarchy> SuperBinding;
 		if (Context.GetInterfaceSuper(Binding, SuperBinding))
diff --git a/Engine/Plugins/Experimental/AnimNext/Source/AnimNext/Private/DecoratorInterfaces/ISmoothBlend.cpp b/Engine/Plugins/Experimental/AnimNext/Source/AnimNext/Private/DecoratorInterfaces/ISmoothBlend.cpp
--- a/Engine/Plugins/Experimental/AnimNext/Source/AnimNext/Private/DecoratorInterfaces/ISmoothBlend.cpp
+++ b/Engine/Plugins/Experimental/AnimNext/Source/AnimNext/Private/DecoratorInterfaces/ISmoothBlend.cpp
@@ -6,7 +6,7 @@
 
 namespace UE::AnimNext
 {
-	float ISmoothBlend::GetBlendTime(FExecutionContext& Context, const TDecoratorBinding<ISmoothBlend>& Binding, int32 ChildIndex) const
+	float ISmoothBlend::GetBlendTime(const FExecutionContext& Context, const TDecoratorBinding<ISmoothBlend>& Binding, int32 ChildIndex) const
 	{
 		TDecoratorBinding<ISmoothBlend> SuperBinding;
 		if (Context.GetInterfaceSuper(Binding, SuperBinding))
diff --git a/Engine/Plugins/Experimental/AnimNext/Source/AnimNext/Private/DecoratorInterfaces/ITimeline.cpp b/Engine/Plugins/Experimental/AnimNext/Source/AnimNext/Private/DecoratorInterfaces/ITimeline.cpp
--- a/Engine/Plugins/Experimental/AnimNext/Source/AnimNext/Private/DecoratorInterfaces/ITimeline.cpp
+++ b/Engine/Plugins/Experimental/AnimNext/Source/AnimNext/Private/DecoratorInterfaces/ITimeline.cpp
@@ -6,7 +6,7 @@
 
 namespace UE::AnimNext
 {
-	double ITimeline::GetPlayRate(FExecutionContext& Context, const TDecoratorBinding<ITimeline>& Binding) const
+	double ITimeline::GetPlayRate(const FExecutionContext& Context, const TDecoratorBinding<ITimeline>& Binding) const
 	{
 		TDecoratorBinding<ITimeline> SuperBinding;
 		if (Context.GetInterfaceSuper(Binding, SuperBinding))
diff --git a/Engine/Plugins/Experimental/AnimNext/Source/AnimNext/Private/DecoratorInterfaces/IUpdate.cpp b/Engine/Plugins/Experimental/AnimNext/Source/AnimNext/Private/DecoratorInterfaces/IUpdate.cpp
--- a/Engine/Plugins/Experimental/AnimNext/Source/AnimNext/Private/DecoratorInterfaces/IUpdate.cpp
+++ b/Engine/Plugins/Experimental/AnimNext/Source/AnimNext/Private/DecoratorInterfaces/IUpdate.cpp
@@ -2,7 +2,6 @@
 
 #include "DecoratorInterfaces/IUpdate.h"
 
-#include "DecoratorBase/ExecutionContext.h"
 #include "DecoratorInterfaces/IHierarchy.h"
 #include "AnimNextStats.h"
 
@@ -10,67 +9,205 @@ DEFINE_STAT(STAT_AnimNext_UpdateGraph);
 
 namespace UE::AnimNext
 {
-	FUpdateTraversalContext::FUpdateTraversalContext(float InDeltaTime)
-		: DeltaTime(InDeltaTime)
+	void IUpdate::PreUpdate(FUpdateTraversalContext& Context, const TDecoratorBinding<IUpdate>& Binding, const FDecoratorUpdateState& DecoratorState) const
 	{
+		TDecoratorBinding<IUpdate> SuperBinding;
+		if (Context.GetInterfaceSuper(Binding, SuperBinding))
+		{
+			SuperBinding.PreUpdate(Context, DecoratorState);
+		}
 	}
 
-	void IUpdate::PreUpdate(FExecutionContext& Context, const TDecoratorBinding<IUpdate>& Binding) const
+	void IUpdate::QueueChildrenForTraversal(FUpdateTraversalContext& Context, const TDecoratorBinding<IUpdate>& Binding, const FDecoratorUpdateState& DecoratorState, FUpdateTraversalQueue& TraversalQueue) const
 	{
 		TDecoratorBinding<IUpdate> SuperBinding;
 		if (Context.GetInterfaceSuper(Binding, SuperBinding))
 		{
-			SuperBinding.PreUpdate(Context);
+			SuperBinding.QueueChildrenForTraversal(Context, DecoratorState, TraversalQueue);
 		}
 	}
 
-	void IUpdate::PostUpdate(FExecutionContext& Context, const TDecoratorBinding<IUpdate>& Binding) const
+	void IUpdate::PostUpdate(FUpdateTraversalContext& Context, const TDecoratorBinding<IUpdate>& Binding, const FDecoratorUpdateState& DecoratorState) const
 	{
 		TDecoratorBinding<IUpdate> SuperBinding;
 		if (Context.GetInterfaceSuper(Binding, SuperBinding))
 		{
-			SuperBinding.PostUpdate(Context);
+			SuperBinding.PostUpdate(Context, DecoratorState);
 		}
 	}
 
-	enum class EUpdateStep
-	{
-		PreUpdate,
-		PostUpdate,
-	};
+	//////////////////////////////////////////////////////////////////////////
+	// Traversal implementation
 
 	// This structure is transient and lives either on the stack or the memstack and its destructor may not be called
 	struct FUpdateEntry
 	{
+		// The decorator state for this entry
+		FDecoratorUpdateState		DecoratorState;
+
 		// The decorator handle that points to our node to update
-		FWeakDecoratorPtr	DecoratorPtr;
+		FWeakDecoratorPtr			DecoratorPtr;
 
-		// Which step we wish to perform when we next see this entry
-		EUpdateStep			DesiredStep = EUpdateStep::PreUpdate;
+		// Whether or not PreUpdate had been called already
+		// TODO: Store bHasPreUpdated in the LSB of the entry pointer to save padding?
+		bool						bHasPreUpdated = false;
 
 		// Once we've called PreUpdate, we cache the decorator binding to avoid a redundant query to call PostUpdate
-		TDecoratorBinding<IUpdate> UpdateDecorator;
+		TDecoratorBinding<IUpdate>	UpdateDecorator;
 
 		// These pointers are mutually exclusive
-		// An entry is either part of the pending stack, the free list, or neither
+		// An entry is either part of the queued update stack, the update stack, the free list, or none of the above
 		union
 		{
-			// Next entry in the linked list of free entries
-			FUpdateEntry* NextFreeEntry = nullptr;
+			// Next entry in the stack of free entries
+			FUpdateEntry*			NextFreeEntry = nullptr;
+
+			// Previous entry on the update stack
+			FUpdateEntry*			PrevUpdateStackEntry;
 
-			// Previous entry below us in the nodes pending stack
-			FUpdateEntry* PrevStackEntry;
+			// Previous entry on the queued update stack
+			FUpdateEntry*			PrevQueuedUpdateStackEntry;
 		};
 
-		FUpdateEntry(const FWeakDecoratorPtr& InDecoratorPtr, EUpdateStep InDesiredStep, FUpdateEntry* InPrevStackEntry = nullptr)
-			: DecoratorPtr(InDecoratorPtr)
-			, DesiredStep(InDesiredStep)
-			, PrevStackEntry(InPrevStackEntry)
+		FUpdateEntry(const FWeakDecoratorPtr& InDecoratorPtr, const FDecoratorUpdateState InDecoratorState)
+			: DecoratorState(InDecoratorState)
+			, DecoratorPtr(InDecoratorPtr)
 		{
 		}
 	};
 
-	void UpdateGraph(FExecutionContext& Context, FUpdateTraversalContext& TraversalContext, FWeakDecoratorPtr GraphRootPtr)
+	FUpdateTraversalContext::FUpdateTraversalContext(const FExecutionContext& InExecutionContext, FMemStack& InMemStack)
+		: FExecutionContextProxy(InExecutionContext)
+		, MemStack(InMemStack)
+	{}
+
+	void FUpdateTraversalContext::PushQueuedUpdateEntries(FUpdateTraversalQueue& TraversalQueue)
+	{
+		// Pop every entry from the queued update stack and push them onto the update stack
+		// reversing their order
+		while (FUpdateEntry* Entry = TraversalQueue.QueuedUpdateStackHead)
+		{
+			// Update our queued stack head
+			TraversalQueue.QueuedUpdateStackHead = Entry->PrevQueuedUpdateStackEntry;
+
+			// Push our new entry onto the update stack
+			PushUpdateEntry(Entry);
+		}
+	}
+
+	void FUpdateTraversalContext::PushUpdateEntry(FUpdateEntry* Entry)
+	{
+		Entry->PrevUpdateStackEntry = UpdateStackHead;
+		UpdateStackHead = Entry;
+	}
+
+	FUpdateEntry* FUpdateTraversalContext::PopUpdateEntry()
+	{
+		FUpdateEntry* ChildEntry = UpdateStackHead;
+		if (ChildEntry != nullptr)
+		{
+			// We have a child, set our new head
+			UpdateStackHead = ChildEntry->PrevUpdateStackEntry;
+		}
+
+		return ChildEntry;
+	}
+
+	void FUpdateTraversalContext::PushFreeEntry(FUpdateEntry* Entry)
+	{
+		Entry->NextFreeEntry = FreeEntryStackHead;
+		FreeEntryStackHead = Entry;
+	}
+
+	FUpdateEntry* FUpdateTraversalContext::GetNewEntry(const FWeakDecoratorPtr& DecoratorPtr, const FDecoratorUpdateState& DecoratorState)
+	{
+		FUpdateEntry* FreeEntry = FreeEntryStackHead;
+		if (FreeEntry != nullptr)
+		{
+			// We have a free entry, set our new head
+			FreeEntryStackHead = FreeEntry->NextFreeEntry;
+
+			// Update our entry
+			FreeEntry->DecoratorState = DecoratorState;
+			FreeEntry->DecoratorPtr = DecoratorPtr;
+			FreeEntry->bHasPreUpdated = false;
+			FreeEntry->NextFreeEntry = nullptr;		// Mark it as not being a member of any list
+		}
+		else
+		{
+			// Allocate a new entry
+			FreeEntry = new(MemStack) FUpdateEntry(DecoratorPtr, DecoratorState);
+		}
+
+		return FreeEntry;
+	}
+
+	FUpdateTraversalQueue::FUpdateTraversalQueue(FUpdateTraversalContext& InTraversalContext)
+		: TraversalContext(InTraversalContext)
+	{
+	}
+
+	void FUpdateTraversalQueue::Push(const FWeakDecoratorPtr& ChildPtr, const FDecoratorUpdateState& ChildDecoratorState)
+	{
+		if (!ChildPtr.IsValid())
+		{
+			return;	// Don't queue invalid pointers
+		}
+
+		FUpdateEntry* ChildEntry = TraversalContext.GetNewEntry(ChildPtr, ChildDecoratorState);
+
+		// We push children that are queued onto a stack
+		// Once pre-update is done, we'll pop queued entries one by one and push them
+		// onto the update stack
+		// This has the effect of reversing the entries so that they are traversed in
+		// the same order they are queued in:
+		//    - First queued will be at the bottom of the queued stack and it ends up at the
+		//      at the top of the update stack (last entry pushed)
+		ChildEntry->PrevQueuedUpdateStackEntry = QueuedUpdateStackHead;
+		QueuedUpdateStackHead = ChildEntry;
+	}
+
+	// Performance note
+	// When we process an animation graph for a frame, typically we'll update first before we evaluate
+	// As a result of this, when we query for the update interface here, we will likely hit cold memory
+	// which will cache miss (by touching the graph instance for the first time).
+	// 
+	// The processor will cache miss and continue to process as many instructions as it can before the
+	// out-of-order execution window fills up. This is problematic here because a lot of the subsequent
+	// instructions depend on the node instance and the interface it returns. The processor will be unable
+	// to execute any of the instructions that follow in the current loop iteration. However, it might
+	// be able to get started on the next node entry which is likely to cache miss as well. Should the
+	// processor make it that far and it turns out that we have to push a child onto the stack, all
+	// of the work it tried to do ahead of time will have to be thrown away.
+	// 
+	// There are two things that we can do here to try and help performance: prefetch ahead and bulk
+	// query.
+	// 
+	// If we prefetch, we have to be careful because we do not know what the node will do in its PreUpdate.
+	// If it turns out that it does a lot of work, our prefetch might end up getting thrown out. This is
+	// because prefetched cache lines typically end up being the first evicted unless they are touched first.
+	// It is thus dangerous to use manual prefetching when the memory access pattern isn't fully known.
+	// In practice, it is likely viable as most nodes won't do too much work.
+	// 
+	// A better approach could be to instead bulk query for our interfaces. We could cache in the FUpdateEntry
+	// the decorator bindings for IUpdate and IHierarchy (and re-use the binding for IUpdate for PostUpdate).
+	// Every iteration we could check how many children are queued up on the stack. We could then grab N
+	// entries (2 to 4) and query their interfaces in bulk. The idea is to clump the cache miss instructions
+	// together and to interleave the interface queries. This will queue up as much work as possible in the
+	// out-of-order execution window that will not be thrown away because of a branch. Eventually the first
+	// interface query will complete and execution will resume here to call PreUpdate etc. This will be able
+	// to happen while the processor still waits on the cache misses and finishes the interface query of the
+	// other bulked children. The same effect could be achieved by querying the interfaces after the call
+	// to GetChildren by bulk querying all of them right then. This way, as soon as the execution window
+	// can clear the end of the loop, it can start working on the next entry which will be warm in the L1
+	// cache allowing the CPU to carry ahead before all child interfaces are fully resolved.
+	// 
+	// The above may seem like a stretch and an insignificant over optimization but it could very well be the
+	// key to unlocking large performance gains during traversal. The above optimization would allow us to
+	// perform as much useful work as possible while waiting for memory, hiding its slow latency by fully
+	// leveraging out-of-order CPU execution.
+
+	void UpdateGraph(const FExecutionContext& Context, const FWeakDecoratorPtr& GraphRootPtr, float DeltaTime)
 	{
 		SCOPE_CYCLE_COUNTER(STAT_AnimNext_UpdateGraph);
 		
@@ -83,139 +220,74 @@ namespace UE::AnimNext
 		FMemMark Mark(MemStack);
 
 		FChildrenArray Children;
-		Children.Reserve(16);
-
-		FScopedTraversalContext ScopedTraversalContext(Context, TraversalContext);
+		TDecoratorBinding<IHierarchy> HierarchyDecorator;
 
-		// Add the graph root to kick start the update process
-		FUpdateEntry GraphRootEntry(GraphRootPtr, EUpdateStep::PreUpdate);
-		FUpdateEntry* NodesPendingUpdateStackTop = &GraphRootEntry;
+		FUpdateTraversalContext TraversalContext(Context, MemStack);
+		FUpdateTraversalQueue TraversalQueue(TraversalContext);
 
-		// List of free entries we can recycle
-		FUpdateEntry* FreeEntryList = nullptr;
+		// Add the graph root to start the update process
+		FUpdateEntry RootEntry(GraphRootPtr, FDecoratorUpdateState(DeltaTime));
+		TraversalContext.PushUpdateEntry(&RootEntry);
 
-		TDecoratorBinding<IHierarchy> HierarchyDecorator;
-
-		// Process every node twice: pre-update and post-update
-		while (NodesPendingUpdateStackTop != nullptr)
+		while (FUpdateEntry* Entry = TraversalContext.PopUpdateEntry())
 		{
-			// Grab the top most entry
-			FUpdateEntry* Entry = NodesPendingUpdateStackTop;
-			bool bIsEntryUsed = true;
-
-			if (Entry->DesiredStep == EUpdateStep::PreUpdate)
+			if (!Entry->bHasPreUpdated)
 			{
 				if (Context.GetInterface(Entry->DecoratorPtr, Entry->UpdateDecorator))
 				{
 					// This is the first time we visit this node, time to pre-update
-					Entry->UpdateDecorator.PreUpdate(Context);
-
-					// Leave our entry on top of the stack, we'll need to call PostUpdate once the children
-					// we'll push on top finish
-					Entry->DesiredStep = EUpdateStep::PostUpdate;
+					Entry->UpdateDecorator.PreUpdate(TraversalContext, Entry->DecoratorState);
+
+					// Make sure that next time we visit this entry, we'll post-update
+					Entry->bHasPreUpdated = true;
+
+					// Push this entry onto the update stack, we'll call it once all our children have finished executing
+					TraversalContext.PushUpdateEntry(Entry);
+
+					// Request that the decorator queues the children it wants to visit
+					// This is a separate function from PreUpdate to simplify traversal management. It is often the case that
+					// the base decorator is the one best placed to figure out how to optimally queue children since it
+					// owns the handles to them. However, if an additive decorator wishes to override PreUpdate, it might want
+					// to perform logic after the base PreUpdate but before children are queued. Without a separate function,
+					// we would have to rewrite the base PreUpdate entirely and use IHierarchy to query the handles of our children.
+					Entry->UpdateDecorator.QueueChildrenForTraversal(TraversalContext, Entry->DecoratorState, TraversalQueue);
+
+					// Iterate over our queued children and push them onto the update stack
+					// We do this to allow children to be queued in traversal order which is intuitive
+					// but to traverse them in that order, they must be pushed in reverse order onto the update stack
+					TraversalContext.PushQueuedUpdateEntries(TraversalQueue);
 				}
 				else
 				{
-					// This node doesn't implement IUpdate, we can pop it from the stack
-					NodesPendingUpdateStackTop = Entry->PrevStackEntry;
-					bIsEntryUsed = false;
-				}
-
-				// Performance note
-				// When we process an animation graph for a frame, typically we'll update first before we evaluate
-				// As a result of this, when we query for the update interface here, we will likely hit cold memory
-				// which will cache miss (by touching the graph instance for the first time).
-				// 
-				// The processor will cache miss and continue to process as many instructions as it can before the
-				// out-of-order execution window fills up. This is problematic here because a lot of the subsequent
-				// instructions depend on the node instance and the interface it returns. The processor will be unable
-				// to execute any of the instructions that follow in the current loop iteration. However, it might
-				// be able to get started on the next node entry which is likely to cache miss as well. Should the
-				// processor make it that far and it turns out that we have to push a child onto the stack, all
-				// of the work it tried to do ahead of time will have to be thrown away.
-				// 
-				// There are two things that we can do here to try and help performance: prefetch ahead and bulk
-				// query.
-				// 
-				// If we prefetch, we have to be careful because we do not know what the node will do in its PreUpdate.
-				// If it turns out that it does a lot of work, our prefetch might end up getting thrown out. This is
-				// because prefetched cache lines typically end up being the first evicted unless they are touched first.
-				// It is thus dangerous to use manual prefetching when the memory access pattern isn't fully known.
-				// In practice, it is likely viable as most nodes won't do too much work.
-				// 
-				// A better approach could be to instead bulk query for our interfaces. We could cache in the FUpdateEntry
-				// the decorator bindings for IUpdate and IHierarchy (and re-use the binding for IUpdate for PostUpdate).
-				// Every iteration we could check how many children are queued up on the stack. We could then grab N
-				// entries (2 to 4) and query their interfaces in bulk. The idea is to clump the cache miss instructions
-				// together and to interleave the interface queries. This will queue up as much work as possible in the
-				// out-of-order execution window that will not be thrown away because of a branch. Eventually the first
-				// interface query will complete and execution will resume here to call PreUpdate etc. This will be able
-				// to happen while the processor still waits on the cache misses and finishes the interface query of the
-				// other bulked children. The same effect could be achieved by querying the interfaces after the call
-				// to GetChildren by bulk querying all of them right then. This way, as soon as the execution window
-				// can clear the end of the loop, it can start working on the next entry which will be warm in the L1
-				// cache allowing the CPU to carry ahead before all child interfaces are fully resolved.
-				// 
-				// The above may seem like a stretch and an insignificant over optimization but it could very well be the
-				// key to unlocking large performance gains during traversal. The above optimization would allow us to
-				// perform as much useful work as possible while waiting for memory, hiding its slow latency by fully
-				// leveraging out-of-order CPU execution.
-
-				if (Context.GetInterface(Entry->DecoratorPtr, HierarchyDecorator))
-				{
-					HierarchyDecorator.GetChildren(Context, Children);
-
-					// Append our children in reserve order so that they are visited in the same order they were added
-					for (int32 ChildIndex = Children.Num() - 1; ChildIndex >= 0; --ChildIndex)
+					// This node doesn't implement IUpdate
+					// We'll grab its children and traverse them
+					if (Context.GetInterface(Entry->DecoratorPtr, HierarchyDecorator))
 					{
-						// Insert our new child on top of the stack
+						HierarchyDecorator.GetChildren(Context, Children);
 
-						FUpdateEntry* ChildEntry;
-						if (FreeEntryList != nullptr)
+						// Append our children in reserve order so that they are visited in the same order they were added
+						for (int32 ChildIndex = Children.Num() - 1; ChildIndex >= 0; --ChildIndex)
 						{
-							// Grab an entry from the free list
-							ChildEntry = FreeEntryList;
-							FreeEntryList = ChildEntry->NextFreeEntry;
-
-							ChildEntry->DecoratorPtr = Children[ChildIndex];
-							ChildEntry->DesiredStep = EUpdateStep::PreUpdate;
-							ChildEntry->PrevStackEntry = NodesPendingUpdateStackTop;
-						}
-						else
-						{
-							// Allocate a new entry
-							ChildEntry = new(MemStack) FUpdateEntry(Children[ChildIndex], EUpdateStep::PreUpdate, NodesPendingUpdateStackTop);
+							FUpdateEntry* ChildEntry = TraversalContext.GetNewEntry(Children[ChildIndex], Entry->DecoratorState);
+							TraversalContext.PushUpdateEntry(ChildEntry);
 						}
 
-						NodesPendingUpdateStackTop = ChildEntry;
+						// Reset our container for the next time we need it
+						Children.Reset();
 					}
 
-					// Reset our container for the next time we need it
-					Children.Reset();
+					// We don't need this entry anymore
+					TraversalContext.PushFreeEntry(Entry);
 				}
-
-				// Break and continue to the next top-most entry
-				// It is either a child ready for its pre-update or our current entry ready for its post-update (leaf)
 			}
 			else
 			{
-				// We've already visited this node once, time to post-update
+				// We've already visited this node once, time to PostUpdate
 				check(Entry->UpdateDecorator.IsValid());
-				Entry->UpdateDecorator.PostUpdate(Context);
-
-				// Now that we are done processing this entry, we can pop it
-				NodesPendingUpdateStackTop = Entry->PrevStackEntry;
-				bIsEntryUsed = false;
+				Entry->UpdateDecorator.PostUpdate(TraversalContext, Entry->DecoratorState);
 
-				// Break and continue to the next top-most entry
-				// It is either a sibling read for its post-update or our parent entry ready for its post-update
-			}
-
-			if (!bIsEntryUsed)
-			{
-				// This entry is no longer used, add it to the free list
-				Entry->NextFreeEntry = FreeEntryList;
-				FreeEntryList = Entry;
+				// We don't need this entry anymore
+				TraversalContext.PushFreeEntry(Entry);
 			}
 		}
 	}
diff --git a/Engine/Plugins/Experimental/AnimNext/Source/AnimNext/Private/Decorators/BlendByBool.cpp b/Engine/Plugins/Experimental/AnimNext/Source/AnimNext/Private/Decorators/BlendByBool.cpp
--- a/Engine/Plugins/Experimental/AnimNext/Source/AnimNext/Private/Decorators/BlendByBool.cpp
+++ b/Engine/Plugins/Experimental/AnimNext/Source/AnimNext/Private/Decorators/BlendByBool.cpp
@@ -5,11 +5,6 @@
 #include "DecoratorBase/ExecutionContext.h"
 #include "EvaluationVM/Tasks/BlendKeyframes.h"
 
-#if 0
-// TODO: Validate blend by bool and blend smoother through a cvar at runtime to select the bool value
-static TAutoConsoleVariable<int32> CVarAnimNextForceBlendBool(TEXT("a.AnimNextForceBlendBool"), -1, TEXT("If != -1, then the value [0 (false), 1 (true)] is used to control Blend By Bool and override its value."));
-#endif
-
 namespace UE::AnimNext
 {
 	AUTO_REGISTER_ANIM_DECORATOR(FBlendByBoolDecorator)
@@ -24,7 +19,7 @@ namespace UE::AnimNext
 	static constexpr int32 TRUE_CHILD_INDEX = 0;
 	static constexpr int32 FALSE_CHILD_INDEX = 1;
 
-	void FBlendByBoolDecorator::PostEvaluate(FExecutionContext& Context, const TDecoratorBinding<IEvaluate>& Binding) const
+	void FBlendByBoolDecorator::PostEvaluate(const FExecutionContext& Context, const TDecoratorBinding<IEvaluate>& Binding) const
 	{
 		const FInstanceData* InstanceData = Binding.GetInstanceData<FInstanceData>();
 
@@ -46,7 +41,7 @@ namespace UE::AnimNext
 		}
 	}
 
-	void FBlendByBoolDecorator::PreUpdate(FExecutionContext& Context, const TDecoratorBinding<IUpdate>& Binding) const
+	void FBlendByBoolDecorator::PreUpdate(FUpdateTraversalContext& Context, const TDecoratorBinding<IUpdate>& Binding, const FDecoratorUpdateState& DecoratorState) const
 	{
 		const FSharedData* SharedData = Binding.GetSharedData<FSharedData>();
 		FInstanceData* InstanceData = Binding.GetInstanceData<FInstanceData>();
@@ -54,7 +49,7 @@ namespace UE::AnimNext
 		TDecoratorBinding<IDiscreteBlend> DiscreteBlendDecorator;
 		Context.GetInterface(Binding, DiscreteBlendDecorator);
 
-		const float DestinationChildIndex = DiscreteBlendDecorator.GetBlendDestinationChildIndex(Context);
+		const int32 DestinationChildIndex = DiscreteBlendDecorator.GetBlendDestinationChildIndex(Context);
 		if (InstanceData->PreviousChildIndex != DestinationChildIndex)
 		{
 			DiscreteBlendDecorator.OnBlendTransition(Context, InstanceData->PreviousChildIndex, DestinationChildIndex);
@@ -63,12 +58,47 @@ namespace UE::AnimNext
 		}
 	}
 
-	uint32 FBlendByBoolDecorator::GetNumChildren(FExecutionContext& Context, const TDecoratorBinding<IHierarchy>& Binding) const
+	void FBlendByBoolDecorator::QueueChildrenForTraversal(FUpdateTraversalContext& Context, const TDecoratorBinding<IUpdate>& Binding, const FDecoratorUpdateState& DecoratorState, FUpdateTraversalQueue& TraversalQueue) const
+	{
+		const FInstanceData* InstanceData = Binding.GetInstanceData<FInstanceData>();
+
+		// The destination child index has been updated in PreUpdate, we can use the cached version
+		const int32 DestinationChildIndex = InstanceData->PreviousChildIndex;
+
+		TDecoratorBinding<IDiscreteBlend> DiscreteBlendDecorator;
+		Context.GetInterface(Binding, DiscreteBlendDecorator);
+
+		const float BlendWeightTrue = DiscreteBlendDecorator.GetBlendWeight(Context, TRUE_CHILD_INDEX);
+		if (InstanceData->TrueChild.IsValid() && FAnimWeight::IsRelevant(BlendWeightTrue))
+		{
+			FDecoratorUpdateState DecoratorStateTrue = DecoratorState.WithWeight(BlendWeightTrue);
+			if (DestinationChildIndex != TRUE_CHILD_INDEX)
+			{
+				DecoratorStateTrue = DecoratorStateTrue.AsBlendingOut();
+			}
+
+			TraversalQueue.Push(InstanceData->TrueChild, DecoratorStateTrue);
+		}
+
+		const float BlendWeightFalse = 1.0f - BlendWeightTrue;
+		if (InstanceData->FalseChild.IsValid() && FAnimWeight::IsRelevant(BlendWeightFalse))
+		{
+			FDecoratorUpdateState DecoratorStateFalse = DecoratorState.WithWeight(BlendWeightFalse);
+			if (DestinationChildIndex != FALSE_CHILD_INDEX)
+			{
+				DecoratorStateFalse = DecoratorStateFalse.AsBlendingOut();
+			}
+
+			TraversalQueue.Push(InstanceData->FalseChild, DecoratorStateFalse);
+		}
+	}
+
+	uint32 FBlendByBoolDecorator::GetNumChildren(const FExecutionContext& Context, const TDecoratorBinding<IHierarchy>& Binding) const
 	{
 		return 2;
 	}
 
-	void FBlendByBoolDecorator::GetChildren(FExecutionContext& Context, const TDecoratorBinding<IHierarchy>& Binding, FChildrenArray& Children) const
+	void FBlendByBoolDecorator::GetChildren(const FExecutionContext& Context, const TDecoratorBinding<IHierarchy>& Binding, FChildrenArray& Children) const
 	{
 		const FInstanceData* InstanceData = Binding.GetInstanceData<FInstanceData>();
 
@@ -77,7 +107,7 @@ namespace UE::AnimNext
 		Children.Add(InstanceData->FalseChild);
 	}
 
-	float FBlendByBoolDecorator::GetBlendWeight(FExecutionContext& Context, const TDecoratorBinding<IDiscreteBlend>& Binding, int32 ChildIndex) const
+	float FBlendByBoolDecorator::GetBlendWeight(const FExecutionContext& Context, const TDecoratorBinding<IDiscreteBlend>& Binding, int32 ChildIndex) const
 	{
 		TDecoratorBinding<IDiscreteBlend> DiscreteBlendDecorator;
 		Context.GetInterface(Binding, DiscreteBlendDecorator);
@@ -99,29 +129,15 @@ namespace UE::AnimNext
 		}
 	}
 
-	int32 FBlendByBoolDecorator::GetBlendDestinationChildIndex(FExecutionContext& Context, const TDecoratorBinding<IDiscreteBlend>& Binding) const
+	int32 FBlendByBoolDecorator::GetBlendDestinationChildIndex(const FExecutionContext& Context, const TDecoratorBinding<IDiscreteBlend>& Binding) const
 	{
 		const FSharedData* SharedData = Binding.GetSharedData<FSharedData>();
 
-#if 0
-		{
-			const int32 BoolOverride = CVarAnimNextForceBlendBool.GetValueOnAnyThread();
-			if (BoolOverride == 0)
-			{
-				return FALSE_CHILD_INDEX;
-			}
-			else if (BoolOverride == 1)
-			{
-				return TRUE_CHILD_INDEX;
-			}
-		}
-#endif
-
 		const bool bCondition = SharedData->GetbCondition(Context, Binding);
 		return bCondition ? TRUE_CHILD_INDEX : FALSE_CHILD_INDEX;
 	}
 
-	void FBlendByBoolDecorator::OnBlendTransition(FExecutionContext& Context, const TDecoratorBinding<IDiscreteBlend>& Binding, int32 OldChildIndex, int32 NewChildIndex) const
+	void FBlendByBoolDecorator::OnBlendTransition(const FExecutionContext& Context, const TDecoratorBinding<IDiscreteBlend>& Binding, int32 OldChildIndex, int32 NewChildIndex) const
 	{
 		TDecoratorBinding<IDiscreteBlend> DiscreteBlendDecorator;
 		Context.GetInterface(Binding, DiscreteBlendDecorator);
@@ -133,7 +149,7 @@ namespace UE::AnimNext
 		DiscreteBlendDecorator.OnBlendTerminated(Context, OldChildIndex);
 	}
 
-	void FBlendByBoolDecorator::OnBlendInitiated(FExecutionContext& Context, const TDecoratorBinding<IDiscreteBlend>& Binding, int32 ChildIndex) const
+	void FBlendByBoolDecorator::OnBlendInitiated(const FExecutionContext& Context, const TDecoratorBinding<IDiscreteBlend>& Binding, int32 ChildIndex) const
 	{
 		const FSharedData* SharedData = Binding.GetSharedData<FSharedData>();
 		FInstanceData* InstanceData = Binding.GetInstanceData<FInstanceData>();
@@ -155,7 +171,7 @@ namespace UE::AnimNext
 		}
 	}
 
-	void FBlendByBoolDecorator::OnBlendTerminated(FExecutionContext& Context, const TDecoratorBinding<IDiscreteBlend>& Binding, int32 ChildIndex) const
+	void FBlendByBoolDecorator::OnBlendTerminated(const FExecutionContext& Context, const TDecoratorBinding<IDiscreteBlend>& Binding, int32 ChildIndex) const
 	{
 		FInstanceData* InstanceData = Binding.GetInstanceData<FInstanceData>();
 
diff --git a/Engine/Plugins/Experimental/AnimNext/Source/AnimNext/Private/Decorators/BlendByBool.h b/Engine/Plugins/Experimental/AnimNext/Source/AnimNext/Private/Decorators/BlendByBool.h
--- a/Engine/Plugins/Experimental/AnimNext/Source/AnimNext/Private/Decorators/BlendByBool.h
+++ b/Engine/Plugins/Experimental/AnimNext/Source/AnimNext/Private/Decorators/BlendByBool.h
@@ -59,20 +59,21 @@ namespace UE::AnimNext
 		};
 
 		// IEvaluate impl
-		virtual void PostEvaluate(FExecutionContext& Context, const TDecoratorBinding<IEvaluate>& Binding) const override;
+		virtual void PostEvaluate(const FExecutionContext& Context, const TDecoratorBinding<IEvaluate>& Binding) const override;
 
 		// IUpdate impl
-		virtual void PreUpdate(FExecutionContext& Context, const TDecoratorBinding<IUpdate>& Binding) const override;
+		virtual void PreUpdate(FUpdateTraversalContext& Context, const TDecoratorBinding<IUpdate>& Binding, const FDecoratorUpdateState& DecoratorState) const override;
+		virtual void QueueChildrenForTraversal(FUpdateTraversalContext& Context, const TDecoratorBinding<IUpdate>& Binding, const FDecoratorUpdateState& DecoratorState, FUpdateTraversalQueue& TraversalQueue) const override;
 
 		// IHierarchy impl
-		virtual uint32 GetNumChildren(FExecutionContext& Context, const TDecoratorBinding<IHierarchy>& Binding) const override;
-		virtual void GetChildren(FExecutionContext& Context, const TDecoratorBinding<IHierarchy>& Binding, FChildrenArray& Children) const override;
+		virtual uint32 GetNumChildren(const FExecutionContext& Context, const TDecoratorBinding<IHierarchy>& Binding) const override;
+		virtual void GetChildren(const FExecutionContext& Context, const TDecoratorBinding<IHierarchy>& Binding, FChildrenArray& Children) const override;
 
 		// IDiscreteBlend impl
-		virtual float GetBlendWeight(FExecutionContext& Context, const TDecoratorBinding<IDiscreteBlend>& Binding, int32 ChildIndex) const override;
-		virtual int32 GetBlendDestinationChildIndex(FExecutionContext& Context, const TDecoratorBinding<IDiscreteBlend>& Binding) const override;
-		virtual void OnBlendTransition(FExecutionContext& Context, const TDecoratorBinding<IDiscreteBlend>& Binding, int32 OldChildIndex, int32 NewChildIndex) const override;
-		virtual void OnBlendInitiated(FExecutionContext& Context, const TDecoratorBinding<IDiscreteBlend>& Binding, int32 ChildIndex) const override;
-		virtual void OnBlendTerminated(FExecutionContext& Context, const TDecoratorBinding<IDiscreteBlend>& Binding, int32 ChildIndex) const override;
+		virtual float GetBlendWeight(const FExecutionContext& Context, const TDecoratorBinding<IDiscreteBlend>& Binding, int32 ChildIndex) const override;
+		virtual int32 GetBlendDestinationChildIndex(const FExecutionContext& Context, const TDecoratorBinding<IDiscreteBlend>& Binding) const override;
+		virtual void OnBlendTransition(const FExecutionContext& Context, const TDecoratorBinding<IDiscreteBlend>& Binding, int32 OldChildIndex, int32 NewChildIndex) const override;
+		virtual void OnBlendInitiated(const FExecutionContext& Context, const TDecoratorBinding<IDiscreteBlend>& Binding, int32 ChildIndex) const override;
+		virtual void OnBlendTerminated(const FExecutionContext& Context, const TDecoratorBinding<IDiscreteBlend>& Binding, int32 ChildIndex) const override;
 	};
 }
diff --git a/Engine/Plugins/Experimental/AnimNext/Source/AnimNext/Private/Decorators/BlendInertializer.cpp b/Engine/Plugins/Experimental/AnimNext/Source/AnimNext/Private/Decorators/BlendInertializer.cpp
--- a/Engine/Plugins/Experimental/AnimNext/Source/AnimNext/Private/Decorators/BlendInertializer.cpp
+++ b/Engine/Plugins/Experimental/AnimNext/Source/AnimNext/Private/Decorators/BlendInertializer.cpp
@@ -13,7 +13,7 @@ namespace UE::AnimNext
 		DEFINE_ANIM_DECORATOR_IMPLEMENTS_INTERFACE(ISmoothBlend)
 	DEFINE_ANIM_DECORATOR_END(FBlendInertializerDecorator)
 
-	void FBlendInertializerDecorator::OnBlendTransition(FExecutionContext& Context, const TDecoratorBinding<IDiscreteBlend>& Binding, int32 OldChildIndex, int32 NewChildIndex) const
+	void FBlendInertializerDecorator::OnBlendTransition(const FExecutionContext& Context, const TDecoratorBinding<IDiscreteBlend>& Binding, int32 OldChildIndex, int32 NewChildIndex) const
 	{
 		// Trigger the new transition
 		IDiscreteBlend::OnBlendTransition(Context, Binding, OldChildIndex, NewChildIndex);
@@ -40,7 +40,7 @@ namespace UE::AnimNext
 #endif
 	}
 
-	float FBlendInertializerDecorator::GetBlendTime(FExecutionContext& Context, const TDecoratorBinding<ISmoothBlend>& Binding, int32 ChildIndex) const
+	float FBlendInertializerDecorator::GetBlendTime(const FExecutionContext& Context, const TDecoratorBinding<ISmoothBlend>& Binding, int32 ChildIndex) const
 	{
 		// We hijack the blend time and always transition instantaneously
 		return 0.0f;
diff --git a/Engine/Plugins/Experimental/AnimNext/Source/AnimNext/Private/Decorators/BlendInertializer.h b/Engine/Plugins/Experimental/AnimNext/Source/AnimNext/Private/Decorators/BlendInertializer.h
--- a/Engine/Plugins/Experimental/AnimNext/Source/AnimNext/Private/Decorators/BlendInertializer.h
+++ b/Engine/Plugins/Experimental/AnimNext/Source/AnimNext/Private/Decorators/BlendInertializer.h
@@ -38,9 +38,9 @@ namespace UE::AnimNext
 		using FInstanceData = FDecorator::FInstanceData;
 
 		// IDiscreteBlend impl
-		virtual void OnBlendTransition(FExecutionContext& Context, const TDecoratorBinding<IDiscreteBlend>& Binding, int32 OldChildIndex, int32 NewChildIndex) const override;
+		virtual void OnBlendTransition(const FExecutionContext& Context, const TDecoratorBinding<IDiscreteBlend>& Binding, int32 OldChildIndex, int32 NewChildIndex) const override;
 
 		// ISmoothBlend impl
-		virtual float GetBlendTime(FExecutionContext& Context, const TDecoratorBinding<ISmoothBlend>& Binding, int32 ChildIndex) const override;
+		virtual float GetBlendTime(const FExecutionContext& Context, const TDecoratorBinding<ISmoothBlend>& Binding, int32 ChildIndex) const override;
 	};
 }
diff --git a/Engine/Plugins/Experimental/AnimNext/Source/AnimNext/Private/Decorators/BlendSmoother.cpp b/Engine/Plugins/Experimental/AnimNext/Source/AnimNext/Private/Decorators/BlendSmoother.cpp
--- a/Engine/Plugins/Experimental/AnimNext/Source/AnimNext/Private/Decorators/BlendSmoother.cpp
+++ b/Engine/Plugins/Experimental/AnimNext/Source/AnimNext/Private/Decorators/BlendSmoother.cpp
@@ -2,6 +2,7 @@
 
 #include "Decorators/BlendSmoother.h"
 
+#include "Animation/AnimTypes.h"
 #include "DecoratorBase/ExecutionContext.h"
 #include "DecoratorInterfaces/IHierarchy.h"
 #include "EvaluationVM/Tasks/BlendKeyframes.h"
@@ -16,7 +17,7 @@ namespace UE::AnimNext
 		DEFINE_ANIM_DECORATOR_IMPLEMENTS_INTERFACE(ISmoothBlend)
 	DEFINE_ANIM_DECORATOR_END(FBlendSmootherDecorator)
 
-	void FBlendSmootherDecorator::PreUpdate(FExecutionContext& Context, const TDecoratorBinding<IUpdate>& Binding) const
+	void FBlendSmootherDecorator::PreUpdate(FUpdateTraversalContext& Context, const TDecoratorBinding<IUpdate>& Binding, const FDecoratorUpdateState& DecoratorState) const
 	{
 		const FSharedData* SharedData = Binding.GetSharedData<FSharedData>();
 		FInstanceData* InstanceData = Binding.GetInstanceData<FInstanceData>();
@@ -28,10 +29,9 @@ namespace UE::AnimNext
 		}
 
 		// Update the decorators below us, they might trigger a transition
-		IUpdate::PreUpdate(Context, Binding);
+		IUpdate::PreUpdate(Context, Binding, DecoratorState);
 
-		const FUpdateTraversalContext& TraversalContext = Context.GetTraversalContext<FUpdateTraversalContext>();
-		const float DeltaTime = TraversalContext.GetDeltaTime();
+		const float DeltaTime = DecoratorState.GetDeltaTime();
 
 		// Advance the weights
 		float SumWeight = 0.0f;
@@ -46,7 +46,14 @@ namespace UE::AnimNext
 
 			ChildBlendData.Blend.Update(DeltaTime);
 
-			const float NewBlendWeight = ChildBlendData.Blend.GetBlendedValue();
+			float NewBlendWeight = ChildBlendData.Blend.GetBlendedValue();
+
+			if (!FAnimWeight::IsRelevant(NewBlendWeight))
+			{
+				// Our new weight is no longer relevant, snap it to zero and normalization below will fix-up the other weights
+				// We'll then terminate the blend below
+				NewBlendWeight = 0.0f;
+			}
 
 			ChildBlendData.Weight = NewBlendWeight;
 			SumWeight += NewBlendWeight;
@@ -89,19 +96,19 @@ namespace UE::AnimNext
 		}
 	}
 
-	float FBlendSmootherDecorator::GetBlendWeight(FExecutionContext& Context, const TDecoratorBinding<IDiscreteBlend>& Binding, int32 ChildIndex) const
+	float FBlendSmootherDecorator::GetBlendWeight(const FExecutionContext& Context, const TDecoratorBinding<IDiscreteBlend>& Binding, int32 ChildIndex) const
 	{
 		const FInstanceData* InstanceData = Binding.GetInstanceData<FInstanceData>();
 		return InstanceData->PerChildBlendData.IsValidIndex(ChildIndex) ? InstanceData->PerChildBlendData[ChildIndex].Weight : -1.0f;
 	}
 
-	const FAlphaBlend* FBlendSmootherDecorator::GetBlendState(FExecutionContext& Context, const TDecoratorBinding<IDiscreteBlend>& Binding, int32 ChildIndex) const
+	const FAlphaBlend* FBlendSmootherDecorator::GetBlendState(const FExecutionContext& Context, const TDecoratorBinding<IDiscreteBlend>& Binding, int32 ChildIndex) const
 	{
 		const FInstanceData* InstanceData = Binding.GetInstanceData<FInstanceData>();
 		return InstanceData->PerChildBlendData.IsValidIndex(ChildIndex) ? &InstanceData->PerChildBlendData[ChildIndex].Blend : nullptr;
 	}
 
-	void FBlendSmootherDecorator::OnBlendTransition(FExecutionContext& Context, const TDecoratorBinding<IDiscreteBlend>& Binding, int32 OldChildIndex, int32 NewChildIndex) const
+	void FBlendSmootherDecorator::OnBlendTransition(const FExecutionContext& Context, const TDecoratorBinding<IDiscreteBlend>& Binding, int32 OldChildIndex, int32 NewChildIndex) const
 	{
 		const FSharedData* SharedData = Binding.GetSharedData<FSharedData>();
 		FInstanceData* InstanceData = Binding.GetInstanceData<FInstanceData>();
@@ -150,13 +157,13 @@ namespace UE::AnimNext
 		DiscreteBlendDecorator.OnBlendInitiated(Context, NewChildIndex);
 	}
 
-	float FBlendSmootherDecorator::GetBlendTime(FExecutionContext& Context, const TDecoratorBinding<ISmoothBlend>& Binding, int32 ChildIndex) const
+	float FBlendSmootherDecorator::GetBlendTime(const FExecutionContext& Context, const TDecoratorBinding<ISmoothBlend>& Binding, int32 ChildIndex) const
 	{
 		const FSharedData* SharedData = Binding.GetSharedData<FSharedData>();
 		return SharedData->BlendTimes.IsValidIndex(ChildIndex) ? SharedData->BlendTimes[ChildIndex] : 0.0f;
 	}
 
-	void FBlendSmootherDecorator::InitializeInstanceData(FExecutionContext& Context, const FDecoratorBinding& Binding, const FSharedData* SharedData, FInstanceData* InstanceData)
+	void FBlendSmootherDecorator::InitializeInstanceData(const FExecutionContext& Context, const FDecoratorBinding& Binding, const FSharedData* SharedData, FInstanceData* InstanceData)
 	{
 		check(InstanceData->PerChildBlendData.IsEmpty());
 
diff --git a/Engine/Plugins/Experimental/AnimNext/Source/AnimNext/Private/Decorators/BlendSmoother.h b/Engine/Plugins/Experimental/AnimNext/Source/AnimNext/Private/Decorators/BlendSmoother.h
--- a/Engine/Plugins/Experimental/AnimNext/Source/AnimNext/Private/Decorators/BlendSmoother.h
+++ b/Engine/Plugins/Experimental/AnimNext/Source/AnimNext/Private/Decorators/BlendSmoother.h
@@ -65,17 +65,17 @@ namespace UE::AnimNext
 		};
 
 		// IUpdate impl
-		virtual void PreUpdate(FExecutionContext& Context, const TDecoratorBinding<IUpdate>& Binding) const override;
+		virtual void PreUpdate(FUpdateTraversalContext& Context, const TDecoratorBinding<IUpdate>& Binding, const FDecoratorUpdateState& DecoratorState) const override;
 
 		// IDiscreteBlend impl
-		virtual float GetBlendWeight(FExecutionContext& Context, const TDecoratorBinding<IDiscreteBlend>& Binding, int32 ChildIndex) const override;
-		virtual const FAlphaBlend* GetBlendState(FExecutionContext& Context, const TDecoratorBinding<IDiscreteBlend>& Binding, int32 ChildIndex) const override;
-		virtual void OnBlendTransition(FExecutionContext& Context, const TDecoratorBinding<IDiscreteBlend>& Binding, int32 OldChildIndex, int32 NewChildIndex) const override;
+		virtual float GetBlendWeight(const FExecutionContext& Context, const TDecoratorBinding<IDiscreteBlend>& Binding, int32 ChildIndex) const override;
+		virtual const FAlphaBlend* GetBlendState(const FExecutionContext& Context, const TDecoratorBinding<IDiscreteBlend>& Binding, int32 ChildIndex) const override;
+		virtual void OnBlendTransition(const FExecutionContext& Context, const TDecoratorBinding<IDiscreteBlend>& Binding, int32 OldChildIndex, int32 NewChildIndex) const override;
 
 		// ISmoothBlend impl
-		virtual float GetBlendTime(FExecutionContext& Context, const TDecoratorBinding<ISmoothBlend>& Binding, int32 ChildIndex) const override;
+		virtual float GetBlendTime(const FExecutionContext& Context, const TDecoratorBinding<ISmoothBlend>& Binding, int32 ChildIndex) const override;
 
 		// Internal impl
-		static void InitializeInstanceData(FExecutionContext& Context, const FDecoratorBinding& Binding, const FSharedData* SharedData, FInstanceData* InstanceData);
+		static void InitializeInstanceData(const FExecutionContext& Context, const FDecoratorBinding& Binding, const FSharedData* SharedData, FInstanceData* InstanceData);
 	};
 }
diff --git a/Engine/Plugins/Experimental/AnimNext/Source/AnimNext/Private/Decorators/BlendSmootherPerBone.cpp b/Engine/Plugins/Experimental/AnimNext/Source/AnimNext/Private/Decorators/BlendSmootherPerBone.cpp
--- a/Engine/Plugins/Experimental/AnimNext/Source/AnimNext/Private/Decorators/BlendSmootherPerBone.cpp
+++ b/Engine/Plugins/Experimental/AnimNext/Source/AnimNext/Private/Decorators/BlendSmootherPerBone.cpp
@@ -19,7 +19,7 @@ namespace UE::AnimNext
 		DEFINE_ANIM_DECORATOR_IMPLEMENTS_INTERFACE(IDiscreteBlend)
 	DEFINE_ANIM_DECORATOR_END(FBlendSmootherPerBoneDecorator)
 
-	void FBlendSmootherPerBoneDecorator::PostEvaluate(FExecutionContext& Context, const TDecoratorBinding<IEvaluate>& Binding) const
+	void FBlendSmootherPerBoneDecorator::PostEvaluate(const FExecutionContext& Context, const TDecoratorBinding<IEvaluate>& Binding) const
 	{
 		const FSharedData* SharedData = Binding.GetSharedData<FSharedData>();
 		const FInstanceData* InstanceData = Binding.GetInstanceData<FInstanceData>();
@@ -67,7 +67,7 @@ namespace UE::AnimNext
 		TraversalContext.AppendTask(FAnimNextNormalizeKeyframeRotationsTask());
 	}
 
-	void FBlendSmootherPerBoneDecorator::PreUpdate(FExecutionContext& Context, const TDecoratorBinding<IUpdate>& Binding) const
+	void FBlendSmootherPerBoneDecorator::PreUpdate(FUpdateTraversalContext& Context, const TDecoratorBinding<IUpdate>& Binding, const FDecoratorUpdateState& DecoratorState) const
 	{
 		const FSharedData* SharedData = Binding.GetSharedData<FSharedData>();
 		FInstanceData* InstanceData = Binding.GetInstanceData<FInstanceData>();
@@ -79,7 +79,7 @@ namespace UE::AnimNext
 		}
 
 		// Update the decorators below us, they might trigger a transition
-		IUpdate::PreUpdate(Context, Binding);
+		IUpdate::PreUpdate(Context, Binding, DecoratorState);
 
 		if (!SharedData->BlendProfile)
 		{
@@ -109,7 +109,7 @@ namespace UE::AnimNext
 		FBlendSampleData::NormalizeDataWeight(InstanceData->PerBoneSampleData);
 	}
 
-	void FBlendSmootherPerBoneDecorator::OnBlendTransition(FExecutionContext& Context, const TDecoratorBinding<IDiscreteBlend>& Binding, int32 OldChildIndex, int32 NewChildIndex) const
+	void FBlendSmootherPerBoneDecorator::OnBlendTransition(const FExecutionContext& Context, const TDecoratorBinding<IDiscreteBlend>& Binding, int32 OldChildIndex, int32 NewChildIndex) const
 	{
 		const FSharedData* SharedData = Binding.GetSharedData<FSharedData>();
 		FInstanceData* InstanceData = Binding.GetInstanceData<FInstanceData>();
@@ -135,7 +135,7 @@ namespace UE::AnimNext
 		}
 	}
 
-	void FBlendSmootherPerBoneDecorator::InitializeInstanceData(FExecutionContext& Context, const FDecoratorBinding& Binding, const FSharedData* SharedData, FInstanceData* InstanceData)
+	void FBlendSmootherPerBoneDecorator::InitializeInstanceData(const FExecutionContext& Context, const FDecoratorBinding& Binding, const FSharedData* SharedData, FInstanceData* InstanceData)
 	{
 		check(InstanceData->PerChildBlendData.IsEmpty());
 
diff --git a/Engine/Plugins/Experimental/AnimNext/Source/AnimNext/Private/Decorators/BlendSmootherPerBone.h b/Engine/Plugins/Experimental/AnimNext/Source/AnimNext/Private/Decorators/BlendSmootherPerBone.h
--- a/Engine/Plugins/Experimental/AnimNext/Source/AnimNext/Private/Decorators/BlendSmootherPerBone.h
+++ b/Engine/Plugins/Experimental/AnimNext/Source/AnimNext/Private/Decorators/BlendSmootherPerBone.h
@@ -56,15 +56,15 @@ namespace UE::AnimNext
 		};
 
 		// IEvaluate impl
-		virtual void PostEvaluate(FExecutionContext& Context, const TDecoratorBinding<IEvaluate>& Binding) const override;
+		virtual void PostEvaluate(const FExecutionContext& Context, const TDecoratorBinding<IEvaluate>& Binding) const override;
 
 		// IUpdate impl
-		virtual void PreUpdate(FExecutionContext& Context, const TDecoratorBinding<IUpdate>& Binding) const override;
+		virtual void PreUpdate(FUpdateTraversalContext& Context, const TDecoratorBinding<IUpdate>& Binding, const FDecoratorUpdateState& DecoratorState) const override;
 
 		// IDiscreteBlend impl
-		virtual void OnBlendTransition(FExecutionContext& Context, const TDecoratorBinding<IDiscreteBlend>& Binding, int32 OldChildIndex, int32 NewChildIndex) const override;
+		virtual void OnBlendTransition(const FExecutionContext& Context, const TDecoratorBinding<IDiscreteBlend>& Binding, int32 OldChildIndex, int32 NewChildIndex) const override;
 
 		// Internal impl
-		static void InitializeInstanceData(FExecutionContext& Context, const FDecoratorBinding& Binding, const FSharedData* SharedData, FInstanceData* InstanceData);
+		static void InitializeInstanceData(const FExecutionContext& Context, const FDecoratorBinding& Binding, const FSharedData* SharedData, FInstanceData* InstanceData);
 	};
 }
diff --git a/Engine/Plugins/Experimental/AnimNext/Source/AnimNext/Private/Decorators/BlendTwoWay.cpp b/Engine/Plugins/Experimental/AnimNext/Source/AnimNext/Private/Decorators/BlendTwoWay.cpp
--- a/Engine/Plugins/Experimental/AnimNext/Source/AnimNext/Private/Decorators/BlendTwoWay.cpp
+++ b/Engine/Plugins/Experimental/AnimNext/Source/AnimNext/Private/Decorators/BlendTwoWay.cpp
@@ -17,7 +17,7 @@ namespace UE::AnimNext
 		DEFINE_ANIM_DECORATOR_IMPLEMENTS_INTERFACE(IUpdate)
 	DEFINE_ANIM_DECORATOR_END(FBlendTwoWayDecorator)
 
-	void FBlendTwoWayDecorator::PostEvaluate(FExecutionContext& Context, const TDecoratorBinding<IEvaluate>& Binding) const
+	void FBlendTwoWayDecorator::PostEvaluate(const FExecutionContext& Context, const TDecoratorBinding<IEvaluate>& Binding) const
 	{
 		const FInstanceData* InstanceData = Binding.GetInstanceData<FInstanceData>();
 
@@ -39,7 +39,7 @@ namespace UE::AnimNext
 		}
 	}
 
-	void FBlendTwoWayDecorator::PreUpdate(FExecutionContext& Context, const TDecoratorBinding<IUpdate>& Binding) const
+	void FBlendTwoWayDecorator::PreUpdate(FUpdateTraversalContext& Context, const TDecoratorBinding<IUpdate>& Binding, const FDecoratorUpdateState& DecoratorState) const
 	{
 		const FSharedData* SharedData = Binding.GetSharedData<FSharedData>();
 		FInstanceData* InstanceData = Binding.GetInstanceData<FInstanceData>();
@@ -47,8 +47,8 @@ namespace UE::AnimNext
 		TDecoratorBinding<IContinuousBlend> ContinuousBlendDecorator;
 		Context.GetInterface(Binding, ContinuousBlendDecorator);
 
-		const float BlendWeight = ContinuousBlendDecorator.GetBlendWeight(Context, 1);
-		if (!FAnimWeight::IsFullWeight(BlendWeight))
+		const float BlendWeightB = ContinuousBlendDecorator.GetBlendWeight(Context, 1);
+		if (!FAnimWeight::IsFullWeight(BlendWeightB))
 		{
 			if (!InstanceData->ChildA.IsValid())
 			{
@@ -56,14 +56,14 @@ namespace UE::AnimNext
 				InstanceData->ChildA = Context.AllocateNodeInstance(Binding, SharedData->ChildA);
 			}
 
-			if (!FAnimWeight::IsRelevant(BlendWeight))
+			if (!FAnimWeight::IsRelevant(BlendWeightB))
 			{
 				// We no longer need this child, release it
 				InstanceData->ChildB.Reset();
 			}
 		}
 
-		if (FAnimWeight::IsRelevant(BlendWeight))
+		if (FAnimWeight::IsRelevant(BlendWeightB))
 		{
 			if (!InstanceData->ChildB.IsValid())
 			{
@@ -71,7 +71,7 @@ namespace UE::AnimNext
 				InstanceData->ChildB = Context.AllocateNodeInstance(Binding, SharedData->ChildB);
 			}
 
-			if (FAnimWeight::IsFullWeight(BlendWeight))
+			if (FAnimWeight::IsFullWeight(BlendWeightB))
 			{
 				// We no longer need this child, release it
 				InstanceData->ChildA.Reset();
@@ -79,12 +79,32 @@ namespace UE::AnimNext
 		}
 	}
 
-	uint32 FBlendTwoWayDecorator::GetNumChildren(FExecutionContext& Context, const TDecoratorBinding<IHierarchy>& Binding) const
+	void FBlendTwoWayDecorator::QueueChildrenForTraversal(FUpdateTraversalContext& Context, const TDecoratorBinding<IUpdate>& Binding, const FDecoratorUpdateState& DecoratorState, FUpdateTraversalQueue& TraversalQueue) const
+	{
+		const FInstanceData* InstanceData = Binding.GetInstanceData<FInstanceData>();
+
+		TDecoratorBinding<IContinuousBlend> ContinuousBlendDecorator;
+		Context.GetInterface(Binding, ContinuousBlendDecorator);
+
+		const float BlendWeightB = ContinuousBlendDecorator.GetBlendWeight(Context, 1);
+		if (InstanceData->ChildA.IsValid())
+		{
+			const float BlendWeightA = 1.0f - BlendWeightB;
+			TraversalQueue.Push(InstanceData->ChildA, DecoratorState.WithWeight(BlendWeightA));
+		}
+
+		if (InstanceData->ChildB.IsValid())
+		{
+			TraversalQueue.Push(InstanceData->ChildB, DecoratorState.WithWeight(BlendWeightB));
+		}
+	}
+
+	uint32 FBlendTwoWayDecorator::GetNumChildren(const FExecutionContext& Context, const TDecoratorBinding<IHierarchy>& Binding) const
 	{
 		return 2;
 	}
 
-	void FBlendTwoWayDecorator::GetChildren(FExecutionContext& Context, const TDecoratorBinding<IHierarchy>& Binding, FChildrenArray& Children) const
+	void FBlendTwoWayDecorator::GetChildren(const FExecutionContext& Context, const TDecoratorBinding<IHierarchy>& Binding, FChildrenArray& Children) const
 	{
 		const FInstanceData* InstanceData = Binding.GetInstanceData<FInstanceData>();
 
@@ -93,7 +113,7 @@ namespace UE::AnimNext
 		Children.Add(InstanceData->ChildB);
 	}
 
-	float FBlendTwoWayDecorator::GetBlendWeight(FExecutionContext& Context, const TDecoratorBinding<IContinuousBlend>& Binding, int32 ChildIndex) const
+	float FBlendTwoWayDecorator::GetBlendWeight(const FExecutionContext& Context, const TDecoratorBinding<IContinuousBlend>& Binding, int32 ChildIndex) const
 	{
 		const FSharedData* SharedData = Binding.GetSharedData<FSharedData>();
 
diff --git a/Engine/Plugins/Experimental/AnimNext/Source/AnimNext/Private/Decorators/BlendTwoWay.h b/Engine/Plugins/Experimental/AnimNext/Source/AnimNext/Private/Decorators/BlendTwoWay.h
--- a/Engine/Plugins/Experimental/AnimNext/Source/AnimNext/Private/Decorators/BlendTwoWay.h
+++ b/Engine/Plugins/Experimental/AnimNext/Source/AnimNext/Private/Decorators/BlendTwoWay.h
@@ -57,16 +57,17 @@ namespace UE::AnimNext
 		};
 
 		// IEvaluate impl
-		virtual void PostEvaluate(FExecutionContext& Context, const TDecoratorBinding<IEvaluate>& Binding) const override;
+		virtual void PostEvaluate(const FExecutionContext& Context, const TDecoratorBinding<IEvaluate>& Binding) const override;
 
 		// IUpdate impl
-		virtual void PreUpdate(FExecutionContext& Context, const TDecoratorBinding<IUpdate>& Binding) const override;
+		virtual void PreUpdate(FUpdateTraversalContext& Context, const TDecoratorBinding<IUpdate>& Binding, const FDecoratorUpdateState& DecoratorState) const override;
+		virtual void QueueChildrenForTraversal(FUpdateTraversalContext& Context, const TDecoratorBinding<IUpdate>& Binding, const FDecoratorUpdateState& DecoratorState, FUpdateTraversalQueue& TraversalQueue) const override;
 
 		// IHierarchy impl
-		virtual uint32 GetNumChildren(FExecutionContext& Context, const TDecoratorBinding<IHierarchy>& Binding) const override;
-		virtual void GetChildren(FExecutionContext& Context, const TDecoratorBinding<IHierarchy>& Binding, FChildrenArray& Children) const override;
+		virtual uint32 GetNumChildren(const FExecutionContext& Context, const TDecoratorBinding<IHierarchy>& Binding) const override;
+		virtual void GetChildren(const FExecutionContext& Context, const TDecoratorBinding<IHierarchy>& Binding, FChildrenArray& Children) const override;
 
 		// IContinuousBlend impl
-		virtual float GetBlendWeight(FExecutionContext& Context, const TDecoratorBinding<IContinuousBlend>& Binding, int32 ChildIndex) const override;
+		virtual float GetBlendWeight(const FExecutionContext& Context, const TDecoratorBinding<IContinuousBlend>& Binding, int32 ChildIndex) const override;
 	};
 }
diff --git a/Engine/Plugins/Experimental/AnimNext/Source/AnimNext/Private/Decorators/ReferencePoseDecorator.cpp b/Engine/Plugins/Experimental/AnimNext/Source/AnimNext/Private/Decorators/ReferencePoseDecorator.cpp
--- a/Engine/Plugins/Experimental/AnimNext/Source/AnimNext/Private/Decorators/ReferencePoseDecorator.cpp
+++ b/Engine/Plugins/Experimental/AnimNext/Source/AnimNext/Private/Decorators/ReferencePoseDecorator.cpp
@@ -13,7 +13,7 @@ namespace UE::AnimNext
 		DEFINE_ANIM_DECORATOR_IMPLEMENTS_INTERFACE(IEvaluate)
 	DEFINE_ANIM_DECORATOR_END(FReferencePoseDecorator)
 
-	void FReferencePoseDecorator::PreEvaluate(FExecutionContext& Context, const TDecoratorBinding<IEvaluate>& Binding) const
+	void FReferencePoseDecorator::PreEvaluate(const FExecutionContext& Context, const TDecoratorBinding<IEvaluate>& Binding) const
 	{
 		const FSharedData* SharedData = Binding.GetSharedData<FSharedData>();
 
diff --git a/Engine/Plugins/Experimental/AnimNext/Source/AnimNext/Private/Decorators/ReferencePoseDecorator.h b/Engine/Plugins/Experimental/AnimNext/Source/AnimNext/Private/Decorators/ReferencePoseDecorator.h
--- a/Engine/Plugins/Experimental/AnimNext/Source/AnimNext/Private/Decorators/ReferencePoseDecorator.h
+++ b/Engine/Plugins/Experimental/AnimNext/Source/AnimNext/Private/Decorators/ReferencePoseDecorator.h
@@ -45,6 +45,6 @@ namespace UE::AnimNext
 		using FSharedData = FAnimNextReferencePoseDecoratorSharedData;
 
 		// IEvaluate impl
-		virtual void PreEvaluate(FExecutionContext& Context, const TDecoratorBinding<IEvaluate>& Binding) const override;
+		virtual void PreEvaluate(const FExecutionContext& Context, const TDecoratorBinding<IEvaluate>& Binding) const override;
 	};
 }
diff --git a/Engine/Plugins/Experimental/AnimNext/Source/AnimNext/Private/Decorators/SequencePlayer.cpp b/Engine/Plugins/Experimental/AnimNext/Source/AnimNext/Private/Decorators/SequencePlayer.cpp
--- a/Engine/Plugins/Experimental/AnimNext/Source/AnimNext/Private/Decorators/SequencePlayer.cpp
+++ b/Engine/Plugins/Experimental/AnimNext/Source/AnimNext/Private/Decorators/SequencePlayer.cpp
@@ -15,7 +15,7 @@ namespace UE::AnimNext
 		DEFINE_ANIM_DECORATOR_IMPLEMENTS_INTERFACE(ITimeline)
 	DEFINE_ANIM_DECORATOR_END(FSequencePlayerDecorator)
 
-	void FSequencePlayerDecorator::FInstanceData::Construct(FExecutionContext& Context, const FDecoratorBinding& Binding)
+	void FSequencePlayerDecorator::FInstanceData::Construct(const FExecutionContext& Context, const FDecoratorBinding& Binding)
 	{
 		const FSharedData* SharedData = Binding.GetSharedData<FSharedData>();
 		if (SharedData->AnimSequence != nullptr)
@@ -26,7 +26,7 @@ namespace UE::AnimNext
 		}
 	}
 
-	void FSequencePlayerDecorator::PreEvaluate(FExecutionContext& Context, const TDecoratorBinding<IEvaluate>& Binding) const
+	void FSequencePlayerDecorator::PreEvaluate(const FExecutionContext& Context, const TDecoratorBinding<IEvaluate>& Binding) const
 	{
 		const FSharedData* SharedData = Binding.GetSharedData<FSharedData>();
 
@@ -40,24 +40,23 @@ namespace UE::AnimNext
 		TraversalContext.AppendTask(Task);
 	}
 
-	double FSequencePlayerDecorator::GetPlayRate(FExecutionContext& Context, const TDecoratorBinding<ITimeline>& Binding) const
+	double FSequencePlayerDecorator::GetPlayRate(const FExecutionContext& Context, const TDecoratorBinding<ITimeline>& Binding) const
 	{
 		const FSharedData* SharedData = Binding.GetSharedData<FSharedData>();
 		return SharedData->GetPlayRate(Context, Binding);
 	}
 
-	void FSequencePlayerDecorator::PreUpdate(FExecutionContext& Context, const TDecoratorBinding<IUpdate>& Binding) const
+	void FSequencePlayerDecorator::PreUpdate(FUpdateTraversalContext& Context, const TDecoratorBinding<IUpdate>& Binding, const FDecoratorUpdateState& DecoratorState) const
 	{
 		const FSharedData* SharedData = Binding.GetSharedData<FSharedData>();
 		if (SharedData->AnimSequence != nullptr)
 		{
 			FInstanceData* InstanceData = Binding.GetInstanceData<FInstanceData>();
-			const FUpdateTraversalContext& TraversalContext = Context.GetTraversalContext<FUpdateTraversalContext>();
 
 			TDecoratorBinding<ITimeline> TimelineDecorator;
 			Context.GetInterface(Binding, TimelineDecorator);
 
-			const float DeltaTime = TraversalContext.GetDeltaTime();
+			const float DeltaTime = DecoratorState.GetDeltaTime();
 			const float PlayRate = (float)TimelineDecorator.GetPlayRate(Context);
 
 			const float EffectiveDelta = FMath::IsNearlyZero(DeltaTime) || FMath::IsNearlyZero(PlayRate) ? 0.f : DeltaTime * PlayRate;
diff --git a/Engine/Plugins/Experimental/AnimNext/Source/AnimNext/Private/Decorators/SequencePlayer.h b/Engine/Plugins/Experimental/AnimNext/Source/AnimNext/Private/Decorators/SequencePlayer.h
--- a/Engine/Plugins/Experimental/AnimNext/Source/AnimNext/Private/Decorators/SequencePlayer.h
+++ b/Engine/Plugins/Experimental/AnimNext/Source/AnimNext/Private/Decorators/SequencePlayer.h
@@ -60,16 +60,16 @@ namespace UE::AnimNext
 			float InternalTimeAccumulator = 0.0f;
 			float PrevInternalTimeAccumulator = 0.0f;
 
-			void Construct(FExecutionContext& Context, const FDecoratorBinding& Binding);
+			void Construct(const FExecutionContext& Context, const FDecoratorBinding& Binding);
 		};
 
 		// IEvaluate impl
-		virtual void PreEvaluate(FExecutionContext& Context, const TDecoratorBinding<IEvaluate>& Binding) const override;
+		virtual void PreEvaluate(const FExecutionContext& Context, const TDecoratorBinding<IEvaluate>& Binding) const override;
 
 		// ITimeline impl
-		virtual double GetPlayRate(FExecutionContext& Context, const TDecoratorBinding<ITimeline>& Binding) const override;
+		virtual double GetPlayRate(const FExecutionContext& Context, const TDecoratorBinding<ITimeline>& Binding) const override;
 
 		// IUpdate impl
-		virtual void PreUpdate(FExecutionContext& Context, const TDecoratorBinding<IUpdate>& Binding) const override;
+		virtual void PreUpdate(FUpdateTraversalContext& Context, const TDecoratorBinding<IUpdate>& Binding, const FDecoratorUpdateState& DecoratorState) const override;
 	};
 }
diff --git a/Engine/Plugins/Experimental/AnimNext/Source/AnimNext/Private/Graph/RigUnit_AnimNextGraphEvaluator.cpp b/Engine/Plugins/Experimental/AnimNext/Source/AnimNext/Private/Graph/RigUnit_AnimNextGraphEvaluator.cpp
--- a/Engine/Plugins/Experimental/AnimNext/Source/AnimNext/Private/Graph/RigUnit_AnimNextGraphEvaluator.cpp
+++ b/Engine/Plugins/Experimental/AnimNext/Source/AnimNext/Private/Graph/RigUnit_AnimNextGraphEvaluator.cpp
@@ -47,8 +47,7 @@ void FRigUnit_AnimNextGraphEvaluator::StaticExecute(FRigVMExtendedExecuteContext
 		const FContext& InterfaceContext = VMExecuteContext.GetContext();
 
 		// Call pre/post update on our graph
-		FUpdateTraversalContext TraversalContext(InterfaceContext.GetDeltaTime());
-		UpdateGraph(Context, TraversalContext, GraphInstancePtr);
+		UpdateGraph(Context, GraphInstancePtr, InterfaceContext.GetDeltaTime());
 	}
 
 	if (EnumHasAnyFlags(SimulationSteps, EAnimNextGraphSimulationSteps::Evaluate))
diff --git a/Engine/Plugins/Experimental/AnimNext/Source/AnimNext/Public/DecoratorBase/Decorator.h b/Engine/Plugins/Experimental/AnimNext/Source/AnimNext/Public/DecoratorBase/Decorator.h
--- a/Engine/Plugins/Experimental/AnimNext/Source/AnimNext/Public/DecoratorBase/Decorator.h
+++ b/Engine/Plugins/Experimental/AnimNext/Source/AnimNext/Public/DecoratorBase/Decorator.h
@@ -29,8 +29,8 @@ class FArchive;
 	static const UE::AnimNext::FDecoratorMemoryLayout DecoratorMemoryDescription; \
 	virtual UE::AnimNext::FDecoratorMemoryLayout GetDecoratorMemoryDescription() const override { return DecoratorMemoryDescription; } \
 	virtual UScriptStruct* GetDecoratorSharedDataStruct() const override { return FSharedData::StaticStruct(); } \
-	virtual void ConstructDecoratorInstance(UE::AnimNext::FExecutionContext& Context, const UE::AnimNext::FDecoratorBinding& Binding) const override; \
-	virtual void DestructDecoratorInstance(UE::AnimNext::FExecutionContext& Context, const UE::AnimNext::FDecoratorBinding& Binding) const override; \
+	virtual void ConstructDecoratorInstance(const UE::AnimNext::FExecutionContext& Context, const UE::AnimNext::FDecoratorBinding& Binding) const override; \
+	virtual void DestructDecoratorInstance(const UE::AnimNext::FExecutionContext& Context, const UE::AnimNext::FDecoratorBinding& Binding) const override; \
 	virtual const UE::AnimNext::IDecoratorInterface* GetDecoratorInterface(UE::AnimNext::FDecoratorInterfaceUID InterfaceUID) const override; \
 	virtual uint32 GetNumLatentDecoratorProperties() const override { return -FSharedData::GetLatentPropertyIndex(~(size_t)0); } \
 	static_assert(std::is_base_of<FAnimNextDecoratorSharedData, FSharedData>::value, "Decorator shared data must derive from FAnimNextDecoratorSharedData"); \
@@ -55,12 +55,12 @@ class FArchive;
 #define DEFINE_ANIM_DECORATOR_BEGIN(DecoratorName) \
 	const UE::AnimNext::FDecoratorMemoryLayout DecoratorName::DecoratorMemoryDescription = \
 		UE::AnimNext::FDecoratorMemoryLayout{ sizeof(DecoratorName), alignof(DecoratorName), sizeof(DecoratorName::FSharedData), alignof(DecoratorName::FSharedData), sizeof(DecoratorName::FInstanceData), alignof(DecoratorName::FInstanceData) }; \
-	void DecoratorName::ConstructDecoratorInstance(UE::AnimNext::FExecutionContext& Context, const UE::AnimNext::FDecoratorBinding& Binding) const \
+	void DecoratorName::ConstructDecoratorInstance(const UE::AnimNext::FExecutionContext& Context, const UE::AnimNext::FDecoratorBinding& Binding) const \
 	{ \
 		FInstanceData* Data = new(Binding.GetInstanceData<FInstanceData>()) FInstanceData(); \
 		Data->Construct(Context, Binding); \
 	} \
-	void DecoratorName::DestructDecoratorInstance(UE::AnimNext::FExecutionContext& Context, const UE::AnimNext::FDecoratorBinding& Binding) const \
+	void DecoratorName::DestructDecoratorInstance(const UE::AnimNext::FExecutionContext& Context, const UE::AnimNext::FDecoratorBinding& Binding) const \
 	{ \
 		FInstanceData* Data = Binding.GetInstanceData<FInstanceData>(); \
 		Data->Destruct(Context, Binding); \
@@ -176,8 +176,8 @@ namespace UE::AnimNext
 
 		// Called when a new instance of the decorator is created or destroyed
 		// Derived types must override this and forward to the instance data constructor/destructor
-		virtual void ConstructDecoratorInstance(FExecutionContext& Context, const FDecoratorBinding& Binding) const = 0;
-		virtual void DestructDecoratorInstance(FExecutionContext& Context, const FDecoratorBinding& Binding) const = 0;
+		virtual void ConstructDecoratorInstance(const FExecutionContext& Context, const FDecoratorBinding& Binding) const = 0;
+		virtual void DestructDecoratorInstance(const FExecutionContext& Context, const FDecoratorBinding& Binding) const = 0;
 
 		// Returns the decorator mode.
 		virtual EDecoratorMode GetDecoratorMode() const = 0;
diff --git a/Engine/Plugins/Experimental/AnimNext/Source/AnimNext/Public/DecoratorBase/DecoratorInstanceData.h b/Engine/Plugins/Experimental/AnimNext/Source/AnimNext/Public/DecoratorBase/DecoratorInstanceData.h
--- a/Engine/Plugins/Experimental/AnimNext/Source/AnimNext/Public/DecoratorBase/DecoratorInstanceData.h
+++ b/Engine/Plugins/Experimental/AnimNext/Source/AnimNext/Public/DecoratorBase/DecoratorInstanceData.h
@@ -27,7 +27,7 @@ namespace UE::AnimNext
 		// You can override this function by adding a new one with the same name on your
 		// derived type.
 		// Decorators are constructed from the bottom to the top.
-		void Construct(FExecutionContext& Context, const FDecoratorBinding& Binding) noexcept
+		void Construct(const FExecutionContext& Context, const FDecoratorBinding& Binding) noexcept
 		{
 		}
 
@@ -36,7 +36,7 @@ namespace UE::AnimNext
 		// You can override this function by adding a new one with the same name on your
 		// derived type.
 		// Decorators are destructed from the top to the bottom.
-		void Destruct(FExecutionContext& Context, const FDecoratorBinding& Binding) noexcept
+		void Destruct(const FExecutionContext& Context, const FDecoratorBinding& Binding) noexcept
 		{
 		}
 	};
diff --git a/Engine/Plugins/Experimental/AnimNext/Source/AnimNext/Public/DecoratorBase/DecoratorPtr.h b/Engine/Plugins/Experimental/AnimNext/Source/AnimNext/Public/DecoratorBase/DecoratorPtr.h
--- a/Engine/Plugins/Experimental/AnimNext/Source/AnimNext/Public/DecoratorBase/DecoratorPtr.h
+++ b/Engine/Plugins/Experimental/AnimNext/Source/AnimNext/Public/DecoratorBase/DecoratorPtr.h
@@ -28,7 +28,7 @@ namespace UE::AnimNext
 		constexpr FDecoratorPtr() noexcept = default;
 
 		// Constructs a pointer handle to the provided instance
-		FDecoratorPtr(FNodeInstance* NodeInstance, uint32 DecoratorIndex_);
+		FDecoratorPtr(FNodeInstance* InNodeInstance, uint32 InDecoratorIndex);
 
 		FDecoratorPtr(const FDecoratorPtr& DecoratorPtr);
 		FDecoratorPtr(FDecoratorPtr&& DecoratorPtr) noexcept;
@@ -70,7 +70,7 @@ namespace UE::AnimNext
 		};
 
 		// Constructs a pointer handle to the provided instance
-		FDecoratorPtr(FNodeInstance* NodeInstance, EFlags Flags, uint32 DecoratorIndex_);
+		FDecoratorPtr(FNodeInstance* InNodeInstance, EFlags InFlags, uint32 InDecoratorIndex);
 
 		// Packed pointer value contains a pointer along with flags in the alignment bits
 		// The pointer part points to a node instance
@@ -101,7 +101,7 @@ namespace UE::AnimNext
 		{}
 
 		// Constructs a weak pointer handle to the provided instance
-		FWeakDecoratorPtr(FNodeInstance* NodeInstance_, uint32 DecoratorIndex_);
+		FWeakDecoratorPtr(FNodeInstance* InNodeInstance, uint32 InDecoratorIndex);
 
 		// Returns a pointer to the node instance
 		constexpr FNodeInstance* GetNodeInstance() const noexcept { return NodeInstance; }
diff --git a/Engine/Plugins/Experimental/AnimNext/Source/AnimNext/Public/DecoratorBase/DecoratorSharedData.h b/Engine/Plugins/Experimental/AnimNext/Source/AnimNext/Public/DecoratorBase/DecoratorSharedData.h
--- a/Engine/Plugins/Experimental/AnimNext/Source/AnimNext/Public/DecoratorBase/DecoratorSharedData.h
+++ b/Engine/Plugins/Experimental/AnimNext/Source/AnimNext/Public/DecoratorBase/DecoratorSharedData.h
@@ -37,7 +37,7 @@
 	} \
 
 #define ANIM_NEXT_IMPL_DEFINE_LATENT_GETTER(PropertyName) \
-	decltype(PropertyName) Get##PropertyName(UE::AnimNext::FExecutionContext& Context, const UE::AnimNext::FDecoratorBinding& Binding) const \
+	decltype(PropertyName) Get##PropertyName(const UE::AnimNext::FExecutionContext& Context, const UE::AnimNext::FDecoratorBinding& Binding) const \
 	{ \
 		/* We need a mapping of latent property name/offset to latent property index */ \
 		/* This can be built once at runtime using the UE reflection and cached on first call or using a constexpr function, see below */ \
diff --git a/Engine/Plugins/Experimental/AnimNext/Source/AnimNext/Public/DecoratorBase/ExecutionContext.h b/Engine/Plugins/Experimental/AnimNext/Source/AnimNext/Public/DecoratorBase/ExecutionContext.h
--- a/Engine/Plugins/Experimental/AnimNext/Source/AnimNext/Public/DecoratorBase/ExecutionContext.h
+++ b/Engine/Plugins/Experimental/AnimNext/Source/AnimNext/Public/DecoratorBase/ExecutionContext.h
@@ -19,8 +19,10 @@ namespace UE::AnimNext
 {
 	struct FNodeDescription;
 	struct FNodeInstance;
+	struct FNodeTemplateRegistry;
 	struct FNodeTemplate;
 	struct FDecorator;
+	struct FDecoratorRegistry;
 	struct FDecoratorTemplate;
 	struct ITraversalContext;
 
@@ -43,7 +45,7 @@ namespace UE::AnimNext
 		// Queries a node for a decorator that implements the specified interface.
 		// If no such decorator exists, nullptr is returned.
 		template<class DecoratorInterface>
-		bool GetInterface(FWeakDecoratorPtr DecoratorPtr, TDecoratorBinding<DecoratorInterface>& InterfaceBinding) const;
+		bool GetInterface(const FWeakDecoratorPtr& DecoratorPtr, TDecoratorBinding<DecoratorInterface>& InterfaceBinding) const;
 
 		// Queries a node for a decorator that implements the specified interface.
 		// If no such decorator exists, nullptr is returned.
@@ -53,7 +55,7 @@ namespace UE::AnimNext
 		// Queries a node for a decorator lower on the stack that implements the specified interface.
 		// If no such decorator exists, nullptr is returned.
 		template<class DecoratorInterface>
-		bool GetInterfaceSuper(FWeakDecoratorPtr DecoratorPtr, TDecoratorBinding<DecoratorInterface>& SuperBinding) const;
+		bool GetInterfaceSuper(const FWeakDecoratorPtr& DecoratorPtr, TDecoratorBinding<DecoratorInterface>& SuperBinding) const;
 
 		// Queries a node for a decorator lower on the stack that implements the specified interface.
 		// If no such decorator exists, nullptr is returned.
@@ -62,14 +64,14 @@ namespace UE::AnimNext
 
 		// Allocates a new node instance from a decorator handle
 		// If the desired decorator lives in the current parent, a weak handle to it will be returned
-		FDecoratorPtr AllocateNodeInstance(const FDecoratorBinding& ParentBinding, FAnimNextDecoratorHandle ChildDecoratorHandle);
+		FDecoratorPtr AllocateNodeInstance(const FDecoratorBinding& ParentBinding, FAnimNextDecoratorHandle ChildDecoratorHandle) const;
 
 		// Allocates a new node instance from a decorator handle
 		// If the desired decorator lives in the current parent, a weak handle to it will be returned
-		FDecoratorPtr AllocateNodeInstance(FWeakDecoratorPtr ParentBinding, FAnimNextDecoratorHandle ChildDecoratorHandle);
+		FDecoratorPtr AllocateNodeInstance(const FWeakDecoratorPtr& ParentBinding, FAnimNextDecoratorHandle ChildDecoratorHandle) const;
 
 		// Releases a node instance that is no longer referenced
-		void ReleaseNodeInstance(FNodeInstance* Node);
+		void ReleaseNodeInstance(FNodeInstance* Node) const;
 
 		// Returns the current strongly typed traversal context or nullptr if not in a traversal
 		template<class TraversalContextType>
@@ -83,16 +85,21 @@ namespace UE::AnimNext
 		FExecutionContext(const FExecutionContext&) = delete;
 		FExecutionContext(FExecutionContext&&) = delete;
 
-		bool GetInterfaceImpl(FDecoratorInterfaceUID InterfaceUID, FWeakDecoratorPtr DecoratorPtr, FDecoratorBinding& InterfaceBinding) const;
-		bool GetInterfaceSuperImpl(FDecoratorInterfaceUID InterfaceUID, FWeakDecoratorPtr DecoratorPtr, FDecoratorBinding& SuperBinding) const;
+		bool GetInterfaceImpl(FDecoratorInterfaceUID InterfaceUID, const FWeakDecoratorPtr& DecoratorPtr, FDecoratorBinding& InterfaceBinding) const;
+		bool GetInterfaceSuperImpl(FDecoratorInterfaceUID InterfaceUID, const FWeakDecoratorPtr& DecoratorPtr, FDecoratorBinding& SuperBinding) const;
 		const void* EvaluateLatentPinImpl(FLatentPropertyHandle LatentPropertyHandle) const;
 
 		const FNodeDescription& GetNodeDescription(FNodeHandle NodeHandle) const;
 		const FNodeTemplate* GetNodeTemplate(const FNodeDescription& NodeDesc) const;
 		const FDecorator* GetDecorator(const FDecoratorTemplate& DecoratorDesc) const;
 
+		// Cached references to the registries we need
+		const FNodeTemplateRegistry& NodeTemplateRegistry;
+		const FDecoratorRegistry& DecoratorRegistry;
+
 		ITraversalContext* TraversalContext = nullptr;
 
+		// Cached properties for the currently executing graph
 		TArrayView<const uint8> GraphSharedData;
 		FRigVMMemoryHandleArray RigVMLatentMemoryHandles;
 		FRigVMExtendedExecuteContext* RigVMExecuteContext;
@@ -106,7 +113,7 @@ namespace UE::AnimNext
 	//////////////////////////////////////////////////////////////////////////
 
 	template<class DecoratorInterface>
-	inline bool FExecutionContext::GetInterface(FWeakDecoratorPtr DecoratorPtr, TDecoratorBinding<DecoratorInterface>& InterfaceBinding) const
+	inline bool FExecutionContext::GetInterface(const FWeakDecoratorPtr& DecoratorPtr, TDecoratorBinding<DecoratorInterface>& InterfaceBinding) const
 	{
 		constexpr FDecoratorInterfaceUID InterfaceUID = DecoratorInterface::InterfaceUID;
 		return GetInterfaceImpl(InterfaceUID, DecoratorPtr, InterfaceBinding);
@@ -119,7 +126,7 @@ namespace UE::AnimNext
 	}
 
 	template<class DecoratorInterface>
-	inline bool FExecutionContext::GetInterfaceSuper(FWeakDecoratorPtr DecoratorPtr, TDecoratorBinding<DecoratorInterface>& SuperBinding) const
+	inline bool FExecutionContext::GetInterfaceSuper(const FWeakDecoratorPtr& DecoratorPtr, TDecoratorBinding<DecoratorInterface>& SuperBinding) const
 	{
 		constexpr FDecoratorInterfaceUID InterfaceUID = DecoratorInterface::InterfaceUID;
 		return GetInterfaceSuperImpl(InterfaceUID, DecoratorPtr, SuperBinding);
@@ -131,7 +138,7 @@ namespace UE::AnimNext
 		return GetInterfaceSuper<DecoratorInterface>(Binding.GetDecoratorPtr(), SuperBinding);
 	}
 
-	inline FDecoratorPtr FExecutionContext::AllocateNodeInstance(const FDecoratorBinding& ParentBinding, FAnimNextDecoratorHandle ChildDecoratorHandle)
+	inline FDecoratorPtr FExecutionContext::AllocateNodeInstance(const FDecoratorBinding& ParentBinding, FAnimNextDecoratorHandle ChildDecoratorHandle) const
 	{
 		return AllocateNodeInstance(ParentBinding.GetDecoratorPtr(), ChildDecoratorHandle);
 	}
diff --git a/Engine/Plugins/Experimental/AnimNext/Source/AnimNext/Public/DecoratorBase/ExecutionContextProxy.h b/Engine/Plugins/Experimental/AnimNext/Source/AnimNext/Public/DecoratorBase/ExecutionContextProxy.h
new file mode 100644
--- /dev/null
+++ b/Engine/Plugins/Experimental/AnimNext/Source/AnimNext/Public/DecoratorBase/ExecutionContextProxy.h
@@ -0,0 +1,87 @@
+// Copyright Epic Games, Inc. All Rights Reserved.
+
+#pragma once
+
+#include "CoreMinimal.h"
+#include "DecoratorBase/ExecutionContext.h"
+
+namespace UE::AnimNext
+{
+	/**
+	 * FExecutionContextProxy
+	 *
+	 * Wraps an execution context and forwards all calls to it by mirroring its API.
+	 */
+	struct FExecutionContextProxy
+	{
+		explicit FExecutionContextProxy(const FExecutionContext& InExecutionContext)
+			: ExecutionContext(InExecutionContext)
+		{}
+
+		// We safely coerce to our wrapped execution context to allow identical usage
+		operator const FExecutionContext& () const { return ExecutionContext; }
+
+		// Queries a node for a decorator that implements the specified interface.
+		// If no such decorator exists, nullptr is returned.
+		template<class DecoratorInterface>
+		bool GetInterface(const FWeakDecoratorPtr& DecoratorPtr, TDecoratorBinding<DecoratorInterface>& InterfaceBinding) const
+		{
+			return ExecutionContext.GetInterface<DecoratorInterface>(DecoratorPtr, InterfaceBinding);
+		}
+
+		// Queries a node for a decorator that implements the specified interface.
+		// If no such decorator exists, nullptr is returned.
+		template<class DecoratorInterface>
+		bool GetInterface(const FDecoratorBinding& Binding, TDecoratorBinding<DecoratorInterface>& InterfaceBinding) const
+		{
+			return ExecutionContext.GetInterface<DecoratorInterface>(Binding, InterfaceBinding);
+		}
+
+		// Queries a node for a decorator lower on the stack that implements the specified interface.
+		// If no such decorator exists, nullptr is returned.
+		template<class DecoratorInterface>
+		bool GetInterfaceSuper(const FWeakDecoratorPtr& DecoratorPtr, TDecoratorBinding<DecoratorInterface>& SuperBinding) const
+		{
+			return ExecutionContext.GetInterfaceSuper<DecoratorInterface>(DecoratorPtr, SuperBinding);
+		}
+
+		// Queries a node for a decorator lower on the stack that implements the specified interface.
+		// If no such decorator exists, nullptr is returned.
+		template<class DecoratorInterface>
+		bool GetInterfaceSuper(const FDecoratorBinding& Binding, TDecoratorBinding<DecoratorInterface>& SuperBinding) const
+		{
+			return ExecutionContext.GetInterfaceSuper<DecoratorInterface>(Binding, SuperBinding);
+		}
+
+		// Allocates a new node instance from a decorator handle
+		// If the desired decorator lives in the current parent, a weak handle to it will be returned
+		FDecoratorPtr AllocateNodeInstance(const FDecoratorBinding& ParentBinding, FAnimNextDecoratorHandle ChildDecoratorHandle) const
+		{
+			return ExecutionContext.AllocateNodeInstance(ParentBinding, ChildDecoratorHandle);
+		}
+
+		// Allocates a new node instance from a decorator handle
+		// If the desired decorator lives in the current parent, a weak handle to it will be returned
+		FDecoratorPtr AllocateNodeInstance(const FWeakDecoratorPtr& ParentBinding, FAnimNextDecoratorHandle ChildDecoratorHandle) const
+		{
+			return ExecutionContext.AllocateNodeInstance(ParentBinding, ChildDecoratorHandle);
+		}
+
+		// Releases a node instance that is no longer referenced
+		void ReleaseNodeInstance(FNodeInstance* Node) const
+		{
+			return ExecutionContext.ReleaseNodeInstance(Node);
+		}
+
+		// Evaluates the latent pin with the specified handle
+		template<typename LatentPinType>
+		LatentPinType EvaluateLatentPin(FLatentPropertyHandle LatentPropertyHandle) const
+		{
+			return ExecutionContext.EvaluateLatentPin<LatentPinType>(LatentPropertyHandle);
+		}
+
+	private:
+		// The execution context that we wrap
+		const FExecutionContext& ExecutionContext;
+	};
+}
diff --git a/Engine/Plugins/Experimental/AnimNext/Source/AnimNext/Public/DecoratorInterfaces/IContinuousBlend.h b/Engine/Plugins/Experimental/AnimNext/Source/AnimNext/Public/DecoratorInterfaces/IContinuousBlend.h
--- a/Engine/Plugins/Experimental/AnimNext/Source/AnimNext/Public/DecoratorInterfaces/IContinuousBlend.h
+++ b/Engine/Plugins/Experimental/AnimNext/Source/AnimNext/Public/DecoratorInterfaces/IContinuousBlend.h
@@ -19,7 +19,7 @@ namespace UE::AnimNext
 		// Returns the blend weight for the specified child
 		// Multiple children can have non-zero weight but their sum must be 1.0
 		// Returns -1.0 if the child index is invalid
-		virtual float GetBlendWeight(FExecutionContext& Context, const TDecoratorBinding<IContinuousBlend>& Binding, int32 ChildIndex) const;
+		virtual float GetBlendWeight(const FExecutionContext& Context, const TDecoratorBinding<IContinuousBlend>& Binding, int32 ChildIndex) const;
 	};
 
 	/**
@@ -29,7 +29,7 @@ namespace UE::AnimNext
 	struct TDecoratorBinding<IContinuousBlend> : FDecoratorBinding
 	{
 		// @see IContinuousBlend::GetBlendWeight
-		float GetBlendWeight(FExecutionContext& Context, int32 ChildIndex) const
+		float GetBlendWeight(const FExecutionContext& Context, int32 ChildIndex) const
 		{
 			return GetInterface()->GetBlendWeight(Context, *this, ChildIndex);
 		}
diff --git a/Engine/Plugins/Experimental/AnimNext/Source/AnimNext/Public/DecoratorInterfaces/IDiscreteBlend.h b/Engine/Plugins/Experimental/AnimNext/Source/AnimNext/Public/DecoratorInterfaces/IDiscreteBlend.h
--- a/Engine/Plugins/Experimental/AnimNext/Source/AnimNext/Public/DecoratorInterfaces/IDiscreteBlend.h
+++ b/Engine/Plugins/Experimental/AnimNext/Source/AnimNext/Public/DecoratorInterfaces/IDiscreteBlend.h
@@ -21,25 +21,25 @@ namespace UE::AnimNext
 		// Returns the blend weight for the specified child
 		// Multiple children can have non-zero weight but their sum must be 1.0
 		// Returns -1.0 if the child index is invalid
-		virtual float GetBlendWeight(FExecutionContext& Context, const TDecoratorBinding<IDiscreteBlend>& Binding, int32 ChildIndex) const;
+		virtual float GetBlendWeight(const FExecutionContext& Context, const TDecoratorBinding<IDiscreteBlend>& Binding, int32 ChildIndex) const;
 
 		// Returns the blend alpha for the specified child
 		// Returns nullptr if the child index is invalid
 		// Allows additive decorators to query the internal state of the base decorator
-		virtual const FAlphaBlend* GetBlendState(FExecutionContext& Context, const TDecoratorBinding<IDiscreteBlend>& Binding, int32 ChildIndex) const;
+		virtual const FAlphaBlend* GetBlendState(const FExecutionContext& Context, const TDecoratorBinding<IDiscreteBlend>& Binding, int32 ChildIndex) const;
 
 		// Returns the blend destination child index (aka the active child index)
 		// Returns INDEX_NONE if no child is active
-		virtual int32 GetBlendDestinationChildIndex(FExecutionContext& Context, const TDecoratorBinding<IDiscreteBlend>& Binding) const;
+		virtual int32 GetBlendDestinationChildIndex(const FExecutionContext& Context, const TDecoratorBinding<IDiscreteBlend>& Binding) const;
 
 		// Called when a blend transition between children occurs
-		virtual void OnBlendTransition(FExecutionContext& Context, const TDecoratorBinding<IDiscreteBlend>& Binding, int32 OldChildIndex, int32 NewChildIndex) const;
+		virtual void OnBlendTransition(const FExecutionContext& Context, const TDecoratorBinding<IDiscreteBlend>& Binding, int32 OldChildIndex, int32 NewChildIndex) const;
 
 		// Called when the blend for specified child is initiated
-		virtual void OnBlendInitiated(FExecutionContext& Context, const TDecoratorBinding<IDiscreteBlend>& Binding, int32 ChildIndex) const;
+		virtual void OnBlendInitiated(const FExecutionContext& Context, const TDecoratorBinding<IDiscreteBlend>& Binding, int32 ChildIndex) const;
 
 		// Called when the blend for specified child terminates
-		virtual void OnBlendTerminated(FExecutionContext& Context, const TDecoratorBinding<IDiscreteBlend>& Binding, int32 ChildIndex) const;
+		virtual void OnBlendTerminated(const FExecutionContext& Context, const TDecoratorBinding<IDiscreteBlend>& Binding, int32 ChildIndex) const;
 	};
 
 	/**
@@ -49,37 +49,37 @@ namespace UE::AnimNext
 	struct TDecoratorBinding<IDiscreteBlend> : FDecoratorBinding
 	{
 		// @see IDiscreteBlend::GetBlendWeight
-		float GetBlendWeight(FExecutionContext& Context, int32 ChildIndex) const
+		float GetBlendWeight(const FExecutionContext& Context, int32 ChildIndex) const
 		{
 			return GetInterface()->GetBlendWeight(Context, *this, ChildIndex);
 		}
 
 		// @see IDiscreteBlend::GetBlendState
-		const FAlphaBlend* GetBlendState(FExecutionContext& Context, int32 ChildIndex) const
+		const FAlphaBlend* GetBlendState(const FExecutionContext& Context, int32 ChildIndex) const
 		{
 			return GetInterface()->GetBlendState(Context, *this, ChildIndex);
 		}
 
 		// @see IDiscreteBlend::GetBlendDestinationChildIndex
-		int32 GetBlendDestinationChildIndex(FExecutionContext& Context) const
+		int32 GetBlendDestinationChildIndex(const FExecutionContext& Context) const
 		{
 			return GetInterface()->GetBlendDestinationChildIndex(Context, *this);
 		}
 
 		// @see IDiscreteBlend::OnBlendTransition
-		void OnBlendTransition(FExecutionContext& Context, int32 OldChildIndex, int32 NewChildIndex) const
+		void OnBlendTransition(const FExecutionContext& Context, int32 OldChildIndex, int32 NewChildIndex) const
 		{
 			GetInterface()->OnBlendTransition(Context, *this, OldChildIndex, NewChildIndex);
 		}
 
 		// @see IDiscreteBlend::OnBlendInitiated
-		void OnBlendInitiated(FExecutionContext& Context, int32 ChildIndex) const
+		void OnBlendInitiated(const FExecutionContext& Context, int32 ChildIndex) const
 		{
 			GetInterface()->OnBlendInitiated(Context, *this, ChildIndex);
 		}
 
 		// @see IDiscreteBlend::OnBlendTerminated
-		void OnBlendTerminated(FExecutionContext& Context, int32 ChildIndex) const
+		void OnBlendTerminated(const FExecutionContext& Context, int32 ChildIndex) const
 		{
 			GetInterface()->OnBlendTerminated(Context, *this, ChildIndex);
 		}
diff --git a/Engine/Plugins/Experimental/AnimNext/Source/AnimNext/Public/DecoratorInterfaces/IEvaluate.h b/Engine/Plugins/Experimental/AnimNext/Source/AnimNext/Public/DecoratorInterfaces/IEvaluate.h
--- a/Engine/Plugins/Experimental/AnimNext/Source/AnimNext/Public/DecoratorInterfaces/IEvaluate.h
+++ b/Engine/Plugins/Experimental/AnimNext/Source/AnimNext/Public/DecoratorInterfaces/IEvaluate.h
@@ -51,10 +51,10 @@ namespace UE::AnimNext
 		DECLARE_ANIM_DECORATOR_INTERFACE(IEvaluate, 0xa303e9e7)
 
 		// Called before a decorator's children are evaluated
-		virtual void PreEvaluate(FExecutionContext& Context, const TDecoratorBinding<IEvaluate>& Binding) const;
+		virtual void PreEvaluate(const FExecutionContext& Context, const TDecoratorBinding<IEvaluate>& Binding) const;
 
 		// Called after a decorator's children have been evaluated
-		virtual void PostEvaluate(FExecutionContext& Context, const TDecoratorBinding<IEvaluate>& Binding) const;
+		virtual void PostEvaluate(const FExecutionContext& Context, const TDecoratorBinding<IEvaluate>& Binding) const;
 	};
 
 	/**
@@ -64,13 +64,13 @@ namespace UE::AnimNext
 	struct TDecoratorBinding<IEvaluate> : FDecoratorBinding
 	{
 		// @see IEvaluate::PreEvaluate
-		void PreEvaluate(FExecutionContext& Context) const
+		void PreEvaluate(const FExecutionContext& Context) const
 		{
 			GetInterface()->PreEvaluate(Context, *this);
 		}
 
 		// @see IEvaluate::PostEvaluate
-		void PostEvaluate(FExecutionContext& Context) const
+		void PostEvaluate(const FExecutionContext& Context) const
 		{
 			GetInterface()->PostEvaluate(Context, *this);
 		}
diff --git a/Engine/Plugins/Experimental/AnimNext/Source/AnimNext/Public/DecoratorInterfaces/IHierarchy.h b/Engine/Plugins/Experimental/AnimNext/Source/AnimNext/Public/DecoratorInterfaces/IHierarchy.h
--- a/Engine/Plugins/Experimental/AnimNext/Source/AnimNext/Public/DecoratorInterfaces/IHierarchy.h
+++ b/Engine/Plugins/Experimental/AnimNext/Source/AnimNext/Public/DecoratorInterfaces/IHierarchy.h
@@ -8,7 +8,9 @@
 
 namespace UE::AnimNext
 {
-	using FChildrenArray = TArray<FWeakDecoratorPtr, TMemStackAllocator<>>;
+	// An array of children pointers
+	// We reserve a small amount inline and spill on the memstack
+	using FChildrenArray = TArray<FWeakDecoratorPtr, TInlineAllocator<8, TMemStackAllocator<>>>;
 
 	/**
 	 * IHierarchy
@@ -21,12 +23,12 @@ namespace UE::AnimNext
 
 		// Returns the number of children
 		// Includes inactive children
-		virtual uint32 GetNumChildren(FExecutionContext& Context, const TDecoratorBinding<IHierarchy>& Binding) const;
+		virtual uint32 GetNumChildren(const FExecutionContext& Context, const TDecoratorBinding<IHierarchy>& Binding) const;
 
 		// Appends weak handles to any children we wish to traverse.
 		// Decorators are responsible for allocating and releasing child instance data.
 		// Empty handles and duplicates can be appended.
-		virtual void GetChildren(FExecutionContext& Context, const TDecoratorBinding<IHierarchy>& Binding, FChildrenArray& Children) const;
+		virtual void GetChildren(const FExecutionContext& Context, const TDecoratorBinding<IHierarchy>& Binding, FChildrenArray& Children) const;
 	};
 
 	/**
@@ -36,13 +38,13 @@ namespace UE::AnimNext
 	struct TDecoratorBinding<IHierarchy> : FDecoratorBinding
 	{
 		// @see IHierarchy::GetNumChildren
-		uint32 GetNumChildren(FExecutionContext& Context) const
+		uint32 GetNumChildren(const FExecutionContext& Context) const
 		{
 			return GetInterface()->GetNumChildren(Context, *this);
 		}
 
 		// @see IHierarchy::GetChildren
-		void GetChildren(FExecutionContext& Context, FChildrenArray& Children) const
+		void GetChildren(const FExecutionContext& Context, FChildrenArray& Children) const
 		{
 			GetInterface()->GetChildren(Context, *this, Children);
 		}
diff --git a/Engine/Plugins/Experimental/AnimNext/Source/AnimNext/Public/DecoratorInterfaces/ISmoothBlend.h b/Engine/Plugins/Experimental/AnimNext/Source/AnimNext/Public/DecoratorInterfaces/ISmoothBlend.h
--- a/Engine/Plugins/Experimental/AnimNext/Source/AnimNext/Public/DecoratorInterfaces/ISmoothBlend.h
+++ b/Engine/Plugins/Experimental/AnimNext/Source/AnimNext/Public/DecoratorInterfaces/ISmoothBlend.h
@@ -17,7 +17,7 @@ namespace UE::AnimNext
 		DECLARE_ANIM_DECORATOR_INTERFACE(ISmoothBlend, 0x1c2c1739)
 
 		// Returns the desired blend time for the specified child
-		virtual float GetBlendTime(FExecutionContext& Context, const TDecoratorBinding<ISmoothBlend>& Binding, int32 ChildIndex) const;
+		virtual float GetBlendTime(const FExecutionContext& Context, const TDecoratorBinding<ISmoothBlend>& Binding, int32 ChildIndex) const;
 	};
 
 	/**
@@ -27,7 +27,7 @@ namespace UE::AnimNext
 	struct TDecoratorBinding<ISmoothBlend> : FDecoratorBinding
 	{
 		// @see ISmoothBlend::GetBlendTime
-		float GetBlendTime(FExecutionContext& Context, int32 ChildIndex) const
+		float GetBlendTime(const FExecutionContext& Context, int32 ChildIndex) const
 		{
 			return GetInterface()->GetBlendTime(Context, *this, ChildIndex);
 		}
diff --git a/Engine/Plugins/Experimental/AnimNext/Source/AnimNext/Public/DecoratorInterfaces/ITimeline.h b/Engine/Plugins/Experimental/AnimNext/Source/AnimNext/Public/DecoratorInterfaces/ITimeline.h
--- a/Engine/Plugins/Experimental/AnimNext/Source/AnimNext/Public/DecoratorInterfaces/ITimeline.h
+++ b/Engine/Plugins/Experimental/AnimNext/Source/AnimNext/Public/DecoratorInterfaces/ITimeline.h
@@ -17,7 +17,7 @@ namespace UE::AnimNext
 		DECLARE_ANIM_DECORATOR_INTERFACE(ITimeline, 0x53760727)
 
 		// Returns the play rate of this timeline
-		virtual double GetPlayRate(FExecutionContext& Context, const TDecoratorBinding<ITimeline>& Binding) const;
+		virtual double GetPlayRate(const FExecutionContext& Context, const TDecoratorBinding<ITimeline>& Binding) const;
 	};
 
 	/**
@@ -27,7 +27,7 @@ namespace UE::AnimNext
 	struct TDecoratorBinding<ITimeline> : FDecoratorBinding
 	{
 		// @see ITimeline::GetPlayRate
-		double GetPlayRate(FExecutionContext& Context) const
+		double GetPlayRate(const FExecutionContext& Context) const
 		{
 			return GetInterface()->GetPlayRate(Context, *this);
 		}
diff --git a/Engine/Plugins/Experimental/AnimNext/Source/AnimNext/Public/DecoratorInterfaces/IUpdate.h b/Engine/Plugins/Experimental/AnimNext/Source/AnimNext/Public/DecoratorInterfaces/IUpdate.h
--- a/Engine/Plugins/Experimental/AnimNext/Source/AnimNext/Public/DecoratorInterfaces/IUpdate.h
+++ b/Engine/Plugins/Experimental/AnimNext/Source/AnimNext/Public/DecoratorInterfaces/IUpdate.h
@@ -3,24 +3,159 @@
 #pragma once
 
 #include "CoreMinimal.h"
+#include "DecoratorBase/ExecutionContextProxy.h"
 #include "DecoratorBase/IDecoratorInterface.h"
-#include "DecoratorBase/ITraversalContext.h"
+
+class FMemStack;
 
 namespace UE::AnimNext
 {
+	struct FUpdateEntry;
+	struct FUpdateTraversalQueue;
+
+	/**
+	  * FDecoratorUpdateState
+	  * 
+	  * State that propagates down the graph during the update traversal.
+	  * 
+	  * We align this to 16 bytes and maintain that same size to ensure efficient copying.
+	  */
+	struct alignas(16) FDecoratorUpdateState final
+	{
+		FDecoratorUpdateState() = default;
+
+		explicit FDecoratorUpdateState(float InDeltaTime)
+			: DeltaTime(InDeltaTime)
+		{}
+
+		// Returns the delta time for this decorator
+		float GetDeltaTime() const { return DeltaTime; }
+
+		// Returns the total weight for this decorator
+		float GetTotalWeight() const { return TotalWeight; }
+
+		// Returns the total trajectory weight for this decorator
+		float GetTotalTrajectoryWeight() const { return TotalTrajectoryWeight; }
+
+		// Returns whether or not this decorator is blending out
+		bool IsBlendingOut() const { return !!bIsBlendingOut; }
+
+		// Creates a new instance of the update state with the total weight scaled by the supplied weight
+		FDecoratorUpdateState WithWeight(float Weight) const
+		{
+			FDecoratorUpdateState Result(*this);
+			Result.TotalWeight *= Weight;
+			return Result;
+		}
+
+		// Creates a new instance of the update state with the total trajectory weight scaled by the supplied weight
+		FDecoratorUpdateState WithTrajectoryWeight(float Weight) const
+		{
+			FDecoratorUpdateState Result(*this);
+			Result.TotalTrajectoryWeight *= Weight;
+			return Result;
+		}
+
+		// Creates a new instance of the update state with the delta time scaled by the supplied scale factor
+		FDecoratorUpdateState WithTimeScale(float TimeScale) const
+		{
+			FDecoratorUpdateState Result(*this);
+			Result.DeltaTime *= TimeScale;
+			return Result;
+		}
+
+		// Creates a new instance of the update state marking it as blending out
+		FDecoratorUpdateState AsBlendingOut() const
+		{
+			FDecoratorUpdateState Result(*this);
+			Result.bIsBlendingOut = true;
+			return Result;
+		}
+
+	private:
+		// The amount of time to move forward by at allowing parents to time scale their children
+		float DeltaTime = 0.0f;
+
+		// The current total weight factoring in all inherited blend weights
+		float TotalWeight = 1.0f;
+
+		// The current total trajectory weight factoring in all inherited blend weights
+		float TotalTrajectoryWeight = 1.0f;
+
+		// Whether or not we are blending out
+		int32 bIsBlendingOut = false;
+	};
+
 	/**
 	 * FUpdateTraversalContext
-	 * 
-	 * Contains all relevant transient data for an update traversal.
+	 *
+	 * Contains all relevant transient data for an update traversal and wraps the execution context.
 	 */
-	struct ANIMNEXT_API FUpdateTraversalContext : ITraversalContext
+	struct FUpdateTraversalContext final : FExecutionContextProxy
 	{
-		explicit FUpdateTraversalContext(float InDeltaTime);
+	private:
+		// Constructs a new traversal context and wraps the specified execution context
+		explicit FUpdateTraversalContext(const FExecutionContext& InExecutionContext, FMemStack& InMemStack);
 
-		float GetDeltaTime() const { return DeltaTime; }
+		// Pops entries from the traversal queue and pushes them onto the update stack
+		void PushQueuedUpdateEntries(FUpdateTraversalQueue& TraversalQueue);
 
-	protected:
-		float DeltaTime;
+		// Pushes an entry onto the update stack
+		void PushUpdateEntry(FUpdateEntry* Entry);
+
+		// Pops an entry from the top of the update stack
+		FUpdateEntry* PopUpdateEntry();
+
+		// Pushes an entry onto the free entry stack
+		void PushFreeEntry(FUpdateEntry* Entry);
+
+		// Returns a new entry suitable for update queuing
+		// If an entry isn't found in the free stack, a new one is allocated from the memstack
+		FUpdateEntry* GetNewEntry(const FWeakDecoratorPtr& DecoratorPtr, const FDecoratorUpdateState& DecoratorState);
+
+		// The memstack of the local thread
+		FMemStack& MemStack;
+
+		// The head pointer of the update stack
+		// This is the traversal execution stack and it contains entries that are
+		// pending their pre-update call and entries waiting for post-update to be called
+		FUpdateEntry* UpdateStackHead = nullptr;
+
+		// The head pointer of the free entry stack
+		// Entries are allocated from the memstack and are re-used in LIFO since they'll
+		// be warmer in the CPU cache
+		FUpdateEntry* FreeEntryStackHead = nullptr;
+
+		friend ANIMNEXT_API void UpdateGraph(const FExecutionContext& Context, const FWeakDecoratorPtr& GraphRootPtr, float DeltaTime);
+		friend FUpdateTraversalQueue;
+	};
+
+	/**
+	 * FUpdateTraversalQueue
+	 *
+	 * A queue of children to traverse.
+	 * @see IUpdate::QueueChildrenForTraversal
+	 */
+	struct FUpdateTraversalQueue final
+	{
+		// Queued a child for traversal
+		// Children are processed in the same order they are queued
+		ANIMNEXT_API void Push(const FWeakDecoratorPtr& ChildPtr, const FDecoratorUpdateState& ChildDecoratorState);
+
+	private:
+		explicit FUpdateTraversalQueue(FUpdateTraversalContext& InTraversalContext);
+
+		// The current traversal context
+		FUpdateTraversalContext& TraversalContext;
+
+		// The head pointer of the queued update stack
+		// When a child is queued for traversal, it is first appended to this stack
+		// Once all children have been queued and pre-update terminates, this stack is
+		// emptied and pushed onto the update stack
+		FUpdateEntry* QueuedUpdateStackHead = nullptr;
+
+		friend ANIMNEXT_API void UpdateGraph(const FExecutionContext& Context, const FWeakDecoratorPtr& GraphRootPtr, float DeltaTime);
+		friend FUpdateTraversalContext;
 	};
 
 	/**
@@ -30,19 +165,26 @@ namespace UE::AnimNext
 	 *
 	 * When a node is visited, PreUpdate is first called on its top decorator. It is responsible for forwarding
 	 * the call to the next decorator that implements this interface on the decorator stack of the node. Once
-	 * all decorators have had the chance to PreUpdate, the children of the decorator are queried through
-	 * the IHierarchy interface. The children will then evaluate and PostUpdate will then be called afterwards
-	 * on the original decorator.
+	 * all decorators have had the chance to PreUpdate, the children will then evaluate and PostUpdate will
+	 * then be called afterwards on the original decorator.
+	 * 
+	 * If this interface is implemented on a node, it is responsible for queuing the children for traversal
+	 * using the QueueChildrenForTraversal function. If the interface isn't implemented, children
+	 * are discovered through the IHierarchy interface and traversal propagates to them using the current
+	 * decorator state.
 	 */
 	struct ANIMNEXT_API IUpdate : IDecoratorInterface
 	{
 		DECLARE_ANIM_DECORATOR_INTERFACE(IUpdate, 0x59d24dc5)
 
 		// Called before a decorator's children are updated
-		virtual void PreUpdate(FExecutionContext& Context, const TDecoratorBinding<IUpdate>& Binding) const;
+		virtual void PreUpdate(FUpdateTraversalContext& Context, const TDecoratorBinding<IUpdate>& Binding, const FDecoratorUpdateState& DecoratorState) const;
+
+		// Called after PreUpdate to request that children be queued with the provided context
+		virtual void QueueChildrenForTraversal(FUpdateTraversalContext& Context, const TDecoratorBinding<IUpdate>& Binding, const FDecoratorUpdateState& DecoratorState, FUpdateTraversalQueue& TraversalQueue) const;
 
 		// Called after a decorator's children have been updated
-		virtual void PostUpdate(FExecutionContext& Context, const TDecoratorBinding<IUpdate>& Binding) const;
+		virtual void PostUpdate(FUpdateTraversalContext& Context, const TDecoratorBinding<IUpdate>& Binding, const FDecoratorUpdateState& DecoratorState) const;
 	};
 
 	/**
@@ -52,15 +194,21 @@ namespace UE::AnimNext
 	struct TDecoratorBinding<IUpdate> : FDecoratorBinding
 	{
 		// @see IUpdate::PreUpdate
-		void PreUpdate(FExecutionContext& Context) const
+		void PreUpdate(FUpdateTraversalContext& Context, const FDecoratorUpdateState& DecoratorState) const
+		{
+			GetInterface()->PreUpdate(Context, *this, DecoratorState);
+		}
+
+		// @see IUpdate::QueueChildrenForTraversal
+		void QueueChildrenForTraversal(FUpdateTraversalContext& Context, const FDecoratorUpdateState& DecoratorState, FUpdateTraversalQueue& TraversalQueue) const
 		{
-			GetInterface()->PreUpdate(Context, *this);
+			GetInterface()->QueueChildrenForTraversal(Context, *this, DecoratorState, TraversalQueue);
 		}
 
 		// @see IUpdate::PostUpdate
-		void PostUpdate(FExecutionContext& Context) const
+		void PostUpdate(FUpdateTraversalContext& Context, const FDecoratorUpdateState& DecoratorState) const
 		{
-			GetInterface()->PostUpdate(Context, *this);
+			GetInterface()->PostUpdate(Context, *this, DecoratorState);
 		}
 
 	protected:
@@ -73,11 +221,10 @@ namespace UE::AnimNext
 	 *
 	 * For each node:
 	 *     - We call PreUpdate on all its decorators
-	 *     - We call GetChildren on all its decorators
-	 *     - We update all children found
+	 *     - We update all children
 	 *     - We call PostUpdate on all its decorators
 	 *
 	 * @see IUpdate::PreUpdate, IUpdate::PostUpdate, IHierarchy::GetChildren
 	 */
-	ANIMNEXT_API void UpdateGraph(FExecutionContext& Context, FUpdateTraversalContext& TraversalContext, FWeakDecoratorPtr GraphRootPtr);
+	ANIMNEXT_API void UpdateGraph(const FExecutionContext& Context, const FWeakDecoratorPtr& GraphRootPtr, float DeltaTime);
 }
diff --git a/Engine/Plugins/Experimental/AnimNext/Source/AnimNextTestSuite/Private/AnimNextDecoratorGraphTest.cpp b/Engine/Plugins/Experimental/AnimNext/Source/AnimNextTestSuite/Private/AnimNextDecoratorGraphTest.cpp
--- a/Engine/Plugins/Experimental/AnimNext/Source/AnimNextTestSuite/Private/AnimNextDecoratorGraphTest.cpp
+++ b/Engine/Plugins/Experimental/AnimNext/Source/AnimNextTestSuite/Private/AnimNextDecoratorGraphTest.cpp
@@ -37,9 +37,9 @@ namespace UE::AnimNext
 		using FSharedData = FTestDecoratorSharedData;
 
 		// IUpdate impl
-		virtual void PostUpdate(FExecutionContext& Context, const TDecoratorBinding<IUpdate>& Binding) const override
+		virtual void PostUpdate(FUpdateTraversalContext& Context, const TDecoratorBinding<IUpdate>& Binding, const FDecoratorUpdateState& DecoratorState) const override
 		{
-			IUpdate::PostUpdate(Context, Binding);
+			IUpdate::PostUpdate(Context, Binding, DecoratorState);
 
 			const FSharedData* SharedData = Binding.GetSharedData<FSharedData>();
 			UE::AnimNext::FParamStack& ParamStack = UE::AnimNext::FParamStack::Get();
@@ -56,7 +56,7 @@ namespace UE::AnimNext
 		}
 
 		// IEvaluate impl
-		virtual void PostEvaluate(FExecutionContext& Context, const TDecoratorBinding<IEvaluate>& Binding) const override
+		virtual void PostEvaluate(const FExecutionContext& Context, const TDecoratorBinding<IEvaluate>& Binding) const override
 		{
 			IEvaluate::PostEvaluate(Context, Binding);
 
diff --git a/Engine/Plugins/Experimental/AnimNext/Source/AnimNextTestSuite/Private/AnimNextDecoratorInterfacesTest.cpp b/Engine/Plugins/Experimental/AnimNext/Source/AnimNextTestSuite/Private/AnimNextDecoratorInterfacesTest.cpp
--- a/Engine/Plugins/Experimental/AnimNext/Source/AnimNextTestSuite/Private/AnimNextDecoratorInterfacesTest.cpp
+++ b/Engine/Plugins/Experimental/AnimNext/Source/AnimNextTestSuite/Private/AnimNextDecoratorInterfacesTest.cpp
@@ -34,39 +34,33 @@ namespace UE::AnimNext
 		static TArray<FDecoratorUID>* EvaluatedDecorators = nullptr;
 	}
 
-	struct FDecoratorWithNoChildren : FBaseDecorator, IHierarchy, IUpdate, IEvaluate
+	struct FDecoratorWithNoChildren : FBaseDecorator, IUpdate, IEvaluate
 	{
 		DECLARE_ANIM_DECORATOR(FDecoratorWithNoChildren, 0xe2400d2a, FBaseDecorator)
 
-		// IHierarchy impl
-		virtual void GetChildren(FExecutionContext& Context, const TDecoratorBinding<IHierarchy>& Binding, FChildrenArray& Children) const override
-		{
-			IHierarchy::GetChildren(Context, Binding, Children);
-		}
-
 		// IUpdate impl
-		virtual void PreUpdate(FExecutionContext& Context, const TDecoratorBinding<IUpdate>& Binding) const override
+		virtual void PreUpdate(FUpdateTraversalContext& Context, const TDecoratorBinding<IUpdate>& Binding, const FDecoratorUpdateState& DecoratorState) const override
 		{
 			if (Private::UpdatedDecorators != nullptr)
 			{
 				Private::UpdatedDecorators->Add(FDecoratorWithNoChildren::DecoratorUID);
 			}
 
-			IUpdate::PreUpdate(Context, Binding);
+			IUpdate::PreUpdate(Context, Binding, DecoratorState);
 		}
 
-		virtual void PostUpdate(FExecutionContext& Context, const TDecoratorBinding<IUpdate>& Binding) const override
+		virtual void PostUpdate(FUpdateTraversalContext& Context, const TDecoratorBinding<IUpdate>& Binding, const FDecoratorUpdateState& DecoratorState) const override
 		{
 			if (Private::UpdatedDecorators != nullptr)
 			{
 				Private::UpdatedDecorators->Add(FDecoratorWithNoChildren::DecoratorUID);
 			}
 
-			IUpdate::PostUpdate(Context, Binding);
+			IUpdate::PostUpdate(Context, Binding, DecoratorState);
 		}
 
 		// IEvaluate impl
-		virtual void PreEvaluate(FExecutionContext& Context, const TDecoratorBinding<IEvaluate>& Binding) const override
+		virtual void PreEvaluate(const FExecutionContext& Context, const TDecoratorBinding<IEvaluate>& Binding) const override
 		{
 			if (Private::EvaluatedDecorators != nullptr)
 			{
@@ -76,7 +70,7 @@ namespace UE::AnimNext
 			IEvaluate::PreEvaluate(Context, Binding);
 		}
 
-		virtual void PostEvaluate(FExecutionContext& Context, const TDecoratorBinding<IEvaluate>& Binding) const override
+		virtual void PostEvaluate(const FExecutionContext& Context, const TDecoratorBinding<IEvaluate>& Binding) const override
 		{
 			if (Private::EvaluatedDecorators != nullptr)
 			{
@@ -88,7 +82,6 @@ namespace UE::AnimNext
 	};
 
 	DEFINE_ANIM_DECORATOR_BEGIN(FDecoratorWithNoChildren)
-		DEFINE_ANIM_DECORATOR_IMPLEMENTS_INTERFACE(IHierarchy)
 		DEFINE_ANIM_DECORATOR_IMPLEMENTS_INTERFACE(IUpdate)
 		DEFINE_ANIM_DECORATOR_IMPLEMENTS_INTERFACE(IEvaluate)
 	DEFINE_ANIM_DECORATOR_END(FDecoratorWithNoChildren)
@@ -104,14 +97,19 @@ namespace UE::AnimNext
 		{
 			FDecoratorPtr Child;
 
-			void Construct(FExecutionContext& Context, const FDecoratorBinding& Binding)
+			void Construct(const FExecutionContext& Context, const FDecoratorBinding& Binding)
 			{
 				Child = Context.AllocateNodeInstance(Binding.GetDecoratorPtr(), Binding.GetSharedData<FSharedData>()->Child);
 			}
 		};
 
 		// IHierarchy impl
-		virtual void GetChildren(FExecutionContext& Context, const TDecoratorBinding<IHierarchy>& Binding, FChildrenArray& Children) const override
+		virtual uint32 GetNumChildren(const FExecutionContext& Context, const TDecoratorBinding<IHierarchy>& Binding) const override
+		{
+			return 1;
+		}
+
+		virtual void GetChildren(const FExecutionContext& Context, const TDecoratorBinding<IHierarchy>& Binding, FChildrenArray& Children) const override
 		{
 			const FInstanceData* InstanceData = Binding.GetInstanceData<FInstanceData>();
 
@@ -135,7 +133,7 @@ namespace UE::AnimNext
 		{
 			FDecoratorPtr Children[2];
 
-			void Construct(FExecutionContext& Context, const FDecoratorBinding& Binding)
+			void Construct(const FExecutionContext& Context, const FDecoratorBinding& Binding)
 			{
 				Children[0] = Context.AllocateNodeInstance(Binding.GetDecoratorPtr(), Binding.GetSharedData<FSharedData>()->Children[0]);
 				Children[1] = Context.AllocateNodeInstance(Binding.GetDecoratorPtr(), Binding.GetSharedData<FSharedData>()->Children[1]);
@@ -143,7 +141,12 @@ namespace UE::AnimNext
 		};
 
 		// IHierarchy impl
-		virtual void GetChildren(FExecutionContext& Context, const TDecoratorBinding<IHierarchy>& Binding, FChildrenArray& Children) const override
+		virtual uint32 GetNumChildren(const FExecutionContext& Context, const TDecoratorBinding<IHierarchy>& Binding) const override
+		{
+			return 2;
+		}
+
+		virtual void GetChildren(const FExecutionContext& Context, const TDecoratorBinding<IHierarchy>& Binding, FChildrenArray& Children) const override
 		{
 			const FInstanceData* InstanceData = Binding.GetInstanceData<FInstanceData>();
 
@@ -154,28 +157,36 @@ namespace UE::AnimNext
 		}
 
 		// IUpdate impl
-		virtual void PreUpdate(FExecutionContext& Context, const TDecoratorBinding<IUpdate>& Binding) const override
+		virtual void PreUpdate(FUpdateTraversalContext& Context, const TDecoratorBinding<IUpdate>& Binding, const FDecoratorUpdateState& DecoratorState) const override
 		{
 			if (Private::UpdatedDecorators != nullptr)
 			{
 				Private::UpdatedDecorators->Add(FDecoratorWithChildren::DecoratorUID);
 			}
 
-			IUpdate::PreUpdate(Context, Binding);
+			IUpdate::PreUpdate(Context, Binding, DecoratorState);
 		}
 
-		virtual void PostUpdate(FExecutionContext& Context, const TDecoratorBinding<IUpdate>& Binding) const override
+		virtual void QueueChildrenForTraversal(FUpdateTraversalContext& Context, const TDecoratorBinding<IUpdate>& Binding, const FDecoratorUpdateState& DecoratorState, FUpdateTraversalQueue& TraversalQueue) const override
+		{
+			const FInstanceData* InstanceData = Binding.GetInstanceData<FInstanceData>();
+
+			TraversalQueue.Push(InstanceData->Children[0], DecoratorState);
+			TraversalQueue.Push(InstanceData->Children[1], DecoratorState);
+		}
+
+		virtual void PostUpdate(FUpdateTraversalContext& Context, const TDecoratorBinding<IUpdate>& Binding, const FDecoratorUpdateState& DecoratorState) const override
 		{
 			if (Private::UpdatedDecorators != nullptr)
 			{
 				Private::UpdatedDecorators->Add(FDecoratorWithChildren::DecoratorUID);
 			}
 
-			IUpdate::PostUpdate(Context, Binding);
+			IUpdate::PostUpdate(Context, Binding, DecoratorState);
 		}
 
 		// IEvaluate impl
-		virtual void PreEvaluate(FExecutionContext& Context, const TDecoratorBinding<IEvaluate>& Binding) const override
+		virtual void PreEvaluate(const FExecutionContext& Context, const TDecoratorBinding<IEvaluate>& Binding) const override
 		{
 			if (Private::EvaluatedDecorators != nullptr)
 			{
@@ -185,7 +196,7 @@ namespace UE::AnimNext
 			IEvaluate::PreEvaluate(Context, Binding);
 		}
 
-		virtual void PostEvaluate(FExecutionContext& Context, const TDecoratorBinding<IEvaluate>& Binding) const override
+		virtual void PostEvaluate(const FExecutionContext& Context, const TDecoratorBinding<IEvaluate>& Binding) const override
 		{
 			if (Private::EvaluatedDecorators != nullptr)
 			{
@@ -361,16 +372,6 @@ bool FAnimationAnimNextRuntimeTest_IHierarchy::RunTest(const FString& InParamete
 		AddErrorIfFalse(ChildrenNodeD[0].IsValid() && ChildrenNodeD[0].GetNodeInstance()->GetNodeHandle() == NodeA, "FAnimationAnimNextRuntimeTest_IHierarchy -> Expected child: NodeA");
 		AddErrorIfFalse(ChildrenNodeD[1].IsValid() && ChildrenNodeD[1].GetNodeInstance()->GetNodeHandle() == NodeC, "FAnimationAnimNextRuntimeTest_IHierarchy -> Expected child: NodeC");
 
-		{
-			TDecoratorBinding<IHierarchy> HierarchyBindingNodeA;
-			AddErrorIfFalse(Context.GetInterface(ChildrenNodeD[0], HierarchyBindingNodeA), "FAnimationAnimNextRuntimeTest_IHierarchy -> IHierarchy not found");
-
-			FChildrenArray ChildrenNodeA;
-			HierarchyBindingNodeA.GetChildren(Context, ChildrenNodeA);
-
-			AddErrorIfFalse(ChildrenNodeA.Num() == 0, "FAnimationAnimNextRuntimeTest_IHierarchy -> Expected 0 children");
-		}
-
 		{
 			TDecoratorBinding<IHierarchy> HierarchyBindingNodeC;
 			AddErrorIfFalse(Context.GetInterface(ChildrenNodeD[1], HierarchyBindingNodeC), "FAnimationAnimNextRuntimeTest_IHierarchy -> IHierarchy not found");
@@ -390,16 +391,6 @@ bool FAnimationAnimNextRuntimeTest_IHierarchy::RunTest(const FString& InParamete
 
 				AddErrorIfFalse(ChildrenNodeB.Num() == 1, "FAnimationAnimNextRuntimeTest_IHierarchy -> Expected 1 child");
 				AddErrorIfFalse(ChildrenNodeB[0].IsValid() && ChildrenNodeB[0].GetNodeInstance()->GetNodeHandle() == NodeA, "FAnimationAnimNextRuntimeTest_IHierarchy -> Expected child: NodeA");
-
-				{
-					TDecoratorBinding<IHierarchy> HierarchyBindingNodeA;
-					AddErrorIfFalse(Context.GetInterface(ChildrenNodeB[0], HierarchyBindingNodeA), "FAnimationAnimNextRuntimeTest_IHierarchy -> IHierarchy not found");
-
-					FChildrenArray ChildrenNodeA;
-					HierarchyBindingNodeA.GetChildren(Context, ChildrenNodeA);
-
-					AddErrorIfFalse(ChildrenNodeA.Num() == 0, "FAnimationAnimNextRuntimeTest_IHierarchy -> Expected 0 children");
-				}
 			}
 		}
 	}
@@ -541,8 +532,7 @@ bool FAnimationAnimNextRuntimeTest_IUpdate::RunTest(const FString& InParameters)
 		AddErrorIfFalse(NodeCPtr.IsValid(), "FAnimationAnimNextRuntimeTest_IUpdate -> Failed to allocate root node instance");
 
 		// Call pre/post update on our graph
-		FUpdateTraversalContext TraversalContext(0.0333f);
-		UpdateGraph(Context, TraversalContext, NodeCPtr);
+		UpdateGraph(Context, NodeCPtr, 0.0333f);
 
 		AddErrorIfFalse(UpdatedDecorators.Num() == 6, "FAnimationAnimNextRuntimeTest_IUpdate -> Expected 6 nodes to have been visited during the update traversal");
 		AddErrorIfFalse(UpdatedDecorators[0] == FDecoratorWithChildren::DecoratorUID, "FAnimationAnimNextRuntimeTest_IUpdate -> Unexpected update order");		// NodeC
