AnimNext: ParamStack changed to use a hash table

This changes parameters from a linear array of entries (with some entries invalid) that is indexed by a globally-agreed upon ID into a hash table that indexes into an array of ptrs to entries. Entry hashes are precomputed and stored in FParamId, but even then the on-demand FName hashing is still much faster than the FRWLock that used to happen.
Performance improvements mostly come from better cache locality, but also from the ability to coalesce parent stacks into the current one, meaning only one hash query is needed in most cases, rather than the recursive call to the parent that we had previously.

diff --git a/Engine/Plugins/Experimental/AnimNext/Source/AnimNext/Internal/Param/ParamDefinition.h b/Engine/Plugins/Experimental/AnimNext/Source/AnimNext/Internal/Param/ParamDefinition.h
--- a/Engine/Plugins/Experimental/AnimNext/Source/AnimNext/Internal/Param/ParamDefinition.h
+++ b/Engine/Plugins/Experimental/AnimNext/Source/AnimNext/Internal/Param/ParamDefinition.h
@@ -30,14 +30,13 @@ enum class EParamDefinitionFlags : uint8
 
 ENUM_CLASS_FLAGS(EParamDefinitionFlags);
 
-// Definition of a parameter for reserved for internal use
+// Definition of a parameter reserved for internal use
 struct FParamDefinition
 {
 	FParamDefinition() = default;
 
 	FParamDefinition(FName InName, const FAnimNextParamType& InType, const FText& InTooltip, EParamDefinitionFlags InFlags = EParamDefinitionFlags::None)
 		: Id(InName)
-		, Name(InName)
 		, TypeHandle(InType.GetHandle())
 		, Type(InType)
 #if WITH_EDITORONLY_DATA
@@ -48,10 +47,10 @@ struct FParamDefinition
 		, Flags(InFlags)
 	{
 		// Param names should not be 'none'
-		check(Name != NAME_None);
+		check(InName != NAME_None);
 		// Param names should not contain periods. Use underscores.
 		// Periods are only a display concern and we want parameters to be expressible as members of objects & structures. 
-		check(!Name.ToString().Contains(TEXT(".")));
+		check(!InName.ToString().Contains(TEXT(".")));
 	}
 
 	FParamDefinition(FName InName, const FAnimNextParamType& InType, EParamDefinitionFlags InFlags = EParamDefinitionFlags::None)
@@ -61,7 +60,6 @@ struct FParamDefinition
 	
 	FParamDefinition(FParamId InId, const FAnimNextParamType& InType, const FText& InTooltip, EParamDefinitionFlags InFlags = EParamDefinitionFlags::None)
 		: Id(InId)
-		, Name(InId.ToName())
 		, TypeHandle(InType.GetHandle())
 		, Type(InType)
 #if WITH_EDITORONLY_DATA
@@ -78,9 +76,18 @@ struct FParamDefinition
 	{
 	}
 
+	// Make a parameter definition from its name and a UObject
+	FParamDefinition(FName InName, const UObject* InObject);
+
+	// Make a parameter definition from its name and a property
+	FParamDefinition(FName InName, const FProperty* InProperty);
+
+	// Make a parameter definition from its name and the return value of a function
+	FParamDefinition(FName InName, const UFunction* InFunction);
+	
 	FName GetName() const
 	{
-		return Name;
+		return Id.GetName();
 	}
 
 	FParamId GetId() const
@@ -104,31 +111,7 @@ private:
 	friend struct FParamAdapter;
 	friend struct UncookedOnly::FUtils;
 
-	FParamDefinition(int32 InIdIndex, FName InName, const FAnimNextParamType& InType, const FText& InTooltip, EParamDefinitionFlags InFlags = EParamDefinitionFlags::None)
-		: Id(InIdIndex)
-		, Name(InName)
-		, TypeHandle(InType.GetHandle())
-		, Type(InType)
-#if WITH_EDITORONLY_DATA
-		, Tooltip(InTooltip)
-#endif
-		, Property(nullptr)
-		, Function(nullptr)
-		, Flags(InFlags)
-	{
-	}
-
-	// Make a parameter definition from its ID and a UObject
-	FParamDefinition(uint32 InIdIndex, FName InName, const UObject* InObject);
-
-	// Make a parameter definition from its ID and a property
-	FParamDefinition(uint32 InIdIndex, FName InName, const FProperty* InProperty);
-
-	// Make a parameter definition from its ID and the return value of a function
-	FParamDefinition(uint32 InIdIndex, FName InName, const UFunction* InFunction);
-
 	FParamId Id;
-	FName Name;
 	FParamTypeHandle TypeHandle;
 	FAnimNextParamType Type;
 #if WITH_EDITORONLY_DATA
diff --git a/Engine/Plugins/Experimental/AnimNext/Source/AnimNext/Internal/Param/RigVMDispatch_GetLayerParameter.h b/Engine/Plugins/Experimental/AnimNext/Source/AnimNext/Internal/Param/RigVMDispatch_GetLayerParameter.h
--- a/Engine/Plugins/Experimental/AnimNext/Source/AnimNext/Internal/Param/RigVMDispatch_GetLayerParameter.h
+++ b/Engine/Plugins/Experimental/AnimNext/Source/AnimNext/Internal/Param/RigVMDispatch_GetLayerParameter.h
@@ -29,7 +29,6 @@ private:
 	virtual const TArray<FRigVMTemplateArgumentInfo>& GetArgumentInfos() const override;
 #if WITH_EDITOR
 	virtual FString GetArgumentMetaData(const FName& InArgumentName, const FName& InMetaDataKey) const override;
-	virtual FString GetArgumentDefaultValue(const FName& InArgumentName, TRigVMTypeIndex InTypeIndex) const override;
 #endif
 	virtual FRigVMTemplateTypeMap OnNewArgumentType(const FName& InArgumentName, TRigVMTypeIndex InTypeIndex) const override;
 	virtual bool IsSingleton() const override { return true; }
diff --git a/Engine/Plugins/Experimental/AnimNext/Source/AnimNext/Internal/Param/RigVMDispatch_GetParameter.h b/Engine/Plugins/Experimental/AnimNext/Source/AnimNext/Internal/Param/RigVMDispatch_GetParameter.h
--- a/Engine/Plugins/Experimental/AnimNext/Source/AnimNext/Internal/Param/RigVMDispatch_GetParameter.h
+++ b/Engine/Plugins/Experimental/AnimNext/Source/AnimNext/Internal/Param/RigVMDispatch_GetParameter.h
@@ -29,7 +29,6 @@ private:
 	virtual const TArray<FRigVMTemplateArgumentInfo>& GetArgumentInfos() const override;
 #if WITH_EDITOR
 	virtual FString GetArgumentMetaData(const FName& InArgumentName, const FName& InMetaDataKey) const override;
-	virtual FString GetArgumentDefaultValue(const FName& InArgumentName, TRigVMTypeIndex InTypeIndex) const override;
 #endif
 	virtual FRigVMTemplateTypeMap OnNewArgumentType(const FName& InArgumentName, TRigVMTypeIndex InTypeIndex) const override;
 	virtual bool IsSingleton() const override { return true; }
diff --git a/Engine/Plugins/Experimental/AnimNext/Source/AnimNext/Internal/Param/RigVMDispatch_SetLayerParameter.h b/Engine/Plugins/Experimental/AnimNext/Source/AnimNext/Internal/Param/RigVMDispatch_SetLayerParameter.h
--- a/Engine/Plugins/Experimental/AnimNext/Source/AnimNext/Internal/Param/RigVMDispatch_SetLayerParameter.h
+++ b/Engine/Plugins/Experimental/AnimNext/Source/AnimNext/Internal/Param/RigVMDispatch_SetLayerParameter.h
@@ -29,7 +29,6 @@ private:
 	virtual const TArray<FRigVMTemplateArgumentInfo>& GetArgumentInfos() const override;
 #if WITH_EDITOR
 	virtual FString GetArgumentMetaData(const FName& InArgumentName, const FName& InMetaDataKey) const override;
-	virtual FString GetArgumentDefaultValue(const FName& InArgumentName, TRigVMTypeIndex InTypeIndex) const override;
 #endif
 	virtual TArray<FRigVMExecuteArgument>& GetExecuteArguments_Impl(const FRigVMDispatchContext& InContext) const override;
 	virtual FRigVMTemplateTypeMap OnNewArgumentType(const FName& InArgumentName, TRigVMTypeIndex InTypeIndex) const override;
diff --git a/Engine/Plugins/Experimental/AnimNext/Source/AnimNext/Private/Component/AnimNextComponent.cpp b/Engine/Plugins/Experimental/AnimNext/Source/AnimNext/Private/Component/AnimNextComponent.cpp
--- a/Engine/Plugins/Experimental/AnimNext/Source/AnimNext/Private/Component/AnimNextComponent.cpp
+++ b/Engine/Plugins/Experimental/AnimNext/Source/AnimNext/Private/Component/AnimNextComponent.cpp
@@ -53,7 +53,7 @@ void UAnimNextComponent::OnRegister()
 
 				if(StackToUse.IsValid())
 				{
-					TArray<TPair<FParamId, Private::FParamEntry>, TInlineAllocator<4>> Params;
+					TArray<Private::FParamEntry, TInlineAllocator<4>> Params;
 					for(UAnimNextComponentParameter* Parameter : ParamPair.Value)
 					{
 						FParamId ParamId;
@@ -64,7 +64,7 @@ void UAnimNextComponent::OnRegister()
 
 						constexpr bool bIsReference = true;
 						constexpr bool bIsMutable = false;
-						Params.Emplace(ParamId, Private::FParamEntry(Type.GetHandle(), TArrayView<uint8>(Value, 1), bIsReference, bIsMutable));
+						Params.Emplace(Private::FParamEntry(ParamId, Type.GetHandle(), TArrayView<uint8>(Value, 1), bIsReference, bIsMutable));
 					}
 
 					FParamStackLayerHandle NewLayer = FParamStack::MakeLayer(Params);
diff --git a/Engine/Plugins/Experimental/AnimNext/Source/AnimNext/Private/Param/ParamAdapter.h b/Engine/Plugins/Experimental/AnimNext/Source/AnimNext/Private/Param/ParamAdapter.h
--- a/Engine/Plugins/Experimental/AnimNext/Source/AnimNext/Private/Param/ParamAdapter.h
+++ b/Engine/Plugins/Experimental/AnimNext/Source/AnimNext/Private/Param/ParamAdapter.h
@@ -43,6 +43,8 @@ struct FParamAdapter
 	FParamDefinition Definition;
 
 	FObjectAdapterFunction Function;
+	
+	uint32 BufferOffset = MAX_uint32;
 };
 
 }
diff --git a/Engine/Plugins/Experimental/AnimNext/Source/AnimNext/Private/Param/ParamDefinition.cpp b/Engine/Plugins/Experimental/AnimNext/Source/AnimNext/Private/Param/ParamDefinition.cpp
--- a/Engine/Plugins/Experimental/AnimNext/Source/AnimNext/Private/Param/ParamDefinition.cpp
+++ b/Engine/Plugins/Experimental/AnimNext/Source/AnimNext/Private/Param/ParamDefinition.cpp
@@ -7,9 +7,8 @@
 namespace UE::AnimNext
 {
 
-FParamDefinition::FParamDefinition(uint32 InIdIndex, FName InName, const UObject* InObject)
-	: Id(FParamId(InIdIndex))
-	, Name(InName)
+FParamDefinition::FParamDefinition(FName InName, const UObject* InObject)
+	: Id(InName)
 	, TypeHandle(FParamTypeHandle::FromObject(InObject))
 	, Type(TypeHandle.GetType())
 #if WITH_EDITORONLY_DATA
@@ -20,9 +19,8 @@ FParamDefinition::FParamDefinition(uint32 InIdIndex, FName InName, const UObject
 {
 }
 
-FParamDefinition::FParamDefinition(uint32 InIdIndex, FName InName, const FProperty* InProperty)
-	: Id(FParamId(InIdIndex))
-	, Name(InName)
+FParamDefinition::FParamDefinition(FName InName, const FProperty* InProperty)
+	: Id(InName)
 	, TypeHandle(FParamTypeHandle::FromProperty(InProperty))
 	, Type(TypeHandle.GetType())
 #if WITH_EDITORONLY_DATA
@@ -32,17 +30,16 @@ FParamDefinition::FParamDefinition(uint32 InIdIndex, FName InName, const FProper
 	, Function(nullptr)
 {
 	// Param names should not contain periods. Periods are only a display concern. Use underscores.
-	check(!Name.ToString().Contains(TEXT(".")));
+	check(!InName.ToString().Contains(TEXT(".")));
 }
 
-FParamDefinition::FParamDefinition(uint32 InIdIndex, FName InName, const UFunction* InFunction)
-	: Id(FParamId(InIdIndex))
-	, Name(InName)
+FParamDefinition::FParamDefinition(FName InName, const UFunction* InFunction)
+	: Id(InName)
 	, Property(nullptr)
 	, Function(InFunction)
 {
 	// Param names should not contain periods. Periods are only a display concern. Use underscores.
-	check(!Name.ToString().Contains(TEXT(".")));
+	check(!InName.ToString().Contains(TEXT(".")));
 
 	const FProperty* ReturnProperty = InFunction->GetReturnProperty();
 	check(ReturnProperty);
diff --git a/Engine/Plugins/Experimental/AnimNext/Source/AnimNext/Private/Param/ParamEntry.cpp b/Engine/Plugins/Experimental/AnimNext/Source/AnimNext/Private/Param/ParamEntry.cpp
--- a/Engine/Plugins/Experimental/AnimNext/Source/AnimNext/Private/Param/ParamEntry.cpp
+++ b/Engine/Plugins/Experimental/AnimNext/Source/AnimNext/Private/Param/ParamEntry.cpp
@@ -2,12 +2,16 @@
 
 #include "Param/ParamEntry.h"
 #include "Param/ParamHelpers.h"
+#include "Param/ParamUtils.h"
+#include "Param/ParamResult.h"
+#include "Param/ParamCompatibility.h"
 
 namespace UE::AnimNext::Private
 {
 
-FParamEntry::FParamEntry(const FParamTypeHandle& InTypeHandle, TArrayView<uint8> InData, bool bInIsReference, bool bInIsMutable)
+FParamEntry::FParamEntry(const FParamId& InId, const FParamTypeHandle& InTypeHandle, TArrayView<uint8> InData, bool bInIsReference, bool bInIsMutable)
 	: Data(nullptr)
+	, Id(InId)
 	, TypeHandle(InTypeHandle)
 	, Size(InData.Num())
 	, Flags(EParamFlags::None)
@@ -49,7 +53,8 @@ FParamEntry::~FParamEntry()
 }
 
 FParamEntry::FParamEntry(const FParamEntry& InOtherParam)
-	: TypeHandle(InOtherParam.TypeHandle)
+	: Id(InOtherParam.Id)
+	, TypeHandle(InOtherParam.TypeHandle)
 	, Size(InOtherParam.Size)
 	, Flags(InOtherParam.Flags)
 {
@@ -65,6 +70,7 @@ FParamEntry::FParamEntry(const FParamEntry& InOtherParam)
 
 FParamEntry& FParamEntry::operator=(const FParamEntry& InOtherParam)
 {
+	Id = InOtherParam.Id;
 	TypeHandle = InOtherParam.TypeHandle;
 	Size = InOtherParam.Size;
 	Flags = InOtherParam.Flags;
@@ -82,7 +88,8 @@ FParamEntry& FParamEntry::operator=(const FParamEntry& InOtherParam)
 }
 
 FParamEntry::FParamEntry(FParamEntry&& InOtherParam) noexcept
-	: TypeHandle(InOtherParam.TypeHandle)
+	: Id(InOtherParam.Id)
+	, TypeHandle(InOtherParam.TypeHandle)
 	, Size(InOtherParam.Size)
 	, Flags(InOtherParam.Flags)
 {
@@ -98,6 +105,7 @@ FParamEntry::FParamEntry(FParamEntry&& InOtherParam) noexcept
 
 FParamEntry& FParamEntry::operator=(FParamEntry&& InOtherParam) noexcept
 {
+	Id = InOtherParam.Id;
 	TypeHandle = InOtherParam.TypeHandle;
 	Size = InOtherParam.Size;
 	Flags = InOtherParam.Flags;
@@ -114,4 +122,79 @@ FParamEntry& FParamEntry::operator=(FParamEntry&& InOtherParam) noexcept
 	return *this;
 }
 
+FParamResult FParamEntry::GetParamData(FParamTypeHandle InTypeHandle, TConstArrayView<uint8>& OutParamData) const
+{
+	FParamCompatibility Compatibility = UE::AnimNext::FParamUtils::GetCompatibility(InTypeHandle, GetTypeHandle());
+	if (!Compatibility.IsCompatible())
+	{
+		return EParamResult::TypeError;
+	}
+
+	OutParamData = GetData();
+	return EParamResult::Success;
+}
+
+FParamResult FParamEntry::GetParamData(FParamTypeHandle InTypeHandle, TConstArrayView<uint8>& OutParamData, FParamTypeHandle& OutParamTypeHandle, FParamCompatibility InRequiredCompatibility) const
+{
+	OutParamTypeHandle = GetTypeHandle();
+
+	FParamCompatibility Compatibility = UE::AnimNext::FParamUtils::GetCompatibility(InTypeHandle, OutParamTypeHandle);
+	if (Compatibility < InRequiredCompatibility)
+	{
+		return EParamResult::TypeError;
+	}
+
+	OutParamData = GetData();
+
+	if (Compatibility == InRequiredCompatibility)
+	{
+		return EParamResult::Success | EParamResult::TypeCompatible;
+	}
+
+	return EParamResult::Success;
+}
+
+FParamResult FParamEntry::GetMutableParamData(FParamTypeHandle InTypeHandle, TArrayView<uint8>& OutParamData)
+{
+	FParamResult AccessResult = EParamResult::Success;
+	FParamCompatibility Compatibility = UE::AnimNext::FParamUtils::GetCompatibility(InTypeHandle, GetTypeHandle());
+	if (!Compatibility.IsCompatible())
+	{
+		AccessResult.Result |= EParamResult::TypeError;
+	}
+
+	AccessResult.Result |= !IsMutable() ? EParamResult::MutabilityError : EParamResult::Success;
+	if (AccessResult.IsSuccessful())
+	{
+		OutParamData = GetMutableData();
+	}
+
+	return AccessResult;
+}
+
+
+FParamResult FParamEntry::GetMutableParamData(FParamTypeHandle InTypeHandle, TArrayView<uint8>& OutParamData, FParamTypeHandle& OutParamTypeHandle, FParamCompatibility InRequiredCompatibility)
+{
+	OutParamTypeHandle = GetTypeHandle();
+
+	FParamResult AccessResult = EParamResult::Success;
+	FParamCompatibility Compatibility = UE::AnimNext::FParamUtils::GetCompatibility(InTypeHandle, OutParamTypeHandle);
+	if (Compatibility < InRequiredCompatibility)
+	{
+		AccessResult.Result |= EParamResult::TypeError;
+	}
+	else if (Compatibility == InRequiredCompatibility)
+	{
+		AccessResult.Result |= EParamResult::TypeCompatible;
+	}
+
+	AccessResult.Result |= !IsMutable() ? EParamResult::MutabilityError : EParamResult::Success;
+	if (AccessResult.IsSuccessful())
+	{
+		OutParamData = GetMutableData();
+	}
+
+	return AccessResult;
+}
+
 }
\ No newline at end of file
diff --git a/Engine/Plugins/Experimental/AnimNext/Source/AnimNext/Private/Param/ParamId.cpp b/Engine/Plugins/Experimental/AnimNext/Source/AnimNext/Private/Param/ParamId.cpp
--- a/Engine/Plugins/Experimental/AnimNext/Source/AnimNext/Private/Param/ParamId.cpp
+++ b/Engine/Plugins/Experimental/AnimNext/Source/AnimNext/Private/Param/ParamId.cpp
@@ -27,8 +27,6 @@ DEFINE_STAT(STAT_AnimNext_ParamIdLock)
 namespace UE::AnimNext
 {
 
-static FRWLock GParamIdLock;
-
 // TODO: Expand scratch buffer system to become a general cache of adapter values.
 // The assumption being that if we are grabbing gameplay-driven values, they do not vary over 
 // the course of a schedule unless the 'tick' of the values (e.g. a component's tick function) runs
@@ -38,23 +36,12 @@ static FRWLock GParamIdLock;
 // types, object lifetimes and GC.
 
 // TODO: Investigate allowing function libraries as adapters
-
-struct FParamData
-{
-	FParamData(FName InName)
-		: Name(InName)
-	{}
-
-	FName Name;
-	int32 AdapterIndex = INDEX_NONE;
-	uint32 BufferOffset = MAX_uint32;
-};
-
+	
 struct FParamIdGlobalData
 {
-	TArray<FParamData> ParamData;
-	TPagedArray<FParamAdapter, sizeof(FParamAdapter) * 256> ParamAdapters;
-	TMap<FName, uint32> NameToParamId;
+	// Hash table used to index the ParamAdapterData 
+	FHashTable HashTable;
+	TArray<FParamAdapter> ParamAdapterData;
 	uint32 SerialNumber = 0;
 	uint32 ResultBufferSize = 0;
 };
@@ -64,26 +51,7 @@ static FParamIdGlobalData GParamIdGlobalData;
 // Buffers to hold values for function-wrapped parameter adapters 
 thread_local TArray<uint8> ParamAdapterResultBuffer;
 thread_local uint32 ParamAdapterBufferSerialNumber = 0;
-
-#if WITH_DEV_AUTOMATION_TESTS
-static FParamIdGlobalData GSandboxedParamIdGlobalData;
-static std::atomic<bool> bGParamIdSandboxed = false;
-#endif
-
-static FParamIdGlobalData& GetParamIdData()
-{
-#if WITH_DEV_AUTOMATION_TESTS		
-	if (bGParamIdSandboxed.load() == true)
-	{
-		return GSandboxedParamIdGlobalData;
-	}
-	else
-#endif
-	{
-		return GParamIdGlobalData;
-	}
-}
-
+	
 void FParamId::Init()
 {
 	RefreshAdapters();
@@ -91,11 +59,8 @@ void FParamId::Init()
 
 void FParamId::Destroy()
 {
-	FRWScopeLock Lock(GParamIdLock, SLT_Write);
-	FParamIdGlobalData& ParamIdData = GetParamIdData();
-	ParamIdData.ParamData.Empty();
-	ParamIdData.ParamAdapters.Empty();
-	ParamIdData.NameToParamId.Empty();
+	GParamIdGlobalData.HashTable.Clear();
+	GParamIdGlobalData.ParamAdapterData.Empty();
 }
 
 void FParamId::RefreshAdapters()
@@ -107,27 +72,75 @@ void FParamId::RefreshAdapters()
 
 void FParamId::ResetAdapters()
 {
-	FParamIdGlobalData& ParamIdGlobalData = GetParamIdData();
+	// Reset existing IDs adapter data
+	for(FParamAdapter& ParamAdapterData : GParamIdGlobalData.ParamAdapterData)
+	{
+		// Unregister this param as built-in
+		FParams::UnregisterBuiltInParameter(ParamAdapterData.Definition.GetName());
+	}
 	
-	ParamIdGlobalData.ResultBufferSize = 0;
-	ParamIdGlobalData.ParamAdapters.Empty();
+	GParamIdGlobalData.HashTable.Clear();
+	GParamIdGlobalData.ParamAdapterData.Empty();
+	GParamIdGlobalData.ResultBufferSize = 0;
+}
 
-	// Reset existing IDs adapter data
-	for(FParamData& ParamData : ParamIdGlobalData.ParamData)
+template<typename ReturnType, bool bRequiresTempStorage>
+static void RegisterBuiltInAdapter(FName InName, FText InTooltipText, FObjectAdapterFunction&& InAdapterFunction, EParamDefinitionFlags InDefinitionFlags = EParamDefinitionFlags::None)
+{
+	FParamDefinition Definition(InName, FAnimNextParamType::GetType<ReturnType>(), InTooltipText, InDefinitionFlags);
+	uint32 Index = GParamIdGlobalData.ParamAdapterData.Emplace(FParamAdapter(MoveTemp(Definition), MoveTemp(InAdapterFunction)));
+	GParamIdGlobalData.HashTable.Add(Definition.GetId().GetHash(), Index);
+	FParams::RegisterBuiltInParameter(Definition);
+
+	if constexpr (bRequiresTempStorage)
 	{
-		ParamData.AdapterIndex = INDEX_NONE;
-		ParamData.BufferOffset = MAX_uint32;
+		FParamAdapter& ParamAdapterData = GParamIdGlobalData.ParamAdapterData[Index];
 
-		// Unregister this param as built-in
-		FParams::UnregisterBuiltInParameter(ParamData.Name);
+		// Calc buffer size for the result
+		ParamAdapterData.BufferOffset = Align(GParamIdGlobalData.ResultBufferSize, alignof(ReturnType));
+		GParamIdGlobalData.ResultBufferSize += sizeof(ReturnType);
 	}
 }
 
-void FParamId::RegisterBuiltInAdapters()
+static void RegisterBuiltInObjectAdapter(FName InName, UObject* InObject, FObjectAdapterFunction&& InAdapterFunction)
 {
-	FRWScopeLock Lock(GParamIdLock, SLT_Write);
-	FParamIdGlobalData& ParamIdGlobalData = GetParamIdData();
+	FParamDefinition Definition(InName, InObject);
+	uint32 Index = GParamIdGlobalData.ParamAdapterData.Emplace(FParamAdapter(MoveTemp(Definition), MoveTemp(InAdapterFunction)));
+	GParamIdGlobalData.HashTable.Add(Definition.GetId().GetHash(), Index);
+	FParams::RegisterBuiltInParameter(Definition);
+
+	FParamAdapter& ParamAdapterData = GParamIdGlobalData.ParamAdapterData[Index];
+
+	// Calc buffer size for the result
+	ParamAdapterData.BufferOffset = Align(GParamIdGlobalData.ResultBufferSize, alignof(UObject*));
+	GParamIdGlobalData.ResultBufferSize += sizeof(UObject*);
+}
+
+static void RegisterBuiltInPropertyAdapter(FName InName, const FProperty* InProperty, FObjectAdapterFunction&& InAdapterFunction)
+{
+	FParamDefinition Definition(InName, InProperty);
+	uint32 Index = GParamIdGlobalData.ParamAdapterData.Emplace(FParamAdapter(MoveTemp(Definition), MoveTemp(InAdapterFunction)));
+	GParamIdGlobalData.HashTable.Add(Definition.GetId().GetHash(), Index);
+	FParams::RegisterBuiltInParameter(Definition);
+}
 
+static void RegisterBuiltInFunctionAdapter(FName InName, UFunction* InFunction, FObjectAdapterFunction&& InAdapterFunction)
+{
+	FParamDefinition Definition(InName, InFunction);
+	uint32 Index = GParamIdGlobalData.ParamAdapterData.Emplace(FParamAdapter(MoveTemp(Definition), MoveTemp(InAdapterFunction)));
+	GParamIdGlobalData.HashTable.Add(Definition.GetId().GetHash(), Index);
+	FParams::RegisterBuiltInParameter(Definition);
+
+	FParamAdapter& ParamAdapterData = GParamIdGlobalData.ParamAdapterData[Index];
+
+	// Calc buffer size for the result
+	const FProperty* ReturnProperty = InFunction->GetReturnProperty();
+	ParamAdapterData.BufferOffset = Align(GParamIdGlobalData.ResultBufferSize, ReturnProperty->GetMinAlignment());
+	GParamIdGlobalData.ResultBufferSize += ReturnProperty->GetSize();
+}
+	
+void FParamId::RegisterBuiltInAdapters()
+{
 	// Register world delta time
 	{
 		FObjectAdapterFunction DeltaTimeFunction = [](UObject* InContextObject, FParamId InId) -> uint8*
@@ -141,16 +154,7 @@ void FParamId::RegisterBuiltInAdapters()
 			return nullptr;
 		};
 
-		FName DeltaTimeName("UE_Frame_DeltaTime");
-		const uint32 DeltaTimeNameFunctionIndex = MakeParamId_NoLock(DeltaTimeName);
-		ParamIdGlobalData.ParamData[DeltaTimeNameFunctionIndex].AdapterIndex = ParamIdGlobalData.ParamAdapters.Num();
-		FText TooltipText = LOCTEXT("DeltaTimeTooltip", "The current delta time");
-		FParamAdapter& NewAdapter = ParamIdGlobalData.ParamAdapters.Emplace_GetRef(FParamDefinition(DeltaTimeNameFunctionIndex, DeltaTimeName, FAnimNextParamType::GetType<float>(), TooltipText), MoveTemp(DeltaTimeFunction));
-		FParams::RegisterBuiltInParameter(NewAdapter.Definition);
-
-		// Calc buffer size for the result
-		ParamIdGlobalData.ParamData[DeltaTimeNameFunctionIndex].BufferOffset = Align(ParamIdGlobalData.ResultBufferSize, alignof(float));
-		ParamIdGlobalData.ResultBufferSize += sizeof(float);
+		RegisterBuiltInAdapter<float, true>("UE_Frame_DeltaTime", LOCTEXT("DeltaTimeTooltip", "The current delta time"), MoveTemp(DeltaTimeFunction));
 	}
 	
 	// Register physics tick function
@@ -164,12 +168,7 @@ void FParamId::RegisterBuiltInAdapters()
 			return nullptr;
 		};
 
-		FName PhysicsTickName("UE_Physics_Tick");
-		const uint32 PhysicsTickFunctionIndex = MakeParamId_NoLock(PhysicsTickName);
-		ParamIdGlobalData.ParamData[PhysicsTickFunctionIndex].AdapterIndex = ParamIdGlobalData.ParamAdapters.Num();
-		FText TooltipText = LOCTEXT("PhysicsTickTooltip", "The tick function of the physics scene");
-		FParamAdapter& NewAdapter = ParamIdGlobalData.ParamAdapters.Emplace_GetRef(FParamDefinition(PhysicsTickFunctionIndex, PhysicsTickName, FAnimNextParamType::GetType<FTickFunction>(), TooltipText), MoveTemp(PhysicsTickFunctionFunction));
-		FParams::RegisterBuiltInParameter(NewAdapter.Definition);
+		RegisterBuiltInAdapter<FEndPhysicsTickFunction, false>("UE_Physics_Tick", LOCTEXT("PhysicsTickTooltip", "The tick function of the physics scene"), MoveTemp(PhysicsTickFunctionFunction));
 	}
 
 	// Register ref pose accessor
@@ -191,24 +190,12 @@ void FParamId::RegisterBuiltInAdapters()
 			return nullptr;
 		};
 
-		FName ReferencePoseName("UE_AnimNextMeshComponent_ReferencePose");
-		const uint32 ReferencePoseIndex = MakeParamId_NoLock(ReferencePoseName);
-		ParamIdGlobalData.ParamData[ReferencePoseIndex].AdapterIndex = ParamIdGlobalData.ParamAdapters.Num();
-		FText TooltipText = LOCTEXT("ReferencePoseTooltip", "The reference pose of the skeletal mesh component");
-		FParamAdapter& NewAdapter = ParamIdGlobalData.ParamAdapters.Emplace_GetRef(FParamDefinition(ReferencePoseIndex, ReferencePoseName, FAnimNextParamType::GetType<FAnimNextGraphReferencePose>(), TooltipText), MoveTemp(ReferencePoseFunction));
-		FParams::RegisterBuiltInParameter(NewAdapter.Definition);
-
-		// Calc buffer size for the result
-		ParamIdGlobalData.ParamData[ReferencePoseIndex].BufferOffset = Align(ParamIdGlobalData.ResultBufferSize, alignof(FAnimNextGraphReferencePose));
-		ParamIdGlobalData.ResultBufferSize += sizeof(FAnimNextGraphReferencePose);
+		RegisterBuiltInAdapter<FAnimNextGraphReferencePose, true>("UE_AnimNextMeshComponent_ReferencePose", LOCTEXT("ReferencePoseTooltip", "The reference pose of the skeletal mesh component"), MoveTemp(ReferencePoseFunction));
 	}
 }
 
 void FParamId::RefreshConfigAdapters()
 {
-	FRWScopeLock Lock(GParamIdLock, SLT_Write);
-	FParamIdGlobalData& ParamIdGlobalData = GetParamIdData();
-
 	for(const FAnimNextObjectAdapterConfig& AdapterConfig : GetDefault<UAnimNextConfig>()->ExposedClasses)
 	{
 		if(AdapterConfig.Class != nullptr && AdapterConfig.RootParameter != NAME_None)
@@ -218,16 +205,13 @@ void FParamId::RefreshConfigAdapters()
 
 			const FString RootParam(AdapterConfig.RootParameter.ToString());
 			const FName RootParamName(*RootParam);
-			const uint32 RootParamIndex = MakeParamId_NoLock(RootParamName);
-			const FString NameParam = RootParam + TEXT("_Name");
-			const FName NameParamName(*NameParam);
-			const uint32 NameParamIndex = MakeParamId_NoLock(NameParamName);
+			const FParamId RootParamId(RootParamName);
 
 			// TODO: add registration mechanism to allow for more than actor/component types
 			FObjectAdapterFunction RootFunction;
 			if(AdapterConfig.Class->IsChildOf(UActorComponent::StaticClass()))
 			{
-				RootFunction = [Class = AdapterConfig.Class, NameParamIndex](UObject* InContextObject, FParamId InId) -> uint8*
+				RootFunction = [Class = AdapterConfig.Class](UObject* InContextObject, FParamId InId) -> uint8*
 				{
 					if (const UActorComponent* ContextComponent = Cast<UActorComponent>(InContextObject))
 					{
@@ -271,16 +255,8 @@ void FParamId::RefreshConfigAdapters()
 			// - When an object's tick function is not concurrently running (if necessary - should be mitigated by the above)
 			
 			// Add root param
-			ParamIdGlobalData.ParamData[RootParamIndex].AdapterIndex = ParamIdGlobalData.ParamAdapters.Num(); 
-			FParamAdapter& NewRootAdapter = ParamIdGlobalData.ParamAdapters.Emplace_GetRef(FParamDefinition(RootParamIndex, RootParamName, AdapterConfig.Class->GetDefaultObject()), MoveTemp(RootFunction));
-			FParams::RegisterBuiltInParameter(NewRootAdapter.Definition);
-
-			// Calc buffer size for the ptr's result
-			ParamIdGlobalData.ParamData[RootParamIndex].BufferOffset = Align(ParamIdGlobalData.ResultBufferSize, alignof(UObject*));
-			ParamIdGlobalData.ResultBufferSize += sizeof(UObject*);
-
-			bool bNameParamAdded = false;
-
+			RegisterBuiltInObjectAdapter(RootParamName, AdapterConfig.Class->GetDefaultObject(), MoveTemp(RootFunction));
+			
 			// Add object's BP-exposed properties
 			if(AdapterConfig.FilterType != EAnimNextObjectAdapterFilterType::AllowOnlyFunctions)
 			{
@@ -301,15 +277,11 @@ void FParamId::RefreshConfigAdapters()
 						{
 							TStringBuilder<128> ParameterNameBuilder;
 							ParameterNameBuilder.Appendf(TEXT("%s_%s"), *RootParam, *Property->GetName());
-							if(NameParam.Compare(ParameterNameBuilder.ToString(), ESearchCase::IgnoreCase) == 0)
-							{
-								bNameParamAdded = true;
-							}
 							
-							FObjectAdapterFunction AdapterFunction = [RootParamIndex, Property, Class = AdapterConfig.Class.Get()](UObject* InContextObject, FParamId InId) -> uint8*
+							FObjectAdapterFunction AdapterFunction = [RootParamId, Property, Class = AdapterConfig.Class.Get()](UObject* InContextObject, FParamId InId) -> uint8*
 							{
 								// Grab object container
-								if(const TObjectPtr<UObject>* Object = FParamStack::Get().GetParamPtr<TObjectPtr<UObject>>(FParamId(RootParamIndex)))
+								if(const TObjectPtr<UObject>* Object = FParamStack::Get().GetParamPtr<TObjectPtr<UObject>>(RootParamId))
 								{
 									check(Class->IsChildOf(Property->GetOwnerClass()));
 									return const_cast<uint8*>(Property->ContainerPtrToValuePtr<uint8>(*Object));
@@ -318,10 +290,7 @@ void FParamId::RefreshConfigAdapters()
 							};
 
 							FName ParameterName(ParameterNameBuilder.ToString());
-							const uint32 NewIdIndex = MakeParamId_NoLock(ParameterName);
-							ParamIdGlobalData.ParamData[NewIdIndex].AdapterIndex = ParamIdGlobalData.ParamAdapters.Num(); 
-							FParamAdapter& NewAdapter = ParamIdGlobalData.ParamAdapters.Emplace_GetRef(FParamDefinition(NewIdIndex, ParameterName, Property), MoveTemp(AdapterFunction));
-							FParams::RegisterBuiltInParameter(NewAdapter.Definition);
+							RegisterBuiltInPropertyAdapter(ParameterName, Property, MoveTemp(AdapterFunction));
 						}
 					}
 				}
@@ -354,17 +323,13 @@ void FParamId::RefreshConfigAdapters()
 							
 							TStringBuilder<128> ParameterNameBuilder;
 							ParameterNameBuilder.Appendf(TEXT("%s_%s"), *RootParam, *FunctionName);
-							if(NameParam.Compare(ParameterNameBuilder.ToString(), ESearchCase::IgnoreCase) == 0)
-							{
-								bNameParamAdded = true;
-							}
 
-							FObjectAdapterFunction AdapterFunction = [RootParamIndex, WeakFunction = TWeakObjectPtr<UFunction>(Function), Class = AdapterConfig.Class.Get()](UObject* InContextObject, FParamId InId) -> uint8*
+							FObjectAdapterFunction AdapterFunction = [RootParamId, WeakFunction = TWeakObjectPtr<UFunction>(Function), Class = AdapterConfig.Class.Get()](UObject* InContextObject, FParamId InId) -> uint8*
 							{
 								// Grab object to call the function with
 								if(UFunction* Function = WeakFunction.Get())
 								{
-									if(const TObjectPtr<UObject>* ObjectPtr = FParamStack::Get().GetParamPtr<TObjectPtr<UObject>>(FParamId(RootParamIndex)))
+									if(const TObjectPtr<UObject>* ObjectPtr = FParamStack::Get().GetParamPtr<TObjectPtr<UObject>>(RootParamId))
 									{
 										if(UObject* Object = *ObjectPtr)
 										{
@@ -385,29 +350,12 @@ void FParamId::RefreshConfigAdapters()
 							};
 
 							FName ParameterName(ParameterNameBuilder.ToString());
-							const uint32 NewIdIndex = MakeParamId_NoLock(ParameterName);
-							ParamIdGlobalData.ParamData[NewIdIndex].AdapterIndex = ParamIdGlobalData.ParamAdapters.Num(); 
-							FParamAdapter& NewAdapter = ParamIdGlobalData.ParamAdapters.Emplace_GetRef(FParamDefinition(NewIdIndex, ParameterName, Function), MoveTemp(AdapterFunction));
-							FParams::RegisterBuiltInParameter(NewAdapter.Definition);
-
-							// Calc buffer size for the function's result
-							ParamIdGlobalData.ParamData[NewIdIndex].BufferOffset = Align(ParamIdGlobalData.ResultBufferSize, ReturnProperty->GetMinAlignment());
-							ParamIdGlobalData.ResultBufferSize += ReturnProperty->GetSize();
+							RegisterBuiltInFunctionAdapter(ParameterName, Function, MoveTemp(AdapterFunction));
 						}
 					}
 				}
 			}
 
-			if(!bNameParamAdded)
-			{
-				// Add name param if we haven't already above
-				FText TooltipText;
-#if WITH_EDITOR
-				TooltipText = FText::Format(LOCTEXT("NameTooltipFormat", "The name of the {0}"), AdapterConfig.Class->GetDisplayNameText());
-#endif
-				FParams::RegisterBuiltInParameter(FParamDefinition(NameParamIndex, NameParamName, FAnimNextParamType::GetType<FName>(), TooltipText));
-			}
-
 			if(AdapterConfig.bRegisterTickFunction)
 			{
 				TStringBuilder<128> ParameterNameBuilder;
@@ -418,9 +366,9 @@ void FParamId::RefreshConfigAdapters()
 				if(AdapterConfig.Class->IsChildOf(UActorComponent::StaticClass()))
 				{
 					ParamType = FAnimNextParamType::GetType<FActorComponentTickFunction>();
-					TickFunctionFunction = [RootParamIndex](UObject* InContextObject, FParamId InId) -> uint8*
+					TickFunctionFunction = [RootParamId](UObject* InContextObject, FParamId InId) -> uint8*
 					{
-						if (const TObjectPtr<UActorComponent>* Component = FParamStack::Get().GetParamPtr<TObjectPtr<UActorComponent>>(FParamId(RootParamIndex)))
+						if (const TObjectPtr<UActorComponent>* Component = FParamStack::Get().GetParamPtr<TObjectPtr<UActorComponent>>(RootParamId))
 						{
 							if(Component->Get())
 							{
@@ -433,9 +381,9 @@ void FParamId::RefreshConfigAdapters()
 				else if(AdapterConfig.Class->IsChildOf(AActor::StaticClass()))
 				{
 					ParamType = FAnimNextParamType::GetType<FActorTickFunction>();
-					TickFunctionFunction = [RootParamIndex](UObject* InContextObject, FParamId InId) -> uint8*
+					TickFunctionFunction = [RootParamId](UObject* InContextObject, FParamId InId) -> uint8*
 					{
-						if (const TObjectPtr<AActor>* Actor = FParamStack::Get().GetParamPtr<TObjectPtr<AActor>>(FParamId(RootParamIndex)))
+						if (const TObjectPtr<AActor>* Actor = FParamStack::Get().GetParamPtr<TObjectPtr<AActor>>(RootParamId))
 						{
 							if(Actor->Get())
 							{
@@ -451,10 +399,7 @@ void FParamId::RefreshConfigAdapters()
 				TooltipText = FText::Format(LOCTEXT("TickFunctionTooltipFormat", "The tick function of a {0}"), AdapterConfig.Class->GetDisplayNameText());
 #endif
 				FName ParameterName(ParameterNameBuilder.ToString());
-				const uint32 NewIdIndex = MakeParamId_NoLock(ParameterName);
-				ParamIdGlobalData.ParamData[NewIdIndex].AdapterIndex = ParamIdGlobalData.ParamAdapters.Num();
-				FParamAdapter& NewAdapter = ParamIdGlobalData.ParamAdapters.Emplace_GetRef(FParamDefinition(NewIdIndex, ParameterName, ParamType, TooltipText, EParamDefinitionFlags::Mutable), MoveTemp(TickFunctionFunction));
-				FParams::RegisterBuiltInParameter(NewAdapter.Definition);
+				RegisterBuiltInAdapter<FTickFunction, false>(ParameterName, TooltipText, MoveTemp(TickFunctionFunction), EParamDefinitionFlags::Mutable);
 			}
 		}
 	}
@@ -463,110 +408,54 @@ void FParamId::RefreshConfigAdapters()
 	FRigVMRegistry::Get().RefreshEngineTypes();
 
 	// Increment result buffer serial so we resize on access
-	ParamIdGlobalData.SerialNumber++;
-	if(ParamIdGlobalData.SerialNumber == 0)
+	GParamIdGlobalData.SerialNumber++;
+	if(GParamIdGlobalData.SerialNumber == 0)
 	{
-		ParamIdGlobalData.SerialNumber++;
+		GParamIdGlobalData.SerialNumber++;
 	}
 }
 
 FParamId::FParamId(FName InName)
+	: Name(InName)
+	, Hash(GetTypeHash(InName))
 {
-	SCOPE_CYCLE_COUNTER(STAT_AnimNext_ParamIdLock);
-
-	FRWScopeLock Lock(GParamIdLock, SLT_Write);
-	ParameterIndex = MakeParamId_NoLock(InName);
-}
-
-uint32 FParamId::MakeParamId_NoLock(FName InName)
-{
-	FParamIdGlobalData& ParamIdData = GetParamIdData();
-
-	// Param names should not be 'none'
-	check(InName != NAME_None);
-	// Param names should not contain periods. Use underscores.
-	// Periods are only a display concern and we want parameters to be expressible as members of objects & structures. 
-	check(!InName.ToString().Contains(TEXT(".")));
-
-	uint32 Index = InvalidIndex;
-	if (const uint32* FoundIndex = ParamIdData.NameToParamId.Find(InName))
-	{
-		Index = *FoundIndex;
-	}
-	else
-	{
-		Index = ParamIdData.ParamData.Num();
-		ParamIdData.NameToParamId.Add(InName, Index);
-		ParamIdData.ParamData.Add(InName);
-	}
-
-	return Index;
 }
 
-FName FParamId::ToName() const
+FParamId::FParamId(FName InName, uint32 InHash)
+	: Name(InName)
+	, Hash(InHash)
 {
-	FRWScopeLock Lock(GParamIdLock, SLT_ReadOnly);
-	return GetParamIdData().ParamData[ParameterIndex].Name;
+	checkSlow(Hash == GetTypeHash(Name));
 }
 
 const FParamAdapter* FParamId::GetAdapter(FParamId InId)
 {
-	FRWScopeLock Lock(GParamIdLock, SLT_ReadOnly);
-	const FParamIdGlobalData& ParamIdData = GetParamIdData();
-	const int32 DefinitionIndex = ParamIdData.ParamData[InId.ParameterIndex].AdapterIndex;
-	return DefinitionIndex != INDEX_NONE ? &ParamIdData.ParamAdapters[DefinitionIndex] : nullptr;
+	for(uint32 Index = GParamIdGlobalData.HashTable.First(InId.GetHash()); GParamIdGlobalData.HashTable.IsValid(Index); Index = GParamIdGlobalData.HashTable.Next(Index))
+	{
+		if(GParamIdGlobalData.ParamAdapterData[Index].Definition.GetName() == InId.GetName())
+		{
+			return &GParamIdGlobalData.ParamAdapterData[Index];
+		}
+	}
+	return nullptr;
 }
 
 uint8* FParamId::GetScratchAreaForParamIdAdapter(FParamId InId)
 {
-	uint32 BufferOffset;
-	uint32 SerialNumber;
-	uint32 ResultBufferSize;
-
+	if(const FParamAdapter* Adapter = GetAdapter(InId))
 	{
-		FRWScopeLock Lock(GParamIdLock, SLT_ReadOnly);
-		const FParamIdGlobalData& ParamIdData = GetParamIdData();
-		BufferOffset = ParamIdData.ParamData[InId.ParameterIndex].BufferOffset;
-		SerialNumber = ParamIdData.SerialNumber;
-		ResultBufferSize = ParamIdData.ResultBufferSize;
-	}
+		if(ParamAdapterBufferSerialNumber != GParamIdGlobalData.SerialNumber)
+		{
+			// realloc thread-local storage for results
+			ParamAdapterResultBuffer.SetNumZeroed(GParamIdGlobalData.ResultBufferSize);
+			ParamAdapterBufferSerialNumber = GParamIdGlobalData.SerialNumber;
+		}
 
-	if(ParamAdapterBufferSerialNumber != SerialNumber)
-	{
-		// realloc thread-local storage for results
-		ParamAdapterResultBuffer.SetNumZeroed(ResultBufferSize);
-		ParamAdapterBufferSerialNumber = SerialNumber;
+		return &ParamAdapterResultBuffer[Adapter->BufferOffset];
 	}
-
-	return &ParamAdapterResultBuffer[BufferOffset];
-}
-
-FParamId FParamId::GetMaxParamId()
-{
-	FRWScopeLock Lock(GParamIdLock, SLT_ReadOnly);
-	return FParamId((uint32)GetParamIdData().ParamData.Num());
+	return nullptr;
 }
 
-#if WITH_DEV_AUTOMATION_TESTS
-void FParamId::BeginTestSandbox()
-{
-	FRWScopeLock Lock(GParamIdLock, SLT_Write);
-	check(bGParamIdSandboxed.load() == false);
-	bGParamIdSandboxed.exchange(true);
-	GSandboxedParamIdGlobalData.NameToParamId.Empty();
-	GSandboxedParamIdGlobalData.ParamData.Empty();
-}
-
-void FParamId::EndTestSandbox()
-{
-	FRWScopeLock Lock(GParamIdLock, SLT_Write);
-	check(bGParamIdSandboxed.load() == true);
-	bGParamIdSandboxed.exchange(false);
-	GSandboxedParamIdGlobalData.NameToParamId.Empty();
-	GSandboxedParamIdGlobalData.ParamData.Empty();
-}
-#endif
-
 }
 
 #undef LOCTEXT_NAMESPACE
\ No newline at end of file
diff --git a/Engine/Plugins/Experimental/AnimNext/Source/AnimNext/Private/Param/ParamStack.cpp b/Engine/Plugins/Experimental/AnimNext/Source/AnimNext/Private/Param/ParamStack.cpp
--- a/Engine/Plugins/Experimental/AnimNext/Source/AnimNext/Private/Param/ParamStack.cpp
+++ b/Engine/Plugins/Experimental/AnimNext/Source/AnimNext/Private/Param/ParamStack.cpp
@@ -13,55 +13,36 @@
 
 DEFINE_STAT(STAT_AnimNext_ParamStack_GetParam);
 DEFINE_STAT(STAT_AnimNext_ParamStack_Adapter);
+DEFINE_STAT(STAT_AnimNext_ParamStack_Coalesce);
+DEFINE_STAT(STAT_AnimNext_ParamStack_Decoalesce);
 
 #define LOCTEXT_NAMESPACE "AnimNextParamStack"
 
 namespace UE::AnimNext
 {
-
+	
 // Stack layer that can own its own data as a UObject or reference an externally owned object
 struct FUObjectLayer : FParamStackLayer
 {
 	FUObjectLayer() = delete;
 
 	explicit FUObjectLayer(UObject* InObject, bool bInMutable)
+		: FParamStackLayer(32)
 	{
 		Object = InObject;
 
 		const UClass* Class = InObject->GetClass();
-
-		TArray<FParamId, TInlineAllocator<128>> CachedIds;
-		TArray<FProperty*, TInlineAllocator<128>> CachedProperties;
-		TArray<FParamTypeHandle, TInlineAllocator<128>> CachedParamTypes;
-
-		// Determine param ID range for this layer
-		MinParamId = MAX_uint32;
-		uint32 MaxParamId = 0;
-
+		
 		for (TFieldIterator<FProperty> PropIt(Class, EFieldIteratorFlags::IncludeSuper); PropIt; ++PropIt)
 		{
 			FParamTypeHandle ParamTypeHandle = FParamTypeHandle::FromProperty(*PropIt);
 			if(ParamTypeHandle.IsValid())
 			{
-				CachedParamTypes.Add(ParamTypeHandle);
-				CachedProperties.Add(*PropIt);
-				const FParamId& ParamId = CachedIds.Add_GetRef(FParamId(PropIt->GetFName()));
-				MinParamId = FMath::Min(ParamId.ToInt(), MinParamId);
-				MaxParamId = FMath::Max(ParamId.ToInt(), MaxParamId);
-			}
-		}
-
-		if (MinParamId <= MaxParamId)
-		{
-			const uint32 ParamRangeSize = (MaxParamId - MinParamId) + 1;
-			Params.SetNumZeroed(ParamRangeSize);
-			for (uint32 PropertyIndex = 0; PropertyIndex < static_cast<uint32>(CachedProperties.Num()); ++PropertyIndex)
-			{
-				const FProperty* Property = CachedProperties[PropertyIndex];
-				const FParamId& ParamId = CachedIds[PropertyIndex];
+				const FProperty* Property = *PropIt;
+				FParamId ParamId(PropIt->GetFName());
 				uint8* DataPtr = Property->ContainerPtrToValuePtr<uint8>(InObject);
-				const uint32 LocalParamIndex = ParamId.ToInt() - MinParamId;
-				Params[LocalParamIndex] = Private::FParamEntry(CachedParamTypes[PropertyIndex], TArrayView<uint8>(DataPtr, Property->GetSize()), true, bInMutable);
+				uint32 ParamIndex = Params.Emplace(ParamId, ParamTypeHandle, TArrayView<uint8>(DataPtr, Property->GetSize()), true, bInMutable);
+				HashTable.Add(ParamId.GetHash(), ParamIndex);
 			}
 		}
 	}
@@ -81,37 +62,20 @@ struct FInstancedPropertyBagLayer : FParamStackLayer, FGCObject
 	FInstancedPropertyBagLayer() = delete;
 
 	explicit FInstancedPropertyBagLayer(FInstancedPropertyBag& InPropertyBag, bool bInMutable)
+		: FParamStackLayer(InPropertyBag.GetPropertyBagStruct() ? InPropertyBag.GetPropertyBagStruct()->GetPropertyDescs().Num() : 0)
 	{
 		if (const UPropertyBag* PropertyBagStruct = InPropertyBag.GetPropertyBagStruct())
 		{
 			TConstArrayView<FPropertyBagPropertyDesc> Descs = PropertyBagStruct->GetPropertyDescs();
-
-			// Determine param ID range for this layer
-			TArray<FParamId, TInlineAllocator<128>> CachedIds;
-			CachedIds.SetNumUninitialized(Descs.Num());
-			MinParamId = MAX_uint32;
-			uint32 MaxParamId = 0;
+			Params.Reserve(Descs.Num());
+			FStructView StructView = InPropertyBag.GetMutableValue();
 			for (uint32 DescIndex = 0; DescIndex < static_cast<uint32>(Descs.Num()); ++DescIndex)
 			{
 				const FPropertyBagPropertyDesc& Desc = Descs[DescIndex];
-				const FParamId& ParamId = CachedIds[DescIndex] = FParamId(Desc.Name);
-				MinParamId = FMath::Min(ParamId.ToInt(), MinParamId);
-				MaxParamId = FMath::Max(ParamId.ToInt(), MaxParamId);
-			}
-
-			if (MinParamId <= MaxParamId)
-			{
-				const uint32 ParamRangeSize = (MaxParamId - MinParamId) + 1;
-				Params.SetNumZeroed(ParamRangeSize);
-				FStructView StructView = InPropertyBag.GetMutableValue();
-				for (uint32 DescIndex = 0; DescIndex < static_cast<uint32>(Descs.Num()); ++DescIndex)
-				{
-					const FPropertyBagPropertyDesc& Desc = Descs[DescIndex];
-					const FParamId& ParamId = CachedIds[DescIndex];
-					uint8* DataPtr = StructView.GetMemory() + Desc.CachedProperty->GetOffset_ForInternal();
-					const uint32 LocalParamIndex = ParamId.ToInt() - MinParamId;
-					Params[LocalParamIndex] = Private::FParamEntry(FParamTypeHandle::FromPropertyBagPropertyDesc(Desc), TArrayView<uint8>(DataPtr, Desc.CachedProperty->GetSize()), true, true);
-				}
+				const FParamId ParamId(Desc.Name);
+				uint8* DataPtr = StructView.GetMemory() + Desc.CachedProperty->GetOffset_ForInternal();
+				uint32 ParamIndex = Params.Emplace(ParamId, FParamTypeHandle::FromPropertyBagPropertyDesc(Desc), TArrayView<uint8>(DataPtr, Desc.CachedProperty->GetSize()), true, true);
+				HashTable.Add(ParamId.GetHash(), ParamIndex);
 			}
 		}
 	}
@@ -176,38 +140,17 @@ struct FRemappedLayer : FParamStackLayer
 	FRemappedLayer() = delete;
 
 	explicit FRemappedLayer(const FParamStackLayer& InLayer, const TMap<FName, FName>& InMapping)
+		: FParamStackLayer(InMapping.Num())
 	{
-		// Determine param ID range for this layer
-		TArray<TPair<FParamId, Private::FParamEntry>, TInlineAllocator<128>> CachedParams;
-		CachedParams.Reserve(InMapping.Num());
-		MinParamId = MAX_uint32;
-		uint32 MaxParamId = 0;
-		for (uint32 LocalParamIndex = 0; LocalParamIndex < static_cast<uint32>(InLayer.Params.Num()); ++LocalParamIndex)
+		Params.Reserve(InMapping.Num());
+		for (const Private::FParamEntry& OtherParamEntry : InLayer.Params)
 		{
-			const Private::FParamEntry& ParamEntry = InLayer.Params[LocalParamIndex];
-			if(ParamEntry.IsValid())
+			if(const FName* RemappedName = InMapping.Find(OtherParamEntry.GetName()))
 			{
-				const uint32 GlobalParamIndex = InLayer.MinParamId + LocalParamIndex;
-				const FParamId SourceParamId(GlobalParamIndex);
-				if(const FName* RemappedName = InMapping.Find(SourceParamId.ToName()))
-				{
-					const FParamId RemappedParamId(*RemappedName);
-					MinParamId = FMath::Min(RemappedParamId.ToInt(), MinParamId);
-					MaxParamId = FMath::Max(RemappedParamId.ToInt(), MaxParamId);
-					CachedParams.Add({ RemappedParamId, ParamEntry });
-				}
-			}
-		}
-
-		if (MinParamId <= MaxParamId)
-		{
-			const uint32 ParamRangeSize = (MaxParamId - MinParamId) + 1;
-			Params.SetNumZeroed(ParamRangeSize);
-			for (uint32 ParamIndex = 0; ParamIndex < static_cast<uint32>(CachedParams.Num()); ++ParamIndex)
-			{
-				const TPair<FParamId, Private::FParamEntry>& ParamIdPair = CachedParams[ParamIndex];
-				const uint32 LocalParamIndex = ParamIdPair.Key.ToInt() - MinParamId;
-				Params[LocalParamIndex] = ParamIdPair.Value;
+				uint32 ParamIndex = Params.Emplace(OtherParamEntry);
+				Private::FParamEntry& ParamEntry = Params[ParamIndex];
+				ParamEntry.Id = FParamId(*RemappedName);
+				HashTable.Add(ParamEntry.Id.GetHash(), ParamIndex);
 			}
 		}
 	}
@@ -218,57 +161,26 @@ FParamStack::FPushedLayer::FPushedLayer(FParamStackLayer& InLayer, FParamStack&
 {
 	SerialNumber = InStack.MakeSerialNumber();
 
-	if (Layer.Params.Num())
+	InStack.EntryStack.Reserve(InStack.EntryStack.Num() + InLayer.Params.Num());
+	for (Private::FParamEntry& Param : InLayer.Params)
 	{
-		InStack.ResizeLayerIndices();
-
-		for (uint32 LocalParamIndex = 0; LocalParamIndex < (uint32)InLayer.Params.Num(); ++LocalParamIndex)
-		{
-			const uint32 GlobalParamIndex = Layer.MinParamId + LocalParamIndex;
-			InStack.LayerIndices[GlobalParamIndex] = InLayer.Params[LocalParamIndex].IsValid() ? 0 : MAX_uint16;
-		}
-	}
-}
-
-FParamStack::FPushedLayer::FPushedLayer(const FPushedLayer& InPreviousLayer, FParamStackLayer& InLayer, FParamStack& InStack)
-	: Layer(InLayer)
-{
-	SerialNumber = InStack.MakeSerialNumber();
-
-	if (Layer.Params.Num())
-	{
-		InStack.ResizeLayerIndices();
-
-		const uint32 NumParams = Layer.Params.Num();
-		const uint32 BaseLayerIndex = InStack.PreviousLayerIndices.Num();
-		InStack.PreviousLayerIndices.SetNum(BaseLayerIndex + NumParams);
-		PreviousLayerIndexStart = BaseLayerIndex;
-		
-		// Update layer indices and build previous layer indices
-		for (uint32 LocalParamIndex = 0; LocalParamIndex < NumParams; ++LocalParamIndex)
-		{
-			const uint32 GlobalParamIndex = Layer.MinParamId + LocalParamIndex;
-			InStack.PreviousLayerIndices[BaseLayerIndex + LocalParamIndex] = InStack.LayerIndices[GlobalParamIndex];
-			if(InLayer.Params[LocalParamIndex].IsValid())
-			{
-				InStack.LayerIndices[GlobalParamIndex] = InStack.Layers.Num() - 1;
-			}
-		}
+		uint32 ParamIndex = InStack.EntryStack.Add(&Param);
+		InStack.LayerHash.Add(Param.GetHash(), ParamIndex);
 	}
 }
 
 // Current stack assigned to this thread
 static thread_local TWeakPtr<FParamStack> GWeakStack;
 
-// Stacks that are asscociated with objects, pending execution of an object's tick function
+// Stacks that are associated with objects, pending execution of an object's tick function
 static TMap<TObjectKey<UObject>, TWeakPtr<FParamStack>> GPendingObjects;
 static FRWLock GPendingObjectsLock;
 
 FParamStack::FParamStack()
+	: LayerHash(1024)
 {
 	Layers.Reserve(8);
-	LayerIndices.Reserve(FParamId::GetMaxParamId().ToInt());
-	PreviousLayerIndices.Reserve(FParamId::GetMaxParamId().ToInt());
+	EntryStack.Reserve(32);
 	OwnedStackLayers.Reserve(8);
 }
 
@@ -306,37 +218,53 @@ void FParamStack::RemoveForPendingObject(const UObject* InObject)
 	GPendingObjects.Remove(InObject);
 }
 
-bool FParamStack::AttachToCurrentThreadForPendingObject(const UObject* InObject)
+bool FParamStack::AttachToCurrentThreadForPendingObject(const UObject* InObject, ECoalesce InCoalesce)
 {
 	FRWScopeLock ScopeLock(GPendingObjectsLock, SLT_ReadOnly);
 	if (TWeakPtr<FParamStack>* PendingStack = GPendingObjects.Find(InObject))
 	{
-		GWeakStack = *PendingStack;
+		AttachToCurrentThread(*PendingStack, InCoalesce);
 		return true;
 	}
 
 	return false;
 }
 
-bool FParamStack::DetachFromCurrentThreadForPendingObject(const UObject* InObject)
+bool FParamStack::DetachFromCurrentThreadForPendingObject(const UObject* InObject, EDecoalesce InDecoalesce)
 {
 	FRWScopeLock ScopeLock(GPendingObjectsLock, SLT_ReadOnly);
 	if (TWeakPtr<FParamStack>* PendingStack = GPendingObjects.Find(InObject))
 	{
-		GWeakStack.Reset();
+		DetachFromCurrentThread(InDecoalesce);
 		return true;
 	}
 
 	return false;
 }
 
-void FParamStack::AttachToCurrentThread(TWeakPtr<FParamStack> InStack)
+void FParamStack::AttachToCurrentThread(TWeakPtr<FParamStack> InStack, ECoalesce InCoalesce)
 {
 	GWeakStack = InStack;
+	
+	if(InCoalesce == ECoalesce::Coalesce)
+	{
+		if(TSharedPtr<FParamStack> PinnedStack = GWeakStack.Pin())
+		{
+			PinnedStack->Coalesce();
+		}
+	}
 }
 
-TWeakPtr<FParamStack> FParamStack::DetachFromCurrentThread()
+TWeakPtr<FParamStack> FParamStack::DetachFromCurrentThread(EDecoalesce InDecoalesce)
 {
+	if(InDecoalesce == EDecoalesce::Decoalesce)
+	{
+		if(TSharedPtr<FParamStack> PinnedStack = GWeakStack.Pin())
+		{
+			PinnedStack->Decoalesce();
+		}
+	}
+	
 	TWeakPtr<FParamStack> Stack = GWeakStack;
 	GWeakStack.Reset();
 	return Stack;
@@ -347,42 +275,20 @@ FParamStack::FPushedLayerHandle FParamStack::PushLayer(const FParamStackLayerHan
 	return PushLayerInternal(*InLayerHandle.Layer.Get());
 }
 
-FParamStack::FPushedLayerHandle FParamStack::PushLayer(TConstArrayView<TPair<FParamId, Private::FParamEntry>> InParams)
+FParamStack::FPushedLayerHandle FParamStack::PushLayer(TConstArrayView<Private::FParamEntry> InParams)
 {
-	if(Layers.Num() < MAX_uint16)
-	{
-		FParamStackLayer& OwnedLayer = OwnedStackLayers.Emplace_GetRef(InParams);
-		OwnedLayer.OwnedStorageOffset = AllocAndCopyOwnedParamStorage(OwnedLayer.Params);
-		return PushLayerInternal(OwnedLayer);
-	}
-	else
-	{
-		UE_LOG(LogAnimation, Warning, TEXT("FParamStack: Could not push a layer: Maximum 65535 stack layers."));
-		return FPushedLayerHandle();
-	}
+	FParamStackLayer& OwnedLayer = OwnedStackLayers.Emplace_GetRef(InParams);
+	OwnedLayer.OwnedStorageOffset = AllocAndCopyOwnedParamStorage(OwnedLayer.Params);
+	OwnedLayer.OwningStack = this;
+	return PushLayerInternal(OwnedLayer);
 }
 
 FParamStack::FPushedLayerHandle FParamStack::PushLayerInternal(FParamStackLayer& InLayer)
 {
 	if(InLayer.Params.Num() > 0)
 	{
-		if (Layers.Num() < MAX_uint16)
-		{
-			if (Layers.Num())
-			{
-				FPushedLayer& NewPushedLayer = Layers.Emplace_GetRef(Layers.Top(), InLayer, *this);
-				return FPushedLayerHandle(Layers.Num() - 1, NewPushedLayer.SerialNumber);
-			}
-			else
-			{
-				FPushedLayer& NewPushedLayer = Layers.Emplace_GetRef(InLayer, *this);
-				return FPushedLayerHandle(Layers.Num() - 1, NewPushedLayer.SerialNumber);
-			}
-		}
-		else
-		{
-			UE_LOG(LogAnimation, Warning, TEXT("FParamStack: Could not push a layer: Maximum 65535 stack layers."));
-		}
+		FPushedLayer& NewPushedLayer = Layers.Emplace_GetRef(InLayer, *this);
+		return FPushedLayerHandle(Layers.Num() - 1, NewPushedLayer.SerialNumber);
 	}
 
 	return FPushedLayerHandle();
@@ -398,33 +304,22 @@ void FParamStack::PopLayer(FPushedLayerHandle InHandle)
 			TEXT("UE::AnimNext::FParamStack::PopLayer: Invalid layer handle supplied (Have: %u, %u, Expected: %u, %u)"), 
 			InHandle.Index, InHandle.SerialNumber, (uint32)Layers.Num() - 1, TopLayer.SerialNumber);
 
-		// Fixup layer indices to previous, if any
-		const uint32 NumParams = TopLayer.Layer.Params.Num();
-		if(NumParams > 0 && PreviousLayerIndices.Num() > 0)
-		{
-			for (uint32 LocalParamIndex = 0; LocalParamIndex < (uint32)NumParams; ++LocalParamIndex)
-			{
-				const uint16 PreviousLayerIndex = PreviousLayerIndices[TopLayer.PreviousLayerIndexStart + LocalParamIndex];
-				if (PreviousLayerIndex != MAX_uint16)
-				{
-					const uint32 GlobalParamIndex = TopLayer.Layer.MinParamId + LocalParamIndex;
-					LayerIndices[GlobalParamIndex] = PreviousLayerIndex;
-				}
-			}
+		// Dont shrink allocs to avoid thrashing
+		constexpr bool bAllowShrinking = false;
 
-			PreviousLayerIndices.SetNum(PreviousLayerIndices.Num() - NumParams);
-		}
-		else
+		// Remove params and indices from hash table
+		const int32 NumLayerParams = TopLayer.Layer.Params.Num();
+		int32 ParamStackIndex = EntryStack.Num() - 1;
+		for(int32 ParamLayerIndex = 0; ParamLayerIndex < NumLayerParams; ++ParamLayerIndex)
 		{
-			// Clear layer indices as this is the last layer to be popped
-			FMemory::Memset(&LayerIndices[0], 0xff, LayerIndices.Num() * LayerIndices.GetTypeSize());
+			LayerHash.Remove(EntryStack[ParamStackIndex]->GetHash(), ParamStackIndex);
+			--ParamStackIndex;
 		}
 
-		// Dont shrink allocs to avoid thrashing
-		constexpr bool bAllowShrinking = false;
+		EntryStack.RemoveAt(EntryStack.Num() - NumLayerParams, NumLayerParams, bAllowShrinking);
 
 		// If we own the layer, pop the owned stack
-		if(TopLayer.Layer.OwnedStorageOffset != MAX_uint32)
+		if(TopLayer.Layer.OwnedStorageOffset != MAX_uint32 && TopLayer.Layer.OwningStack == this)
 		{
 			check(OwnedStackLayers.Num() && &TopLayer.Layer == &OwnedStackLayers[OwnedStackLayers.Num() - 1]);
 			OwnedStackLayers.Pop(bAllowShrinking);
@@ -474,7 +369,7 @@ FParamStackLayerHandle FParamStack::MakeRemappedLayer(const FParamStackLayerHand
 	return FParamStackLayerHandle(MoveTemp(Layer));
 }
 
-FParamStackLayerHandle FParamStack::MakeLayer(TConstArrayView<TPair<FParamId, Private::FParamEntry>> InParams)
+FParamStackLayerHandle FParamStack::MakeLayer(TConstArrayView<Private::FParamEntry> InParams)
 {
 	TUniquePtr<FParamStackLayer> Layer = MakeUnique<FParamStackLayer>(InParams);
 	return FParamStackLayerHandle(MoveTemp(Layer));
@@ -488,18 +383,27 @@ FParamResult FParamStack::GetParamData(FParamId InId, FParamTypeHandle InTypeHan
 
 FParamResult FParamStack::GetParamData(FParamId InId, FParamTypeHandle InTypeHandle, TConstArrayView<uint8>& OutParamData, FParamTypeHandle& OutParamTypeHandle, FParamCompatibility InRequiredCompatibility) const
 {
+	SCOPE_CYCLE_COUNTER(STAT_AnimNext_ParamStack_GetParam);
+	
 	const FParamResult Result = GetParamDataInternal(InId, InTypeHandle, OutParamData, OutParamTypeHandle, InRequiredCompatibility);
 	if (Result.IsInScope())
 	{
 		return Result;
 	}
-	else if (TSharedPtr<const FParamStack> ParentStack = WeakParentStack.Pin())
+
+	if(!bIsCoalesced)
 	{
-		return ParentStack->GetParamData(InId, InTypeHandle, OutParamData, OutParamTypeHandle, InRequiredCompatibility);
+		if (TSharedPtr<const FParamStack> ParentStack = WeakParentStack.Pin())
+		{
+			return ParentStack->GetParamData(InId, InTypeHandle, OutParamData, OutParamTypeHandle, InRequiredCompatibility);
+		}
 	}
+	
 	// Check if we have a built in adapter to fall back on
-	else if(const FParamAdapter* Adapter = InId.GetAdapter(InId))
+	if(const FParamAdapter* Adapter = InId.GetAdapter(InId))
 	{
+		SCOPE_CYCLE_COUNTER(STAT_AnimNext_ParamStack_Adapter);
+		
 		return Adapter->GetParamData(InTypeHandle, OutParamData, OutParamTypeHandle, InRequiredCompatibility);
 	}
 
@@ -508,20 +412,18 @@ FParamResult FParamStack::GetParamData(FParamId InId, FParamTypeHandle InTypeHan
 
 FParamResult FParamStack::GetParamDataInternal(FParamId InId, FParamTypeHandle InTypeHandle, TConstArrayView<uint8>& OutParamData, FParamTypeHandle& OutParamTypeHandle, FParamCompatibility InRequiredCompatibility) const
 {
-	if (Layers.Num() == 0 || InId.ToInt() >= (uint32)LayerIndices.Num() || LayerIndices[InId.ToInt()] == MAX_uint16)
+	const Private::FParamEntry* ParamPtr = FindParam(InId);
+	if (ParamPtr == nullptr)
 	{
 		return EParamResult::NotInScope;
 	}
-
-	const FPushedLayer& Layer = Layers[LayerIndices[InId.ToInt()]];
-	const FParamResult Result = Layer.Layer.GetParamData(InId, InTypeHandle, OutParamData, OutParamTypeHandle, InRequiredCompatibility);
+	
+	const FParamResult Result = ParamPtr->GetParamData(InTypeHandle, OutParamData, OutParamTypeHandle, InRequiredCompatibility);
 	if(Result.IsInScope())
 	{
 		return Result;
 	}
 
-
-
 	return EParamResult::NotInScope;
 }
 
@@ -533,25 +435,34 @@ FParamResult FParamStack::GetMutableParamData(FParamId InId, FParamTypeHandle In
 
 FParamResult FParamStack::GetMutableParamData(FParamId InId, FParamTypeHandle InTypeHandle, TArrayView<uint8>& OutParamData, FParamTypeHandle& OutParamTypeHandle, FParamCompatibility InRequiredCompatibility)
 {
+	SCOPE_CYCLE_COUNTER(STAT_AnimNext_ParamStack_GetParam);
+
 	const FParamResult Result = GetMutableParamDataInternal(InId, InTypeHandle, OutParamData, OutParamTypeHandle, InRequiredCompatibility);
 	if (Result.IsInScope())
 	{
 		return Result;
 	}
-	else if (TSharedPtr<const FParamStack> ParentStack = WeakParentStack.Pin())
+
+	if(!bIsCoalesced)
 	{
-		// we use a dummy here because if the data is present in a parent, it must be immutable anyways and we will early out
-		TConstArrayView<uint8> ParamData;
-		FParamResult ParentResult = ParentStack->GetParamData(InId, InTypeHandle, ParamData, OutParamTypeHandle, InRequiredCompatibility);
-		if (ParentResult.IsInScope())
+		if (TSharedPtr<const FParamStack> ParentStack = WeakParentStack.Pin())
 		{
-			// Parent data is immutable
-			return ParentResult.Result & EParamResult::MutabilityError;
+			// we use a dummy here because if the data is present in a parent, it must be immutable anyways and we will early out
+			TConstArrayView<uint8> ParamData;
+			FParamResult ParentResult = ParentStack->GetParamData(InId, InTypeHandle, ParamData, OutParamTypeHandle, InRequiredCompatibility);
+			if (ParentResult.IsInScope())
+			{
+				// Parent data is immutable
+				return ParentResult.Result & EParamResult::MutabilityError;
+			}
 		}
 	}
+
 	// Check if we have a built in adapter to fall back on
-	else if(const FParamAdapter* Adapter = InId.GetAdapter(InId))
+	if(const FParamAdapter* Adapter = InId.GetAdapter(InId))
 	{
+		SCOPE_CYCLE_COUNTER(STAT_AnimNext_ParamStack_Adapter);
+		
 		return Adapter->GetMutableParamData(InTypeHandle, OutParamData, OutParamTypeHandle, InRequiredCompatibility);
 	}
 
@@ -560,13 +471,13 @@ FParamResult FParamStack::GetMutableParamData(FParamId InId, FParamTypeHandle In
 
 FParamResult FParamStack::GetMutableParamDataInternal(FParamId InId, FParamTypeHandle InTypeHandle, TArrayView<uint8>& OutParamData, FParamTypeHandle& OutParamTypeHandle, FParamCompatibility InRequiredCompatibility)
 {
-	if (Layers.Num() == 0 || InId.ToInt() >= (uint32)LayerIndices.Num() || LayerIndices[InId.ToInt()] == MAX_uint16)
+	Private::FParamEntry* ParamPtr = FindMutableParam(InId);
+	if (ParamPtr == nullptr)
 	{
 		return EParamResult::NotInScope;
 	}
-
-	FPushedLayer& Layer = Layers[LayerIndices[InId.ToInt()]];
-	const FParamResult Result =  Layer.Layer.GetMutableParamData(InId, InTypeHandle, OutParamData, OutParamTypeHandle, InRequiredCompatibility);
+	
+	const FParamResult Result = ParamPtr->GetMutableParamData(InTypeHandle, OutParamData, OutParamTypeHandle, InRequiredCompatibility);
 	if(Result.IsInScope())
 	{
 		return Result;
@@ -577,45 +488,31 @@ FParamResult FParamStack::GetMutableParamDataInternal(FParamId InId, FParamTypeH
 
 bool FParamStack::IsMutableParam(FParamId InId) const
 {
-	if (Layers.Num() == 0 || InId.ToInt() >= (uint32)LayerIndices.Num() || LayerIndices[InId.ToInt()] == MAX_uint16)
+	const Private::FParamEntry* ParamPtr = FindParam(InId);
+	if (ParamPtr == nullptr)
 	{
 		return false;
 	}
-
-	const FPushedLayer& Layer = Layers[LayerIndices[InId.ToInt()]];
-	const uint32 LocalParamIndex = InId.ToInt() - Layer.Layer.MinParamId;
-
-	const Private::FParamEntry& Param = Layer.Layer.Params[LocalParamIndex];
-
-	return Param.IsMutable();
+	
+	return ParamPtr->IsMutable();
 }
 
 bool FParamStack::IsReferenceParam(FParamId InId) const
 {
-	if (Layers.Num() == 0 || InId.ToInt() >= (uint32)LayerIndices.Num() || LayerIndices[InId.ToInt()] == MAX_uint16)
+	const Private::FParamEntry* ParamPtr = FindParam(InId);
+	if(ParamPtr == nullptr)
 	{
 		return false;
 	}
-
-	const FPushedLayer& Layer = Layers[LayerIndices[InId.ToInt()]];
-	const uint32 LocalParamIndex = InId.ToInt() - Layer.Layer.MinParamId;
-
-	const Private::FParamEntry& Param = Layer.Layer.Params[LocalParamIndex];
-
-	return Param.IsReference();
+	
+	return ParamPtr->IsReference();
 }
 
 bool FParamStack::LayerContainsParam(const FParamStackLayerHandle& InHandle, FName InKey)
 {
 	const FParamId ParamIdToFind(InKey);
-	const uint32 MinParamId = InHandle.Layer->MinParamId;
-	const uint32 NumParams = InHandle.Layer->Params.Num();
-	if (ParamIdToFind.ToInt() >= MinParamId && ParamIdToFind.ToInt() < MinParamId + NumParams)
-	{
-		const uint32 LocalParamIndex = ParamIdToFind.ToInt() - MinParamId;
-		return InHandle.Layer->Params[LocalParamIndex].IsValid();
-	}
-	return false;
+	const Private::FParamEntry* ParamPtr = InHandle.Layer->FindEntry(ParamIdToFind);
+	return ParamPtr != nullptr;
 }
 
 uint32 FParamStack::AllocAndCopyOwnedParamStorage(TArrayView<Private::FParamEntry> InParams)
@@ -672,26 +569,81 @@ void FParamStack::FreeOwnedParamStorage(uint32 InOffset)
 	OwnedLayerParamStorage.SetNum(InOffset, false);
 }
 
-void FParamStack::ResizeLayerIndices()
+uint32 FParamStack::MakeSerialNumber()
 {
-	const uint32 NumParams = FParamId::GetMaxParamId().ToInt();
-	const uint32 NumLayerIndices = LayerIndices.Num();
+	++SerialNumber;
+	if (SerialNumber == 0)
+	{
+		++SerialNumber;
+	}
+	return SerialNumber;
+}
 
-	if (NumParams > NumLayerIndices)
+const Private::FParamEntry* FParamStack::FindParam(FParamId InId) const
+{
+	for(uint32 Index = LayerHash.First(InId.GetHash()); LayerHash.IsValid(Index); Index = LayerHash.Next(Index))
 	{
-		LayerIndices.SetNum(NumParams);
-		FMemory::Memset(&LayerIndices[NumLayerIndices], 0xff, (NumParams - NumLayerIndices) * sizeof(uint16));
+		if (EntryStack[Index]->GetName() == InId.GetName())
+		{
+			return EntryStack[Index];
+		}
 	}
+	return nullptr;
+}
+	
+Private::FParamEntry* FParamStack::FindMutableParam(FParamId InId)
+{
+	return const_cast<Private::FParamEntry*>(FindParam(InId));
 }
 
-uint32 FParamStack::MakeSerialNumber()
+void FParamStack::Coalesce()
 {
-	++SerialNumber;
-	if (SerialNumber == 0)
+	SCOPE_CYCLE_COUNTER(STAT_AnimNext_ParamStack_Coalesce);
+
+	check(!bIsCoalesced);
+	check(Layers.IsEmpty() && EntryStack.IsEmpty());
+
+	// Find stacks to coalesce
+	TArray<const FParamStack*, TInlineAllocator<8>> Stacks;
+	if(WeakParentStack.IsValid())
 	{
-		++SerialNumber;
+		int32 NumCoalescedLayers = 0;
+		const FParamStack* RootStack = WeakParentStack.Pin().Get();
+		while(RootStack)
+		{
+			NumCoalescedLayers += RootStack->Layers.Num();
+			Stacks.Add(RootStack);
+			RootStack = RootStack->WeakParentStack.Pin().Get();
+		}
+
+		CoalesceLayerHandles.Reserve(NumCoalescedLayers);
+		for(int32 StackIndex = Stacks.Num() - 1; StackIndex >= 0; --StackIndex)
+		{
+			const FParamStack* Stack = Stacks[StackIndex];
+			for(const FPushedLayer& PushedLayer : Stack->Layers)
+			{
+				CoalesceLayerHandles.Add(PushLayerInternal(PushedLayer.Layer));
+			}
+		}
 	}
-	return SerialNumber;
+
+	bIsCoalesced = true;
+}
+
+void FParamStack::Decoalesce()
+{
+	SCOPE_CYCLE_COUNTER(STAT_AnimNext_ParamStack_Decoalesce);
+
+	check(bIsCoalesced);
+
+	for(int32 LayerIndex = CoalesceLayerHandles.Num() - 1; LayerIndex >= 0; --LayerIndex)
+	{
+		PopLayer(CoalesceLayerHandles[LayerIndex]);
+	}
+
+	CoalesceLayerHandles.Reset();
+
+	bIsCoalesced = false;
 }
 
 }
diff --git a/Engine/Plugins/Experimental/AnimNext/Source/AnimNext/Private/Param/ParamStackLayer.cpp b/Engine/Plugins/Experimental/AnimNext/Source/AnimNext/Private/Param/ParamStackLayer.cpp
--- a/Engine/Plugins/Experimental/AnimNext/Source/AnimNext/Private/Param/ParamStackLayer.cpp
+++ b/Engine/Plugins/Experimental/AnimNext/Source/AnimNext/Private/Param/ParamStackLayer.cpp
@@ -8,143 +8,86 @@
 namespace UE::AnimNext
 {
 
+FParamStackLayer::FParamStackLayer(uint32 InParamCount)
+	: HashTable(FMath::Max<uint32>(32u, FMath::RoundUpToPowerOfTwo(InParamCount + 1)), InParamCount)
+{
+}
+	
 FParamStackLayer::~FParamStackLayer()
 {
 }
 
-FParamStackLayer::FParamStackLayer(TConstArrayView<TPair<FParamId, Private::FParamEntry>> InParams)
+FParamStackLayer::FParamStackLayer(TConstArrayView<Private::FParamEntry> InParams)
+	: FParamStackLayer(InParams.Num())
 {
-	MinParamId = MAX_uint32;
-	uint32 MaxParamId = 0;
+	Params.Reserve(InParams.Num());
 	for (uint32 ParamIndex = 0; ParamIndex < static_cast<uint32>(InParams.Num()); ++ParamIndex)
 	{
-		MinParamId = FMath::Min(InParams[ParamIndex].Key.ToInt(), MinParamId);
-		MaxParamId = FMath::Max(InParams[ParamIndex].Key.ToInt(), MaxParamId);
-	}
-
-	if (MinParamId <= MaxParamId)
-	{
-		const uint32 ParamRangeSize = (MaxParamId - MinParamId) + 1;
-		Params.SetNumZeroed(ParamRangeSize);
-		for (uint32 ParamIndex = 0; ParamIndex < static_cast<uint32>(InParams.Num()); ++ParamIndex)
-		{
-			const TPair<FParamId, Private::FParamEntry>& Pair = InParams[ParamIndex];
-			const uint32 LocalParamIndex = Pair.Key.ToInt() - MinParamId;
-			Params[LocalParamIndex] = Pair.Value;
-		}
+		const Private::FParamEntry& Param = InParams[ParamIndex];
+		Params.Add(Param);
+		HashTable.Add(Param.GetHash(), ParamIndex);
 	}
 }
-
-
+	
 FParamResult FParamStackLayer::GetParamData(FParamId InId, FParamTypeHandle InTypeHandle, TConstArrayView<uint8>& OutParamData) const
 {
-	const uint32 LocalParamIndex = InId.ToInt() - MinParamId;
-
-	const Private::FParamEntry& Param = Params[LocalParamIndex];
-	if (!Param.IsValid())
+	const Private::FParamEntry* ParamPtr = FindEntry(InId);
+	if (ParamPtr == nullptr)
 	{
 		return EParamResult::NotInScope;
 	}
 
-	FParamCompatibility Compatibility = UE::AnimNext::FParamUtils::GetCompatibility(InTypeHandle, Param.GetTypeHandle());
-	if (!Compatibility.IsCompatible())
-	{
-		return EParamResult::TypeError;
-	}
-
-	OutParamData = Param.GetData();
-	return EParamResult::Success;
+	return ParamPtr->GetParamData(InTypeHandle, OutParamData);
 }
 
 FParamResult FParamStackLayer::GetParamData(FParamId InId, FParamTypeHandle InTypeHandle, TConstArrayView<uint8>& OutParamData, FParamTypeHandle& OutParamTypeHandle, FParamCompatibility InRequiredCompatibility) const
 {
-	const uint32 LocalParamIndex = InId.ToInt() - MinParamId;
-
-	if(!Params.IsValidIndex(LocalParamIndex))
-	{
-		return EParamResult::NotInScope;
-	}
-
-	const Private::FParamEntry& Param = Params[LocalParamIndex];
-	if (!Param.IsValid())
+	const Private::FParamEntry* ParamPtr = FindEntry(InId);
+	if (ParamPtr == nullptr)
 	{
 		return EParamResult::NotInScope;
 	}
 
-	OutParamTypeHandle = Param.GetTypeHandle();
-
-	FParamCompatibility Compatibility = UE::AnimNext::FParamUtils::GetCompatibility(InTypeHandle, OutParamTypeHandle);
-	if (Compatibility < InRequiredCompatibility)
-	{
-		return EParamResult::TypeError;
-	}
-
-	OutParamData = Param.GetData();
-
-	if (Compatibility == InRequiredCompatibility)
-	{
-		return EParamResult::Success | EParamResult::TypeCompatible;
-	}
-
-	return EParamResult::Success;
+	return ParamPtr->GetParamData(InTypeHandle, OutParamData, OutParamTypeHandle, InRequiredCompatibility);
 }
 
 FParamResult FParamStackLayer::GetMutableParamData(FParamId InId, FParamTypeHandle InTypeHandle, TArrayView<uint8>& OutParamData)
 {
-	const uint32 LocalParamIndex = InId.ToInt() - MinParamId;
-
-	Private::FParamEntry& Param = Params[LocalParamIndex];
-	if (!Param.IsValid())
+	Private::FParamEntry* ParamPtr = FindMutableEntry(InId);
+	if (ParamPtr == nullptr)
 	{
 		return EParamResult::NotInScope;
 	}
 
-	FParamResult AccessResult = EParamResult::Success;
-	FParamCompatibility Compatibility = UE::AnimNext::FParamUtils::GetCompatibility(InTypeHandle, Param.GetTypeHandle());
-	if (!Compatibility.IsCompatible())
-	{
-		AccessResult.Result |= EParamResult::TypeError;
-	}
-
-	AccessResult.Result |= !Param.IsMutable() ? EParamResult::MutabilityError : EParamResult::Success;
-	if (AccessResult.IsSuccessful())
-	{
-		OutParamData = Param.GetMutableData();
-	}
-
-	return AccessResult;
+	return ParamPtr->GetMutableParamData(InTypeHandle, OutParamData);
 }
 
 FParamResult FParamStackLayer::GetMutableParamData(FParamId InId, FParamTypeHandle InTypeHandle, TArrayView<uint8>& OutParamData, FParamTypeHandle& OutParamTypeHandle, FParamCompatibility InRequiredCompatibility)
 {
-	const uint32 LocalParamIndex = InId.ToInt() - MinParamId;
-
-	Private::FParamEntry& Param = Params[LocalParamIndex];
-	if (!Param.IsValid())
+	Private::FParamEntry* ParamPtr = FindMutableEntry(InId);
+	if (ParamPtr == nullptr)
 	{
 		return EParamResult::NotInScope;
 	}
 
-	OutParamTypeHandle = Param.GetTypeHandle();
-
-	FParamResult AccessResult = EParamResult::Success;
-	FParamCompatibility Compatibility = UE::AnimNext::FParamUtils::GetCompatibility(InTypeHandle, OutParamTypeHandle);
-	if (Compatibility < InRequiredCompatibility)
-	{
-		AccessResult.Result |= EParamResult::TypeError;
-	}
-	else if (Compatibility == InRequiredCompatibility)
-	{
-		AccessResult.Result |= EParamResult::TypeCompatible;
-	}
+	return ParamPtr->GetMutableParamData(InTypeHandle, OutParamData, OutParamTypeHandle, InRequiredCompatibility);
+}
 
-	AccessResult.Result |= !Param.IsMutable() ? EParamResult::MutabilityError : EParamResult::Success;
-	if (AccessResult.IsSuccessful())
+const Private::FParamEntry* FParamStackLayer::FindEntry(FParamId InId) const
+{
+	for(uint32 Index = HashTable.First(InId.GetHash()); HashTable.IsValid(Index); Index = HashTable.Next(Index))
 	{
-		OutParamData = Param.GetMutableData();
+		if(Params[Index].GetName() == InId.GetName())
+		{
+			return &Params[Index];
+		}
 	}
-
-	return AccessResult;
+	return nullptr;
 }
 
+Private::FParamEntry* FParamStackLayer::FindMutableEntry(FParamId InId) const
+{
+	return const_cast<Private::FParamEntry*>(FindEntry(InId));
+}
+	
 }
\ No newline at end of file
diff --git a/Engine/Plugins/Experimental/AnimNext/Source/AnimNext/Private/Param/ParamStackLayer.h b/Engine/Plugins/Experimental/AnimNext/Source/AnimNext/Private/Param/ParamStackLayer.h
--- a/Engine/Plugins/Experimental/AnimNext/Source/AnimNext/Private/Param/ParamStackLayer.h
+++ b/Engine/Plugins/Experimental/AnimNext/Source/AnimNext/Private/Param/ParamStackLayer.h
@@ -5,6 +5,7 @@
 #include "CoreMinimal.h"
 #include "Param/ParamId.h"
 #include "Param/ParamCompatibility.h"
+#include "Containers/HashTable.h"
 
 struct FInstancedPropertyBag;
 
@@ -12,6 +13,7 @@ namespace UE::AnimNext
 {
 	struct FParamResult;
 	struct FParamTypeHandle;
+	struct FParamStack;
 }
 
 namespace UE::AnimNext::Private
@@ -27,13 +29,15 @@ namespace UE::AnimNext
 // Memory ownership for params is variable depending on the layer's subclass
 struct FParamStackLayer
 {
-	FParamStackLayer() = default;
+	FParamStackLayer() = delete;
+
 	virtual ~FParamStackLayer();
 
-	friend struct FParamStack;
-	friend struct FParamStackLayerHandle;
+	// Constructor that reserves space for the specified parameters
+	explicit FParamStackLayer(uint32 InParamCount);
 
-	explicit FParamStackLayer(TConstArrayView<TPair<FParamId, Private::FParamEntry>> InParams);
+	// Constructor for set of user parameters
+	explicit FParamStackLayer(TConstArrayView<Private::FParamEntry> InParams);
 
 	ANIMNEXT_API FParamResult GetParamData(FParamId InId, FParamTypeHandle InTypeHandle, TConstArrayView<uint8>& OutParamData) const;
 	ANIMNEXT_API FParamResult GetParamData(FParamId InId, FParamTypeHandle InTypeHandle, TConstArrayView<uint8>& OutParamData, FParamTypeHandle& OutParamTypeHandle, FParamCompatibility InRequiredCompatibility = FParamCompatibility::Equal()) const;
@@ -44,11 +48,17 @@ struct FParamStackLayer
 	ANIMNEXT_API virtual UObject* AsUObject() { return nullptr; }
 	ANIMNEXT_API virtual FInstancedPropertyBag* AsInstancedPropertyBag() { return nullptr; }
 
+	const Private::FParamEntry* FindEntry(FParamId InId) const;
+	Private::FParamEntry* FindMutableEntry(FParamId InId) const;
+	
 	// Params that this layer supplies
 	TArray<Private::FParamEntry> Params;
 
-	// ID that the param indices start at. Maps the global param ID range into the range for this layer
-	uint32 MinParamId = 0;
+	// Hash table for param IDs
+	FHashTable HashTable;
+
+	// Owning stack for this layer if it is owned internally, otherwise nullptr
+	FParamStack* OwningStack = nullptr;
 
 	// Storage offset for this layer if it is owned internally, otherwise MAX_uint32
 	uint32 OwnedStorageOffset = MAX_uint32;
diff --git a/Engine/Plugins/Experimental/AnimNext/Source/AnimNext/Private/Param/ParamStackLayerHandle.cpp b/Engine/Plugins/Experimental/AnimNext/Source/AnimNext/Private/Param/ParamStackLayerHandle.cpp
--- a/Engine/Plugins/Experimental/AnimNext/Source/AnimNext/Private/Param/ParamStackLayerHandle.cpp
+++ b/Engine/Plugins/Experimental/AnimNext/Source/AnimNext/Private/Param/ParamStackLayerHandle.cpp
@@ -65,22 +65,18 @@ FParamResult FParamStackLayerHandle::GetValueRaw(FParamId InParamId, FParamTypeH
 	return Layer->GetParamData(InParamId, InTypeHandle, OutData);
 }
 
-FParamResult FParamStackLayerHandle::SetValuesInternal(TConstArrayView<TPair<FParamId, Private::FParamEntry>> InParams)
+FParamResult FParamStackLayerHandle::SetValuesInternal(TConstArrayView<Private::FParamEntry> InParams)
 {
 	check(Layer.IsValid());
 
-	const uint32 MinParamId = Layer->MinParamId;
-	for (const TPair<FParamId, Private::FParamEntry>& ParamPair : InParams)
+	for (const Private::FParamEntry& Param : InParams)
 	{
-		const uint32 LocalParamIndex = ParamPair.Key.ToInt() - MinParamId;
-		if(Layer->Params.IsValidIndex(LocalParamIndex))
+		if(Private::FParamEntry* ParamPtr = Layer->FindMutableEntry(Param.GetId()))
 		{
-			Private::FParamEntry& Param = Layer->Params[LocalParamIndex];
-			
-			check(Param.GetTypeHandle() == ParamPair.Value.GetTypeHandle());
-			check(Param.IsMutable());
+			check(ParamPtr->GetTypeHandle() == Param.GetTypeHandle());
+			check(ParamPtr->IsMutable());
 
-			FParamHelpers::Copy(ParamPair.Value.GetTypeHandle(), ParamPair.Value.GetData(), Param.GetMutableData());
+			FParamHelpers::Copy(ParamPtr->GetTypeHandle(), Param.GetData(), ParamPtr->GetMutableData());
 			return EParamResult::Success;
 		}
 	}
diff --git a/Engine/Plugins/Experimental/AnimNext/Source/AnimNext/Private/Param/Params.cpp b/Engine/Plugins/Experimental/AnimNext/Source/AnimNext/Private/Param/Params.cpp
--- a/Engine/Plugins/Experimental/AnimNext/Source/AnimNext/Private/Param/Params.cpp
+++ b/Engine/Plugins/Experimental/AnimNext/Source/AnimNext/Private/Param/Params.cpp
@@ -11,7 +11,7 @@ static TMap<FName, FParamDefinition> RegisteredParams;
 
 FParamDefinition FParams::RegisterBuiltInParameter(const FParamDefinition& InDefinition)
 {
-	return RegisteredParams.Add(InDefinition.Name, InDefinition);
+	return RegisteredParams.Add(InDefinition.GetName(), InDefinition);
 }
 
 void FParams::UnregisterBuiltInParameter(FName InName)
diff --git a/Engine/Plugins/Experimental/AnimNext/Source/AnimNext/Private/Param/RigVMDispatch_GetLayerParameter.cpp b/Engine/Plugins/Experimental/AnimNext/Source/AnimNext/Private/Param/RigVMDispatch_GetLayerParameter.cpp
--- a/Engine/Plugins/Experimental/AnimNext/Source/AnimNext/Private/Param/RigVMDispatch_GetLayerParameter.cpp
+++ b/Engine/Plugins/Experimental/AnimNext/Source/AnimNext/Private/Param/RigVMDispatch_GetLayerParameter.cpp
@@ -46,16 +46,6 @@ FString FRigVMDispatch_GetLayerParameter::GetArgumentMetaData(const FName& InArg
 
 	return Super::GetArgumentMetaData(InArgumentName, InMetaDataKey);
 }
-
-FString FRigVMDispatch_GetLayerParameter::GetArgumentDefaultValue(const FName& InArgumentName, TRigVMTypeIndex InTypeIndex) const
-{
-	if (InArgumentName == ParameterIdName)
-	{
-		return FString::Printf(TEXT("%u"), 0xffffffff);
-	}
-
-	return Super::GetArgumentDefaultValue(InArgumentName, InTypeIndex);
-}
 #endif
 
 const TArray<FRigVMTemplateArgumentInfo>& FRigVMDispatch_GetLayerParameter::GetArgumentInfos() const
@@ -97,10 +87,10 @@ void FRigVMDispatch_GetLayerParameter::Execute(FRigVMExtendedExecuteContext& InC
 	check(ValueProperty);
 	uint8* TargetDataPtr = Handles[1].GetData();
 
-	uint32& ParameterId = *(uint32*)Handles[2].GetData();
-	if (ParameterId == FParamId::InvalidIndex)
+	uint32& ParameterHash = *(uint32*)Handles[2].GetData();
+	if (ParameterHash == 0 && Parameter != NAME_None)
 	{
-		ParameterId = FParamId(Parameter).ToInt();
+		ParameterHash = GetTypeHash(Parameter);
 	}
 
 	uint32& TypeHandle = *(uint32*)Handles[3].GetData();
@@ -111,7 +101,7 @@ void FRigVMDispatch_GetLayerParameter::Execute(FRigVMExtendedExecuteContext& InC
 
 	FAnimNextParameterExecuteContext& ParamContext = InContext.GetPublicData<FAnimNextParameterExecuteContext>();
 	TConstArrayView<uint8> SourceData;
-	if (ParamContext.GetLayerHandle().GetValueRaw(FParamId(ParameterId), FParamTypeHandle::FromRaw(TypeHandle), SourceData).IsSuccessful())
+	if (ParamContext.GetLayerHandle().GetValueRaw(FParamId(Parameter, ParameterHash), FParamTypeHandle::FromRaw(TypeHandle), SourceData).IsSuccessful())
 	{
 		ValueProperty->CopyCompleteValue(TargetDataPtr, SourceData.GetData());
 	}
diff --git a/Engine/Plugins/Experimental/AnimNext/Source/AnimNext/Private/Param/RigVMDispatch_GetParameter.cpp b/Engine/Plugins/Experimental/AnimNext/Source/AnimNext/Private/Param/RigVMDispatch_GetParameter.cpp
--- a/Engine/Plugins/Experimental/AnimNext/Source/AnimNext/Private/Param/RigVMDispatch_GetParameter.cpp
+++ b/Engine/Plugins/Experimental/AnimNext/Source/AnimNext/Private/Param/RigVMDispatch_GetParameter.cpp
@@ -48,16 +48,6 @@ FString FRigVMDispatch_GetParameter::GetArgumentMetaData(const FName& InArgument
 
 	return Super::GetArgumentMetaData(InArgumentName, InMetaDataKey);
 }
-
-FString FRigVMDispatch_GetParameter::GetArgumentDefaultValue(const FName& InArgumentName, TRigVMTypeIndex InTypeIndex) const
-{
-	if (InArgumentName == ParameterIdName)
-	{
-		return FString::Printf(TEXT("%u"), 0xffffffff);
-	}
-
-	return Super::GetArgumentDefaultValue(InArgumentName, InTypeIndex);
-}
 #endif
 
 const TArray<FRigVMTemplateArgumentInfo>& FRigVMDispatch_GetParameter::GetArgumentInfos() const
@@ -99,10 +89,10 @@ void FRigVMDispatch_GetParameter::Execute(FRigVMExtendedExecuteContext& InContex
 	check(ValueProperty);
 	uint8* TargetDataPtr = Handles[1].GetData();
 
-	uint32& ParameterId = *(uint32*)Handles[2].GetData();
-	if (ParameterId == FParamId::InvalidIndex)
+	uint32& ParameterHash = *(uint32*)Handles[2].GetData();
+	if (ParameterHash == 0 && Parameter != NAME_None)
 	{
-		ParameterId = FParamId(Parameter).ToInt();
+		ParameterHash = GetTypeHash(Parameter);
 	}
 
 	uint32& TypeHandle = *(uint32*)Handles[3].GetData();
@@ -112,7 +102,7 @@ void FRigVMDispatch_GetParameter::Execute(FRigVMExtendedExecuteContext& InContex
 	}
 
 	TConstArrayView<uint8> SourceData;
-	if (FParamStack::Get().GetParamData(FParamId(ParameterId), FParamTypeHandle::FromRaw(TypeHandle), SourceData).IsSuccessful())
+	if (FParamStack::Get().GetParamData(FParamId(Parameter, ParameterHash), FParamTypeHandle::FromRaw(TypeHandle), SourceData).IsSuccessful())
 	{
 		ValueProperty->CopyCompleteValue(TargetDataPtr, SourceData.GetData());
 	}
diff --git a/Engine/Plugins/Experimental/AnimNext/Source/AnimNext/Private/Param/RigVMDispatch_SetLayerParameter.cpp b/Engine/Plugins/Experimental/AnimNext/Source/AnimNext/Private/Param/RigVMDispatch_SetLayerParameter.cpp
--- a/Engine/Plugins/Experimental/AnimNext/Source/AnimNext/Private/Param/RigVMDispatch_SetLayerParameter.cpp
+++ b/Engine/Plugins/Experimental/AnimNext/Source/AnimNext/Private/Param/RigVMDispatch_SetLayerParameter.cpp
@@ -49,16 +49,6 @@ FString FRigVMDispatch_SetLayerParameter::GetArgumentMetaData(const FName& InArg
 
 	return Super::GetArgumentMetaData(InArgumentName, InMetaDataKey);
 }
-
-FString FRigVMDispatch_SetLayerParameter::GetArgumentDefaultValue(const FName& InArgumentName, TRigVMTypeIndex InTypeIndex) const
-{
-	if (InArgumentName == ParameterIdName)
-	{
-		return FString::Printf(TEXT("%u"), 0xffffffff);
-	}
-
-	return Super::GetArgumentDefaultValue(InArgumentName, InTypeIndex);
-}
 #endif
 
 const TArray<FRigVMTemplateArgumentInfo>& FRigVMDispatch_SetLayerParameter::GetArgumentInfos() const
@@ -106,10 +96,10 @@ void FRigVMDispatch_SetLayerParameter::Execute(FRigVMExtendedExecuteContext& InC
 	check(ValueProperty);
 	const uint8* SourceData =  Handles[1].GetData();
 
-	uint32& RawParameterId = *(uint32*)Handles[2].GetData();
-	if (RawParameterId == FParamId::InvalidIndex && Parameter != NAME_None)
+	uint32& ParameterHash = *(uint32*)Handles[2].GetData();
+	if (ParameterHash == 0 && Parameter != NAME_None)
 	{
-		RawParameterId = FParamId(Parameter).ToInt();
+		ParameterHash = GetTypeHash(Parameter);
 	}
 
 	uint32& RawTypeHandle = *(uint32*)Handles[3].GetData();
@@ -118,7 +108,7 @@ void FRigVMDispatch_SetLayerParameter::Execute(FRigVMExtendedExecuteContext& InC
 		RawTypeHandle = FParamTypeHandle::FromProperty(ValueProperty).ToRaw();
 	}
 
-	FParamId ParameterId(RawParameterId);
+	FParamId ParameterId(Parameter, ParameterHash);
 	FParamTypeHandle TypeHandle = FParamTypeHandle::FromRaw(RawTypeHandle);
 	if(ParameterId.IsValid() && TypeHandle.IsValid())
 	{
diff --git a/Engine/Plugins/Experimental/AnimNext/Source/AnimNext/Private/Scheduler/ScheduleTickFunction.cpp b/Engine/Plugins/Experimental/AnimNext/Source/AnimNext/Private/Scheduler/ScheduleTickFunction.cpp
--- a/Engine/Plugins/Experimental/AnimNext/Source/AnimNext/Private/Scheduler/ScheduleTickFunction.cpp
+++ b/Engine/Plugins/Experimental/AnimNext/Source/AnimNext/Private/Scheduler/ScheduleTickFunction.cpp
@@ -120,11 +120,11 @@ void FScheduleTickFunction::RunSchedule(TConstArrayView<FAnimNextScheduleInstruc
 			{
 				uint32 TaskIndex = Instruction.Operand;
 				FScheduleInstanceData& InstanceData = ScheduleContext.GetInstanceData();
-				FParamStack::AttachToCurrentThread(InstanceData.GetParamStack(Schedule->Tasks[TaskIndex].ParamScopeIndex));
+				FParamStack::AttachToCurrentThread(InstanceData.GetParamStack(Schedule->Tasks[TaskIndex].ParamScopeIndex), FParamStack::ECoalesce::Coalesce);
 
 				Schedule->Tasks[TaskIndex].RunGraph(ScheduleContext);
 
-				FParamStack::DetachFromCurrentThread();
+				FParamStack::DetachFromCurrentThread(FParamStack::EDecoalesce::Decoalesce);
 				break;
 			}
 		case EAnimNextScheduleScheduleOpcode::BeginRunExternalTask:
@@ -146,8 +146,6 @@ void FScheduleTickFunction::RunSchedule(TConstArrayView<FAnimNextScheduleInstruc
 				{
 					if (const UObject* Object = InTargetObjects[InstructionIndex].Get())
 					{
-						FScheduleInstanceData& InstanceData = ScheduleContext.GetInstanceData();
-						uint32 ExternalTaskIndex = Instruction.Operand;
 						FParamStack::RemoveForPendingObject(Object);
 					}
 				}
@@ -170,7 +168,7 @@ void FScheduleTickFunction::RunSchedule(TConstArrayView<FAnimNextScheduleInstruc
 
 				uint32 ScopeEntryIndex = Instruction.Operand;
 				FScheduleInstanceData& InstanceData = ScheduleContext.GetInstanceData();
-				FParamStack::AttachToCurrentThread(InstanceData.GetParamStack(Schedule->ParamScopeEntryTasks[ScopeEntryIndex].ParamScopeIndex));
+				FParamStack::AttachToCurrentThread(InstanceData.GetParamStack(Schedule->ParamScopeEntryTasks[ScopeEntryIndex].ParamScopeIndex), FParamStack::ECoalesce::Coalesce);
 
 				Schedule->ParamScopeEntryTasks[ScopeEntryIndex].RunParamScopeEntry(ScheduleContext);
 
@@ -188,7 +186,7 @@ void FScheduleTickFunction::RunSchedule(TConstArrayView<FAnimNextScheduleInstruc
 
 				Schedule->ParamScopeExitTasks[ScopeExitIndex].RunParamScopeExit(ScheduleContext);
 
-				FParamStack::DetachFromCurrentThread();
+				FParamStack::DetachFromCurrentThread(FParamStack::EDecoalesce::Decoalesce);
 				break;
 			}
 		default:
diff --git a/Engine/Plugins/Experimental/AnimNext/Source/AnimNext/Public/AnimNextStats.h b/Engine/Plugins/Experimental/AnimNext/Source/AnimNext/Public/AnimNextStats.h
--- a/Engine/Plugins/Experimental/AnimNext/Source/AnimNext/Public/AnimNextStats.h
+++ b/Engine/Plugins/Experimental/AnimNext/Source/AnimNext/Public/AnimNextStats.h
@@ -28,6 +28,8 @@ DECLARE_CYCLE_STAT_EXTERN(TEXT("AnimNext: Param Type Handle Lock"), STAT_AnimNex
 DECLARE_CYCLE_STAT_EXTERN(TEXT("AnimNext: Param ID Lock"), STAT_AnimNext_ParamIdLock, STATGROUP_AnimNext, ANIMNEXT_API);
 DECLARE_CYCLE_STAT_EXTERN(TEXT("AnimNext: Param Stack Get Param"), STAT_AnimNext_ParamStack_GetParam, STATGROUP_AnimNext, ANIMNEXT_API);
 DECLARE_CYCLE_STAT_EXTERN(TEXT("AnimNext: Param Stack Adapter"), STAT_AnimNext_ParamStack_Adapter, STATGROUP_AnimNext, ANIMNEXT_API);
+DECLARE_CYCLE_STAT_EXTERN(TEXT("AnimNext: Param Stack Coalesce"), STAT_AnimNext_ParamStack_Coalesce, STATGROUP_AnimNext, ANIMNEXT_API);
+DECLARE_CYCLE_STAT_EXTERN(TEXT("AnimNext: Param Stack Decoalesce"), STAT_AnimNext_ParamStack_Decoalesce, STATGROUP_AnimNext, ANIMNEXT_API);
 DECLARE_CYCLE_STAT_EXTERN(TEXT("AnimNext: Run Param Block"), STAT_AnimNext_ParamBlock_UpdateLayer, STATGROUP_AnimNext, ANIMNEXT_API);
 
 DECLARE_CYCLE_STAT_EXTERN(TEXT("AnimNext: Generate Reference Pose"), STAT_AnimNext_GenerateReferencePose, STATGROUP_AnimNext, ANIMNEXT_API);
diff --git a/Engine/Plugins/Experimental/AnimNext/Source/AnimNext/Public/Param/ParamEntry.h b/Engine/Plugins/Experimental/AnimNext/Source/AnimNext/Public/Param/ParamEntry.h
--- a/Engine/Plugins/Experimental/AnimNext/Source/AnimNext/Public/Param/ParamEntry.h
+++ b/Engine/Plugins/Experimental/AnimNext/Source/AnimNext/Public/Param/ParamEntry.h
@@ -4,6 +4,7 @@
 
 #include "CoreMinimal.h"
 #include "Param/ParamTypeHandle.h"
+#include "Param/ParamId.h"
 
 class UAnimNextComponent;
 
@@ -14,6 +15,8 @@ namespace UE::AnimNext
 	struct FInstancedPropertyBagLayer;
 	struct FParamStackLayerHandle;
 	struct FRemappedLayer;
+	struct FParamResult;
+	struct FParamCompatibility;
 }
 
 namespace UE::AnimNext
@@ -52,8 +55,9 @@ private:
 	friend struct UE::AnimNext::FParamStackLayerHandle;
 	friend struct UE::AnimNext::FRemappedLayer;
 	friend class ::UAnimNextComponent;
+	template<typename ElementType, typename AllocatorType> friend class ::TArray;
 
-	ANIMNEXT_API FParamEntry(const FParamTypeHandle& InTypeHandle, TArrayView<uint8> InData, bool bInIsReference, bool bInIsMutable);
+	ANIMNEXT_API FParamEntry(const FParamId& InId, const FParamTypeHandle& InTypeHandle, TArrayView<uint8> InData, bool bInIsReference, bool bInIsMutable);
 
 	// Get the type handle of this param 
 	FParamTypeHandle GetTypeHandle() const { return TypeHandle; }
@@ -80,6 +84,10 @@ private:
 		}
 	}
 
+	FParamResult GetParamData(FParamTypeHandle InTypeHandle, TConstArrayView<uint8>& OutParamData) const;
+
+	FParamResult GetParamData(FParamTypeHandle InTypeHandle, TConstArrayView<uint8>& OutParamData, FParamTypeHandle& OutParamTypeHandle, FParamCompatibility InRequiredCompatibility) const;
+
 	// Get an mutable view of the parameter's data, returns an empty array view if this parameter is immutable
 	TArrayView<uint8> GetMutableData()
 	{
@@ -97,15 +105,40 @@ private:
 		return TArrayView<uint8>();
 	}
 
+	FParamResult GetMutableParamData(FParamTypeHandle InTypeHandle, TArrayView<uint8>& OutParamData);
+
+	FParamResult GetMutableParamData(FParamTypeHandle InTypeHandle, TArrayView<uint8>& OutParamData, FParamTypeHandle& OutParamTypeHandle, FParamCompatibility InRequiredCompatibility);
+
 	// Check whether this param represents valid data. Note that this doesn't check the type hande for validity.
 	bool IsValid() const
 	{
 		return Size > 0;
 	}
 
+	// Get the name of this parameter
+	FName GetName() const
+	{
+		return Id.GetName();
+	}
+
+	// Get the hash of this parameter
+	uint32 GetHash() const
+	{
+		return Id.GetHash();
+	}
+
+	// Get the Id of this parameter
+	const FParamId& GetId() const
+	{
+		return Id;
+	}
+	
 	// Raw ptr to the data, or the data itself if we have EFlags::Embedded
 	void* Data = nullptr;
 
+	// ID of the parameter
+	FParamId Id;
+	
 	// The type of the param
 	FParamTypeHandle TypeHandle;
 
diff --git a/Engine/Plugins/Experimental/AnimNext/Source/AnimNext/Public/Param/ParamId.h b/Engine/Plugins/Experimental/AnimNext/Source/AnimNext/Public/Param/ParamId.h
--- a/Engine/Plugins/Experimental/AnimNext/Source/AnimNext/Public/Param/ParamId.h
+++ b/Engine/Plugins/Experimental/AnimNext/Source/AnimNext/Public/Param/ParamId.h
@@ -25,7 +25,7 @@ namespace UE::AnimNext::Tests
 namespace UE::AnimNext
 {
 
-// Global identifier used to index into dense parameter arrays
+// Global identifier used to avoid re-hashing parameter names
 struct ANIMNEXT_API FParamId
 {
 	friend struct FParamStack;
@@ -39,40 +39,35 @@ struct ANIMNEXT_API FParamId
 	friend class FModule;
 	friend struct FRemappedLayer;
 
-	static constexpr uint32 InvalidIndex = MAX_uint32;
-
 	FParamId() = default;
 	FParamId(const FParamId& InName) = default;
 	FParamId& operator=(const FParamId& InName) = default;
 
-	// Make a parameter ID from an FName
+	// Make a parameter ID from an FName, generating the hash
 	explicit FParamId(FName InName);
 
-	// Get the index of this param
-	uint32 ToInt() const
+	// Make a parameter ID from a name and hash
+	explicit FParamId(FName InName, uint32 InHash);
+
+	// Get the name of this param
+	FName GetName() const
 	{
-		return ParameterIndex;
+		return Name;
 	}
 
-	// Get the name that this parameter was created from
-	FName ToName() const;
+	// Get the hash of this param
+	uint32 GetHash() const
+	{
+		return Hash;
+	}
 
 	// Check if this ID represents a valid parameter
 	bool IsValid() const
 	{
-		return ParameterIndex != InvalidIndex;
+		return Hash != 0;
 	}
 
 private:
-	// Make a parameter ID from an index (for internal use)
-	explicit FParamId(uint32 InParameterIndex)
-		: ParameterIndex(InParameterIndex)
-	{
-	}
-
-	// Make a new parameter ID (internal usage only)
-	static uint32 MakeParamId_NoLock(FName InName);
-
 	// Initialize the parameter ID system
 	static void Init();
 
@@ -87,14 +82,9 @@ private:
 
 	// Register all built-in adapters
 	static void RegisterBuiltInAdapters();
-	
+
 	// Refresh all adapters defined in config
 	static void RefreshConfigAdapters();
-	
-	// Get the maximum parameter ID that can exist at present
-	// Note that this can change via concurrent modifications when new parameters are created by 
-	// different threads
-	static FParamId GetMaxParamId();
 
 	// Get any built in adapter for the supplied parameter
 	static const FParamAdapter* GetAdapter(FParamId InId);
@@ -108,16 +98,13 @@ private:
 	{
 		return reinterpret_cast<ReturnType*>(GetScratchAreaForParamIdAdapter(InId));
 	}
-	
-#if WITH_DEV_AUTOMATION_TESTS
-	// Used to isolate global param IDs from automated tests
-	static void BeginTestSandbox();
-	static void EndTestSandbox();
-#endif
 
 private:
-	// Stable index
-	uint32 ParameterIndex = InvalidIndex;
+	// Parameter name
+	FName Name;
+
+	// Name hash
+	uint32 Hash = 0;
 };
 
 } // end namespace UE::AnimNext
diff --git a/Engine/Plugins/Experimental/AnimNext/Source/AnimNext/Public/Param/ParamStack.h b/Engine/Plugins/Experimental/AnimNext/Source/AnimNext/Public/Param/ParamStack.h
--- a/Engine/Plugins/Experimental/AnimNext/Source/AnimNext/Public/Param/ParamStack.h
+++ b/Engine/Plugins/Experimental/AnimNext/Source/AnimNext/Public/Param/ParamStack.h
@@ -15,6 +15,7 @@
 #include "ParamEntry.h"
 #include "Param/ParamResult.h"
 #include "AnimNextStats.h"
+#include "Containers/HashTable.h"
 
 class UAnimNextSchedule;
 class UAnimNextSchedulerWorldSubsystem;
@@ -85,14 +86,8 @@ private:
 	{
 		FPushedLayer(FParamStackLayer& InLayer, FParamStack& InStack);
 
-		FPushedLayer(const FPushedLayer& InPreviousLayer, FParamStackLayer& InLayer, FParamStack& InStack);
-
 		FParamStackLayer& Layer;
 
-		// Index offset for the previous layer for each active param.
-		// Offset into FParamStack::PreviousLayerIndices.
-		uint32 PreviousLayerIndexStart = MAX_uint32;
-
 		// Serial number used for identifying layers to pop
 		uint32 SerialNumber = 0;
 	};
@@ -191,7 +186,7 @@ public:
 	FPushedLayerHandle PushValues(Args&&... InValues)
 	{
 		constexpr int32 NumItems = sizeof...(InValues) / 2;
-		TArray<TPair<FParamId, Private::FParamEntry>, TInlineAllocator<NumItems>> ParamIdValues;
+		TArray<Private::FParamEntry, TInlineAllocator<NumItems>> ParamIdValues;
 		ParamIdValues.Reserve(NumItems);
 		return PushValuesHelper(ParamIdValues, Forward<Args>(InValues)...);
 	}
@@ -239,7 +234,7 @@ public:
 	static FParamStackLayerHandle MakeValuesLayer(Args&&... InValues)
 	{
 		constexpr int32 NumItems = sizeof...(InValues) / 2;
-		TArray<TPair<FParamId, Private::FParamEntry>, TInlineAllocator<NumItems>> ParamIdValues;
+		TArray<Private::FParamEntry, TInlineAllocator<NumItems>> ParamIdValues;
 		ParamIdValues.Reserve(NumItems);
 		return MakeValuesLayerHelper(ParamIdValues, Forward<Args>(InValues)...);
 	}
@@ -248,8 +243,6 @@ public:
 	// Asserts if the layer supplied is valid and not the top layer
 	ANIMNEXT_API void PopLayer(FPushedLayerHandle InLayer);
 
-
-
 	// Get a pointer to a parameter's value given a FParamId.
 	// @param	InParamId			Parameter ID to find the currently-pushed value for
 	// @param	OutResult			Optional ptr to aresult that will be filled if an error occurs
@@ -411,21 +404,44 @@ public:
 private:
 	ANIMNEXT_API FParamStack();
 
-	void SetParent(TWeakPtr<const FParamStack> InParent);
+	ANIMNEXT_API void SetParent(TWeakPtr<const FParamStack> InParent);
 
 	// Get the param stack that is attached to the current thread
 	static ANIMNEXT_API TWeakPtr<FParamStack> GetForCurrentThread();
 
+	// Behavior flag for coalescing parent stacks on thread attachment 
+	enum class ECoalesce
+	{
+		// Do nothing
+		None,
+
+		// Coalesce parent stacks into this stack
+		Coalesce
+	};
+
 	// Attach a param stack to the current thread
-	static ANIMNEXT_API void AttachToCurrentThread(TWeakPtr<FParamStack> InStack);
+	static ANIMNEXT_API void AttachToCurrentThread(TWeakPtr<FParamStack> InStack, ECoalesce InCoalesce = ECoalesce::None);
 
 	// Attach a param stack that is 'pending' for the specified object to the current thread.
 	// @return true if the stack was successfully attached
-	static ANIMNEXT_API bool AttachToCurrentThreadForPendingObject(const UObject* InObject);
+	static ANIMNEXT_API bool AttachToCurrentThreadForPendingObject(const UObject* InObject, ECoalesce InCoalesce = ECoalesce::None);
+
+	// Behavior flag for de-coalescing parent stacks on thread detachment
+	enum class EDecoalesce
+	{
+		// Do nothing
+		None,
 
+		// Remove coalesced parent stacks from this stack
+		Decoalesce
+	};
+
+	// Detach a param stack to the current thread
+	static ANIMNEXT_API TWeakPtr<FParamStack> DetachFromCurrentThread(EDecoalesce InDecoalesce = EDecoalesce::None);
+	
 	// Detach a param stack that is 'pending' for the specified object from the current thread.
 	// @return true if the stack was successfully detached
-	static ANIMNEXT_API bool DetachFromCurrentThreadForPendingObject(const UObject* InObject);
+	static ANIMNEXT_API bool DetachFromCurrentThreadForPendingObject(const UObject* InObject, EDecoalesce InDecoalesce = EDecoalesce::None);
 
 	// Adds a stack asscociated with a 'pending' object - that is, an object that is due to have some
 	// asscociated logic run that needs to access the relevant stack.
@@ -435,21 +451,19 @@ private:
 	// Removes a stack asscociated with a 'pending' object. Asserts if any existing stack still has outstanding references.
 	static ANIMNEXT_API void RemoveForPendingObject(const UObject* InObject);
 
-	// Detach a param stack to the current thread
-	static ANIMNEXT_API TWeakPtr<FParamStack> DetachFromCurrentThread();
-
 	// Create a cached parameter layer from set of params. Mutabilty is on a per-parameter basis.
-	static ANIMNEXT_API FParamStackLayerHandle MakeLayer(TConstArrayView<TPair<FParamId, Private::FParamEntry>> InParams);
+	static ANIMNEXT_API FParamStackLayerHandle MakeLayer(TConstArrayView<Private::FParamEntry> InParams);
 
 	// Push an internally-owned layer. Copies parameter data to internal storage.
-	ANIMNEXT_API FPushedLayerHandle PushLayer(TConstArrayView<TPair<FParamId, Private::FParamEntry>> InParams);
+	ANIMNEXT_API FPushedLayerHandle PushLayer(TConstArrayView<Private::FParamEntry> InParams);
 
 	// Recursive helper function for PushValues
 	template <uint32 NumItems, typename FirstType, typename SecondType, typename... OtherTypes>
-	FPushedLayerHandle PushValuesHelper(TArray<TPair<FParamId, Private::FParamEntry>, TInlineAllocator<NumItems>>& InArray, FirstType&& InFirst, SecondType&& InSecond, OtherTypes&&... InOthers)
+	FPushedLayerHandle PushValuesHelper(TArray<Private::FParamEntry, TInlineAllocator<NumItems>>& InArray, FirstType&& InFirst, SecondType&& InSecond, OtherTypes&&... InOthers)
 	{
-		InArray.Emplace(FParamId(InFirst),
+		InArray.Emplace(
 			Private::FParamEntry(
+				FParamId(InFirst),
 				FParamTypeHandle::GetHandle<std::remove_reference_t<SecondType>>(),
 				TArrayView<uint8>(const_cast<uint8*>(reinterpret_cast<const uint8*>(&InSecond)), sizeof(std::remove_reference_t<SecondType>)),
 				std::is_reference_v<SecondType>,
@@ -469,10 +483,11 @@ private:
 
 	// Recursive helper function for MakeValuesLayer
 	template <uint32 NumItems, typename FirstType, typename SecondType, typename... OtherTypes>
-	static FParamStackLayerHandle MakeValuesLayerHelper(TArray<TPair<FParamId, Private::FParamEntry>, TInlineAllocator<NumItems>>& InArray, FirstType&& InFirst, SecondType&& InSecond, OtherTypes&&... InOthers)
+	static FParamStackLayerHandle MakeValuesLayerHelper(TArray<Private::FParamEntry, TInlineAllocator<NumItems>>& InArray, FirstType&& InFirst, SecondType&& InSecond, OtherTypes&&... InOthers)
 	{
-		InArray.Emplace(FParamId(InFirst), 
+		InArray.Emplace( 
 			Private::FParamEntry(
+				FParamId(InFirst),
 				FParamTypeHandle::GetHandle<std::remove_reference_t<SecondType>>(),
 				TArrayView<uint8>(const_cast<uint8*>(reinterpret_cast<const uint8*>(&InSecond)), sizeof(std::remove_reference_t<SecondType>)),
 				std::is_reference_v<SecondType>,
@@ -527,15 +542,29 @@ private:
 	// Frees any owned storage, restores owned storage offset
 	void FreeOwnedParamStorage(uint32 InOffset);
 
-	// Resize layer indices to deal with any new params we have seen since the stack was created
-	void ResizeLayerIndices();
-
 	// Get a new serial number for a pushed layer
 	uint32 MakeSerialNumber();
 
+	// Find the topmost value for the specified parameter
+	const Private::FParamEntry* FindParam(FParamId InId) const;
+	Private::FParamEntry* FindMutableParam(FParamId InId);
+
+	// Combine this stack with its parents to allow a single hash query to discover parameters
+	// Requires the stack is empty before this is called
+	ANIMNEXT_API void Coalesce();
+
+	// Remove all combined parameters that were added in Coalesce()
+	ANIMNEXT_API void Decoalesce();
+
 	// Layer stack
 	TArray<FPushedLayer> Layers;
 
+	// Stack of entry ptrs, pointing to Layers, indexed by LayerHash
+	TArray<Private::FParamEntry*> EntryStack;
+
+	// Hash table for the current layer index per-parameter. Indexes EntryStack.
+	FHashTable LayerHash;
+
 	// Owned layers, paged for stable addresses as layers reference them by ptr
 	// Grows with each new pushed owned layer
 	TPagedArray<FParamStackLayer, 4096> OwnedStackLayers;
@@ -544,20 +573,17 @@ private:
 	// Grows with each new pushed layer, free'd when the param stack is destroyed
 	TPagedArray<uint8, 4096> OwnedLayerParamStorage;
 
-	// Indicies into layers. These are the indices into the layer stack where the current 'top' value resides.
-	// Grows on creation and when new parameter IDs are encountered
-	TArray<uint16> LayerIndices;
-
-	// Indices for the previous layer for each active param. MAX_uint16 if there is no previous layer for the param.
-	// Pushed layers hold views into this array.
-	// Grows with each new pushed layer, free'd when the param stack is destroyed
-	TArray<uint16> PreviousLayerIndices;
-
 	// Parent stack from outer scope
 	TWeakPtr<const FParamStack> WeakParentStack;
 
+	// Coalesced layer handles
+	TArray<FPushedLayerHandle> CoalesceLayerHandles;
+
 	// Serial number for stack layers
 	uint32 SerialNumber = 0;
+
+	// Flag to indicate coalesced query path
+	bool bIsCoalesced = false;
 };
 
 }
diff --git a/Engine/Plugins/Experimental/AnimNext/Source/AnimNext/Public/Param/ParamStackLayerHandle.h b/Engine/Plugins/Experimental/AnimNext/Source/AnimNext/Public/Param/ParamStackLayerHandle.h
--- a/Engine/Plugins/Experimental/AnimNext/Source/AnimNext/Public/Param/ParamStackLayerHandle.h
+++ b/Engine/Plugins/Experimental/AnimNext/Source/AnimNext/Public/Param/ParamStackLayerHandle.h
@@ -75,7 +75,7 @@ public:
 	FParamResult SetValues(Args&&... InValues)
 	{
 		constexpr int32 NumItems = sizeof...(InValues) / 2;
-		TArray<TPair<FParamId, Private::FParamEntry>, TInlineAllocator<NumItems>> ParamIdValues;
+		TArray<Private::FParamEntry, TInlineAllocator<NumItems>> ParamIdValues;
 		ParamIdValues.Reserve(NumItems);
 		return SetValuesHelper(ParamIdValues, Forward<Args>(InValues)...);
 	}
@@ -170,7 +170,7 @@ private:
 	friend class UAnimNextParameterBlock;
 
 	// Set parameter values
-	ANIMNEXT_API FParamResult SetValuesInternal(TConstArrayView<TPair<FParamId, Private::FParamEntry>> InParams);
+	ANIMNEXT_API FParamResult SetValuesInternal(TConstArrayView<Private::FParamEntry> InParams);
 
 	// Helpers for As() to hide Layer
 	ANIMNEXT_API UObject* GetUObjectFromLayer() const;
@@ -178,10 +178,11 @@ private:
 
 	// Recursive helper function for SetValues
 	template <uint32 NumItems, typename FirstType, typename SecondType, typename... OtherTypes>
-	FParamResult SetValuesHelper(TArray<TPair<FParamId, Private::FParamEntry>, TInlineAllocator<NumItems>>& InArray, FirstType&& InFirst, SecondType&& InSecond, OtherTypes&&... InOthers)
+	FParamResult SetValuesHelper(TArray<Private::FParamEntry, TInlineAllocator<NumItems>>& InArray, FirstType&& InFirst, SecondType&& InSecond, OtherTypes&&... InOthers)
 	{
-		InArray.Emplace(FParamId(InFirst),
+		InArray.Emplace(
 			Private::FParamEntry(
+				FParamId(InFirst),
 				FParamTypeHandle::GetHandle<std::remove_reference_t<SecondType>>(),
 				TArrayView<uint8>(const_cast<uint8*>(reinterpret_cast<const uint8*>(&InSecond)), sizeof(std::remove_reference_t<SecondType>)),
 				std::is_reference_v<SecondType>,
diff --git a/Engine/Plugins/Experimental/AnimNext/Source/AnimNextTestSuite/Private/AnimNextParametersTest.cpp b/Engine/Plugins/Experimental/AnimNext/Source/AnimNextTestSuite/Private/AnimNextParametersTest.cpp
--- a/Engine/Plugins/Experimental/AnimNext/Source/AnimNextTestSuite/Private/AnimNextParametersTest.cpp
+++ b/Engine/Plugins/Experimental/AnimNext/Source/AnimNextTestSuite/Private/AnimNextParametersTest.cpp
@@ -294,8 +294,6 @@ IMPLEMENT_SIMPLE_AUTOMATION_TEST(FParamStackTest, "Animation.AnimNext.Parameters
 
 bool FParamStackTest::RunTest(const FString& InParameters)
 {
-	FParamId::BeginTestSandbox();
-
 	FParamId ParamIds[] = { FParamId("Param0"), FParamId("Param1"), FParamId("Param2"), FParamId("Param3") };
 
 	// Check immediate value interfaces
@@ -548,17 +546,40 @@ bool FParamStackTest::RunTest(const FString& InParameters)
 			AddErrorIfFalse(Stack.GetParam<float>(ParamName) == (float)Index, FString::Printf(TEXT("%s != %.0f"), *ParamName.ToString(), (float)Index));
 		}
 
-		Stack.PushValue(FName("Param", MAX_uint16), (float)MAX_uint16);
-		AddErrorIfFalse(Stack.GetParamPtr<float>(FName("Param", MAX_uint16)) == nullptr, TEXT("Overflowed stack value is non-null"));
-
 		for (int32 Index = NumParams - 1; Index >= 0; --Index)
 		{
 			Stack.PopLayer(Handles[Index]);
 		}
 	}
 
-	FParamId::EndTestSandbox();
+	// Check parent links & coalescing of parent stacks
+	{
+		TSharedPtr<FParamStack> Stack1 = MakeShared<FParamStack>();
+		TSharedPtr<FParamStack> Stack2 = MakeShared<FParamStack>();
+		Stack2->SetParent(Stack1);
+
+		FParamStack::FPushedLayerHandle Handle1 = Stack1->PushValues("Param2", 3.0f, "Param1", 3.0f);
+
+		AddErrorIfFalse(Stack2->GetParam<float>("Param1") == 3.0f, "Decoalesced Param1 != 3.0f");
+		AddErrorIfFalse(Stack2->GetParam<float>("Param2") == 3.0f, "Decoalesced Param2 != 3.0f");
+
+		Stack2->Coalesce();
 
+		FParamStack::FPushedLayerHandle Handle2 = Stack2->PushValues("Param3", 5.0f, "Param0", 6.0f);
+
+		AddErrorIfFalse(Stack2->GetParam<float>("Param1") == 3.0f, "Coalesced Param1 != 3.0f");
+		AddErrorIfFalse(Stack2->GetParam<float>("Param2") == 3.0f, "Coalesced Param2 != 3.0f");
+		AddErrorIfFalse(Stack2->GetParam<float>("Param3") == 5.0f, "Coalesced Param3 != 5.0f");
+		AddErrorIfFalse(Stack2->GetParam<float>("Param0") == 6.0f, "Coalesced Param0 != 6.0f");
+
+		Stack2->PopLayer(Handle2);
+
+		Stack2->Decoalesce();
+
+		AddErrorIfFalse(Stack2->GetParam<float>("Param1") == 3.0f, "Decoalesced Param1 != 3.0f");
+		AddErrorIfFalse(Stack2->GetParam<float>("Param2") == 3.0f, "Decoalesced Param2 != 3.0f");
+	}
+	
 	return true;
 }
 
