Add sync group support to AnimNext through a decorator
Add support for custom AnimNext graph instance components

diff --git a/Engine/Plugins/Experimental/AnimNext/Source/AnimNext/Internal/Graph/AnimNextExecuteContext.h b/Engine/Plugins/Experimental/AnimNext/Source/AnimNext/Internal/Graph/AnimNextExecuteContext.h
--- a/Engine/Plugins/Experimental/AnimNext/Source/AnimNext/Internal/Graph/AnimNextExecuteContext.h
+++ b/Engine/Plugins/Experimental/AnimNext/Source/AnimNext/Internal/Graph/AnimNextExecuteContext.h
@@ -8,6 +8,7 @@
 #include "DecoratorBase/DecoratorPtr.h"
 #include "AnimNextExecuteContext.generated.h"
 
+struct FAnimNextGraphInstance;
 class UAnimNextGraph;
 
 namespace UE::AnimNext
@@ -44,15 +45,15 @@ struct FAnimNextExecuteContext : public FRigVMExecuteContext
 
 	const UE::AnimNext::FContext& GetContext() const
 	{
-		check(Context);
+		check(Context != nullptr);
 		return *Context;
 	}
 
-	const UAnimNextGraph* GetGraph() const { return Graph; }
-
-	const TArrayView<const uint8>& GetSharedDataBuffer() const { return SharedDataBuffer; }
-
-	UE::AnimNext::FWeakDecoratorPtr GetGraphInstancePtr() const { return GraphInstancePtr; }
+	FAnimNextGraphInstance& GetGraphInstance() const
+	{
+		check(GraphInstance != nullptr);
+		return *GraphInstance;
+	}
 
 	EAnimNextGraphSimulationSteps GetSimulationSteps() const { return SimulationSteps; }
 
@@ -62,9 +63,7 @@ struct FAnimNextExecuteContext : public FRigVMExecuteContext
 
 		const FAnimNextExecuteContext* OtherContext = (const FAnimNextExecuteContext*)InOtherContext;
 		Context = OtherContext->Context;
-		Graph = OtherContext->Graph; 
-		SharedDataBuffer = OtherContext->SharedDataBuffer;
-		GraphInstancePtr = OtherContext->GraphInstancePtr;
+		GraphInstance = OtherContext->GraphInstance;
 		SimulationSteps = OtherContext->SimulationSteps;
 	}
 
@@ -74,11 +73,9 @@ private:
 		Context = &InContext;
 	}
 
-	void InitializeWithGraph(const UAnimNextGraph* InGraph, TArrayView<const uint8> InSharedDataBuffer, UE::AnimNext::FWeakDecoratorPtr InGraphInstancePtr)
+	void SetGraphInstance(FAnimNextGraphInstance& InGraphInstance)
 	{
-		Graph = InGraph;
-		SharedDataBuffer = InSharedDataBuffer;
-		GraphInstancePtr = InGraphInstancePtr;
+		GraphInstance = &InGraphInstance;
 	}
 
 	void SetSimulationSteps(EAnimNextGraphSimulationSteps InSimulationSteps)
@@ -90,16 +87,12 @@ private:
 	void DebugReset()
 	{
 		Context = nullptr;
-		Graph = nullptr;
-		SharedDataBuffer = TArrayView<const uint8>();
-		GraphInstancePtr.Reset();
+		GraphInstance = nullptr;
 		SimulationSteps = EAnimNextGraphSimulationSteps::None;
 	}
 
 	const UE::AnimNext::FContext* Context = nullptr;
-	const UAnimNextGraph* Graph = nullptr;
-	TArrayView<const uint8> SharedDataBuffer;
-	UE::AnimNext::FWeakDecoratorPtr GraphInstancePtr;
+	FAnimNextGraphInstance* GraphInstance = nullptr;
 	EAnimNextGraphSimulationSteps SimulationSteps = EAnimNextGraphSimulationSteps::None;
 
 	friend class UAnimNextGraph;
diff --git a/Engine/Plugins/Experimental/AnimNext/Source/AnimNext/Internal/Graph/AnimNextGraph.h b/Engine/Plugins/Experimental/AnimNext/Source/AnimNext/Internal/Graph/AnimNextGraph.h
--- a/Engine/Plugins/Experimental/AnimNext/Source/AnimNext/Internal/Graph/AnimNextGraph.h
+++ b/Engine/Plugins/Experimental/AnimNext/Source/AnimNext/Internal/Graph/AnimNextGraph.h
@@ -7,6 +7,7 @@
 #include "DecoratorBase/DecoratorPtr.h"
 #include "DecoratorBase/DecoratorHandle.h"
 #include "DecoratorBase/EntryPointHandle.h"
+#include "Graph/GraphInstanceComponent.h"
 #include "Graph/RigUnit_AnimNextGraphEvaluator.h"
 #include "Param/ParamId.h"
 #include "Scheduler/IAnimNextScheduleTermInterface.h"
@@ -18,11 +19,14 @@ class UEdGraph;
 class UAnimNextGraph;
 class UAnimGraphNode_AnimNextGraph;
 struct FAnimNode_AnimNextGraph;
+struct FRigUnit_AnimNextGraphEvaluator;
 enum class EAnimNextGraphSimulationSteps;
 
 namespace UE::AnimNext
 {
 	struct FContext;
+	struct FExecutionContext;
+	struct FTestUtils;
 }
 
 namespace UE::AnimNext::UncookedOnly
@@ -41,6 +45,8 @@ namespace UE::AnimNext::Graph
 	extern ANIMNEXT_API const FName ResultName;
 }
 
+using GraphInstanceComponentMapType = TMap<FName, TSharedPtr<UE::AnimNext::FGraphInstanceComponent>>;
+
 // Represents an instance of an AnimNext graph
 // This struct uses UE reflection because we wish for the GC to keep the graph
 // alive while we own a reference to it. It is not intended to be serialized on disk with a live instance.
@@ -74,10 +80,34 @@ struct ANIMNEXT_API FAnimNextGraphInstance
 	// Returns true if we have a live graph instance, false otherwise
 	bool IsValid() const;
 
+	// Returns the graph used by this instance or nullptr if the instance is invalid
+	const UAnimNextGraph* GetGraph() const;
+
 	// Check to see if this instance data matches the provided graph
 	bool UsesGraph(const UAnimNextGraph* InGraph) const;
 
+	// Returns a typed graph instance component, creating it lazily the first time it is queried
+	template<class ComponentType>
+	ComponentType& GetComponent();
+
+	// Returns a typed graph instance component pointer if found or nullptr otherwise
+	template<class ComponentType>
+	ComponentType* TryGetComponent();
+
+	// Returns a typed graph instance component pointer if found or nullptr otherwise
+	template<class ComponentType>
+	const ComponentType* TryGetComponent() const;
+
+	// Returns const iterators to the graph instance component container
+	GraphInstanceComponentMapType::TConstIterator GetComponentIterator() const;
+
 private:
+	// Returns a pointer to the specified component, or nullptr if not found
+	UE::AnimNext::FGraphInstanceComponent* TryGetComponent(int32 ComponentNameHash, FName ComponentName) const;
+
+	// Adds the specified component and returns a reference to it
+	UE::AnimNext::FGraphInstanceComponent& AddComponent(int32 ComponentNameHash, FName ComponentName, TSharedPtr<UE::AnimNext::FGraphInstanceComponent>&& Component);
+
 	// Hard reference to the graph used to create this instance to ensure we can release it safely
 	UPROPERTY()
 	TObjectPtr<const UAnimNextGraph> Graph;
@@ -89,8 +119,11 @@ private:
 	UPROPERTY()
 	FRigVMExtendedExecuteContext ExtendedExecuteContext;
 
-	// The graph is the one that allocates instances
-	friend UAnimNextGraph;
+	// Graph instance components that persist from update to update
+	GraphInstanceComponentMapType Components;
+
+	friend UAnimNextGraph;					// The graph is the one that allocates instances
+	friend FRigUnit_AnimNextGraphEvaluator;	// We evaluate the instance
 };
 
 // A user-created graph of logic used to supply data
@@ -137,8 +170,10 @@ protected:
 	friend class UAnimNextGraph_EditorData;
 	friend struct UE::AnimNext::UncookedOnly::FUtils;
 	friend class UE::AnimNext::Editor::FGraphEditor;
+	friend struct UE::AnimNext::FTestUtils;
 	friend struct FAnimNextGraphInstance;
 	friend class UAnimGraphNode_AnimNextGraph;
+	friend UE::AnimNext::FExecutionContext;
 
 #if WITH_EDITORONLY_DATA
 	mutable FRWLock GraphInstancesLock;
@@ -200,3 +235,37 @@ protected:
 	TArray<uint8> SharedDataArchiveBuffer;
 #endif
 };
+
+//////////////////////////////////////////////////////////////////////////
+
+template<class ComponentType>
+ComponentType& FAnimNextGraphInstance::GetComponent()
+{
+	const FName ComponentName = ComponentType::StaticComponentName();
+	const int32 ComponentNameHash = GetTypeHash(ComponentName);
+
+	if (UE::AnimNext::FGraphInstanceComponent* Component = TryGetComponent(ComponentNameHash, ComponentName))
+	{
+		return *static_cast<ComponentType*>(Component);
+	}
+
+	return static_cast<ComponentType&>(AddComponent(ComponentNameHash, ComponentName, MakeShared<ComponentType>()));
+}
+
+template<class ComponentType>
+ComponentType* FAnimNextGraphInstance::TryGetComponent()
+{
+	const FName ComponentName = ComponentType::StaticComponentName();
+	const int32 ComponentNameHash = GetTypeHash(ComponentName);
+
+	return *static_cast<ComponentType*>(TryGetComponent(ComponentNameHash, ComponentName));
+}
+
+template<class ComponentType>
+const ComponentType* FAnimNextGraphInstance::TryGetComponent() const
+{
+	const FName ComponentName = ComponentType::StaticComponentName();
+	const int32 ComponentNameHash = GetTypeHash(ComponentName);
+
+	return *static_cast<ComponentType*>(TryGetComponent(ComponentNameHash, ComponentName));
+}
diff --git a/Engine/Plugins/Experimental/AnimNext/Source/AnimNext/Private/DecoratorBase/ExecutionContext.cpp b/Engine/Plugins/Experimental/AnimNext/Source/AnimNext/Private/DecoratorBase/ExecutionContext.cpp
--- a/Engine/Plugins/Experimental/AnimNext/Source/AnimNext/Private/DecoratorBase/ExecutionContext.cpp
+++ b/Engine/Plugins/Experimental/AnimNext/Source/AnimNext/Private/DecoratorBase/ExecutionContext.cpp
@@ -18,22 +18,24 @@ namespace UE::AnimNext
 		static thread_local UE::AnimNext::FExecutionContext* GThreadLocalExecutionContext = nullptr;
 	}
 
-	FExecutionContext::FExecutionContext(TArrayView<const uint8> InGraphSharedData)
+	FExecutionContext::FExecutionContext(FAnimNextGraphInstance& InGraphInstance)
 		: NodeTemplateRegistry(FNodeTemplateRegistry::Get())
 		, DecoratorRegistry(FDecoratorRegistry::Get())
-		, GraphSharedData(InGraphSharedData)
-		, RigVMLatentMemoryHandles()
-		, RigVMExecuteContext(nullptr)
+		, Graph(InGraphInstance.GetGraph())
+		, GraphInstance(&InGraphInstance)
+		, GraphSharedData(Graph->SharedDataBuffer)
 	{
 		// There can be only one execution context alive per thread
 		ensure(Private::GThreadLocalExecutionContext == nullptr);
 		Private::GThreadLocalExecutionContext = this;
 	}
 
-	FExecutionContext::FExecutionContext(TArrayView<const uint8> InGraphSharedData, FRigVMExtendedExecuteContext& InRigVMExecuteContext, FRigVMMemoryHandleArray InRigVMLatentMemoryHandles)
+	FExecutionContext::FExecutionContext(FAnimNextGraphInstance& InGraphInstance, FRigVMExtendedExecuteContext& InRigVMExecuteContext, FRigVMMemoryHandleArray InRigVMLatentMemoryHandles)
 		: NodeTemplateRegistry(FNodeTemplateRegistry::Get())
 		, DecoratorRegistry(FDecoratorRegistry::Get())
-		, GraphSharedData(InGraphSharedData)
+		, Graph(InGraphInstance.GetGraph())
+		, GraphInstance(&InGraphInstance)
+		, GraphSharedData(Graph->SharedDataBuffer)
 		, RigVMLatentMemoryHandles(InRigVMLatentMemoryHandles)
 		, RigVMExecuteContext(&InRigVMExecuteContext)
 	{
diff --git a/Engine/Plugins/Experimental/AnimNext/Source/AnimNext/Private/DecoratorInterfaces/IGroupSynchronization.cpp b/Engine/Plugins/Experimental/AnimNext/Source/AnimNext/Private/DecoratorInterfaces/IGroupSynchronization.cpp
new file mode 100644
--- /dev/null
+++ b/Engine/Plugins/Experimental/AnimNext/Source/AnimNext/Private/DecoratorInterfaces/IGroupSynchronization.cpp
@@ -0,0 +1,50 @@
+// Copyright Epic Games, Inc. All Rights Reserved.
+
+#include "DecoratorInterfaces/IGroupSynchronization.h"
+
+#include "DecoratorBase/ExecutionContext.h"
+
+namespace UE::AnimNext
+{
+	FName IGroupSynchronization::GetGroupName(const FExecutionContext& Context, const TDecoratorBinding<IGroupSynchronization>& Binding) const
+	{
+		TDecoratorBinding<IGroupSynchronization> SuperBinding;
+		if (Context.GetInterfaceSuper(Binding, SuperBinding))
+		{
+			return SuperBinding.GetGroupName(Context);
+		}
+
+		return NAME_None;
+	}
+
+	EAnimGroupRole::Type IGroupSynchronization::GetGroupRole(const FExecutionContext& Context, const TDecoratorBinding<IGroupSynchronization>& Binding) const
+	{
+		TDecoratorBinding<IGroupSynchronization> SuperBinding;
+		if (Context.GetInterfaceSuper(Binding, SuperBinding))
+		{
+			return SuperBinding.GetGroupRole(Context);
+		}
+
+		return EAnimGroupRole::CanBeLeader;
+	}
+
+	float IGroupSynchronization::AdvanceBy(const FExecutionContext& Context, const TDecoratorBinding<IGroupSynchronization>& Binding, float DeltaTime) const
+	{
+		TDecoratorBinding<IGroupSynchronization> SuperBinding;
+		if (Context.GetInterfaceSuper(Binding, SuperBinding))
+		{
+			return SuperBinding.AdvanceBy(Context, DeltaTime);
+		}
+
+		return 0.0f;
+	}
+
+	void IGroupSynchronization::AdvanceToRatio(const FExecutionContext& Context, const TDecoratorBinding<IGroupSynchronization>& Binding, float ProgressRatio) const
+	{
+		TDecoratorBinding<IGroupSynchronization> SuperBinding;
+		if (Context.GetInterfaceSuper(Binding, SuperBinding))
+		{
+			SuperBinding.AdvanceToRatio(Context, ProgressRatio);
+		}
+	}
+}
diff --git a/Engine/Plugins/Experimental/AnimNext/Source/AnimNext/Private/DecoratorInterfaces/ITimeline.cpp b/Engine/Plugins/Experimental/AnimNext/Source/AnimNext/Private/DecoratorInterfaces/ITimeline.cpp
--- a/Engine/Plugins/Experimental/AnimNext/Source/AnimNext/Private/DecoratorInterfaces/ITimeline.cpp
+++ b/Engine/Plugins/Experimental/AnimNext/Source/AnimNext/Private/DecoratorInterfaces/ITimeline.cpp
@@ -6,7 +6,7 @@
 
 namespace UE::AnimNext
 {
-	double ITimeline::GetPlayRate(const FExecutionContext& Context, const TDecoratorBinding<ITimeline>& Binding) const
+	float ITimeline::GetPlayRate(const FExecutionContext& Context, const TDecoratorBinding<ITimeline>& Binding) const
 	{
 		TDecoratorBinding<ITimeline> SuperBinding;
 		if (Context.GetInterfaceSuper(Binding, SuperBinding))
@@ -14,6 +14,26 @@ namespace UE::AnimNext
 			return SuperBinding.GetPlayRate(Context);
 		}
 
-		return 1.0;
+		return 1.0f;
+	}
+
+	float ITimeline::AdvanceBy(const FExecutionContext& Context, const TDecoratorBinding<ITimeline>& Binding, float DeltaTime) const
+	{
+		TDecoratorBinding<ITimeline> SuperBinding;
+		if (Context.GetInterfaceSuper(Binding, SuperBinding))
+		{
+			return SuperBinding.AdvanceBy(Context, DeltaTime);
+		}
+
+		return 0.0f;
+	}
+
+	void ITimeline::AdvanceToRatio(const FExecutionContext& Context, const TDecoratorBinding<ITimeline>& Binding, float ProgressRatio) const
+	{
+		TDecoratorBinding<ITimeline> SuperBinding;
+		if (Context.GetInterfaceSuper(Binding, SuperBinding))
+		{
+			SuperBinding.AdvanceToRatio(Context, ProgressRatio);
+		}
 	}
 }
diff --git a/Engine/Plugins/Experimental/AnimNext/Source/AnimNext/Private/DecoratorInterfaces/IUpdate.cpp b/Engine/Plugins/Experimental/AnimNext/Source/AnimNext/Private/DecoratorInterfaces/IUpdate.cpp
--- a/Engine/Plugins/Experimental/AnimNext/Source/AnimNext/Private/DecoratorInterfaces/IUpdate.cpp
+++ b/Engine/Plugins/Experimental/AnimNext/Source/AnimNext/Private/DecoratorInterfaces/IUpdate.cpp
@@ -225,6 +225,12 @@ namespace UE::AnimNext
 		FUpdateTraversalContext TraversalContext(Context, MemStack);
 		FUpdateTraversalQueue TraversalQueue(TraversalContext);
 
+		// Before we start the traversal, we give the graph instance components the chance to do some work
+		for (auto It = Context.GetComponentIterator(); It; ++It)
+		{
+			It.Value()->PreUpdate(TraversalContext);
+		}
+
 		// Add the graph root to start the update process
 		FUpdateEntry RootEntry(GraphRootPtr, FDecoratorUpdateState(DeltaTime));
 		TraversalContext.PushUpdateEntry(&RootEntry);
@@ -290,5 +296,11 @@ namespace UE::AnimNext
 				TraversalContext.PushFreeEntry(Entry);
 			}
 		}
+
+		// After we finish the traversal, we give the graph instance components the chance to do some work
+		for (auto It = Context.GetComponentIterator(); It; ++It)
+		{
+			It.Value()->PostUpdate(TraversalContext);
+		}
 	}
 }
diff --git a/Engine/Plugins/Experimental/AnimNext/Source/AnimNext/Private/Decorators/SequencePlayer.cpp b/Engine/Plugins/Experimental/AnimNext/Source/AnimNext/Private/Decorators/SequencePlayer.cpp
--- a/Engine/Plugins/Experimental/AnimNext/Source/AnimNext/Private/Decorators/SequencePlayer.cpp
+++ b/Engine/Plugins/Experimental/AnimNext/Source/AnimNext/Private/Decorators/SequencePlayer.cpp
@@ -2,6 +2,7 @@
 
 #include "Decorators/SequencePlayer.h"
 
+#include "AnimationRuntime.h"
 #include "DecoratorBase/ExecutionContext.h"
 #include "EvaluationVM/Tasks/PushAnimSequenceKeyframe.h"
 
@@ -22,7 +23,6 @@ namespace UE::AnimNext
 		{
 			const float SequenceLength = SharedData->AnimSequence->GetPlayLength();
 			InternalTimeAccumulator = FMath::Clamp(SharedData->GetStartPosition(Context, Binding), 0.0f, SequenceLength);
-			PrevInternalTimeAccumulator = InternalTimeAccumulator;
 		}
 	}
 
@@ -40,41 +40,53 @@ namespace UE::AnimNext
 		TraversalContext.AppendTask(Task);
 	}
 
-	double FSequencePlayerDecorator::GetPlayRate(const FExecutionContext& Context, const TDecoratorBinding<ITimeline>& Binding) const
+	float FSequencePlayerDecorator::GetPlayRate(const FExecutionContext& Context, const TDecoratorBinding<ITimeline>& Binding) const
 	{
 		const FSharedData* SharedData = Binding.GetSharedData<FSharedData>();
 		return SharedData->GetPlayRate(Context, Binding);
 	}
 
-	void FSequencePlayerDecorator::PreUpdate(FUpdateTraversalContext& Context, const TDecoratorBinding<IUpdate>& Binding, const FDecoratorUpdateState& DecoratorState) const
+	float FSequencePlayerDecorator::AdvanceBy(const FExecutionContext& Context, const TDecoratorBinding<ITimeline>& Binding, float DeltaTime) const
 	{
 		const FSharedData* SharedData = Binding.GetSharedData<FSharedData>();
-		if (SharedData->AnimSequence != nullptr)
+		if (UAnimSequence* AnimSeq = SharedData->AnimSequence.Get())
 		{
 			FInstanceData* InstanceData = Binding.GetInstanceData<FInstanceData>();
 
 			TDecoratorBinding<ITimeline> TimelineDecorator;
 			Context.GetInterface(Binding, TimelineDecorator);
 
-			const float DeltaTime = DecoratorState.GetDeltaTime();
-			const float PlayRate = (float)TimelineDecorator.GetPlayRate(Context);
+			const float PlayRate = TimelineDecorator.GetPlayRate(Context);
+			const bool bIsLooping = SharedData->GetbLoop(Context, Binding);
+			const float SequenceLength = AnimSeq->GetPlayLength();
 
-			const float EffectiveDelta = FMath::IsNearlyZero(DeltaTime) || FMath::IsNearlyZero(PlayRate) ? 0.f : DeltaTime * PlayRate;
+			FAnimationRuntime::AdvanceTime(bIsLooping, DeltaTime * PlayRate, InstanceData->InternalTimeAccumulator, SequenceLength);
 
-			const bool bIsLooping = SharedData->GetbLoop(Context, Binding);
+			return FMath::Clamp(InstanceData->InternalTimeAccumulator / SequenceLength, 0.0f, 1.0f);
+		}
 
-			const float SequenceLength = SharedData->AnimSequence->GetPlayLength();
-			float CurrentTime = bIsLooping
-				? FMath::Fmod(InstanceData->InternalTimeAccumulator + EffectiveDelta, SequenceLength)
-				: FMath::Clamp(InstanceData->InternalTimeAccumulator + EffectiveDelta, 0.f, SequenceLength);
+		return 0.0f;
+	}
 
-			if (bIsLooping && CurrentTime < 0.f)
-			{
-				CurrentTime += SequenceLength;
-			}
+	void FSequencePlayerDecorator::AdvanceToRatio(const FExecutionContext& Context, const TDecoratorBinding<ITimeline>& Binding, float ProgressRatio) const
+	{
+		const FSharedData* SharedData = Binding.GetSharedData<FSharedData>();
+		if (UAnimSequence* AnimSeq = SharedData->AnimSequence.Get())
+		{
+			FInstanceData* InstanceData = Binding.GetInstanceData<FInstanceData>();
 
-			InstanceData->PrevInternalTimeAccumulator = InstanceData->InternalTimeAccumulator;
-			InstanceData->InternalTimeAccumulator = CurrentTime;
+			const float SequenceLength = AnimSeq->GetPlayLength();
+
+			InstanceData->InternalTimeAccumulator = FMath::Clamp(ProgressRatio, 0.0f, 1.0f) * SequenceLength;
 		}
 	}
+
+	void FSequencePlayerDecorator::PreUpdate(FUpdateTraversalContext& Context, const TDecoratorBinding<IUpdate>& Binding, const FDecoratorUpdateState& DecoratorState) const
+	{
+		// We just advance the timeline
+		TDecoratorBinding<ITimeline> TimelineDecorator;
+		Context.GetInterface(Binding, TimelineDecorator);
+
+		TimelineDecorator.AdvanceBy(Context, DecoratorState.GetDeltaTime());
+	}
 }
diff --git a/Engine/Plugins/Experimental/AnimNext/Source/AnimNext/Private/Decorators/SequencePlayer.h b/Engine/Plugins/Experimental/AnimNext/Source/AnimNext/Private/Decorators/SequencePlayer.h
--- a/Engine/Plugins/Experimental/AnimNext/Source/AnimNext/Private/Decorators/SequencePlayer.h
+++ b/Engine/Plugins/Experimental/AnimNext/Source/AnimNext/Private/Decorators/SequencePlayer.h
@@ -58,7 +58,6 @@ namespace UE::AnimNext
 		struct FInstanceData : FDecorator::FInstanceData
 		{
 			float InternalTimeAccumulator = 0.0f;
-			float PrevInternalTimeAccumulator = 0.0f;
 
 			void Construct(const FExecutionContext& Context, const FDecoratorBinding& Binding);
 		};
@@ -67,7 +66,9 @@ namespace UE::AnimNext
 		virtual void PreEvaluate(const FExecutionContext& Context, const TDecoratorBinding<IEvaluate>& Binding) const override;
 
 		// ITimeline impl
-		virtual double GetPlayRate(const FExecutionContext& Context, const TDecoratorBinding<ITimeline>& Binding) const override;
+		virtual float GetPlayRate(const FExecutionContext& Context, const TDecoratorBinding<ITimeline>& Binding) const override;
+		virtual float AdvanceBy(const FExecutionContext& Context, const TDecoratorBinding<ITimeline>& Binding, float DeltaTime) const override;
+		virtual void AdvanceToRatio(const FExecutionContext& Context, const TDecoratorBinding<ITimeline>& Binding, float ProgressRatio) const override;
 
 		// IUpdate impl
 		virtual void PreUpdate(FUpdateTraversalContext& Context, const TDecoratorBinding<IUpdate>& Binding, const FDecoratorUpdateState& DecoratorState) const override;
diff --git a/Engine/Plugins/Experimental/AnimNext/Source/AnimNext/Private/Decorators/SynchronizeUsingGroups.cpp b/Engine/Plugins/Experimental/AnimNext/Source/AnimNext/Private/Decorators/SynchronizeUsingGroups.cpp
new file mode 100644
--- /dev/null
+++ b/Engine/Plugins/Experimental/AnimNext/Source/AnimNext/Private/Decorators/SynchronizeUsingGroups.cpp
@@ -0,0 +1,111 @@
+// Copyright Epic Games, Inc. All Rights Reserved.
+
+#include "Decorators/SynchronizeUsingGroups.h"
+#include "Graph/SyncGroup_GraphInstanceComponent.h"
+
+namespace UE::AnimNext
+{
+	AUTO_REGISTER_ANIM_DECORATOR(FSynchronizeUsingGroupsDecorator)
+
+	DEFINE_ANIM_DECORATOR_BEGIN(FSynchronizeUsingGroupsDecorator)
+		DEFINE_ANIM_DECORATOR_IMPLEMENTS_INTERFACE(IUpdate)
+		DEFINE_ANIM_DECORATOR_IMPLEMENTS_INTERFACE(IGroupSynchronization)
+		DEFINE_ANIM_DECORATOR_IMPLEMENTS_INTERFACE(ITimeline)
+	DEFINE_ANIM_DECORATOR_END(FSynchronizeUsingGroupsDecorator)
+
+	void FSynchronizeUsingGroupsDecorator::PreUpdate(FUpdateTraversalContext& Context, const TDecoratorBinding<IUpdate>& Binding, const FDecoratorUpdateState& DecoratorState) const
+	{
+		FInstanceData* InstanceData = Binding.GetInstanceData<FInstanceData>();
+
+		TDecoratorBinding<IGroupSynchronization> GroupSyncDecorator;
+		Context.GetInterface(Binding, GroupSyncDecorator);
+
+		const FName GroupName = GroupSyncDecorator.GetGroupName(Context);
+		const bool bHasGroupName = !GroupName.IsNone();
+
+		// If we have a group name, we are active
+		// Freeze the timeline, our sync group will control it
+		InstanceData->bFreezeTimeline = bHasGroupName;
+
+		// Forward the PreUpdate call, if the timeline attempts to update, we'll do nothing if we are frozen
+		IUpdate::PreUpdate(Context, Binding, DecoratorState);
+
+		if (!bHasGroupName)
+		{
+			// If no group name is specified, this decorator is inactive
+			return;
+		}
+
+		const EAnimGroupRole::Type GroupRole = GroupSyncDecorator.GetGroupRole(Context);
+
+		// Append this decorator to our group, we'll need to synchronize it
+		FSyncGroupGraphInstanceComponent& Component = Context.GetComponent<FSyncGroupGraphInstanceComponent>();
+		Component.RegisterWithGroup(GroupName, GroupRole, Binding.GetDecoratorPtr(), DecoratorState);
+	}
+
+	FName FSynchronizeUsingGroupsDecorator::GetGroupName(const FExecutionContext& Context, const TDecoratorBinding<IGroupSynchronization>& Binding) const
+	{
+		const FSharedData* SharedData = Binding.GetSharedData<FSharedData>();
+		return SharedData->GroupName;
+	}
+
+	EAnimGroupRole::Type FSynchronizeUsingGroupsDecorator::GetGroupRole(const FExecutionContext& Context, const TDecoratorBinding<IGroupSynchronization>& Binding) const
+	{
+		const FSharedData* SharedData = Binding.GetSharedData<FSharedData>();
+		return SharedData->GroupRole;
+	}
+
+	float FSynchronizeUsingGroupsDecorator::AdvanceBy(const FExecutionContext& Context, const TDecoratorBinding<IGroupSynchronization>& Binding, float DeltaTime) const
+	{
+		FInstanceData* InstanceData = Binding.GetInstanceData<FInstanceData>();
+
+		// When the group advances the timeline, we thaw it to advance
+		InstanceData->bFreezeTimeline = false;
+
+		TDecoratorBinding<ITimeline> TimelineDecorator;
+		Context.GetInterface(Binding, TimelineDecorator);
+
+		const float ProgressRatio = TimelineDecorator.AdvanceBy(Context, DeltaTime);
+
+		InstanceData->bFreezeTimeline = true;
+
+		return ProgressRatio;
+	}
+
+	void FSynchronizeUsingGroupsDecorator::AdvanceToRatio(const FExecutionContext& Context, const TDecoratorBinding<IGroupSynchronization>& Binding, float ProgressRatio) const
+	{
+		FInstanceData* InstanceData = Binding.GetInstanceData<FInstanceData>();
+
+		// When the group advances the timeline, we thaw it to advance
+		InstanceData->bFreezeTimeline = false;
+
+		TDecoratorBinding<ITimeline> TimelineDecorator;
+		Context.GetInterface(Binding, TimelineDecorator);
+
+		TimelineDecorator.AdvanceToRatio(Context, ProgressRatio);
+
+		InstanceData->bFreezeTimeline = true;
+	}
+
+	float FSynchronizeUsingGroupsDecorator::AdvanceBy(const FExecutionContext& Context, const TDecoratorBinding<ITimeline>& Binding, float DeltaTime) const
+	{
+		const FInstanceData* InstanceData = Binding.GetInstanceData<FInstanceData>();
+		if (InstanceData->bFreezeTimeline)
+		{
+			return 0.0f;	// If the timeline is frozen, we don't advance
+		}
+
+		return ITimeline::AdvanceBy(Context, Binding, DeltaTime);
+	}
+
+	void FSynchronizeUsingGroupsDecorator::AdvanceToRatio(const FExecutionContext& Context, const TDecoratorBinding<ITimeline>& Binding, float ProgressRatio) const
+	{
+		const FInstanceData* InstanceData = Binding.GetInstanceData<FInstanceData>();
+		if (InstanceData->bFreezeTimeline)
+		{
+			return;			// If the timeline is frozen, we don't advance
+		}
+
+		ITimeline::AdvanceToRatio(Context, Binding, ProgressRatio);
+	}
+}
diff --git a/Engine/Plugins/Experimental/AnimNext/Source/AnimNext/Private/Decorators/SynchronizeUsingGroups.h b/Engine/Plugins/Experimental/AnimNext/Source/AnimNext/Private/Decorators/SynchronizeUsingGroups.h
new file mode 100644
--- /dev/null
+++ b/Engine/Plugins/Experimental/AnimNext/Source/AnimNext/Private/Decorators/SynchronizeUsingGroups.h
@@ -0,0 +1,60 @@
+// Copyright Epic Games, Inc. All Rights Reserved.
+
+#pragma once
+
+#include "CoreMinimal.h"
+#include "Animation/AnimationAsset.h"
+#include "DecoratorBase/Decorator.h"
+#include "DecoratorInterfaces/IGroupSynchronization.h"
+#include "DecoratorInterfaces/IUpdate.h"
+#include "DecoratorInterfaces/ITimeline.h"
+
+#include "SynchronizeUsingGroups.generated.h"
+
+USTRUCT()
+struct FAnimNextSynchronizeUsingGroupsDecoratorSharedData : public FAnimNextDecoratorSharedData
+{
+	GENERATED_BODY()
+
+	// The group name
+	// If no name is provided, this decorator is inactive
+	UPROPERTY(meta = (Inline))
+	FName GroupName;
+
+	// The role this player can assume within the group
+	UPROPERTY(meta = (Inline))
+	TEnumAsByte<EAnimGroupRole::Type> GroupRole = EAnimGroupRole::CanBeLeader;
+};
+
+namespace UE::AnimNext
+{
+	/**
+	 * FSynchronizeUsingGroupsDecorator
+	 * 
+	 * A decorator that synchronizes animation sequence playback using named groups.
+	 */
+	struct FSynchronizeUsingGroupsDecorator : FAdditiveDecorator, IUpdate, IGroupSynchronization, ITimeline
+	{
+		DECLARE_ANIM_DECORATOR(FSynchronizeUsingGroupsDecorator, 0x09b4c174, FAdditiveDecorator)
+
+		using FSharedData = FAnimNextSynchronizeUsingGroupsDecoratorSharedData;
+
+		struct FInstanceData : FDecoratorInstanceData
+		{
+			bool bFreezeTimeline = false;
+		};
+
+		// IUpdate impl
+		virtual void PreUpdate(FUpdateTraversalContext& Context, const TDecoratorBinding<IUpdate>& Binding, const FDecoratorUpdateState& DecoratorState) const override;
+
+		// IGroupSynchronization impl
+		virtual FName GetGroupName(const FExecutionContext& Context, const TDecoratorBinding<IGroupSynchronization>& Binding) const override;
+		virtual EAnimGroupRole::Type GetGroupRole(const FExecutionContext& Context, const TDecoratorBinding<IGroupSynchronization>& Binding) const override;
+		virtual float AdvanceBy(const FExecutionContext& Context, const TDecoratorBinding<IGroupSynchronization>& Binding, float DeltaTime) const override;
+		virtual void AdvanceToRatio(const FExecutionContext& Context, const TDecoratorBinding<IGroupSynchronization>& Binding, float ProgressRatio) const override;
+
+		// ITimeline impl
+		virtual float AdvanceBy(const FExecutionContext& Context, const TDecoratorBinding<ITimeline>& Binding, float DeltaTime) const override;
+		virtual void AdvanceToRatio(const FExecutionContext& Context, const TDecoratorBinding<ITimeline>& Binding, float ProgressRatio) const override;
+	};
+}
diff --git a/Engine/Plugins/Experimental/AnimNext/Source/AnimNext/Private/Graph/AnimNextGraph.cpp b/Engine/Plugins/Experimental/AnimNext/Source/AnimNext/Private/Graph/AnimNextGraph.cpp
--- a/Engine/Plugins/Experimental/AnimNext/Source/AnimNext/Private/Graph/AnimNextGraph.cpp
+++ b/Engine/Plugins/Experimental/AnimNext/Source/AnimNext/Private/Graph/AnimNextGraph.cpp
@@ -134,11 +134,14 @@ void FAnimNextGraphInstance::Release()
 		return;
 	}
 
-	// We need an execution context for this graph to be active when we delete the graph instance
-	UE::AnimNext::FExecutionContext Context(Graph->SharedDataBuffer);
+	{
+		// We need an execution context for this graph to be active when we delete the graph instance
+		UE::AnimNext::FExecutionContext Context(*this);
+		GraphInstancePtr.Reset();
+	}
 
-	GraphInstancePtr.Reset();
 	ExtendedExecuteContext.Reset();
+	Components.Empty();
 
 #if WITH_EDITORONLY_DATA
 	FRWScopeLock Lock(Graph->GraphInstancesLock, SLT_Write);
@@ -154,11 +157,36 @@ bool FAnimNextGraphInstance::IsValid() const
 	return GraphInstancePtr.IsValid();
 }
 
+const UAnimNextGraph* FAnimNextGraphInstance::GetGraph() const
+{
+	return Graph;
+}
+
 bool FAnimNextGraphInstance::UsesGraph(const UAnimNextGraph* InGraph) const
 {
 	return Graph == InGraph;
 }
 
+UE::AnimNext::FGraphInstanceComponent* FAnimNextGraphInstance::TryGetComponent(int32 ComponentNameHash, FName ComponentName) const
+{
+	if (const TSharedPtr<UE::AnimNext::FGraphInstanceComponent>* Component = Components.FindByHash(ComponentNameHash, ComponentName))
+	{
+		return Component->Get();
+	}
+
+	return nullptr;
+}
+
+UE::AnimNext::FGraphInstanceComponent& FAnimNextGraphInstance::AddComponent(int32 ComponentNameHash, FName ComponentName, TSharedPtr<UE::AnimNext::FGraphInstanceComponent>&& Component)
+{
+	return *Components.AddByHash(ComponentNameHash, ComponentName, MoveTemp(Component)).Get();
+}
+
+GraphInstanceComponentMapType::TConstIterator FAnimNextGraphInstance::GetComponentIterator() const
+{
+	return Components.CreateConstIterator();
+}
+
 UAnimNextGraph::UAnimNextGraph(const FObjectInitializer& ObjectInitializer)
 	: Super(ObjectInitializer)
 {
@@ -178,13 +206,14 @@ void UAnimNextGraph::AllocateInstance(FAnimNextGraphInstance& Instance) const
 
 	Instance.Graph = this;
 
-	UE::AnimNext::FExecutionContext Context(SharedDataBuffer);
-	Instance.GraphInstancePtr = Context.AllocateNodeInstance(UE::AnimNext::FWeakDecoratorPtr(), ResolvedRootDecoratorHandle);
-
 	Instance.ExtendedExecuteContext.CopyMemoryStorage(ExtendedExecuteContext);
-
 	VM->InitializeInstance(Instance.ExtendedExecuteContext);
 
+	{
+		UE::AnimNext::FExecutionContext Context(Instance);
+		Instance.GraphInstancePtr = Context.AllocateNodeInstance(UE::AnimNext::FWeakDecoratorPtr(), ResolvedRootDecoratorHandle);
+	}
+
 #if WITH_EDITORONLY_DATA
 	FRWScopeLock Lock(GraphInstancesLock, SLT_Write);
 	check(!GraphInstances.Contains(&Instance));
@@ -200,7 +229,7 @@ void UAnimNextGraph::Run(const UE::AnimNext::FContext& Context, FAnimNextGraphIn
 	{
 		FAnimNextExecuteContext& AnimNextContext = GraphInstance.ExtendedExecuteContext.GetPublicDataSafe<FAnimNextExecuteContext>();
 		AnimNextContext.SetContextData(Context);
-		AnimNextContext.InitializeWithGraph(this, SharedDataBuffer, GraphInstance.GraphInstancePtr);
+		AnimNextContext.SetGraphInstance(GraphInstance);
 		AnimNextContext.SetSimulationSteps(SimulationSteps);
 
 		VM->ExecuteVM(GraphInstance.ExtendedExecuteContext, FRigUnit_AnimNextShimRoot::EventName);
diff --git a/Engine/Plugins/Experimental/AnimNext/Source/AnimNext/Private/Graph/RigUnit_AnimNextGraphEvaluator.cpp b/Engine/Plugins/Experimental/AnimNext/Source/AnimNext/Private/Graph/RigUnit_AnimNextGraphEvaluator.cpp
--- a/Engine/Plugins/Experimental/AnimNext/Source/AnimNext/Private/Graph/RigUnit_AnimNextGraphEvaluator.cpp
+++ b/Engine/Plugins/Experimental/AnimNext/Source/AnimNext/Private/Graph/RigUnit_AnimNextGraphEvaluator.cpp
@@ -36,10 +36,12 @@ void FRigUnit_AnimNextGraphEvaluator::StaticExecute(FRigVMExtendedExecuteContext
 	using namespace UE::AnimNext;
 
 	const FAnimNextExecuteContext& VMExecuteContext = RigVMExecuteContext.GetPublicData<FAnimNextExecuteContext>();
+	FAnimNextGraphInstance& GraphInstance = VMExecuteContext.GetGraphInstance();
 
 	// Setup what we need to execute
-	FExecutionContext Context(VMExecuteContext.GetSharedDataBuffer(), RigVMExecuteContext, RigVMMemoryHandles);
-	FWeakDecoratorPtr GraphInstancePtr = VMExecuteContext.GetGraphInstancePtr();
+	FExecutionContext Context(GraphInstance, RigVMExecuteContext, RigVMMemoryHandles);
+
+	const FWeakDecoratorPtr& GraphInstancePtr = GraphInstance.GraphInstancePtr;
 	const EAnimNextGraphSimulationSteps SimulationSteps = VMExecuteContext.GetSimulationSteps();
 
 	if (EnumHasAnyFlags(SimulationSteps, EAnimNextGraphSimulationSteps::Update))
@@ -60,7 +62,7 @@ void FRigUnit_AnimNextGraphEvaluator::StaticExecute(FRigVMExtendedExecuteContext
 		{
 			FParamStack& ParamStack = FParamStack::Get();
 
-			const UAnimNextGraph* Graph = VMExecuteContext.GetGraph();
+			const UAnimNextGraph* Graph = GraphInstance.GetGraph();
 			const FAnimNextGraphReferencePose* GraphReferencePosePtr = ParamStack.GetParamPtr<FAnimNextGraphReferencePose>(Graph->GetReferencePoseParam());
 			const int32* GraphLODLevelPtr = ParamStack.GetParamPtr<int32>(Graph->GetCurrentLODParam());
 			static FParamId ResultId("UE_Internal_ResultPose");
diff --git a/Engine/Plugins/Experimental/AnimNext/Source/AnimNext/Private/Graph/SyncGroup_GraphInstanceComponent.cpp b/Engine/Plugins/Experimental/AnimNext/Source/AnimNext/Private/Graph/SyncGroup_GraphInstanceComponent.cpp
new file mode 100644
--- /dev/null
+++ b/Engine/Plugins/Experimental/AnimNext/Source/AnimNext/Private/Graph/SyncGroup_GraphInstanceComponent.cpp
@@ -0,0 +1,97 @@
+// Copyright Epic Games, Inc. All Rights Reserved.
+
+#include "Graph/SyncGroup_GraphInstanceComponent.h"
+#include "DecoratorInterfaces/IGroupSynchronization.h"
+
+namespace UE::AnimNext
+{
+	void FSyncGroupGraphInstanceComponent::RegisterWithGroup(FName GroupName, EAnimGroupRole::Type GroupRole, const FWeakDecoratorPtr& DecoratorPtr, const FDecoratorUpdateState& DecoratorState)
+	{
+		FSyncGroupState& GroupState = SyncGroupMap.FindOrAdd(GroupName);
+		GroupState.Members.Add(FSyncGroupMember{ DecoratorState, DecoratorPtr, GroupRole });
+	}
+
+	void FSyncGroupGraphInstanceComponent::PreUpdate(FUpdateTraversalContext& Context)
+	{
+		// Reset our group state, we want to start fresh every update
+		SyncGroupMap.Reset();
+	}
+
+	void FSyncGroupGraphInstanceComponent::PostUpdate(FUpdateTraversalContext& Context)
+	{
+		TDecoratorBinding<IGroupSynchronization> GroupSyncDecorator;
+
+		// Now that we have discovered all groups and their memberships, we can perform synchronization
+		for (const TTuple<FName, FSyncGroupState>& It : SyncGroupMap)
+		{
+			const FSyncGroupState& GroupState = It.Value;
+
+			// Find our leader by looking at the total weight and the group role
+			const int32 NumMembers = GroupState.Members.Num();
+			check(NumMembers != 0);	// Groups should never be empty
+
+			int32 LeaderIndex = INDEX_NONE;
+			float LeaderTotalWeight = -1.0f;
+
+			for (int32 MemberIndex = 0; MemberIndex < NumMembers; ++MemberIndex)
+			{
+				const FSyncGroupMember& GroupMember = GroupState.Members[MemberIndex];
+
+				switch (GroupMember.GroupRole)
+				{
+				case EAnimGroupRole::CanBeLeader:
+				case EAnimGroupRole::TransitionLeader:
+					// Highest weight is the leader
+					if (GroupMember.DecoratorState.GetTotalWeight() > LeaderTotalWeight)
+					{
+						LeaderIndex = MemberIndex;
+						LeaderTotalWeight = GroupMember.DecoratorState.GetTotalWeight();
+					}
+					break;
+				case EAnimGroupRole::AlwaysLeader:
+					// Always set the leader index
+					LeaderIndex = MemberIndex;
+					LeaderTotalWeight = 2.0f;		// Some high value
+					break;
+				default:
+				case EAnimGroupRole::AlwaysFollower:
+				case EAnimGroupRole::TransitionFollower:
+					// Never set the leader index
+					// If we find no leader, we'll use the first index as set below
+					break;
+				}
+			}
+
+			if (LeaderIndex == INDEX_NONE)
+			{
+				// If none of the entries wish to be a leader, grab the first and force it
+				LeaderIndex = 0;
+			}
+
+			// Now that we found our leader, advance it by the delta time
+			float LeaderProgressRatio;
+			{
+				const FSyncGroupMember& GroupLeader = GroupState.Members[LeaderIndex];
+
+				ensure(Context.GetInterface(GroupLeader.DecoratorPtr, GroupSyncDecorator));
+
+				LeaderProgressRatio = GroupSyncDecorator.AdvanceBy(Context, GroupLeader.DecoratorState.GetDeltaTime());
+			}
+
+			// Advance every follower to the same progress ratio as the leader
+			for (int32 MemberIndex = 0; MemberIndex < NumMembers; ++MemberIndex)
+			{
+				if (MemberIndex == LeaderIndex)
+				{
+					continue;	// Ignore the leader, it already advanced
+				}
+
+				const FSyncGroupMember& GroupMember = GroupState.Members[MemberIndex];
+
+				ensure(Context.GetInterface(GroupMember.DecoratorPtr, GroupSyncDecorator));
+
+				GroupSyncDecorator.AdvanceToRatio(Context, LeaderProgressRatio);
+			}
+		}
+	}
+}
diff --git a/Engine/Plugins/Experimental/AnimNext/Source/AnimNext/Private/Graph/SyncGroup_GraphInstanceComponent.h b/Engine/Plugins/Experimental/AnimNext/Source/AnimNext/Private/Graph/SyncGroup_GraphInstanceComponent.h
new file mode 100644
--- /dev/null
+++ b/Engine/Plugins/Experimental/AnimNext/Source/AnimNext/Private/Graph/SyncGroup_GraphInstanceComponent.h
@@ -0,0 +1,43 @@
+// Copyright Epic Games, Inc. All Rights Reserved.
+
+#pragma once
+
+#include "CoreMinimal.h"
+#include "Animation/AnimationAsset.h"
+#include "DecoratorBase/DecoratorPtr.h"
+#include "DecoratorInterfaces/IUpdate.h"
+#include "Graph/GraphInstanceComponent.h"
+
+namespace UE::AnimNext
+{
+	/**
+	 * FSyncGroupGraphInstanceComponent
+	 *
+	 * This component maintains the necessary state to support group based synchronization.
+	 */
+	struct FSyncGroupGraphInstanceComponent : public FGraphInstanceComponent
+	{
+		DECLARE_ANIM_GRAPH_INSTANCE_COMPONENT(FSyncGroupGraphInstanceComponent)
+
+		void RegisterWithGroup(FName GroupName, EAnimGroupRole::Type GroupRole, const FWeakDecoratorPtr& DecoratorPtr, const FDecoratorUpdateState& DecoratorState);
+
+		// FGraphInstanceComponent impl
+		virtual void PreUpdate(FUpdateTraversalContext& Context) override;
+		virtual void PostUpdate(FUpdateTraversalContext& Context) override;
+
+	private:
+		struct FSyncGroupMember
+		{
+			FDecoratorUpdateState				DecoratorState;
+			FWeakDecoratorPtr					DecoratorPtr;
+			TEnumAsByte<EAnimGroupRole::Type>	GroupRole;
+		};
+
+		struct FSyncGroupState
+		{
+			TArray<FSyncGroupMember> Members;
+		};
+
+		TMap<FName, FSyncGroupState> SyncGroupMap;
+	};
+}
diff --git a/Engine/Plugins/Experimental/AnimNext/Source/AnimNext/Public/DecoratorBase/ExecutionContext.h b/Engine/Plugins/Experimental/AnimNext/Source/AnimNext/Public/DecoratorBase/ExecutionContext.h
--- a/Engine/Plugins/Experimental/AnimNext/Source/AnimNext/Public/DecoratorBase/ExecutionContext.h
+++ b/Engine/Plugins/Experimental/AnimNext/Source/AnimNext/Public/DecoratorBase/ExecutionContext.h
@@ -9,6 +9,7 @@
 #include "DecoratorBase/DecoratorInterfaceUID.h"
 #include "DecoratorBase/LatentPropertyHandle.h"
 #include "DecoratorBase/NodeHandle.h"
+#include "Graph/AnimNextGraph.h"
 
 struct FRigVMExtendedExecuteContext;
 struct FRigVMMemoryHandle;
@@ -33,11 +34,11 @@ namespace UE::AnimNext
 	 */
 	struct ANIMNEXT_API FExecutionContext
 	{
-		// Creates an execution context for the specified graph
-		explicit FExecutionContext(TArrayView<const uint8> InGraphSharedData);
+		// Creates an execution context for the specified graph instance
+		explicit FExecutionContext(FAnimNextGraphInstance& InGraphInstance);
 
 		// Creates an execution context for the specified graph with RigVM latent pin support
-		FExecutionContext(TArrayView<const uint8> InGraphSharedData, FRigVMExtendedExecuteContext& InRigVMExecuteContext, FRigVMMemoryHandleArray InRigVMLatentMemoryHandles);
+		FExecutionContext(FAnimNextGraphInstance& InGraphInstance, FRigVMExtendedExecuteContext& InRigVMExecuteContext, FRigVMMemoryHandleArray InRigVMLatentMemoryHandles);
 
 		// Destroys the execution context
 		~FExecutionContext();
@@ -81,9 +82,24 @@ namespace UE::AnimNext
 		template<typename LatentPinType>
 		LatentPinType EvaluateLatentPin(FLatentPropertyHandle LatentPropertyHandle) const;
 
+		// Returns a typed graph instance component, creating it lazily the first time it is queried
+		template<class ComponentType>
+		ComponentType& GetComponent() const;
+
+		// Returns a typed graph instance component pointer if found or nullptr otherwise
+		template<class ComponentType>
+		ComponentType* TryGetComponent() const;
+
+		// Returns const iterators to the graph instance component container
+		GraphInstanceComponentMapType::TConstIterator GetComponentIterator() const { return GraphInstance->GetComponentIterator(); }
+
+		// Returns the bound graph instance
+		const FAnimNextGraphInstance& GetGraphInstance() const { return *GraphInstance; }
+
 	private:
+		// No copy or move
 		FExecutionContext(const FExecutionContext&) = delete;
-		FExecutionContext(FExecutionContext&&) = delete;
+		FExecutionContext& operator=(const FExecutionContext&) = delete;
 
 		bool GetInterfaceImpl(FDecoratorInterfaceUID InterfaceUID, const FWeakDecoratorPtr& DecoratorPtr, FDecoratorBinding& InterfaceBinding) const;
 		bool GetInterfaceSuperImpl(FDecoratorInterfaceUID InterfaceUID, const FWeakDecoratorPtr& DecoratorPtr, FDecoratorBinding& SuperBinding) const;
@@ -100,9 +116,11 @@ namespace UE::AnimNext
 		ITraversalContext* TraversalContext = nullptr;
 
 		// Cached properties for the currently executing graph
+		const UAnimNextGraph* Graph = nullptr;
+		FAnimNextGraphInstance* GraphInstance = nullptr;
 		TArrayView<const uint8> GraphSharedData;
 		FRigVMMemoryHandleArray RigVMLatentMemoryHandles;
-		FRigVMExtendedExecuteContext* RigVMExecuteContext;
+		FRigVMExtendedExecuteContext* RigVMExecuteContext = nullptr;
 
 		friend struct FScopedTraversalContext;
 	};
@@ -162,4 +180,16 @@ namespace UE::AnimNext
 
 		return *reinterpret_cast<const LatentPinType*>(EvaluateLatentPinImpl(LatentPropertyHandle));
 	}
+
+	template<class ComponentType>
+	ComponentType& FExecutionContext::GetComponent() const
+	{
+		return GraphInstance->GetComponent<ComponentType>();
+	}
+
+	template<class ComponentType>
+	ComponentType* FExecutionContext::TryGetComponent() const
+	{
+		return GraphInstance->TryGetComponent<ComponentType>();
+	}
 }
diff --git a/Engine/Plugins/Experimental/AnimNext/Source/AnimNext/Public/DecoratorBase/ExecutionContextProxy.h b/Engine/Plugins/Experimental/AnimNext/Source/AnimNext/Public/DecoratorBase/ExecutionContextProxy.h
--- a/Engine/Plugins/Experimental/AnimNext/Source/AnimNext/Public/DecoratorBase/ExecutionContextProxy.h
+++ b/Engine/Plugins/Experimental/AnimNext/Source/AnimNext/Public/DecoratorBase/ExecutionContextProxy.h
@@ -80,6 +80,32 @@ namespace UE::AnimNext
 			return ExecutionContext.EvaluateLatentPin<LatentPinType>(LatentPropertyHandle);
 		}
 
+		// Returns a typed graph instance component, creating it lazily the first time it is queried
+		template<class ComponentType>
+		ComponentType& GetComponent() const
+		{
+			return ExecutionContext.GetComponent<ComponentType>();
+		}
+
+		// Returns a typed graph instance component pointer if found or nullptr otherwise
+		template<class ComponentType>
+		ComponentType* TryGetComponent() const
+		{
+			return ExecutionContext.TryGetComponent<ComponentType>();
+		}
+
+		// Returns const iterators to the graph instance component container
+		GraphInstanceComponentMapType::TConstIterator GetComponentIterator() const
+		{
+			return ExecutionContext.GetComponentIterator();
+		}
+
+		// Returns the bound graph instance
+		const FAnimNextGraphInstance& GetGraphInstance() const
+		{
+			return ExecutionContext.GetGraphInstance();
+		}
+
 	private:
 		// The execution context that we wrap
 		const FExecutionContext& ExecutionContext;
diff --git a/Engine/Plugins/Experimental/AnimNext/Source/AnimNext/Public/DecoratorInterfaces/IGroupSynchronization.h b/Engine/Plugins/Experimental/AnimNext/Source/AnimNext/Public/DecoratorInterfaces/IGroupSynchronization.h
new file mode 100644
--- /dev/null
+++ b/Engine/Plugins/Experimental/AnimNext/Source/AnimNext/Public/DecoratorInterfaces/IGroupSynchronization.h
@@ -0,0 +1,68 @@
+// Copyright Epic Games, Inc. All Rights Reserved.
+
+#pragma once
+
+#include "CoreMinimal.h"
+#include "Animation/AnimationAsset.h"
+#include "DecoratorBase/IDecoratorInterface.h"
+
+namespace UE::AnimNext
+{
+	/**
+	 * IGroupSynchronization
+	 *
+	 * This interface exposes group synchronization related information and behavior.
+	 */
+	struct ANIMNEXT_API IGroupSynchronization : IDecoratorInterface
+	{
+		DECLARE_ANIM_DECORATOR_INTERFACE(IGroupSynchronization, 0xf607d0fd)
+
+		// Returns the group name used for synchronization
+		virtual FName GetGroupName(const FExecutionContext& Context, const TDecoratorBinding<IGroupSynchronization>& Binding) const;
+
+		// Returns the group role used for synchronization
+		virtual EAnimGroupRole::Type GetGroupRole(const FExecutionContext& Context, const TDecoratorBinding<IGroupSynchronization>& Binding) const;
+
+		// Called by the sync group graph instance component once a group has been synchronized to advance time on the leader
+		// Returns the progress ratio of playback: 0.0 = start of animation, 1.0 = end of animation
+		virtual float AdvanceBy(const FExecutionContext& Context, const TDecoratorBinding<IGroupSynchronization>& Binding, float DeltaTime) const;
+
+		// Called by the sync group graph instance component once a group has been synchronized to advance time on each follower
+		// Progress ratio must be between [0.0, 1.0]
+		virtual void AdvanceToRatio(const FExecutionContext& Context, const TDecoratorBinding<IGroupSynchronization>& Binding, float ProgressRatio) const;
+	};
+
+	/**
+	 * Specialization for decorator binding.
+	 */
+	template<>
+	struct TDecoratorBinding<IGroupSynchronization> : FDecoratorBinding
+	{
+		// @see IGroupSynchronization::GetGroupName
+		FName GetGroupName(const FExecutionContext& Context) const
+		{
+			return GetInterface()->GetGroupName(Context, *this);
+		}
+
+		// @see IGroupSynchronization::GetGroupRole
+		EAnimGroupRole::Type GetGroupRole(const FExecutionContext& Context) const
+		{
+			return GetInterface()->GetGroupRole(Context, *this);
+		}
+
+		// @see IGroupSynchronization::AdvanceBy
+		float AdvanceBy(const FExecutionContext& Context, float DeltaTime) const
+		{
+			return GetInterface()->AdvanceBy(Context, *this, DeltaTime);
+		}
+
+		// @see IGroupSynchronization::AdvanceToRatio
+		void AdvanceToRatio(const FExecutionContext& Context, float ProgressRatio) const
+		{
+			GetInterface()->AdvanceToRatio(Context, *this, ProgressRatio);
+		}
+
+	protected:
+		const IGroupSynchronization* GetInterface() const { return GetInterfaceTyped<IGroupSynchronization>(); }
+	};
+}
diff --git a/Engine/Plugins/Experimental/AnimNext/Source/AnimNext/Public/DecoratorInterfaces/ITimeline.h b/Engine/Plugins/Experimental/AnimNext/Source/AnimNext/Public/DecoratorInterfaces/ITimeline.h
--- a/Engine/Plugins/Experimental/AnimNext/Source/AnimNext/Public/DecoratorInterfaces/ITimeline.h
+++ b/Engine/Plugins/Experimental/AnimNext/Source/AnimNext/Public/DecoratorInterfaces/ITimeline.h
@@ -17,7 +17,15 @@ namespace UE::AnimNext
 		DECLARE_ANIM_DECORATOR_INTERFACE(ITimeline, 0x53760727)
 
 		// Returns the play rate of this timeline
-		virtual double GetPlayRate(const FExecutionContext& Context, const TDecoratorBinding<ITimeline>& Binding) const;
+		virtual float GetPlayRate(const FExecutionContext& Context, const TDecoratorBinding<ITimeline>& Binding) const;
+
+		// Advances time by the provided delta time (positive or negative) on this timeline
+		// Returns the progress ratio of playback: 0.0 = start of animation, 1.0 = end of animation
+		virtual float AdvanceBy(const FExecutionContext& Context, const TDecoratorBinding<ITimeline>& Binding, float DeltaTime) const;
+
+		// Advances time to the specified progress ratio on this timeline
+		// Progress ratio must be between [0.0, 1.0]
+		virtual void AdvanceToRatio(const FExecutionContext& Context, const TDecoratorBinding<ITimeline>& Binding, float ProgressRatio) const;
 	};
 
 	/**
@@ -27,11 +35,23 @@ namespace UE::AnimNext
 	struct TDecoratorBinding<ITimeline> : FDecoratorBinding
 	{
 		// @see ITimeline::GetPlayRate
-		double GetPlayRate(const FExecutionContext& Context) const
+		float GetPlayRate(const FExecutionContext& Context) const
 		{
 			return GetInterface()->GetPlayRate(Context, *this);
 		}
 
+		// @see ITimeline::AdvanceBy
+		float AdvanceBy(const FExecutionContext& Context, float DeltaTime) const
+		{
+			return GetInterface()->AdvanceBy(Context, *this, DeltaTime);
+		}
+
+		// @see ITimeline::AdvanceToRatio
+		void AdvanceToRatio(const FExecutionContext& Context, float ProgressRatio) const
+		{
+			GetInterface()->AdvanceToRatio(Context, *this, ProgressRatio);
+		}
+
 	protected:
 		const ITimeline* GetInterface() const { return GetInterfaceTyped<ITimeline>(); }
 	};
diff --git a/Engine/Plugins/Experimental/AnimNext/Source/AnimNext/Public/Graph/GraphInstanceComponent.h b/Engine/Plugins/Experimental/AnimNext/Source/AnimNext/Public/Graph/GraphInstanceComponent.h
new file mode 100644
--- /dev/null
+++ b/Engine/Plugins/Experimental/AnimNext/Source/AnimNext/Public/Graph/GraphInstanceComponent.h
@@ -0,0 +1,36 @@
+// Copyright Epic Games, Inc. All Rights Reserved.
+
+#pragma once
+
+#include "CoreMinimal.h"
+
+// Declares a graph instance component and implements the necessary boilerplate
+#define DECLARE_ANIM_GRAPH_INSTANCE_COMPONENT(ComponentType) \
+	static FName StaticComponentName() { return FName(TEXT(#ComponentType)); } \
+	virtual FName GetComponentName() const override { return StaticComponentName(); }
+
+namespace UE::AnimNext
+{
+	struct FUpdateTraversalContext;
+
+	/**
+	 * FGraphInstanceComponent
+	 *
+	 * A graph instance component is attached and owned by a graph instance.
+	 * It persists as long as it is needed.
+	 */
+	struct FGraphInstanceComponent
+	{
+		virtual ~FGraphInstanceComponent() {}
+
+		static FName StaticComponentName() { return FName(TEXT("FGraphInstanceComponent")); }
+		virtual FName GetComponentName() const { return StaticComponentName(); }
+
+		// Called before the update traversal begins, before any node has been visited
+		// Note that PreUpdate won't be called if a component is created during the update traversal until the next update
+		virtual void PreUpdate(FUpdateTraversalContext& Context) {}
+
+		// Called after the update traversal completes, after every node has been visited
+		virtual void PostUpdate(FUpdateTraversalContext& Context) {}
+	};
+}
diff --git a/Engine/Plugins/Experimental/AnimNext/Source/AnimNextTestSuite/Private/AnimNextDecoratorBaseTest.cpp b/Engine/Plugins/Experimental/AnimNext/Source/AnimNextTestSuite/Private/AnimNextDecoratorBaseTest.cpp
--- a/Engine/Plugins/Experimental/AnimNext/Source/AnimNextTestSuite/Private/AnimNextDecoratorBaseTest.cpp
+++ b/Engine/Plugins/Experimental/AnimNext/Source/AnimNextTestSuite/Private/AnimNextDecoratorBaseTest.cpp
@@ -18,6 +18,7 @@
 #include "DecoratorBase/NodeInstance.h"
 #include "DecoratorBase/NodeTemplateBuilder.h"
 #include "DecoratorBase/NodeTemplateRegistry.h"
+#include "Graph/GraphFactory.h"
 
 //****************************************************************************
 // AnimNext Runtime DecoratorBase Tests
@@ -500,6 +501,10 @@ bool FAnimationAnimNextRuntimeTest_NodeLifetime::RunTest(const FString& InParame
 	AUTO_REGISTER_ANIM_DECORATOR(FDecoratorAB_Add)
 	AUTO_REGISTER_ANIM_DECORATOR(FDecoratorAC_Add)
 
+	UFactory* GraphFactory = NewObject<UAnimNextGraphFactory>();
+	UAnimNextGraph* AnimNextGraph = CastChecked<UAnimNextGraph>(GraphFactory->FactoryCreateNew(UAnimNextGraph::StaticClass(), GetTransientPackage(), TEXT("TestAnimNextGraph"), RF_Transient, nullptr, nullptr, NAME_None));
+	UE_RETURN_ON_ERROR(AnimNextGraph != nullptr, "FAnimationAnimNextRuntimeTest_NodeLifetime -> Failed to create animation graph");
+
 	FScopedClearNodeTemplateRegistry ScopedClearNodeTemplateRegistry;
 	FNodeTemplateRegistry& Registry = FNodeTemplateRegistry::Get();
 
@@ -517,8 +522,7 @@ bool FAnimationAnimNextRuntimeTest_NodeLifetime::RunTest(const FString& InParame
 	FNodeTemplateRegistryHandle TemplateHandle0 = Registry.FindOrAdd(NodeTemplate0);
 	AddErrorIfFalse(TemplateHandle0.IsValid(), "FAnimationAnimNextRuntimeTest_NodeLifetime -> Registry should contain our template");
 
-	FNodeHandle Node0;
-	FNodeHandle Node1;
+	TArray<FNodeHandle> NodeHandles;
 
 	// Write our graph
 	TArray<uint8> GraphSharedDataArchiveBuffer;
@@ -526,13 +530,13 @@ bool FAnimationAnimNextRuntimeTest_NodeLifetime::RunTest(const FString& InParame
 	{
 		FDecoratorWriter DecoratorWriter;
 
-		Node0 = DecoratorWriter.RegisterNode(*NodeTemplate0);
-		Node1 = DecoratorWriter.RegisterNode(*NodeTemplate0);
+		NodeHandles.Add(DecoratorWriter.RegisterNode(*NodeTemplate0));
+		NodeHandles.Add(DecoratorWriter.RegisterNode(*NodeTemplate0));
 
 		// We don't have decorator properties
 
 		DecoratorWriter.BeginNodeWriting();
-		DecoratorWriter.WriteNode(Node0,
+		DecoratorWriter.WriteNode(NodeHandles[0],
 			[](uint32 DecoratorIndex, const FString& PropertyName)
 			{
 				return FString();
@@ -541,7 +545,7 @@ bool FAnimationAnimNextRuntimeTest_NodeLifetime::RunTest(const FString& InParame
 			{
 				return false;
 			});
-		DecoratorWriter.WriteNode(Node1,
+		DecoratorWriter.WriteNode(NodeHandles[1],
 			[](uint32 DecoratorIndex, const FString& PropertyName)
 			{
 				return FString();
@@ -558,31 +562,24 @@ bool FAnimationAnimNextRuntimeTest_NodeLifetime::RunTest(const FString& InParame
 	}
 
 	// Read our graph
-	TArray<uint8> GraphSharedDataBuffer;
-	{
-		FMemoryReader GraphSharedDataArchive(GraphSharedDataArchiveBuffer);
-		FDecoratorReader DecoratorReader(GraphReferencedObjects, GraphSharedDataArchive);
+	FTestUtils::LoadFromArchiveBuffer(*AnimNextGraph, NodeHandles, GraphSharedDataArchiveBuffer);
 
-		FDecoratorReader::EErrorState ErrorState = DecoratorReader.ReadGraph(GraphSharedDataBuffer);
-		AddErrorIfFalse(ErrorState == FDecoratorReader::EErrorState::None, "FAnimationAnimNextRuntimeTest_NodeLifetime -> Failed to read graph shared data");
+	FAnimNextGraphInstance GraphInstance;
+	AnimNextGraph->AllocateInstance(GraphInstance);
 
-		Node0 = DecoratorReader.ResolveNodeHandle(Node0);
-		Node1 = DecoratorReader.ResolveNodeHandle(Node1);
-	}
-
-	FExecutionContext Context(GraphSharedDataBuffer);
+	FExecutionContext Context(GraphInstance);
 
 	// Validate handle bookkeeping
 	{
-		FDecoratorBinding RootBinding;					// Empty, no parent
-		FAnimNextDecoratorHandle DecoratorHandle00(Node0, 0);	// Point to first node, first base decorator
+		FDecoratorBinding RootBinding;									// Empty, no parent
+		FAnimNextDecoratorHandle DecoratorHandle00(NodeHandles[0], 0);	// Point to first node, first base decorator
 
 		// Allocate a node
 		FDecoratorPtr DecoratorPtr00 = Context.AllocateNodeInstance(RootBinding, DecoratorHandle00);
 		AddErrorIfFalse(DecoratorPtr00.IsValid(), "FAnimationAnimNextRuntimeTest_NodeLifetime -> Failed to allocate a node instance");
 		AddErrorIfFalse(DecoratorPtr00.GetDecoratorIndex() == 0, "FAnimationAnimNextRuntimeTest_NodeLifetime -> Allocated decorator pointer should point to root decorator");
 		AddErrorIfFalse(!DecoratorPtr00.IsWeak(), "FAnimationAnimNextRuntimeTest_NodeLifetime -> Allocated decorator pointer should not be weak, we have no parent");
-		AddErrorIfFalse(DecoratorPtr00.GetNodeInstance()->GetNodeHandle() == Node0, "FAnimationAnimNextRuntimeTest_NodeLifetime -> Allocated node should point to the provided node handle");
+		AddErrorIfFalse(DecoratorPtr00.GetNodeInstance()->GetNodeHandle() == NodeHandles[0], "FAnimationAnimNextRuntimeTest_NodeLifetime -> Allocated node should point to the provided node handle");
 		AddErrorIfFalse(DecoratorPtr00.GetNodeInstance()->GetReferenceCount() == 1, "FAnimationAnimNextRuntimeTest_NodeLifetime -> Allocated node should have a single reference");
 
 		{
@@ -624,10 +621,10 @@ bool FAnimationAnimNextRuntimeTest_NodeLifetime::RunTest(const FString& InParame
 
 	// Validate parent support
 	{
-		FDecoratorBinding RootBinding;					// Empty, no parent
-		FAnimNextDecoratorHandle DecoratorHandle00(Node0, 0);	// Point to first node, first base decorator
-		FAnimNextDecoratorHandle DecoratorHandle03(Node0, 3);	// Point to first node, second base decorator
-		FAnimNextDecoratorHandle DecoratorHandle10(Node1, 0);	// Point to second node, first base decorator
+		FDecoratorBinding RootBinding;									// Empty, no parent
+		FAnimNextDecoratorHandle DecoratorHandle00(NodeHandles[0], 0);	// Point to first node, first base decorator
+		FAnimNextDecoratorHandle DecoratorHandle03(NodeHandles[0], 3);	// Point to first node, second base decorator
+		FAnimNextDecoratorHandle DecoratorHandle10(NodeHandles[1], 0);	// Point to second node, first base decorator
 
 		// Allocate our first node
 		FDecoratorPtr DecoratorPtr00 = Context.AllocateNodeInstance(RootBinding, DecoratorHandle00);
@@ -639,7 +636,7 @@ bool FAnimationAnimNextRuntimeTest_NodeLifetime::RunTest(const FString& InParame
 		AddErrorIfFalse(DecoratorPtr03.IsValid(), "FAnimationAnimNextRuntimeTest_NodeLifetime -> Failed to allocate a node instance");
 		AddErrorIfFalse(DecoratorPtr03.GetDecoratorIndex() == 3, "FAnimationAnimNextRuntimeTest_NodeLifetime -> Allocated decorator pointer should point to fourth decorator");
 		AddErrorIfFalse(DecoratorPtr03.IsWeak(), "FAnimationAnimNextRuntimeTest_NodeLifetime -> Allocated decorator pointer should be weak, we have the same parent");
-		AddErrorIfFalse(DecoratorPtr03.GetNodeInstance()->GetNodeHandle() == Node0, "FAnimationAnimNextRuntimeTest_NodeLifetime -> Allocated node should point to the provided node handle");
+		AddErrorIfFalse(DecoratorPtr03.GetNodeInstance()->GetNodeHandle() == NodeHandles[0], "FAnimationAnimNextRuntimeTest_NodeLifetime -> Allocated node should point to the provided node handle");
 		AddErrorIfFalse(DecoratorPtr03.GetNodeInstance() == DecoratorPtr00.GetNodeInstance(), "FAnimationAnimNextRuntimeTest_NodeLifetime -> Handles should point to the same node instance");
 		AddErrorIfFalse(DecoratorPtr03.GetNodeInstance()->GetReferenceCount() == 1, "FAnimationAnimNextRuntimeTest_NodeLifetime -> Allocated node should have one reference");
 
@@ -649,7 +646,7 @@ bool FAnimationAnimNextRuntimeTest_NodeLifetime::RunTest(const FString& InParame
 		AddErrorIfFalse(DecoratorPtr10.IsValid(), "FAnimationAnimNextRuntimeTest_NodeLifetime -> Failed to allocate a node instance");
 		AddErrorIfFalse(DecoratorPtr10.GetDecoratorIndex() == 0, "FAnimationAnimNextRuntimeTest_NodeLifetime -> Allocated decorator pointer should point to first decorator");
 		AddErrorIfFalse(!DecoratorPtr10.IsWeak(), "FAnimationAnimNextRuntimeTest_NodeLifetime -> Allocated decorator pointer should not be weak, we have the same parent but a different node handle");
-		AddErrorIfFalse(DecoratorPtr10.GetNodeInstance()->GetNodeHandle() == Node1, "FAnimationAnimNextRuntimeTest_NodeLifetime -> Allocated node should point to the provided node handle");
+		AddErrorIfFalse(DecoratorPtr10.GetNodeInstance()->GetNodeHandle() == NodeHandles[1], "FAnimationAnimNextRuntimeTest_NodeLifetime -> Allocated node should point to the provided node handle");
 		AddErrorIfFalse(DecoratorPtr10.GetNodeInstance() != DecoratorPtr00.GetNodeInstance(), "FAnimationAnimNextRuntimeTest_NodeLifetime -> Handles should not point to the same node instance");
 		AddErrorIfFalse(DecoratorPtr10.GetNodeInstance()->GetReferenceCount() == 1, "FAnimationAnimNextRuntimeTest_NodeLifetime -> Allocated node should have one reference");
 	}
@@ -663,8 +660,8 @@ bool FAnimationAnimNextRuntimeTest_NodeLifetime::RunTest(const FString& InParame
 		Private::DestructedDecorators = &DestructedDecorators;
 
 		{
-			FDecoratorBinding RootBinding;					// Empty, no parent
-			FAnimNextDecoratorHandle DecoratorHandle00(Node0, 0);	// Point to first node, first base decorator
+			FDecoratorBinding RootBinding;									// Empty, no parent
+			FAnimNextDecoratorHandle DecoratorHandle00(NodeHandles[0], 0);	// Point to first node, first base decorator
 
 			// Allocate our node instance
 			FDecoratorPtr DecoratorPtr00 = Context.AllocateNodeInstance(RootBinding, DecoratorHandle00);
@@ -713,6 +710,10 @@ bool FAnimationAnimNextRuntimeTest_GetDecoratorInterface::RunTest(const FString&
 	AUTO_REGISTER_ANIM_DECORATOR(FDecoratorAB_Add)
 	AUTO_REGISTER_ANIM_DECORATOR(FDecoratorAC_Add)
 
+	UFactory* GraphFactory = NewObject<UAnimNextGraphFactory>();
+	UAnimNextGraph* AnimNextGraph = CastChecked<UAnimNextGraph>(GraphFactory->FactoryCreateNew(UAnimNextGraph::StaticClass(), GetTransientPackage(), TEXT("TestAnimNextGraph"), RF_Transient, nullptr, nullptr, NAME_None));
+	UE_RETURN_ON_ERROR(AnimNextGraph != nullptr, "FAnimationAnimNextRuntimeTest_GetDecoratorInterface -> Failed to create animation graph");
+
 	FScopedClearNodeTemplateRegistry ScopedClearNodeTemplateRegistry;
 	FNodeTemplateRegistry& Registry = FNodeTemplateRegistry::Get();
 
@@ -730,7 +731,7 @@ bool FAnimationAnimNextRuntimeTest_GetDecoratorInterface::RunTest(const FString&
 	FNodeTemplateRegistryHandle TemplateHandle0 = Registry.FindOrAdd(NodeTemplate0);
 	AddErrorIfFalse(TemplateHandle0.IsValid(), "FAnimationAnimNextRuntimeTest_GetDecoratorInterface -> Registry should contain our template");
 
-	FNodeHandle Node0;
+	TArray<FNodeHandle> NodeHandles;
 
 	// Write our graph
 	TArray<uint8> GraphSharedDataArchiveBuffer;
@@ -738,12 +739,12 @@ bool FAnimationAnimNextRuntimeTest_GetDecoratorInterface::RunTest(const FString&
 	{
 		FDecoratorWriter DecoratorWriter;
 
-		Node0 = DecoratorWriter.RegisterNode(*NodeTemplate0);
+		NodeHandles.Add(DecoratorWriter.RegisterNode(*NodeTemplate0));
 
 		// We don't have decorator properties
 
 		DecoratorWriter.BeginNodeWriting();
-		DecoratorWriter.WriteNode(Node0,
+		DecoratorWriter.WriteNode(NodeHandles[0],
 			[](uint32 DecoratorIndex, const FString& PropertyName)
 			{
 				return FString();
@@ -760,23 +761,17 @@ bool FAnimationAnimNextRuntimeTest_GetDecoratorInterface::RunTest(const FString&
 	}
 
 	// Read our graph
-	TArray<uint8> GraphSharedDataBuffer;
-	{
-		FMemoryReader GraphSharedDataArchive(GraphSharedDataArchiveBuffer);
-		FDecoratorReader DecoratorReader(GraphReferencedObjects, GraphSharedDataArchive);
-
-		FDecoratorReader::EErrorState ErrorState = DecoratorReader.ReadGraph(GraphSharedDataBuffer);
-		AddErrorIfFalse(ErrorState == FDecoratorReader::EErrorState::None, "FAnimationAnimNextRuntimeTest_GetDecoratorInterface -> Failed to read graph shared data");
+	FTestUtils::LoadFromArchiveBuffer(*AnimNextGraph, NodeHandles, GraphSharedDataArchiveBuffer);
 
-		Node0 = DecoratorReader.ResolveNodeHandle(Node0);
-	}
+	FAnimNextGraphInstance GraphInstance;
+	AnimNextGraph->AllocateInstance(GraphInstance);
 
-	FExecutionContext Context(GraphSharedDataBuffer);
+	FExecutionContext Context(GraphInstance);
 
 	// Validate from the first base decorator
 	{
-		FDecoratorBinding ParentBinding;				// Empty, no parent
-		FAnimNextDecoratorHandle DecoratorHandle00(Node0, 0);	// Point to first node, first base decorator
+		FDecoratorBinding ParentBinding;								// Empty, no parent
+		FAnimNextDecoratorHandle DecoratorHandle00(NodeHandles[0], 0);	// Point to first node, first base decorator
 
 		FDecoratorPtr DecoratorPtr00 = Context.AllocateNodeInstance(ParentBinding, DecoratorHandle00);
 		AddErrorIfFalse(DecoratorPtr00.IsValid(), "FAnimationAnimNextRuntimeTest_GetDecoratorInterface -> Failed to allocate a node instance");
@@ -787,7 +782,7 @@ bool FAnimationAnimNextRuntimeTest_GetDecoratorInterface::RunTest(const FString&
 		AddErrorIfFalse(Binding00C.IsValid(), "FAnimationAnimNextRuntimeTest_GetDecoratorInterface -> InterfaceC binding not valid");
 		AddErrorIfFalse(Binding00C.GetInterfaceUID() == IInterfaceC::InterfaceUID, "FAnimationAnimNextRuntimeTest_GetDecoratorInterface -> Unexpected interface UID found in decorator binding");
 		AddErrorIfFalse(Binding00C.GetDecoratorPtr().GetDecoratorIndex() == 2, "FAnimationAnimNextRuntimeTest_GetDecoratorInterface -> InterfaceC not found on expected decorator");
-		AddErrorIfFalse(Binding00C.GetDecoratorPtr().GetNodeInstance()->GetNodeHandle() == Node0, "FAnimationAnimNextRuntimeTest_GetDecoratorInterface -> InterfaceC not found on expected node");
+		AddErrorIfFalse(Binding00C.GetDecoratorPtr().GetNodeInstance()->GetNodeHandle() == NodeHandles[0], "FAnimationAnimNextRuntimeTest_GetDecoratorInterface -> InterfaceC not found on expected node");
 		AddErrorIfFalse(Binding00C.GetSharedData<FDecoratorAC_Add::FSharedData>()->DecoratorUID == FDecoratorAC_Add::DecoratorUID, "FAnimationAnimNextRuntimeTest_GetDecoratorInterface -> Unexpected shared data in decorator binding");
 		AddErrorIfFalse(Binding00C.GetInstanceData<FDecoratorAC_Add::FInstanceData>()->DecoratorUID == FDecoratorAC_Add::DecoratorUID, "FAnimationAnimNextRuntimeTest_GetDecoratorInterface -> Unexpected instance data in decorator binding");
 
@@ -796,7 +791,7 @@ bool FAnimationAnimNextRuntimeTest_GetDecoratorInterface::RunTest(const FString&
 		AddErrorIfFalse(Binding00B.IsValid(), "FAnimationAnimNextRuntimeTest_GetDecoratorInterface -> InterfaceB binding not valid");
 		AddErrorIfFalse(Binding00B.GetInterfaceUID() == IInterfaceB::InterfaceUID, "FAnimationAnimNextRuntimeTest_GetDecoratorInterface -> Unexpected interface UID found in decorator binding");
 		AddErrorIfFalse(Binding00B.GetDecoratorPtr().GetDecoratorIndex() == 1, "FAnimationAnimNextRuntimeTest_GetDecoratorInterface -> InterfaceB not found on expected decorator");
-		AddErrorIfFalse(Binding00B.GetDecoratorPtr().GetNodeInstance()->GetNodeHandle() == Node0, "FAnimationAnimNextRuntimeTest_GetDecoratorInterface -> InterfaceB not found on expected node");
+		AddErrorIfFalse(Binding00B.GetDecoratorPtr().GetNodeInstance()->GetNodeHandle() == NodeHandles[0], "FAnimationAnimNextRuntimeTest_GetDecoratorInterface -> InterfaceB not found on expected node");
 		AddErrorIfFalse(Binding00B.GetSharedData<FDecoratorAB_Add::FSharedData>()->DecoratorUID == FDecoratorAB_Add::DecoratorUID, "FAnimationAnimNextRuntimeTest_GetDecoratorInterface -> Unexpected shared data in decorator binding");
 		AddErrorIfFalse(Binding00B.GetInstanceData<FDecoratorAB_Add::FInstanceData>()->DecoratorUID == FDecoratorAB_Add::DecoratorUID, "FAnimationAnimNextRuntimeTest_GetDecoratorInterface -> Unexpected instance data in decorator binding");
 
@@ -805,7 +800,7 @@ bool FAnimationAnimNextRuntimeTest_GetDecoratorInterface::RunTest(const FString&
 		AddErrorIfFalse(Binding00A.IsValid(), "FAnimationAnimNextRuntimeTest_GetDecoratorInterface -> InterfaceA binding not valid");
 		AddErrorIfFalse(Binding00A.GetInterfaceUID() == IInterfaceA::InterfaceUID, "FAnimationAnimNextRuntimeTest_GetDecoratorInterface -> Unexpected interface UID found in decorator binding");
 		AddErrorIfFalse(Binding00A.GetDecoratorPtr().GetDecoratorIndex() == 2, "FAnimationAnimNextRuntimeTest_GetDecoratorInterface -> InterfaceA not found on expected decorator");
-		AddErrorIfFalse(Binding00A.GetDecoratorPtr().GetNodeInstance()->GetNodeHandle() == Node0, "FAnimationAnimNextRuntimeTest_GetDecoratorInterface -> InterfaceA not found on expected node");
+		AddErrorIfFalse(Binding00A.GetDecoratorPtr().GetNodeInstance()->GetNodeHandle() == NodeHandles[0], "FAnimationAnimNextRuntimeTest_GetDecoratorInterface -> InterfaceA not found on expected node");
 		AddErrorIfFalse(Binding00A.GetSharedData<FDecoratorAC_Add::FSharedData>()->DecoratorUID == FDecoratorAC_Add::DecoratorUID, "FAnimationAnimNextRuntimeTest_GetDecoratorInterface -> Unexpected shared data in decorator binding");
 		AddErrorIfFalse(Binding00A.GetInstanceData<FDecoratorAC_Add::FInstanceData>()->DecoratorUID == FDecoratorAC_Add::DecoratorUID, "FAnimationAnimNextRuntimeTest_GetDecoratorInterface -> Unexpected instance data in decorator binding");
 
@@ -875,8 +870,8 @@ bool FAnimationAnimNextRuntimeTest_GetDecoratorInterface::RunTest(const FString&
 
 	// Validate from the second base decorator
 	{
-		FDecoratorBinding ParentBinding;				// Empty, no parent
-		FAnimNextDecoratorHandle DecoratorHandle03(Node0, 3);	// Point to first node, second base decorator
+		FDecoratorBinding ParentBinding;								// Empty, no parent
+		FAnimNextDecoratorHandle DecoratorHandle03(NodeHandles[0], 3);	// Point to first node, second base decorator
 
 		FDecoratorPtr DecoratorPtr03 = Context.AllocateNodeInstance(ParentBinding, DecoratorHandle03);
 		AddErrorIfFalse(DecoratorPtr03.IsValid(), "FAnimationAnimNextRuntimeTest_GetDecoratorInterface -> Failed to allocate a node instance");
@@ -887,7 +882,7 @@ bool FAnimationAnimNextRuntimeTest_GetDecoratorInterface::RunTest(const FString&
 		AddErrorIfFalse(Binding03C.IsValid(), "FAnimationAnimNextRuntimeTest_GetDecoratorInterface -> InterfaceC binding not valid");
 		AddErrorIfFalse(Binding03C.GetInterfaceUID() == IInterfaceC::InterfaceUID, "FAnimationAnimNextRuntimeTest_GetDecoratorInterface -> Unexpected interface UID found in decorator binding");
 		AddErrorIfFalse(Binding03C.GetDecoratorPtr().GetDecoratorIndex() == 4, "FAnimationAnimNextRuntimeTest_GetDecoratorInterface -> InterfaceC not found on expected decorator");
-		AddErrorIfFalse(Binding03C.GetDecoratorPtr().GetNodeInstance()->GetNodeHandle() == Node0, "FAnimationAnimNextRuntimeTest_GetDecoratorInterface -> InterfaceC not found on expected node");
+		AddErrorIfFalse(Binding03C.GetDecoratorPtr().GetNodeInstance()->GetNodeHandle() == NodeHandles[0], "FAnimationAnimNextRuntimeTest_GetDecoratorInterface -> InterfaceC not found on expected node");
 		AddErrorIfFalse(Binding03C.GetSharedData<FDecoratorAC_Add::FSharedData>()->DecoratorUID == FDecoratorAC_Add::DecoratorUID, "FAnimationAnimNextRuntimeTest_GetDecoratorInterface -> Unexpected shared data in decorator binding");
 		AddErrorIfFalse(Binding03C.GetInstanceData<FDecoratorAC_Add::FInstanceData>()->DecoratorUID == FDecoratorAC_Add::DecoratorUID, "FAnimationAnimNextRuntimeTest_GetDecoratorInterface -> Unexpected instance data in decorator binding");
 
@@ -900,7 +895,7 @@ bool FAnimationAnimNextRuntimeTest_GetDecoratorInterface::RunTest(const FString&
 		AddErrorIfFalse(Binding03A.IsValid(), "FAnimationAnimNextRuntimeTest_GetDecoratorInterface -> InterfaceA binding not valid");
 		AddErrorIfFalse(Binding03A.GetInterfaceUID() == IInterfaceA::InterfaceUID, "FAnimationAnimNextRuntimeTest_GetDecoratorInterface -> Unexpected interface UID found in decorator binding");
 		AddErrorIfFalse(Binding03A.GetDecoratorPtr().GetDecoratorIndex() == 4, "FAnimationAnimNextRuntimeTest_GetDecoratorInterface -> InterfaceA not found on expected decorator");
-		AddErrorIfFalse(Binding03A.GetDecoratorPtr().GetNodeInstance()->GetNodeHandle() == Node0, "FAnimationAnimNextRuntimeTest_GetDecoratorInterface -> InterfaceA not found on expected node");
+		AddErrorIfFalse(Binding03A.GetDecoratorPtr().GetNodeInstance()->GetNodeHandle() == NodeHandles[0], "FAnimationAnimNextRuntimeTest_GetDecoratorInterface -> InterfaceA not found on expected node");
 		AddErrorIfFalse(Binding03A.GetSharedData<FDecoratorAC_Add::FSharedData>()->DecoratorUID == FDecoratorAC_Add::DecoratorUID, "FAnimationAnimNextRuntimeTest_GetDecoratorInterface -> Unexpected shared data in decorator binding");
 		AddErrorIfFalse(Binding03A.GetInstanceData<FDecoratorAC_Add::FInstanceData>()->DecoratorUID == FDecoratorAC_Add::DecoratorUID, "FAnimationAnimNextRuntimeTest_GetDecoratorInterface -> Unexpected instance data in decorator binding");
 
@@ -983,6 +978,10 @@ bool FAnimationAnimNextRuntimeTest_GetDecoratorInterfaceSuper::RunTest(const FSt
 	AUTO_REGISTER_ANIM_DECORATOR(FDecoratorAB_Add)
 	AUTO_REGISTER_ANIM_DECORATOR(FDecoratorAC_Add)
 
+	UFactory* GraphFactory = NewObject<UAnimNextGraphFactory>();
+	UAnimNextGraph* AnimNextGraph = CastChecked<UAnimNextGraph>(GraphFactory->FactoryCreateNew(UAnimNextGraph::StaticClass(), GetTransientPackage(), TEXT("TestAnimNextGraph"), RF_Transient, nullptr, nullptr, NAME_None));
+	UE_RETURN_ON_ERROR(AnimNextGraph != nullptr, "FAnimationAnimNextRuntimeTest_GetDecoratorInterfaceSuper -> Failed to create animation graph");
+
 	FScopedClearNodeTemplateRegistry ScopedClearNodeTemplateRegistry;
 	FNodeTemplateRegistry& Registry = FNodeTemplateRegistry::Get();
 
@@ -1001,7 +1000,7 @@ bool FAnimationAnimNextRuntimeTest_GetDecoratorInterfaceSuper::RunTest(const FSt
 	AddErrorIfFalse(Registry.GetNum() == 1, "FAnimationAnimNextRuntimeTest_GetDecoratorInterfaceSuper -> Registry should contain 1 template");
 	AddErrorIfFalse(TemplateHandle0.IsValid(), "FAnimationAnimNextRuntimeTest_GetDecoratorInterfaceSuper -> Registry should contain our template");
 
-	FNodeHandle Node0;
+	TArray<FNodeHandle> NodeHandles;
 
 	// Write our graph
 	TArray<uint8> GraphSharedDataArchiveBuffer;
@@ -1009,12 +1008,12 @@ bool FAnimationAnimNextRuntimeTest_GetDecoratorInterfaceSuper::RunTest(const FSt
 	{
 		FDecoratorWriter DecoratorWriter;
 
-		Node0 = DecoratorWriter.RegisterNode(*NodeTemplate0);
+		NodeHandles.Add(DecoratorWriter.RegisterNode(*NodeTemplate0));
 
 		// We don't have decorator properties
 
 		DecoratorWriter.BeginNodeWriting();
-		DecoratorWriter.WriteNode(Node0,
+		DecoratorWriter.WriteNode(NodeHandles[0],
 			[](uint32 DecoratorIndex, const FString& PropertyName)
 			{
 				return FString();
@@ -1031,23 +1030,17 @@ bool FAnimationAnimNextRuntimeTest_GetDecoratorInterfaceSuper::RunTest(const FSt
 	}
 
 	// Read our graph
-	TArray<uint8> GraphSharedDataBuffer;
-	{
-		FMemoryReader GraphSharedDataArchive(GraphSharedDataArchiveBuffer);
-		FDecoratorReader DecoratorReader(GraphReferencedObjects, GraphSharedDataArchive);
-
-		FDecoratorReader::EErrorState ErrorState = DecoratorReader.ReadGraph(GraphSharedDataBuffer);
-		AddErrorIfFalse(ErrorState == FDecoratorReader::EErrorState::None, "FAnimationAnimNextRuntimeTest_GetDecoratorInterfaceSuper -> Failed to read graph shared data");
+	FTestUtils::LoadFromArchiveBuffer(*AnimNextGraph, NodeHandles, GraphSharedDataArchiveBuffer);
 
-		Node0 = DecoratorReader.ResolveNodeHandle(Node0);
-	}
+	FAnimNextGraphInstance GraphInstance;
+	AnimNextGraph->AllocateInstance(GraphInstance);
 
-	FExecutionContext Context(GraphSharedDataBuffer);
+	FExecutionContext Context(GraphInstance);
 
 	// Validate from the first base decorator
 	{
-		FDecoratorBinding ParentBinding;				// Empty, no parent
-		FAnimNextDecoratorHandle DecoratorHandle00(Node0, 0);	// Point to first node, first base decorator
+		FDecoratorBinding ParentBinding;								// Empty, no parent
+		FAnimNextDecoratorHandle DecoratorHandle00(NodeHandles[0], 0);	// Point to first node, first base decorator
 
 		FDecoratorPtr DecoratorPtr00 = Context.AllocateNodeInstance(ParentBinding, DecoratorHandle00);
 		AddErrorIfFalse(DecoratorPtr00.IsValid(), "FAnimationAnimNextRuntimeTest_GetDecoratorInterfaceSuper -> Failed to allocate a node instance");
@@ -1078,7 +1071,7 @@ bool FAnimationAnimNextRuntimeTest_GetDecoratorInterfaceSuper::RunTest(const FSt
 			AddErrorIfFalse(SuperBinding02A.IsValid(), "FAnimationAnimNextRuntimeTest_GetDecoratorInterfaceSuper -> InterfaceA binding not valid");
 			AddErrorIfFalse(SuperBinding02A.GetInterfaceUID() == IInterfaceA::InterfaceUID, "FAnimationAnimNextRuntimeTest_GetDecoratorInterfaceSuper -> Unexpected interface UID found in decorator binding");
 			AddErrorIfFalse(SuperBinding02A.GetDecoratorPtr().GetDecoratorIndex() == 1, "FAnimationAnimNextRuntimeTest_GetDecoratorInterfaceSuper -> InterfaceA not found on expected decorator");
-			AddErrorIfFalse(SuperBinding02A.GetDecoratorPtr().GetNodeInstance()->GetNodeHandle() == Node0, "FAnimationAnimNextRuntimeTest_GetDecoratorInterfaceSuper -> InterfaceA not found on expected node");
+			AddErrorIfFalse(SuperBinding02A.GetDecoratorPtr().GetNodeInstance()->GetNodeHandle() == NodeHandles[0], "FAnimationAnimNextRuntimeTest_GetDecoratorInterfaceSuper -> InterfaceA not found on expected node");
 			AddErrorIfFalse(SuperBinding02A.GetSharedData<FDecoratorAB_Add::FSharedData>()->DecoratorUID == FDecoratorAB_Add::DecoratorUID, "FAnimationAnimNextRuntimeTest_GetDecoratorInterfaceSuper -> Unexpected shared data in decorator binding");
 			AddErrorIfFalse(SuperBinding02A.GetInstanceData<FDecoratorAB_Add::FInstanceData>()->DecoratorUID == FDecoratorAB_Add::DecoratorUID, "FAnimationAnimNextRuntimeTest_GetDecoratorInterfaceSuper -> Unexpected instance data in decorator binding");
 
@@ -1093,7 +1086,7 @@ bool FAnimationAnimNextRuntimeTest_GetDecoratorInterfaceSuper::RunTest(const FSt
 			AddErrorIfFalse(SuperBinding01A.IsValid(), "FAnimationAnimNextRuntimeTest_GetDecoratorInterfaceSuper -> InterfaceA binding not valid");
 			AddErrorIfFalse(SuperBinding01A.GetInterfaceUID() == IInterfaceA::InterfaceUID, "FAnimationAnimNextRuntimeTest_GetDecoratorInterfaceSuper -> Unexpected interface UID found in decorator binding");
 			AddErrorIfFalse(SuperBinding01A.GetDecoratorPtr().GetDecoratorIndex() == 0, "FAnimationAnimNextRuntimeTest_GetDecoratorInterfaceSuper -> InterfaceA not found on expected decorator");
-			AddErrorIfFalse(SuperBinding01A.GetDecoratorPtr().GetNodeInstance()->GetNodeHandle() == Node0, "FAnimationAnimNextRuntimeTest_GetDecoratorInterfaceSuper -> InterfaceA not found on expected node");
+			AddErrorIfFalse(SuperBinding01A.GetDecoratorPtr().GetNodeInstance()->GetNodeHandle() == NodeHandles[0], "FAnimationAnimNextRuntimeTest_GetDecoratorInterfaceSuper -> InterfaceA not found on expected node");
 			AddErrorIfFalse(SuperBinding01A.GetSharedData<FDecoratorA_Base::FSharedData>()->DecoratorUID == FDecoratorA_Base::DecoratorUID, "FAnimationAnimNextRuntimeTest_GetDecoratorInterfaceSuper -> Unexpected shared data in decorator binding");
 			AddErrorIfFalse(SuperBinding01A.GetInstanceData<FDecoratorA_Base::FInstanceData>()->DecoratorUID == FDecoratorA_Base::DecoratorUID, "FAnimationAnimNextRuntimeTest_GetDecoratorInterfaceSuper -> Unexpected instance data in decorator binding");
 
@@ -1115,8 +1108,8 @@ bool FAnimationAnimNextRuntimeTest_GetDecoratorInterfaceSuper::RunTest(const FSt
 
 	// Validate from the second base decorator
 	{
-		FDecoratorBinding ParentBinding;				// Empty, no parent
-		FAnimNextDecoratorHandle DecoratorHandle03(Node0, 3);	// Point to first node, second base decorator
+		FDecoratorBinding ParentBinding;								// Empty, no parent
+		FAnimNextDecoratorHandle DecoratorHandle03(NodeHandles[0], 3);	// Point to first node, second base decorator
 
 		FDecoratorPtr DecoratorPtr03 = Context.AllocateNodeInstance(ParentBinding, DecoratorHandle03);
 		AddErrorIfFalse(DecoratorPtr03.IsValid(), "FAnimationAnimNextRuntimeTest_GetDecoratorInterfaceSuper -> Failed to allocate a node instance");
@@ -1147,7 +1140,7 @@ bool FAnimationAnimNextRuntimeTest_GetDecoratorInterfaceSuper::RunTest(const FSt
 			AddErrorIfFalse(SuperBinding04A.IsValid(), "FAnimationAnimNextRuntimeTest_GetDecoratorInterfaceSuper -> InterfaceA binding not valid");
 			AddErrorIfFalse(SuperBinding04A.GetInterfaceUID() == IInterfaceA::InterfaceUID, "FAnimationAnimNextRuntimeTest_GetDecoratorInterfaceSuper -> Unexpected interface UID found in decorator binding");
 			AddErrorIfFalse(SuperBinding04A.GetDecoratorPtr().GetDecoratorIndex() == 3, "FAnimationAnimNextRuntimeTest_GetDecoratorInterfaceSuper -> InterfaceA not found on expected decorator");
-			AddErrorIfFalse(SuperBinding04A.GetDecoratorPtr().GetNodeInstance()->GetNodeHandle() == Node0, "FAnimationAnimNextRuntimeTest_GetDecoratorInterfaceSuper -> InterfaceA not found on expected node");
+			AddErrorIfFalse(SuperBinding04A.GetDecoratorPtr().GetNodeInstance()->GetNodeHandle() == NodeHandles[0], "FAnimationAnimNextRuntimeTest_GetDecoratorInterfaceSuper -> InterfaceA not found on expected node");
 			AddErrorIfFalse(SuperBinding04A.GetSharedData<FDecoratorA_Base::FSharedData>()->DecoratorUID == FDecoratorA_Base::DecoratorUID, "FAnimationAnimNextRuntimeTest_GetDecoratorInterfaceSuper -> Unexpected shared data in decorator binding");
 			AddErrorIfFalse(SuperBinding04A.GetInstanceData<FDecoratorA_Base::FInstanceData>()->DecoratorUID == FDecoratorA_Base::DecoratorUID, "FAnimationAnimNextRuntimeTest_GetDecoratorInterfaceSuper -> Unexpected instance data in decorator binding");
 
@@ -1184,6 +1177,10 @@ bool FAnimationAnimNextRuntimeTest_DecoratorSerialization::RunTest(const FString
 	AUTO_REGISTER_ANIM_DECORATOR(FDecoratorSerialization_Add)
 	AUTO_REGISTER_ANIM_DECORATOR(FDecoratorNativeSerialization_Add)
 
+	UFactory* GraphFactory = NewObject<UAnimNextGraphFactory>();
+	UAnimNextGraph* AnimNextGraph = CastChecked<UAnimNextGraph>(GraphFactory->FactoryCreateNew(UAnimNextGraph::StaticClass(), GetTransientPackage(), TEXT("TestAnimNextGraph"), RF_Transient, nullptr, nullptr, NAME_None));
+	UE_RETURN_ON_ERROR(AnimNextGraph != nullptr, "FAnimationAnimNextRuntimeTest_DecoratorSerialization -> Failed to create animation graph");
+
 	FScopedClearNodeTemplateRegistry ScopedClearNodeTemplateRegistry;
 	FNodeTemplateRegistry& Registry = FNodeTemplateRegistry::Get();
 
@@ -1283,8 +1280,7 @@ bool FAnimationAnimNextRuntimeTest_DecoratorSerialization::RunTest(const FString
 	DecoratorNativeRef1.String = TEXT("sample string 1235");
 	DecoratorNativeRef1.Name = FName(TEXT("sample name 9991785"));
 
-	FNodeHandle Node0;
-	FNodeHandle Node1;
+	TArray<FNodeHandle> NodeHandles;
 
 	// Write our graph
 	TArray<uint8> GraphSharedDataArchiveBuffer;
@@ -1292,8 +1288,8 @@ bool FAnimationAnimNextRuntimeTest_DecoratorSerialization::RunTest(const FString
 	{
 		FDecoratorWriter DecoratorWriter;
 
-		Node0 = DecoratorWriter.RegisterNode(*NodeTemplate0);
-		Node1 = DecoratorWriter.RegisterNode(*NodeTemplate0);
+		NodeHandles.Add(DecoratorWriter.RegisterNode(*NodeTemplate0));
+		NodeHandles.Add(DecoratorWriter.RegisterNode(*NodeTemplate0));
 
 		// We don't have decorator properties
 		TArray<TMap<FString, FString>> DecoratorProperties0;
@@ -1357,7 +1353,7 @@ bool FAnimationAnimNextRuntimeTest_DecoratorSerialization::RunTest(const FString
 		DecoratorProperties1[2].Add(TEXT("Name"), ToString<FDecoratorNativeSerialization_Add::FSharedData>(TEXT("Name"), DecoratorNativeRef1.Name));
 
 		DecoratorWriter.BeginNodeWriting();
-		DecoratorWriter.WriteNode(Node0,
+		DecoratorWriter.WriteNode(NodeHandles[0],
 			[&DecoratorProperties0](uint32 DecoratorIndex, const FString& PropertyName)
 			{
 				return DecoratorProperties0[DecoratorIndex][PropertyName];
@@ -1366,7 +1362,7 @@ bool FAnimationAnimNextRuntimeTest_DecoratorSerialization::RunTest(const FString
 			{
 				return false;
 			});
-		DecoratorWriter.WriteNode(Node1,
+		DecoratorWriter.WriteNode(NodeHandles[1],
 			[&DecoratorProperties1](uint32 DecoratorIndex, const FString& PropertyName)
 			{
 				return DecoratorProperties1[DecoratorIndex][PropertyName];
@@ -1389,27 +1385,18 @@ bool FAnimationAnimNextRuntimeTest_DecoratorSerialization::RunTest(const FString
 		AddErrorIfFalse(Registry.GetNum() == 0, "FAnimationAnimNextRuntimeTest_DecoratorSerialization -> Registry should contain 0 templates");
 
 		// Read our graph
-		TArray<uint8> GraphSharedDataBuffer;
-		{
-			FMemoryReader GraphSharedDataArchive(GraphSharedDataArchiveBuffer);
-			FDecoratorReader DecoratorReader(GraphReferencedObjects, GraphSharedDataArchive);
-
-			FDecoratorReader::EErrorState ErrorState = DecoratorReader.ReadGraph(GraphSharedDataBuffer);
-			AddErrorIfFalse(ErrorState == FDecoratorReader::EErrorState::None, "FAnimationAnimNextRuntimeTest_DecoratorSerialization -> Failed to read graph shared data");
+		FTestUtils::LoadFromArchiveBuffer(*AnimNextGraph, NodeHandles, GraphSharedDataArchiveBuffer);
 
-			Node0 = DecoratorReader.ResolveNodeHandle(Node0);
-			Node1 = DecoratorReader.ResolveNodeHandle(Node1);
-
-			AddErrorIfFalse(Registry.GetNum() == 1, "FAnimationAnimNextRuntimeTest_DecoratorSerialization -> Registry should contain 1 template");
-		}
+		FAnimNextGraphInstance GraphInstance;
+		AnimNextGraph->AllocateInstance(GraphInstance);
 
-		FExecutionContext Context(GraphSharedDataBuffer);
+		FExecutionContext Context(GraphInstance);
 
 		// Validate decorator serialization
 		{
-			FDecoratorBinding ParentBinding;				// Empty, no parent
-			FAnimNextDecoratorHandle DecoratorHandle0(Node0, 0);	// Point to first node, first base decorator
-			FAnimNextDecoratorHandle DecoratorHandle1(Node1, 0);	// Point to second node, first base decorator
+			FDecoratorBinding ParentBinding;								// Empty, no parent
+			FAnimNextDecoratorHandle DecoratorHandle0(NodeHandles[0], 0);	// Point to first node, first base decorator
+			FAnimNextDecoratorHandle DecoratorHandle1(NodeHandles[1], 0);	// Point to second node, first base decorator
 
 			FDecoratorPtr DecoratorPtr0 = Context.AllocateNodeInstance(ParentBinding, DecoratorHandle0);
 			AddErrorIfFalse(DecoratorPtr0.IsValid(), "FAnimationAnimNextRuntimeTest_DecoratorSerialization -> Failed to allocate a node instance");
diff --git a/Engine/Plugins/Experimental/AnimNext/Source/AnimNextTestSuite/Private/AnimNextDecoratorInterfacesTest.cpp b/Engine/Plugins/Experimental/AnimNext/Source/AnimNextTestSuite/Private/AnimNextDecoratorInterfacesTest.cpp
--- a/Engine/Plugins/Experimental/AnimNext/Source/AnimNextTestSuite/Private/AnimNextDecoratorInterfacesTest.cpp
+++ b/Engine/Plugins/Experimental/AnimNext/Source/AnimNextTestSuite/Private/AnimNextDecoratorInterfacesTest.cpp
@@ -21,6 +21,7 @@
 #include "DecoratorInterfaces/IEvaluate.h"
 #include "DecoratorInterfaces/IHierarchy.h"
 #include "DecoratorInterfaces/IUpdate.h"
+#include "Graph/GraphFactory.h"
 
 //****************************************************************************
 // AnimNext Runtime DecoratorInterfaces Tests
@@ -224,6 +225,10 @@ bool FAnimationAnimNextRuntimeTest_IHierarchy::RunTest(const FString& InParamete
 	AUTO_REGISTER_ANIM_DECORATOR(FDecoratorWithOneChild)
 	AUTO_REGISTER_ANIM_DECORATOR(FDecoratorWithChildren)
 
+	UFactory* GraphFactory = NewObject<UAnimNextGraphFactory>();
+	UAnimNextGraph* AnimNextGraph = CastChecked<UAnimNextGraph>(GraphFactory->FactoryCreateNew(UAnimNextGraph::StaticClass(), GetTransientPackage(), TEXT("TestAnimNextGraph"), RF_Transient, nullptr, nullptr, NAME_None));
+	UE_RETURN_ON_ERROR(AnimNextGraph != nullptr, "FAnimationAnimNextRuntimeTest_IHierarchy -> Failed to create animation graph");
+
 	FScopedClearNodeTemplateRegistry ScopedClearNodeTemplateRegistry;
 	FNodeTemplateRegistry& Registry = FNodeTemplateRegistry::Get();
 
@@ -258,10 +263,7 @@ bool FAnimationAnimNextRuntimeTest_IHierarchy::RunTest(const FString& InParamete
 	// NodeC has two children: NodeA and NodeB (but both decorators are base, only NodeB will be referenced)
 	// NodeD has two children: NodeA and NodeC
 
-	FNodeHandle NodeA;
-	FNodeHandle NodeB;
-	FNodeHandle NodeC;
-	FNodeHandle NodeD;
+	TArray<FNodeHandle> NodeHandles;
 
 	// Write our graph
 	TArray<uint8> GraphSharedDataArchiveBuffer;
@@ -269,10 +271,10 @@ bool FAnimationAnimNextRuntimeTest_IHierarchy::RunTest(const FString& InParamete
 	{
 		FDecoratorWriter DecoratorWriter;
 
-		NodeA = DecoratorWriter.RegisterNode(*NodeTemplateA);
-		NodeB = DecoratorWriter.RegisterNode(*NodeTemplateB);
-		NodeC = DecoratorWriter.RegisterNode(*NodeTemplateC);
-		NodeD = DecoratorWriter.RegisterNode(*NodeTemplateD);
+		NodeHandles.Add(DecoratorWriter.RegisterNode(*NodeTemplateA));
+		NodeHandles.Add(DecoratorWriter.RegisterNode(*NodeTemplateB));
+		NodeHandles.Add(DecoratorWriter.RegisterNode(*NodeTemplateC));
+		NodeHandles.Add(DecoratorWriter.RegisterNode(*NodeTemplateD));
 
 		// We don't have decorator properties
 		TArray<TMap<FString, FString>> DecoratorPropertiesA;
@@ -280,20 +282,20 @@ bool FAnimationAnimNextRuntimeTest_IHierarchy::RunTest(const FString& InParamete
 
 		TArray<TMap<FString, FString>> DecoratorPropertiesB;
 		DecoratorPropertiesB.AddDefaulted(NodeTemplateDecoratorListB.Num());
-		DecoratorPropertiesB[0].Add(TEXT("Child"), ToString<FDecoratorWithOneChild::FSharedData>(TEXT("Child"), FAnimNextDecoratorHandle(NodeA)));
+		DecoratorPropertiesB[0].Add(TEXT("Child"), ToString<FDecoratorWithOneChild::FSharedData>(TEXT("Child"), FAnimNextDecoratorHandle(NodeHandles[0])));
 
 		TArray<TMap<FString, FString>> DecoratorPropertiesC;
 		DecoratorPropertiesC.AddDefaulted(NodeTemplateDecoratorListC.Num());
-		DecoratorPropertiesC[0].Add(TEXT("Child"), ToString<FDecoratorWithOneChild::FSharedData>(TEXT("Child"), FAnimNextDecoratorHandle(NodeA)));
-		DecoratorPropertiesC[1].Add(TEXT("Child"), ToString<FDecoratorWithOneChild::FSharedData>(TEXT("Child"), FAnimNextDecoratorHandle(NodeB)));
+		DecoratorPropertiesC[0].Add(TEXT("Child"), ToString<FDecoratorWithOneChild::FSharedData>(TEXT("Child"), FAnimNextDecoratorHandle(NodeHandles[0])));
+		DecoratorPropertiesC[1].Add(TEXT("Child"), ToString<FDecoratorWithOneChild::FSharedData>(TEXT("Child"), FAnimNextDecoratorHandle(NodeHandles[1])));
 
 		TArray<TMap<FString, FString>> DecoratorPropertiesD;
 		DecoratorPropertiesD.AddDefaulted(NodeTemplateDecoratorListD.Num());
-		FAnimNextDecoratorHandle ChildrenHandlesD[2] = { FAnimNextDecoratorHandle(NodeA), FAnimNextDecoratorHandle(NodeC, 1) };
+		FAnimNextDecoratorHandle ChildrenHandlesD[2] = { FAnimNextDecoratorHandle(NodeHandles[0]), FAnimNextDecoratorHandle(NodeHandles[2], 1)};
 		DecoratorPropertiesD[0].Add(TEXT("Children"), ToString<FDecoratorWithChildren::FSharedData>(TEXT("Children"), ChildrenHandlesD));
 
 		DecoratorWriter.BeginNodeWriting();
-		DecoratorWriter.WriteNode(NodeA,
+		DecoratorWriter.WriteNode(NodeHandles[0],
 			[&DecoratorPropertiesA](uint32 DecoratorIndex, const FString& PropertyName)
 			{
 				return DecoratorPropertiesA[DecoratorIndex][PropertyName];
@@ -302,7 +304,7 @@ bool FAnimationAnimNextRuntimeTest_IHierarchy::RunTest(const FString& InParamete
 			{
 				return false;
 			});
-		DecoratorWriter.WriteNode(NodeB,
+		DecoratorWriter.WriteNode(NodeHandles[1],
 			[&DecoratorPropertiesB](uint32 DecoratorIndex, const FString& PropertyName)
 			{
 				return DecoratorPropertiesB[DecoratorIndex][PropertyName];
@@ -311,7 +313,7 @@ bool FAnimationAnimNextRuntimeTest_IHierarchy::RunTest(const FString& InParamete
 			{
 				return false;
 			});
-		DecoratorWriter.WriteNode(NodeC,
+		DecoratorWriter.WriteNode(NodeHandles[2],
 			[&DecoratorPropertiesC](uint32 DecoratorIndex, const FString& PropertyName)
 			{
 				return DecoratorPropertiesC[DecoratorIndex][PropertyName];
@@ -320,7 +322,7 @@ bool FAnimationAnimNextRuntimeTest_IHierarchy::RunTest(const FString& InParamete
 			{
 				return false;
 			});
-		DecoratorWriter.WriteNode(NodeD,
+		DecoratorWriter.WriteNode(NodeHandles[3],
 			[&DecoratorPropertiesD](uint32 DecoratorIndex, const FString& PropertyName)
 			{
 				return DecoratorPropertiesD[DecoratorIndex][PropertyName];
@@ -337,27 +339,18 @@ bool FAnimationAnimNextRuntimeTest_IHierarchy::RunTest(const FString& InParamete
 	}
 
 	// Read our graph
-	TArray<uint8> GraphSharedDataBuffer;
-	{
-		FMemoryReader GraphSharedDataArchive(GraphSharedDataArchiveBuffer);
-		FDecoratorReader DecoratorReader(GraphReferencedObjects, GraphSharedDataArchive);
-
-		FDecoratorReader::EErrorState ErrorState = DecoratorReader.ReadGraph(GraphSharedDataBuffer);
-		AddErrorIfFalse(ErrorState == FDecoratorReader::EErrorState::None, "FAnimationAnimNextRuntimeTest_IHierarchy -> Failed to read graph shared data");
+	FTestUtils::LoadFromArchiveBuffer(*AnimNextGraph, NodeHandles, GraphSharedDataArchiveBuffer);
 
-		NodeA = DecoratorReader.ResolveNodeHandle(NodeA);
-		NodeB = DecoratorReader.ResolveNodeHandle(NodeB);
-		NodeC = DecoratorReader.ResolveNodeHandle(NodeC);
-		NodeD = DecoratorReader.ResolveNodeHandle(NodeD);
-	}
+	FAnimNextGraphInstance GraphInstance;
+	AnimNextGraph->AllocateInstance(GraphInstance);
 
-	FExecutionContext Context(GraphSharedDataBuffer);
+	FExecutionContext Context(GraphInstance);
 
 	{
 		FMemMark Mark(FMemStack::Get());
 
-		FWeakDecoratorPtr NullPtr;				// Empty, no parent
-		FAnimNextDecoratorHandle RootHandle(NodeD);		// Point to NodeD, first base decorator
+		FWeakDecoratorPtr NullPtr;								// Empty, no parent
+		FAnimNextDecoratorHandle RootHandle(NodeHandles[3]);	// Point to NodeD, first base decorator
 
 		FDecoratorPtr NodeDPtr = Context.AllocateNodeInstance(NullPtr, RootHandle);
 		AddErrorIfFalse(NodeDPtr.IsValid(), "FAnimationAnimNextRuntimeTest_IHierarchy -> Failed to allocate root node instance");
@@ -369,8 +362,8 @@ bool FAnimationAnimNextRuntimeTest_IHierarchy::RunTest(const FString& InParamete
 		HierarchyBindingNodeD.GetChildren(Context, ChildrenNodeD);
 
 		AddErrorIfFalse(ChildrenNodeD.Num() == 2, "FAnimationAnimNextRuntimeTest_IHierarchy -> Expected 2 children");
-		AddErrorIfFalse(ChildrenNodeD[0].IsValid() && ChildrenNodeD[0].GetNodeInstance()->GetNodeHandle() == NodeA, "FAnimationAnimNextRuntimeTest_IHierarchy -> Expected child: NodeA");
-		AddErrorIfFalse(ChildrenNodeD[1].IsValid() && ChildrenNodeD[1].GetNodeInstance()->GetNodeHandle() == NodeC, "FAnimationAnimNextRuntimeTest_IHierarchy -> Expected child: NodeC");
+		AddErrorIfFalse(ChildrenNodeD[0].IsValid() && ChildrenNodeD[0].GetNodeInstance()->GetNodeHandle() == NodeHandles[0], "FAnimationAnimNextRuntimeTest_IHierarchy -> Expected child: NodeA");
+		AddErrorIfFalse(ChildrenNodeD[1].IsValid() && ChildrenNodeD[1].GetNodeInstance()->GetNodeHandle() == NodeHandles[2], "FAnimationAnimNextRuntimeTest_IHierarchy -> Expected child: NodeC");
 
 		{
 			TDecoratorBinding<IHierarchy> HierarchyBindingNodeC;
@@ -380,7 +373,7 @@ bool FAnimationAnimNextRuntimeTest_IHierarchy::RunTest(const FString& InParamete
 			HierarchyBindingNodeC.GetChildren(Context, ChildrenNodeC);
 
 			AddErrorIfFalse(ChildrenNodeC.Num() == 1, "FAnimationAnimNextRuntimeTest_IHierarchy -> Expected 2 children");
-			AddErrorIfFalse(ChildrenNodeC[0].IsValid() && ChildrenNodeC[0].GetNodeInstance()->GetNodeHandle() == NodeB, "FAnimationAnimNextRuntimeTest_IHierarchy -> Expected child: NodeB");
+			AddErrorIfFalse(ChildrenNodeC[0].IsValid() && ChildrenNodeC[0].GetNodeInstance()->GetNodeHandle() == NodeHandles[1], "FAnimationAnimNextRuntimeTest_IHierarchy -> Expected child: NodeB");
 
 			{
 				TDecoratorBinding<IHierarchy> HierarchyBindingNodeB;
@@ -390,7 +383,7 @@ bool FAnimationAnimNextRuntimeTest_IHierarchy::RunTest(const FString& InParamete
 				HierarchyBindingNodeB.GetChildren(Context, ChildrenNodeB);
 
 				AddErrorIfFalse(ChildrenNodeB.Num() == 1, "FAnimationAnimNextRuntimeTest_IHierarchy -> Expected 1 child");
-				AddErrorIfFalse(ChildrenNodeB[0].IsValid() && ChildrenNodeB[0].GetNodeInstance()->GetNodeHandle() == NodeA, "FAnimationAnimNextRuntimeTest_IHierarchy -> Expected child: NodeA");
+				AddErrorIfFalse(ChildrenNodeB[0].IsValid() && ChildrenNodeB[0].GetNodeInstance()->GetNodeHandle() == NodeHandles[0], "FAnimationAnimNextRuntimeTest_IHierarchy -> Expected child: NodeA");
 			}
 		}
 	}
@@ -415,6 +408,10 @@ bool FAnimationAnimNextRuntimeTest_IUpdate::RunTest(const FString& InParameters)
 	AUTO_REGISTER_ANIM_DECORATOR(FDecoratorWithOneChild)
 	AUTO_REGISTER_ANIM_DECORATOR(FDecoratorWithChildren)
 
+	UFactory* GraphFactory = NewObject<UAnimNextGraphFactory>();
+	UAnimNextGraph* AnimNextGraph = CastChecked<UAnimNextGraph>(GraphFactory->FactoryCreateNew(UAnimNextGraph::StaticClass(), GetTransientPackage(), TEXT("TestAnimNextGraph"), RF_Transient, nullptr, nullptr, NAME_None));
+	UE_RETURN_ON_ERROR(AnimNextGraph != nullptr, "FAnimationAnimNextRuntimeTest_IUpdate -> Failed to create animation graph");
+
 	FScopedClearNodeTemplateRegistry ScopedClearNodeTemplateRegistry;
 	FNodeTemplateRegistry& Registry = FNodeTemplateRegistry::Get();
 
@@ -442,9 +439,7 @@ bool FAnimationAnimNextRuntimeTest_IUpdate::RunTest(const FString& InParameters)
 	// NodeB has one child: NodeA (it doesn't update)
 	// NodeD has two children: NodeA and NodeB
 
-	FNodeHandle NodeA;
-	FNodeHandle NodeB;
-	FNodeHandle NodeC;
+	TArray<FNodeHandle> NodeHandles;
 
 	// Write our graph
 	TArray<uint8> GraphSharedDataArchiveBuffer;
@@ -452,9 +447,9 @@ bool FAnimationAnimNextRuntimeTest_IUpdate::RunTest(const FString& InParameters)
 	{
 		FDecoratorWriter DecoratorWriter;
 
-		NodeA = DecoratorWriter.RegisterNode(*NodeTemplateA);
-		NodeB = DecoratorWriter.RegisterNode(*NodeTemplateB);
-		NodeC = DecoratorWriter.RegisterNode(*NodeTemplateC);
+		NodeHandles.Add(DecoratorWriter.RegisterNode(*NodeTemplateA));
+		NodeHandles.Add(DecoratorWriter.RegisterNode(*NodeTemplateB));
+		NodeHandles.Add(DecoratorWriter.RegisterNode(*NodeTemplateC));
 
 		// We don't have decorator properties
 		TArray<TMap<FString, FString>> DecoratorPropertiesA;
@@ -462,15 +457,15 @@ bool FAnimationAnimNextRuntimeTest_IUpdate::RunTest(const FString& InParameters)
 
 		TArray<TMap<FString, FString>> DecoratorPropertiesB;
 		DecoratorPropertiesB.AddDefaulted(NodeTemplateDecoratorListB.Num());
-		DecoratorPropertiesB[0].Add(TEXT("Child"), ToString<FDecoratorWithOneChild::FSharedData>(TEXT("Child"), FAnimNextDecoratorHandle(NodeA)));
+		DecoratorPropertiesB[0].Add(TEXT("Child"), ToString<FDecoratorWithOneChild::FSharedData>(TEXT("Child"), FAnimNextDecoratorHandle(NodeHandles[0])));
 
 		TArray<TMap<FString, FString>> DecoratorPropertiesC;
 		DecoratorPropertiesC.AddDefaulted(NodeTemplateDecoratorListC.Num());
-		FAnimNextDecoratorHandle ChildrenHandlesC[2] = { FAnimNextDecoratorHandle(NodeA), FAnimNextDecoratorHandle(NodeB) };
+		FAnimNextDecoratorHandle ChildrenHandlesC[2] = { FAnimNextDecoratorHandle(NodeHandles[0]), FAnimNextDecoratorHandle(NodeHandles[1])};
 		DecoratorPropertiesC[0].Add(TEXT("Children"), ToString<FDecoratorWithChildren::FSharedData>(TEXT("Children"), ChildrenHandlesC));
 
 		DecoratorWriter.BeginNodeWriting();
-		DecoratorWriter.WriteNode(NodeA,
+		DecoratorWriter.WriteNode(NodeHandles[0],
 			[&DecoratorPropertiesA](uint32 DecoratorIndex, const FString& PropertyName)
 			{
 				return DecoratorPropertiesA[DecoratorIndex][PropertyName];
@@ -479,7 +474,7 @@ bool FAnimationAnimNextRuntimeTest_IUpdate::RunTest(const FString& InParameters)
 			{
 				return false;
 			});
-		DecoratorWriter.WriteNode(NodeB,
+		DecoratorWriter.WriteNode(NodeHandles[1],
 			[&DecoratorPropertiesB](uint32 DecoratorIndex, const FString& PropertyName)
 			{
 				return DecoratorPropertiesB[DecoratorIndex][PropertyName];
@@ -488,7 +483,7 @@ bool FAnimationAnimNextRuntimeTest_IUpdate::RunTest(const FString& InParameters)
 			{
 				return false;
 			});
-		DecoratorWriter.WriteNode(NodeC,
+		DecoratorWriter.WriteNode(NodeHandles[2],
 			[&DecoratorPropertiesC](uint32 DecoratorIndex, const FString& PropertyName)
 			{
 				return DecoratorPropertiesC[DecoratorIndex][PropertyName];
@@ -505,28 +500,20 @@ bool FAnimationAnimNextRuntimeTest_IUpdate::RunTest(const FString& InParameters)
 	}
 
 	// Read our graph
-	TArray<uint8> GraphSharedDataBuffer;
-	{
-		FMemoryReader GraphSharedDataArchive(GraphSharedDataArchiveBuffer);
-		FDecoratorReader DecoratorReader(GraphReferencedObjects, GraphSharedDataArchive);
+	FTestUtils::LoadFromArchiveBuffer(*AnimNextGraph, NodeHandles, GraphSharedDataArchiveBuffer);
 
-		FDecoratorReader::EErrorState ErrorState = DecoratorReader.ReadGraph(GraphSharedDataBuffer);
-		AddErrorIfFalse(ErrorState == FDecoratorReader::EErrorState::None, "FAnimationAnimNextRuntimeTest_IUpdate -> Failed to read graph shared data");
+	FAnimNextGraphInstance GraphInstance;
+	AnimNextGraph->AllocateInstance(GraphInstance);
 
-		NodeA = DecoratorReader.ResolveNodeHandle(NodeA);
-		NodeB = DecoratorReader.ResolveNodeHandle(NodeB);
-		NodeC = DecoratorReader.ResolveNodeHandle(NodeC);
-	}
-
-	FExecutionContext Context(GraphSharedDataBuffer);
+	FExecutionContext Context(GraphInstance);
 
 	{
 		TArray<FDecoratorUID> UpdatedDecorators;
 
 		Private::UpdatedDecorators = &UpdatedDecorators;
 
-		FWeakDecoratorPtr NullPtr;				// Empty, no parent
-		FAnimNextDecoratorHandle RootHandle(NodeC);		// Point to NodeC, first base decorator
+		FWeakDecoratorPtr NullPtr;								// Empty, no parent
+		FAnimNextDecoratorHandle RootHandle(NodeHandles[2]);	// Point to NodeC, first base decorator
 
 		FDecoratorPtr NodeCPtr = Context.AllocateNodeInstance(NullPtr, RootHandle);
 		AddErrorIfFalse(NodeCPtr.IsValid(), "FAnimationAnimNextRuntimeTest_IUpdate -> Failed to allocate root node instance");
@@ -564,6 +551,10 @@ bool FAnimationAnimNextRuntimeTest_IEvaluate::RunTest(const FString& InParameter
 	AUTO_REGISTER_ANIM_DECORATOR(FDecoratorWithOneChild)
 	AUTO_REGISTER_ANIM_DECORATOR(FDecoratorWithChildren)
 
+	UFactory* GraphFactory = NewObject<UAnimNextGraphFactory>();
+	UAnimNextGraph* AnimNextGraph = CastChecked<UAnimNextGraph>(GraphFactory->FactoryCreateNew(UAnimNextGraph::StaticClass(), GetTransientPackage(), TEXT("TestAnimNextGraph"), RF_Transient, nullptr, nullptr, NAME_None));
+	UE_RETURN_ON_ERROR(AnimNextGraph != nullptr, "FAnimationAnimNextRuntimeTest_IEvaluate -> Failed to create animation graph");
+
 	FScopedClearNodeTemplateRegistry ScopedClearNodeTemplateRegistry;
 	FNodeTemplateRegistry& Registry = FNodeTemplateRegistry::Get();
 
@@ -591,9 +582,7 @@ bool FAnimationAnimNextRuntimeTest_IEvaluate::RunTest(const FString& InParameter
 	// NodeB has one child: NodeA (it doesn't evaluate)
 	// NodeD has two children: NodeA and NodeB
 
-	FNodeHandle NodeA;
-	FNodeHandle NodeB;
-	FNodeHandle NodeC;
+	TArray<FNodeHandle> NodeHandles;
 
 	// Write our graph
 	TArray<uint8> GraphSharedDataArchiveBuffer;
@@ -601,9 +590,9 @@ bool FAnimationAnimNextRuntimeTest_IEvaluate::RunTest(const FString& InParameter
 	{
 		FDecoratorWriter DecoratorWriter;
 
-		NodeA = DecoratorWriter.RegisterNode(*NodeTemplateA);
-		NodeB = DecoratorWriter.RegisterNode(*NodeTemplateB);
-		NodeC = DecoratorWriter.RegisterNode(*NodeTemplateC);
+		NodeHandles.Add(DecoratorWriter.RegisterNode(*NodeTemplateA));
+		NodeHandles.Add(DecoratorWriter.RegisterNode(*NodeTemplateB));
+		NodeHandles.Add(DecoratorWriter.RegisterNode(*NodeTemplateC));
 
 		// We don't have decorator properties
 		TArray<TMap<FString, FString>> DecoratorPropertiesA;
@@ -611,16 +600,16 @@ bool FAnimationAnimNextRuntimeTest_IEvaluate::RunTest(const FString& InParameter
 
 		TArray<TMap<FString, FString>> DecoratorPropertiesB;
 		DecoratorPropertiesB.AddDefaulted(NodeTemplateDecoratorListB.Num());
-		DecoratorPropertiesB[0].Add(TEXT("Child"), ToString<FDecoratorWithOneChild::FSharedData>(TEXT("Child"), FAnimNextDecoratorHandle(NodeA)));
+		DecoratorPropertiesB[0].Add(TEXT("Child"), ToString<FDecoratorWithOneChild::FSharedData>(TEXT("Child"), FAnimNextDecoratorHandle(NodeHandles[0])));
 
 		TArray<TMap<FString, FString>> DecoratorPropertiesC;
 		DecoratorPropertiesC.AddDefaulted(NodeTemplateDecoratorListC.Num());
 
-		FAnimNextDecoratorHandle ChildrenHandlesC[2] = { FAnimNextDecoratorHandle(NodeA), FAnimNextDecoratorHandle(NodeB) };
+		FAnimNextDecoratorHandle ChildrenHandlesC[2] = { FAnimNextDecoratorHandle(NodeHandles[0]), FAnimNextDecoratorHandle(NodeHandles[1])};
 		DecoratorPropertiesC[0].Add(TEXT("Children"), ToString<FDecoratorWithChildren::FSharedData>(TEXT("Children"), ChildrenHandlesC));
 
 		DecoratorWriter.BeginNodeWriting();
-		DecoratorWriter.WriteNode(NodeA,
+		DecoratorWriter.WriteNode(NodeHandles[0],
 			[&DecoratorPropertiesA](uint32 DecoratorIndex, const FString& PropertyName)
 			{
 				return DecoratorPropertiesA[DecoratorIndex][PropertyName];
@@ -629,7 +618,7 @@ bool FAnimationAnimNextRuntimeTest_IEvaluate::RunTest(const FString& InParameter
 			{
 				return false;
 			});
-		DecoratorWriter.WriteNode(NodeB,
+		DecoratorWriter.WriteNode(NodeHandles[1],
 			[&DecoratorPropertiesB](uint32 DecoratorIndex, const FString& PropertyName)
 			{
 				return DecoratorPropertiesB[DecoratorIndex][PropertyName];
@@ -638,7 +627,7 @@ bool FAnimationAnimNextRuntimeTest_IEvaluate::RunTest(const FString& InParameter
 			{
 				return false;
 			});
-		DecoratorWriter.WriteNode(NodeC,
+		DecoratorWriter.WriteNode(NodeHandles[2],
 			[&DecoratorPropertiesC](uint32 DecoratorIndex, const FString& PropertyName)
 			{
 				return DecoratorPropertiesC[DecoratorIndex][PropertyName];
@@ -655,28 +644,20 @@ bool FAnimationAnimNextRuntimeTest_IEvaluate::RunTest(const FString& InParameter
 	}
 
 	// Read our graph
-	TArray<uint8> GraphSharedDataBuffer;
-	{
-		FMemoryReader GraphSharedDataArchive(GraphSharedDataArchiveBuffer);
-		FDecoratorReader DecoratorReader(GraphReferencedObjects, GraphSharedDataArchive);
-
-		FDecoratorReader::EErrorState ErrorState = DecoratorReader.ReadGraph(GraphSharedDataBuffer);
-		AddErrorIfFalse(ErrorState == FDecoratorReader::EErrorState::None, "FAnimationAnimNextRuntimeTest_IEvaluate -> Failed to read graph shared data");
+	FTestUtils::LoadFromArchiveBuffer(*AnimNextGraph, NodeHandles, GraphSharedDataArchiveBuffer);
 
-		NodeA = DecoratorReader.ResolveNodeHandle(NodeA);
-		NodeB = DecoratorReader.ResolveNodeHandle(NodeB);
-		NodeC = DecoratorReader.ResolveNodeHandle(NodeC);
-	}
+	FAnimNextGraphInstance GraphInstance;
+	AnimNextGraph->AllocateInstance(GraphInstance);
 
-	FExecutionContext Context(GraphSharedDataBuffer);
+	FExecutionContext Context(GraphInstance);
 
 	{
 		TArray<FDecoratorUID> EvaluatedDecorators;
 
 		Private::EvaluatedDecorators = &EvaluatedDecorators;
 
-		FWeakDecoratorPtr NullPtr;				// Empty, no parent
-		FAnimNextDecoratorHandle RootHandle(NodeC);		// Point to NodeC, first base decorator
+		FWeakDecoratorPtr NullPtr;								// Empty, no parent
+		FAnimNextDecoratorHandle RootHandle(NodeHandles[2]);	// Point to NodeC, first base decorator
 
 		FDecoratorPtr NodeCPtr = Context.AllocateNodeInstance(NullPtr, RootHandle);
 		AddErrorIfFalse(NodeCPtr.IsValid(), "FAnimationAnimNextRuntimeTest_IEvaluate -> Failed to allocate root node instance");
diff --git a/Engine/Plugins/Experimental/AnimNext/Source/AnimNextTestSuite/Private/AnimNextRuntimeTest.cpp b/Engine/Plugins/Experimental/AnimNext/Source/AnimNextTestSuite/Private/AnimNextRuntimeTest.cpp
--- a/Engine/Plugins/Experimental/AnimNext/Source/AnimNextTestSuite/Private/AnimNextRuntimeTest.cpp
+++ b/Engine/Plugins/Experimental/AnimNext/Source/AnimNextTestSuite/Private/AnimNextRuntimeTest.cpp
@@ -3,7 +3,10 @@
 
 #include "AnimNextRuntimeTest.h"
 
+#include "DecoratorBase/DecoratorReader.h"
+#include "Graph/AnimNextGraph.h"
 #include "Misc/AutomationTest.h"
+#include "Serialization/MemoryReader.h"
 
 #if WITH_DEV_AUTOMATION_TESTS
 
@@ -20,5 +23,42 @@ namespace UE::AnimNext
 		FNodeTemplateRegistry& Registry = FNodeTemplateRegistry::Get();
 		Swap(Registry, TmpRegistry);
 	}
+
+	bool FTestUtils::LoadFromArchiveBuffer(UAnimNextGraph& Graph, TArray<FNodeHandle>& NodeHandles, const TArray<uint8>& SharedDataArchiveBuffer)
+	{
+		FAnimNextGraphEvaluatorExecuteDefinition ExecuteDefinition;
+		ExecuteDefinition.Hash = 0;
+		ExecuteDefinition.MethodName = TEXT("Execute_0");
+
+		Graph.ExecuteDefinition = ExecuteDefinition;
+		Graph.SharedDataArchiveBuffer = SharedDataArchiveBuffer;
+		Graph.GraphReferencedObjects.Empty();
+		Graph.RootDecoratorHandle = FAnimNextEntryPointHandle(NodeHandles[0]);
+
+		// Reconstruct our graph shared data
+		FMemoryReader GraphSharedDataArchive(SharedDataArchiveBuffer);
+		FDecoratorReader DecoratorReader(Graph.GraphReferencedObjects, GraphSharedDataArchive);
+
+		const FDecoratorReader::EErrorState ErrorState = DecoratorReader.ReadGraph(Graph.SharedDataBuffer);
+		if (ErrorState == FDecoratorReader::EErrorState::None)
+		{
+			Graph.ResolvedRootDecoratorHandle = DecoratorReader.ResolveEntryPointHandle(Graph.RootDecoratorHandle);
+
+			for (FNodeHandle& NodeHandle : NodeHandles)
+			{
+				NodeHandle = DecoratorReader.ResolveNodeHandle(NodeHandle);
+			}
+
+			// Make sure our execute method is registered
+			FRigUnit_AnimNextGraphEvaluator::RegisterExecuteMethod(ExecuteDefinition);
+			return true;
+		}
+		else
+		{
+			Graph.SharedDataBuffer.Empty(0);
+			Graph.ResolvedRootDecoratorHandle = FAnimNextDecoratorHandle();
+			return false;
+		}
+	}
 }
 #endif
diff --git a/Engine/Plugins/Experimental/AnimNext/Source/AnimNextTestSuite/Private/AnimNextRuntimeTest.h b/Engine/Plugins/Experimental/AnimNext/Source/AnimNextTestSuite/Private/AnimNextRuntimeTest.h
--- a/Engine/Plugins/Experimental/AnimNext/Source/AnimNextTestSuite/Private/AnimNextRuntimeTest.h
+++ b/Engine/Plugins/Experimental/AnimNext/Source/AnimNextTestSuite/Private/AnimNextRuntimeTest.h
@@ -13,8 +13,12 @@
 #include "DecoratorBase/NodeTemplate.h"
 #include "DecoratorBase/NodeTemplateRegistry.h"
 
+class UAnimNextGraph;
+
 namespace UE::AnimNext
 {
+	struct FNodeHandle;
+
 	// Converts a property value into its string representation using UE reflection
 	template<class DecoratorSharedDataType, typename PropertyType>
 	static FString ToString(const FString& PropertyName, PropertyType PropertyValue)
@@ -75,6 +79,13 @@ namespace UE::AnimNext
 
 		FNodeTemplateRegistry TmpRegistry;
 	};
+
+	struct FTestUtils final
+	{
+		// Loads the graph data from the provided archive buffer and returns true on success, false otherwise
+		// On success, we resolve every node handle provided as argument
+		static bool LoadFromArchiveBuffer(UAnimNextGraph& Graph, TArray<FNodeHandle>& NodeHandles, const TArray<uint8>& SharedDataArchiveBuffer);
+	};
 }
 
 #endif
diff --git a/Engine/Source/Runtime/Engine/Private/Animation/AnimationRuntime.cpp b/Engine/Source/Runtime/Engine/Private/Animation/AnimationRuntime.cpp
--- a/Engine/Source/Runtime/Engine/Private/Animation/AnimationRuntime.cpp
+++ b/Engine/Source/Runtime/Engine/Private/Animation/AnimationRuntime.cpp
@@ -1438,40 +1438,42 @@ void FAnimationRuntime::MirrorPose(FCompactPose& Pose, const UMirrorDataTable& M
 /** 
  * return ETypeAdvanceAnim type
  */
-ETypeAdvanceAnim FAnimationRuntime::AdvanceTime(const bool& bAllowLooping, const float& MoveDelta, float& InOutTime, const float& EndTime)
+ETypeAdvanceAnim FAnimationRuntime::AdvanceTime(const bool bAllowLooping, const float MoveDelta, float& InOutTime, const float EndTime)
 {
-	InOutTime += MoveDelta;
+	float NewTime = InOutTime + MoveDelta;
 
-	if( InOutTime < 0.f || InOutTime > EndTime )
+	if (NewTime < 0.f || NewTime > EndTime)
 	{
-		if( bAllowLooping )
+		if (bAllowLooping)
 		{
-			if( EndTime != 0.f )
+			if (EndTime != 0.f)
 			{
-				InOutTime	= FMath::Fmod(InOutTime, EndTime);
+				NewTime = FMath::Fmod(NewTime, EndTime);
 				// Fmod doesn't give result that falls into (0, EndTime), but one that falls into (-EndTime, EndTime). Negative values need to be handled in custom way
-				if( InOutTime < 0.f )
+				if (NewTime < 0.f)
 				{
-					InOutTime += EndTime;
+					NewTime += EndTime;
 				}
 			}
 			else
 			{
 				// end time is 0.f
-				InOutTime = 0.f;
+				NewTime = 0.f;
 			}
 
 			// it has been looped
+			InOutTime = NewTime;
 			return ETAA_Looped;
 		}
 		else 
 		{
 			// If not, snap time to end of sequence and stop playing.
-			InOutTime = FMath::Clamp(InOutTime, 0.f, EndTime);
+			InOutTime = FMath::Clamp(NewTime, 0.f, EndTime);
 			return ETAA_Finished;
 		}
 	}
 
+	InOutTime = NewTime;
 	return ETAA_Default;
 }
 
diff --git a/Engine/Source/Runtime/Engine/Public/AnimationRuntime.h b/Engine/Source/Runtime/Engine/Public/AnimationRuntime.h
--- a/Engine/Source/Runtime/Engine/Public/AnimationRuntime.h
+++ b/Engine/Source/Runtime/Engine/Public/AnimationRuntime.h
@@ -545,7 +545,7 @@ public:
 	 *
 	 * return ETypeAdvanceAnim type
 	 */
-	static ENGINE_API ETypeAdvanceAnim AdvanceTime(const bool& bAllowLooping, const float& MoveDelta, float& InOutTime, const float& EndTime);
+	static ENGINE_API ETypeAdvanceAnim AdvanceTime(const bool bAllowLooping, const float MoveDelta, float& InOutTime, const float EndTime);
 
 	static ENGINE_API void TickBlendWeight(float DeltaTime, float DesiredWeight, float& Weight, float& BlendTime);
 	/** 
