Migrate AnimNext evaluate traversal to use new proxy pattern

diff --git a/Engine/Plugins/Experimental/AnimNext/Source/AnimNext/Private/DecoratorBase/ITraversalContext.cpp b/Engine/Plugins/Experimental/AnimNext/Source/AnimNext/Private/DecoratorBase/ITraversalContext.cpp
deleted file mode 100644
diff --git a/Engine/Plugins/Experimental/AnimNext/Source/AnimNext/Private/DecoratorInterfaces/IEvaluate.cpp b/Engine/Plugins/Experimental/AnimNext/Source/AnimNext/Private/DecoratorInterfaces/IEvaluate.cpp
--- a/Engine/Plugins/Experimental/AnimNext/Source/AnimNext/Private/DecoratorInterfaces/IEvaluate.cpp
+++ b/Engine/Plugins/Experimental/AnimNext/Source/AnimNext/Private/DecoratorInterfaces/IEvaluate.cpp
@@ -10,38 +10,13 @@ DEFINE_STAT(STAT_AnimNext_EvaluateGraph);
 
 namespace UE::AnimNext
 {
-	/**
-	 * FScopedEvaluationProgram
-	 *
-	 * Pushes and pops a evaluation program onto the evaluation traversal context.
-	 * During the lifetime of an instance of this object, the traversal context will
-	 * return the provided evaluation program as the current evaluation program.
-	 *
-	 * @see FEvaluateTraversalContext
-	 */
-	struct ANIMNEXT_API FScopedEvaluationProgram final
+	FEvaluateTraversalContext::FEvaluateTraversalContext(const FExecutionContext& InExecutionContext, FEvaluationProgram& InEvaluationProgram)
+		: FExecutionContextProxy(InExecutionContext)
+		, EvaluationProgram(InEvaluationProgram)
 	{
-		FScopedEvaluationProgram(FEvaluateTraversalContext& InTraversalContext, FEvaluationProgram& InEvaluationProgram)
-			: TraversalContext(InTraversalContext)
-			, OldEvaluationProgram(InTraversalContext.EvaluationProgram)
-		{
-			InTraversalContext.EvaluationProgram = &InEvaluationProgram;
-		}
-
-		~FScopedEvaluationProgram()
-		{
-			TraversalContext.EvaluationProgram = OldEvaluationProgram;
-		}
-
-		FScopedEvaluationProgram(const FScopedEvaluationProgram&) = delete;
-		FScopedEvaluationProgram& operator=(const FScopedEvaluationProgram&) = delete;
-
-	private:
-		FEvaluateTraversalContext& TraversalContext;
-		FEvaluationProgram* OldEvaluationProgram;
-	};
+	}
 
-	void IEvaluate::PreEvaluate(const FExecutionContext& Context, const TDecoratorBinding<IEvaluate>& Binding) const
+	void IEvaluate::PreEvaluate(FEvaluateTraversalContext& Context, const TDecoratorBinding<IEvaluate>& Binding) const
 	{
 		TDecoratorBinding<IEvaluate> SuperBinding;
 		if (Context.GetInterfaceSuper(Binding, SuperBinding))
@@ -50,7 +25,7 @@ namespace UE::AnimNext
 		}
 	}
 
-	void IEvaluate::PostEvaluate(const FExecutionContext& Context, const TDecoratorBinding<IEvaluate>& Binding) const
+	void IEvaluate::PostEvaluate(FEvaluateTraversalContext& Context, const TDecoratorBinding<IEvaluate>& Binding) const
 	{
 		TDecoratorBinding<IEvaluate> SuperBinding;
 		if (Context.GetInterfaceSuper(Binding, SuperBinding))
@@ -96,7 +71,7 @@ namespace UE::AnimNext
 		}
 	};
 
-	FEvaluationProgram EvaluateGraph(FExecutionContext& Context, FEvaluateTraversalContext& TraversalContext, FWeakDecoratorPtr GraphRootPtr)
+	FEvaluationProgram EvaluateGraph(FExecutionContext& Context, FWeakDecoratorPtr GraphRootPtr)
 	{
 		SCOPE_CYCLE_COUNTER(STAT_AnimNext_EvaluateGraph);
 		
@@ -113,8 +88,7 @@ namespace UE::AnimNext
 		FChildrenArray Children;
 		Children.Reserve(16);
 
-		FScopedEvaluationProgram ScopedEvaluationProgram(TraversalContext, EvaluationProgram);
-		FScopedTraversalContext ScopedTraversalContext(Context, TraversalContext);
+		FEvaluateTraversalContext TraversalContext(Context, EvaluationProgram);
 
 		// Add the graph root to kick start the evaluation process
 		FEvaluateEntry GraphRootEntry(GraphRootPtr, EEvaluateStep::PreEvaluate);
@@ -137,7 +111,7 @@ namespace UE::AnimNext
 				if (Context.GetInterface(Entry->DecoratorPtr, Entry->EvaluateDecorator))
 				{
 					// This is the first time we visit this node, time to pre-evaluate
-					Entry->EvaluateDecorator.PreEvaluate(Context);
+					Entry->EvaluateDecorator.PreEvaluate(TraversalContext);
 
 					// Leave our entry on top of the stack, we'll need to call PostEvaluate once the children
 					// we'll push on top finish
@@ -190,7 +164,7 @@ namespace UE::AnimNext
 			{
 				// We've already visited this node once, time to post-update
 				check(Entry->EvaluateDecorator.IsValid());
-				Entry->EvaluateDecorator.PostEvaluate(Context);
+				Entry->EvaluateDecorator.PostEvaluate(TraversalContext);
 
 				// Now that we are done processing this entry, we can pop it
 				NodesPendingUpdateStackTop = Entry->PrevStackEntry;
diff --git a/Engine/Plugins/Experimental/AnimNext/Source/AnimNext/Private/Decorators/BlendByBool.cpp b/Engine/Plugins/Experimental/AnimNext/Source/AnimNext/Private/Decorators/BlendByBool.cpp
--- a/Engine/Plugins/Experimental/AnimNext/Source/AnimNext/Private/Decorators/BlendByBool.cpp
+++ b/Engine/Plugins/Experimental/AnimNext/Source/AnimNext/Private/Decorators/BlendByBool.cpp
@@ -19,7 +19,7 @@ namespace UE::AnimNext
 	static constexpr int32 TRUE_CHILD_INDEX = 0;
 	static constexpr int32 FALSE_CHILD_INDEX = 1;
 
-	void FBlendByBoolDecorator::PostEvaluate(const FExecutionContext& Context, const TDecoratorBinding<IEvaluate>& Binding) const
+	void FBlendByBoolDecorator::PostEvaluate(FEvaluateTraversalContext& Context, const TDecoratorBinding<IEvaluate>& Binding) const
 	{
 		const FInstanceData* InstanceData = Binding.GetInstanceData<FInstanceData>();
 
@@ -31,9 +31,7 @@ namespace UE::AnimNext
 			Context.GetInterface(Binding, DiscreteBlendDecorator);
 
 			const float BlendWeight = DiscreteBlendDecorator.GetBlendWeight(Context, FALSE_CHILD_INDEX);
-
-			FEvaluateTraversalContext& TraversalContext = Context.GetTraversalContext<FEvaluateTraversalContext>();
-			TraversalContext.AppendTask(FAnimNextBlendTwoKeyframesTask::Make(BlendWeight));
+			Context.AppendTask(FAnimNextBlendTwoKeyframesTask::Make(BlendWeight));
 		}
 		else
 		{
diff --git a/Engine/Plugins/Experimental/AnimNext/Source/AnimNext/Private/Decorators/BlendByBool.h b/Engine/Plugins/Experimental/AnimNext/Source/AnimNext/Private/Decorators/BlendByBool.h
--- a/Engine/Plugins/Experimental/AnimNext/Source/AnimNext/Private/Decorators/BlendByBool.h
+++ b/Engine/Plugins/Experimental/AnimNext/Source/AnimNext/Private/Decorators/BlendByBool.h
@@ -59,7 +59,7 @@ namespace UE::AnimNext
 		};
 
 		// IEvaluate impl
-		virtual void PostEvaluate(const FExecutionContext& Context, const TDecoratorBinding<IEvaluate>& Binding) const override;
+		virtual void PostEvaluate(FEvaluateTraversalContext& Context, const TDecoratorBinding<IEvaluate>& Binding) const override;
 
 		// IUpdate impl
 		virtual void PreUpdate(FUpdateTraversalContext& Context, const TDecoratorBinding<IUpdate>& Binding, const FDecoratorUpdateState& DecoratorState) const override;
diff --git a/Engine/Plugins/Experimental/AnimNext/Source/AnimNext/Private/Decorators/BlendSmootherPerBone.cpp b/Engine/Plugins/Experimental/AnimNext/Source/AnimNext/Private/Decorators/BlendSmootherPerBone.cpp
--- a/Engine/Plugins/Experimental/AnimNext/Source/AnimNext/Private/Decorators/BlendSmootherPerBone.cpp
+++ b/Engine/Plugins/Experimental/AnimNext/Source/AnimNext/Private/Decorators/BlendSmootherPerBone.cpp
@@ -19,7 +19,7 @@ namespace UE::AnimNext
 		DEFINE_ANIM_DECORATOR_IMPLEMENTS_INTERFACE(IDiscreteBlend)
 	DEFINE_ANIM_DECORATOR_END(FBlendSmootherPerBoneDecorator)
 
-	void FBlendSmootherPerBoneDecorator::PostEvaluate(const FExecutionContext& Context, const TDecoratorBinding<IEvaluate>& Binding) const
+	void FBlendSmootherPerBoneDecorator::PostEvaluate(FEvaluateTraversalContext& Context, const TDecoratorBinding<IEvaluate>& Binding) const
 	{
 		const FSharedData* SharedData = Binding.GetSharedData<FSharedData>();
 		const FInstanceData* InstanceData = Binding.GetInstanceData<FInstanceData>();
@@ -39,8 +39,6 @@ namespace UE::AnimNext
 			return;	// If we don't have at least 2 children, there is nothing to do
 		}
 
-		FEvaluateTraversalContext& TraversalContext = Context.GetTraversalContext<FEvaluateTraversalContext>();
-
 		// Children are visited depth first, in the order returned
 		// As such, when we evaluate the task program, the keyframe of the last child will be
 		// on top of the keyframe stack
@@ -51,7 +49,7 @@ namespace UE::AnimNext
 			const int32 ChildIndex = NumChildren - 1;
 			const FBlendSampleData& PoseSampleData = InstanceData->PerBoneSampleData[ChildIndex];
 
-			TraversalContext.AppendTask(FAnimNextBlendOverwriteKeyframePerBoneWithScaleTask::Make(SharedData->BlendProfile, PoseSampleData, PoseSampleData.TotalWeight));
+			Context.AppendTask(FAnimNextBlendOverwriteKeyframePerBoneWithScaleTask::Make(SharedData->BlendProfile, PoseSampleData, PoseSampleData.TotalWeight));
 		}
 
 		// Other children accumulate with scale
@@ -60,11 +58,11 @@ namespace UE::AnimNext
 			const FBlendSampleData& PoseSampleDataA = InstanceData->PerBoneSampleData[ChildIndex];
 			const FBlendSampleData& PoseSampleDataB = InstanceData->PerBoneSampleData[ChildIndex + 1];	// Above on the keyframe stack
 
-			TraversalContext.AppendTask(FAnimNextBlendAddKeyframePerBoneWithScaleTask::Make(SharedData->BlendProfile, PoseSampleDataA, PoseSampleDataB, PoseSampleDataA.TotalWeight));
+			Context.AppendTask(FAnimNextBlendAddKeyframePerBoneWithScaleTask::Make(SharedData->BlendProfile, PoseSampleDataA, PoseSampleDataB, PoseSampleDataA.TotalWeight));
 		}
 
 		// Once we are done, we normalize rotations
-		TraversalContext.AppendTask(FAnimNextNormalizeKeyframeRotationsTask());
+		Context.AppendTask(FAnimNextNormalizeKeyframeRotationsTask());
 	}
 
 	void FBlendSmootherPerBoneDecorator::PreUpdate(FUpdateTraversalContext& Context, const TDecoratorBinding<IUpdate>& Binding, const FDecoratorUpdateState& DecoratorState) const
diff --git a/Engine/Plugins/Experimental/AnimNext/Source/AnimNext/Private/Decorators/BlendSmootherPerBone.h b/Engine/Plugins/Experimental/AnimNext/Source/AnimNext/Private/Decorators/BlendSmootherPerBone.h
--- a/Engine/Plugins/Experimental/AnimNext/Source/AnimNext/Private/Decorators/BlendSmootherPerBone.h
+++ b/Engine/Plugins/Experimental/AnimNext/Source/AnimNext/Private/Decorators/BlendSmootherPerBone.h
@@ -56,7 +56,7 @@ namespace UE::AnimNext
 		};
 
 		// IEvaluate impl
-		virtual void PostEvaluate(const FExecutionContext& Context, const TDecoratorBinding<IEvaluate>& Binding) const override;
+		virtual void PostEvaluate(FEvaluateTraversalContext& Context, const TDecoratorBinding<IEvaluate>& Binding) const override;
 
 		// IUpdate impl
 		virtual void PreUpdate(FUpdateTraversalContext& Context, const TDecoratorBinding<IUpdate>& Binding, const FDecoratorUpdateState& DecoratorState) const override;
diff --git a/Engine/Plugins/Experimental/AnimNext/Source/AnimNext/Private/Decorators/BlendTwoWay.cpp b/Engine/Plugins/Experimental/AnimNext/Source/AnimNext/Private/Decorators/BlendTwoWay.cpp
--- a/Engine/Plugins/Experimental/AnimNext/Source/AnimNext/Private/Decorators/BlendTwoWay.cpp
+++ b/Engine/Plugins/Experimental/AnimNext/Source/AnimNext/Private/Decorators/BlendTwoWay.cpp
@@ -17,7 +17,7 @@ namespace UE::AnimNext
 		DEFINE_ANIM_DECORATOR_IMPLEMENTS_INTERFACE(IUpdate)
 	DEFINE_ANIM_DECORATOR_END(FBlendTwoWayDecorator)
 
-	void FBlendTwoWayDecorator::PostEvaluate(const FExecutionContext& Context, const TDecoratorBinding<IEvaluate>& Binding) const
+	void FBlendTwoWayDecorator::PostEvaluate(FEvaluateTraversalContext& Context, const TDecoratorBinding<IEvaluate>& Binding) const
 	{
 		const FInstanceData* InstanceData = Binding.GetInstanceData<FInstanceData>();
 
@@ -29,9 +29,7 @@ namespace UE::AnimNext
 			Context.GetInterface(Binding, ContinuousBlendDecorator);
 
 			const float BlendWeight = ContinuousBlendDecorator.GetBlendWeight(Context, 1);
-
-			FEvaluateTraversalContext& TraversalContext = Context.GetTraversalContext<FEvaluateTraversalContext>();
-			TraversalContext.AppendTask(FAnimNextBlendTwoKeyframesTask::Make(BlendWeight));
+			Context.AppendTask(FAnimNextBlendTwoKeyframesTask::Make(BlendWeight));
 		}
 		else
 		{
diff --git a/Engine/Plugins/Experimental/AnimNext/Source/AnimNext/Private/Decorators/BlendTwoWay.h b/Engine/Plugins/Experimental/AnimNext/Source/AnimNext/Private/Decorators/BlendTwoWay.h
--- a/Engine/Plugins/Experimental/AnimNext/Source/AnimNext/Private/Decorators/BlendTwoWay.h
+++ b/Engine/Plugins/Experimental/AnimNext/Source/AnimNext/Private/Decorators/BlendTwoWay.h
@@ -57,7 +57,7 @@ namespace UE::AnimNext
 		};
 
 		// IEvaluate impl
-		virtual void PostEvaluate(const FExecutionContext& Context, const TDecoratorBinding<IEvaluate>& Binding) const override;
+		virtual void PostEvaluate(FEvaluateTraversalContext& Context, const TDecoratorBinding<IEvaluate>& Binding) const override;
 
 		// IUpdate impl
 		virtual void PreUpdate(FUpdateTraversalContext& Context, const TDecoratorBinding<IUpdate>& Binding, const FDecoratorUpdateState& DecoratorState) const override;
diff --git a/Engine/Plugins/Experimental/AnimNext/Source/AnimNext/Private/Decorators/ReferencePoseDecorator.cpp b/Engine/Plugins/Experimental/AnimNext/Source/AnimNext/Private/Decorators/ReferencePoseDecorator.cpp
--- a/Engine/Plugins/Experimental/AnimNext/Source/AnimNext/Private/Decorators/ReferencePoseDecorator.cpp
+++ b/Engine/Plugins/Experimental/AnimNext/Source/AnimNext/Private/Decorators/ReferencePoseDecorator.cpp
@@ -13,14 +13,13 @@ namespace UE::AnimNext
 		DEFINE_ANIM_DECORATOR_IMPLEMENTS_INTERFACE(IEvaluate)
 	DEFINE_ANIM_DECORATOR_END(FReferencePoseDecorator)
 
-	void FReferencePoseDecorator::PreEvaluate(const FExecutionContext& Context, const TDecoratorBinding<IEvaluate>& Binding) const
+	void FReferencePoseDecorator::PreEvaluate(FEvaluateTraversalContext& Context, const TDecoratorBinding<IEvaluate>& Binding) const
 	{
 		const FSharedData* SharedData = Binding.GetSharedData<FSharedData>();
 
 		FAnimNextPushReferenceKeyframeTask Task;
 		Task.bIsAdditive = SharedData->ReferencePoseType == EAnimNextReferencePoseType::AdditiveIdentity;
 
-		FEvaluateTraversalContext& TraversalContext = Context.GetTraversalContext<FEvaluateTraversalContext>();
-		TraversalContext.AppendTask(Task);
+		Context.AppendTask(Task);
 	}
 }
diff --git a/Engine/Plugins/Experimental/AnimNext/Source/AnimNext/Private/Decorators/ReferencePoseDecorator.h b/Engine/Plugins/Experimental/AnimNext/Source/AnimNext/Private/Decorators/ReferencePoseDecorator.h
--- a/Engine/Plugins/Experimental/AnimNext/Source/AnimNext/Private/Decorators/ReferencePoseDecorator.h
+++ b/Engine/Plugins/Experimental/AnimNext/Source/AnimNext/Private/Decorators/ReferencePoseDecorator.h
@@ -45,6 +45,6 @@ namespace UE::AnimNext
 		using FSharedData = FAnimNextReferencePoseDecoratorSharedData;
 
 		// IEvaluate impl
-		virtual void PreEvaluate(const FExecutionContext& Context, const TDecoratorBinding<IEvaluate>& Binding) const override;
+		virtual void PreEvaluate(FEvaluateTraversalContext& Context, const TDecoratorBinding<IEvaluate>& Binding) const override;
 	};
 }
diff --git a/Engine/Plugins/Experimental/AnimNext/Source/AnimNext/Private/Decorators/SequencePlayer.cpp b/Engine/Plugins/Experimental/AnimNext/Source/AnimNext/Private/Decorators/SequencePlayer.cpp
--- a/Engine/Plugins/Experimental/AnimNext/Source/AnimNext/Private/Decorators/SequencePlayer.cpp
+++ b/Engine/Plugins/Experimental/AnimNext/Source/AnimNext/Private/Decorators/SequencePlayer.cpp
@@ -26,7 +26,7 @@ namespace UE::AnimNext
 		}
 	}
 
-	void FSequencePlayerDecorator::PreEvaluate(const FExecutionContext& Context, const TDecoratorBinding<IEvaluate>& Binding) const
+	void FSequencePlayerDecorator::PreEvaluate(FEvaluateTraversalContext& Context, const TDecoratorBinding<IEvaluate>& Binding) const
 	{
 		const FSharedData* SharedData = Binding.GetSharedData<FSharedData>();
 
@@ -36,8 +36,7 @@ namespace UE::AnimNext
 		FAnimNextAnimSequenceKeyframeTask Task = FAnimNextAnimSequenceKeyframeTask::MakeFromSampleTime(SharedData->AnimSequence, InstanceData->InternalTimeAccumulator, bInterpolate);
 		Task.bExtractTrajectory = true;	/*Output.AnimInstanceProxy->ShouldExtractRootMotion()*/
 
-		FEvaluateTraversalContext& TraversalContext = Context.GetTraversalContext<FEvaluateTraversalContext>();
-		TraversalContext.AppendTask(Task);
+		Context.AppendTask(Task);
 	}
 
 	float FSequencePlayerDecorator::GetPlayRate(const FExecutionContext& Context, const TDecoratorBinding<ITimeline>& Binding) const
diff --git a/Engine/Plugins/Experimental/AnimNext/Source/AnimNext/Private/Decorators/SequencePlayer.h b/Engine/Plugins/Experimental/AnimNext/Source/AnimNext/Private/Decorators/SequencePlayer.h
--- a/Engine/Plugins/Experimental/AnimNext/Source/AnimNext/Private/Decorators/SequencePlayer.h
+++ b/Engine/Plugins/Experimental/AnimNext/Source/AnimNext/Private/Decorators/SequencePlayer.h
@@ -63,7 +63,7 @@ namespace UE::AnimNext
 		};
 
 		// IEvaluate impl
-		virtual void PreEvaluate(const FExecutionContext& Context, const TDecoratorBinding<IEvaluate>& Binding) const override;
+		virtual void PreEvaluate(FEvaluateTraversalContext& Context, const TDecoratorBinding<IEvaluate>& Binding) const override;
 
 		// ITimeline impl
 		virtual float GetPlayRate(const FExecutionContext& Context, const TDecoratorBinding<ITimeline>& Binding) const override;
diff --git a/Engine/Plugins/Experimental/AnimNext/Source/AnimNext/Private/Graph/RigUnit_AnimNextGraphEvaluator.cpp b/Engine/Plugins/Experimental/AnimNext/Source/AnimNext/Private/Graph/RigUnit_AnimNextGraphEvaluator.cpp
--- a/Engine/Plugins/Experimental/AnimNext/Source/AnimNext/Private/Graph/RigUnit_AnimNextGraphEvaluator.cpp
+++ b/Engine/Plugins/Experimental/AnimNext/Source/AnimNext/Private/Graph/RigUnit_AnimNextGraphEvaluator.cpp
@@ -55,8 +55,7 @@ void FRigUnit_AnimNextGraphEvaluator::StaticExecute(FRigVMExtendedExecuteContext
 	if (EnumHasAnyFlags(SimulationSteps, EAnimNextGraphSimulationSteps::Evaluate))
 	{
 		// Call pre/post evaluate on our graph
-		FEvaluateTraversalContext TraversalContext;
-		FEvaluationProgram EvaluationProgram = EvaluateGraph(Context, TraversalContext, GraphInstancePtr);
+		FEvaluationProgram EvaluationProgram = EvaluateGraph(Context, GraphInstancePtr);
 
 		if (!EvaluationProgram.IsEmpty())
 		{
diff --git a/Engine/Plugins/Experimental/AnimNext/Source/AnimNext/Public/DecoratorBase/ExecutionContext.h b/Engine/Plugins/Experimental/AnimNext/Source/AnimNext/Public/DecoratorBase/ExecutionContext.h
--- a/Engine/Plugins/Experimental/AnimNext/Source/AnimNext/Public/DecoratorBase/ExecutionContext.h
+++ b/Engine/Plugins/Experimental/AnimNext/Source/AnimNext/Public/DecoratorBase/ExecutionContext.h
@@ -25,7 +25,6 @@ namespace UE::AnimNext
 	struct FDecorator;
 	struct FDecoratorRegistry;
 	struct FDecoratorTemplate;
-	struct ITraversalContext;
 
 	/**
 	 * Execution Context
@@ -74,10 +73,6 @@ namespace UE::AnimNext
 		// Releases a node instance that is no longer referenced
 		void ReleaseNodeInstance(FNodeInstance* Node) const;
 
-		// Returns the current strongly typed traversal context or nullptr if not in a traversal
-		template<class TraversalContextType>
-		TraversalContextType& GetTraversalContext() const { return *static_cast<TraversalContextType*>(TraversalContext); }	// TODO: Add a casting check for safety
-
 		// Evaluates the latent pin with the specified handle
 		template<typename LatentPinType>
 		LatentPinType EvaluateLatentPin(FLatentPropertyHandle LatentPropertyHandle) const;
@@ -113,22 +108,19 @@ namespace UE::AnimNext
 		const FNodeTemplateRegistry& NodeTemplateRegistry;
 		const FDecoratorRegistry& DecoratorRegistry;
 
-		ITraversalContext* TraversalContext = nullptr;
-
 		// Cached properties for the currently executing graph
 		const UAnimNextGraph* Graph = nullptr;
 		FAnimNextGraphInstance* GraphInstance = nullptr;
 		TArrayView<const uint8> GraphSharedData;
 		FRigVMMemoryHandleArray RigVMLatentMemoryHandles;
 		FRigVMExtendedExecuteContext* RigVMExecuteContext = nullptr;
-
-		friend struct FScopedTraversalContext;
 	};
 
 	// Returns a pointer to the current execution context if present, nullptr otherwise.
 	FExecutionContext* GetThreadExecutionContext();
 
 	//////////////////////////////////////////////////////////////////////////
+	// Inline implementations
 
 	template<class DecoratorInterface>
 	inline bool FExecutionContext::GetInterface(const FWeakDecoratorPtr& DecoratorPtr, TDecoratorBinding<DecoratorInterface>& InterfaceBinding) const
diff --git a/Engine/Plugins/Experimental/AnimNext/Source/AnimNext/Public/DecoratorBase/ITraversalContext.h b/Engine/Plugins/Experimental/AnimNext/Source/AnimNext/Public/DecoratorBase/ITraversalContext.h
deleted file mode 100644
diff --git a/Engine/Plugins/Experimental/AnimNext/Source/AnimNext/Public/DecoratorInterfaces/IEvaluate.h b/Engine/Plugins/Experimental/AnimNext/Source/AnimNext/Public/DecoratorInterfaces/IEvaluate.h
--- a/Engine/Plugins/Experimental/AnimNext/Source/AnimNext/Public/DecoratorInterfaces/IEvaluate.h
+++ b/Engine/Plugins/Experimental/AnimNext/Source/AnimNext/Public/DecoratorInterfaces/IEvaluate.h
@@ -3,8 +3,8 @@
 #pragma once
 
 #include "CoreMinimal.h"
+#include "DecoratorBase/ExecutionContextProxy.h"
 #include "DecoratorBase/IDecoratorInterface.h"
-#include "DecoratorBase/ITraversalContext.h"
 #include "EvaluationVM/EvaluationProgram.h"
 #include "EvaluationVM/KeyframeState.h"
 
@@ -13,23 +13,23 @@ namespace UE::AnimNext
 	/**
 	 * FEvaluateTraversalContext
 	 *
-	 * Contains all relevant transient data for an evaluate traversal.
+	 * Contains all relevant transient data for an evaluate traversal and wraps the execution context.
 	 */
-	struct ANIMNEXT_API FEvaluateTraversalContext : ITraversalContext
+	struct FEvaluateTraversalContext final : FExecutionContextProxy
 	{
-		FEvaluateTraversalContext() = default;
-
 		// Appends a new task into the evaluation program, tasks mutate state in the order they have been appended in
 		// This means that child nodes need to evaluate first, tasks will usually be appended in IEvaluate::PostEvaluate
 		// Tasks are moved into their final memory location, caller can allocate the task anywhere, it is no longer needed after this operation
 		// @see FEvaluationProgram, FEvaluationTask, FEvaluationVM
 		template<class TaskType>
-		void AppendTask(TaskType&& Task) { EvaluationProgram->AppendTask(MoveTemp(Task)); }
+		void AppendTask(TaskType&& Task) { EvaluationProgram.AppendTask(MoveTemp(Task)); }
 
 	private:
-		FEvaluationProgram* EvaluationProgram = nullptr;
+		FEvaluateTraversalContext(const FExecutionContext& InExecutionContext, FEvaluationProgram& InEvaluationProgram);
+
+		FEvaluationProgram& EvaluationProgram;
 
-		friend struct FScopedEvaluationProgram;
+		friend ANIMNEXT_API FEvaluationProgram EvaluateGraph(FExecutionContext& Context, FWeakDecoratorPtr GraphRootPtr);
 	};
 
 	/**
@@ -51,10 +51,10 @@ namespace UE::AnimNext
 		DECLARE_ANIM_DECORATOR_INTERFACE(IEvaluate, 0xa303e9e7)
 
 		// Called before a decorator's children are evaluated
-		virtual void PreEvaluate(const FExecutionContext& Context, const TDecoratorBinding<IEvaluate>& Binding) const;
+		virtual void PreEvaluate(FEvaluateTraversalContext& Context, const TDecoratorBinding<IEvaluate>& Binding) const;
 
 		// Called after a decorator's children have been evaluated
-		virtual void PostEvaluate(const FExecutionContext& Context, const TDecoratorBinding<IEvaluate>& Binding) const;
+		virtual void PostEvaluate(FEvaluateTraversalContext& Context, const TDecoratorBinding<IEvaluate>& Binding) const;
 	};
 
 	/**
@@ -64,13 +64,13 @@ namespace UE::AnimNext
 	struct TDecoratorBinding<IEvaluate> : FDecoratorBinding
 	{
 		// @see IEvaluate::PreEvaluate
-		void PreEvaluate(const FExecutionContext& Context) const
+		void PreEvaluate(FEvaluateTraversalContext& Context) const
 		{
 			GetInterface()->PreEvaluate(Context, *this);
 		}
 
 		// @see IEvaluate::PostEvaluate
-		void PostEvaluate(const FExecutionContext& Context) const
+		void PostEvaluate(FEvaluateTraversalContext& Context) const
 		{
 			GetInterface()->PostEvaluate(Context, *this);
 		}
@@ -91,5 +91,5 @@ namespace UE::AnimNext
 	 *
 	 * @see IEvaluate::PreEvaluate, IEvaluate::PostEvaluate, IHierarchy::GetChildren
 	 */
-	[[nodiscard]] ANIMNEXT_API FEvaluationProgram EvaluateGraph(FExecutionContext& Context, FEvaluateTraversalContext& TraversalContext, FWeakDecoratorPtr GraphRootPtr);
+	[[nodiscard]] ANIMNEXT_API FEvaluationProgram EvaluateGraph(FExecutionContext& Context, FWeakDecoratorPtr GraphRootPtr);
 }
diff --git a/Engine/Plugins/Experimental/AnimNext/Source/AnimNextTestSuite/Private/AnimNextDecoratorGraphTest.cpp b/Engine/Plugins/Experimental/AnimNext/Source/AnimNextTestSuite/Private/AnimNextDecoratorGraphTest.cpp
--- a/Engine/Plugins/Experimental/AnimNext/Source/AnimNextTestSuite/Private/AnimNextDecoratorGraphTest.cpp
+++ b/Engine/Plugins/Experimental/AnimNext/Source/AnimNextTestSuite/Private/AnimNextDecoratorGraphTest.cpp
@@ -56,7 +56,7 @@ namespace UE::AnimNext
 		}
 
 		// IEvaluate impl
-		virtual void PostEvaluate(const FExecutionContext& Context, const TDecoratorBinding<IEvaluate>& Binding) const override
+		virtual void PostEvaluate(FEvaluateTraversalContext& Context, const TDecoratorBinding<IEvaluate>& Binding) const override
 		{
 			IEvaluate::PostEvaluate(Context, Binding);
 
diff --git a/Engine/Plugins/Experimental/AnimNext/Source/AnimNextTestSuite/Private/AnimNextDecoratorInterfacesTest.cpp b/Engine/Plugins/Experimental/AnimNext/Source/AnimNextTestSuite/Private/AnimNextDecoratorInterfacesTest.cpp
--- a/Engine/Plugins/Experimental/AnimNext/Source/AnimNextTestSuite/Private/AnimNextDecoratorInterfacesTest.cpp
+++ b/Engine/Plugins/Experimental/AnimNext/Source/AnimNextTestSuite/Private/AnimNextDecoratorInterfacesTest.cpp
@@ -61,7 +61,7 @@ namespace UE::AnimNext
 		}
 
 		// IEvaluate impl
-		virtual void PreEvaluate(const FExecutionContext& Context, const TDecoratorBinding<IEvaluate>& Binding) const override
+		virtual void PreEvaluate(FEvaluateTraversalContext& Context, const TDecoratorBinding<IEvaluate>& Binding) const override
 		{
 			if (Private::EvaluatedDecorators != nullptr)
 			{
@@ -71,7 +71,7 @@ namespace UE::AnimNext
 			IEvaluate::PreEvaluate(Context, Binding);
 		}
 
-		virtual void PostEvaluate(const FExecutionContext& Context, const TDecoratorBinding<IEvaluate>& Binding) const override
+		virtual void PostEvaluate(FEvaluateTraversalContext& Context, const TDecoratorBinding<IEvaluate>& Binding) const override
 		{
 			if (Private::EvaluatedDecorators != nullptr)
 			{
@@ -187,7 +187,7 @@ namespace UE::AnimNext
 		}
 
 		// IEvaluate impl
-		virtual void PreEvaluate(const FExecutionContext& Context, const TDecoratorBinding<IEvaluate>& Binding) const override
+		virtual void PreEvaluate(FEvaluateTraversalContext& Context, const TDecoratorBinding<IEvaluate>& Binding) const override
 		{
 			if (Private::EvaluatedDecorators != nullptr)
 			{
@@ -197,7 +197,7 @@ namespace UE::AnimNext
 			IEvaluate::PreEvaluate(Context, Binding);
 		}
 
-		virtual void PostEvaluate(const FExecutionContext& Context, const TDecoratorBinding<IEvaluate>& Binding) const override
+		virtual void PostEvaluate(FEvaluateTraversalContext& Context, const TDecoratorBinding<IEvaluate>& Binding) const override
 		{
 			if (Private::EvaluatedDecorators != nullptr)
 			{
@@ -663,8 +663,7 @@ bool FAnimationAnimNextRuntimeTest_IEvaluate::RunTest(const FString& InParameter
 		AddErrorIfFalse(NodeCPtr.IsValid(), "FAnimationAnimNextRuntimeTest_IEvaluate -> Failed to allocate root node instance");
 
 		// Call pre/post evaluate on our graph
-		FEvaluateTraversalContext TraversalContext;
-		(void)EvaluateGraph(Context, TraversalContext, NodeCPtr);
+		(void)EvaluateGraph(Context, NodeCPtr);
 
 		AddErrorIfFalse(EvaluatedDecorators.Num() == 6, "FAnimationAnimNextRuntimeTest_IEvaluate -> Expected 6 nodes to have been visited during the evaluate traversal");
 		AddErrorIfFalse(EvaluatedDecorators[0] == FDecoratorWithChildren::DecoratorUID, "FAnimationAnimNextRuntimeTest_IEvaluate -> Unexpected evaluate order");		// NodeC
