Add support for AnimNext graph instance data to contain strong UObject references and provide a mechanism for GC to find them

diff --git a/Engine/Plugins/Experimental/AnimNext/Source/AnimNext/Internal/Graph/AnimNextGraph.h b/Engine/Plugins/Experimental/AnimNext/Source/AnimNext/Internal/Graph/AnimNextGraph.h
--- a/Engine/Plugins/Experimental/AnimNext/Source/AnimNext/Internal/Graph/AnimNextGraph.h
+++ b/Engine/Plugins/Experimental/AnimNext/Source/AnimNext/Internal/Graph/AnimNextGraph.h
@@ -86,6 +86,9 @@ struct ANIMNEXT_API FAnimNextGraphInstance
 	// Check to see if this instance data matches the provided graph
 	bool UsesGraph(const UAnimNextGraph* InGraph) const;
 
+	// Adds strong/hard object references during GC
+	void AddStructReferencedObjects(class FReferenceCollector& Collector);
+
 	// Returns a typed graph instance component, creating it lazily the first time it is queried
 	template<class ComponentType>
 	ComponentType& GetComponent();
@@ -126,6 +129,15 @@ private:
 	friend FRigUnit_AnimNextGraphEvaluator;	// We evaluate the instance
 };
 
+template<>
+struct TStructOpsTypeTraits<FAnimNextGraphInstance> : public TStructOpsTypeTraitsBase2<FAnimNextGraphInstance>
+{
+	enum
+	{
+		WithAddStructReferencedObjects = true,
+	};
+};
+
 // A user-created graph of logic used to supply data
 UCLASS(BlueprintType)
 class ANIMNEXT_API UAnimNextGraph :  public URigVMHost, public IAnimNextScheduleTermInterface
@@ -258,7 +270,7 @@ ComponentType* FAnimNextGraphInstance::TryGetComponent()
 	const FName ComponentName = ComponentType::StaticComponentName();
 	const int32 ComponentNameHash = GetTypeHash(ComponentName);
 
-	return *static_cast<ComponentType*>(TryGetComponent(ComponentNameHash, ComponentName));
+	return static_cast<ComponentType*>(TryGetComponent(ComponentNameHash, ComponentName));
 }
 
 template<class ComponentType>
@@ -267,5 +279,5 @@ const ComponentType* FAnimNextGraphInstance::TryGetComponent() const
 	const FName ComponentName = ComponentType::StaticComponentName();
 	const int32 ComponentNameHash = GetTypeHash(ComponentName);
 
-	return *static_cast<ComponentType*>(TryGetComponent(ComponentNameHash, ComponentName));
+	return static_cast<ComponentType*>(TryGetComponent(ComponentNameHash, ComponentName));
 }
diff --git a/Engine/Plugins/Experimental/AnimNext/Source/AnimNext/Private/DecoratorInterfaces/IGarbageColllection.cpp b/Engine/Plugins/Experimental/AnimNext/Source/AnimNext/Private/DecoratorInterfaces/IGarbageColllection.cpp
new file mode 100644
--- /dev/null
+++ b/Engine/Plugins/Experimental/AnimNext/Source/AnimNext/Private/DecoratorInterfaces/IGarbageColllection.cpp
@@ -0,0 +1,30 @@
+// Copyright Epic Games, Inc. All Rights Reserved.
+
+#include "DecoratorInterfaces/IGarbageCollection.h"
+
+#include "DecoratorBase/ExecutionContext.h"
+#include "Graph/GC_GraphInstanceComponent.h"
+
+namespace UE::AnimNext
+{
+	void IGarbageCollection::RegisterWithGC(const FExecutionContext& Context, const FDecoratorBinding& Binding)
+	{
+		FGCGraphInstanceComponent& Component = Context.GetComponent<FGCGraphInstanceComponent>();
+		Component.Register(Context.GetGraphInstance(), Binding.GetDecoratorPtr());
+	}
+
+	void IGarbageCollection::UnregisterWithGC(const FExecutionContext& Context, const FDecoratorBinding& Binding)
+	{
+		FGCGraphInstanceComponent& Component = Context.GetComponent<FGCGraphInstanceComponent>();
+		Component.Unregister(Binding.GetDecoratorPtr());
+	}
+
+	void IGarbageCollection::AddReferencedObjects(const FExecutionContext& Context, const TDecoratorBinding<IGarbageCollection>& Binding, FReferenceCollector& Collector) const
+	{
+		TDecoratorBinding<IGarbageCollection> SuperBinding;
+		if (Context.GetInterfaceSuper(Binding, SuperBinding))
+		{
+			SuperBinding.AddReferencedObjects(Context, Collector);
+		}
+	}
+}
diff --git a/Engine/Plugins/Experimental/AnimNext/Source/AnimNext/Private/Graph/AnimNextGraph.cpp b/Engine/Plugins/Experimental/AnimNext/Source/AnimNext/Private/Graph/AnimNextGraph.cpp
--- a/Engine/Plugins/Experimental/AnimNext/Source/AnimNext/Private/Graph/AnimNextGraph.cpp
+++ b/Engine/Plugins/Experimental/AnimNext/Source/AnimNext/Private/Graph/AnimNextGraph.cpp
@@ -10,6 +10,7 @@
 #include "DecoratorBase/DecoratorReader.h"
 #include "DecoratorBase/ExecutionContext.h"
 #include "Graph/AnimNext_LODPose.h"
+#include "Graph/GC_GraphInstanceComponent.h"
 #include "Serialization/MemoryReader.h"
 #include "AnimNextStats.h"
 
@@ -167,6 +168,14 @@ bool FAnimNextGraphInstance::UsesGraph(const UAnimNextGraph* InGraph) const
 	return Graph == InGraph;
 }
 
+void FAnimNextGraphInstance::AddStructReferencedObjects(FReferenceCollector& Collector)
+{
+	if (const UE::AnimNext::FGCGraphInstanceComponent* Component = TryGetComponent<UE::AnimNext::FGCGraphInstanceComponent>())
+	{
+		Component->AddReferencedObjects(Collector);
+	}
+}
+
 UE::AnimNext::FGraphInstanceComponent* FAnimNextGraphInstance::TryGetComponent(int32 ComponentNameHash, FName ComponentName) const
 {
 	if (const TSharedPtr<UE::AnimNext::FGraphInstanceComponent>* Component = Components.FindByHash(ComponentNameHash, ComponentName))
diff --git a/Engine/Plugins/Experimental/AnimNext/Source/AnimNext/Private/Graph/GC_GraphInstanceComponent.cpp b/Engine/Plugins/Experimental/AnimNext/Source/AnimNext/Private/Graph/GC_GraphInstanceComponent.cpp
new file mode 100644
--- /dev/null
+++ b/Engine/Plugins/Experimental/AnimNext/Source/AnimNext/Private/Graph/GC_GraphInstanceComponent.cpp
@@ -0,0 +1,42 @@
+// Copyright Epic Games, Inc. All Rights Reserved.
+
+#include "Graph/GC_GraphInstanceComponent.h"
+
+#include "DecoratorBase/ExecutionContext.h"
+#include "DecoratorInterfaces/IGarbageCollection.h"
+
+namespace UE::AnimNext
+{
+	void FGCGraphInstanceComponent::Register(FAnimNextGraphInstance& GraphInstance, const FWeakDecoratorPtr& DecoratorPtr)
+	{
+		DecoratorsWithReferences.Add(FEntry(GraphInstance, DecoratorPtr));
+	}
+
+	void FGCGraphInstanceComponent::Unregister(const FWeakDecoratorPtr& DecoratorPtr)
+	{
+		const int32 EntryIndex = DecoratorsWithReferences.IndexOfByPredicate(
+			[&DecoratorPtr](const FEntry& Entry)
+			{
+				return Entry.DecoratorPtr == DecoratorPtr;
+			});
+
+		if (ensure(EntryIndex != INDEX_NONE))
+		{
+			DecoratorsWithReferences.RemoveAtSwap(EntryIndex);
+		}
+	}
+
+	void FGCGraphInstanceComponent::AddReferencedObjects(FReferenceCollector& Collector) const
+	{
+		TDecoratorBinding<IGarbageCollection> GCDecorator;
+
+		// TODO: If we kept the entries sorted by graph instance, we could re-use the execution context
+		for (const FEntry& Entry : DecoratorsWithReferences)
+		{
+			FExecutionContext Context(Entry.GraphInstance);
+			ensure(Context.GetInterface(Entry.DecoratorPtr, GCDecorator));
+
+			GCDecorator.AddReferencedObjects(Context, Collector);
+		}
+	}
+}
diff --git a/Engine/Plugins/Experimental/AnimNext/Source/AnimNext/Private/Graph/GC_GraphInstanceComponent.h b/Engine/Plugins/Experimental/AnimNext/Source/AnimNext/Private/Graph/GC_GraphInstanceComponent.h
new file mode 100644
--- /dev/null
+++ b/Engine/Plugins/Experimental/AnimNext/Source/AnimNext/Private/Graph/GC_GraphInstanceComponent.h
@@ -0,0 +1,49 @@
+// Copyright Epic Games, Inc. All Rights Reserved.
+
+#pragma once
+
+#include "CoreMinimal.h"
+#include "DecoratorBase/DecoratorPtr.h"
+#include "Graph/GraphInstanceComponent.h"
+
+struct FAnimNextGraphInstance;
+class FReferenceCollector;
+
+namespace UE::AnimNext
+{
+	/**
+	 * FGCGraphInstanceComponent
+	 *
+	 * This component maintains the necessary state to garbage collection.
+	 */
+	struct FGCGraphInstanceComponent : public FGraphInstanceComponent
+	{
+		DECLARE_ANIM_GRAPH_INSTANCE_COMPONENT(FGCGraphInstanceComponent)
+
+		// Registers the provided decorator with the GC system
+		// Once registered, IGarbageCollection::AddReferencedObjects will be called on it during GC
+		void Register(FAnimNextGraphInstance& GraphInstance, const FWeakDecoratorPtr& DecoratorPtr);
+
+		// Unregisters the provided decorator from the GC system
+		void Unregister(const FWeakDecoratorPtr& DecoratorPtr);
+
+		// Called during garbage collection to collect strong object references
+		void AddReferencedObjects(FReferenceCollector& Collector) const;
+
+	private:
+		struct FEntry
+		{
+			FEntry(FAnimNextGraphInstance& InGraphInstance, const FWeakDecoratorPtr& InDecoratorPtr)
+				: GraphInstance(InGraphInstance)
+				, DecoratorPtr(InDecoratorPtr)
+			{
+			}
+
+			FAnimNextGraphInstance& GraphInstance;
+			FWeakDecoratorPtr DecoratorPtr;
+		};
+
+		// List of decorator handles that contain UObject references and implement IGarbageCollection
+		TArray<FEntry> DecoratorsWithReferences;
+	};
+}
diff --git a/Engine/Plugins/Experimental/AnimNext/Source/AnimNext/Public/DecoratorBase/ExecutionContext.h b/Engine/Plugins/Experimental/AnimNext/Source/AnimNext/Public/DecoratorBase/ExecutionContext.h
--- a/Engine/Plugins/Experimental/AnimNext/Source/AnimNext/Public/DecoratorBase/ExecutionContext.h
+++ b/Engine/Plugins/Experimental/AnimNext/Source/AnimNext/Public/DecoratorBase/ExecutionContext.h
@@ -89,7 +89,7 @@ namespace UE::AnimNext
 		GraphInstanceComponentMapType::TConstIterator GetComponentIterator() const { return GraphInstance->GetComponentIterator(); }
 
 		// Returns the bound graph instance
-		const FAnimNextGraphInstance& GetGraphInstance() const { return *GraphInstance; }
+		FAnimNextGraphInstance& GetGraphInstance() const { return *GraphInstance; }
 
 	private:
 		// No copy or move
diff --git a/Engine/Plugins/Experimental/AnimNext/Source/AnimNext/Public/DecoratorBase/ExecutionContextProxy.h b/Engine/Plugins/Experimental/AnimNext/Source/AnimNext/Public/DecoratorBase/ExecutionContextProxy.h
--- a/Engine/Plugins/Experimental/AnimNext/Source/AnimNext/Public/DecoratorBase/ExecutionContextProxy.h
+++ b/Engine/Plugins/Experimental/AnimNext/Source/AnimNext/Public/DecoratorBase/ExecutionContextProxy.h
@@ -101,7 +101,7 @@ namespace UE::AnimNext
 		}
 
 		// Returns the bound graph instance
-		const FAnimNextGraphInstance& GetGraphInstance() const
+		FAnimNextGraphInstance& GetGraphInstance() const
 		{
 			return ExecutionContext.GetGraphInstance();
 		}
diff --git a/Engine/Plugins/Experimental/AnimNext/Source/AnimNext/Public/DecoratorInterfaces/IGarbageCollection.h b/Engine/Plugins/Experimental/AnimNext/Source/AnimNext/Public/DecoratorInterfaces/IGarbageCollection.h
new file mode 100644
--- /dev/null
+++ b/Engine/Plugins/Experimental/AnimNext/Source/AnimNext/Public/DecoratorInterfaces/IGarbageCollection.h
@@ -0,0 +1,48 @@
+// Copyright Epic Games, Inc. All Rights Reserved.
+
+#pragma once
+
+#include "CoreMinimal.h"
+#include "DecoratorBase/IDecoratorInterface.h"
+
+class FReferenceCollector;
+
+namespace UE::AnimNext
+{
+	/**
+	 * IGarbageCollection
+	 * 
+	 * This interface exposes garbage collection reference tracking.
+	 */
+	struct ANIMNEXT_API IGarbageCollection : IDecoratorInterface
+	{
+		DECLARE_ANIM_DECORATOR_INTERFACE(IGarbageCollection, 0x231a2017)
+
+		// Registers the provided binding for GC callback
+		// Once registered, AddReferencedObjects is called during GC to collect references
+		static void RegisterWithGC(const FExecutionContext& Context, const FDecoratorBinding& Binding);
+
+		// Unregisters the provided binding from GC callback
+		static void UnregisterWithGC(const FExecutionContext& Context, const FDecoratorBinding& Binding);
+
+		// Called when garbage collection requests hard/strong object references
+		// @see UObject::AddReferencedObjects
+		virtual void AddReferencedObjects(const FExecutionContext& Context, const TDecoratorBinding<IGarbageCollection>& Binding, FReferenceCollector& Collector) const;
+	};
+
+	/**
+	 * Specialization for decorator binding.
+	 */
+	template<>
+	struct TDecoratorBinding<IGarbageCollection> : FDecoratorBinding
+	{
+		// @see IGarbageCollection::AddReferencedObjects
+		void AddReferencedObjects(const FExecutionContext& Context, FReferenceCollector& Collector) const
+		{
+			GetInterface()->AddReferencedObjects(Context, *this, Collector);
+		}
+
+	protected:
+		const IGarbageCollection* GetInterface() const { return GetInterfaceTyped<IGarbageCollection>(); }
+	};
+}
