Fix for Anim Next Parameter Block TLS context data release issue

#rnx

diff --git a/Engine/Plugins/Experimental/AnimNext/Source/AnimNext/Internal/Param/AnimNextParameterBlock.h b/Engine/Plugins/Experimental/AnimNext/Source/AnimNext/Internal/Param/AnimNextParameterBlock.h
--- a/Engine/Plugins/Experimental/AnimNext/Source/AnimNext/Internal/Param/AnimNextParameterBlock.h
+++ b/Engine/Plugins/Experimental/AnimNext/Source/AnimNext/Internal/Param/AnimNextParameterBlock.h
@@ -88,6 +88,7 @@ class UAnimNextParameterBlock : public URigVMHost, public IAnimNextParameterSour
 	virtual bool ShouldCacheLayer(const UE::AnimNext::FParamStackLayerHandle& InHandle) const override;
 
 	// UObject interface
+	virtual void BeginDestroy() override;
 	virtual void PostLoad() override;
 	virtual void GetAssetRegistryTags(TArray<FAssetRegistryTag>& OutTags) const override;
 
diff --git a/Engine/Plugins/Experimental/AnimNext/Source/AnimNext/Internal/RigVMRuntimeDataRegistry.h b/Engine/Plugins/Experimental/AnimNext/Source/AnimNext/Internal/RigVMRuntimeDataRegistry.h
--- a/Engine/Plugins/Experimental/AnimNext/Source/AnimNext/Internal/RigVMRuntimeDataRegistry.h
+++ b/Engine/Plugins/Experimental/AnimNext/Source/AnimNext/Internal/RigVMRuntimeDataRegistry.h
@@ -9,55 +9,70 @@
 namespace UE::AnimNext
 {
 
+using FRigVMRuntimeDataID = TObjectKey<URigVM>;
+
 struct FRigVMRuntimeData
 {
 	FRigVMExtendedExecuteContext Context;
 };
 
-using FRigVMRuntimeDataID = TObjectKey<URigVM>;
 
 /**
-	* RigVMRuntimeDataRegistry
-	* 
-	* A global registry of all existing VMs that require TLS data instantiation
-	* 
-	*/
+* RigVMRuntimeDataRegistry
+* 
+* A global registry of all existing VMs that require TLS data instantiation
+* 
+*/
 struct ANIMNEXT_API FRigVMRuntimeDataRegistry final
 {
 	/**
-	 * Finds the VM runtime data ID in the TLS storage
+	 * Finds or adds one VM runtime data instance for the passed ID in the TLS.
+	 * If it exists and the Context hash is correct, returns the TLS weak pointer.
+	 * If not found or the context hash is not correct, it creates a new instance in the Global storage and adds a Weak pointer to the TLS.
+	 *
+	 * @param RigVMRuntimeDataID The key to find or associate (if created) the instance with.
+	 * @param ReferenceContext Reference context to copy data from (if not found or found but the context hash is not correct)
+	 * @param bOutHasBeenUpdated Out : True If the ID did not exist or existed but the Context hash was not correct
+	 * @return A weak pointer to the data as stored in the TLS map.
+	 */
+	static TWeakPtr<FRigVMRuntimeData> FindOrAddLocalRuntimeData(const FRigVMRuntimeDataID& RigVMRuntimeDataID, const FRigVMExtendedExecuteContext& ReferenceContext);
+
+	/**
+	 * Finds the instance VM runtime data ID in the TLS storage
 	 *
-	 * @param RigVMRuntimeDataID The key to associate the value with.
-	 * @return A pointer to the data as stored in the map. The data is only valid until the next change to any key in the map.
+	 * @param RigVMRuntimeDataID The key to find the associated data with.
+	 * @return A pointer to the data as stored in the TLS map.
 	 */
-	static FRigVMRuntimeData* FindRuntimeData(const FRigVMRuntimeDataID& RigVMRuntimeDataID);
+	static TWeakPtr<FRigVMRuntimeData> FindLocalRuntimeData(const FRigVMRuntimeDataID& RigVMRuntimeDataID);
 
 	/**
-	 * Adds VM runtime data for the passed ID
+	 * Adds one VM runtime data instance for the passed ID in the Global storage and adds a Weak pointer to the TLS.
+	 * Adding a new instance is only allowed if it does not already exist in the TLS
 	 *
-	 * @param RigVMRuntimeDataID The key to associate the value with.
+	 * @param RigVMRuntimeDataID The key to associate the data with.
 	 * @param ReferenceContext Reference context to copy data from
-	 * @return A pointer to the data as stored in the map. The data is only valid until the next change to any key in the map.
+	 * @return A weak pointer to the data as stored in the TLS map.
 	 */
-	static FRigVMRuntimeData* AddRuntimeData(const FRigVMRuntimeDataID& RigVMRuntimeDataID, const FRigVMExtendedExecuteContext& ReferenceContext);
+	static TWeakPtr<FRigVMRuntimeData> AddRuntimeData(const FRigVMRuntimeDataID& RigVMRuntimeDataID, const FRigVMExtendedExecuteContext& ReferenceContext);
 
-	
 	/**
-	 * Finds or adds the VM runtime data ID
+	 * Destroys all the VM runtime data instances for the passed ID. This should be only called when the VM is destroyed.
 	 *
-	 * @param RigVMRuntimeDataID The key to associate the value with.
-	 * @param ReferenceContext Reference context to copy data from (if not found)
-	 * @return A pointer to the data as stored in the map. The data is only valid until the next change to any key in the map.
+	 * @param RigVMRuntimeDataID The key to find the associated data with.
 	 */
-	static FRigVMRuntimeData* FindOrAddRuntimeData(const FRigVMRuntimeDataID& RigVMRuntimeDataID, const FRigVMExtendedExecuteContext& ReferenceContext);
+	static void ReleaseAllVMRuntimeData(const FRigVMRuntimeDataID& RigVMRuntimeDataID);
 
 private:
+	static TSharedPtr<FRigVMRuntimeData> AddGlobalRuntimeData(const FRigVMRuntimeDataID& RigVMRuntimeDataID, const FRigVMExtendedExecuteContext& ReferenceContext);
+	static void ReleaseAllGlobalRuntimeData(const FRigVMRuntimeDataID& RigVMRuntimeDataID);
 
-	// Post GC callback to signal a compactation has to be done (and perform it on game thread)
+	// Post GC callback to signal a compaction has to be done (and perform it on game thread)
 	static void HandlePostGarbageCollect();
 
-	// Checks if any of the stored VM datas have been deleted and removes deleted elements.
-	static void PerformStorageCompaction();
+	// Checks if any of the stored VM datas have been deleted and removes deleted elements. Main thread only.
+	static void PerformGlobalStorageCompaction();
+	// Checks if any of the stored VM datas have been deleted and removes deleted elements in the TLS storage.
+	static void PerformLocalStorageCompaction();
 
 	// Module lifetime functions
 	static void Init();
diff --git a/Engine/Plugins/Experimental/AnimNext/Source/AnimNext/Private/Param/AnimNextParameterBlock.cpp b/Engine/Plugins/Experimental/AnimNext/Source/AnimNext/Private/Param/AnimNextParameterBlock.cpp
--- a/Engine/Plugins/Experimental/AnimNext/Source/AnimNext/Private/Param/AnimNextParameterBlock.cpp
+++ b/Engine/Plugins/Experimental/AnimNext/Source/AnimNext/Private/Param/AnimNextParameterBlock.cpp
@@ -25,18 +25,16 @@ void UAnimNextParameterBlock::UpdateLayer(UE::AnimNext::FParamStackLayerHandle&
 	
 	if (VM)
 	{
-		UE::AnimNext::FRigVMRuntimeData* RuntimeData = UE::AnimNext::FRigVMRuntimeDataRegistry::FindRuntimeData(VM);
-		if (RuntimeData == nullptr || RuntimeData->Context.VMHash != VM->GetVMHash())
+		if (TSharedPtr<UE::AnimNext::FRigVMRuntimeData> RuntimeData = UE::AnimNext::FRigVMRuntimeDataRegistry::FindOrAddLocalRuntimeData(VM, GetRigVMExtendedExecuteContext()).Pin())
 		{
-			RuntimeData = UE::AnimNext::FRigVMRuntimeDataRegistry::AddRuntimeData(VM, GetRigVMExtendedExecuteContext());
-			VM->InitializeInstance(RuntimeData->Context, false); // TODO zzz : Temp until VM supports WorkData cached handles using offsets (UE-197067)
-			RuntimeData->Context.VMHash = VM->GetVMHash();
-		}
+			FRigVMExtendedExecuteContext& Context = RuntimeData->Context;
+
+			check(Context.VMHash == VM->GetVMHash());
 
-		FRigVMExtendedExecuteContext& Context = RuntimeData->Context;
-		FAnimNextParameterExecuteContext& AnimNextParameterContext = Context.GetPublicDataSafe<FAnimNextParameterExecuteContext>();
-		AnimNextParameterContext.SetParamContextData(InHandle);
-		VM->ExecuteVM(Context, FRigUnit_AnimNextBeginExecution::EventName);
+			FAnimNextParameterExecuteContext& AnimNextParameterContext = Context.GetPublicDataSafe<FAnimNextParameterExecuteContext>();
+			AnimNextParameterContext.SetParamContextData(InHandle);
+			VM->ExecuteVM(Context, FRigUnit_AnimNextBeginExecution::EventName);
+		}
 	}
 }
 
@@ -63,6 +61,16 @@ bool UAnimNextParameterBlock::ShouldCacheLayer(const UE::AnimNext::FParamStackLa
 	return false;
 }
 
+void UAnimNextParameterBlock::BeginDestroy()
+{
+	Super::BeginDestroy();
+
+	if (VM)
+	{
+		UE::AnimNext::FRigVMRuntimeDataRegistry::ReleaseAllVMRuntimeData(VM);
+	}
+}
+
 void UAnimNextParameterBlock::PostLoad()
 {
 	Super::PostLoad();
diff --git a/Engine/Plugins/Experimental/AnimNext/Source/AnimNext/Private/RigVMRuntimeDataRegistry.cpp b/Engine/Plugins/Experimental/AnimNext/Source/AnimNext/Private/RigVMRuntimeDataRegistry.cpp
--- a/Engine/Plugins/Experimental/AnimNext/Source/AnimNext/Private/RigVMRuntimeDataRegistry.cpp
+++ b/Engine/Plugins/Experimental/AnimNext/Source/AnimNext/Private/RigVMRuntimeDataRegistry.cpp
@@ -2,6 +2,7 @@
 
 #include "RigVMRuntimeDataRegistry.h"
 #include "Misc/ScopeRWLock.h"
+#include "RigVMCore/RigVM.h"
 
 namespace UE::AnimNext
 {
@@ -14,8 +15,13 @@ static FDelegateHandle PostGarbageCollectHandle;
 
 static std::atomic<uint32> GCCycle = 0; // Main thread GC counter, incremented on each main thread GC cycle
 
+static FRWLock GlobalRuntimeDataStorageLock;
+static TMultiMap<FRigVMRuntimeDataID, TSharedPtr<FRigVMRuntimeData>> GlobalRuntimeDataStorage;
+
+
 static thread_local uint32 LocalGCCycle = 0;	// Local thread GC counter, used to compare with main and trigger compaction if different
-static thread_local TMap<FRigVMRuntimeDataID, FRigVMRuntimeData> RuntimeDataStorage;
+static thread_local TMap<FRigVMRuntimeDataID, TWeakPtr<FRigVMRuntimeData>> LocalRuntimeDataStorage;
+
 
 } // end namespace Private
 
@@ -37,68 +43,115 @@ static thread_local TMap<FRigVMRuntimeDataID, FRigVMRuntimeData> RuntimeDataStor
 		Private::bInitialized = false;
 
 		FCoreUObjectDelegates::GetPostGarbageCollect().Remove(Private::PostGarbageCollectHandle);
-		Private::RuntimeDataStorage.Empty();
+		Private::GlobalRuntimeDataStorage.Empty();
 	}
 }
 
-/*static*/ FRigVMRuntimeData* FRigVMRuntimeDataRegistry::FindRuntimeData(const FRigVMRuntimeDataID& RigVMRuntimeDataID)
+/*static*/ TWeakPtr<FRigVMRuntimeData> FRigVMRuntimeDataRegistry::FindOrAddLocalRuntimeData(const FRigVMRuntimeDataID& RigVMRuntimeDataID, const FRigVMExtendedExecuteContext& ReferenceContext)
+{
+	check(Private::bInitialized);
+
+	TWeakPtr<FRigVMRuntimeData> WeakRigVMRuntimeData = FindLocalRuntimeData(RigVMRuntimeDataID);
+	if (TSharedPtr<FRigVMRuntimeData> RigVMRuntimeData = WeakRigVMRuntimeData.Pin())
+	{
+		if (RigVMRuntimeData->Context.VMHash != ReferenceContext.VMHash)
+		{
+			RigVMRuntimeData->Context = ReferenceContext;
+			if (URigVM* VM = RigVMRuntimeDataID.ResolveObjectPtr())
+			{
+				VM->InitializeInstance(RigVMRuntimeData->Context, false);
+			}
+		}
+	}
+	else
+	{
+		WeakRigVMRuntimeData = AddRuntimeData(RigVMRuntimeDataID, ReferenceContext);
+	}
+
+	return WeakRigVMRuntimeData;
+}
+
+/*static*/ TWeakPtr<FRigVMRuntimeData> FRigVMRuntimeDataRegistry::FindLocalRuntimeData(const FRigVMRuntimeDataID& RigVMRuntimeDataID)
 {
 	check(Private::bInitialized);
 
 	const uint32 CurrentCycle = Private::GCCycle;
 	if (CurrentCycle != Private::LocalGCCycle)
 	{
-		PerformStorageCompaction();
+		PerformLocalStorageCompaction();
 
 		Private::LocalGCCycle = CurrentCycle;
 	}
 
-	return Private::RuntimeDataStorage.Find(RigVMRuntimeDataID);
+	return Private::LocalRuntimeDataStorage.FindRef(RigVMRuntimeDataID);
 }
 
-/*static*/ FRigVMRuntimeData* FRigVMRuntimeDataRegistry::AddRuntimeData(const FRigVMRuntimeDataID& RigVMRuntimeDataID, const FRigVMExtendedExecuteContext& ReferenceContext)
+/*static*/ TWeakPtr<FRigVMRuntimeData> FRigVMRuntimeDataRegistry::AddRuntimeData(const FRigVMRuntimeDataID& RigVMRuntimeDataID, const FRigVMExtendedExecuteContext& ReferenceContext)
 {
 	check(Private::bInitialized);
 
-	FRigVMRuntimeData* RigVMRuntimeData = nullptr;
+	check(Private::LocalRuntimeDataStorage.FindRef(RigVMRuntimeDataID).IsValid() == false);
 
-	FRigVMRuntimeData& NewRigVMRuntimeData = Private::RuntimeDataStorage.Add(RigVMRuntimeDataID, FRigVMRuntimeData());
-	NewRigVMRuntimeData.Context = ReferenceContext;
-	RigVMRuntimeData = &NewRigVMRuntimeData;
+	TSharedPtr<FRigVMRuntimeData> RigVMRuntimeData = AddGlobalRuntimeData(RigVMRuntimeDataID, ReferenceContext);
+	Private::LocalRuntimeDataStorage.Add(RigVMRuntimeDataID, RigVMRuntimeData);
+	if (URigVM* VM = RigVMRuntimeDataID.ResolveObjectPtr())
+	{
+		VM->InitializeInstance(RigVMRuntimeData->Context, false);
+	}
 
 	return RigVMRuntimeData;
 }
 
-/*static*/ FRigVMRuntimeData* FRigVMRuntimeDataRegistry::FindOrAddRuntimeData(const FRigVMRuntimeDataID& RigVMRuntimeDataID, const FRigVMExtendedExecuteContext& ReferenceContext)
+/*static*/ void FRigVMRuntimeDataRegistry::ReleaseAllVMRuntimeData(const FRigVMRuntimeDataID& RigVMRuntimeDataID)
 {
 	check(Private::bInitialized);
 
-	FRigVMRuntimeData* RigVMRuntimeData = nullptr;
+	ReleaseAllGlobalRuntimeData(RigVMRuntimeDataID);
+}
 
-	if (RigVMRuntimeData = FindRuntimeData(RigVMRuntimeDataID); RigVMRuntimeData != nullptr)
-	{
-		if (RigVMRuntimeData->Context.VMHash != ReferenceContext.VMHash)
-		{
-			RigVMRuntimeData->Context = ReferenceContext;
-		}
+/*static*/ TSharedPtr<FRigVMRuntimeData> FRigVMRuntimeDataRegistry::AddGlobalRuntimeData(const FRigVMRuntimeDataID& RigVMRuntimeDataID, const FRigVMExtendedExecuteContext& ReferenceContext)
+{
+	TSharedPtr<FRigVMRuntimeData> RigVMRuntimeData = nullptr;
 
-		return RigVMRuntimeData;
+	{
+		FRWScopeLock Lock(Private::GlobalRuntimeDataStorageLock, SLT_Write);
+		RigVMRuntimeData = Private::GlobalRuntimeDataStorage.Add(RigVMRuntimeDataID, MakeShared<FRigVMRuntimeData>());
+		RigVMRuntimeData->Context = ReferenceContext;
 	}
 
-	return AddRuntimeData(RigVMRuntimeDataID, ReferenceContext);
+	return RigVMRuntimeData;
+}
+
+/*static*/ void FRigVMRuntimeDataRegistry::ReleaseAllGlobalRuntimeData(const FRigVMRuntimeDataID& RigVMRuntimeDataID)
+{
+	FRWScopeLock Lock(Private::GlobalRuntimeDataStorageLock, SLT_Write);
+	Private::GlobalRuntimeDataStorage.Remove(RigVMRuntimeDataID);
 }
 
 /*static*/ void FRigVMRuntimeDataRegistry::HandlePostGarbageCollect()
 {
 	Private::GCCycle++;
-	Private::LocalGCCycle = Private::GCCycle; // avoid additional compactions on main thread
+	Private::LocalGCCycle = Private::GCCycle; // Avoid additional compaction on main thread
 
-	PerformStorageCompaction();
+	PerformLocalStorageCompaction();
+	PerformGlobalStorageCompaction();
+}
+
+/*static*/ void FRigVMRuntimeDataRegistry::PerformGlobalStorageCompaction()
+{
+	for (auto Iter = Private::GlobalRuntimeDataStorage.CreateIterator(); Iter; ++Iter)
+	{
+		const FRigVMRuntimeDataID& RuntimeDataID = Iter.Key();
+		if (RuntimeDataID.ResolveObjectPtr() == nullptr)
+		{
+			Iter.RemoveCurrent();
+		}
+	}
 }
 
-/*static*/ void FRigVMRuntimeDataRegistry::PerformStorageCompaction()
+/*static*/ void FRigVMRuntimeDataRegistry::PerformLocalStorageCompaction()
 {
-	for (auto Iter = Private::RuntimeDataStorage.CreateIterator(); Iter; ++Iter)
+	for (auto Iter = Private::LocalRuntimeDataStorage.CreateIterator(); Iter; ++Iter)
 	{
 		const FRigVMRuntimeDataID& RuntimeDataID = Iter.Key();
 		if (RuntimeDataID.ResolveObjectPtr() == nullptr)
