Adding new AnimNext EdGraphNode details customization

diff --git a/Engine/Plugins/Experimental/AnimNext/Source/AnimNext/Internal/Graph/RigDecorator_AnimNextCppDecorator.h b/Engine/Plugins/Experimental/AnimNext/Source/AnimNext/Internal/Graph/RigDecorator_AnimNextCppDecorator.h
--- a/Engine/Plugins/Experimental/AnimNext/Source/AnimNext/Internal/Graph/RigDecorator_AnimNextCppDecorator.h
+++ b/Engine/Plugins/Experimental/AnimNext/Source/AnimNext/Internal/Graph/RigDecorator_AnimNextCppDecorator.h
@@ -23,6 +23,8 @@ struct ANIMNEXT_API FRigDecorator_AnimNextCppDecorator : public FRigVMDecorator
 
 #if WITH_EDITOR
 	virtual void GetProgrammaticPins(URigVMController* InController, int32 InParentPinIndex, const FString& InDefaultValue, struct FRigVMPinInfoArray& OutPinArray) const override;
+	
+	virtual UScriptStruct* GetDecoratorSharedDataStruct() const { return DecoratorSharedDataStruct; }
 
 	const UE::AnimNext::FDecorator* GetDecorator() const;
 #endif
diff --git a/Engine/Plugins/Experimental/AnimNext/Source/AnimNext/Private/Decorators/BlendByBool.h b/Engine/Plugins/Experimental/AnimNext/Source/AnimNext/Private/Decorators/BlendByBool.h
--- a/Engine/Plugins/Experimental/AnimNext/Source/AnimNext/Private/Decorators/BlendByBool.h
+++ b/Engine/Plugins/Experimental/AnimNext/Source/AnimNext/Private/Decorators/BlendByBool.h
@@ -12,7 +12,7 @@
 
 #include "BlendByBool.generated.h"
 
-USTRUCT()
+USTRUCT(meta = (DisplayName = "Blend By Bool"))
 struct FAnimNextBlendByBoolDecoratorSharedData : public FAnimNextDecoratorSharedData
 {
 	GENERATED_BODY()
@@ -26,7 +26,7 @@ struct FAnimNextBlendByBoolDecoratorSharedData : public FAnimNextDecoratorShared
 	FAnimNextDecoratorHandle FalseChild;
 
 	/** The boolean condition that decides which child is active. */
-	UPROPERTY()
+	UPROPERTY(EditAnywhere, Category = "Default")
 	bool bCondition = false;
 
 	// Latent pin support boilerplate
diff --git a/Engine/Plugins/Experimental/AnimNext/Source/AnimNext/Private/Decorators/BlendInertializer.h b/Engine/Plugins/Experimental/AnimNext/Source/AnimNext/Private/Decorators/BlendInertializer.h
--- a/Engine/Plugins/Experimental/AnimNext/Source/AnimNext/Private/Decorators/BlendInertializer.h
+++ b/Engine/Plugins/Experimental/AnimNext/Source/AnimNext/Private/Decorators/BlendInertializer.h
@@ -12,14 +12,14 @@
 
 class UBlendProfile;
 
-USTRUCT()
+USTRUCT(meta = (DisplayName = "Blend Inertializer"))
 struct FAnimNextBlendInertializerDecoratorSharedData : public FAnimNextDecoratorSharedData
 {
 	GENERATED_BODY()
 
 	/** Blend profile that configures how fast to blend each bone. */
 	// TODO: Can't show list of blend profiles, we need to find a skeleton to perform the lookup with
-	UPROPERTY(meta = (Inline, UseAsBlendProfile = true))
+	UPROPERTY(EditAnywhere, Category = "Default", meta = (Inline, UseAsBlendProfile = true))
 	TObjectPtr<UBlendProfile> BlendProfile = nullptr;
 };
 
diff --git a/Engine/Plugins/Experimental/AnimNext/Source/AnimNext/Private/Decorators/BlendSmoother.h b/Engine/Plugins/Experimental/AnimNext/Source/AnimNext/Private/Decorators/BlendSmoother.h
--- a/Engine/Plugins/Experimental/AnimNext/Source/AnimNext/Private/Decorators/BlendSmoother.h
+++ b/Engine/Plugins/Experimental/AnimNext/Source/AnimNext/Private/Decorators/BlendSmoother.h
@@ -14,21 +14,21 @@
 
 class UCurveFloat;
 
-USTRUCT()
+USTRUCT(meta = (DisplayName = "Blend Smoother"))
 struct FAnimNextBlendSmootherDecoratorSharedData : public FAnimNextDecoratorSharedData
 {
 	GENERATED_BODY()
 
 	/** How long to take when blending into each child. */
-	UPROPERTY(meta = (Inline))
+	UPROPERTY(EditAnywhere, Category = "Default", meta = (Inline))
 	TArray<float> BlendTimes;
 
 	/** What type of blend equation to use when converting the time elapsed into a blend weight. */
-	UPROPERTY(meta = (Inline))
+	UPROPERTY(EditAnywhere, Category = "Default", meta = (Inline))
 	EAlphaBlendOption BlendType = EAlphaBlendOption::Linear;
 
 	/** Custom curve to use when the Custom blend type is used. */
-	UPROPERTY(meta = (Inline))
+	UPROPERTY(EditAnywhere, Category = "Default", meta = (Inline))
 	TObjectPtr<UCurveFloat> CustomBlendCurve;
 };
 
diff --git a/Engine/Plugins/Experimental/AnimNext/Source/AnimNext/Private/Decorators/BlendSmootherPerBone.h b/Engine/Plugins/Experimental/AnimNext/Source/AnimNext/Private/Decorators/BlendSmootherPerBone.h
--- a/Engine/Plugins/Experimental/AnimNext/Source/AnimNext/Private/Decorators/BlendSmootherPerBone.h
+++ b/Engine/Plugins/Experimental/AnimNext/Source/AnimNext/Private/Decorators/BlendSmootherPerBone.h
@@ -15,14 +15,14 @@
 
 class UBlendProfile;
 
-USTRUCT()
+USTRUCT(meta = (DisplayName = "Blend Smoother Per Bone"))
 struct FAnimNextBlendSmootherPerBoneDecoratorSharedData : public FAnimNextDecoratorSharedData
 {
 	GENERATED_BODY()
 
 	/** Blend profile that configures how fast to blend each bone. */
 	// TODO: Can't show list of blend profiles, we need to find a skeleton to perform the lookup with
-	UPROPERTY(meta = (Inline, UseAsBlendProfile=true))
+	UPROPERTY(EditAnywhere, Category = "Default", meta = (Inline, UseAsBlendProfile=true))
 	TObjectPtr<UBlendProfile> BlendProfile = nullptr;
 };
 
diff --git a/Engine/Plugins/Experimental/AnimNext/Source/AnimNext/Private/Decorators/BlendTwoWay.h b/Engine/Plugins/Experimental/AnimNext/Source/AnimNext/Private/Decorators/BlendTwoWay.h
--- a/Engine/Plugins/Experimental/AnimNext/Source/AnimNext/Private/Decorators/BlendTwoWay.h
+++ b/Engine/Plugins/Experimental/AnimNext/Source/AnimNext/Private/Decorators/BlendTwoWay.h
@@ -12,7 +12,7 @@
 
 #include "BlendTwoWay.generated.h"
 
-USTRUCT()
+USTRUCT(meta = (DisplayName = "Blend Two Way"))
 struct FAnimNextBlendTwoWayDecoratorSharedData : public FAnimNextDecoratorSharedData
 {
 	GENERATED_BODY()
@@ -26,7 +26,7 @@ struct FAnimNextBlendTwoWayDecoratorSharedData : public FAnimNextDecoratorShared
 	FAnimNextDecoratorHandle ChildB;
 
 	/** How much to blend our two children: 0.0 is fully child A while 1.0 is fully child B. */
-	UPROPERTY()
+	UPROPERTY(EditAnywhere, Category = "Default")
 	float BlendWeight = 0.0f;
 
 	// Latent pin support boilerplate
diff --git a/Engine/Plugins/Experimental/AnimNext/Source/AnimNext/Private/Decorators/ReferencePoseDecorator.h b/Engine/Plugins/Experimental/AnimNext/Source/AnimNext/Private/Decorators/ReferencePoseDecorator.h
--- a/Engine/Plugins/Experimental/AnimNext/Source/AnimNext/Private/Decorators/ReferencePoseDecorator.h
+++ b/Engine/Plugins/Experimental/AnimNext/Source/AnimNext/Private/Decorators/ReferencePoseDecorator.h
@@ -21,13 +21,13 @@ enum class EAnimNextReferencePoseType : int32
 	AdditiveIdentity,
 };
 
-USTRUCT()
+USTRUCT(meta = (DisplayName = "Reference Pose"))
 struct FAnimNextReferencePoseDecoratorSharedData : public FAnimNextDecoratorSharedData
 {
 	GENERATED_BODY()
 
 	/** The type of the reference pose. */
-	UPROPERTY(meta = (Inline))
+	UPROPERTY(EditAnywhere, Category = "Default", meta = (Inline))
 	EAnimNextReferencePoseType ReferencePoseType = EAnimNextReferencePoseType::MeshLocalSpace;
 };
 
diff --git a/Engine/Plugins/Experimental/AnimNext/Source/AnimNext/Private/Decorators/SequencePlayer.h b/Engine/Plugins/Experimental/AnimNext/Source/AnimNext/Private/Decorators/SequencePlayer.h
--- a/Engine/Plugins/Experimental/AnimNext/Source/AnimNext/Private/Decorators/SequencePlayer.h
+++ b/Engine/Plugins/Experimental/AnimNext/Source/AnimNext/Private/Decorators/SequencePlayer.h
@@ -11,25 +11,25 @@
 
 #include "SequencePlayer.generated.h"
 
-USTRUCT()
+USTRUCT(meta = (DisplayName = "Sequence Player"))
 struct FAnimNextSequencePlayerDecoratorSharedData : public FAnimNextDecoratorSharedData
 {
 	GENERATED_BODY()
 
 	/** The sequence to play. */
-	UPROPERTY(meta = (Inline))
+	UPROPERTY(EditAnywhere, Category = "Default", meta = (Inline))
 	TObjectPtr<UAnimSequence> AnimSequence;
 
 	/** The play rate multiplier at which this sequence plays. */
-	UPROPERTY()
+	UPROPERTY(EditAnywhere, Category = "Default")
 	float PlayRate = 1.0f;
 
 	/** The time at which we should start playing this sequence. */
-	UPROPERTY()
+	UPROPERTY(EditAnywhere, Category = "Default")
 	float StartPosition = 0.0f;
 
 	/** Whether or not this sequence playback will loop. */
-	UPROPERTY()
+	UPROPERTY(EditAnywhere, Category = "Default")
 	bool bLoop = 0.0f;
 
 	// Latent pin support boilerplate
diff --git a/Engine/Plugins/Experimental/AnimNext/Source/AnimNext/Private/Decorators/SynchronizeUsingGroups.h b/Engine/Plugins/Experimental/AnimNext/Source/AnimNext/Private/Decorators/SynchronizeUsingGroups.h
--- a/Engine/Plugins/Experimental/AnimNext/Source/AnimNext/Private/Decorators/SynchronizeUsingGroups.h
+++ b/Engine/Plugins/Experimental/AnimNext/Source/AnimNext/Private/Decorators/SynchronizeUsingGroups.h
@@ -11,18 +11,18 @@
 
 #include "SynchronizeUsingGroups.generated.h"
 
-USTRUCT()
+USTRUCT(meta = (DisplayName = "Synchronize Using Groups"))
 struct FAnimNextSynchronizeUsingGroupsDecoratorSharedData : public FAnimNextDecoratorSharedData
 {
 	GENERATED_BODY()
 
 	// The group name
 	// If no name is provided, this decorator is inactive
-	UPROPERTY(meta = (Inline))
+	UPROPERTY(EditAnywhere, Category = "Default", meta = (Inline))
 	FName GroupName;
 
 	// The role this player can assume within the group
-	UPROPERTY(meta = (Inline))
+	UPROPERTY(EditAnywhere, Category = "Default", meta = (Inline))
 	TEnumAsByte<EAnimGroupRole::Type> GroupRole = EAnimGroupRole::CanBeLeader;
 };
 
diff --git a/Engine/Plugins/Experimental/AnimNext/Source/AnimNextEditor/Private/AnimNextEditorModule.cpp b/Engine/Plugins/Experimental/AnimNext/Source/AnimNextEditor/Private/AnimNextEditorModule.cpp
--- a/Engine/Plugins/Experimental/AnimNext/Source/AnimNextEditor/Private/AnimNextEditorModule.cpp
+++ b/Engine/Plugins/Experimental/AnimNext/Source/AnimNextEditor/Private/AnimNextEditorModule.cpp
@@ -23,6 +23,8 @@
 #include "Workspace/AnimNextWorkspaceEditor.h"
 #include "Param/ParameterBlockParameterCustomization.h"
 #include "Param/AnimNextParameterBlockParameter.h"
+#include "Graph/AnimNextGraph_EdGraphNode.h"
+#include "Graph/AnimNextGraph_EdGraphNodeCustomization.h"
 
 #define LOCTEXT_NAMESPACE "AnimNextEditorModule"
 
@@ -56,6 +58,9 @@ class FModule : public IModule
 		PropertyModule.RegisterCustomClassLayout(UAnimNextParameterBlockParameter::StaticClass()->GetFName(), 
 			FOnGetDetailCustomizationInstance::CreateLambda([] { return MakeShared<FParameterBlockParameterCustomization>(); }));
 
+		PropertyModule.RegisterCustomClassLayout(UAnimNextGraph_EdGraphNode::StaticClass()->GetFName(),
+			FOnGetDetailCustomizationInstance::CreateLambda([] { return MakeShared<FAnimNextGraph_EdGraphNodeCustomization>(); }));
+
 		AnimNextGraphPanelNodeFactory = MakeShared<FAnimNextGraphPanelNodeFactory>();
 		FEdGraphUtilities::RegisterVisualNodeFactory(AnimNextGraphPanelNodeFactory);
 
@@ -142,6 +147,9 @@ class FModule : public IModule
 		{
 			UAnimNextGraph* Graph = CastChecked<UAnimNextGraph>(InAsset);
 			UAnimNextGraph_EditorData* EditorData = UncookedOnly::FUtils::GetEditorData(Graph);
+			
+			EditorData->GetRigVMGraphModifiedEvent().RemoveAll(&InEditor.Get());
+			EditorData->GetRigVMGraphModifiedEvent().AddSP(InEditor, &FWorkspaceEditor::OnGraphModified);
 
 			return SNew(SAnimNextGraphView, EditorData)
 				.OnOpenGraph_Lambda([WeakEditor = TWeakPtr<FWorkspaceEditor>(InEditor)](URigVMGraph* InGraph)
diff --git a/Engine/Plugins/Experimental/AnimNext/Source/AnimNextEditor/Private/Workspace/AnimNextWorkspaceEditor.cpp b/Engine/Plugins/Experimental/AnimNext/Source/AnimNextEditor/Private/Workspace/AnimNextWorkspaceEditor.cpp
--- a/Engine/Plugins/Experimental/AnimNext/Source/AnimNextEditor/Private/Workspace/AnimNextWorkspaceEditor.cpp
+++ b/Engine/Plugins/Experimental/AnimNext/Source/AnimNextEditor/Private/Workspace/AnimNextWorkspaceEditor.cpp
@@ -1,4 +1,4 @@
-﻿// Copyright Epic Games, Inc. All Rights Reserved.
+// Copyright Epic Games, Inc. All Rights Reserved.
 
 #include "AnimNextWorkspaceEditor.h"
 
@@ -26,6 +26,8 @@
 #include "Scheduler/AnimNextSchedule.h"
 #include "Framework/Application/SlateApplication.h"
 #include "Graph/AnimNextGraph.h"
+#include "Graph/AnimNextGraph_EdGraphNode.h"
+#include "GraphEditAction.h"
 
 #define LOCTEXT_NAMESPACE "AnimNextWorkspaceEditor"
 
@@ -202,6 +204,11 @@ void FWorkspaceEditor::OpenWorkspaceForAsset(UObject* InAsset, EOpenWorkspaceMet
 	}
 }
 
+void FWorkspaceEditor::OnGraphSelectionChanged(const TSet<UObject*>& NewSelection)
+{
+	SetSelectedObjects(NewSelection.Array());
+}
+
 void FWorkspaceEditor::RestoreEditedObjectState()
 {
 	for (const FEditedDocumentInfo& Document : Workspace->LastEditedDocuments)
@@ -318,10 +325,23 @@ void FWorkspaceEditor::SaveAsset_Execute()
 	}
 }
 
+void FWorkspaceEditor::OnGraphModified(ERigVMGraphNotifType Type, URigVMGraph* Graph, UObject* Subject)
+{
+	if (Type == ERigVMGraphNotifType::PinDefaultValueChanged)
+	{
+		if (DetailsView.IsValid())
+		{
+			DetailsView->ForceRefresh();
+		}
+	}
+}
+
 void FWorkspaceEditor::SetFocusedGraphEditor(TSharedPtr<SGraphEditor> InGraphEditor)
 {
 	// Update the graph editor that is currently focused
 	FocusedGraphEdPtr = InGraphEditor;
+
+	SetSelectedObjects({});
 }
 
 UEdGraph* FWorkspaceEditor::GetFocusedGraph() const
@@ -548,6 +568,15 @@ bool FWorkspaceEditor::OnRequestClose(EAssetEditorCloseReason InCloseReason)
 	return true;
 }
 
+void FWorkspaceEditor::OnClose()
+{
+	if (DetailsView.IsValid())
+	{
+		DetailsView->SetObject(nullptr);
+		DetailsView.Reset();
+	}
+}
+
 }
 
 #undef LOCTEXT_NAMESPACE
\ No newline at end of file
diff --git a/Engine/Plugins/Experimental/AnimNext/Source/AnimNextEditor/Private/Workspace/AnimNextWorkspaceEditor.h b/Engine/Plugins/Experimental/AnimNext/Source/AnimNextEditor/Private/Workspace/AnimNextWorkspaceEditor.h
--- a/Engine/Plugins/Experimental/AnimNext/Source/AnimNextEditor/Private/Workspace/AnimNextWorkspaceEditor.h
+++ b/Engine/Plugins/Experimental/AnimNext/Source/AnimNextEditor/Private/Workspace/AnimNextWorkspaceEditor.h
@@ -1,4 +1,4 @@
-﻿// Copyright Epic Games, Inc. All Rights Reserved.
+// Copyright Epic Games, Inc. All Rights Reserved.
 
 #pragma once
 
@@ -70,6 +70,8 @@ public:
 	// Open an asset inside the workspace editor.
 	static void OpenWorkspaceForAsset(UObject* InAsset, EOpenWorkspaceMethod InOpenMethod);
 
+	void OnGraphSelectionChanged(const TSet<UObject*>& NewSelection);
+
 private:
 	friend class FWorkspaceEditorMode;
 	friend struct FGraphDocumentSummoner;
@@ -87,6 +89,7 @@ private:
 	virtual void InitToolMenuContext(FToolMenuContext& InMenuContext) override;
 	virtual void SaveAsset_Execute() override;
 	virtual bool OnRequestClose(EAssetEditorCloseReason InCloseReason) override;
+	virtual void OnClose() override;
 
 	// FAssetEditorToolkit interface
 	virtual void GetSaveableObjects(TArray<UObject*>& OutObjects) const override;
@@ -140,6 +143,8 @@ private:
 
 	void HandleSaveDocumentState(UObject* InObject);
 
+	void OnGraphModified(ERigVMGraphNotifType Type, URigVMGraph* Graph, UObject* Subject);
+
 	// The asset we are editing
 	UAnimNextWorkspace* Workspace = nullptr;
 
diff --git a/Engine/Plugins/Experimental/AnimNext/Source/AnimNextEditor/Private/Workspace/GraphDocumentSummoner.cpp b/Engine/Plugins/Experimental/AnimNext/Source/AnimNextEditor/Private/Workspace/GraphDocumentSummoner.cpp
--- a/Engine/Plugins/Experimental/AnimNext/Source/AnimNextEditor/Private/Workspace/GraphDocumentSummoner.cpp
+++ b/Engine/Plugins/Experimental/AnimNext/Source/AnimNextEditor/Private/Workspace/GraphDocumentSummoner.cpp
@@ -1,4 +1,4 @@
-﻿// Copyright Epic Games, Inc. All Rights Reserved.
+// Copyright Epic Games, Inc. All Rights Reserved.
 
 #include "GraphDocumentSummoner.h"
 
@@ -86,6 +86,7 @@ TSharedRef<SWidget> FGraphDocumentSummoner::CreateTabBodyForObject(const FWorkfl
 {
 	SGraphEditor::FGraphEditorEvents Events;
 	Events.OnCreateActionMenu = SGraphEditor::FOnCreateActionMenu::CreateSP(this, &FGraphDocumentSummoner::OnCreateGraphActionMenu);
+	Events.OnSelectionChanged = SGraphEditor::FOnSelectionChanged::CreateSP(HostingAppPtr.Pin().Get(), &FWorkspaceEditor::OnGraphSelectionChanged);
 	
 	return SNew(SGraphEditor)
 		.AdditionalCommands(CommandList)
diff --git a/Engine/Plugins/Experimental/AnimNext/Source/AnimNextUncookedOnly/AnimNextUncookedOnly.Build.cs b/Engine/Plugins/Experimental/AnimNext/Source/AnimNextUncookedOnly/AnimNextUncookedOnly.Build.cs
--- a/Engine/Plugins/Experimental/AnimNext/Source/AnimNextUncookedOnly/AnimNextUncookedOnly.Build.cs
+++ b/Engine/Plugins/Experimental/AnimNext/Source/AnimNextUncookedOnly/AnimNextUncookedOnly.Build.cs
@@ -30,6 +30,7 @@ namespace UnrealBuildTool.Rules
 					"Slate",
 					"SlateCore",
 					"StructUtils",
+					"StructUtilsEditor",
 					"ToolMenus",
 				}
 			);
diff --git a/Engine/Plugins/Experimental/AnimNext/Source/AnimNextUncookedOnly/Internal/Graph/AnimNextGraph_EdGraphNodeCustomization.h b/Engine/Plugins/Experimental/AnimNext/Source/AnimNextUncookedOnly/Internal/Graph/AnimNextGraph_EdGraphNodeCustomization.h
new file mode 100644
--- /dev/null
+++ b/Engine/Plugins/Experimental/AnimNext/Source/AnimNextUncookedOnly/Internal/Graph/AnimNextGraph_EdGraphNodeCustomization.h
@@ -0,0 +1,35 @@
+// Copyright Epic Games, Inc. All Rights Reserved.
+
+#pragma once
+
+#include "IDetailCustomization.h"
+#include "IPropertyTypeCustomization.h"
+#include "PropertyEditorModule.h"
+#include "RigVMCore/RigVMMemoryStorageStruct.h"
+
+class URigVMPin;
+class IDetailCategoryBuilder;
+class UAnimNextGraph_EdGraphNode;
+
+namespace UE::AnimNext::Editor
+{
+
+class ANIMNEXTUNCOOKEDONLY_API FAnimNextGraph_EdGraphNodeCustomization : public IDetailCustomization
+{
+private:
+
+	/** Called when details should be customized */
+	virtual void CustomizeDetails(IDetailLayoutBuilder& DetailBuilder) override;
+
+	FText GetName() const;
+	void SetName(const FText& InNewText, ETextCommit::Type InCommitType);
+	bool OnVerifyNameChanged(const FText& InText, FText& OutErrorMessage);
+
+	static void GenerateMemoryStorage(const TArray<URigVMPin*>& ModelPinsToDisplay, FRigVMMemoryStorageStruct& MemoryStorage);
+	static void PopulateCategory(IDetailCategoryBuilder& Category, const TArray<URigVMPin*>& ModelPinsToDisplay, FRigVMMemoryStorageStruct& MemoryStorage, UAnimNextGraph_EdGraphNode* EdGraphNode);
+
+
+	FRigVMMemoryStorageStruct MemoryStorage;
+};
+
+}
diff --git a/Engine/Plugins/Experimental/AnimNext/Source/AnimNextUncookedOnly/Internal/Graph/AnimNextGraph_EditorData.h b/Engine/Plugins/Experimental/AnimNext/Source/AnimNextUncookedOnly/Internal/Graph/AnimNextGraph_EditorData.h
--- a/Engine/Plugins/Experimental/AnimNext/Source/AnimNextUncookedOnly/Internal/Graph/AnimNextGraph_EditorData.h
+++ b/Engine/Plugins/Experimental/AnimNext/Source/AnimNextUncookedOnly/Internal/Graph/AnimNextGraph_EditorData.h
@@ -121,6 +121,11 @@ public:
 	UEdGraph* CreateEdGraph(URigVMGraph* InRigVMGraph, bool bForce);
 	bool RemoveEdGraph(URigVMGraph* InModel);
 
+	FRigVMGraphModifiedEvent& GetRigVMGraphModifiedEvent()
+	{
+		return RigVMGraphModifiedEvent;
+	}
+
 	ANIMNEXTUNCOOKEDONLY_API void Initialize(bool bRecompileVM);
 
 protected:
diff --git a/Engine/Plugins/Experimental/AnimNext/Source/AnimNextUncookedOnly/Private/Graph/AnimNextGraph_EdGraphNodeCustomization.cpp b/Engine/Plugins/Experimental/AnimNext/Source/AnimNextUncookedOnly/Private/Graph/AnimNextGraph_EdGraphNodeCustomization.cpp
new file mode 100644
--- /dev/null
+++ b/Engine/Plugins/Experimental/AnimNext/Source/AnimNextUncookedOnly/Private/Graph/AnimNextGraph_EdGraphNodeCustomization.cpp
@@ -0,0 +1,220 @@
+// Copyright Epic Games, Inc. All Rights Reserved.
+
+#include "Graph/AnimNextGraph_EdGraphNodeCustomization.h"
+
+#include "DetailLayoutBuilder.h"
+#include "DetailCategoryBuilder.h"
+#include "Graph/AnimNextGraph_EdGraphNode.h"
+#include "Graph/RigDecorator_AnimNextCppDecorator.h"
+#include "Logging/LogScopedVerbosityOverride.h"
+#include "InstancedPropertyBagStructureDataProvider.h"
+#include "RigVMModel/RigVMController.h"
+
+#define LOCTEXT_NAMESPACE "EdGraphNodeCustomization"
+
+namespace UE::AnimNext::Editor
+{
+
+void FAnimNextGraph_EdGraphNodeCustomization::CustomizeDetails(IDetailLayoutBuilder& DetailBuilder)
+{
+	TArray<TWeakObjectPtr<UObject>> Objects;
+	DetailBuilder.GetObjectsBeingCustomized(Objects);
+
+	if (Objects.IsEmpty() || Objects.Num() > 1)
+	{
+		return;
+	}
+
+	if (UAnimNextGraph_EdGraphNode* EdGraphNode = Cast<UAnimNextGraph_EdGraphNode>(Objects[0].Get()))
+	{
+		if (URigVMNode* ModelNode = EdGraphNode->GetModelNode())
+		{
+			// If there are any decorators in the node, treat it as a decorator stack
+			const TArray<URigVMPin*> DecoratorPins = ModelNode->GetDecoratorPins();
+			if (DecoratorPins.Num() > 0)
+			{
+				// For each Decorator (represented as a pin in the node)
+				for (const URigVMPin* DecoratorPin : DecoratorPins)
+				{
+					if (DecoratorPin->IsExecuteContext())
+					{
+						continue;
+					}
+
+					// Create a temporary decorator instance, in order to get the correct DecoratorSharedDataStruct
+					if (TSharedPtr<FStructOnScope> ScopedDecorator = ModelNode->GetDecoratorInstance(DecoratorPin->GetFName()))
+					{
+						// Create a scoped struct with the Decorator Shared Instance Data
+						const FRigVMDecorator* Decorator = (FRigVMDecorator*)ScopedDecorator->GetStructMemory();
+
+						// Import the default pin values into the Decorator Shared Instance Data
+						TSharedPtr<FStructOnScope> ScopedSharedData;
+						if (UScriptStruct* DecoratorSharedInstanceData = Decorator->GetDecoratorSharedDataStruct())
+						{
+							ScopedSharedData = MakeShared<FStructOnScope>(DecoratorSharedInstanceData);
+
+							FRigVMPinDefaultValueImportErrorContext ErrorPipe;
+							LOG_SCOPE_VERBOSITY_OVERRIDE(LogExec, ELogVerbosity::Verbose);
+							const FString DefaultValue = DecoratorPin->GetDefaultValue();
+							DecoratorSharedInstanceData->ImportText(*DefaultValue, ScopedSharedData->GetStructMemory(), nullptr, PPF_None, &ErrorPipe, DecoratorSharedInstanceData->GetName());
+						}
+						else // or use the decorator itself for simple decorators
+						{
+							ScopedSharedData = ScopedDecorator;
+						}
+
+						// Create a category with the DisplayName of the decorator Shared Data
+						IDetailCategoryBuilder& ParameterCategory = DetailBuilder.EditCategory(*ScopedSharedData->GetStruct()->GetDisplayNameText().ToString(), FText::GetEmpty(), ECategoryPriority::Important);
+
+						// Now add all the structure properties and handle modifications
+						TArray<IDetailPropertyRow*> OutPropertiesRow;
+						ParameterCategory.AddAllExternalStructureProperties(ScopedSharedData.ToSharedRef(), EPropertyLocation::Default, &OutPropertiesRow);
+						for (IDetailPropertyRow* DetailPropertyRow : OutPropertiesRow)
+						{
+							if (TSharedPtr<IPropertyHandle> PropertyHandle = DetailPropertyRow->GetPropertyHandle(); PropertyHandle.IsValid())
+							{
+								const FProperty* Property = PropertyHandle->GetProperty();
+
+								const auto ModifyAsset = [EdGraphNode]()
+								{
+									EdGraphNode->Modify();
+								};
+
+								const auto UpdatePinDefaultValue = [EdGraphNode, Property, ScopedSharedData]()
+								{
+									// Extract the value from the property and assign it to the Pin as a default value (via Schema)
+									const uint8* StructMemberMemoryPtr = Property->ContainerPtrToValuePtr<uint8>(ScopedSharedData->GetStructMemory());
+									const FString ValueStr = FRigVMStruct::ExportToFullyQualifiedText(Property, StructMemberMemoryPtr, true);
+
+									for (UEdGraphPin* EdGraphPin : EdGraphNode->Pins)
+									{
+										// Find the EdGraphPin that corresponds to the Property
+										if (EdGraphPin->GetFName().ToString().EndsWith(Property->GetFName().ToString(), ESearchCase::CaseSensitive))
+										{
+											if (URigVMPin* ModelPin = EdGraphNode->FindModelPinFromGraphPin(EdGraphPin))
+											{
+												// Set the default value using the Controller
+												EdGraphNode->GetController()->SetPinDefaultValue(ModelPin->GetPinPath(), ValueStr);
+											}
+											break;
+										}
+									}
+								};
+
+								if (Decorator->GetDecoratorSharedDataStruct() != nullptr) // TODO zzz : Investigate why shared works with value and decorator struct with child value
+								{
+									PropertyHandle->SetOnPropertyValuePreChange(FSimpleDelegate::CreateLambda(ModifyAsset));
+									PropertyHandle->SetOnPropertyValueChanged(FSimpleDelegate::CreateLambda(UpdatePinDefaultValue));
+								}
+								else
+								{
+									PropertyHandle->SetOnChildPropertyValuePreChange(FSimpleDelegate::CreateLambda(ModifyAsset));
+									PropertyHandle->SetOnChildPropertyValueChanged(FSimpleDelegate::CreateLambda(UpdatePinDefaultValue));
+								}
+							}
+						}
+					}
+				}
+			}
+			else
+			{
+				// For nodes that aren't decorators, we display the pins as properties
+				const TArray<URigVMPin*> ModelPins = ModelNode->GetPins();
+				TArray<URigVMPin*> PinsToDisplay;
+				PinsToDisplay.Reserve(ModelPins.Num());
+
+				// First, obtain the pins to display
+				for (URigVMPin* Pin : ModelPins)
+				{
+					if (Pin->IsExecuteContext())
+					{
+						continue;
+					}
+					const ERigVMPinDirection PinDirection = Pin->GetDirection();
+					if (PinDirection != ERigVMPinDirection::Hidden && (PinDirection == ERigVMPinDirection::IO || PinDirection == ERigVMPinDirection::Input))
+					{
+						PinsToDisplay.Add(Pin);
+					}
+				}
+
+				if (PinsToDisplay.Num() > 0)
+				{
+					// Create the category where the properties will be added, using the Title of the EdGraphNode
+					IDetailCategoryBuilder& ParameterCategory = DetailBuilder.EditCategory(*EdGraphNode->GetNodeTitle(ENodeTitleType::ListView).ToString(), FText::GetEmpty(), ECategoryPriority::Default);
+
+					// Then, create a custom property bag to store the data, initializating the properties in the struct with the pin default values
+					GenerateMemoryStorage(PinsToDisplay, MemoryStorage);
+					
+					// And finally, add the struct properties to the customization
+					PopulateCategory(ParameterCategory, PinsToDisplay, MemoryStorage, EdGraphNode);
+				}
+			}
+		}
+	}
+}
+
+void FAnimNextGraph_EdGraphNodeCustomization::GenerateMemoryStorage(const TArray<URigVMPin*>& ModelPinsToDisplay, FRigVMMemoryStorageStruct& MemoryStorage)
+{
+	TArray<FRigVMPropertyDescription> PropertyDescriptions;
+	PropertyDescriptions.Reserve(ModelPinsToDisplay.Num());
+
+	for (const URigVMPin* ModelPin : ModelPinsToDisplay)
+	{
+		FRigVMPropertyDescription& PropertyDesc = PropertyDescriptions.AddDefaulted_GetRef();
+
+		PropertyDesc.Name = ModelPin->GetFName();
+		PropertyDesc.Property = nullptr;
+		PropertyDesc.CPPType = ModelPin->GetCPPType();
+		PropertyDesc.CPPTypeObject = ModelPin->GetCPPTypeObject();
+		PropertyDesc.Containers;
+		if (ModelPin->IsArray())
+		{
+			PropertyDesc.Containers.Add(EPinContainerType::Array);
+		}
+		PropertyDesc.DefaultValue = ModelPin->GetDefaultValue();
+	}
+
+	MemoryStorage.AddProperties(PropertyDescriptions);
+}
+
+void FAnimNextGraph_EdGraphNodeCustomization::PopulateCategory(IDetailCategoryBuilder& Category, const TArray<URigVMPin*>& ModelPinsToDisplay, FRigVMMemoryStorageStruct& MemoryStorage, UAnimNextGraph_EdGraphNode* EdGraphNode)
+{
+	for (URigVMPin* ModelPin : ModelPinsToDisplay)
+	{
+		FAddPropertyParams AddPropertyParams;
+		IDetailPropertyRow* DetailPropertyRow = Category.AddExternalStructureProperty(MakeShared<FInstancePropertyBagStructureDataProvider>(MemoryStorage), ModelPin->GetFName(), EPropertyLocation::Default, AddPropertyParams);
+
+		if (TSharedPtr<IPropertyHandle> Handle = DetailPropertyRow->GetPropertyHandle(); Handle.IsValid())
+		{
+			Handle->SetOnChildPropertyValuePreChange(FSimpleDelegate::CreateLambda([EdGraphNode]()
+			{
+				EdGraphNode->Modify(); // needed to enable the transaction when we modify the PropertyBag
+			}));
+			Handle->SetOnPropertyValueChanged(FSimpleDelegate::CreateLambda([&MemoryStorage, EdGraphNode, ModelPin]()
+			{
+				const FString ValueStr = MemoryStorage.GetDataAsStringByName(ModelPin->GetFName());
+				// Set the default value using the Controller
+				EdGraphNode->GetController()->SetPinDefaultValue(ModelPin->GetPinPath(), ValueStr);
+			}));
+		}
+	}
+}
+
+FText FAnimNextGraph_EdGraphNodeCustomization::GetName() const
+{
+	return FText();
+}
+
+void FAnimNextGraph_EdGraphNodeCustomization::SetName(const FText& InNewText, ETextCommit::Type InCommitType)
+{
+}
+
+bool FAnimNextGraph_EdGraphNodeCustomization::OnVerifyNameChanged(const FText& InText, FText& OutErrorMessage)
+{
+	return true;
+}
+
+
+}
+
+#undef LOCTEXT_NAMESPACE
diff --git a/Engine/Plugins/Experimental/AnimNext/Source/AnimNextUncookedOnly/Private/UncookedOnlyUtils.cpp b/Engine/Plugins/Experimental/AnimNext/Source/AnimNextUncookedOnly/Private/UncookedOnlyUtils.cpp
--- a/Engine/Plugins/Experimental/AnimNext/Source/AnimNextUncookedOnly/Private/UncookedOnlyUtils.cpp
+++ b/Engine/Plugins/Experimental/AnimNext/Source/AnimNextUncookedOnly/Private/UncookedOnlyUtils.cpp
@@ -439,14 +439,16 @@ void FUtils::Compile(UAnimNextGraph* InGraph)
 	// Link our shim and evaluator nodes together using the execution context
 	TempController->AddLink(
 		TempShimRootNode->FindPin(GET_MEMBER_NAME_STRING_CHECKED(FRigUnit_AnimNextShimRoot, ExecuteContext)),
-		GraphEvaluatorNode->FindPin(GET_MEMBER_NAME_STRING_CHECKED(FRigUnit_AnimNextGraphEvaluator, ExecuteContext)));
+		GraphEvaluatorNode->FindPin(GET_MEMBER_NAME_STRING_CHECKED(FRigUnit_AnimNextGraphEvaluator, ExecuteContext)),
+		false);
 
 	// Link our latent pins
 	for (const FRigVMPinInfo& LatentPin : LatentPins)
 	{
 		TempController->AddLink(
 			LatentPinMapping[LatentPin.Name],
-			GraphEvaluatorNode->FindPin(LatentPin.Name.ToString()));
+			GraphEvaluatorNode->FindPin(LatentPin.Name.ToString()),
+			false);
 	}
 
 	FDecoratorWriter DecoratorWriter;
diff --git a/Engine/Plugins/Runtime/RigVM/Source/RigVM/Public/RigVMCore/RigVMDecorator.h b/Engine/Plugins/Runtime/RigVM/Source/RigVM/Public/RigVMCore/RigVMDecorator.h
--- a/Engine/Plugins/Runtime/RigVM/Source/RigVM/Public/RigVMCore/RigVMDecorator.h
+++ b/Engine/Plugins/Runtime/RigVM/Source/RigVM/Public/RigVMCore/RigVMDecorator.h
@@ -65,6 +65,8 @@ public:
 	// allows the decorator to return dynamic pins (parent pin index must be INDEX_NONE or point to a valid index of the parent pin in the OutPinArray)
 	virtual void GetProgrammaticPins(URigVMController* InController, int32 InParentPinIndex, const FString& InDefaultValue, struct FRigVMPinInfoArray& OutPinArray) const {}
 
+	virtual UScriptStruct* GetDecoratorSharedDataStruct() const { return nullptr; }
+
 #endif
 
 private:
