CIS fix : removed unused expression

#rnx

diff --git a/Engine/Plugins/Experimental/AnimNext/Source/AnimNextUncookedOnly/Private/Graph/AnimNextGraph_EdGraphNodeCustomization.cpp b/Engine/Plugins/Experimental/AnimNext/Source/AnimNextUncookedOnly/Private/Graph/AnimNextGraph_EdGraphNodeCustomization.cpp
--- a/Engine/Plugins/Experimental/AnimNext/Source/AnimNextUncookedOnly/Private/Graph/AnimNextGraph_EdGraphNodeCustomization.cpp
+++ b/Engine/Plugins/Experimental/AnimNext/Source/AnimNextUncookedOnly/Private/Graph/AnimNextGraph_EdGraphNodeCustomization.cpp
@@ -166,7 +166,6 @@ void FAnimNextGraph_EdGraphNodeCustomization::GenerateMemoryStorage(const TArray
 		PropertyDesc.Property = nullptr;
 		PropertyDesc.CPPType = ModelPin->GetCPPType();
 		PropertyDesc.CPPTypeObject = ModelPin->GetCPPTypeObject();
-		PropertyDesc.Containers;
 		if (ModelPin->IsArray())
 		{
 			PropertyDesc.Containers.Add(EPinContainerType::Array);
