Rework AnimNext and RigVM integration

RigVM is now only called into when evaluating the latent pins avoiding the need to evaluate everything within the context of RigVM.

diff --git a/Engine/Plugins/Experimental/AnimNext/Source/AnimNext/Internal/Graph/AnimNextExecuteContext.h b/Engine/Plugins/Experimental/AnimNext/Source/AnimNext/Internal/Graph/AnimNextExecuteContext.h
--- a/Engine/Plugins/Experimental/AnimNext/Source/AnimNext/Internal/Graph/AnimNextExecuteContext.h
+++ b/Engine/Plugins/Experimental/AnimNext/Source/AnimNext/Internal/Graph/AnimNextExecuteContext.h
@@ -5,37 +5,8 @@
 #include "CoreMinimal.h"
 #include "RigVMCore/RigVMExecuteContext.h"
 #include "Units/RigUnit.h"
-#include "DecoratorBase/DecoratorPtr.h"
 #include "AnimNextExecuteContext.generated.h"
 
-struct FAnimNextGraphInstance;
-class UAnimNextGraph;
-
-namespace UE::AnimNext
-{
-	struct FContext;
-}
-
-/**
-  * An enum to control which simulation steps should be performed when the animation graph is processed.
-  */
-enum class EAnimNextGraphSimulationSteps
-{
-	// No step is performed
-	None = 0x00,
-
-	// Performs the update step
-	Update = 0x01,
-
-	// Performs the evaluate step
-	Evaluate = 0x02,
-
-	// Perform all simulation steps
-	All = 0xff,
-};
-
-ENUM_CLASS_FLAGS(EAnimNextGraphSimulationSteps)
-
 USTRUCT(BlueprintType)
 struct FAnimNextExecuteContext : public FRigVMExecuteContext
 {
@@ -43,60 +14,36 @@ struct FAnimNextExecuteContext : public FRigVMExecuteContext
 
 	FAnimNextExecuteContext() = default;
 
-	const UE::AnimNext::FContext& GetContext() const
-	{
-		check(Context != nullptr);
-		return *Context;
-	}
-
-	FAnimNextGraphInstance& GetGraphInstance() const
-	{
-		check(GraphInstance != nullptr);
-		return *GraphInstance;
-	}
-
-	EAnimNextGraphSimulationSteps GetSimulationSteps() const { return SimulationSteps; }
+	int32 GetLatentPinIndex() const { return LatentPinIndex; }
+	void* GetDestinationPtr() const { return DestinationPtr; }
 
 	virtual void Copy(const FRigVMExecuteContext* InOtherContext) override
 	{
 		Super::Copy(InOtherContext);
 
 		const FAnimNextExecuteContext* OtherContext = (const FAnimNextExecuteContext*)InOtherContext;
-		Context = OtherContext->Context;
-		GraphInstance = OtherContext->GraphInstance;
-		SimulationSteps = OtherContext->SimulationSteps;
+		LatentPinIndex = OtherContext->LatentPinIndex;
+		DestinationPtr = OtherContext->DestinationPtr;
 	}
 
 private:
-	void SetContextData(const UE::AnimNext::FContext& InContext)
-	{
-		Context = &InContext;
-	}
-
-	void SetGraphInstance(FAnimNextGraphInstance& InGraphInstance)
-	{
-		GraphInstance = &InGraphInstance;
-	}
-
-	void SetSimulationSteps(EAnimNextGraphSimulationSteps InSimulationSteps)
+	void SetupForExecution(int32 InLatentPinIndex, void* InDestinationPtr)
 	{
-		SimulationSteps = InSimulationSteps;
+		LatentPinIndex = InLatentPinIndex;
+		DestinationPtr = InDestinationPtr;
 	}
 
 	// Call this to reset the context to its original state to detect stale usage, can't call it Reset due to virtual in base with that name
 	void DebugReset()
 	{
-		Context = nullptr;
-		GraphInstance = nullptr;
-		SimulationSteps = EAnimNextGraphSimulationSteps::None;
+		LatentPinIndex = INDEX_NONE;
+		DestinationPtr = nullptr;
 	}
 
-	const UE::AnimNext::FContext* Context = nullptr;
-	FAnimNextGraphInstance* GraphInstance = nullptr;
-	EAnimNextGraphSimulationSteps SimulationSteps = EAnimNextGraphSimulationSteps::None;
+	int32 LatentPinIndex = INDEX_NONE;
+	void* DestinationPtr = nullptr;
 
-	friend class UAnimNextGraph;
-	friend class UAnimNextParameterBlock;
+	friend struct FAnimNextGraphInstance;
 };
 
 USTRUCT(meta=(ExecuteContext="FAnimNextExecuteContext"))
diff --git a/Engine/Plugins/Experimental/AnimNext/Source/AnimNext/Internal/Graph/AnimNextGraph.h b/Engine/Plugins/Experimental/AnimNext/Source/AnimNext/Internal/Graph/AnimNextGraph.h
--- a/Engine/Plugins/Experimental/AnimNext/Source/AnimNext/Internal/Graph/AnimNextGraph.h
+++ b/Engine/Plugins/Experimental/AnimNext/Source/AnimNext/Internal/Graph/AnimNextGraph.h
@@ -20,7 +20,6 @@ class UAnimNextGraph;
 class UAnimGraphNode_AnimNextGraph;
 struct FAnimNode_AnimNextGraph;
 struct FRigUnit_AnimNextGraphEvaluator;
-enum class EAnimNextGraphSimulationSteps;
 
 namespace UE::AnimNext
 {
@@ -83,6 +82,9 @@ struct ANIMNEXT_API FAnimNextGraphInstance
 	// Returns the graph used by this instance or nullptr if the instance is invalid
 	const UAnimNextGraph* GetGraph() const;
 
+	// Returns a weak handle to the root decorator instance
+	UE::AnimNext::FWeakDecoratorPtr GetGraphRootPtr() const;
+
 	// Check to see if this instance data matches the provided graph
 	bool UsesGraph(const UAnimNextGraph* InGraph) const;
 
@@ -111,6 +113,9 @@ private:
 	// Adds the specified component and returns a reference to it
 	UE::AnimNext::FGraphInstanceComponent& AddComponent(int32 ComponentNameHash, FName ComponentName, TSharedPtr<UE::AnimNext::FGraphInstanceComponent>&& Component);
 
+	// Executes a latent RigVM pin and writes the result into the destination pointer
+	void ExecuteLatentPin(int32 LatentPinIndex, void* DestinationPtr);
+
 	// Hard reference to the graph used to create this instance to ensure we can release it safely
 	UPROPERTY()
 	TObjectPtr<const UAnimNextGraph> Graph;
@@ -127,6 +132,7 @@ private:
 
 	friend UAnimNextGraph;					// The graph is the one that allocates instances
 	friend FRigUnit_AnimNextGraphEvaluator;	// We evaluate the instance
+	friend UE::AnimNext::FExecutionContext;
 };
 
 template<>
@@ -157,9 +163,6 @@ public:
 	// Allocates an instance of the graph, retain the handle and use it with the Run() function to evaluate it
 	void AllocateInstance(FAnimNextGraphInstance& Instance) const;
 
-	// Run the specified simulation steps on the provided graph with the given context
-	void Run(const UE::AnimNext::FContext& Context, FAnimNextGraphInstance& GraphInstance, EAnimNextGraphSimulationSteps SimulationSteps) const;
-
 	// Get the parameter to use to access the reference pose
 	UE::AnimNext::FParamId GetReferencePoseParam() const { return ReferencePoseId; }
 
diff --git a/Engine/Plugins/Experimental/AnimNext/Source/AnimNext/Internal/Graph/RigUnit_AnimNextEndExecution.h b/Engine/Plugins/Experimental/AnimNext/Source/AnimNext/Internal/Graph/RigUnit_AnimNextEndExecution.h
deleted file mode 100644
diff --git a/Engine/Plugins/Experimental/AnimNext/Source/AnimNext/Private/DecoratorBase/ExecutionContext.cpp b/Engine/Plugins/Experimental/AnimNext/Source/AnimNext/Private/DecoratorBase/ExecutionContext.cpp
--- a/Engine/Plugins/Experimental/AnimNext/Source/AnimNext/Private/DecoratorBase/ExecutionContext.cpp
+++ b/Engine/Plugins/Experimental/AnimNext/Source/AnimNext/Private/DecoratorBase/ExecutionContext.cpp
@@ -30,20 +30,6 @@ namespace UE::AnimNext
 		Private::GThreadLocalExecutionContext = this;
 	}
 
-	FExecutionContext::FExecutionContext(FAnimNextGraphInstance& InGraphInstance, FRigVMExtendedExecuteContext& InRigVMExecuteContext, FRigVMMemoryHandleArray InRigVMLatentMemoryHandles)
-		: NodeTemplateRegistry(FNodeTemplateRegistry::Get())
-		, DecoratorRegistry(FDecoratorRegistry::Get())
-		, Graph(InGraphInstance.GetGraph())
-		, GraphInstance(&InGraphInstance)
-		, GraphSharedData(Graph->SharedDataBuffer)
-		, RigVMLatentMemoryHandles(InRigVMLatentMemoryHandles)
-		, RigVMExecuteContext(&InRigVMExecuteContext)
-	{
-		// There can be only one execution context alive per thread
-		ensure(Private::GThreadLocalExecutionContext == nullptr);
-		Private::GThreadLocalExecutionContext = this;
-	}
-
 	FExecutionContext::~FExecutionContext()
 	{
 		// There can be only one execution context alive per thread
@@ -336,23 +322,10 @@ namespace UE::AnimNext
 		return false;
 	}
 
-	const void* FExecutionContext::EvaluateLatentPinImpl(FLatentPropertyHandle LatentPropertyHandle) const
+	void FExecutionContext::EvaluateLatentPinImpl(FLatentPropertyHandle LatentPropertyHandle, void* DestinationPtr) const
 	{
 		check(LatentPropertyHandle.IsValid());
-		check(RigVMLatentMemoryHandles.IsValidIndex(LatentPropertyHandle.GetLatentPropertyIndex()));
-
-		FRigVMMemoryHandle& MemoryHandle = RigVMLatentMemoryHandles[LatentPropertyHandle.GetLatentPropertyIndex()];
-
-		// This should be an assert. If this triggers, it means that we have a bug in how lazy memory handles
-		// are assigned during compilation. We keep it as an ensure because in this case, we can recover
-		// as even if the memory handle isn't lazy, it remains valid and we can use it. It won't have the
-		// value we expect but it'll work. The ensure will signal that we need to fix the bug.
-		if (ensure(MemoryHandle.IsLazy()))
-		{
-			MemoryHandle.ComputeLazyValueIfNecessary(*RigVMExecuteContext, RigVMExecuteContext->GetSlice().GetIndex());
-		}
-
-		return MemoryHandle.GetData();
+		GraphInstance->ExecuteLatentPin(LatentPropertyHandle.GetLatentPropertyIndex(), DestinationPtr);
 	}
 
 	const FNodeDescription& FExecutionContext::GetNodeDescription(FNodeHandle NodeHandle) const
diff --git a/Engine/Plugins/Experimental/AnimNext/Source/AnimNext/Private/Graph/AnimGraph/AnimNode_AnimNextGraph.cpp b/Engine/Plugins/Experimental/AnimNext/Source/AnimNext/Private/Graph/AnimGraph/AnimNode_AnimNextGraph.cpp
--- a/Engine/Plugins/Experimental/AnimNext/Source/AnimNext/Private/Graph/AnimGraph/AnimNode_AnimNextGraph.cpp
+++ b/Engine/Plugins/Experimental/AnimNext/Source/AnimNext/Private/Graph/AnimGraph/AnimNode_AnimNextGraph.cpp
@@ -17,6 +17,9 @@
 #include "Param/ParamStack.h"
 #include "AnimGraphParamStackScope.h"
 #include "Scheduler/ScheduleContext.h"
+#include "DecoratorInterfaces/IEvaluate.h"
+#include "DecoratorInterfaces/IUpdate.h"
+#include "EvaluationVM/EvaluationVM.h"
 
 #include UE_INLINE_GENERATED_CPP_BY_NAME(AnimNode_AnimNextGraph)
 
@@ -63,19 +66,17 @@ void FAnimNode_AnimNextGraph::Update_AnyThread(const FAnimationUpdateContext& Co
 
 		PropagateInputProperties(Context.AnimInstanceProxy->GetAnimInstanceObject());
 
-		// Populate our param stack since our instance data might need it during construction
 		const int32 LODLevel = Context.AnimInstanceProxy->GetLODLevel();
 
 		UE::AnimNext::FAnimGraphParamStackScope Scope(Context);
 
 		FParamStack& ParamStack = FParamStack::Get();
 		FParamStack::FPushedLayerHandle LayerHandle = ParamStack.PushValues(
-			"GraphLODLevel", LODLevel
+			AnimNextGraph->GetCurrentLODParam(), LODLevel
 		);
 
-		FContext AnimNextContext(Context.GetDeltaTime());
-
-		AnimNextGraph->Run(AnimNextContext, GraphInstance, EAnimNextGraphSimulationSteps::Update);
+		FExecutionContext ExecutionContext(GraphInstance);
+		UE::AnimNext::UpdateGraph(ExecutionContext, GraphInstance.GetGraphRootPtr(), Context.GetDeltaTime());
 
 		ParamStack.PopLayer(LayerHandle);
 	}
@@ -103,7 +104,7 @@ void FAnimNode_AnimNextGraph::Initialize_AnyThread(const FAnimationInitializeCon
 
 		FParamStack& ParamStack = FParamStack::Get();
 		FParamStack::FPushedLayerHandle LayerHandle = ParamStack.PushValues(
-			"GraphLODLevel", LODLevel
+			AnimNextGraph->GetCurrentLODParam(), LODLevel
 		);
 
 		AnimNextGraph->AllocateInstance(GraphInstance);
@@ -140,8 +141,6 @@ void FAnimNode_AnimNextGraph::Evaluate_AnyThread(FPoseContext& Output)
 
 		const int32 LODLevel = Output.AnimInstanceProxy->GetLODLevel();
 
-		UE::AnimNext::FContext Context(0.0f);
-
 		FAnimNextGraphLODPose ResultPose(FLODPoseHeap(RefPose, LODLevel, true, Output.ExpectsAdditivePose()));
 
 		FScheduleContext ScheduleContext(SkeletalMeshComponent);
@@ -151,12 +150,35 @@ void FAnimNode_AnimNextGraph::Evaluate_AnyThread(FPoseContext& Output)
 		FParamStack& ParamStack = FParamStack::Get();
 		FParamStack::FPushedLayerHandle LayerHandle = ParamStack.PushValues(
 			AnimNextGraph->GetReferencePoseParam(), GraphReferencePose,
-			"UE_Internal_ResultPose", ResultPose,
-			AnimNextGraph->GetCurrentLODParam(), LODLevel,
-			"UE_Internal_GraphExpectsAdditive", Output.ExpectsAdditivePose()
+			AnimNextGraph->GetCurrentLODParam(), LODLevel
 		);
 
-		AnimNextGraph->Run(Context, GraphInstance, EAnimNextGraphSimulationSteps::Evaluate);
+		{
+			FExecutionContext ExecutionContext(GraphInstance);
+			FEvaluationProgram EvaluationProgram = UE::AnimNext::EvaluateGraph(ExecutionContext, GraphInstance.GetGraphRootPtr());
+
+			FEvaluationVM EvaluationVM(EEvaluationFlags::All, RefPose, LODLevel);
+			bool bHasValidOutput = false;
+
+			if (!EvaluationProgram.IsEmpty())
+			{
+				EvaluationProgram.Execute(EvaluationVM);
+
+				TUniquePtr<FKeyframeState> EvaluatedKeyframe;
+				if (EvaluationVM.PopValue(KEYFRAME_STACK_NAME, EvaluatedKeyframe))
+				{
+					ResultPose.LODPose.CopyFrom(EvaluatedKeyframe->Pose);
+					bHasValidOutput = true;
+				}
+			}
+
+			if (!bHasValidOutput)
+			{
+				// We need to output a valid pose, generate one
+				FKeyframeState ReferenceKeyframe = EvaluationVM.MakeReferenceKeyframe(Output.ExpectsAdditivePose());
+				ResultPose.LODPose.CopyFrom(ReferenceKeyframe.Pose);
+			}
+		}
 
 		FGenerationTools::RemapPose(ResultPose.LODPose, Output);
 
diff --git a/Engine/Plugins/Experimental/AnimNext/Source/AnimNext/Private/Graph/AnimNextGraph.cpp b/Engine/Plugins/Experimental/AnimNext/Source/AnimNext/Private/Graph/AnimNextGraph.cpp
--- a/Engine/Plugins/Experimental/AnimNext/Source/AnimNext/Private/Graph/AnimNextGraph.cpp
+++ b/Engine/Plugins/Experimental/AnimNext/Source/AnimNext/Private/Graph/AnimNextGraph.cpp
@@ -14,7 +14,7 @@
 #include "Serialization/MemoryReader.h"
 #include "AnimNextStats.h"
 
-DEFINE_STAT(STAT_AnimNext_Graph);
+DEFINE_STAT(STAT_AnimNext_Graph_RigVM);
 DEFINE_STAT(STAT_AnimNext_Graph_AllocateInstance);
 
 #include UE_INLINE_GENERATED_CPP_BY_NAME(AnimNextGraph)
@@ -163,6 +163,11 @@ const UAnimNextGraph* FAnimNextGraphInstance::GetGraph() const
 	return Graph;
 }
 
+UE::AnimNext::FWeakDecoratorPtr FAnimNextGraphInstance::GetGraphRootPtr() const
+{
+	return GraphInstancePtr;
+}
+
 bool FAnimNextGraphInstance::UsesGraph(const UAnimNextGraph* InGraph) const
 {
 	return Graph == InGraph;
@@ -196,6 +201,27 @@ GraphInstanceComponentMapType::TConstIterator FAnimNextGraphInstance::GetCompone
 	return Components.CreateConstIterator();
 }
 
+void FAnimNextGraphInstance::ExecuteLatentPin(int32 LatentPinIndex, void* DestinationPtr)
+{
+	SCOPE_CYCLE_COUNTER(STAT_AnimNext_Graph_RigVM);
+
+	if (!IsValid())
+	{
+		return;
+	}
+
+	if (URigVM* VM = Graph->VM)
+	{
+		FAnimNextExecuteContext& AnimNextContext = ExtendedExecuteContext.GetPublicDataSafe<FAnimNextExecuteContext>();
+		AnimNextContext.SetupForExecution(LatentPinIndex, DestinationPtr);
+
+		VM->ExecuteVM(ExtendedExecuteContext, FRigUnit_AnimNextShimRoot::EventName);
+
+		// Reset the context to avoid issues if we forget to reset it the next time we use it
+		AnimNextContext.DebugReset();
+	}
+}
+
 UAnimNextGraph::UAnimNextGraph(const FObjectInitializer& ObjectInitializer)
 	: Super(ObjectInitializer)
 {
@@ -223,29 +249,22 @@ void UAnimNextGraph::AllocateInstance(FAnimNextGraphInstance& Instance) const
 		Instance.GraphInstancePtr = Context.AllocateNodeInstance(UE::AnimNext::FWeakDecoratorPtr(), ResolvedRootDecoratorHandle);
 	}
 
-#if WITH_EDITORONLY_DATA
-	FRWScopeLock Lock(GraphInstancesLock, SLT_Write);
-	check(!GraphInstances.Contains(&Instance));
-	GraphInstances.Add(&Instance);
-#endif
-}
-
-void UAnimNextGraph::Run(const UE::AnimNext::FContext& Context, FAnimNextGraphInstance& GraphInstance, EAnimNextGraphSimulationSteps SimulationSteps) const
-{
-	SCOPE_CYCLE_COUNTER(STAT_AnimNext_Graph);
-
-	if (VM && GraphInstance.IsValid())
+	if (!Instance.GraphInstancePtr.IsValid())
 	{
-		FAnimNextExecuteContext& AnimNextContext = GraphInstance.ExtendedExecuteContext.GetPublicDataSafe<FAnimNextExecuteContext>();
-		AnimNextContext.SetContextData(Context);
-		AnimNextContext.SetGraphInstance(GraphInstance);
-		AnimNextContext.SetSimulationSteps(SimulationSteps);
-
-		VM->ExecuteVM(GraphInstance.ExtendedExecuteContext, FRigUnit_AnimNextShimRoot::EventName);
+		// We failed to allocate our instance, clear everything
+		Instance.Graph = nullptr;
+		Instance.ExtendedExecuteContext.Reset();
+		Instance.Components.Empty();
+	}
 
-		// Reset the context to avoid issues if we forget to reset it the next time we use it
-		AnimNextContext.DebugReset();
+#if WITH_EDITORONLY_DATA
+	if (Instance.GraphInstancePtr.IsValid())
+	{
+		FRWScopeLock Lock(GraphInstancesLock, SLT_Write);
+		check(!GraphInstances.Contains(&Instance));
+		GraphInstances.Add(&Instance);
 	}
+#endif
 }
 
 TArray<FRigVMExternalVariable> UAnimNextGraph::GetRigVMExternalVariables()
diff --git a/Engine/Plugins/Experimental/AnimNext/Source/AnimNext/Private/Graph/RigUnit_AnimNextEndExecution.cpp b/Engine/Plugins/Experimental/AnimNext/Source/AnimNext/Private/Graph/RigUnit_AnimNextEndExecution.cpp
deleted file mode 100644
diff --git a/Engine/Plugins/Experimental/AnimNext/Source/AnimNext/Private/Graph/RigUnit_AnimNextGraphEvaluator.cpp b/Engine/Plugins/Experimental/AnimNext/Source/AnimNext/Private/Graph/RigUnit_AnimNextGraphEvaluator.cpp
--- a/Engine/Plugins/Experimental/AnimNext/Source/AnimNext/Private/Graph/RigUnit_AnimNextGraphEvaluator.cpp
+++ b/Engine/Plugins/Experimental/AnimNext/Source/AnimNext/Private/Graph/RigUnit_AnimNextGraphEvaluator.cpp
@@ -1,15 +1,7 @@
 // Copyright Epic Games, Inc. All Rights Reserved.
 
 #include "Graph/RigUnit_AnimNextGraphEvaluator.h"
-#include "DecoratorBase/ExecutionContext.h"
-#include "DecoratorInterfaces/IUpdate.h"
-#include "DecoratorInterfaces/IEvaluate.h"
-#include "EvaluationVM/EvaluationVM.h"
 #include "Context.h"
-#include "Graph/AnimNextGraph.h"
-#include "Param/ParamStack.h"
-#include "Graph/AnimNext_LODPose.h"
-#include "AnimNextStats.h"
 
 #include UE_INLINE_GENERATED_CPP_BY_NAME(RigUnit_AnimNextGraphEvaluator)
 
@@ -33,62 +25,33 @@ namespace UE::AnimNext::Private
 
 void FRigUnit_AnimNextGraphEvaluator::StaticExecute(FRigVMExtendedExecuteContext& RigVMExecuteContext, FRigVMMemoryHandleArray RigVMMemoryHandles, FRigVMPredicateBranchArray RigVMBranches)
 {
-	using namespace UE::AnimNext;
-
 	const FAnimNextExecuteContext& VMExecuteContext = RigVMExecuteContext.GetPublicData<FAnimNextExecuteContext>();
-	FAnimNextGraphInstance& GraphInstance = VMExecuteContext.GetGraphInstance();
 
-	// Setup what we need to execute
-	FExecutionContext Context(GraphInstance, RigVMExecuteContext, RigVMMemoryHandles);
+	const int32 LatentPinIndex = VMExecuteContext.GetLatentPinIndex();
+	if (LatentPinIndex == INDEX_NONE)
+	{
+		return;
+	}
 
-	const FWeakDecoratorPtr& GraphInstancePtr = GraphInstance.GraphInstancePtr;
-	const EAnimNextGraphSimulationSteps SimulationSteps = VMExecuteContext.GetSimulationSteps();
+	FRigVMMemoryHandle& MemoryHandle = RigVMMemoryHandles[LatentPinIndex];
 
-	if (EnumHasAnyFlags(SimulationSteps, EAnimNextGraphSimulationSteps::Update))
+	// This should be an assert. If this triggers, it means that we have a bug in how lazy memory handles
+	// are assigned during compilation. We keep it as an ensure because in this case, we can recover
+	// as even if the memory handle isn't lazy, it remains valid and we can use it. It won't have the
+	// value we expect but it'll work. The ensure will signal that we need to fix the bug.
+	if (ensure(MemoryHandle.IsLazy()))
 	{
-		const FContext& InterfaceContext = VMExecuteContext.GetContext();
-
-		// Call pre/post update on our graph
-		UpdateGraph(Context, GraphInstancePtr, InterfaceContext.GetDeltaTime());
+		MemoryHandle.ComputeLazyValueIfNecessary(RigVMExecuteContext, RigVMExecuteContext.GetSlice().GetIndex());
 	}
 
-	if (EnumHasAnyFlags(SimulationSteps, EAnimNextGraphSimulationSteps::Evaluate))
-	{
-		// Call pre/post evaluate on our graph
-		FEvaluationProgram EvaluationProgram = EvaluateGraph(Context, GraphInstancePtr);
+	const uint8* SourcePtr = MemoryHandle.GetData();
+	uint8* DestinationPtr = (uint8*)VMExecuteContext.GetDestinationPtr();
 
-		if (!EvaluationProgram.IsEmpty())
-		{
-			FParamStack& ParamStack = FParamStack::Get();
-
-			const UAnimNextGraph* Graph = GraphInstance.GetGraph();
-			const FAnimNextGraphReferencePose* GraphReferencePosePtr = ParamStack.GetParamPtr<FAnimNextGraphReferencePose>(Graph->GetReferencePoseParam());
-			const int32* GraphLODLevelPtr = ParamStack.GetParamPtr<int32>(Graph->GetCurrentLODParam());
-			static FParamId ResultId("UE_Internal_ResultPose");
-			FAnimNextGraphLODPose* ResultPosePtr = ParamStack.GetMutableParamPtr<FAnimNextGraphLODPose>(ResultId);
-			static FParamId ExpectsAdditiveId("UE_Internal_GraphExpectsAdditive");
-			const bool* bExpectsAdditivePtr = ParamStack.GetParamPtr<bool>(ExpectsAdditiveId);
-
-			if(GraphReferencePosePtr && GraphLODLevelPtr && ResultPosePtr && bExpectsAdditivePtr)
-			{
-				FEvaluationVM EvaluationVM(EEvaluationFlags::All, *GraphReferencePosePtr->ReferencePose, *GraphLODLevelPtr);
-				EvaluationProgram.Execute(EvaluationVM);
-
-				TUniquePtr<FKeyframeState> EvaluatedKeyframe;
-				if (EvaluationVM.PopValue(KEYFRAME_STACK_NAME, EvaluatedKeyframe))
-				{
-					ResultPosePtr->LODPose.CopyFrom(EvaluatedKeyframe->Pose);
-				}
-				else
-				{
-					// We need to output a valid pose, generate one
-					
-					FKeyframeState ReferenceKeyframe = EvaluationVM.MakeReferenceKeyframe(*bExpectsAdditivePtr);
-					ResultPosePtr->LODPose.CopyFrom(ReferenceKeyframe.Pose);
-				}
-			}
-		}
-	}
+	// Copy from our source into our destination
+	// We assume the source and destination properties are identical
+	URigVMMemoryStorage::CopyProperty(
+		MemoryHandle.GetProperty(), DestinationPtr,
+		MemoryHandle.GetProperty(), SourcePtr);
 }
 
 void FRigUnit_AnimNextGraphEvaluator::RegisterExecuteMethod(const FAnimNextGraphEvaluatorExecuteDefinition& ExecuteDefinition)
diff --git a/Engine/Plugins/Experimental/AnimNext/Source/AnimNext/Private/Param/ParamStack.cpp b/Engine/Plugins/Experimental/AnimNext/Source/AnimNext/Private/Param/ParamStack.cpp
--- a/Engine/Plugins/Experimental/AnimNext/Source/AnimNext/Private/Param/ParamStack.cpp
+++ b/Engine/Plugins/Experimental/AnimNext/Source/AnimNext/Private/Param/ParamStack.cpp
@@ -272,7 +272,12 @@ TWeakPtr<FParamStack> FParamStack::DetachFromCurrentThread(EDecoalesce InDecoale
 
 FParamStack::FPushedLayerHandle FParamStack::PushLayer(const FParamStackLayerHandle& InLayerHandle)
 {
-	return PushLayerInternal(*InLayerHandle.Layer.Get());
+	if (InLayerHandle.IsValid())
+	{
+		return PushLayerInternal(*InLayerHandle.Layer.Get());
+	}
+
+	return FPushedLayerHandle();
 }
 
 FParamStack::FPushedLayerHandle FParamStack::PushLayer(TConstArrayView<Private::FParamEntry> InParams)
diff --git a/Engine/Plugins/Experimental/AnimNext/Source/AnimNext/Private/Scheduler/AnimNextScheduleGraphTask.cpp b/Engine/Plugins/Experimental/AnimNext/Source/AnimNext/Private/Scheduler/AnimNextScheduleGraphTask.cpp
--- a/Engine/Plugins/Experimental/AnimNext/Source/AnimNext/Private/Scheduler/AnimNextScheduleGraphTask.cpp
+++ b/Engine/Plugins/Experimental/AnimNext/Source/AnimNext/Private/Scheduler/AnimNextScheduleGraphTask.cpp
@@ -5,6 +5,9 @@
 #include "Scheduler/ScheduleContext.h"
 #include "Graph/AnimNextGraph.h"
 #include "Context.h"
+#include "DecoratorInterfaces/IEvaluate.h"
+#include "DecoratorInterfaces/IUpdate.h"
+#include "EvaluationVM/EvaluationVM.h"
 #include "Graph/AnimNextExecuteContext.h"
 #include "Graph/AnimNext_LODPose.h"
 #include "AnimNextStats.h"
@@ -96,23 +99,6 @@ void FAnimNextScheduleGraphTask::RunGraph(const UE::AnimNext::FScheduleContext&
 	check(OutputPose->LODPose.LODLevel == *GraphLODLevel);
 
 	// Push parameter layers that translate schedule data to graph inputs
-	// TODO: This should probably be reworked, its overly complex for what it does!
-	static FParamId ResultId("UE_Internal_ResultPose");
-	static FParamId ExpectsAdditiveId("UE_Internal_GraphExpectsAdditive");
-
-	if(!InstanceData.GraphInputLayers[TaskIndex].IsValid())
-	{
-		InstanceData.GraphInputLayers[TaskIndex] = FParamStack::MakeValuesLayer(
-			ResultId, *OutputPose,
-			ExpectsAdditiveId, false);
-	}
-	else
-	{
-		InstanceData.GraphInputLayers[TaskIndex].SetValues(
-			ResultId, *OutputPose,
-			ExpectsAdditiveId, false);
-	}
-
 	FParamStack::FPushedLayerHandle LayerHandle = ParamStack.PushLayer(InstanceData.GraphInputLayers[TaskIndex]);
 
 	// Internally we use memstack allocation, so we need a mark here
@@ -123,9 +109,35 @@ void FAnimNextScheduleGraphTask::RunGraph(const UE::AnimNext::FScheduleContext&
 	// This reduces churn internally by avoiding a chunk to be repeatedly allocated and freed as we push/pop marks
 	MemStack.Alloc(size_t(FPageAllocator::SmallPageSize) + 1, 16);
 
-	const float DeltaTime = InContext.GetDeltaTime();
-	const FContext Context(DeltaTime);
-	GraphToRun->Run(Context, InstanceData.GraphInstanceData[TaskIndex], EAnimNextGraphSimulationSteps::All);
+	FExecutionContext Context(InstanceData.GraphInstanceData[TaskIndex]);
+
+	UE::AnimNext::UpdateGraph(Context, InstanceData.GraphInstanceData[TaskIndex].GetGraphRootPtr(), InContext.GetDeltaTime());
+
+	{
+		FEvaluationProgram EvaluationProgram = UE::AnimNext::EvaluateGraph(Context, InstanceData.GraphInstanceData[TaskIndex].GetGraphRootPtr());
+
+		FEvaluationVM EvaluationVM(EEvaluationFlags::All, *GraphReferencePose->ReferencePose, *GraphLODLevel);
+		bool bHasValidOutput = false;
+
+		if (!EvaluationProgram.IsEmpty())
+		{
+			EvaluationProgram.Execute(EvaluationVM);
+
+			TUniquePtr<FKeyframeState> EvaluatedKeyframe;
+			if (EvaluationVM.PopValue(KEYFRAME_STACK_NAME, EvaluatedKeyframe))
+			{
+				OutputPose->LODPose.CopyFrom(EvaluatedKeyframe->Pose);
+				bHasValidOutput = true;
+			}
+		}
+
+		if (!bHasValidOutput)
+		{
+			// We need to output a valid pose, generate one
+			FKeyframeState ReferenceKeyframe = EvaluationVM.MakeReferenceKeyframe(false);
+			OutputPose->LODPose.CopyFrom(ReferenceKeyframe.Pose);
+		}
+	}
 
 	ParamStack.PopLayer(LayerHandle);
 }
diff --git a/Engine/Plugins/Experimental/AnimNext/Source/AnimNext/Public/AnimNextStats.h b/Engine/Plugins/Experimental/AnimNext/Source/AnimNext/Public/AnimNextStats.h
--- a/Engine/Plugins/Experimental/AnimNext/Source/AnimNext/Public/AnimNextStats.h
+++ b/Engine/Plugins/Experimental/AnimNext/Source/AnimNext/Public/AnimNextStats.h
@@ -15,7 +15,7 @@ DECLARE_CYCLE_STAT_EXTERN(TEXT("AnimNext Task: Scope Entry"), STAT_AnimNext_Task
 DECLARE_CYCLE_STAT_EXTERN(TEXT("AnimNext Task: Scope Exit"), STAT_AnimNext_Task_ScopeExit, STATGROUP_AnimNext, ANIMNEXT_API);
 
 DECLARE_CYCLE_STAT_EXTERN(TEXT("AnimNext: Allocate Graph Instance"), STAT_AnimNext_Graph_AllocateInstance, STATGROUP_AnimNext, ANIMNEXT_API);
-DECLARE_CYCLE_STAT_EXTERN(TEXT("AnimNext: Run Graph"), STAT_AnimNext_Graph, STATGROUP_AnimNext, ANIMNEXT_API);
+DECLARE_CYCLE_STAT_EXTERN(TEXT("AnimNext: Graph RigVM"), STAT_AnimNext_Graph_RigVM, STATGROUP_AnimNext, ANIMNEXT_API);
 DECLARE_CYCLE_STAT_EXTERN(TEXT("AnimNext: Update Graph"), STAT_AnimNext_UpdateGraph, STATGROUP_AnimNext, ANIMNEXT_API);
 DECLARE_CYCLE_STAT_EXTERN(TEXT("AnimNext: Evaluate Graph"), STAT_AnimNext_EvaluateGraph, STATGROUP_AnimNext, ANIMNEXT_API);
 DECLARE_CYCLE_STAT_EXTERN(TEXT("AnimNext: Execute Evaluation Program"), STAT_AnimNext_EvaluationProgram_Execute, STATGROUP_AnimNext, ANIMNEXT_API);
diff --git a/Engine/Plugins/Experimental/AnimNext/Source/AnimNext/Public/DecoratorBase/ExecutionContext.h b/Engine/Plugins/Experimental/AnimNext/Source/AnimNext/Public/DecoratorBase/ExecutionContext.h
--- a/Engine/Plugins/Experimental/AnimNext/Source/AnimNext/Public/DecoratorBase/ExecutionContext.h
+++ b/Engine/Plugins/Experimental/AnimNext/Source/AnimNext/Public/DecoratorBase/ExecutionContext.h
@@ -11,11 +11,6 @@
 #include "DecoratorBase/NodeHandle.h"
 #include "Graph/AnimNextGraph.h"
 
-struct FRigVMExtendedExecuteContext;
-struct FRigVMMemoryHandle;
-
-using FRigVMMemoryHandleArray = TArrayView<FRigVMMemoryHandle>;
-
 namespace UE::AnimNext
 {
 	struct FNodeDescription;
@@ -36,9 +31,6 @@ namespace UE::AnimNext
 		// Creates an execution context for the specified graph instance
 		explicit FExecutionContext(FAnimNextGraphInstance& InGraphInstance);
 
-		// Creates an execution context for the specified graph with RigVM latent pin support
-		FExecutionContext(FAnimNextGraphInstance& InGraphInstance, FRigVMExtendedExecuteContext& InRigVMExecuteContext, FRigVMMemoryHandleArray InRigVMLatentMemoryHandles);
-
 		// Destroys the execution context
 		~FExecutionContext();
 
@@ -98,7 +90,7 @@ namespace UE::AnimNext
 
 		bool GetInterfaceImpl(FDecoratorInterfaceUID InterfaceUID, const FWeakDecoratorPtr& DecoratorPtr, FDecoratorBinding& InterfaceBinding) const;
 		bool GetInterfaceSuperImpl(FDecoratorInterfaceUID InterfaceUID, const FWeakDecoratorPtr& DecoratorPtr, FDecoratorBinding& SuperBinding) const;
-		const void* EvaluateLatentPinImpl(FLatentPropertyHandle LatentPropertyHandle) const;
+		void EvaluateLatentPinImpl(FLatentPropertyHandle LatentPropertyHandle, void* DestinationPtr) const;
 
 		const FNodeDescription& GetNodeDescription(FNodeHandle NodeHandle) const;
 		const FNodeTemplate* GetNodeTemplate(const FNodeDescription& NodeDesc) const;
@@ -112,8 +104,6 @@ namespace UE::AnimNext
 		const UAnimNextGraph* Graph = nullptr;
 		FAnimNextGraphInstance* GraphInstance = nullptr;
 		TArrayView<const uint8> GraphSharedData;
-		FRigVMMemoryHandleArray RigVMLatentMemoryHandles;
-		FRigVMExtendedExecuteContext* RigVMExecuteContext = nullptr;
 	};
 
 	// Returns a pointer to the current execution context if present, nullptr otherwise.
@@ -156,21 +146,15 @@ namespace UE::AnimNext
 	template<typename LatentPinType>
 	inline LatentPinType FExecutionContext::EvaluateLatentPin(FLatentPropertyHandle LatentPropertyHandle) const
 	{
-		// Latent pin handle needs to be valid
-		ensure(LatentPropertyHandle.IsValid());
-		if (!LatentPropertyHandle.IsValid())
-		{
-			return LatentPinType();
-		}
+		LatentPinType Result{};
 
-		// We need latent pin memory handles, also implies we have a valid RigVM execution context
-		ensure(RigVMLatentMemoryHandles.IsValidIndex(LatentPropertyHandle.GetLatentPropertyIndex()));
-		if (!RigVMLatentMemoryHandles.IsValidIndex(LatentPropertyHandle.GetLatentPropertyIndex()))
+		// Latent pin handle needs to be valid
+		if (ensure(LatentPropertyHandle.IsValid()))
 		{
-			return LatentPinType();
+			EvaluateLatentPinImpl(LatentPropertyHandle, &Result);
 		}
 
-		return *reinterpret_cast<const LatentPinType*>(EvaluateLatentPinImpl(LatentPropertyHandle));
+		return Result;
 	}
 
 	template<class ComponentType>
diff --git a/Engine/Plugins/Experimental/AnimNext/Source/AnimNextTestSuite/Private/AnimNextDecoratorGraphTest.cpp b/Engine/Plugins/Experimental/AnimNext/Source/AnimNextTestSuite/Private/AnimNextDecoratorGraphTest.cpp
--- a/Engine/Plugins/Experimental/AnimNext/Source/AnimNextTestSuite/Private/AnimNextDecoratorGraphTest.cpp
+++ b/Engine/Plugins/Experimental/AnimNext/Source/AnimNextTestSuite/Private/AnimNextDecoratorGraphTest.cpp
@@ -252,9 +252,7 @@ bool FAnimationAnimNextRuntimeTest_GraphExecute::RunTest(const FString& InParame
 	FAnimNextGraphInstance GraphInstance;
 	AnimNextGraph->AllocateInstance(GraphInstance);
 
-	FContext Context(1.0f / 30.0f);
-
-	FParamStack::FPushedLayerHandle LayerHandle = Context.GetMutableParamStack().PushValues(
+	FParamStack::FPushedLayerHandle LayerHandle = ParamStack->PushValues(
 		"UpdateCount", (int32)0,
 		"EvaluateCount", (int32)0,
 		"SomeInt32", (int32)0,
@@ -264,17 +262,21 @@ bool FAnimationAnimNextRuntimeTest_GraphExecute::RunTest(const FString& InParame
 		"SomeLatentFloat", 0.0f
 	);
 
-	AnimNextGraph->Run(Context, GraphInstance, EAnimNextGraphSimulationSteps::All);
+	{
+		FExecutionContext Context(GraphInstance);
+		UE::AnimNext::UpdateGraph(Context, GraphInstance.GetGraphRootPtr(), 1.0f / 30.0f);
+		(void)UE::AnimNext::EvaluateGraph(Context, GraphInstance.GetGraphRootPtr());
+	}
 
-	AddErrorIfFalse(Context.GetParamStack().GetParam<int32>("UpdateCount") == 1, "FAnimationAnimNextRuntimeTest_GraphExecute -> Unexpected update count");
-	AddErrorIfFalse(Context.GetParamStack().GetParam<int32>("EvaluateCount") == 1, "FAnimationAnimNextRuntimeTest_GraphExecute -> Unexpected evaluate count");
-	AddErrorIfFalse(Context.GetParamStack().GetParam<int32>("SomeInt32") == 78, "FAnimationAnimNextRuntimeTest_GraphExecute -> Unexpected SomeInt32 value");
-	AddErrorIfFalse(Context.GetParamStack().GetParam<float>("SomeFloat") == 142.33f, "FAnimationAnimNextRuntimeTest_GraphExecute -> Unexpected SomeFloat value");
-	AddErrorIfFalse(Context.GetParamStack().GetParam<int32>("SomeLatentInt32") == 3, "FAnimationAnimNextRuntimeTest_GraphExecute -> Unexpected SomeLatentInt32 value");
-	AddErrorIfFalse(Context.GetParamStack().GetParam<int32>("SomeOtherLatentInt32") == 3, "FAnimationAnimNextRuntimeTest_GraphExecute -> Unexpected SomeOtherLatentInt32 value");
-	AddErrorIfFalse(Context.GetParamStack().GetParam<float>("SomeLatentFloat") == 34.0f, "FAnimationAnimNextRuntimeTest_GraphExecute -> Unexpected SomeLatentFloat value");
+	AddErrorIfFalse(ParamStack->GetParam<int32>("UpdateCount") == 1, "FAnimationAnimNextRuntimeTest_GraphExecute -> Unexpected update count");
+	AddErrorIfFalse(ParamStack->GetParam<int32>("EvaluateCount") == 1, "FAnimationAnimNextRuntimeTest_GraphExecute -> Unexpected evaluate count");
+	AddErrorIfFalse(ParamStack->GetParam<int32>("SomeInt32") == 78, "FAnimationAnimNextRuntimeTest_GraphExecute -> Unexpected SomeInt32 value");
+	AddErrorIfFalse(ParamStack->GetParam<float>("SomeFloat") == 142.33f, "FAnimationAnimNextRuntimeTest_GraphExecute -> Unexpected SomeFloat value");
+	AddErrorIfFalse(ParamStack->GetParam<int32>("SomeLatentInt32") == 3, "FAnimationAnimNextRuntimeTest_GraphExecute -> Unexpected SomeLatentInt32 value");
+	AddErrorIfFalse(ParamStack->GetParam<int32>("SomeOtherLatentInt32") == 3, "FAnimationAnimNextRuntimeTest_GraphExecute -> Unexpected SomeOtherLatentInt32 value");
+	AddErrorIfFalse(ParamStack->GetParam<float>("SomeLatentFloat") == 34.0f, "FAnimationAnimNextRuntimeTest_GraphExecute -> Unexpected SomeLatentFloat value");
 
-	Context.GetMutableParamStack().PopLayer(LayerHandle);
+	ParamStack->PopLayer(LayerHandle);
 	GraphInstance.Release();
 
 	FParamStack::DetachFromCurrentThread();
@@ -393,9 +395,7 @@ bool FAnimationAnimNextRuntimeTest_GraphExecuteLatent::RunTest(const FString& In
 	FAnimNextGraphInstance GraphInstance;
 	AnimNextGraph->AllocateInstance(GraphInstance);
 
-	FContext Context(1.0f / 30.0f);
-
-	FParamStack::FPushedLayerHandle LayerHandle = Context.GetMutableParamStack().PushValues(
+	FParamStack::FPushedLayerHandle LayerHandle = ParamStack->PushValues(
 		"UpdateCount", (int32)0,
 		"EvaluateCount", (int32)0,
 		"SomeSourceInt", (int32)1223,
@@ -406,17 +406,21 @@ bool FAnimationAnimNextRuntimeTest_GraphExecuteLatent::RunTest(const FString& In
 		"SomeLatentFloat", 0.0f
 	);
 
-	AnimNextGraph->Run(Context, GraphInstance, EAnimNextGraphSimulationSteps::All);
+	{
+		FExecutionContext Context(GraphInstance);
+		UE::AnimNext::UpdateGraph(Context, GraphInstance.GetGraphRootPtr(), 1.0f / 30.0f);
+		(void)UE::AnimNext::EvaluateGraph(Context, GraphInstance.GetGraphRootPtr());
+	}
 
-	AddErrorIfFalse(Context.GetParamStack().GetParam<int32>("UpdateCount") == 1, "FAnimationAnimNextRuntimeTest_GraphExecuteLatent -> Unexpected update count");
-	AddErrorIfFalse(Context.GetParamStack().GetParam<int32>("EvaluateCount") == 1, "FAnimationAnimNextRuntimeTest_GraphExecuteLatent -> Unexpected evaluate count");
-	AddErrorIfFalse(Context.GetParamStack().GetParam<int32>("SomeInt32") == 78, "FAnimationAnimNextRuntimeTest_GraphExecuteLatent -> Unexpected SomeInt32 value");
-	AddErrorIfFalse(Context.GetParamStack().GetParam<float>("SomeFloat") == 142.33f, "FAnimationAnimNextRuntimeTest_GraphExecuteLatent -> Unexpected SomeFloat value");
-	AddErrorIfFalse(Context.GetParamStack().GetParam<int32>("SomeLatentInt32") == 33, "FAnimationAnimNextRuntimeTest_GraphExecuteLatent -> Unexpected SomeLatentInt32 value");
-	AddErrorIfFalse(Context.GetParamStack().GetParam<int32>("SomeOtherLatentInt32") == 1223, "FAnimationAnimNextRuntimeTest_GraphExecuteLatent -> Unexpected SomeOtherLatentInt32 value");
-	AddErrorIfFalse(Context.GetParamStack().GetParam<float>("SomeLatentFloat") == 1123.31f, "FAnimationAnimNextRuntimeTest_GraphExecuteLatent -> Unexpected SomeLatentFloat value");
+	AddErrorIfFalse(ParamStack->GetParam<int32>("UpdateCount") == 1, "FAnimationAnimNextRuntimeTest_GraphExecuteLatent -> Unexpected update count");
+	AddErrorIfFalse(ParamStack->GetParam<int32>("EvaluateCount") == 1, "FAnimationAnimNextRuntimeTest_GraphExecuteLatent -> Unexpected evaluate count");
+	AddErrorIfFalse(ParamStack->GetParam<int32>("SomeInt32") == 78, "FAnimationAnimNextRuntimeTest_GraphExecuteLatent -> Unexpected SomeInt32 value");
+	AddErrorIfFalse(ParamStack->GetParam<float>("SomeFloat") == 142.33f, "FAnimationAnimNextRuntimeTest_GraphExecuteLatent -> Unexpected SomeFloat value");
+	AddErrorIfFalse(ParamStack->GetParam<int32>("SomeLatentInt32") == 33, "FAnimationAnimNextRuntimeTest_GraphExecuteLatent -> Unexpected SomeLatentInt32 value");
+	AddErrorIfFalse(ParamStack->GetParam<int32>("SomeOtherLatentInt32") == 1223, "FAnimationAnimNextRuntimeTest_GraphExecuteLatent -> Unexpected SomeOtherLatentInt32 value");
+	AddErrorIfFalse(ParamStack->GetParam<float>("SomeLatentFloat") == 1123.31f, "FAnimationAnimNextRuntimeTest_GraphExecuteLatent -> Unexpected SomeLatentFloat value");
 
-	Context.GetMutableParamStack().PopLayer(LayerHandle);
+	ParamStack->PopLayer(LayerHandle);
 	GraphInstance.Release();
 
 	FParamStack::DetachFromCurrentThread();
