Fix for AnimNext VM rejecting execution of the Bytecode due to Context public data type mismatch

#rnx

diff --git a/Engine/Plugins/Experimental/AnimNext/Source/AnimNext/Private/Graph/AnimNextGraph.cpp b/Engine/Plugins/Experimental/AnimNext/Source/AnimNext/Private/Graph/AnimNextGraph.cpp
--- a/Engine/Plugins/Experimental/AnimNext/Source/AnimNext/Private/Graph/AnimNextGraph.cpp
+++ b/Engine/Plugins/Experimental/AnimNext/Source/AnimNext/Private/Graph/AnimNextGraph.cpp
@@ -225,6 +225,7 @@ void FAnimNextGraphInstance::ExecuteLatentPin(int32 LatentPinIndex, void* Destin
 UAnimNextGraph::UAnimNextGraph(const FObjectInitializer& ObjectInitializer)
 	: Super(ObjectInitializer)
 {
+	ExtendedExecuteContext.SetContextPublicDataStruct(FAnimNextExecuteContext::StaticStruct());
 	SetRigVMExtendedExecuteContext(&ExtendedExecuteContext);
 }
 
diff --git a/Engine/Plugins/Experimental/AnimNext/Source/AnimNext/Private/Param/AnimNextParameterBlock.cpp b/Engine/Plugins/Experimental/AnimNext/Source/AnimNext/Private/Param/AnimNextParameterBlock.cpp
--- a/Engine/Plugins/Experimental/AnimNext/Source/AnimNext/Private/Param/AnimNextParameterBlock.cpp
+++ b/Engine/Plugins/Experimental/AnimNext/Source/AnimNext/Private/Param/AnimNextParameterBlock.cpp
@@ -16,6 +16,7 @@ DEFINE_STAT(STAT_AnimNext_ParamBlock_UpdateLayer);
 UAnimNextParameterBlock::UAnimNextParameterBlock(const FObjectInitializer& ObjectInitializer)
 	: Super(ObjectInitializer)
 {
+	BaseRigVMContext.SetContextPublicDataStruct(FAnimNextParameterExecuteContext::StaticStruct());
 	SetRigVMExtendedExecuteContext(&BaseRigVMContext);
 }
 
