Add reference to the owning graph in every AnimNext node instance to decouple and simplify execution binding. Allows any external/internal system to call into AnimNext with minimal complexity.

diff --git a/Engine/Plugins/Experimental/AnimNext/Source/AnimNext/Internal/Graph/AnimNextGraph.h b/Engine/Plugins/Experimental/AnimNext/Source/AnimNext/Internal/Graph/AnimNextGraph.h
--- a/Engine/Plugins/Experimental/AnimNext/Source/AnimNext/Internal/Graph/AnimNextGraph.h
+++ b/Engine/Plugins/Experimental/AnimNext/Source/AnimNext/Internal/Graph/AnimNextGraph.h
@@ -160,7 +160,7 @@ public:
 	// IAnimNextScheduleTermInterface interface
 	virtual TConstArrayView<UE::AnimNext::FScheduleTerm> GetTerms() const override;
 
-	// Allocates an instance of the graph, retain the handle and use it with the Run() function to evaluate it
+	// Allocates an instance of the graph
 	void AllocateInstance(FAnimNextGraphInstance& Instance) const;
 
 	// Get the parameter to use to access the reference pose
diff --git a/Engine/Plugins/Experimental/AnimNext/Source/AnimNext/Private/DecoratorBase/DecoratorPtr.cpp b/Engine/Plugins/Experimental/AnimNext/Source/AnimNext/Private/DecoratorBase/DecoratorPtr.cpp
--- a/Engine/Plugins/Experimental/AnimNext/Source/AnimNext/Private/DecoratorBase/DecoratorPtr.cpp
+++ b/Engine/Plugins/Experimental/AnimNext/Source/AnimNext/Private/DecoratorBase/DecoratorPtr.cpp
@@ -1,6 +1,8 @@
 // Copyright Epic Games, Inc. All Rights Reserved.
 
 #include "DecoratorBase/DecoratorPtr.h"
+
+#include "DecoratorBase/ExecutionContext.h"
 #include "DecoratorBase/NodeInstance.h"
 
 namespace UE::AnimNext
@@ -74,14 +76,7 @@ namespace UE::AnimNext
 			}
 		}
 
-		// Only decrement the reference count if we aren't a weak handle
-		if (!IsWeak())
-		{
-			if (FNodeInstance* OldNode = GetNodeInstance())
-			{
-				OldNode->ReleaseReference();
-			}
-		}
+		Reset();
 
 		PackedPointerAndFlags = DecoratorPtr.PackedPointerAndFlags;
 		DecoratorIndex = DecoratorPtr.DecoratorIndex;
@@ -99,12 +94,13 @@ namespace UE::AnimNext
 
 	void FDecoratorPtr::Reset()
 	{
-		// Only decrement the reference count if we aren't a weak handle
+		// Only decrement the reference count if we aren't a weak handle and if we are valid
 		if (!IsWeak())
 		{
 			if (FNodeInstance* Node = GetNodeInstance())
 			{
-				Node->ReleaseReference();
+				FExecutionContext Context(Node->GetOwner());
+				Context.ReleaseNodeInstance(*this);
 			}
 		}
 
diff --git a/Engine/Plugins/Experimental/AnimNext/Source/AnimNext/Private/DecoratorBase/ExecutionContext.cpp b/Engine/Plugins/Experimental/AnimNext/Source/AnimNext/Private/DecoratorBase/ExecutionContext.cpp
--- a/Engine/Plugins/Experimental/AnimNext/Source/AnimNext/Private/DecoratorBase/ExecutionContext.cpp
+++ b/Engine/Plugins/Experimental/AnimNext/Source/AnimNext/Private/DecoratorBase/ExecutionContext.cpp
@@ -13,49 +13,70 @@
 
 namespace UE::AnimNext
 {
-	namespace Private
+	FExecutionContext::FExecutionContext()
+		: NodeTemplateRegistry(FNodeTemplateRegistry::Get())
+		, DecoratorRegistry(FDecoratorRegistry::Get())
+		, GraphInstance(nullptr)
 	{
-		static thread_local UE::AnimNext::FExecutionContext* GThreadLocalExecutionContext = nullptr;
 	}
 
 	FExecutionContext::FExecutionContext(FAnimNextGraphInstance& InGraphInstance)
-		: NodeTemplateRegistry(FNodeTemplateRegistry::Get())
-		, DecoratorRegistry(FDecoratorRegistry::Get())
-		, Graph(InGraphInstance.GetGraph())
-		, GraphInstance(&InGraphInstance)
-		, GraphSharedData(Graph->SharedDataBuffer)
+		: FExecutionContext()
+	{
+		BindTo(InGraphInstance);
+	}
+
+	void FExecutionContext::BindTo(FAnimNextGraphInstance& InGraphInstance)
+	{
+		if (GraphInstance == &InGraphInstance)
+		{
+			return;	// Already bound to this graph instance, nothing to do
+		}
+
+		GraphInstance = &InGraphInstance;
+		GraphSharedData = InGraphInstance.GetGraph()->SharedDataBuffer;
+	}
+
+	void FExecutionContext::BindTo(const FWeakDecoratorPtr& DecoratorPtr)
+	{
+		if (const FNodeInstance* NodeInstance = DecoratorPtr.GetNodeInstance())
+		{
+			BindTo(NodeInstance->GetOwner());
+		}
+	}
+
+	bool FExecutionContext::IsBound() const
 	{
-		// There can be only one execution context alive per thread
-		ensure(Private::GThreadLocalExecutionContext == nullptr);
-		Private::GThreadLocalExecutionContext = this;
+		return GraphInstance != nullptr;
 	}
 
-	FExecutionContext::~FExecutionContext()
+	bool FExecutionContext::IsBoundTo(const FAnimNextGraphInstance& InGraphInstance) const
 	{
-		// There can be only one execution context alive per thread
-		ensure(Private::GThreadLocalExecutionContext == this);
-		Private::GThreadLocalExecutionContext = nullptr;
+		return GraphInstance == &InGraphInstance;
 	}
 
 	FDecoratorPtr FExecutionContext::AllocateNodeInstance(const FWeakDecoratorPtr& ParentBinding, FAnimNextDecoratorHandle ChildDecoratorHandle) const
 	{
-		ensure(ChildDecoratorHandle.IsValid());
-		if (!ChildDecoratorHandle.IsValid())
+		if (!ensure(ChildDecoratorHandle.IsValid()))
 		{
 			return FDecoratorPtr();	// Attempting to allocate a node using an invalid decorator handle
 		}
 
+		if (!ensure(IsBound()))
+		{
+			return FDecoratorPtr();	// The execution context must be bound to a valid graph instance
+		}
+
 		const FNodeHandle ChildNodeHandle = ChildDecoratorHandle.GetNodeHandle();
 		const FNodeDescription& NodeDesc = GetNodeDescription(ChildNodeHandle);
 
 		const FNodeTemplate* NodeTemplate = GetNodeTemplate(NodeDesc);
-		ensure(NodeTemplate != nullptr);
-		if (NodeTemplate == nullptr)
+		if (!ensure(NodeTemplate != nullptr))
 		{
 			return FDecoratorPtr();	// Node template wasn't found, node descriptor is perhaps corrupted
 		}
 
-		const uint32_t ChildDecoratorIndex = ChildDecoratorHandle.GetDecoratorIndex();
+		const uint32 ChildDecoratorIndex = ChildDecoratorHandle.GetDecoratorIndex();
 
 		if (ChildDecoratorIndex >= NodeTemplate->GetNumDecorators())
 		{
@@ -79,7 +100,7 @@ namespace UE::AnimNext
 
 		const uint32 InstanceSize = NodeTemplate->GetNodeInstanceDataSize();
 		uint8* NodeInstanceBuffer = reinterpret_cast<uint8*>(FMemory::Malloc(InstanceSize, 16));
-		FNodeInstance* NodeInstance = new(NodeInstanceBuffer) FNodeInstance(ChildNodeHandle);
+		FNodeInstance* NodeInstance = new(NodeInstanceBuffer) FNodeInstance(*GraphInstance, ChildNodeHandle);
 
 		// Start construction with the bottom decorator
 		const FDecoratorTemplate* StartDesc = DecoratorDescs;
@@ -107,16 +128,25 @@ namespace UE::AnimNext
 		return FDecoratorPtr(NodeInstance, ChildDecoratorIndex);
 	}
 
-	void FExecutionContext::ReleaseNodeInstance(FNodeInstance* NodeInstance) const
+	void FExecutionContext::ReleaseNodeInstance(FDecoratorPtr& NodePtr) const
 	{
-		ensure(NodeInstance != nullptr && NodeInstance->IsValid());
-		if (NodeInstance == nullptr || !NodeInstance->IsValid())
+		if (!NodePtr.IsValid())
 		{
-			return;	// Invalid node instance provided
+			return;
 		}
 
-		ensure(NodeInstance->GetReferenceCount() == 0);
-		if (NodeInstance->GetReferenceCount() != 0)
+		FNodeInstance* NodeInstance = NodePtr.GetNodeInstance();
+
+		if (!ensure(IsBoundTo(NodeInstance->GetOwner())))
+		{
+			return;	// The execution context isn't bound to the right graph instance
+		}
+
+		// Reset the handle here to simplify the multiple return statements below
+		NodePtr.PackedPointerAndFlags = 0;
+		NodePtr.DecoratorIndex = 0;
+
+		if (NodeInstance->RemoveReference())
 		{
 			return;	// Node instance still has references, we can't release it
 		}
@@ -124,10 +154,9 @@ namespace UE::AnimNext
 		const FNodeDescription& NodeDesc = GetNodeDescription(NodeInstance->GetNodeHandle());
 
 		const FNodeTemplate* NodeTemplate = GetNodeTemplate(NodeDesc);
-		ensure(NodeTemplate != nullptr);
-		if (NodeTemplate == nullptr)
+		if (!ensure(NodeTemplate != nullptr))
 		{
-			return;	// Node template wasn't found, node descriptor is perhaps corrupted
+			return;	// Node template wasn't found, node descriptor is perhaps corrupted (we'll leak the node memory)
 		}
 
 		const FDecoratorTemplate* DecoratorDescs = NodeTemplate->GetDecorators();
@@ -222,11 +251,15 @@ namespace UE::AnimNext
 
 		FNodeInstance* NodeInstance = DecoratorPtr.GetNodeInstance();
 
+		if (!ensure(IsBoundTo(NodeInstance->GetOwner())))
+		{
+			return false;	// The execution context isn't bound to the right graph instance
+		}
+
 		const FNodeDescription& NodeDesc = GetNodeDescription(NodeInstance->GetNodeHandle());
 
 		const FNodeTemplate* NodeTemplate = GetNodeTemplate(NodeDesc);
-		ensure(NodeTemplate != nullptr);
-		if (NodeTemplate == nullptr)
+		if (!ensure(NodeTemplate != nullptr))
 		{
 			return false;	// Node template wasn't found, node descriptor is perhaps corrupted
 		}
@@ -274,11 +307,15 @@ namespace UE::AnimNext
 
 		FNodeInstance* NodeInstance = DecoratorPtr.GetNodeInstance();
 
+		if (!ensure(IsBoundTo(NodeInstance->GetOwner())))
+		{
+			return false;	// The execution context isn't bound to the right graph instance
+		}
+
 		const FNodeDescription& NodeDesc = GetNodeDescription(NodeInstance->GetNodeHandle());
 
 		const FNodeTemplate* NodeTemplate = GetNodeTemplate(NodeDesc);
-		ensure(NodeTemplate != nullptr);
-		if (NodeTemplate == nullptr)
+		if (!ensure(NodeTemplate != nullptr))
 		{
 			return false;	// Node template wasn't found, node descriptor is perhaps corrupted
 		}
@@ -345,9 +382,4 @@ namespace UE::AnimNext
 		check(Template.GetRegistryHandle().IsValid());
 		return DecoratorRegistry.Find(Template.GetRegistryHandle());
 	}
-
-	FExecutionContext* GetThreadExecutionContext()
-	{
-		return Private::GThreadLocalExecutionContext;
-	}
 }
diff --git a/Engine/Plugins/Experimental/AnimNext/Source/AnimNext/Private/DecoratorBase/NodeInstance.cpp b/Engine/Plugins/Experimental/AnimNext/Source/AnimNext/Private/DecoratorBase/NodeInstance.cpp
deleted file mode 100644
diff --git a/Engine/Plugins/Experimental/AnimNext/Source/AnimNext/Private/DecoratorInterfaces/IEvaluate.cpp b/Engine/Plugins/Experimental/AnimNext/Source/AnimNext/Private/DecoratorInterfaces/IEvaluate.cpp
--- a/Engine/Plugins/Experimental/AnimNext/Source/AnimNext/Private/DecoratorInterfaces/IEvaluate.cpp
+++ b/Engine/Plugins/Experimental/AnimNext/Source/AnimNext/Private/DecoratorInterfaces/IEvaluate.cpp
@@ -71,7 +71,12 @@ namespace UE::AnimNext
 		}
 	};
 
-	FEvaluationProgram EvaluateGraph(FExecutionContext& Context, FWeakDecoratorPtr GraphRootPtr)
+	FEvaluationProgram EvaluateGraph(FAnimNextGraphInstance& GraphInstance)
+	{
+		return EvaluateGraph(GraphInstance.GetGraphRootPtr());
+	}
+
+	FEvaluationProgram EvaluateGraph(const FWeakDecoratorPtr& GraphRootPtr)
 	{
 		SCOPE_CYCLE_COUNTER(STAT_AnimNext_EvaluateGraph);
 		
@@ -88,7 +93,8 @@ namespace UE::AnimNext
 		FChildrenArray Children;
 		Children.Reserve(16);
 
-		FEvaluateTraversalContext TraversalContext(Context, EvaluationProgram);
+		FExecutionContext ExecutionContext;
+		FEvaluateTraversalContext TraversalContext(ExecutionContext, EvaluationProgram);
 
 		// Add the graph root to kick start the evaluation process
 		FEvaluateEntry GraphRootEntry(GraphRootPtr, EEvaluateStep::PreEvaluate);
@@ -106,9 +112,14 @@ namespace UE::AnimNext
 			FEvaluateEntry* Entry = NodesPendingUpdateStackTop;
 			bool bIsEntryUsed = true;
 
+			const FWeakDecoratorPtr& EntryDecoratorPtr = Entry->DecoratorPtr;
+
+			// Make sure the execution context is bound to our graph instance
+			ExecutionContext.BindTo(EntryDecoratorPtr);
+
 			if (Entry->DesiredStep == EEvaluateStep::PreEvaluate)
 			{
-				if (Context.GetInterface(Entry->DecoratorPtr, Entry->EvaluateDecorator))
+				if (ExecutionContext.GetInterface(EntryDecoratorPtr, Entry->EvaluateDecorator))
 				{
 					// This is the first time we visit this node, time to pre-evaluate
 					Entry->EvaluateDecorator.PreEvaluate(TraversalContext);
@@ -124,9 +135,9 @@ namespace UE::AnimNext
 					bIsEntryUsed = false;
 				}
 
-				if (Context.GetInterface(Entry->DecoratorPtr, HierarchyDecorator))
+				if (ExecutionContext.GetInterface(EntryDecoratorPtr, HierarchyDecorator))
 				{
-					HierarchyDecorator.GetChildren(Context, Children);
+					HierarchyDecorator.GetChildren(ExecutionContext, Children);
 
 					// Append our children in reserve order so that they are visited in the same order they were added
 					for (int32 ChildIndex = Children.Num() - 1; ChildIndex >= 0; --ChildIndex)
diff --git a/Engine/Plugins/Experimental/AnimNext/Source/AnimNext/Private/DecoratorInterfaces/IUpdate.cpp b/Engine/Plugins/Experimental/AnimNext/Source/AnimNext/Private/DecoratorInterfaces/IUpdate.cpp
--- a/Engine/Plugins/Experimental/AnimNext/Source/AnimNext/Private/DecoratorInterfaces/IUpdate.cpp
+++ b/Engine/Plugins/Experimental/AnimNext/Source/AnimNext/Private/DecoratorInterfaces/IUpdate.cpp
@@ -207,11 +207,11 @@ namespace UE::AnimNext
 	// perform as much useful work as possible while waiting for memory, hiding its slow latency by fully
 	// leveraging out-of-order CPU execution.
 
-	void UpdateGraph(const FExecutionContext& Context, const FWeakDecoratorPtr& GraphRootPtr, float DeltaTime)
+	void UpdateGraph(FAnimNextGraphInstance& GraphInstance, float DeltaTime)
 	{
 		SCOPE_CYCLE_COUNTER(STAT_AnimNext_UpdateGraph);
 		
-		if (!GraphRootPtr.IsValid())
+		if (!GraphInstance.IsValid())
 		{
 			return;	// Nothing to update
 		}
@@ -222,24 +222,31 @@ namespace UE::AnimNext
 		FChildrenArray Children;
 		TDecoratorBinding<IHierarchy> HierarchyDecorator;
 
-		FUpdateTraversalContext TraversalContext(Context, MemStack);
+		FExecutionContext ExecutionContext;
+		FUpdateTraversalContext TraversalContext(ExecutionContext, MemStack);
 		FUpdateTraversalQueue TraversalQueue(TraversalContext);
 
 		// Before we start the traversal, we give the graph instance components the chance to do some work
-		for (auto It = Context.GetComponentIterator(); It; ++It)
+		ExecutionContext.BindTo(GraphInstance);
+		for (auto It = ExecutionContext.GetComponentIterator(); It; ++It)
 		{
-			It.Value()->PreUpdate(TraversalContext);
+			It.Value()->PreUpdate(ExecutionContext);
 		}
 
 		// Add the graph root to start the update process
-		FUpdateEntry RootEntry(GraphRootPtr, FDecoratorUpdateState(DeltaTime));
+		FUpdateEntry RootEntry(GraphInstance.GetGraphRootPtr(), FDecoratorUpdateState(DeltaTime));
 		TraversalContext.PushUpdateEntry(&RootEntry);
 
 		while (FUpdateEntry* Entry = TraversalContext.PopUpdateEntry())
 		{
+			const FWeakDecoratorPtr& EntryDecoratorPtr = Entry->DecoratorPtr;
+
+			// Make sure the execution context is bound to our graph instance
+			ExecutionContext.BindTo(EntryDecoratorPtr);
+
 			if (!Entry->bHasPreUpdated)
 			{
-				if (Context.GetInterface(Entry->DecoratorPtr, Entry->UpdateDecorator))
+				if (ExecutionContext.GetInterface(EntryDecoratorPtr, Entry->UpdateDecorator))
 				{
 					// This is the first time we visit this node, time to pre-update
 					Entry->UpdateDecorator.PreUpdate(TraversalContext, Entry->DecoratorState);
@@ -267,9 +274,9 @@ namespace UE::AnimNext
 				{
 					// This node doesn't implement IUpdate
 					// We'll grab its children and traverse them
-					if (Context.GetInterface(Entry->DecoratorPtr, HierarchyDecorator))
+					if (ExecutionContext.GetInterface(EntryDecoratorPtr, HierarchyDecorator))
 					{
-						HierarchyDecorator.GetChildren(Context, Children);
+						HierarchyDecorator.GetChildren(ExecutionContext, Children);
 
 						// Append our children in reserve order so that they are visited in the same order they were added
 						for (int32 ChildIndex = Children.Num() - 1; ChildIndex >= 0; --ChildIndex)
@@ -298,9 +305,10 @@ namespace UE::AnimNext
 		}
 
 		// After we finish the traversal, we give the graph instance components the chance to do some work
-		for (auto It = Context.GetComponentIterator(); It; ++It)
+		ExecutionContext.BindTo(GraphInstance);
+		for (auto It = ExecutionContext.GetComponentIterator(); It; ++It)
 		{
-			It.Value()->PostUpdate(TraversalContext);
+			It.Value()->PostUpdate(ExecutionContext);
 		}
 	}
 }
diff --git a/Engine/Plugins/Experimental/AnimNext/Source/AnimNext/Private/Graph/AnimGraph/AnimNode_AnimNextGraph.cpp b/Engine/Plugins/Experimental/AnimNext/Source/AnimNext/Private/Graph/AnimGraph/AnimNode_AnimNextGraph.cpp
--- a/Engine/Plugins/Experimental/AnimNext/Source/AnimNext/Private/Graph/AnimGraph/AnimNode_AnimNextGraph.cpp
+++ b/Engine/Plugins/Experimental/AnimNext/Source/AnimNext/Private/Graph/AnimGraph/AnimNode_AnimNextGraph.cpp
@@ -75,8 +75,7 @@ void FAnimNode_AnimNextGraph::Update_AnyThread(const FAnimationUpdateContext& Co
 			AnimNextGraph->GetCurrentLODParam(), LODLevel
 		);
 
-		FExecutionContext ExecutionContext(GraphInstance);
-		UE::AnimNext::UpdateGraph(ExecutionContext, GraphInstance.GetGraphRootPtr(), Context.GetDeltaTime());
+		UE::AnimNext::UpdateGraph(GraphInstance, Context.GetDeltaTime());
 
 		ParamStack.PopLayer(LayerHandle);
 	}
@@ -154,8 +153,7 @@ void FAnimNode_AnimNextGraph::Evaluate_AnyThread(FPoseContext& Output)
 		);
 
 		{
-			FExecutionContext ExecutionContext(GraphInstance);
-			FEvaluationProgram EvaluationProgram = UE::AnimNext::EvaluateGraph(ExecutionContext, GraphInstance.GetGraphRootPtr());
+			const FEvaluationProgram EvaluationProgram = UE::AnimNext::EvaluateGraph(GraphInstance);
 
 			FEvaluationVM EvaluationVM(EEvaluationFlags::All, RefPose, LODLevel);
 			bool bHasValidOutput = false;
diff --git a/Engine/Plugins/Experimental/AnimNext/Source/AnimNext/Private/Graph/AnimNextGraph.cpp b/Engine/Plugins/Experimental/AnimNext/Source/AnimNext/Private/Graph/AnimNextGraph.cpp
--- a/Engine/Plugins/Experimental/AnimNext/Source/AnimNext/Private/Graph/AnimNextGraph.cpp
+++ b/Engine/Plugins/Experimental/AnimNext/Source/AnimNext/Private/Graph/AnimNextGraph.cpp
@@ -135,12 +135,7 @@ void FAnimNextGraphInstance::Release()
 		return;
 	}
 
-	{
-		// We need an execution context for this graph to be active when we delete the graph instance
-		UE::AnimNext::FExecutionContext Context(*this);
-		GraphInstancePtr.Reset();
-	}
-
+	GraphInstancePtr.Reset();
 	ExtendedExecuteContext.Reset();
 	Components.Empty();
 
diff --git a/Engine/Plugins/Experimental/AnimNext/Source/AnimNext/Private/Graph/GC_GraphInstanceComponent.cpp b/Engine/Plugins/Experimental/AnimNext/Source/AnimNext/Private/Graph/GC_GraphInstanceComponent.cpp
--- a/Engine/Plugins/Experimental/AnimNext/Source/AnimNext/Private/Graph/GC_GraphInstanceComponent.cpp
+++ b/Engine/Plugins/Experimental/AnimNext/Source/AnimNext/Private/Graph/GC_GraphInstanceComponent.cpp
@@ -28,12 +28,13 @@ namespace UE::AnimNext
 
 	void FGCGraphInstanceComponent::AddReferencedObjects(FReferenceCollector& Collector) const
 	{
+		FExecutionContext Context;
 		TDecoratorBinding<IGarbageCollection> GCDecorator;
 
 		// TODO: If we kept the entries sorted by graph instance, we could re-use the execution context
 		for (const FEntry& Entry : DecoratorsWithReferences)
 		{
-			FExecutionContext Context(Entry.GraphInstance);
+			Context.BindTo(Entry.GraphInstance);
 			ensure(Context.GetInterface(Entry.DecoratorPtr, GCDecorator));
 
 			GCDecorator.AddReferencedObjects(Context, Collector);
diff --git a/Engine/Plugins/Experimental/AnimNext/Source/AnimNext/Private/Graph/SyncGroup_GraphInstanceComponent.cpp b/Engine/Plugins/Experimental/AnimNext/Source/AnimNext/Private/Graph/SyncGroup_GraphInstanceComponent.cpp
--- a/Engine/Plugins/Experimental/AnimNext/Source/AnimNext/Private/Graph/SyncGroup_GraphInstanceComponent.cpp
+++ b/Engine/Plugins/Experimental/AnimNext/Source/AnimNext/Private/Graph/SyncGroup_GraphInstanceComponent.cpp
@@ -11,13 +11,13 @@ namespace UE::AnimNext
 		GroupState.Members.Add(FSyncGroupMember{ DecoratorState, DecoratorPtr, GroupRole });
 	}
 
-	void FSyncGroupGraphInstanceComponent::PreUpdate(FUpdateTraversalContext& Context)
+	void FSyncGroupGraphInstanceComponent::PreUpdate(FExecutionContext& Context)
 	{
 		// Reset our group state, we want to start fresh every update
 		SyncGroupMap.Reset();
 	}
 
-	void FSyncGroupGraphInstanceComponent::PostUpdate(FUpdateTraversalContext& Context)
+	void FSyncGroupGraphInstanceComponent::PostUpdate(FExecutionContext& Context)
 	{
 		TDecoratorBinding<IGroupSynchronization> GroupSyncDecorator;
 
@@ -73,6 +73,7 @@ namespace UE::AnimNext
 			{
 				const FSyncGroupMember& GroupLeader = GroupState.Members[LeaderIndex];
 
+				Context.BindTo(GroupLeader.DecoratorPtr);
 				ensure(Context.GetInterface(GroupLeader.DecoratorPtr, GroupSyncDecorator));
 
 				LeaderProgressRatio = GroupSyncDecorator.AdvanceBy(Context, GroupLeader.DecoratorState.GetDeltaTime());
@@ -88,6 +89,7 @@ namespace UE::AnimNext
 
 				const FSyncGroupMember& GroupMember = GroupState.Members[MemberIndex];
 
+				Context.BindTo(GroupMember.DecoratorPtr);
 				ensure(Context.GetInterface(GroupMember.DecoratorPtr, GroupSyncDecorator));
 
 				GroupSyncDecorator.AdvanceToRatio(Context, LeaderProgressRatio);
diff --git a/Engine/Plugins/Experimental/AnimNext/Source/AnimNext/Private/Graph/SyncGroup_GraphInstanceComponent.h b/Engine/Plugins/Experimental/AnimNext/Source/AnimNext/Private/Graph/SyncGroup_GraphInstanceComponent.h
--- a/Engine/Plugins/Experimental/AnimNext/Source/AnimNext/Private/Graph/SyncGroup_GraphInstanceComponent.h
+++ b/Engine/Plugins/Experimental/AnimNext/Source/AnimNext/Private/Graph/SyncGroup_GraphInstanceComponent.h
@@ -22,8 +22,8 @@ namespace UE::AnimNext
 		void RegisterWithGroup(FName GroupName, EAnimGroupRole::Type GroupRole, const FWeakDecoratorPtr& DecoratorPtr, const FDecoratorUpdateState& DecoratorState);
 
 		// FGraphInstanceComponent impl
-		virtual void PreUpdate(FUpdateTraversalContext& Context) override;
-		virtual void PostUpdate(FUpdateTraversalContext& Context) override;
+		virtual void PreUpdate(FExecutionContext& Context) override;
+		virtual void PostUpdate(FExecutionContext& Context) override;
 
 	private:
 		struct FSyncGroupMember
diff --git a/Engine/Plugins/Experimental/AnimNext/Source/AnimNext/Private/Scheduler/AnimNextScheduleGraphTask.cpp b/Engine/Plugins/Experimental/AnimNext/Source/AnimNext/Private/Scheduler/AnimNextScheduleGraphTask.cpp
--- a/Engine/Plugins/Experimental/AnimNext/Source/AnimNext/Private/Scheduler/AnimNextScheduleGraphTask.cpp
+++ b/Engine/Plugins/Experimental/AnimNext/Source/AnimNext/Private/Scheduler/AnimNextScheduleGraphTask.cpp
@@ -109,12 +109,10 @@ void FAnimNextScheduleGraphTask::RunGraph(const UE::AnimNext::FScheduleContext&
 	// This reduces churn internally by avoiding a chunk to be repeatedly allocated and freed as we push/pop marks
 	MemStack.Alloc(size_t(FPageAllocator::SmallPageSize) + 1, 16);
 
-	FExecutionContext Context(InstanceData.GraphInstanceData[TaskIndex]);
-
-	UE::AnimNext::UpdateGraph(Context, InstanceData.GraphInstanceData[TaskIndex].GetGraphRootPtr(), InContext.GetDeltaTime());
+	UE::AnimNext::UpdateGraph(InstanceData.GraphInstanceData[TaskIndex], InContext.GetDeltaTime());
 
 	{
-		FEvaluationProgram EvaluationProgram = UE::AnimNext::EvaluateGraph(Context, InstanceData.GraphInstanceData[TaskIndex].GetGraphRootPtr());
+		const FEvaluationProgram EvaluationProgram = UE::AnimNext::EvaluateGraph(InstanceData.GraphInstanceData[TaskIndex]);
 
 		FEvaluationVM EvaluationVM(EEvaluationFlags::All, *GraphReferencePose->ReferencePose, *GraphLODLevel);
 		bool bHasValidOutput = false;
diff --git a/Engine/Plugins/Experimental/AnimNext/Source/AnimNext/Public/DecoratorBase/ExecutionContext.h b/Engine/Plugins/Experimental/AnimNext/Source/AnimNext/Public/DecoratorBase/ExecutionContext.h
--- a/Engine/Plugins/Experimental/AnimNext/Source/AnimNext/Public/DecoratorBase/ExecutionContext.h
+++ b/Engine/Plugins/Experimental/AnimNext/Source/AnimNext/Public/DecoratorBase/ExecutionContext.h
@@ -24,15 +24,28 @@ namespace UE::AnimNext
 	/**
 	 * Execution Context
 	 * 
-	 * The execution context holds internal state during traversals of the animation graph.
+	 * The execution context aims to centralize the decorator query API.
+	 * It is meant to be bound to a graph instance and re-used by the nodes/decorators within.
 	 */
 	struct ANIMNEXT_API FExecutionContext
 	{
-		// Creates an execution context for the specified graph instance
+		// Creates an unbound execution context
+		FExecutionContext();
+
+		// Creates an execution context and binds it to the specified graph instance
 		explicit FExecutionContext(FAnimNextGraphInstance& InGraphInstance);
 
-		// Destroys the execution context
-		~FExecutionContext();
+		// Binds the execution context to the specified graph instance if it differs from the currently bound instance
+		void BindTo(FAnimNextGraphInstance& InGraphInstance);
+
+		// Binds the execution context to the graph instance that owns the specified decorator if it differs from the currently bound instance
+		void BindTo(const FWeakDecoratorPtr& DecoratorPtr);
+
+		// Returns whether or not this execution context is bound to a graph instance
+		bool IsBound() const;
+
+		// Returns whether or not this execution context is bound to the specified graph instance
+		bool IsBoundTo(const FAnimNextGraphInstance& InGraphInstance) const;
 
 		// Queries a node for a decorator that implements the specified interface.
 		// If no such decorator exists, nullptr is returned.
@@ -62,8 +75,9 @@ namespace UE::AnimNext
 		// If the desired decorator lives in the current parent, a weak handle to it will be returned
 		FDecoratorPtr AllocateNodeInstance(const FWeakDecoratorPtr& ParentBinding, FAnimNextDecoratorHandle ChildDecoratorHandle) const;
 
-		// Releases a node instance that is no longer referenced
-		void ReleaseNodeInstance(FNodeInstance* Node) const;
+		// Decrements the reference count of the provided node pointer and releases it if
+		// there are no more references remaining, reseting the pointer in the process
+		void ReleaseNodeInstance(FDecoratorPtr& NodePtr) const;
 
 		// Evaluates the latent pin with the specified handle
 		template<typename LatentPinType>
@@ -101,14 +115,10 @@ namespace UE::AnimNext
 		const FDecoratorRegistry& DecoratorRegistry;
 
 		// Cached properties for the currently executing graph
-		const UAnimNextGraph* Graph = nullptr;
 		FAnimNextGraphInstance* GraphInstance = nullptr;
 		TArrayView<const uint8> GraphSharedData;
 	};
 
-	// Returns a pointer to the current execution context if present, nullptr otherwise.
-	FExecutionContext* GetThreadExecutionContext();
-
 	//////////////////////////////////////////////////////////////////////////
 	// Inline implementations
 
diff --git a/Engine/Plugins/Experimental/AnimNext/Source/AnimNext/Public/DecoratorBase/ExecutionContextProxy.h b/Engine/Plugins/Experimental/AnimNext/Source/AnimNext/Public/DecoratorBase/ExecutionContextProxy.h
--- a/Engine/Plugins/Experimental/AnimNext/Source/AnimNext/Public/DecoratorBase/ExecutionContextProxy.h
+++ b/Engine/Plugins/Experimental/AnimNext/Source/AnimNext/Public/DecoratorBase/ExecutionContextProxy.h
@@ -16,10 +16,26 @@ namespace UE::AnimNext
 	{
 		explicit FExecutionContextProxy(const FExecutionContext& InExecutionContext)
 			: ExecutionContext(InExecutionContext)
-		{}
+		{
+		}
 
 		// We safely coerce to our wrapped execution context to allow identical usage
-		operator const FExecutionContext& () const { return ExecutionContext; }
+		operator const FExecutionContext& () const
+		{
+			return ExecutionContext;
+		}
+
+		// Returns whether or not this execution context is bound to a graph instance
+		bool IsBound() const
+		{
+			return ExecutionContext.IsBound();
+		}
+
+		// Returns whether or not this execution context is bound to the specified graph instance
+		bool IsBoundTo(const FAnimNextGraphInstance& InGraphInstance) const
+		{
+			return ExecutionContext.IsBoundTo(InGraphInstance);
+		}
 
 		// Queries a node for a decorator that implements the specified interface.
 		// If no such decorator exists, nullptr is returned.
@@ -67,10 +83,11 @@ namespace UE::AnimNext
 			return ExecutionContext.AllocateNodeInstance(ParentBinding, ChildDecoratorHandle);
 		}
 
-		// Releases a node instance that is no longer referenced
-		void ReleaseNodeInstance(FNodeInstance* Node) const
+		// Decrements the reference count of the provided node pointer and releases it if
+		// there are no more references remaining, reseting the pointer in the process
+		void ReleaseNodeInstance(FDecoratorPtr& NodePtr) const
 		{
-			return ExecutionContext.ReleaseNodeInstance(Node);
+			return ExecutionContext.ReleaseNodeInstance(NodePtr);
 		}
 
 		// Evaluates the latent pin with the specified handle
diff --git a/Engine/Plugins/Experimental/AnimNext/Source/AnimNext/Public/DecoratorBase/NodeInstance.h b/Engine/Plugins/Experimental/AnimNext/Source/AnimNext/Public/DecoratorBase/NodeInstance.h
--- a/Engine/Plugins/Experimental/AnimNext/Source/AnimNext/Public/DecoratorBase/NodeInstance.h
+++ b/Engine/Plugins/Experimental/AnimNext/Source/AnimNext/Public/DecoratorBase/NodeInstance.h
@@ -6,6 +6,8 @@
 #include "DecoratorBase/DecoratorTemplate.h"
 #include "DecoratorBase/NodeHandle.h"
 
+struct FAnimNextGraphInstance;
+
 namespace UE::AnimNext
 {
 	/**
@@ -16,18 +18,21 @@ namespace UE::AnimNext
 	 * 
 	 * @see FNodeTemplate, FDecoratorTemplate
 	 */
-	struct ANIMNEXT_API FNodeInstance
+	struct FNodeInstance
 	{
 		// Largest allowed size for a node instance and the instance data of its decorators
 		// We use unsigned 16 bit offsets within the NodeTemplate/DecoratorTemplate
 		static constexpr uint32 MAXIMUM_NODE_INSTANCE_DATA_SIZE = 64 * 1024;
 
-		FNodeInstance(const FNodeInstance&) = delete;
-		FNodeInstance& operator=(const FNodeInstance&) = delete;
-
 		// Returns whether the node instance is valid or not
 		bool IsValid() const { return NodeHandle.IsValid(); }
 
+		// Returns whether or not this node instance is owned by the specified graph instance
+		bool IsOwnedBy(const FAnimNextGraphInstance& GraphInstance) const { return &Owner == &GraphInstance; }
+
+		// Returns the graph instance that owns this node instance
+		FAnimNextGraphInstance& GetOwner() const { return Owner; }
+
 		// Returns a handle to the shared data for this node
 		FNodeHandle GetNodeHandle() const { return NodeHandle; }
 
@@ -35,11 +40,20 @@ namespace UE::AnimNext
 		uint32 GetReferenceCount() const { return ReferenceCount; }
 
 	private:
-		explicit FNodeInstance(FNodeHandle NodeHandle);
+		FNodeInstance(FAnimNextGraphInstance& InOwner, FNodeHandle InNodeHandle) : Owner(InOwner), ReferenceCount(0), NodeHandle(InNodeHandle) {}
 		~FNodeInstance() { check(ReferenceCount == 0); }
 
+		// Cannot move or copy node instances
+		FNodeInstance(const FNodeInstance&) = delete;
+		FNodeInstance& operator=(const FNodeInstance&) = delete;
+
+		// Increments the reference count
 		void AddReference() { ReferenceCount++; }
-		void ReleaseReference();
+
+		// Decrements the reference count and returns true if any references remain
+		bool RemoveReference() { check(ReferenceCount > 0); return ReferenceCount-- != 1; }
+
+		FAnimNextGraphInstance& Owner;	// The graph instance that owns this node instance
 
 		uint32		ReferenceCount;		// how many non-weak FDecoratorPtr handles point to us, not thread safe
 		FNodeHandle	NodeHandle;			// relative to root of sub-graph, should this be a pointer?
diff --git a/Engine/Plugins/Experimental/AnimNext/Source/AnimNext/Public/DecoratorInterfaces/IEvaluate.h b/Engine/Plugins/Experimental/AnimNext/Source/AnimNext/Public/DecoratorInterfaces/IEvaluate.h
--- a/Engine/Plugins/Experimental/AnimNext/Source/AnimNext/Public/DecoratorInterfaces/IEvaluate.h
+++ b/Engine/Plugins/Experimental/AnimNext/Source/AnimNext/Public/DecoratorInterfaces/IEvaluate.h
@@ -29,7 +29,7 @@ namespace UE::AnimNext
 
 		FEvaluationProgram& EvaluationProgram;
 
-		friend ANIMNEXT_API FEvaluationProgram EvaluateGraph(FExecutionContext& Context, FWeakDecoratorPtr GraphRootPtr);
+		friend ANIMNEXT_API FEvaluationProgram EvaluateGraph(const FWeakDecoratorPtr& GraphRootPtr);
 	};
 
 	/**
@@ -79,9 +79,24 @@ namespace UE::AnimNext
 		const IEvaluate* GetInterface() const { return GetInterfaceTyped<IEvaluate>(); }
 	};
 
+	/**
+	 * Evaluates a sub-graph starting at its root and produces an evaluation program.
+	 * Evaluation should be deterministic and repeated calls should yield the same evaluation program.
+	 *
+	 * For each node:
+	 *     - We call PreEvaluate on all its decorators
+	 *     - We call GetChildren on all its decorators
+	 *     - We evaluate all children found
+	 *     - We call PostEvaluate on all its decorators
+	 *
+	 * @see IEvaluate::PreEvaluate, IEvaluate::PostEvaluate, IHierarchy::GetChildren
+	 */
+	[[nodiscard]] ANIMNEXT_API FEvaluationProgram EvaluateGraph(FAnimNextGraphInstance& GraphInstance);
+
 	/**
 	 * Evaluates a sub-graph starting at its root and produces an evaluation program.
 	 * Evaluation starts at the top of the stack that includes the graph root decorator.
+	 * Evaluation should be deterministic and repeated calls should yield the same evaluation program.
 	 *
 	 * For each node:
 	 *     - We call PreEvaluate on all its decorators
@@ -91,5 +106,5 @@ namespace UE::AnimNext
 	 *
 	 * @see IEvaluate::PreEvaluate, IEvaluate::PostEvaluate, IHierarchy::GetChildren
 	 */
-	[[nodiscard]] ANIMNEXT_API FEvaluationProgram EvaluateGraph(FExecutionContext& Context, FWeakDecoratorPtr GraphRootPtr);
+	[[nodiscard]] ANIMNEXT_API FEvaluationProgram EvaluateGraph(const FWeakDecoratorPtr& GraphRootPtr);
 }
diff --git a/Engine/Plugins/Experimental/AnimNext/Source/AnimNext/Public/DecoratorInterfaces/IUpdate.h b/Engine/Plugins/Experimental/AnimNext/Source/AnimNext/Public/DecoratorInterfaces/IUpdate.h
--- a/Engine/Plugins/Experimental/AnimNext/Source/AnimNext/Public/DecoratorInterfaces/IUpdate.h
+++ b/Engine/Plugins/Experimental/AnimNext/Source/AnimNext/Public/DecoratorInterfaces/IUpdate.h
@@ -6,6 +6,7 @@
 #include "DecoratorBase/ExecutionContextProxy.h"
 #include "DecoratorBase/IDecoratorInterface.h"
 
+struct FAnimNextGraphInstance;
 class FMemStack;
 
 namespace UE::AnimNext
@@ -126,7 +127,7 @@ namespace UE::AnimNext
 		// be warmer in the CPU cache
 		FUpdateEntry* FreeEntryStackHead = nullptr;
 
-		friend ANIMNEXT_API void UpdateGraph(const FExecutionContext& Context, const FWeakDecoratorPtr& GraphRootPtr, float DeltaTime);
+		friend ANIMNEXT_API void UpdateGraph(FAnimNextGraphInstance& GraphInstance, float DeltaTime);
 		friend FUpdateTraversalQueue;
 	};
 
@@ -154,7 +155,7 @@ namespace UE::AnimNext
 		// emptied and pushed onto the update stack
 		FUpdateEntry* QueuedUpdateStackHead = nullptr;
 
-		friend ANIMNEXT_API void UpdateGraph(const FExecutionContext& Context, const FWeakDecoratorPtr& GraphRootPtr, float DeltaTime);
+		friend ANIMNEXT_API void UpdateGraph(FAnimNextGraphInstance& GraphInstance, float DeltaTime);
 		friend FUpdateTraversalContext;
 	};
 
@@ -217,7 +218,6 @@ namespace UE::AnimNext
 
 	/**
 	 * Updates a sub-graph starting at its root.
-	 * Update starts at the top of the stack that includes the graph root decorator.
 	 *
 	 * For each node:
 	 *     - We call PreUpdate on all its decorators
@@ -226,5 +226,5 @@ namespace UE::AnimNext
 	 *
 	 * @see IUpdate::PreUpdate, IUpdate::PostUpdate, IHierarchy::GetChildren
 	 */
-	ANIMNEXT_API void UpdateGraph(const FExecutionContext& Context, const FWeakDecoratorPtr& GraphRootPtr, float DeltaTime);
+	ANIMNEXT_API void UpdateGraph(FAnimNextGraphInstance& GraphInstance, float DeltaTime);
 }
diff --git a/Engine/Plugins/Experimental/AnimNext/Source/AnimNext/Public/Graph/GraphInstanceComponent.h b/Engine/Plugins/Experimental/AnimNext/Source/AnimNext/Public/Graph/GraphInstanceComponent.h
--- a/Engine/Plugins/Experimental/AnimNext/Source/AnimNext/Public/Graph/GraphInstanceComponent.h
+++ b/Engine/Plugins/Experimental/AnimNext/Source/AnimNext/Public/Graph/GraphInstanceComponent.h
@@ -11,7 +11,7 @@
 
 namespace UE::AnimNext
 {
-	struct FUpdateTraversalContext;
+	struct FExecutionContext;
 
 	/**
 	 * FGraphInstanceComponent
@@ -28,9 +28,11 @@ namespace UE::AnimNext
 
 		// Called before the update traversal begins, before any node has been visited
 		// Note that PreUpdate won't be called if a component is created during the update traversal until the next update
-		virtual void PreUpdate(FUpdateTraversalContext& Context) {}
+		// The execution context provided is bound to the graph root and can be bound to anything the component wishes
+		virtual void PreUpdate(FExecutionContext& Context) {}
 
 		// Called after the update traversal completes, after every node has been visited
-		virtual void PostUpdate(FUpdateTraversalContext& Context) {}
+		// The execution context provided is bound to the graph root and can be bound to anything the component wishes
+		virtual void PostUpdate(FExecutionContext& Context) {}
 	};
 }
diff --git a/Engine/Plugins/Experimental/AnimNext/Source/AnimNextTestSuite/Private/AnimNextDecoratorBaseTest.cpp b/Engine/Plugins/Experimental/AnimNext/Source/AnimNextTestSuite/Private/AnimNextDecoratorBaseTest.cpp
--- a/Engine/Plugins/Experimental/AnimNext/Source/AnimNextTestSuite/Private/AnimNextDecoratorBaseTest.cpp
+++ b/Engine/Plugins/Experimental/AnimNext/Source/AnimNextTestSuite/Private/AnimNextDecoratorBaseTest.cpp
@@ -36,13 +36,13 @@ namespace UE::AnimNext
 	{
 		DECLARE_ANIM_DECORATOR_INTERFACE(IInterfaceA, 0x34cb8e62)
 
-		virtual void FuncA(FExecutionContext& Context, const TDecoratorBinding<IInterfaceA>& Binding) const;
+		virtual void FuncA(const FExecutionContext& Context, const TDecoratorBinding<IInterfaceA>& Binding) const;
 	};
 
 	template<>
 	struct TDecoratorBinding<IInterfaceA> : FDecoratorBinding
 	{
-		void FuncA(FExecutionContext& Context) const
+		void FuncA(const FExecutionContext& Context) const
 		{
 			GetInterface()->FuncA(Context, *this);
 		}
@@ -51,7 +51,7 @@ namespace UE::AnimNext
 		const IInterfaceA* GetInterface() const { return GetInterfaceTyped<IInterfaceA>(); }
 	};
 
-	void IInterfaceA::FuncA(FExecutionContext& Context, const TDecoratorBinding<IInterfaceA>& Binding) const
+	void IInterfaceA::FuncA(const FExecutionContext& Context, const TDecoratorBinding<IInterfaceA>& Binding) const
 	{
 		TDecoratorBinding<IInterfaceA> SuperBinding;
 		if (Context.GetInterfaceSuper(Binding, SuperBinding))
@@ -66,13 +66,13 @@ namespace UE::AnimNext
 	{
 		DECLARE_ANIM_DECORATOR_INTERFACE(IInterfaceB, 0x33cb8ccf)
 
-		virtual void FuncB(FExecutionContext& Context, const TDecoratorBinding<IInterfaceB>& Binding) const;
+		virtual void FuncB(const FExecutionContext& Context, const TDecoratorBinding<IInterfaceB>& Binding) const;
 	};
 
 	template<>
 	struct TDecoratorBinding<IInterfaceB> : FDecoratorBinding
 	{
-		void FuncB(FExecutionContext& Context) const
+		void FuncB(const FExecutionContext& Context) const
 		{
 			GetInterface()->FuncB(Context, *this);
 		}
@@ -81,7 +81,7 @@ namespace UE::AnimNext
 		const IInterfaceB* GetInterface() const { return GetInterfaceTyped<IInterfaceB>(); }
 	};
 
-	void IInterfaceB::FuncB(FExecutionContext& Context, const TDecoratorBinding<IInterfaceB>& Binding) const
+	void IInterfaceB::FuncB(const FExecutionContext& Context, const TDecoratorBinding<IInterfaceB>& Binding) const
 	{
 		TDecoratorBinding<IInterfaceB> SuperBinding;
 		if (Context.GetInterfaceSuper(Binding, SuperBinding))
@@ -96,13 +96,13 @@ namespace UE::AnimNext
 	{
 		DECLARE_ANIM_DECORATOR_INTERFACE(IInterfaceC, 0x32cb8b3c)
 
-		virtual void FuncC(FExecutionContext& Context, const TDecoratorBinding<IInterfaceC>& Binding) const;
+		virtual void FuncC(const FExecutionContext& Context, const TDecoratorBinding<IInterfaceC>& Binding) const;
 	};
 
 	template<>
 	struct TDecoratorBinding<IInterfaceC> : FDecoratorBinding
 	{
-		void FuncC(FExecutionContext& Context) const
+		void FuncC(const FExecutionContext& Context) const
 		{
 			GetInterface()->FuncC(Context, *this);
 		}
@@ -111,7 +111,7 @@ namespace UE::AnimNext
 		const IInterfaceC* GetInterface() const { return GetInterfaceTyped<IInterfaceC>(); }
 	};
 
-	void IInterfaceC::FuncC(FExecutionContext& Context, const TDecoratorBinding<IInterfaceC>& Binding) const
+	void IInterfaceC::FuncC(const FExecutionContext& Context, const TDecoratorBinding<IInterfaceC>& Binding) const
 	{
 		TDecoratorBinding<IInterfaceC> SuperBinding;
 		if (Context.GetInterfaceSuper(Binding, SuperBinding))
@@ -150,7 +150,7 @@ namespace UE::AnimNext
 		};
 
 		// IInterfaceA impl
-		virtual void FuncA(FExecutionContext& Context, const TDecoratorBinding<IInterfaceA>& Binding) const override
+		virtual void FuncA(const FExecutionContext& Context, const TDecoratorBinding<IInterfaceA>& Binding) const override
 		{
 		}
 	};
@@ -189,12 +189,12 @@ namespace UE::AnimNext
 		};
 
 		// IInterfaceA impl
-		virtual void FuncA(FExecutionContext& Context, const TDecoratorBinding<IInterfaceA>& Binding) const override
+		virtual void FuncA(const FExecutionContext& Context, const TDecoratorBinding<IInterfaceA>& Binding) const override
 		{
 		}
 
 		// IInterfaceB impl
-		virtual void FuncB(FExecutionContext& Context, const TDecoratorBinding<IInterfaceB>& Binding) const override
+		virtual void FuncB(const FExecutionContext& Context, const TDecoratorBinding<IInterfaceB>& Binding) const override
 		{
 		}
 	};
@@ -234,12 +234,12 @@ namespace UE::AnimNext
 		};
 
 		// IInterfaceA impl
-		virtual void FuncA(FExecutionContext& Context, const TDecoratorBinding<IInterfaceA>& Binding) const override
+		virtual void FuncA(const FExecutionContext& Context, const TDecoratorBinding<IInterfaceA>& Binding) const override
 		{
 		}
 
 		// IInterfaceC impl
-		virtual void FuncC(FExecutionContext& Context, const TDecoratorBinding<IInterfaceC>& Binding) const override
+		virtual void FuncC(const FExecutionContext& Context, const TDecoratorBinding<IInterfaceC>& Binding) const override
 		{
 		}
 	};
diff --git a/Engine/Plugins/Experimental/AnimNext/Source/AnimNextTestSuite/Private/AnimNextDecoratorGraphTest.cpp b/Engine/Plugins/Experimental/AnimNext/Source/AnimNextTestSuite/Private/AnimNextDecoratorGraphTest.cpp
--- a/Engine/Plugins/Experimental/AnimNext/Source/AnimNextTestSuite/Private/AnimNextDecoratorGraphTest.cpp
+++ b/Engine/Plugins/Experimental/AnimNext/Source/AnimNextTestSuite/Private/AnimNextDecoratorGraphTest.cpp
@@ -263,9 +263,8 @@ bool FAnimationAnimNextRuntimeTest_GraphExecute::RunTest(const FString& InParame
 	);
 
 	{
-		FExecutionContext Context(GraphInstance);
-		UE::AnimNext::UpdateGraph(Context, GraphInstance.GetGraphRootPtr(), 1.0f / 30.0f);
-		(void)UE::AnimNext::EvaluateGraph(Context, GraphInstance.GetGraphRootPtr());
+		UE::AnimNext::UpdateGraph(GraphInstance, 1.0f / 30.0f);
+		(void)UE::AnimNext::EvaluateGraph(GraphInstance);
 	}
 
 	AddErrorIfFalse(ParamStack->GetParam<int32>("UpdateCount") == 1, "FAnimationAnimNextRuntimeTest_GraphExecute -> Unexpected update count");
@@ -407,9 +406,8 @@ bool FAnimationAnimNextRuntimeTest_GraphExecuteLatent::RunTest(const FString& In
 	);
 
 	{
-		FExecutionContext Context(GraphInstance);
-		UE::AnimNext::UpdateGraph(Context, GraphInstance.GetGraphRootPtr(), 1.0f / 30.0f);
-		(void)UE::AnimNext::EvaluateGraph(Context, GraphInstance.GetGraphRootPtr());
+		UE::AnimNext::UpdateGraph(GraphInstance, 1.0f / 30.0f);
+		(void)UE::AnimNext::EvaluateGraph(GraphInstance);
 	}
 
 	AddErrorIfFalse(ParamStack->GetParam<int32>("UpdateCount") == 1, "FAnimationAnimNextRuntimeTest_GraphExecuteLatent -> Unexpected update count");
diff --git a/Engine/Plugins/Experimental/AnimNext/Source/AnimNextTestSuite/Private/AnimNextDecoratorInterfacesTest.cpp b/Engine/Plugins/Experimental/AnimNext/Source/AnimNextTestSuite/Private/AnimNextDecoratorInterfacesTest.cpp
--- a/Engine/Plugins/Experimental/AnimNext/Source/AnimNextTestSuite/Private/AnimNextDecoratorInterfacesTest.cpp
+++ b/Engine/Plugins/Experimental/AnimNext/Source/AnimNextTestSuite/Private/AnimNextDecoratorInterfacesTest.cpp
@@ -519,7 +519,7 @@ bool FAnimationAnimNextRuntimeTest_IUpdate::RunTest(const FString& InParameters)
 		AddErrorIfFalse(NodeCPtr.IsValid(), "FAnimationAnimNextRuntimeTest_IUpdate -> Failed to allocate root node instance");
 
 		// Call pre/post update on our graph
-		UpdateGraph(Context, NodeCPtr, 0.0333f);
+		UpdateGraph(GraphInstance, 0.0333f);
 
 		AddErrorIfFalse(UpdatedDecorators.Num() == 6, "FAnimationAnimNextRuntimeTest_IUpdate -> Expected 6 nodes to have been visited during the update traversal");
 		AddErrorIfFalse(UpdatedDecorators[0] == FDecoratorWithChildren::DecoratorUID, "FAnimationAnimNextRuntimeTest_IUpdate -> Unexpected update order");		// NodeC
@@ -649,8 +649,6 @@ bool FAnimationAnimNextRuntimeTest_IEvaluate::RunTest(const FString& InParameter
 	FAnimNextGraphInstance GraphInstance;
 	AnimNextGraph->AllocateInstance(GraphInstance);
 
-	FExecutionContext Context(GraphInstance);
-
 	{
 		TArray<FDecoratorUID> EvaluatedDecorators;
 
@@ -659,11 +657,12 @@ bool FAnimationAnimNextRuntimeTest_IEvaluate::RunTest(const FString& InParameter
 		FWeakDecoratorPtr NullPtr;								// Empty, no parent
 		FAnimNextDecoratorHandle RootHandle(NodeHandles[2]);	// Point to NodeC, first base decorator
 
+		FExecutionContext Context(GraphInstance);
 		FDecoratorPtr NodeCPtr = Context.AllocateNodeInstance(NullPtr, RootHandle);
 		AddErrorIfFalse(NodeCPtr.IsValid(), "FAnimationAnimNextRuntimeTest_IEvaluate -> Failed to allocate root node instance");
 
 		// Call pre/post evaluate on our graph
-		(void)EvaluateGraph(Context, NodeCPtr);
+		(void)EvaluateGraph(GraphInstance);
 
 		AddErrorIfFalse(EvaluatedDecorators.Num() == 6, "FAnimationAnimNextRuntimeTest_IEvaluate -> Expected 6 nodes to have been visited during the evaluate traversal");
 		AddErrorIfFalse(EvaluatedDecorators[0] == FDecoratorWithChildren::DecoratorUID, "FAnimationAnimNextRuntimeTest_IEvaluate -> Unexpected evaluate order");		// NodeC
