AnimNext library/parameter removal
- Parameter name/types are now exposed directly through AssetRegistry
- Some UX improvements with the use of known pin/parameter types
- Removal of all library/parameter related code/UI

diff --git a/Engine/Plugins/Experimental/AnimNext/Source/AnimNext/Internal/Graph/AnimNextGraph.h b/Engine/Plugins/Experimental/AnimNext/Source/AnimNext/Internal/Graph/AnimNextGraph.h
--- a/Engine/Plugins/Experimental/AnimNext/Source/AnimNext/Internal/Graph/AnimNextGraph.h
+++ b/Engine/Plugins/Experimental/AnimNext/Source/AnimNext/Internal/Graph/AnimNextGraph.h
@@ -169,6 +169,9 @@ public:
 	// Get the parameter to use to access the current LOD
 	UE::AnimNext::FParamId GetCurrentLODParam() const { return CurrentLODId; }
 
+#if WITH_EDITOR
+	virtual void GetAssetRegistryTags(TArray<FAssetRegistryTag>& OutTags) const override;
+#endif	
 protected:
 	// Support rig VM execution
 	TArray<FRigVMExternalVariable> GetRigVMExternalVariables();
diff --git a/Engine/Plugins/Experimental/AnimNext/Source/AnimNext/Internal/Param/AnimNextParameterBlock.h b/Engine/Plugins/Experimental/AnimNext/Source/AnimNext/Internal/Param/AnimNextParameterBlock.h
--- a/Engine/Plugins/Experimental/AnimNext/Source/AnimNext/Internal/Param/AnimNextParameterBlock.h
+++ b/Engine/Plugins/Experimental/AnimNext/Source/AnimNext/Internal/Param/AnimNextParameterBlock.h
@@ -32,36 +32,6 @@ namespace UE::AnimNext::Editor
 	class FParameterBlockParameterCustomization;
 }
 
-// Library entry used to export to asset registry
-USTRUCT()
-struct FAnimNextParameterBlockAssetRegistryExportEntry
-{
-	GENERATED_BODY()
-
-	FAnimNextParameterBlockAssetRegistryExportEntry() = default;
-
-	FAnimNextParameterBlockAssetRegistryExportEntry(FName InName, const FSoftObjectPath& InLibrary)
-		: Name(InName)
-		, Library(InLibrary)
-	{}
-	
-	UPROPERTY()
-	FName Name;
-
-	UPROPERTY()
-	FSoftObjectPath Library;
-};
-
-// Library used to export to asset registry
-USTRUCT()
-struct FAnimNextParameterBlockAssetRegistryExports
-{
-	GENERATED_BODY()
-
-	UPROPERTY()
-	TArray<FAnimNextParameterBlockAssetRegistryExportEntry> Bindings;
-};
-
 /** An asset used to define AnimNext parameters and their bindings */
 UCLASS(MinimalAPI, BlueprintType)
 class UAnimNextParameterBlock : public URigVMHost, public IAnimNextParameterSourceInterface
diff --git a/Engine/Plugins/Experimental/AnimNext/Source/AnimNext/Private/Graph/AnimNextGraph.cpp b/Engine/Plugins/Experimental/AnimNext/Source/AnimNext/Private/Graph/AnimNextGraph.cpp
--- a/Engine/Plugins/Experimental/AnimNext/Source/AnimNext/Private/Graph/AnimNextGraph.cpp
+++ b/Engine/Plugins/Experimental/AnimNext/Source/AnimNext/Private/Graph/AnimNextGraph.cpp
@@ -263,6 +263,19 @@ void UAnimNextGraph::AllocateInstance(FAnimNextGraphInstance& Instance) const
 #endif
 }
 
+
+#if WITH_EDITOR
+void UAnimNextGraph::GetAssetRegistryTags(TArray<FAssetRegistryTag>& OutTags) const
+{
+	Super::GetAssetRegistryTags(OutTags);
+
+	if (EditorData)
+	{		
+		EditorData->GetAssetRegistryTags(OutTags);
+	}
+}
+#endif
+
 TArray<FRigVMExternalVariable> UAnimNextGraph::GetRigVMExternalVariables()
 {
 	return TArray<FRigVMExternalVariable>(); 
diff --git a/Engine/Plugins/Experimental/AnimNext/Source/AnimNext/Private/Param/ParamType.cpp b/Engine/Plugins/Experimental/AnimNext/Source/AnimNext/Private/Param/ParamType.cpp
--- a/Engine/Plugins/Experimental/AnimNext/Source/AnimNext/Private/Param/ParamType.cpp
+++ b/Engine/Plugins/Experimental/AnimNext/Source/AnimNext/Private/Param/ParamType.cpp
@@ -11,6 +11,7 @@
 #include "GameFramework/CharacterMovementComponent.h"
 #include "Components/SkeletalMeshComponent.h"
 #include "Animation/AnimSequence.h"
+#include "RigVMCore/RigVMTemplate.h"
 
 FAnimNextParamType::FAnimNextParamType(EValueType InValueType, EContainerType InContainerType, const UObject* InValueTypeObject)
 	: ValueTypeObject(InValueTypeObject)
@@ -19,6 +20,80 @@ FAnimNextParamType::FAnimNextParamType(EValueType InValueType, EContainerType In
 {
 }
 
+FAnimNextParamType FAnimNextParamType::FromRigVMTemplateArgument(const FRigVMTemplateArgumentType& RigVMType)
+{
+	FAnimNextParamType Type;	
+	const FString CPPType = RigVMType.CPPType.ToString();
+	
+	if (RigVMTypeUtils::IsArrayType(CPPType))
+	{
+		Type.ContainerType = EPropertyBagContainerType::Array;
+	}
+
+	static const FName IntTypeName(TEXT("int")); // type used by some engine tests
+	static const FName Int64TypeName(TEXT("Int64"));
+	static const FName UInt64TypeName(TEXT("UInt64"));
+
+	if (CPPType == RigVMTypeUtils::BoolTypeName)
+	{
+		Type.ValueType = EPropertyBagPropertyType::Bool;
+	}
+	else if (CPPType == RigVMTypeUtils::Int32TypeName || CPPType == IntTypeName)
+	{
+		Type.ValueType = EPropertyBagPropertyType::Int32;
+	}
+	else if (CPPType == RigVMTypeUtils::UInt32TypeName)
+	{
+		Type.ValueType = EPropertyBagPropertyType::UInt32;
+	}
+	else if (CPPType == Int64TypeName)
+	{
+		Type.ValueType = EPropertyBagPropertyType::Int64;
+	}
+	else if (CPPType == UInt64TypeName)
+	{
+		Type.ValueType = EPropertyBagPropertyType::UInt64;
+	}
+	else if (CPPType == RigVMTypeUtils::FloatTypeName)
+	{
+		Type.ValueType = EPropertyBagPropertyType::Float;
+	}
+	else if (CPPType == RigVMTypeUtils::DoubleTypeName)
+	{
+		Type.ValueType = EPropertyBagPropertyType::Double;
+	}
+	else if (CPPType == RigVMTypeUtils::FNameTypeName)
+	{
+		Type.ValueType = EPropertyBagPropertyType::Name;
+	}
+	else if (CPPType == RigVMTypeUtils::FStringTypeName)
+	{
+		Type.ValueType = EPropertyBagPropertyType::String;
+	}
+	else if (UScriptStruct* ScriptStruct = Cast<UScriptStruct>(RigVMType.CPPTypeObject))
+	{
+		Type.ValueType = EPropertyBagPropertyType::Struct;
+		Type.ValueTypeObject = ScriptStruct;
+	}
+	else if (UEnum* Enum = Cast<UEnum>(RigVMType.CPPTypeObject))
+	{
+		Type.ValueType = EPropertyBagPropertyType::Enum;
+		Type.ValueTypeObject = Enum;
+	}
+	else if (UObject* Object = Cast<UObject>(RigVMType.CPPTypeObject))
+	{
+		Type.ValueType = EPropertyBagPropertyType::Object;	
+		Type.ValueTypeObject = Object;
+	}
+	else
+	{
+		ensureMsgf(false, TEXT("Unsupported type : %s"), *CPPType);
+		Type.ValueType = EPropertyBagPropertyType::None;
+	}
+
+	return Type;
+}
+
 bool FAnimNextParamType::IsValidObject() const
 {
 	switch(ValueType)
diff --git a/Engine/Plugins/Experimental/AnimNext/Source/AnimNext/Private/Param/RigVMDispatch_GetLayerParameter.cpp b/Engine/Plugins/Experimental/AnimNext/Source/AnimNext/Private/Param/RigVMDispatch_GetLayerParameter.cpp
--- a/Engine/Plugins/Experimental/AnimNext/Source/AnimNext/Private/Param/RigVMDispatch_GetLayerParameter.cpp
+++ b/Engine/Plugins/Experimental/AnimNext/Source/AnimNext/Private/Param/RigVMDispatch_GetLayerParameter.cpp
@@ -6,8 +6,6 @@
 #include "Graph/AnimNextExecuteContext.h"
 #include "Param/ParamStack.h"
 
-class UAnimNextParameter;
-
 const FName FRigVMDispatch_GetLayerParameter::ValueName = TEXT("Value");
 const FName FRigVMDispatch_GetLayerParameter::TypeHandleName = TEXT("Type");
 const FName FRigVMDispatch_GetLayerParameter::ParameterName = TEXT("Parameter");
diff --git a/Engine/Plugins/Experimental/AnimNext/Source/AnimNext/Private/Param/RigVMDispatch_GetParameter.cpp b/Engine/Plugins/Experimental/AnimNext/Source/AnimNext/Private/Param/RigVMDispatch_GetParameter.cpp
--- a/Engine/Plugins/Experimental/AnimNext/Source/AnimNext/Private/Param/RigVMDispatch_GetParameter.cpp
+++ b/Engine/Plugins/Experimental/AnimNext/Source/AnimNext/Private/Param/RigVMDispatch_GetParameter.cpp
@@ -8,8 +8,6 @@
 #include "Param/ParamStack.h"
 #include "Param/ParamStack.h"
 
-class UAnimNextParameter;
-
 const FName FRigVMDispatch_GetParameter::ValueName = TEXT("Value");
 const FName FRigVMDispatch_GetParameter::TypeHandleName = TEXT("Type");
 const FName FRigVMDispatch_GetParameter::ParameterName = TEXT("Parameter");
diff --git a/Engine/Plugins/Experimental/AnimNext/Source/AnimNext/Private/Param/RigVMDispatch_SetLayerParameter.cpp b/Engine/Plugins/Experimental/AnimNext/Source/AnimNext/Private/Param/RigVMDispatch_SetLayerParameter.cpp
--- a/Engine/Plugins/Experimental/AnimNext/Source/AnimNext/Private/Param/RigVMDispatch_SetLayerParameter.cpp
+++ b/Engine/Plugins/Experimental/AnimNext/Source/AnimNext/Private/Param/RigVMDispatch_SetLayerParameter.cpp
@@ -8,8 +8,6 @@
 #include "Param/ParamTypeHandle.h"
 #include "Param/ParamStackLayerHandle.h"
 
-class UAnimNextParameter;
-
 const FName FRigVMDispatch_SetLayerParameter::ExecuteContextName = TEXT("ExecuteContext");
 const FName FRigVMDispatch_SetLayerParameter::ValueName = TEXT("Value");
 const FName FRigVMDispatch_SetLayerParameter::TypeHandleName = TEXT("Type");
diff --git a/Engine/Plugins/Experimental/AnimNext/Source/AnimNext/Public/Param/ParamType.h b/Engine/Plugins/Experimental/AnimNext/Source/AnimNext/Public/Param/ParamType.h
--- a/Engine/Plugins/Experimental/AnimNext/Source/AnimNext/Public/Param/ParamType.h
+++ b/Engine/Plugins/Experimental/AnimNext/Source/AnimNext/Public/Param/ParamType.h
@@ -9,6 +9,7 @@
 #include "UObject/SoftObjectPtr.h"
 #include "UObject/ReflectedTypeAccessors.h"
 #include "Concepts/BaseStructureProvider.h"
+#include "RigVMCore/RigVMTemplate.h"
 #include "ParamType.generated.h"
 
 namespace UE::AnimNext
@@ -70,6 +71,8 @@ public:
 	/** Construct a parameter type from the passed in value, container and object type. */
 	FAnimNextParamType(EValueType InValueType, EContainerType InContainerType = EContainerType::None, const UObject* InValueTypeObject = nullptr);
 
+	/** Construct a parameter type from the passed in FRigVMTemplateArgumentType. */
+	static FAnimNextParamType FromRigVMTemplateArgument(const FRigVMTemplateArgumentType& RigVMType);
 private:
 	/** Pointer to object that defines the Enum, Struct, or Class. */
 	UPROPERTY()
diff --git a/Engine/Plugins/Experimental/AnimNext/Source/AnimNextEditor/Internal/Param/ParameterLibraryFactory.h b/Engine/Plugins/Experimental/AnimNext/Source/AnimNextEditor/Internal/Param/ParameterLibraryFactory.h
deleted file mode 100644
diff --git a/Engine/Plugins/Experimental/AnimNext/Source/AnimNextEditor/Internal/Param/ParameterPickerArgs.h b/Engine/Plugins/Experimental/AnimNext/Source/AnimNextEditor/Internal/Param/ParameterPickerArgs.h
--- a/Engine/Plugins/Experimental/AnimNext/Source/AnimNextEditor/Internal/Param/ParameterPickerArgs.h
+++ b/Engine/Plugins/Experimental/AnimNext/Source/AnimNextEditor/Internal/Param/ParameterPickerArgs.h
@@ -10,14 +10,36 @@
 namespace UE::AnimNext::Editor
 {
 
+struct FParameterToAdd
+{
+	FParameterToAdd() = default;
+
+	FParameterToAdd(const FAnimNextParamType& InType, FName InName)
+		: Type(InType)
+		, Name(InName)
+	{}
+
+	bool IsValid() const
+	{
+		return Name != NAME_None && Type.IsValid(); 
+	}
+
+	bool IsValid(FText& OutReason) const;
+
+	// Type
+	FAnimNextParamType Type;
+
+	// Name for parameter
+	FName Name;
+};
+
 // A parameter asset, optionally bound in a block
 struct FParameterBindingReference
 {
 	FParameterBindingReference() = default;
 
-	FParameterBindingReference(const FName& InParameter, const FAssetData& InLibrary = FAssetData(), const FAssetData& InBlock = FAssetData())
+	FParameterBindingReference(const FName& InParameter, const FAssetData& InBlock = FAssetData())
 		: Parameter(InParameter)
-		, Library(InLibrary)
 		, Block(InBlock)
 	{
 	}
@@ -25,8 +47,8 @@ struct FParameterBindingReference
 	// Parameter name
 	FName Parameter;
 
-	// Optional library asset (as this can represent built-in parameters too)
-	FAssetData Library;
+	// Asset (first found in asset registry) that the parameter is used in
+	FAssetData Asset;
 
 	// Optional block asset that the parameter is bound in
 	FAssetData Block;
@@ -38,6 +60,9 @@ DECLARE_DELEGATE_OneParam(FOnGetParameterBindings, TArray<FParameterBindingRefer
 // Delegate called when a parameter has been picked. Block argument is invalid when an unbound parameter is chosen.
 DECLARE_DELEGATE_OneParam(FOnParameterPicked, const FParameterBindingReference& /*InParameterBinding*/);
 
+// Delegate called when a parameter is due to be added.
+DECLARE_DELEGATE_OneParam(FOnAddParameter, const FParameterToAdd& /*InParameterToAdd*/);
+
 // Result of a filter operation via FOnFilterParameter
 enum class EFilterParameterResult : int32
 {
@@ -65,6 +90,9 @@ struct FParameterPickerArgs
 	// Delegate called when a single parameter has been picked
 	FOnParameterPicked OnParameterPicked;
 
+	// Delegate called when a parameter, or set of parameters is added
+	FOnAddParameter OnAddParameter;
+
 	// Delegate called to filter parameters for display to the user
 	FOnFilterParameter OnFilterParameter;
 
@@ -83,9 +111,6 @@ struct FParameterPickerArgs
 	// Whether we should show parameters that are built in
 	bool bShowBuiltInParameters = true;
 
-	// Whether we should show the library alongside parameters
-	bool bShowLibraries = true;
-
 	// Whether we should show the block alongside bound parameters
 	bool bShowBlocks = true;
 
diff --git a/Engine/Plugins/Experimental/AnimNext/Source/AnimNextEditor/Private/AnimNextEditorModule.cpp b/Engine/Plugins/Experimental/AnimNext/Source/AnimNextEditor/Private/AnimNextEditorModule.cpp
--- a/Engine/Plugins/Experimental/AnimNext/Source/AnimNextEditor/Private/AnimNextEditorModule.cpp
+++ b/Engine/Plugins/Experimental/AnimNext/Source/AnimNextEditor/Private/AnimNextEditorModule.cpp
@@ -14,10 +14,8 @@
 #include "Graph/AnimNextGraph.h"
 #include "Param/AnimNextParameterBlock.h"
 #include "Param/AnimNextParameterBlockEntry.h"
-#include "Param/AnimNextParameterLibrary.h"
 #include "Param/IAnimNextParameterBlockGraphInterface.h"
 #include "Param/SParameterBlockView.h"
-#include "Param/SParameterLibraryView.h"
 #include "Graph/SAnimNextGraphView.h"
 #include "Scheduler/AnimNextSchedule.h"
 #include "Workspace/AnimNextWorkspaceEditor.h"
@@ -119,18 +117,6 @@ class FModule : public IModule
 				});
 		});
 
-		FWorkspaceEditor::RegisterAssetDocumentWidget(UAnimNextParameterLibrary::StaticClass()->GetFName(), [](TSharedRef<FWorkspaceEditor> InEditor, UObject* InAsset)
-		{
-			UAnimNextParameterLibrary* ParameterLibrary = CastChecked<UAnimNextParameterLibrary>(InAsset);
-			return SNew(SParameterLibraryView, ParameterLibrary)
-				.OnSelectionChanged_Lambda([WeakEditor = TWeakPtr<FWorkspaceEditor>(InEditor)](const TArray<UObject*>& InObjects)
-				{
-					if(TSharedPtr<FWorkspaceEditor> Editor = WeakEditor.Pin())
-					{
-						Editor->SetSelectedObjects(InObjects);
-					}
-				});
-		});
 
 		FWorkspaceEditor::RegisterAssetDocumentWidget(UAnimNextSchedule::StaticClass()->GetFName(), [](TSharedRef<FWorkspaceEditor> InEditor, UObject* InAsset)
 		{
diff --git a/Engine/Plugins/Experimental/AnimNext/Source/AnimNextEditor/Private/EditorUtils.cpp b/Engine/Plugins/Experimental/AnimNext/Source/AnimNextEditor/Private/EditorUtils.cpp
--- a/Engine/Plugins/Experimental/AnimNext/Source/AnimNextEditor/Private/EditorUtils.cpp
+++ b/Engine/Plugins/Experimental/AnimNext/Source/AnimNextEditor/Private/EditorUtils.cpp
@@ -6,7 +6,6 @@
 #include "Kismet2/Kismet2NameValidators.h"
 #include "Param/AnimNextParameterBlock_EditorData.h"
 #include "Param/AnimNextParameterBlock.h"
-#include "Param/AnimNextParameterLibrary.h"
 #include "PropertyBagDetails.h"
 #include "Graph/AnimNextGraphEntry.h"
 #include "AssetRegistry/AssetRegistryModule.h"
@@ -16,6 +15,9 @@
 #include "Param/Params.h"
 #include "Workspace/AnimNextWorkspace.h"
 #include "UncookedOnlyUtils.h"
+#include "Param/RigVMDispatch_GetLayerParameter.h"
+#include "Param/RigVMDispatch_GetParameter.h"
+#include "Param/RigVMDispatch_SetLayerParameter.h"
 
 namespace UE::AnimNext::Editor
 {
@@ -434,14 +436,15 @@ void FUtils::GetFilteredVariableTypeTree(TArray<TSharedPtr<UEdGraphSchema_K2::FP
 	}
 };
 
-FName FUtils::GetNewParameterNameInLibrary(const FAssetData& InLibraryAsset, const TCHAR* InBaseName, TArrayView<FName> InAdditionalExistingNames)
+
+FName FUtils::GetNewParameterName(const TCHAR* InBaseName, TArrayView<FName> InAdditionalExistingNames)
 {
-	FAnimNextParameterLibraryAssetRegistryExports Exports;
-	UncookedOnly::FUtils::GetExportedParametersForLibrary(InLibraryAsset, Exports);
+	FAnimNextParameterProviderAssetRegistryExports Exports;
+	UE::AnimNext::UncookedOnly::FUtils::GetExportedParametersFromAssetRegistry(Exports);
 
 	auto NameExists = [&Exports, &InAdditionalExistingNames](const TCHAR* InName)
 	{
-		for(const FAnimNextParameterLibraryAssetRegistryExportEntry& Parameter : Exports.Parameters)
+		for(const FAnimNextParameterAssetRegistryExportEntry& Parameter : Exports.Parameters)
 		{
 			if(Parameter.Name.ToString() == InName)
 			{
@@ -481,38 +484,18 @@ FName FUtils::GetNewParameterNameInLibrary(const FAssetData& InLibraryAsset, con
 	return NAME_None;
 }
 
-bool FUtils::DoesParameterExistInLibrary(const FAssetData& InLibraryAsset, const FName InParameterName)
-{
-	FAnimNextParameterLibraryAssetRegistryExports Exports;
-	UncookedOnly::FUtils::GetExportedParametersForLibrary(InLibraryAsset, Exports);
-
-	for(const FAnimNextParameterLibraryAssetRegistryExportEntry& Parameter : Exports.Parameters)
-	{
-		if(Parameter.Name == InParameterName)
-		{
-			return true;
-		}
-	}
-
-	return false;
-}
-
-FAnimNextParamType FUtils::GetParameterTypeFromLibraryExports(FName InName, const FAnimNextParameterLibraryAssetRegistryExports& InExports)
+bool FUtils::DoesParameterNameExist(const FName InName)
 {
-	for(const FAnimNextParameterLibraryAssetRegistryExportEntry& Export : InExports.Parameters)
-	{
-		if(Export.Name == InName)
-		{
-			return Export.Type;
-		}
-	}
-	return FAnimNextParamType();
+	FAnimNextParameterProviderAssetRegistryExports Exports;
+	UncookedOnly::FUtils::GetExportedParametersFromAssetRegistry(Exports);
+	return Exports.Parameters.ContainsByPredicate([InName](const FAnimNextParameterAssetRegistryExportEntry& Entry) { return Entry.Name == InName; });
 }
 
-bool FUtils::GetExportedBindingsForBlock(const FAssetData& InLibraryAsset, FAnimNextParameterBlockAssetRegistryExports& OutExports)
+bool FUtils::DoesParameterNameExistInAsset(const FName InName, const FAssetData& InAsset)
 {
-	const FString TagValue = InLibraryAsset.GetTagValueRef<FString>(UAnimNextParameterBlock_EditorData::ExportsAssetRegistryTag);
-	return FAnimNextParameterBlockAssetRegistryExports::StaticStruct()->ImportText(*TagValue, &OutExports, nullptr, PPF_None, nullptr, FAnimNextParameterBlockAssetRegistryExports::StaticStruct()->GetName()) != nullptr;
+	FAnimNextParameterProviderAssetRegistryExports Exports;
+	UncookedOnly::FUtils::GetExportedParametersForAsset(InAsset, Exports);
+	return Exports.Parameters.ContainsByPredicate([InName](const FAnimNextParameterAssetRegistryExportEntry& Entry) { return Entry.Name == InName; });
 }
 
 bool FUtils::GetExportedAssetsForWorkspace(const FAssetData& InWorkspaceAsset, FAnimNextWorkspaceAssetRegistryExports& OutExports)
diff --git a/Engine/Plugins/Experimental/AnimNext/Source/AnimNextEditor/Private/EditorUtils.h b/Engine/Plugins/Experimental/AnimNext/Source/AnimNextEditor/Private/EditorUtils.h
--- a/Engine/Plugins/Experimental/AnimNext/Source/AnimNextEditor/Private/EditorUtils.h
+++ b/Engine/Plugins/Experimental/AnimNext/Source/AnimNextEditor/Private/EditorUtils.h
@@ -13,12 +13,12 @@ struct FAnimNextParamType;
 class UAnimNextParameterBlock;
 class UAnimNextParameterBlockBinding;
 class UAnimNextParameterBlock_EditorData;
-class UAnimNextParameter;
 class URigVMController;
-struct FAnimNextParameterLibraryAssetRegistryExports;
 struct FAnimNextParameterBlockAssetRegistryExports;
 struct FAnimNextWorkspaceAssetRegistryExports;
 
+struct FAnimNextParameterProviderAssetRegistryExports;
+
 namespace UE::AnimNext::Editor
 {
 
@@ -36,14 +36,10 @@ struct FUtils
 
 	static void GetFilteredVariableTypeTree(TArray<TSharedPtr<UEdGraphSchema_K2::FPinTypeTreeInfo>>& TypeTree, ETypeTreeFilter TypeTreeFilter);
 
-	static FName GetNewParameterNameInLibrary(const FAssetData& InLibraryAsset, const TCHAR* InBaseName, TArrayView<FName> InAdditionalExistingNames);
-
-	static bool DoesParameterExistInLibrary(const FAssetData& InLibraryAsset, const FName InParameterName);
-
-	static FAnimNextParamType GetParameterTypeFromLibraryExports(FName InName, const FAnimNextParameterLibraryAssetRegistryExports& InExports);
-
-	static bool GetExportedBindingsForBlock(const FAssetData& InBlockAsset, FAnimNextParameterBlockAssetRegistryExports& OutExports);
+	static FName GetNewParameterName(const TCHAR* InBaseName, TArrayView<FName> InAdditionalExistingNames);
 
+	static bool DoesParameterNameExist(const FName InName);
+	static bool DoesParameterNameExistInAsset(const FName InName, const FAssetData& InAsset);
 
 	static bool GetExportedAssetsForWorkspace(const FAssetData& InWorkspaceAsset, FAnimNextWorkspaceAssetRegistryExports& OutExports);
 };
diff --git a/Engine/Plugins/Experimental/AnimNext/Source/AnimNextEditor/Private/Param/AnimNextParameterLibraryEditor.cpp b/Engine/Plugins/Experimental/AnimNext/Source/AnimNextEditor/Private/Param/AnimNextParameterLibraryEditor.cpp
deleted file mode 100644
diff --git a/Engine/Plugins/Experimental/AnimNext/Source/AnimNextEditor/Private/Param/AnimNextParameterLibraryEditor.h b/Engine/Plugins/Experimental/AnimNext/Source/AnimNextEditor/Private/Param/AnimNextParameterLibraryEditor.h
deleted file mode 100644
diff --git a/Engine/Plugins/Experimental/AnimNext/Source/AnimNextEditor/Private/Param/AnimNextParameterSettings.cpp b/Engine/Plugins/Experimental/AnimNext/Source/AnimNextEditor/Private/Param/AnimNextParameterSettings.cpp
--- a/Engine/Plugins/Experimental/AnimNext/Source/AnimNextEditor/Private/Param/AnimNextParameterSettings.cpp
+++ b/Engine/Plugins/Experimental/AnimNext/Source/AnimNextEditor/Private/Param/AnimNextParameterSettings.cpp
@@ -20,17 +20,4 @@ const FAnimNextParamType& UAnimNextParameterSettings::GetLastParameterType() con
 void UAnimNextParameterSettings::SetLastParameterType(const FAnimNextParamType& InLastParameterType)
 {
 	LastParameterType = InLastParameterType;
-}
-
-FAssetData UAnimNextParameterSettings::GetLastLibrary() const
-{
-	IAssetRegistry& AssetRegistryModule = FModuleManager::Get().LoadModuleChecked<FAssetRegistryModule>(TEXT("AssetRegistry")).Get();
-	return AssetRegistryModule.GetAssetByObjectPath(LastLibrary);
-}
-
-void UAnimNextParameterSettings::SetLastLibrary(const FAssetData& InLastLibrary)
-{
-	LastLibrary = InLastLibrary.GetSoftObjectPath();
-
-	SaveConfig();
 }
\ No newline at end of file
diff --git a/Engine/Plugins/Experimental/AnimNext/Source/AnimNextEditor/Private/Param/AnimNextParameterSettings.h b/Engine/Plugins/Experimental/AnimNext/Source/AnimNextEditor/Private/Param/AnimNextParameterSettings.h
--- a/Engine/Plugins/Experimental/AnimNext/Source/AnimNextEditor/Private/Param/AnimNextParameterSettings.h
+++ b/Engine/Plugins/Experimental/AnimNext/Source/AnimNextEditor/Private/Param/AnimNextParameterSettings.h
@@ -20,17 +20,7 @@ public:
 
 	// Set the type of the last parameter type that we created
 	void SetLastParameterType(const FAnimNextParamType& InLastParameterType);
-
-	// Get the library of the last parameter type that we created
-	FAssetData GetLastLibrary() const;
-
-	// Set the library of the last parameter type that we created
-	void SetLastLibrary(const FAssetData& InLastLibrary);
-
 private:
-	UPROPERTY(config)
-	FSoftObjectPath LastLibrary;
-
 	UPROPERTY(Transient)
 	FAnimNextParamType LastParameterType = FAnimNextParamType::GetType<bool>();
 };
\ No newline at end of file
diff --git a/Engine/Plugins/Experimental/AnimNext/Source/AnimNextEditor/Private/Param/ParamNamePropertyCustomization.cpp b/Engine/Plugins/Experimental/AnimNext/Source/AnimNextEditor/Private/Param/ParamNamePropertyCustomization.cpp
--- a/Engine/Plugins/Experimental/AnimNext/Source/AnimNextEditor/Private/Param/ParamNamePropertyCustomization.cpp
+++ b/Engine/Plugins/Experimental/AnimNext/Source/AnimNextEditor/Private/Param/ParamNamePropertyCustomization.cpp
@@ -52,7 +52,6 @@ void FParamNamePropertyTypeCustomization::CustomizeHeader(TSharedRef<IPropertyHa
 	FAnimNextParamType FilterType = FAnimNextParamType::FromString(ParamTypeString);
 	
 	FParameterPickerArgs PickerArgs;
-	PickerArgs.bShowLibraries = false;
 	PickerArgs.bShowBlocks = false;
 	PickerArgs.bMultiSelect = false;
 	PickerArgs.OnFilterParameterType = FOnFilterParameterType::CreateLambda([FilterType](const FAnimNextParamType& InParameterType)
diff --git a/Engine/Plugins/Experimental/AnimNext/Source/AnimNextEditor/Private/Param/ParameterLibraryEditorMode.cpp b/Engine/Plugins/Experimental/AnimNext/Source/AnimNextEditor/Private/Param/ParameterLibraryEditorMode.cpp
deleted file mode 100644
diff --git a/Engine/Plugins/Experimental/AnimNext/Source/AnimNextEditor/Private/Param/ParameterLibraryEditorMode.h b/Engine/Plugins/Experimental/AnimNext/Source/AnimNextEditor/Private/Param/ParameterLibraryEditorMode.h
deleted file mode 100644
diff --git a/Engine/Plugins/Experimental/AnimNext/Source/AnimNextEditor/Private/Param/ParameterLibraryFactory.cpp b/Engine/Plugins/Experimental/AnimNext/Source/AnimNextEditor/Private/Param/ParameterLibraryFactory.cpp
deleted file mode 100644
diff --git a/Engine/Plugins/Experimental/AnimNext/Source/AnimNextEditor/Private/Param/ParametersAssetDefinitions.cpp b/Engine/Plugins/Experimental/AnimNext/Source/AnimNextEditor/Private/Param/ParametersAssetDefinitions.cpp
--- a/Engine/Plugins/Experimental/AnimNext/Source/AnimNextEditor/Private/Param/ParametersAssetDefinitions.cpp
+++ b/Engine/Plugins/Experimental/AnimNext/Source/AnimNextEditor/Private/Param/ParametersAssetDefinitions.cpp
@@ -2,45 +2,12 @@
 
 #include "ParametersAssetDefinitions.h"
 #include "AnimNextParameterBlockEditor.h"
-#include "AnimNextParameterLibraryEditor.h"
 #include "Param/AnimNextParameterBlock.h"
-#include "Param/AnimNextParameterLibrary.h"
 #include "Workspace/AnimNextWorkspaceEditor.h"
 #include "EditorCVars.h"
 
 #define LOCTEXT_NAMESPACE "AnimNextAssetDefinitions"
 
-EAssetCommandResult UAssetDefinition_AnimNextParameterLibrary::OpenAssets(const FAssetOpenArgs& OpenArgs) const
-{
-	using namespace UE::AnimNext::Editor;
-
-	for (UAnimNextParameterLibrary* Asset : OpenArgs.LoadObjects<UAnimNextParameterLibrary>())
-	{
-		if(CVars::GUseWorkspaceEditor.GetValueOnGameThread())
-		{
-			FWorkspaceEditor::OpenWorkspaceForAsset(Asset, FWorkspaceEditor::EOpenWorkspaceMethod::Default);
-		}
-		else
-		{
-			TSharedRef<FParameterLibraryEditor> Editor = MakeShared<UE::AnimNext::Editor::FParameterLibraryEditor>();
-			Editor->InitEditor(OpenArgs.GetToolkitMode(), OpenArgs.ToolkitHost, Asset);
-		}
-	}
-
-	return EAssetCommandResult::Handled;
-}
-
-FText UAssetDefinition_AnimNextParameter::GetObjectDisplayNameText(UObject* Object) const
-{
-	UAnimNextParameter* Parameter = CastChecked<UAnimNextParameter>(Object);
-	if(UAnimNextParameterLibrary* Library = CastChecked<UAnimNextParameterLibrary>(Object->GetOuter()))
-	{
-		return FText::Format(LOCTEXT("LibraryParameterDisplayFormat", "{0}.{1}"), FText::FromString(Library->GetName()), FText::FromString(Parameter->GetName()));
-	}
-
-	return FText::FromString(Parameter->GetName());
-}
-
 EAssetCommandResult UAssetDefinition_AnimNextParameterBlock::OpenAssets(const FAssetOpenArgs& OpenArgs) const
 {
 	using namespace UE::AnimNext::Editor;
diff --git a/Engine/Plugins/Experimental/AnimNext/Source/AnimNextEditor/Private/Param/ParametersAssetDefinitions.h b/Engine/Plugins/Experimental/AnimNext/Source/AnimNextEditor/Private/Param/ParametersAssetDefinitions.h
--- a/Engine/Plugins/Experimental/AnimNext/Source/AnimNextEditor/Private/Param/ParametersAssetDefinitions.h
+++ b/Engine/Plugins/Experimental/AnimNext/Source/AnimNextEditor/Private/Param/ParametersAssetDefinitions.h
@@ -3,44 +3,12 @@
 #pragma once
 
 #include "AssetDefinitionDefault.h"
-#include "Param/AnimNextParameter.h"
 #include "Param/AnimNextParameterBlockBinding.h"
 #include "Param/AnimNextParameterBlock.h"
-#include "Param/AnimNextParameterLibrary.h"
 #include "ParametersAssetDefinitions.generated.h"
 
 #define LOCTEXT_NAMESPACE "AnimNextAssetDefinitions"
 
-UCLASS()
-class UAssetDefinition_AnimNextParameterLibrary : public UAssetDefinitionDefault
-{
-	GENERATED_BODY()
-
-public:
-	// UAssetDefinition interface
-	virtual FText GetAssetDisplayName() const override { return LOCTEXT("AnimNextParameterLibrary", "AnimNext Parameter Library"); }
-	virtual FLinearColor GetAssetColor() const override { return FLinearColor(FColor(128,64,64)); }
-	virtual TSoftClassPtr<UObject> GetAssetClass() const override { return UAnimNextParameterLibrary::StaticClass(); }
-	virtual EAssetCommandResult OpenAssets(const FAssetOpenArgs& OpenArgs) const override;
-	virtual TConstArrayView<FAssetCategoryPath> GetAssetCategories() const override
-	{
-		static const auto Categories = { FAssetCategoryPath(EAssetCategoryPaths::Animation, LOCTEXT("AnimNextSubMenu", "AnimNext")) };
-		return Categories;
-	}
-};
-
-UCLASS()
-class UAssetDefinition_AnimNextParameter : public UAssetDefinitionDefault
-{
-	GENERATED_BODY()
-
-public:
-	// UAssetDefinition interface
-	virtual FText GetAssetDisplayName() const override { return LOCTEXT("AnimNextParameter", "AnimNext Parameter"); }
-	virtual FLinearColor GetAssetColor() const override { return FLinearColor(FColor(128,64,64)); }
-	virtual TSoftClassPtr<UObject> GetAssetClass() const override { return UAnimNextParameter::StaticClass(); }
-	virtual FText GetObjectDisplayNameText(UObject* Object) const override;
-};
 
 UCLASS()
 class UAssetDefinition_AnimNextParameterBlock : public UAssetDefinitionDefault
diff --git a/Engine/Plugins/Experimental/AnimNext/Source/AnimNextEditor/Private/Param/ParametersGraphPanelPinFactory.cpp b/Engine/Plugins/Experimental/AnimNext/Source/AnimNextEditor/Private/Param/ParametersGraphPanelPinFactory.cpp
--- a/Engine/Plugins/Experimental/AnimNext/Source/AnimNextEditor/Private/Param/ParametersGraphPanelPinFactory.cpp
+++ b/Engine/Plugins/Experimental/AnimNext/Source/AnimNextEditor/Private/Param/ParametersGraphPanelPinFactory.cpp
@@ -26,7 +26,8 @@ TSharedPtr<SGraphPin> FParametersGraphPanelPinFactory::CreatePin_Internal(UEdGra
 			if(ModelPin->GetCustomWidgetName() == "ParamName")
 			{
 				return SNew(SGraphPinParamName, InPin)
-					.ModelPin(ModelPin);
+					.ModelPin(ModelPin)
+					.GraphNode(RigNode);
 			}
 		}
 	}
diff --git a/Engine/Plugins/Experimental/AnimNext/Source/AnimNextEditor/Private/Param/SAddParametersDialog.cpp b/Engine/Plugins/Experimental/AnimNext/Source/AnimNextEditor/Private/Param/SAddParametersDialog.cpp
--- a/Engine/Plugins/Experimental/AnimNext/Source/AnimNextEditor/Private/Param/SAddParametersDialog.cpp
+++ b/Engine/Plugins/Experimental/AnimNext/Source/AnimNextEditor/Private/Param/SAddParametersDialog.cpp
@@ -22,8 +22,6 @@
 #include "Widgets/Text/SInlineEditableTextBlock.h"
 #include "SSimpleButton.h"
 #include "SSimpleComboButton.h"
-#include "Param/AnimNextParameterLibrary.h"
-#include "Param/ParameterLibraryFactory.h"
 #include "ToolMenus.h"
 
 #define LOCTEXT_NAMESPACE "SAddParametersDialog"
@@ -35,7 +33,6 @@ namespace AddParametersDialog
 {
 static FName Column_Name(TEXT("Name"));
 static FName Column_Type(TEXT("Type"));
-static FName Column_Library(TEXT("Library"));
 static FName SelectLibraryMenuName(TEXT("AnimNext.AddParametersDialog.SelectedLibraryMenu"));
 }
 
@@ -50,11 +47,6 @@ bool FParameterToAdd::IsValid(FText& OutReason) const
 	{
 		OutReason = LOCTEXT("InvalidParameterType", "Invalid Parameter Type");
 	}
-
-	if(!Library.IsValid())
-	{
-		OutReason = LOCTEXT("InvalidParameterLibrary", "Invalid Parameter Library");
-	}
 	
 	return true; 
 }
@@ -63,7 +55,7 @@ void SAddParametersDialog::Construct(const FArguments& InArgs)
 {
 	using namespace AddParametersDialog;
 
-	Library = InArgs._Library;
+	TargetBlock = InArgs._Block;
 
 	SWindow::Construct(SWindow::FArguments()
 		.Title(LOCTEXT("WindowTitle", "Add Parameters"))
@@ -110,11 +102,6 @@ void SAddParametersDialog::Construct(const FArguments& InArgs)
 						.DefaultLabel(LOCTEXT("TypeColumnHeader", "Type"))
 						.ToolTipText(LOCTEXT("TypeColumnHeaderTooltip", "The type of the new parameter"))
 						.FillWidth(0.25f)
-
-						+SHeaderRow::Column(Column_Library)
-						.DefaultLabel(LOCTEXT("LibraryColumnHeader", "Library"))
-						.ToolTipText(LOCTEXT("LibraryColumnHeaderTooltip", "The library the new parameter will be created in"))
-						.FillWidth(0.5f)
 					)
 				]
 				+SVerticalBox::Slot()
@@ -191,23 +178,15 @@ void SAddParametersDialog::Construct(const FArguments& InArgs)
 void SAddParametersDialog::AddEntry()
 {
 	const UAnimNextParameterSettings* Settings = GetDefault<UAnimNextParameterSettings>();
-	FAssetData LibraryAsset = Settings->GetLastLibrary();
-	if(Library != nullptr)
-	{
-		LibraryAsset = FAssetData(Library);
-	}
 	
 	TArray<FName> PendingNames;
 	PendingNames.Reserve(Entries.Num());
 	for(const TSharedRef<FParameterToAdd>& QueuedAdd : Entries)
 	{
-		if(QueuedAdd->Library == LibraryAsset)
-		{
-			PendingNames.Add(QueuedAdd->Name);
-		}
+		PendingNames.Add(QueuedAdd->Name);
 	}
-	FName ParameterName = FUtils::GetNewParameterNameInLibrary(LibraryAsset, TEXT("NewParameter"), PendingNames);
-	Entries.Add(MakeShared<FParameterToAdd>(Settings->GetLastParameterType(), ParameterName, Settings->GetLastLibrary()));
+	FName ParameterName = FUtils::GetNewParameterName(TEXT("NewParameter"), PendingNames);
+	Entries.Add(MakeShared<FParameterToAdd>(Settings->GetLastParameterType(), ParameterName));
 
 	RefreshEntries();
 }
@@ -263,11 +242,10 @@ class SParameterToAdd : public SMultiColumnTableRow<TSharedRef<FParameterToAdd>>
 							return false;
 						}
 
-						FName Name(*NewString);
-
-						if(FUtils::DoesParameterExistInLibrary(Entry->Library, Name))
+						const FName Name(*NewString);
+						if(FUtils::DoesParameterNameExist(Name))
 						{
-							OutErrorText = LOCTEXT("Error_NameExists", "This name already exists in the specified library");
+							OutErrorText = LOCTEXT("Error_NameExists", "This name already exists in the project");
 							return false;
 						}
 
@@ -304,28 +282,6 @@ class SParameterToAdd : public SMultiColumnTableRow<TSharedRef<FParameterToAdd>>
 						.Font(IDetailLayoutBuilder::GetDetailFont())
 				];
 		}
-		else if(InColumnName == Column_Library)
-		{
-			return
-				SNew(SBox)
-				.HAlign(HAlign_Left)
-				.VAlign(VAlign_Center)
-				.IsEnabled(WeakDialog.Pin()->Library == nullptr)
-				[
-					SNew(SSimpleComboButton)
-					.UsesSmallText(true)
-					.HasDownArrow(true)
-					.Text_Lambda([this]()
-					{
-						return FText::FromName(Entry->Library.PackageName);
-					})
-					.ToolTipText_Lambda([this]()
-					{
-						return FText::Format(LOCTEXT("LibraryTooltip", "The library of the new parameter.\n{0}"), FText::FromName(Entry->Library.PackageName));
-					})
-					.OnGetMenuContent(WeakDialog.Pin().Get(), &SAddParametersDialog::HandleGetAddParameterMenuContent, Entry)
-				];
-		}
 
 		return SNullWidget::NullWidget;
 	}
@@ -365,95 +321,6 @@ TSharedRef<SWidget> SAddParametersDialog::HandleGetAddParameterMenuContent(TShar
 
 	UToolMenus* ToolMenus = UToolMenus::Get();
 
-	if(!ToolMenus->IsMenuRegistered(SelectLibraryMenuName))
-	{
-		UToolMenu* Menu = ToolMenus->RegisterMenu(SelectLibraryMenuName);
-
-		{
-			FToolMenuSection& Section = Menu->AddSection("AddNewLibrary", LOCTEXT("AddNewLibraryMenuSection", "Add New Library"));
-			Section.AddSubMenu("AddNewLibrary",
-				LOCTEXT("AddNewLibraryLabel", "Add New Library"),
-				LOCTEXT("AddNewLibraryTooltip", "Add a new parameter library"),
-				FNewToolMenuDelegate::CreateLambda([](UToolMenu* InMenu)
-				{
-					if(UAddParameterDialogMenuContext* MenuContext = InMenu->FindContext<UAddParameterDialogMenuContext>())
-					{
-						FPathPickerConfig PathPickerConfig;
-						PathPickerConfig.bShowFavorites = false;
-						PathPickerConfig.OnPathSelected = FOnPathSelected::CreateLambda([MenuContext](const FString& InPath)
-						{
-							FSlateApplication::Get().DismissAllMenus();
-
-							if(TSharedPtr<FParameterToAdd> Entry = MenuContext->Entry.Pin())
-							{
-								IAssetTools& AssetTools = FModuleManager::LoadModuleChecked<FAssetToolsModule>("AssetTools").Get();
-								if(UObject* NewAsset = AssetTools.CreateAsset(TEXT("NewParameterLibrary"), InPath, UAnimNextParameterLibrary::StaticClass(), NewObject<UAnimNextParameterLibraryFactory>()))
-								{
-									Entry->Library = FAssetData(NewAsset);
-
-									UAnimNextParameterSettings* Settings = GetMutableDefault<UAnimNextParameterSettings>();
-									Settings->SetLastLibrary(Entry->Library);
-								}
-							}
-						});
-
-						FContentBrowserModule& ContentBrowserModule = FModuleManager::LoadModuleChecked<FContentBrowserModule>("ContentBrowser");
-						FToolMenuSection& Section = InMenu->AddSection("ChooseLibrary", LOCTEXT("ChooseLibraryMenuSection", "Choose Existing Library"));
-						Section.AddEntry(FToolMenuEntry::InitWidget(
-								"LibraryPicker",
-								SNew(SBox)
-								.WidthOverride(300.0f)
-								.HeightOverride(400.0f)
-								[
-									ContentBrowserModule.Get().CreatePathPicker(PathPickerConfig)
-								],
-								FText::GetEmpty(),
-								true));
-					}
-				}));
-		}
-
-		{
-			
-
-			FToolMenuSection& Section = Menu->AddDynamicSection("ChooseLibrary",  FNewToolMenuDelegate::CreateLambda([](UToolMenu* InMenu)
-			{
-				if(UAddParameterDialogMenuContext* MenuContext = InMenu->FindContext<UAddParameterDialogMenuContext>())
-				{
-					FAssetPickerConfig AssetPickerConfig;
-					AssetPickerConfig.SelectionMode = ESelectionMode::Single;
-					AssetPickerConfig.InitialAssetViewType = EAssetViewType::List;
-					AssetPickerConfig.Filter.ClassPaths = { UAnimNextParameterLibrary::StaticClass()->GetClassPathName() };
-					AssetPickerConfig.OnAssetSelected = FOnAssetSelected::CreateLambda([MenuContext](const FAssetData& InAssetData)
-					{
-						FSlateApplication::Get().DismissAllMenus();
-
-						if(TSharedPtr<FParameterToAdd> Entry = MenuContext->Entry.Pin())
-						{
-							Entry->Library = InAssetData;
-
-							UAnimNextParameterSettings* Settings = GetMutableDefault<UAnimNextParameterSettings>();
-							Settings->SetLastLibrary(InAssetData);
-						}
-					});
-
-					FContentBrowserModule& ContentBrowserModule = FModuleManager::LoadModuleChecked<FContentBrowserModule>("ContentBrowser");
-					FToolMenuSection& Section = InMenu->AddSection("ChooseLibrary", LOCTEXT("ChooseLibraryMenuSection", "Choose Existing Library"));
-					Section.AddEntry(FToolMenuEntry::InitWidget(
-							"LibraryPicker",
-							SNew(SBox)
-							.WidthOverride(300.0f)
-							.HeightOverride(400.0f)
-							[
-								ContentBrowserModule.Get().CreateAssetPicker(AssetPickerConfig)
-							],
-							FText::GetEmpty(),
-							true));
-				}
-			}));
-		}
-	}
-
 	UAddParameterDialogMenuContext* MenuContext = NewObject<UAddParameterDialogMenuContext>();
 	MenuContext->AddParametersDialog = SharedThis(this);
 	MenuContext->Entry = InEntry;
diff --git a/Engine/Plugins/Experimental/AnimNext/Source/AnimNextEditor/Private/Param/SAddParametersDialog.h b/Engine/Plugins/Experimental/AnimNext/Source/AnimNextEditor/Private/Param/SAddParametersDialog.h
--- a/Engine/Plugins/Experimental/AnimNext/Source/AnimNextEditor/Private/Param/SAddParametersDialog.h
+++ b/Engine/Plugins/Experimental/AnimNext/Source/AnimNextEditor/Private/Param/SAddParametersDialog.h
@@ -9,47 +9,22 @@
 
 class SWrapBox;
 class UAnimNextParameterLibrary;
+class UAnimNextParameterBlock_EditorData;
 
 namespace UE::AnimNext::Editor
 {
+	struct FParameterToAdd;
 
-struct FParameterToAdd
-{
-	FParameterToAdd() = default;
-
-	FParameterToAdd(const FAnimNextParamType& InType, FName InName, const FAssetData& InLibrary)
-		: Type(InType)
-		, Name(InName)
-		, Library(InLibrary)
-	{}
-
-	bool IsValid() const
-	{
-		return Name != NAME_None && Type.IsValid() && Library.IsValid(); 
-	}
-
-	bool IsValid(FText& OutReason) const;
-	
-	// Type
-	FAnimNextParamType Type;
-
-	// Name for parameter
-	FName Name;
-
-	// Parameter library
-	FAssetData Library;
-};
 
 class SAddParametersDialog : public SWindow
 {
 public:
 	SLATE_BEGIN_ARGS(SAddParametersDialog)
-		: _Library(nullptr)
-		, _AllowMultiple(true)
+		: _Block(nullptr),
+		_AllowMultiple(true)
 	{}
 
-	/** The default library to create new parameters in */
-	SLATE_ARGUMENT(UAnimNextParameterLibrary*, Library)
+	SLATE_ARGUMENT(UAnimNextParameterBlock_EditorData*, Block)
 
 	/** Whether we allow multiple parameters to be added or just one at a time */
 	SLATE_ARGUMENT(bool, AllowMultiple)
@@ -78,8 +53,7 @@ private:
 
 	TArray<TSharedRef<FParameterToAdd>> Entries;
 
-	// The fixed library to use. If this is NULL, any library can be used.
-	UAnimNextParameterLibrary* Library = nullptr;
+	UAnimNextParameterBlock_EditorData* TargetBlock = nullptr;
 	
 	bool bCancelPressed = false;
 };
diff --git a/Engine/Plugins/Experimental/AnimNext/Source/AnimNextEditor/Private/Param/SGraphPinParamName.cpp b/Engine/Plugins/Experimental/AnimNext/Source/AnimNextEditor/Private/Param/SGraphPinParamName.cpp
--- a/Engine/Plugins/Experimental/AnimNext/Source/AnimNextEditor/Private/Param/SGraphPinParamName.cpp
+++ b/Engine/Plugins/Experimental/AnimNext/Source/AnimNextEditor/Private/Param/SGraphPinParamName.cpp
@@ -5,6 +5,9 @@
 #include "SParameterPickerCombo.h"
 #include "ScopedTransaction.h"
 #include "UncookedOnlyUtils.h"
+#include "Param/RigVMDispatch_GetLayerParameter.h"
+#include "Param/RigVMDispatch_GetParameter.h"
+#include "Param/RigVMDispatch_SetLayerParameter.h"
 
 #define LOCTEXT_NAMESPACE "SGraphPinParamName"
 
@@ -14,6 +17,7 @@ namespace UE::AnimNext::Editor
 void SGraphPinParamName::Construct(const FArguments& InArgs, UEdGraphPin* InPin)
 {
 	ModelPin = InArgs._ModelPin;
+	Node = InArgs._GraphNode;
 
 	SGraphPin::Construct(SGraphPin::FArguments(), InPin);
 }
@@ -21,9 +25,24 @@ void SGraphPinParamName::Construct(const FArguments& InArgs, UEdGraphPin* InPin)
 TSharedRef<SWidget> SGraphPinParamName::GetDefaultValueWidget()
 {
 	FParameterPickerArgs Args;
-	Args.bShowLibraries = false;
 	Args.bShowBlocks = false;
 	Args.bMultiSelect = false;
+
+	// Check whether this is a Set/Get parameter from block node, and if so only show bound parameters
+	if (const URigVMTemplateNode* TemplateNode = Cast<URigVMTemplateNode>(ModelPin->GetOuter()))
+	{
+		const FRigVMDispatchFactory* GetLayerParameterFactory = FRigVMRegistry::Get().FindOrAddDispatchFactory(FRigVMDispatch_GetLayerParameter::StaticStruct());
+		const FName GetLayerParameterNotation = GetLayerParameterFactory->GetTemplate()->GetNotation();
+
+		const FRigVMDispatchFactory* SetLayerParameterFactory = FRigVMRegistry::Get().FindOrAddDispatchFactory(FRigVMDispatch_SetLayerParameter::StaticStruct());
+		const FName SetLayerParameterNotation = SetLayerParameterFactory->GetTemplate()->GetNotation();
+
+		if (TemplateNode->GetNotation() == GetLayerParameterNotation || TemplateNode->GetNotation() == SetLayerParameterNotation)
+		{
+			Args.bShowUnboundParameters = false;
+		}
+	}
+	
 	Args.OnParameterPicked = FOnParameterPicked::CreateLambda([this](const FParameterBindingReference& InParameterBinding)
 	{
 		if(ModelPin)
@@ -34,6 +53,17 @@ TSharedRef<SWidget> SGraphPinParamName::GetDefaultValueWidget()
 		}
 	});
 	
+	Args.OnFilterParameterType = FOnFilterParameterType::CreateLambda([this](const FAnimNextParamType& InParamType)-> EFilterParameterResult
+	{
+		if(Node && ModelPin->IsLinked())
+		{
+			const FAnimNextParamType Type = FAnimNextParamType::FromRigVMTemplateArgument(ModelPin->GetTemplateArgumentType());
+			return Type.IsValid() && Type == InParamType ? EFilterParameterResult::Include : EFilterParameterResult::Exclude;
+		}
+
+		return EFilterParameterResult::Include;
+	});
+	
 	return SNew(SParameterPickerCombo)
 		.PickerArgs(Args)
 		.OnGetParameterName_Lambda([this]()
diff --git a/Engine/Plugins/Experimental/AnimNext/Source/AnimNextEditor/Private/Param/SGraphPinParamName.h b/Engine/Plugins/Experimental/AnimNext/Source/AnimNextEditor/Private/Param/SGraphPinParamName.h
--- a/Engine/Plugins/Experimental/AnimNext/Source/AnimNextEditor/Private/Param/SGraphPinParamName.h
+++ b/Engine/Plugins/Experimental/AnimNext/Source/AnimNextEditor/Private/Param/SGraphPinParamName.h
@@ -5,6 +5,7 @@
 #include "SGraphPin.h"
 
 class URigVMPin;
+class URigVMEdGraphNode;
 
 namespace UE::AnimNext::Editor
 {
@@ -15,6 +16,7 @@ class SGraphPinParamName : public SGraphPin
 	SLATE_BEGIN_ARGS(SGraphPinParamName) {}
 
 	SLATE_ARGUMENT(URigVMPin*, ModelPin)
+	SLATE_ARGUMENT(URigVMEdGraphNode*, GraphNode)
 
 	SLATE_END_ARGS()
 
@@ -24,6 +26,7 @@ private:
 	virtual TSharedRef<SWidget>	GetDefaultValueWidget() override;
 
 	URigVMPin* ModelPin = nullptr;
+	URigVMEdGraphNode* Node = nullptr;
 };
 
 }
\ No newline at end of file
diff --git a/Engine/Plugins/Experimental/AnimNext/Source/AnimNextEditor/Private/Param/SParameterBlockView.cpp b/Engine/Plugins/Experimental/AnimNext/Source/AnimNextEditor/Private/Param/SParameterBlockView.cpp
--- a/Engine/Plugins/Experimental/AnimNext/Source/AnimNextEditor/Private/Param/SParameterBlockView.cpp
+++ b/Engine/Plugins/Experimental/AnimNext/Source/AnimNextEditor/Private/Param/SParameterBlockView.cpp
@@ -3,12 +3,10 @@
 #include "SParameterBlockView.h"
 
 #include "Param/AnimNextParameterBlock.h"
-#include "Param/AnimNextParameter.h"
 #include "Param/AnimNextParameterBlock_EditorData.h"
 #include "Widgets/Layout/SBox.h"
 #include "Widgets/Input/SSearchBox.h"
 #include "Param/AnimNextParameterBlockEntry.h"
-#include "Param/AnimNextParameterLibrary.h"
 #include "DetailLayoutBuilder.h"
 #include "ISourceControlModule.h"
 #include "ISourceControlProvider.h"
@@ -694,19 +692,6 @@ class SParameterBlockViewRow : public SMultiColumnTableRow<TSharedRef<FParameter
 							{
 								TextBuilder.AppendLine(FText::FromName(ExternalPackage->GetFName()));
 							}
-
-							if(const IAnimNextParameterBlockParameterInterface* Binding = Cast<IAnimNextParameterBlockParameterInterface>(BlockEntry))
-							{
-								if(const UAnimNextParameter* Parameter = Binding->GetParameter())
-								{
-									const UPackage* ParameterPackage = Parameter->GetPackage();
-									check(ParameterPackage);
-									if(ParameterPackage->IsDirty())
-									{
-										TextBuilder.AppendLine(FText::FromName(ParameterPackage->GetFName()));
-									}
-								}
-							}
 						}
 
 						return TextBuilder.ToText();
@@ -723,18 +708,6 @@ class SParameterBlockViewRow : public SMultiColumnTableRow<TSharedRef<FParameter
 								bIsDirty = true;
 							}
 
-							if(const IAnimNextParameterBlockParameterInterface* Binding = Cast<IAnimNextParameterBlockParameterInterface>(BlockEntry))
-							{
-								if(const UAnimNextParameter* Parameter = Binding->GetParameter())
-								{
-									const UPackage* ParameterPackage = Parameter->GetPackage();
-									check(ParameterPackage);
-									if(ParameterPackage->IsDirty())
-									{
-										bIsDirty = true;
-									}
-								}
-							}
 
 							return bIsDirty ? FAppStyle::GetBrush("ContentBrowser.ContentDirty") : nullptr;
 						}
@@ -750,18 +723,28 @@ class SParameterBlockViewRow : public SMultiColumnTableRow<TSharedRef<FParameter
 					SNew(SBox)
 					.HAlign(HAlign_Left)
 					.VAlign(VAlign_Center)
-					.IsEnabled(false)
 					[
 						SNew(SPinTypeSelector, FGetPinTypeTree::CreateStatic(&Editor::FUtils::GetFilteredVariableTypeTree))
 							.TargetPinType_Lambda([this]()
 							{
-								if(IAnimNextParameterBlockParameterInterface* Binding = Cast<IAnimNextParameterBlockParameterInterface>(Entry->WeakEntry.Get()))
+								if(const IAnimNextParameterBlockParameterInterface* Binding = Cast<IAnimNextParameterBlockParameterInterface>(Entry->WeakEntry.Get()))
 								{
 									return UncookedOnly::FUtils::GetPinTypeFromParamType(Binding->GetParamType());
 								}
 
 								return FEdGraphPinType();
 							})
+							.OnPinTypeChanged_Lambda([this](const FEdGraphPinType& PinType)
+							{
+								if(IAnimNextParameterBlockParameterInterface* Binding = Cast<IAnimNextParameterBlockParameterInterface>(Entry->WeakEntry.Get()))
+								{
+									const FAnimNextParamType ParamType = UncookedOnly::FUtils::GetParamTypeFromPinType(PinType);
+									if(ParamType.IsValid())
+									{
+										Binding->SetParamType(ParamType);
+									}
+								}
+							})
 							.Schema(GetDefault<UPropertyBagSchema>())
 							.bAllowArrays(true)
 							.TypeTreeFilter(ETypeTreeFilter::None)
@@ -829,10 +812,10 @@ class SParameterBlockViewRow : public SMultiColumnTableRow<TSharedRef<FParameter
 							}
 						}
 
-						if(IAnimNextParameterBlockReferenceInterface* Reference = Cast<IAnimNextParameterBlockReferenceInterface>(Entry->WeakEntry.Get()))
+						if(const IAnimNextParameterBlockReferenceInterface* Reference = Cast<IAnimNextParameterBlockReferenceInterface>(Entry->WeakEntry.Get()))
 						{
 							const FName Name(*NewString);
-							if(!FUtils::DoesParameterExistInLibrary(Reference->GetBlock(), Name))
+							if(!FUtils::DoesParameterNameExistInAsset(Name, Reference->GetBlock()))
 							{
 								OutErrorText = LOCTEXT("Error_NameDoesNotExist", "This name does not exist in the specified block");
 								return false;
@@ -1014,11 +997,10 @@ TSharedRef<ITableRow> SParameterBlockView::HandleGenerateRow(TSharedRef<FParamet
 				{
 					FParameterPickerArgs Args;
 					Args.bMultiSelect = false;
-					Args.bShowLibraries = false;
 					Args.bShowBlocks = false;
-					Args.bShowBoundParameters = true;
+					Args.bShowBoundParameters = false;
 					Args.bShowBuiltInParameters = false; // Built-In paameters Disabled for MVP
-					Args.OnFilterParameter = FOnFilterParameter::CreateLambda([this](const FParameterBindingReference& InParameterBinding)
+				Args.OnFilterParameter = FOnFilterParameter::CreateLambda([this](const FParameterBindingReference& InParameterBinding)
 					{
 						// Skip params that are already bound in this block
 						if(InParameterBinding.Block == BlockAssetData)
@@ -1028,18 +1010,34 @@ TSharedRef<ITableRow> SParameterBlockView::HandleGenerateRow(TSharedRef<FParamet
 						
 						return EFilterParameterResult::Include;
 					});
-					Args.OnParameterPicked = FOnParameterPicked::CreateLambda([this](const FParameterBindingReference& InParameterBinding)
+
+					Args.OnAddParameter = FOnAddParameter::CreateLambda([this](const FParameterToAdd& ParameterToAdd)
 					{
 						FSlateApplication::Get().DismissAllMenus();
 
+						check(EditorData->FindBinding(ParameterToAdd.Name) == nullptr);
 						FScopedTransaction Transaction(LOCTEXT("AddParameter", "Add parameter"));
-						PendingSelection.Empty();
+						UAnimNextParameterBlockParameter* Parameter = EditorData->AddParameter(ParameterToAdd.Name, ParameterToAdd.Type);
 
-						// Create a new entry for the parameter
-						UAnimNextParameterLibrary* Library = Cast<UAnimNextParameterLibrary>(InParameterBinding.Library.GetAsset());
-						UAnimNextParameterBlockParameter* Parameter = EditorData->AddParameter(InParameterBinding.Parameter, Library);
+						PendingSelection.Empty();
 						PendingSelection.Add(Parameter);
 					});
+					Args.OnParameterPicked = FOnParameterPicked::CreateLambda([this](const FParameterBindingReference& InParameterBinding)
+					{
+						FSlateApplication::Get().DismissAllMenus();			
+						PendingSelection.Empty();
+						UAnimNextParameterBlockEntry* Binding = EditorData->FindBinding(InParameterBinding.Parameter);
+						if (Binding == nullptr)
+						{
+							const FAnimNextParamType Type = UE::AnimNext::UncookedOnly::FUtils::GetParameterTypeFromName(InParameterBinding.Parameter);
+							if (Type.IsValid())
+							{
+								Binding = EditorData->AddParameter(InParameterBinding.Parameter, Type);
+							};
+							PendingSelection.Add(Binding);	
+						}
+					});
+					
 					return SNew(SParameterPicker)
 						.Args(Args);
 				})
diff --git a/Engine/Plugins/Experimental/AnimNext/Source/AnimNextEditor/Private/Param/SParameterBlockView.h b/Engine/Plugins/Experimental/AnimNext/Source/AnimNextEditor/Private/Param/SParameterBlockView.h
--- a/Engine/Plugins/Experimental/AnimNext/Source/AnimNextEditor/Private/Param/SParameterBlockView.h
+++ b/Engine/Plugins/Experimental/AnimNext/Source/AnimNextEditor/Private/Param/SParameterBlockView.h
@@ -8,7 +8,6 @@
 #include "Templates/SharedPointer.h"
 
 class UAnimNextParameterBlockBinding;
-class UAnimNextParameter;
 class UAnimNextParameterBlock_EditorData;
 class UAnimNextParameterBlockEntry;
 class URigVMGraph;
diff --git a/Engine/Plugins/Experimental/AnimNext/Source/AnimNextEditor/Private/Param/SParameterLibraryView.cpp b/Engine/Plugins/Experimental/AnimNext/Source/AnimNextEditor/Private/Param/SParameterLibraryView.cpp
deleted file mode 100644
diff --git a/Engine/Plugins/Experimental/AnimNext/Source/AnimNextEditor/Private/Param/SParameterLibraryView.h b/Engine/Plugins/Experimental/AnimNext/Source/AnimNextEditor/Private/Param/SParameterLibraryView.h
deleted file mode 100644
diff --git a/Engine/Plugins/Experimental/AnimNext/Source/AnimNextEditor/Private/Param/SParameterPicker.cpp b/Engine/Plugins/Experimental/AnimNext/Source/AnimNextEditor/Private/Param/SParameterPicker.cpp
--- a/Engine/Plugins/Experimental/AnimNext/Source/AnimNextEditor/Private/Param/SParameterPicker.cpp
+++ b/Engine/Plugins/Experimental/AnimNext/Source/AnimNextEditor/Private/Param/SParameterPicker.cpp
@@ -7,9 +7,7 @@
 #include "AssetRegistry/AssetRegistryModule.h"
 #include "Kismet2/BlueprintEditorUtils.h"
 #include "Param/ParamType.h"
-#include "Param/AnimNextParameter.h"
 #include "Param/AnimNextParameterBlock_EditorData.h"
-#include "Param/AnimNextParameterLibrary.h"
 #include "DetailLayoutBuilder.h"
 #include "EditorUtils.h"
 #include "SAddParametersDialog.h"
@@ -32,7 +30,6 @@ namespace UE::AnimNext::Editor
 namespace ParameterPicker
 {
 static FName Column_Parameter(TEXT("Parameter"));
-static FName Column_Library(TEXT("Library"));
 static FName Column_Block(TEXT("Block"));
 static FName Column_Type(TEXT("Type"));
 static FName Column_New(TEXT("New"));
@@ -156,15 +153,13 @@ void SParameterPicker::Construct(const FArguments& InArgs)
 				if(ParametersToAdd.Num() > 0)
 				{
 					FScopedTransaction Transaction(LOCTEXT("AddParameter", "Add parameter"));
-
-					// Create a new parameter in the supplied library
-					UAnimNextParameterLibrary* Library = Cast<UAnimNextParameterLibrary>(ParametersToAdd[0].Library.GetAsset());
-					UAnimNextParameter* NewParameter = Library->AddParameter(ParametersToAdd[0].Name, ParametersToAdd[0].Type);
-							
-					FParameterBindingReference Reference;
-					Reference.Parameter = ParametersToAdd[0].Name;
-					Reference.Library = ParametersToAdd[0].Library;
-					Args.OnParameterPicked.ExecuteIfBound(Reference);
+					for (const FParameterToAdd& ParameterToAdd : ParametersToAdd)
+					{
+						FParameterBindingReference Reference;
+						Reference.Parameter = ParametersToAdd[0].Name;
+						Args.OnAddParameter.ExecuteIfBound(ParametersToAdd[0]);
+						Args.OnParameterPicked.ExecuteIfBound(Reference);
+					}
 				}
 			}
 			return FReply::Handled();
@@ -237,16 +232,8 @@ void SParameterPicker::Construct(const FArguments& InArgs)
 		]
 	];
 
-	if (Args.bShowLibraries)
-	{
-		HeaderRow->AddColumn(
-			SHeaderRow::Column(Column_Library)
-			.DefaultLabel(LOCTEXT("LibraryColumnHeader", "Library"))
-			.ToolTipText(LOCTEXT("LibraryColumnHeaderTooltip", "The library that the parameter is declared in"))
-			.FillWidth(0.33f));
-	}
 	
-	if (Args.bShowLibraries)
+	if (Args.bShowBlocks)
 	{
 		HeaderRow->AddColumn(
 			SHeaderRow::Column(Column_Block)
@@ -255,6 +242,52 @@ void SParameterPicker::Construct(const FArguments& InArgs)
 			.FillWidth(0.33f));
 	}
 
+	if(Args.bAllowNew)
+	{
+		HeaderRow->AddColumn(
+			SHeaderRow::Column(Column_New)
+			.DefaultLabel(FText::GetEmpty())
+			.HeaderContentPadding(FMargin(0.0f))
+			.FixedWidth(24.0f)
+			.HeaderContent()
+			[
+				SNew(SButton)
+				.ToolTipText(LOCTEXT("AddColumnHeaderTooltip", "Add a new parameter at global scope"))
+				.ButtonStyle(FAppStyle::Get(), "HoverHintOnly")
+				.OnClicked_Lambda([this]()
+				{
+					FSlateApplication::Get().DismissAllMenus();
+					TSharedRef<SAddParametersDialog> AddParametersDialog =
+						SNew(SAddParametersDialog)
+						.AllowMultiple(false);
+					TArray<FParameterToAdd> ParametersToAdd;
+					if(AddParametersDialog->ShowModal(ParametersToAdd))
+					{
+						if(ParametersToAdd.Num() > 0)
+						{
+							FScopedTransaction Transaction(LOCTEXT("AddParameter", "Add parameter"));
+
+							FParameterBindingReference Reference;
+							Reference.Parameter = ParametersToAdd[0].Name;
+							Args.OnParameterPicked.ExecuteIfBound(Reference);
+						}
+					}
+					return FReply::Handled();
+				})
+				[
+					SNew(SBox)
+					.WidthOverride(16.0f)
+					.HeightOverride(16.0f)
+					.VAlign(VAlign_Center)
+					.HAlign(HAlign_Center)
+					[
+						SNew(SImage)
+						.ColorAndOpacity(FSlateColor::UseForeground())
+						.Image(FAppStyle::GetBrush("Icons.Plus"))
+					]
+				]
+			]);
+	}
 	
 	RefreshEntries();
 }
@@ -269,21 +302,8 @@ void SParameterPicker::RefreshEntries()
 
 	TSet<TTuple<FName, FAssetData>> BoundParameters;
 
-	TMap<FAssetData, FAnimNextParameterLibraryAssetRegistryExports> LibraryExportMap;
-
-	auto GetExportsForLibrary = [&LibraryExportMap](const FAssetData& InLibrary) -> const FAnimNextParameterLibraryAssetRegistryExports&
-	{
-		if(FAnimNextParameterLibraryAssetRegistryExports* ExistingExports = LibraryExportMap.Find(InLibrary))
-		{
-			return *ExistingExports;
-		}
-
-		FAnimNextParameterLibraryAssetRegistryExports& NewExports = LibraryExportMap.Add(InLibrary);
-		UncookedOnly::FUtils::GetExportedParametersForLibrary(InLibrary, NewExports);
-		return NewExports;
-	};
-	
 	// Find all blocks and their bound parameters
+	if(Args.bShowBoundParameters)
 	{
 		ARFilter.ClassPaths = { UAnimNextParameterBlock::StaticClass()->GetClassPathName() };
 		
@@ -292,29 +312,22 @@ void SParameterPicker::RefreshEntries()
 
 		for(const FAssetData& BlockAsset : BlockAssets)
 		{
-			FAnimNextParameterBlockAssetRegistryExports Exports;
-			if(FUtils::GetExportedBindingsForBlock(BlockAsset, Exports))
+			FAnimNextParameterProviderAssetRegistryExports Exports;
+			if(UncookedOnly::FUtils::GetExportedParametersForAsset(BlockAsset, Exports))
 			{
-				for(const FAnimNextParameterBlockAssetRegistryExportEntry& Export : Exports.Bindings)
+				for(const FAnimNextParameterAssetRegistryExportEntry& Export : Exports.Parameters)
 				{
-					const FSoftObjectPath LibraryAssetPath(Export.Library);
-					FAssetData LibraryAsset = AssetRegistry.GetAssetByObjectPath(LibraryAssetPath);
-					if(LibraryAsset.IsValid())
+					BoundParameters.Add({ Export.Name, BlockAsset });
+					FParameterBindingReference NewReference(Export.Name, BlockAsset);
+					if(!Args.OnFilterParameter.IsBound() || Args.OnFilterParameter.Execute(NewReference) == EFilterParameterResult::Include)
 					{
-						BoundParameters.Add({ Export.Name, LibraryAsset });
-
-						FParameterBindingReference NewReference(Export.Name, LibraryAsset, BlockAsset);
-						if(!Args.OnFilterParameter.IsBound() || Args.OnFilterParameter.Execute(NewReference) == EFilterParameterResult::Include)
+						FAnimNextParamType ParamType = UE::AnimNext::UncookedOnly::FUtils::GetParameterTypeFromName(Export.Name);
+						if(!Args.OnFilterParameterType.IsBound() || Args.OnFilterParameterType.Execute(ParamType) == EFilterParameterResult::Include)
 						{
-							const FAnimNextParameterLibraryAssetRegistryExports& LibraryExports = GetExportsForLibrary(LibraryAsset);
-							FAnimNextParamType ParamType = FUtils::GetParameterTypeFromLibraryExports(Export.Name, LibraryExports);
-							if(!Args.OnFilterParameterType.IsBound() || Args.OnFilterParameterType.Execute(ParamType) == EFilterParameterResult::Include)
+							if (Args.bShowBoundParameters && EnumHasAnyFlags(Export.Flags, EAnimNextParameterFlags::Bound))
 							{
-								if (Args.bShowBoundParameters)
-								{
-									TSharedRef<FParameterPickerEntry> NewEntry = MakeShared<FParameterPickerEntry>(NewReference, ParamType);
-									Entries.Add(NewEntry);
-								}
+								TSharedRef<FParameterPickerEntry> NewEntry = MakeShared<FParameterPickerEntry>(NewReference, ParamType);
+								Entries.Add(NewEntry);
 							}
 						}
 					}
@@ -323,27 +336,25 @@ void SParameterPicker::RefreshEntries()
 		}
 	}
 
-	// Find all library parameters (that have not already been added as bound above)
+	// Find all parameters (that have not already been added as bound above)
 	if(Args.bShowUnboundParameters)
 	{
-		ARFilter.ClassPaths = { UAnimNextParameterLibrary::StaticClass()->GetClassPathName() };
-
-		TArray<FAssetData> LibraryAssets;
-		AssetRegistry.GetAssets(ARFilter, LibraryAssets);
-
-		for(const FAssetData& LibraryAsset : LibraryAssets)
+		FAnimNextParameterProviderAssetRegistryExports AllExports;
+		if(UE::AnimNext::UncookedOnly::FUtils::GetExportedParametersFromAssetRegistry(AllExports))
 		{
-			const FAnimNextParameterLibraryAssetRegistryExports& Exports = GetExportsForLibrary(LibraryAsset);
-			for(const FAnimNextParameterLibraryAssetRegistryExportEntry& Export : Exports.Parameters)
+			for(const FAnimNextParameterAssetRegistryExportEntry& ExportEntry : AllExports.Parameters)
 			{
-				if(!BoundParameters.Contains( { Export.Name, LibraryAsset } ))
+				if(!BoundParameters.Contains( { ExportEntry.Name, ExportEntry.ReferencingAsset } ))
 				{
-					FParameterBindingReference NewReference(Export.Name, LibraryAsset);
+					FParameterBindingReference NewReference;
+					NewReference.Asset = ExportEntry.ReferencingAsset;
+					NewReference.Parameter = ExportEntry.Name;
+
 					if(!Args.OnFilterParameter.IsBound() || Args.OnFilterParameter.Execute(NewReference) == EFilterParameterResult::Include)
 					{
-						if (!Args.OnFilterParameterType.IsBound() || Args.OnFilterParameterType.Execute(Export.Type) == EFilterParameterResult::Include)
+						if (!Args.OnFilterParameterType.IsBound() || Args.OnFilterParameterType.Execute(ExportEntry.Type) == EFilterParameterResult::Include)
 						{
-							TSharedRef<FParameterPickerEntry> NewEntry = MakeShared<FParameterPickerEntry>(NewReference, Export.Type);
+							TSharedRef<FParameterPickerEntry> NewEntry = MakeShared<FParameterPickerEntry>(NewReference, ExportEntry.Type);
 							Entries.Add(NewEntry);
 						}
 					}
@@ -555,6 +566,10 @@ class SParameterPickerRow : public SMultiColumnTableRow<TSharedRef<FParameterPic
 					SNew(SImage)
 					.Image(Entry->PinIcon)
 					.ColorAndOpacity(Entry->PinColor)
+					.ToolTipText_Lambda([this]() -> FText
+					{
+						return FText::FromString(Entry->ParamType.ToString());	
+					})
 				];
 		}
 		else if(InColumnName == Column_Parameter)
@@ -578,29 +593,13 @@ class SParameterPickerRow : public SMultiColumnTableRow<TSharedRef<FParameterPic
 					SNew(STextBlock)
 					.Font(IDetailLayoutBuilder::GetDetailFont())
 					.Text(FText::FromString(Entry->DisplayString))
-					.ToolTipText(UncookedOnly::FUtils::GetParameterDisplayNameText(Entry->Binding.Parameter))
+					.ToolTipText(FText::Format(LOCTEXT("ParameterPathTooltip", "{1}::{0}"), UncookedOnly::FUtils::GetParameterDisplayNameText(Entry->Binding.Parameter), FText::FromName(Entry->Binding.Asset.AssetName)))
 					.HighlightText_Lambda([this]()
 					{
 						return ParameterPicker.Pin()->FilterText;
 					})
 				];
 		}
-		else if(InColumnName == Column_Library)
-		{
-			if(Entry->Binding.Library.IsValid())
-			{
-				
-				return
-					SNew(SBox)
-					.VAlign(VAlign_Center)
-					[
-						SNew(STextBlock)
-						.Font(IDetailLayoutBuilder::GetDetailFont())
-						.Text(FText::FromName(Entry->Binding.Library.AssetName))
-						.ToolTipText(FText::FromName(Entry->Binding.Library.PackageName))
-					];
-			}
-		}
 		else if(InColumnName == Column_Block)
 		{
 			if(Entry->Binding.Block.IsValid())
diff --git a/Engine/Plugins/Experimental/AnimNext/Source/AnimNextEditor/Private/Workspace/AnimNextWorkspace.cpp b/Engine/Plugins/Experimental/AnimNext/Source/AnimNextEditor/Private/Workspace/AnimNextWorkspace.cpp
--- a/Engine/Plugins/Experimental/AnimNext/Source/AnimNextEditor/Private/Workspace/AnimNextWorkspace.cpp
+++ b/Engine/Plugins/Experimental/AnimNext/Source/AnimNextEditor/Private/Workspace/AnimNextWorkspace.cpp
@@ -4,7 +4,6 @@
 
 #include "Graph/AnimNextGraph.h"
 #include "Param/AnimNextParameterBlock.h"
-#include "Param/AnimNextParameterLibrary.h"
 #include "Scheduler/AnimNextSchedule.h"
 
 const FName UAnimNextWorkspace::ExportsAssetRegistryTag = TEXT("Exports");
@@ -178,7 +177,6 @@ bool UAnimNextWorkspace::RemoveAssets(TConstArrayView<FAssetData> InAssets, bool
 const TArray<FTopLevelAssetPath>& UAnimNextWorkspace::GetSupportedAssetClassPaths()
 {
 	static TArray<FTopLevelAssetPath> SupportedAssets = {
-		UAnimNextParameterLibrary::StaticClass()->GetClassPathName(),
 		UAnimNextParameterBlock::StaticClass()->GetClassPathName(),
 		UAnimNextGraph::StaticClass()->GetClassPathName(),
 		UAnimNextSchedule::StaticClass()->GetClassPathName(),
diff --git a/Engine/Plugins/Experimental/AnimNext/Source/AnimNextEditor/Private/Workspace/AnimNextWorkspaceEditor.cpp b/Engine/Plugins/Experimental/AnimNext/Source/AnimNextEditor/Private/Workspace/AnimNextWorkspaceEditor.cpp
--- a/Engine/Plugins/Experimental/AnimNext/Source/AnimNextEditor/Private/Workspace/AnimNextWorkspaceEditor.cpp
+++ b/Engine/Plugins/Experimental/AnimNext/Source/AnimNextEditor/Private/Workspace/AnimNextWorkspaceEditor.cpp
@@ -21,7 +21,6 @@
 #include "Graph/AnimNextGraphDocumentSummoner.h"
 #include "Param/AnimNextParameterBlock.h"
 #include "Param/AnimNextParameterBlock_EditorData.h"
-#include "Param/AnimNextParameterLibrary.h"
 #include "Param/ParameterBlockGraphDocumentSummoner.h"
 #include "Scheduler/AnimNextSchedule.h"
 #include "Framework/Application/SlateApplication.h"
@@ -81,7 +80,6 @@ void FWorkspaceEditor::InitEditor(const EToolkitMode::Type InMode, const TShared
 	{
 		UAnimNextGraph::StaticClass()->GetClassPathName(),
 		UAnimNextParameterBlock::StaticClass()->GetClassPathName(),
-		UAnimNextParameterLibrary::StaticClass()->GetClassPathName(),
 	});
 	LeftAssetDocumentSummoner->OnSaveDocumentState().BindSP(this, &FWorkspaceEditor::HandleSaveDocumentState);
 	DocumentManager->RegisterDocumentFactory(LeftAssetDocumentSummoner);
diff --git a/Engine/Plugins/Experimental/AnimNext/Source/AnimNextTestSuite/Private/AnimNextEditorTest.cpp b/Engine/Plugins/Experimental/AnimNext/Source/AnimNextTestSuite/Private/AnimNextEditorTest.cpp
--- a/Engine/Plugins/Experimental/AnimNext/Source/AnimNextTestSuite/Private/AnimNextEditorTest.cpp
+++ b/Engine/Plugins/Experimental/AnimNext/Source/AnimNextTestSuite/Private/AnimNextEditorTest.cpp
@@ -5,10 +5,9 @@
 #include "Param/AnimNextParameterBlock_EditorData.h"
 #include "Misc/AutomationTest.h"
 #include "Param/ParameterBlockFactory.h"
-#include "Param/ParameterLibraryFactory.h"
 #include "Param/AnimNextParameterBlock.h"
+#include "Param/AnimNextParameterBlockParameter.h"
 #include "Param/AnimNextParameterBlockBinding.h"
-#include "Param/AnimNextParameterLibrary.h"
 #include "Animation/AnimSequence.h"
 #if WITH_EDITOR
 #include "ScopedTransaction.h"
@@ -20,64 +19,13 @@
 
 #if WITH_DEV_AUTOMATION_TESTS && WITH_EDITOR
 
-IMPLEMENT_SIMPLE_AUTOMATION_TEST(FAnimationAnimNextParametersEditorTest_Library, "Animation.AnimNext.Parameters.Editor.Library", EAutomationTestFlags::EditorContext | EAutomationTestFlags::EngineFilter)
-
-bool FAnimationAnimNextParametersEditorTest_Library::RunTest(const FString& InParameters)
-{
-	using namespace UE::AnimNext;
-
-	UFactory* LibraryFactory = NewObject<UAnimNextParameterLibraryFactory>();
-	UAnimNextParameterLibrary* Library = Cast<UAnimNextParameterLibrary>(LibraryFactory->FactoryCreateNew(UAnimNextParameterLibrary::StaticClass(), GetTransientPackage(), TEXT("TestAnimNextParameterLibrary"), RF_Transient, nullptr, nullptr, NAME_None));
-	if(Library == nullptr)
-	{
-		AddError(TEXT("Could not create parameter library."));
-		return false;
-	}
-
-	AddErrorIfFalse(Library->AddParameter(TEXT("TestParam0"), FAnimNextParamType::GetType<bool>()) != nullptr, TEXT("Could not create new bool parameter in library."));
-	AddErrorIfFalse(Library->AddParameter(TEXT("TestParam1"), FAnimNextParamType::GetType<uint8>()) != nullptr, TEXT("Could not create new uint8 parameter in library."));
-	AddErrorIfFalse(Library->AddParameter(TEXT("TestParam2"), FAnimNextParamType::GetType<int32>()) != nullptr, TEXT("Could not create new int32 parameter in library."));
-	AddErrorIfFalse(Library->AddParameter(TEXT("TestParam3"), FAnimNextParamType::GetType<int64>()) != nullptr, TEXT("Could not create new int64 parameter in library."));
-	AddErrorIfFalse(Library->AddParameter(TEXT("TestParam4"), FAnimNextParamType::GetType<float>()) != nullptr, TEXT("Could not create new float parameter in library."));
-	AddErrorIfFalse(Library->AddParameter(TEXT("TestParam5"), FAnimNextParamType::GetType<double>()) != nullptr, TEXT("Could not create new double parameter in library."));
-	AddErrorIfFalse(Library->AddParameter(TEXT("TestParam6"), FAnimNextParamType::GetType<FName>()) != nullptr, TEXT("Could not create new FName parameter in library."));
-	AddErrorIfFalse(Library->AddParameter(TEXT("TestParam7"), FAnimNextParamType::GetType<FString>()) != nullptr, TEXT("Could not create new FString parameter in library."));
-	AddErrorIfFalse(Library->AddParameter(TEXT("TestParam8"), FAnimNextParamType::GetType<FText>()) != nullptr, TEXT("Could not create new FText parameter in library."));
-	AddErrorIfFalse(Library->AddParameter(TEXT("TestParam9"), FAnimNextParamType::GetType<EPropertyBagPropertyType>()) != nullptr, TEXT("Could not create new enum parameter in library."));
-	AddErrorIfFalse(Library->AddParameter(TEXT("TestParam10"), FAnimNextParamType::GetType<FVector>()) != nullptr, TEXT("Could not create new FVector parameter in library."));
-	AddErrorIfFalse(Library->AddParameter(TEXT("TestParam11"), FAnimNextParamType::GetType<FQuat>()) != nullptr, TEXT("Could not create new FQuat parameter in library."));
-	AddErrorIfFalse(Library->AddParameter(TEXT("TestParam12"), FAnimNextParamType::GetType<FTransform>()) != nullptr, TEXT("Could not create new FTransform parameter in library."));
-	AddErrorIfFalse(Library->AddParameter(TEXT("TestParam13"), FAnimNextParamType::GetType<TObjectPtr<UObject>>()) != nullptr, TEXT("Could not create new TObjectPtr<UObject> parameter in library."));
-	AddErrorIfFalse(Library->AddParameter(TEXT("TestParam14"), FAnimNextParamType::GetType<TObjectPtr<UAnimSequence>>()) != nullptr, TEXT("Could not create new TObjectPtr<UAnimSequence> parameter in library."));
-	AddErrorIfFalse(Library->AddParameter(TEXT("TestParam15"), FAnimNextParamType::GetType<TArray<float>>()) != nullptr, TEXT("Could not create new TArray<float> parameter in library."));
-	AddErrorIfFalse(Library->AddParameter(TEXT("TestParam16"), FAnimNextParamType::GetType<TArray<TObjectPtr<UAnimSequence>>>()) != nullptr, TEXT("Could not create new TArray<TObjectPtr<UAnimSequence>> parameter in library."));
-
-	AddErrorIfFalse(Library->RemoveParameter(TEXT("TestParam0")), TEXT("Could not remove parameter from library."));
-	AddErrorIfFalse(Library->RemoveParameters({ TEXT("TestParam1"), TEXT("TestParam2") }), TEXT("Could not remove parameters from library."));
-	AddErrorIfFalse(Library->FindParameter(TEXT("TestParam3")) != nullptr, TEXT("Could not find parameter in library."));
-	
-	CollectGarbage(GARBAGE_COLLECTION_KEEPFLAGS);
-
-	return true;
-}
-
 IMPLEMENT_SIMPLE_AUTOMATION_TEST(FAnimationAnimNextParametersEditorTest_Block, "Animation.AnimNext.Parameters.Editor.Block", EAutomationTestFlags::EditorContext | EAutomationTestFlags::EngineFilter)
 
 bool FAnimationAnimNextParametersEditorTest_Block::RunTest(const FString& InParameters)
 {
 	using namespace UE::AnimNext;
 
-	UFactory* LibraryFactory = NewObject<UAnimNextParameterLibraryFactory>();
-	UAnimNextParameterLibrary* Library = Cast<UAnimNextParameterLibrary>(LibraryFactory->FactoryCreateNew(UAnimNextParameterLibrary::StaticClass(), GetTransientPackage(), TEXT("TestAnimNextParameterLibrary"), RF_Transient, nullptr, nullptr, NAME_None));
-	if(Library == nullptr)
-	{
-		AddError(TEXT("Could not create parameter library."));
-		return false;
-	}
-
-	AddErrorIfFalse(Library->AddParameter(TEXT("TestParam"), FAnimNextParamType::GetType<bool>()) != nullptr, TEXT("Could not create new parameter in library."));
-
-	TStrongObjectPtr<UFactory> BlockFactory(NewObject<UAnimNextParameterBlockFactory>());
+	const TStrongObjectPtr<UFactory> BlockFactory(NewObject<UAnimNextParameterBlockFactory>());
 	UAnimNextParameterBlock* Block = Cast<UAnimNextParameterBlock>(BlockFactory->FactoryCreateNew(UAnimNextParameterBlock::StaticClass(), GetTransientPackage(), TEXT("TestAnimNextParameterBlock"), RF_Transient, nullptr, nullptr, NAME_None));
 	if(Block == nullptr)
 	{
@@ -92,11 +40,70 @@ bool FAnimationAnimNextParametersEditorTest_Block::RunTest(const FString& InPara
 		return false;
 	}
 
+	static FName TestParameterName = TEXT("TestParam");
+	
+	// AddParameter
+	{
+		FScopedTransaction Transaction(FText::GetEmpty());		
+		UAnimNextParameterBlockParameter* Parameter = EditorData->AddParameter(TestParameterName, FAnimNextParamType::GetType<bool>());
+
+		if (AddErrorIfFalse(Parameter != nullptr, TEXT("Could not create new parameter in block.")))
+		{
+			AddErrorIfFalse(Parameter->GetParamType() == FAnimNextParamType::GetType<bool>(), TEXT("Incorrect parameter type found"));			
+		}
+	}
+		
+	AddExpectedError(TEXT("UAnimNextParameterBlock_EditorData::AddParameter: A parameter already exists for the supplied parameter name."));
+	AddErrorIfFalse(EditorData->AddParameter(TestParameterName, FAnimNextParamType::GetType<bool>()) == nullptr, TEXT("Expected duplicate parameter name argument to fail"));
+
+	GEditor->UndoTransaction();
+	AddErrorIfFalse(EditorData->Entries.Num() == 0, FString::Printf(TEXT("Unexpected entry count found in parameter block (Have %d, expected 0)."), EditorData->Entries.Num()));
+
+	GEditor->RedoTransaction();
+	AddErrorIfFalse(EditorData->Entries.Num() == 1, FString::Printf(TEXT("Unexpected entry count found in parameter block (Have %d, expected 1)."), EditorData->Entries.Num()));
+	
+	GEditor->UndoTransaction();
+	
+	// Failure cases
+	AddExpectedError(TEXT("UAnimNextParameterBlock_EditorData::AddParameter: Invalid parameter name supplied."));
+	AddErrorIfFalse(EditorData->AddParameter(NAME_None, FAnimNextParamType::GetType<bool>()) == nullptr, TEXT("Expected invalid argument to fail"));
+
+	auto TestParameterType = [this, EditorData](FAnimNextParamType InType)
+	{
+		UAnimNextParameterBlockParameter* TypedParameter = EditorData->AddParameter(TEXT("TestParam0"), InType);
+		const bool bValidParameter = TypedParameter != nullptr;
+		if (bValidParameter && AddErrorIfFalse(bValidParameter, FString::Printf(TEXT("Could not create new parameter of type %s in block."), *InType.ToString())))
+		{
+			AddErrorIfFalse(TypedParameter->GetParamType() == InType, TEXT("Incorrect parameter type found"));
+			EditorData->RemoveEntry(TypedParameter);
+		}		
+	};
+	
+	// Various types
+	TestParameterType(FAnimNextParamType::GetType<bool>());
+	TestParameterType(FAnimNextParamType::GetType<uint8>());
+	TestParameterType(FAnimNextParamType::GetType<int32>());
+	TestParameterType(FAnimNextParamType::GetType<int64>());
+	TestParameterType(FAnimNextParamType::GetType<float>());
+	TestParameterType(FAnimNextParamType::GetType<double>());
+	TestParameterType(FAnimNextParamType::GetType<FName>());
+	TestParameterType(FAnimNextParamType::GetType<FString>());
+	TestParameterType(FAnimNextParamType::GetType<FText>());
+	TestParameterType(FAnimNextParamType::GetType<EPropertyBagPropertyType>());
+	TestParameterType(FAnimNextParamType::GetType<FVector>());
+	TestParameterType(FAnimNextParamType::GetType<FQuat>());
+	TestParameterType(FAnimNextParamType::GetType<FTransform>());
+	TestParameterType(FAnimNextParamType::GetType<TObjectPtr<UObject>>());
+	TestParameterType(FAnimNextParamType::GetType<TObjectPtr<UAnimSequence>>());
+	TestParameterType(FAnimNextParamType::GetType<TArray<float>>());
+	TestParameterType(FAnimNextParamType::GetType<TArray<TObjectPtr<UAnimSequence>>>());
+
 	// AddBinding
+	static FName TestBindingName = TEXT("TestParamBinding");
 	UAnimNextParameterBlockBinding* Binding = nullptr;
 	{
 		FScopedTransaction Transaction(FText::GetEmpty());
-		Binding = EditorData->AddBinding(TEXT("TestParam"), Library);
+		Binding = EditorData->AddBinding(TestBindingName, FAnimNextParamType::GetType<int32>());
 		AddErrorIfFalse(Binding != nullptr, TEXT("Could not create new binding in block."));
 	}
 
@@ -104,12 +111,11 @@ bool FAnimationAnimNextParametersEditorTest_Block::RunTest(const FString& InPara
 	AddErrorIfFalse(EditorData->Entries.Num() == 0, FString::Printf(TEXT("Unexpected entry count found in parameter block (Have %d, expected 0)."), EditorData->Entries.Num()));
 
 	GEditor->RedoTransaction();
-	AddErrorIfFalse(EditorData->Entries.Num() == 1, FString::Printf(TEXT("Unexpected entry count found in parameter block (Have %d, expected 1)."), EditorData->Entries.Num()));
 
 	// RemoveAllBindings
 	{
 		FScopedTransaction Transaction(FText::GetEmpty());
-		AddErrorIfFalse(EditorData->RemoveAllBindings(TEXT("TestParam")), TEXT("Failed to remove binding."));
+		AddErrorIfFalse(EditorData->RemoveAllBindings(TestBindingName), TEXT("Failed to remove binding."));
 	}
 
 	GEditor->UndoTransaction();
@@ -137,11 +143,9 @@ bool FAnimationAnimNextParametersEditorTest_Block::RunTest(const FString& InPara
 	AddErrorIfFalse(EditorData->Entries.Num() == 1, FString::Printf(TEXT("Unexpected entry count found in parameter block (Have %d, expected 1)."), EditorData->Entries.Num()));
 
 	// FindBinding
-	AddErrorIfFalse(EditorData->FindBinding(TEXT("TestParam")) != nullptr, TEXT("Could not find binding in block."));
-
-	GEditor->RedoTransaction();
+	AddErrorIfFalse(EditorData->FindBinding(TestBindingName) != nullptr, TEXT("Could not find binding in block."));
 
-	UAnimNextParameterBlock* OtherBlock = Cast<UAnimNextParameterBlock>(BlockFactory->FactoryCreateNew(UAnimNextParameterBlock::StaticClass(), GetTransientPackage(), TEXT("TestAnimNextParameterBlock2"), RF_Transient, nullptr, nullptr, NAME_None));
+	const UAnimNextParameterBlock* OtherBlock = Cast<UAnimNextParameterBlock>(BlockFactory->FactoryCreateNew(UAnimNextParameterBlock::StaticClass(), GetTransientPackage(), TEXT("TestAnimNextParameterBlock2"), RF_Transient, nullptr, nullptr, NAME_None));
 	if(OtherBlock == nullptr)
 	{
 		AddError(TEXT("Could not create additional parameter block."));
@@ -158,7 +162,7 @@ bool FAnimationAnimNextParametersEditorTest_Block::RunTest(const FString& InPara
 	// AddBindingReference
 	{
 		FScopedTransaction Transaction(FText::GetEmpty());
-		UAnimNextParameterBlockBindingReference* Ref = OtherEditorData->AddBindingReference(TEXT("TestParam"), Library, Block);
+		UAnimNextParameterBlockBindingReference* Ref = OtherEditorData->AddBindingReference(TestBindingName, Block);
 		AddErrorIfFalse(Ref != nullptr, TEXT("Could not create new binding reference in block."));
 	}
 
@@ -169,15 +173,9 @@ bool FAnimationAnimNextParametersEditorTest_Block::RunTest(const FString& InPara
 	AddErrorIfFalse(OtherEditorData->Entries.Num() == 1, FString::Printf(TEXT("Unexpected entry count found in parameter block (Have %d, expected 1)."), OtherEditorData->Entries.Num()));
 
 	// FindBinding
-	AddErrorIfFalse(OtherEditorData->FindBinding(TEXT("TestParam")) != nullptr, TEXT("Could not find binding refernce in block."));
-
-	UAnimNextParameterBlockParameter* Parameter = nullptr;
-	{
-		FScopedTransaction Transaction(FText::GetEmpty());
-		Parameter = EditorData->AddParameter(TEXT("TestParam"), Library);
-		AddErrorIfFalse(Parameter != nullptr, TEXT("Could not create new parameter in block."));
-	}
+	AddErrorIfFalse(OtherEditorData->FindBinding(TestBindingName) != nullptr, TEXT("Could not find binding refernce in block."));
 
+	GEditor->UndoTransaction();
 	GEditor->UndoTransaction();
 	AddErrorIfFalse(EditorData->Entries.Num() == 0, FString::Printf(TEXT("Unexpected entry count found in parameter block (Have %d, expected 0)."), EditorData->Entries.Num()));
 
@@ -212,15 +210,12 @@ bool FAnimationAnimNextParametersEditorTest_Python::RunTest(const FString& InPar
 
 	const TCHAR* Script = TEXT(
 		"asset_tools = unreal.AssetToolsHelpers.get_asset_tools()\n"
-		"library = unreal.AssetTools.create_asset(asset_tools, asset_name = \"TestLibrary\", package_path = \"/Game/\", asset_class = unreal.AnimNextParameterLibrary, factory = unreal.AnimNextParameterLibraryFactory())\n"
 		"block = unreal.AssetTools.create_asset(asset_tools, asset_name = \"TestBlock\", package_path = \"/Game/\", asset_class = unreal.AnimNextParameterBlock, factory = unreal.AnimNextParameterBlockFactory())\n"
 		"other_block = unreal.AssetTools.create_asset(asset_tools, asset_name = \"TestBlock1\", package_path = \"/Game/\", asset_class = unreal.AnimNextParameterBlock, factory = unreal.AnimNextParameterBlockFactory())\n"
-		"library.add_parameter(name = \"TestParam\", value_type = unreal.PropertyBagPropertyType.BOOL, container_type = unreal.PropertyBagContainerType.NONE)\n"
-		"block.add_binding(name = \"TestParam\", library = library)\n"
-		"other_block.add_binding_reference(name = \"TestParam\", library = library, referenced_block = block)\n"
-		"block.add_parameter(name = \"TestParam\", library = library)\n"
+		"block.add_parameter(name = \"TestParam\", value_type = unreal.PropertyBagPropertyType.BOOL, container_type = unreal.PropertyBagContainerType.NONE)\n"
+		"block.add_binding(name = \"TestParamBinding\", value_type = unreal.PropertyBagPropertyType.BOOL, container_type = unreal.PropertyBagContainerType.NONE)\n"
+		"other_block.add_binding_reference(name = \"TestParamBinding\", referenced_block = block)\n"
 		"block.add_graph(name = \"TestGraph\")\n"
-		"unreal.EditorAssetLibrary.delete_loaded_asset(library)\n"
 		"unreal.EditorAssetLibrary.delete_loaded_asset(block)\n"
 		"unreal.EditorAssetLibrary.delete_loaded_asset(other_block)\n"
 	);
diff --git a/Engine/Plugins/Experimental/AnimNext/Source/AnimNextUncookedOnly/Internal/Graph/AnimNextGraph_EditorData.h b/Engine/Plugins/Experimental/AnimNext/Source/AnimNextUncookedOnly/Internal/Graph/AnimNextGraph_EditorData.h
--- a/Engine/Plugins/Experimental/AnimNext/Source/AnimNextUncookedOnly/Internal/Graph/AnimNextGraph_EditorData.h
+++ b/Engine/Plugins/Experimental/AnimNext/Source/AnimNextUncookedOnly/Internal/Graph/AnimNextGraph_EditorData.h
@@ -97,6 +97,7 @@ public:
 	virtual void Serialize(FArchive& Ar) override;
 	virtual void PostLoad() override;
 	virtual bool IsEditorOnly() const override { return true; }
+	virtual void GetAssetRegistryTags(TArray<FAssetRegistryTag>& OutTags) const override;
 #if WITH_EDITOR
 	void HandlePackageDone(const FEndLoadPackageContext& Context);
 	void HandlePackageDone();
diff --git a/Engine/Plugins/Experimental/AnimNext/Source/AnimNextUncookedOnly/Internal/Param/AnimNextParameter.h b/Engine/Plugins/Experimental/AnimNext/Source/AnimNextUncookedOnly/Internal/Param/AnimNextParameter.h
deleted file mode 100644
diff --git a/Engine/Plugins/Experimental/AnimNext/Source/AnimNextUncookedOnly/Internal/Param/AnimNextParameterBlockBinding.h b/Engine/Plugins/Experimental/AnimNext/Source/AnimNextUncookedOnly/Internal/Param/AnimNextParameterBlockBinding.h
--- a/Engine/Plugins/Experimental/AnimNext/Source/AnimNextUncookedOnly/Internal/Param/AnimNextParameterBlockBinding.h
+++ b/Engine/Plugins/Experimental/AnimNext/Source/AnimNextUncookedOnly/Internal/Param/AnimNextParameterBlockBinding.h
@@ -9,7 +9,6 @@
 #include "AnimNextParameterBlockBinding.generated.h"
 
 class UAssetDefinition_AnimNextParameterBlockBinding;
-class UAnimNextParameter;
 class UAnimNextParameterLibrary;
 class URigVMGraph;
 
@@ -31,9 +30,8 @@ class UAnimNextParameterBlockBinding : public UAnimNextParameterBlockEntry, publ
 	// IAnimNextParameterBlockParameterInterface interface
 	virtual FAnimNextParamType GetParamType() const override;
 	virtual void SetParameterName(FName InName, bool bSetupUndoRedo = true) override;
-	virtual FName GetParameterName() const override;
-	virtual const UAnimNextParameter* GetParameter() const override;
-	virtual const UAnimNextParameterLibrary* GetLibrary() const override;
+	virtual FName GetParameterName() const override;	
+	virtual bool SetParamType(const FAnimNextParamType& InType, bool bSetupUndoRedo = true) override;
 
 	// IAnimNextParameterBlockGraphInterface interface
 	virtual URigVMGraph* GetGraph() const override { return BindingGraph; }
@@ -43,15 +41,14 @@ class UAnimNextParameterBlockBinding : public UAnimNextParameterBlockEntry, publ
 	// UAnimNextParameterBlockEntry interface
 	virtual FText GetDisplayName() const override;
 	virtual FText GetDisplayNameTooltip() const override;
-	virtual void GetEditedObjects(TArray<UObject*>& OutObjects) const override;
 
 	/** Parameter name we reference */
 	UPROPERTY(VisibleAnywhere, Category = Parameter)
 	FName ParameterName;
 
-	/** Parameter library we reference */
-	UPROPERTY(VisibleAnywhere, Category = Parameter)
-	TObjectPtr<UAnimNextParameterLibrary> Library;
+	/** The parameter's type */
+	UPROPERTY(EditAnywhere, Category = "Parameter", AssetRegistrySearchable)
+	FAnimNextParamType Type = FAnimNextParamType::GetType<bool>();
 
 	/** Binding graph */
 	UPROPERTY()
diff --git a/Engine/Plugins/Experimental/AnimNext/Source/AnimNextUncookedOnly/Internal/Param/AnimNextParameterBlockBindingReference.h b/Engine/Plugins/Experimental/AnimNext/Source/AnimNextUncookedOnly/Internal/Param/AnimNextParameterBlockBindingReference.h
--- a/Engine/Plugins/Experimental/AnimNext/Source/AnimNextUncookedOnly/Internal/Param/AnimNextParameterBlockBindingReference.h
+++ b/Engine/Plugins/Experimental/AnimNext/Source/AnimNextUncookedOnly/Internal/Param/AnimNextParameterBlockBindingReference.h
@@ -8,7 +8,6 @@
 #include "IAnimNextParameterBlockParameterInterface.h"
 #include "AnimNextParameterBlockBindingReference.generated.h"
 
-class UAnimNextParameter;
 class UAnimNextParameterBlock;
 class UAnimNextParameterLibrary;
 class UAssetDefinition_AnimNextParameterBlockBindingReference;
@@ -31,9 +30,8 @@ class UAnimNextParameterBlockBindingReference : public UAnimNextParameterBlockEn
 	// IAnimNextParameterBlockParameterInterface interface
 	virtual FAnimNextParamType GetParamType() const override;
 	virtual FName GetParameterName() const override;
-	virtual void SetParameterName(FName InName, bool bSetupUndoRedo = true) override;
-	virtual const UAnimNextParameter* GetParameter() const override;
-	virtual const UAnimNextParameterLibrary* GetLibrary() const override;
+	virtual void SetParameterName(FName InName, bool bSetupUndoRedo = true) override;	
+	virtual bool SetParamType(const FAnimNextParamType& InType, bool bSetupUndoRedo = true) override { return false; }
 
 	// IAnimNextParameterBlockReferenceInterface interface
 	virtual const UAnimNextParameterBlock* GetBlock() const override;
@@ -47,10 +45,6 @@ class UAnimNextParameterBlockBindingReference : public UAnimNextParameterBlockEn
 	UPROPERTY(VisibleAnywhere, Category = Parameter)
 	FName ParameterName;
 
-	/** Parameter library we reference */
-	UPROPERTY(VisibleAnywhere, Category = Parameter)
-	TObjectPtr<UAnimNextParameterLibrary> Library;
-
 	/** Parameter block we reference */
 	UPROPERTY()
 	TObjectPtr<UAnimNextParameterBlock> Block;
diff --git a/Engine/Plugins/Experimental/AnimNext/Source/AnimNextUncookedOnly/Internal/Param/AnimNextParameterBlockParameter.h b/Engine/Plugins/Experimental/AnimNext/Source/AnimNextUncookedOnly/Internal/Param/AnimNextParameterBlockParameter.h
--- a/Engine/Plugins/Experimental/AnimNext/Source/AnimNextUncookedOnly/Internal/Param/AnimNextParameterBlockParameter.h
+++ b/Engine/Plugins/Experimental/AnimNext/Source/AnimNextUncookedOnly/Internal/Param/AnimNextParameterBlockParameter.h
@@ -21,6 +21,7 @@ class UAnimNextParameterBlockParameter : public UAnimNextParameterBlockEntry, pu
 	GENERATED_BODY()
 
 	friend class UAnimNextParameterBlock_EditorData;
+	friend class FAnimationAnimNextParametersEditorTest_Block;
 	friend class UE::AnimNext::Editor::FParameterBlockParameterCustomization;
 	
 	// UAnimNextParameterBlockEntry interface
@@ -31,14 +32,17 @@ class UAnimNextParameterBlockParameter : public UAnimNextParameterBlockEntry, pu
 	virtual FAnimNextParamType GetParamType() const override;
 	virtual void SetParameterName(FName InName, bool bSetupUndoRedo = true) override;
 	virtual FName GetParameterName() const override;
-	virtual const UAnimNextParameter* GetParameter() const override;
-	virtual const UAnimNextParameterLibrary* GetLibrary() const override;
+	virtual bool SetParamType(const FAnimNextParamType& InType, bool bSetupUndoRedo = true) override;
 	
 	/** Parameter name we reference */
 	UPROPERTY(VisibleAnywhere, Category = Parameter)
 	FName ParameterName;
 
-	/** Parameter library we reference */
-	UPROPERTY(VisibleAnywhere, Category = Parameter)
-	TObjectPtr<UAnimNextParameterLibrary> Library;
+	/** The parameter's type */
+	UPROPERTY(EditAnywhere, Category = "Parameter", AssetRegistrySearchable)
+	FAnimNextParamType Type = FAnimNextParamType::GetType<bool>();
+	
+	/** Comment to display in editor */
+	UPROPERTY(EditAnywhere, Category = "Parameter", meta=(MultiLine))
+	FString Comment;
 };
\ No newline at end of file
diff --git a/Engine/Plugins/Experimental/AnimNext/Source/AnimNextUncookedOnly/Internal/Param/AnimNextParameterBlock_EditorData.h b/Engine/Plugins/Experimental/AnimNext/Source/AnimNextUncookedOnly/Internal/Param/AnimNextParameterBlock_EditorData.h
--- a/Engine/Plugins/Experimental/AnimNext/Source/AnimNextUncookedOnly/Internal/Param/AnimNextParameterBlock_EditorData.h
+++ b/Engine/Plugins/Experimental/AnimNext/Source/AnimNextUncookedOnly/Internal/Param/AnimNextParameterBlock_EditorData.h
@@ -8,15 +8,15 @@
 #include "RigVMCore/RigVMGraphFunctionHost.h"
 #include "RigVMBlueprint.h"
 #include "Kismet/BlueprintFunctionLibrary.h"
+#include "Param/ParamType.h"
 #include "AnimNextParameterBlock_EditorData.generated.h"
 
 class UAnimNextParameterBlock;
 class UAnimNextParameterBlockEntry;
 class UAnimNextParameterBlockParameter;
-class UAnimNextParameter;
 class UAnimNextParameterBlock_Controller;
 enum class ERigVMGraphNotifType : uint8;
-class FAnimationAnimNextParametersEditorTest_Block;;
+class FAnimationAnimNextParametersEditorTest_Block;
 
 namespace UE::AnimNext::UncookedOnly
 {
@@ -60,16 +60,16 @@ class ANIMNEXTUNCOOKEDONLY_API UAnimNextParameterBlockLibrary : public UBlueprin
 	GENERATED_BODY()
 
 	UFUNCTION(BlueprintCallable, Category = "AnimNext|Parameter Block", meta=(ScriptMethod))
-	static UAnimNextParameterBlockParameter* AddParameter(UAnimNextParameterBlock* InBlock, FName InName, UAnimNextParameterLibrary* InLibrary, bool bSetupUndoRedo = true, bool bPrintPythonCommand = true);
+	static UAnimNextParameterBlockParameter* AddParameter(UAnimNextParameterBlock* InBlock, FName InName, EPropertyBagPropertyType InValueType, EPropertyBagContainerType InContainerType = EPropertyBagContainerType::None, const UObject* InValueTypeObject = nullptr, bool bSetupUndoRedo = true, bool bPrintPythonCommand = true);
 
 	UFUNCTION(BlueprintCallable, Category = "AnimNext|Parameter Block", meta=(ScriptMethod))
 	static UAnimNextParameterBlockGraph* AddGraph(UAnimNextParameterBlock* InBlock, FName InName, bool bSetupUndoRedo = true, bool bPrintPythonCommand = true);
 	
 	UFUNCTION(BlueprintCallable, Category = "AnimNext|Parameter Block", meta=(ScriptMethod))
-	static UAnimNextParameterBlockBinding* AddBinding(UAnimNextParameterBlock* InBlock, FName InName, UAnimNextParameterLibrary* InLibrary, bool bSetupUndoRedo = true, bool bPrintPythonCommand = true);
+	static UAnimNextParameterBlockBinding* AddBinding(UAnimNextParameterBlock* InBlock, FName InName, EPropertyBagPropertyType InValueType, EPropertyBagContainerType InContainerType = EPropertyBagContainerType::None, const UObject* InValueTypeObject = nullptr, bool bSetupUndoRedo = true, bool bPrintPythonCommand = true);
 
 	UFUNCTION(BlueprintCallable, Category = "AnimNext|Parameter Block", meta=(ScriptMethod))
-	static UAnimNextParameterBlockBindingReference* AddBindingReference(UAnimNextParameterBlock* InBlock, FName InName, UAnimNextParameterLibrary* InLibrary, UAnimNextParameterBlock* InReferencedBlock, bool bSetupUndoRedo = true, bool bPrintPythonCommand = true);
+	static UAnimNextParameterBlockBindingReference* AddBindingReference(UAnimNextParameterBlock* InBlock, FName InName, UAnimNextParameterBlock* InReferencedBlock, bool bSetupUndoRedo = true, bool bPrintPythonCommand = true);
 
 	UFUNCTION(BlueprintCallable, Category = "AnimNext|Parameter Block", meta=(ScriptMethod))
 	static bool RemoveAllBindings(UAnimNextParameterBlock* InBlock, FName InName, bool bSetupUndoRedo = true, bool bPrintPythonCommand = true);
@@ -98,7 +98,8 @@ class UAnimNextParameterBlock_EditorData : public UObject, public IRigVMClientHo
 	friend class UAnimNextParameterBlockFactory;
 	friend class UAnimNextParameterBlockEntry;
 	friend class UAnimNextParameterBlock_EdGraph;
-	friend class UAnimNextParameterBlock_Controller;
+	friend class UAnimNextParameterBlockParameter;
+	friend class UAnimNextParameterBlockBindingReference;
 	friend struct UE::AnimNext::Editor::FUtils;
 	friend struct UE::AnimNext::UncookedOnly::FUtils;
 	friend struct UE::AnimNext::UncookedOnly::FUtilsPrivate;
@@ -112,13 +113,13 @@ class UAnimNextParameterBlock_EditorData : public UObject, public IRigVMClientHo
 	friend class UE::AnimNext::Editor::SParameterBlockViewRow;
 	friend class UE::AnimNext::Editor::FParameterBlockParameterCustomization;
 
-	ANIMNEXTUNCOOKEDONLY_API UAnimNextParameterBlockParameter* AddParameter(FName InName, UAnimNextParameterLibrary* InLibrary, bool bSetupUndoRedo = true, bool bPrintPythonCommand = true);
+	ANIMNEXTUNCOOKEDONLY_API UAnimNextParameterBlockParameter* AddParameter(FName InName, FAnimNextParamType InType, bool bSetupUndoRedo = true, bool bPrintPythonCommand = true);
 
 	ANIMNEXTUNCOOKEDONLY_API UAnimNextParameterBlockGraph* AddGraph(FName InName, bool bSetupUndoRedo = true, bool bPrintPythonCommand = true);
 	
-	ANIMNEXTUNCOOKEDONLY_API UAnimNextParameterBlockBinding* AddBinding(FName InName, UAnimNextParameterLibrary* InLibrary, bool bSetupUndoRedo = true, bool bPrintPythonCommand = true);
+	ANIMNEXTUNCOOKEDONLY_API UAnimNextParameterBlockBinding* AddBinding(FName InName, FAnimNextParamType InType, bool bSetupUndoRedo = true, bool bPrintPythonCommand = true);
 
-	ANIMNEXTUNCOOKEDONLY_API UAnimNextParameterBlockBindingReference* AddBindingReference(FName InName, UAnimNextParameterLibrary* InLibrary, UAnimNextParameterBlock* InBlock, bool bSetupUndoRedo = true, bool bPrintPythonCommand = true);
+	ANIMNEXTUNCOOKEDONLY_API UAnimNextParameterBlockBindingReference* AddBindingReference(FName InName, UAnimNextParameterBlock* InBlock, bool bSetupUndoRedo = true, bool bPrintPythonCommand = true);
 
 	ANIMNEXTUNCOOKEDONLY_API bool RemoveAllBindings(FName InName, bool bSetupUndoRedo = true, bool bPrintPythonCommand = true);
 
diff --git a/Engine/Plugins/Experimental/AnimNext/Source/AnimNextUncookedOnly/Internal/Param/AnimNextParameterLibrary.h b/Engine/Plugins/Experimental/AnimNext/Source/AnimNextUncookedOnly/Internal/Param/AnimNextParameterLibrary.h
deleted file mode 100644
diff --git a/Engine/Plugins/Experimental/AnimNext/Source/AnimNextUncookedOnly/Internal/Param/IAnimNextParameterBlockParameterInterface.h b/Engine/Plugins/Experimental/AnimNext/Source/AnimNextUncookedOnly/Internal/Param/IAnimNextParameterBlockParameterInterface.h
--- a/Engine/Plugins/Experimental/AnimNext/Source/AnimNextUncookedOnly/Internal/Param/IAnimNextParameterBlockParameterInterface.h
+++ b/Engine/Plugins/Experimental/AnimNext/Source/AnimNextUncookedOnly/Internal/Param/IAnimNextParameterBlockParameterInterface.h
@@ -6,7 +6,6 @@
 #include "UObject/Interface.h"
 #include "IAnimNextParameterBlockParameterInterface.generated.h"
 
-class UAnimNextParameter;
 class UAnimNextParameterLibrary;
 class UAnimNextParameterBlock_EditorData;
 
@@ -38,6 +37,7 @@ class ANIMNEXTUNCOOKEDONLY_API IAnimNextParameterBlockParameterInterface
 	GENERATED_BODY()
 
 	friend class UAnimNextParameterBlock_EditorData;
+	friend class UAnimNextParameterBlockBindingReference;
 	friend class UE::AnimNext::Editor::SParameterBlockView;
 	friend class UE::AnimNext::Editor::SParameterBlockViewRow;
 	friend struct UE::AnimNext::UncookedOnly::FUtils;
@@ -46,15 +46,12 @@ class ANIMNEXTUNCOOKEDONLY_API IAnimNextParameterBlockParameterInterface
 	// Get the parameter type
 	virtual FAnimNextParamType GetParamType() const = 0;
 
+	// Set the parameter type
+	virtual bool SetParamType(const FAnimNextParamType& InType, bool bSetupUndoRedo = true) = 0;
+
 	// Get the parameter name
 	virtual FName GetParameterName() const = 0;
 
 	// Set the parameter name
 	virtual void SetParameterName(FName InName, bool bSetupUndoRedo = true) = 0;
-
-	// Resolve the underlying parameter in the library
-	virtual const UAnimNextParameter* GetParameter() const = 0;
-
-	// Get the parameter library hosting the parameter
-	virtual const UAnimNextParameterLibrary* GetLibrary() const = 0;
 };
\ No newline at end of file
diff --git a/Engine/Plugins/Experimental/AnimNext/Source/AnimNextUncookedOnly/Internal/Param/IAnimNextParameterBlockReferenceInterface.h b/Engine/Plugins/Experimental/AnimNext/Source/AnimNextUncookedOnly/Internal/Param/IAnimNextParameterBlockReferenceInterface.h
--- a/Engine/Plugins/Experimental/AnimNext/Source/AnimNextUncookedOnly/Internal/Param/IAnimNextParameterBlockReferenceInterface.h
+++ b/Engine/Plugins/Experimental/AnimNext/Source/AnimNextUncookedOnly/Internal/Param/IAnimNextParameterBlockReferenceInterface.h
@@ -7,7 +7,6 @@
 #include "IAnimNextParameterBlockReferenceInterface.generated.h"
 
 class URigVMGraph;
-class UAnimNextParameter;
 class UAnimNextParameterBlock;
 
 namespace UE::AnimNext::Editor
diff --git a/Engine/Plugins/Experimental/AnimNext/Source/AnimNextUncookedOnly/Internal/UncookedOnlyUtils.h b/Engine/Plugins/Experimental/AnimNext/Source/AnimNextUncookedOnly/Internal/UncookedOnlyUtils.h
--- a/Engine/Plugins/Experimental/AnimNext/Source/AnimNextUncookedOnly/Internal/UncookedOnlyUtils.h
+++ b/Engine/Plugins/Experimental/AnimNext/Source/AnimNextUncookedOnly/Internal/UncookedOnlyUtils.h
@@ -8,18 +8,74 @@
 #include "Param/ParamTypeHandle.h"
 #include "RigVMCore/RigVMTemplate.h"
 
-struct FAnimNextParameterLibraryAssetRegistryExports;
+#include "UncookedOnlyUtils.generated.h"
+
 class UAnimNextGraph;
 class UAnimNextGraph_EditorData;
 class UAnimNextGraph_EdGraph;
 class URigVMController;
 class URigVMGraph;
-class UAnimNextParameter;
 class UAnimNextParameterBlock;
 class UAnimNextParameterBlock_EditorData;
 class UAnimNextGraph_EdGraph;
 struct FEdGraphPinType;
 
+namespace UE
+{
+	namespace AnimNext
+	{
+		static const FName ExportsAnimNextAssetRegistryTag = TEXT("AnimNextExports");
+	}
+}
+
+UENUM()
+enum class EAnimNextParameterFlags
+{
+	NoFlags = 0x0,
+	Private = 0x1,
+	Read = 0x02,
+	Write = 0x04,
+	Bound = 0x08,
+	Max
+};
+
+ENUM_CLASS_FLAGS(EAnimNextParameterFlags)
+
+USTRUCT()
+struct FAnimNextParameterAssetRegistryExportEntry
+{
+	GENERATED_BODY()
+
+	FAnimNextParameterAssetRegistryExportEntry() = default;
+	
+	FAnimNextParameterAssetRegistryExportEntry(FName InName, const FAnimNextParamType& InType, EAnimNextParameterFlags InFlags = EAnimNextParameterFlags::NoFlags)
+		: Name(InName)
+		, Type(InType)
+		, Flags(InFlags) 
+	{}
+	
+	UPROPERTY()
+	FName Name;
+
+	UPROPERTY()
+	FAnimNextParamType Type;
+
+	// Asset, found first in asset-registry, that references this parameter entry
+	FAssetData ReferencingAsset;
+
+	UPROPERTY()
+	EAnimNextParameterFlags Flags;
+};
+
+USTRUCT()
+struct FAnimNextParameterProviderAssetRegistryExports
+{
+	GENERATED_BODY()
+
+	UPROPERTY()
+	TArray<FAnimNextParameterAssetRegistryExportEntry> Parameters;
+};
+
 namespace UE::AnimNext::UncookedOnly
 {
 
@@ -114,8 +170,11 @@ struct ANIMNEXTUNCOOKEDONLY_API FUtils
 			OutNodes.Append(GraphNodes);
 		}
 	}
+	
+	static bool GetExportedParametersForAsset(const FAssetData& InAsset, FAnimNextParameterProviderAssetRegistryExports& OutExports);
+	static bool GetExportedParametersFromAssetRegistry(FAnimNextParameterProviderAssetRegistryExports& OutExports);
 
-	static bool GetExportedParametersForLibrary(const FAssetData& InLibraryAsset, FAnimNextParameterLibraryAssetRegistryExports& OutExports);
+	static void GetGraphParameters(const URigVMGraph* Graph,FAnimNextParameterProviderAssetRegistryExports& OutExports);
 	
 	// Attempts to determine the type from a parameter name
 	// If the name cannot be found, the returned type will be invalid
diff --git a/Engine/Plugins/Experimental/AnimNext/Source/AnimNextUncookedOnly/Private/Graph/AnimGraph/AnimGraphNodeBinding_AnimNextParameters.cpp b/Engine/Plugins/Experimental/AnimNext/Source/AnimNextUncookedOnly/Private/Graph/AnimGraph/AnimGraphNodeBinding_AnimNextParameters.cpp
--- a/Engine/Plugins/Experimental/AnimNext/Source/AnimNextUncookedOnly/Private/Graph/AnimGraph/AnimGraphNodeBinding_AnimNextParameters.cpp
+++ b/Engine/Plugins/Experimental/AnimNext/Source/AnimNextUncookedOnly/Private/Graph/AnimGraph/AnimGraphNodeBinding_AnimNextParameters.cpp
@@ -549,7 +549,6 @@ TSharedRef<SWidget> UAnimGraphNodeBinding_AnimNextParameters::MakePropertyBindin
 					return (FParamUtils::GetCompatibility(ParamTypeHandle, InParamType.GetHandle()).IsCompatibleWithDataLoss()) ? EFilterParameterResult::Include : EFilterParameterResult::Exclude;
 				});
 				Args.bMultiSelect = false;
-				Args.bShowLibraries = false;
 				Args.bShowBlocks = false;
 			
 				MenuBuilder.AddWidget(
diff --git a/Engine/Plugins/Experimental/AnimNext/Source/AnimNextUncookedOnly/Private/Graph/AnimNextGraph_EditorData.cpp b/Engine/Plugins/Experimental/AnimNext/Source/AnimNextUncookedOnly/Private/Graph/AnimNextGraph_EditorData.cpp
--- a/Engine/Plugins/Experimental/AnimNext/Source/AnimNextUncookedOnly/Private/Graph/AnimNextGraph_EditorData.cpp
+++ b/Engine/Plugins/Experimental/AnimNext/Source/AnimNextUncookedOnly/Private/Graph/AnimNextGraph_EditorData.cpp
@@ -17,6 +17,8 @@
 #include "UObject/ObjectSaveContext.h"
 #include "Kismet2/BlueprintEditorUtils.h"
 #include "Graph/AnimNextGraphEntry.h"
+#include "Param/RigVMDispatch_GetLayerParameter.h"
+#include "Param/RigVMDispatch_GetParameter.h"
 
 namespace UE::AnimNext::Graph::Private
 {
@@ -166,6 +168,28 @@ void UAnimNextGraph_EditorData::PostLoad()
 }
 
 #if WITH_EDITOR
+void UAnimNextGraph_EditorData::GetAssetRegistryTags(TArray<FAssetRegistryTag>& OutTags) const
+{
+	UObject::GetAssetRegistryTags(OutTags);
+
+	FAnimNextParameterProviderAssetRegistryExports ExportParameters;
+	
+	for (const TObjectPtr<UAnimNextGraph_EdGraph>& Graph : Graphs)
+	{
+		if (const URigVMGraph* RigGraph = Graph->GetModel())
+		{
+			UE::AnimNext::UncookedOnly::FUtils::GetGraphParameters(RigGraph, ExportParameters);
+		}
+	}
+
+	if (ExportParameters.Parameters.Num())
+	{
+		FString TagValue;	
+		FAnimNextParameterProviderAssetRegistryExports::StaticStruct()->ExportText(TagValue, &ExportParameters, nullptr, nullptr, PPF_None, nullptr);
+		OutTags.Add(FAssetRegistryTag(UE::AnimNext::ExportsAnimNextAssetRegistryTag, TagValue, FAssetRegistryTag::TT_Hidden));	
+	}
+}
+
 void UAnimNextGraph_EditorData::HandlePackageDone(const FEndLoadPackageContext& Context)
 {
 	if (!Context.LoadedPackages.Contains(GetPackage()))
diff --git a/Engine/Plugins/Experimental/AnimNext/Source/AnimNextUncookedOnly/Private/Param/AnimNextParameter.cpp b/Engine/Plugins/Experimental/AnimNext/Source/AnimNextUncookedOnly/Private/Param/AnimNextParameter.cpp
deleted file mode 100644
diff --git a/Engine/Plugins/Experimental/AnimNext/Source/AnimNextUncookedOnly/Private/Param/AnimNextParameterBlockBinding.cpp b/Engine/Plugins/Experimental/AnimNext/Source/AnimNextUncookedOnly/Private/Param/AnimNextParameterBlockBinding.cpp
--- a/Engine/Plugins/Experimental/AnimNext/Source/AnimNextUncookedOnly/Private/Param/AnimNextParameterBlockBinding.cpp
+++ b/Engine/Plugins/Experimental/AnimNext/Source/AnimNextUncookedOnly/Private/Param/AnimNextParameterBlockBinding.cpp
@@ -1,22 +1,12 @@
 // Copyright Epic Games, Inc. All Rights Reserved.
 
 #include "Param/AnimNextParameterBlockBinding.h"
-#include "Param/AnimNextParameter.h"
-#include "Param/AnimNextParameterLibrary.h"
 
 #define LOCTEXT_NAMESPACE "AnimNextParameterBlockBinding"
 
 FAnimNextParamType UAnimNextParameterBlockBinding::GetParamType() const
 {
-	if(Library)
-	{
-		if(const UAnimNextParameter* Parameter = Library->FindParameter(ParameterName))
-		{
-			return Parameter->GetType();
-		}
-	}
-
-	return FAnimNextParamType();
+	return Type;
 }
 
 FName UAnimNextParameterBlockBinding::GetParameterName() const
@@ -24,31 +14,30 @@ FName UAnimNextParameterBlockBinding::GetParameterName() const
 	return ParameterName;
 }
 
-void UAnimNextParameterBlockBinding::SetParameterName(FName InName, bool bSetupUndoRedo)
+bool UAnimNextParameterBlockBinding::SetParamType(const FAnimNextParamType& InType, bool bSetupUndoRedo)
 {
 	if(bSetupUndoRedo)
 	{
 		Modify();
 	}
-
-	ParameterName = InName;
+	
+	Type = InType;
 
 	BroadcastModified();
+
+	return true;
 }
 
-const UAnimNextParameter* UAnimNextParameterBlockBinding::GetParameter() const
+void UAnimNextParameterBlockBinding::SetParameterName(FName InName, bool bSetupUndoRedo)
 {
-	if(Library)
+	if(bSetupUndoRedo)
 	{
-		return Library->FindParameter(ParameterName);
+		Modify();
 	}
 
-	return nullptr;
-}
+	ParameterName = InName;
 
-const UAnimNextParameterLibrary* UAnimNextParameterBlockBinding::GetLibrary() const
-{
-	return Library;
+	BroadcastModified();
 }
 
 FText UAnimNextParameterBlockBinding::GetDisplayName() const
@@ -60,21 +49,7 @@ FText UAnimNextParameterBlockBinding::GetDisplayNameTooltip() const
 {
 	FTextBuilder TextBuilder;
 	TextBuilder.AppendLine(FText::FromName(ParameterName));
-	TextBuilder.AppendLine(Library ? FText::FromString(Library->GetPathName()) : LOCTEXT("MissingLibrary", "Missing parameter library"));
 	return TextBuilder.ToText();
 }
 
-void UAnimNextParameterBlockBinding::GetEditedObjects(TArray<UObject*>& OutObjects) const
-{
-	if(Library)
-	{
-		OutObjects.Add(Library);
-	}
-
-	if(const UAnimNextParameter* Parameter = GetParameter())
-	{
-		OutObjects.Add(const_cast<UAnimNextParameter*>(Parameter));
-	}
-}
-
 #undef LOCTEXT_NAMESPACE
\ No newline at end of file
diff --git a/Engine/Plugins/Experimental/AnimNext/Source/AnimNextUncookedOnly/Private/Param/AnimNextParameterBlockBindingReference.cpp b/Engine/Plugins/Experimental/AnimNext/Source/AnimNextUncookedOnly/Private/Param/AnimNextParameterBlockBindingReference.cpp
--- a/Engine/Plugins/Experimental/AnimNext/Source/AnimNextUncookedOnly/Private/Param/AnimNextParameterBlockBindingReference.cpp
+++ b/Engine/Plugins/Experimental/AnimNext/Source/AnimNextUncookedOnly/Private/Param/AnimNextParameterBlockBindingReference.cpp
@@ -1,19 +1,22 @@
 ﻿// Copyright Epic Games, Inc. All Rights Reserved.
 
 #include "Param/AnimNextParameterBlockBindingReference.h"
-#include "Param/AnimNextParameter.h"
-#include "Param/AnimNextParameterLibrary.h"
+
+#include "UncookedOnlyUtils.h"
 #include "Param/AnimNextParameterBlock.h"
 
 #define LOCTEXT_NAMESPACE "AnimNextParameterBlockBindingReference"
 
 FAnimNextParamType UAnimNextParameterBlockBindingReference::GetParamType() const
 {
-	if(Library)
+	if (Block != nullptr)
 	{
-		if(const UAnimNextParameter* Parameter = Library->FindParameter(ParameterName))
+		if(UAnimNextParameterBlockEntry* ReferencedEntry = UE::AnimNext::UncookedOnly::FUtils::GetEditorData(Block)->FindBinding(ParameterName))
 		{
-			return Parameter->GetType();
+			if(const IAnimNextParameterBlockParameterInterface* Binding = Cast<IAnimNextParameterBlockParameterInterface>(ReferencedEntry))
+			{
+				return Binding->GetParamType();
+			}
 		}
 	}
 
@@ -30,21 +33,6 @@ void UAnimNextParameterBlockBindingReference::SetParameterName(FName InName, boo
 	ParameterName = InName;
 }
 
-const UAnimNextParameter* UAnimNextParameterBlockBindingReference::GetParameter() const
-{
-	if(Library)
-	{
-		return Library->FindParameter(ParameterName);
-	}
-
-	return nullptr;
-}
-
-const UAnimNextParameterLibrary* UAnimNextParameterBlockBindingReference::GetLibrary() const
-{
-	return Library;
-}
-
 const UAnimNextParameterBlock* UAnimNextParameterBlockBindingReference::GetBlock() const
 {
 	return Block;
@@ -59,22 +47,11 @@ FText UAnimNextParameterBlockBindingReference::GetDisplayNameTooltip() const
 {
 	FTextBuilder TextBuilder;
 	TextBuilder.AppendLine(FText::FromName(ParameterName));
-	TextBuilder.AppendLine(Library ? FText::FromString(Library->GetPathName()) : LOCTEXT("MissingLibrary", "Missing parameter library"));
 	return TextBuilder.ToText();
 }
 
 void UAnimNextParameterBlockBindingReference::GetEditedObjects(TArray<UObject*>& OutObjects) const
 {
-	if(Library)
-	{
-		OutObjects.Add(Library);
-	}
-
-	if(const UAnimNextParameter* Parameter = GetParameter())
-	{
-		OutObjects.Add(const_cast<UAnimNextParameter*>(Parameter));
-	}
-	
 	if(Block)
 	{
 		OutObjects.Add(Block);
diff --git a/Engine/Plugins/Experimental/AnimNext/Source/AnimNextUncookedOnly/Private/Param/AnimNextParameterBlockParameter.cpp b/Engine/Plugins/Experimental/AnimNext/Source/AnimNextUncookedOnly/Private/Param/AnimNextParameterBlockParameter.cpp
--- a/Engine/Plugins/Experimental/AnimNext/Source/AnimNextUncookedOnly/Private/Param/AnimNextParameterBlockParameter.cpp
+++ b/Engine/Plugins/Experimental/AnimNext/Source/AnimNextUncookedOnly/Private/Param/AnimNextParameterBlockParameter.cpp
@@ -1,31 +1,20 @@
 // Copyright Epic Games, Inc. All Rights Reserved.
 
 #include "Param/AnimNextParameterBlockParameter.h"
-#include "Param/AnimNextParameter.h"
-#include "Param/AnimNextParameterLibrary.h"
+#include "Param/AnimNextParameterBlock_EditorData.h"
 #include "Param/Params.h"
 
 FAnimNextParamType UAnimNextParameterBlockParameter::GetParamType() const
 {
 	using namespace UE::AnimNext;
-	
-	if(Library)
-	{
-		if(const UAnimNextParameter* Parameter = Library->FindParameter(ParameterName))
-		{
-			return Parameter->GetType();
-		}
-	}
-	else
+
+	// Look in built-in parameters
+	if(const FParamDefinition* ParamDefinition = FParams::FindBuiltInParameter(ParameterName))
 	{
-		// Look in built-in parameters
-		if(const FParamDefinition* ParamDefinition = FParams::FindBuiltInParameter(ParameterName))
-		{
-			return ParamDefinition->GetType();
-		}
+		return ParamDefinition->GetType();
 	}
 
-	return FAnimNextParamType();
+	return Type;
 }
 
 FName UAnimNextParameterBlockParameter::GetParameterName() const
@@ -33,41 +22,29 @@ FName UAnimNextParameterBlockParameter::GetParameterName() const
 	return ParameterName;
 }
 
-void UAnimNextParameterBlockParameter::SetParameterName(FName InName, bool bSetupUndoRedo)
+bool UAnimNextParameterBlockParameter::SetParamType(const FAnimNextParamType& InType, bool bSetupUndoRedo)
 {
 	if(bSetupUndoRedo)
 	{
 		Modify();
 	}
-
-	if (Library)
-	{
-		UAnimNextParameter* Parameter = Library->FindParameter(ParameterName);
-		if (Parameter)
-		{
-			if (Parameter->Rename(*InName.ToString(), Parameter->GetOuter()))
-			{
-				ParameterName = InName;
-			}
-		}
-	}
+	
+	Type = InType;
 
 	BroadcastModified();
+
+	return true;
 }
 
-const UAnimNextParameter* UAnimNextParameterBlockParameter::GetParameter() const
+void UAnimNextParameterBlockParameter::SetParameterName(FName InName, bool bSetupUndoRedo)
 {
-	if(Library)
+	if(bSetupUndoRedo)
 	{
-		return Library->FindParameter(ParameterName);
+		Modify();
 	}
 
-	return nullptr;
-}
-
-const UAnimNextParameterLibrary* UAnimNextParameterBlockParameter::GetLibrary() const
-{
-	return Library;
+	ParameterName = InName;
+	BroadcastModified();
 }
 
 FText UAnimNextParameterBlockParameter::GetDisplayName() const
diff --git a/Engine/Plugins/Experimental/AnimNext/Source/AnimNextUncookedOnly/Private/Param/AnimNextParameterBlock_EditorData.cpp b/Engine/Plugins/Experimental/AnimNext/Source/AnimNextUncookedOnly/Private/Param/AnimNextParameterBlock_EditorData.cpp
--- a/Engine/Plugins/Experimental/AnimNext/Source/AnimNextUncookedOnly/Private/Param/AnimNextParameterBlock_EditorData.cpp
+++ b/Engine/Plugins/Experimental/AnimNext/Source/AnimNextUncookedOnly/Private/Param/AnimNextParameterBlock_EditorData.cpp
@@ -15,14 +15,14 @@
 #include "Param/AnimNextParameterBlock.h"
 #include "Param/AnimNextParameterBlockEntry.h"
 #include "ExternalPackageHelper.h"
-#include "Param/AnimNextParameter.h"
 #include "Param/AnimNextParameterBlockBinding.h"
 #include "Param/AnimNextParameterBlockBindingReference.h"
-#include "Param/AnimNextParameterLibrary.h"
 #include "Misc/TransactionObjectEvent.h"
 #include "ObjectTools.h"
 #include "Param/AnimNextParameterBlockGraph.h"
 #include "Param/AnimNextParameterBlockParameter.h"
+#include "Param/RigVMDispatch_GetLayerParameter.h"
+#include "Param/RigVMDispatch_GetParameter.h"
 
 #if WITH_EDITORONLY_DATA
 
@@ -71,12 +71,7 @@ UAnimNextParameterBlock_EditorData::UAnimNextParameterBlock_EditorData(const FOb
 
 #if WITH_EDITORONLY_DATA
 
-UAnimNextParameterBlockParameter* UAnimNextParameterBlockLibrary::AddParameter(UAnimNextParameterBlock* InBlock, FName InName, UAnimNextParameterLibrary* InLibrary, bool bSetupUndoRedo, bool bPrintPythonCommand)
-{
-	return UE::AnimNext::UncookedOnly::FUtils::GetEditorData(InBlock)->AddParameter(InName, InLibrary, bSetupUndoRedo, bPrintPythonCommand);
-}
-
-UAnimNextParameterBlockParameter* UAnimNextParameterBlock_EditorData::AddParameter(FName InName, UAnimNextParameterLibrary* InLibrary, bool bSetupUndoRedo, bool bPrintPythonCommand)
+UAnimNextParameterBlockParameter* UAnimNextParameterBlock_EditorData::AddParameter(FName InName, FAnimNextParamType InType, bool bSetupUndoRedo, bool bPrintPythonCommand)
 {
 	if(InName == NAME_None)
 	{
@@ -84,17 +79,6 @@ UAnimNextParameterBlockParameter* UAnimNextParameterBlock_EditorData::AddParamet
 		return nullptr;
 	}
 
-	if(InLibrary != nullptr)
-	{
-		// Check if parameter exists in library
-		UAnimNextParameter* Parameter = InLibrary->FindParameter(InName);
-		if(Parameter == nullptr)
-		{
-			ReportError(TEXT("UAnimNextParameterBlock_EditorData::AddParameter: Parameter does not exist in library."));
-			return nullptr;
-		}
-	}
-
 	// Check for duplicate parameter
 	const bool bAlreadyExists = Entries.ContainsByPredicate([InName](const UAnimNextParameterBlockEntry* InEntry)
 	{
@@ -113,7 +97,7 @@ UAnimNextParameterBlockParameter* UAnimNextParameterBlock_EditorData::AddParamet
 
 	UAnimNextParameterBlockParameter* NewEntry = UE::AnimNext::Parameters::Private::CreateNewParameterBlockEntry<UAnimNextParameterBlockParameter>(this);
 	NewEntry->ParameterName = InName;
-	NewEntry->Library = InLibrary;
+	NewEntry->Type = InType;
 
 	if(bSetupUndoRedo)
 	{
@@ -129,6 +113,12 @@ UAnimNextParameterBlockParameter* UAnimNextParameterBlock_EditorData::AddParamet
 }
 
 
+UAnimNextParameterBlockParameter* UAnimNextParameterBlockLibrary::AddParameter(UAnimNextParameterBlock* InBlock, FName InName, EPropertyBagPropertyType InValueType,
+	EPropertyBagContainerType InContainerType, const UObject* InValueTypeObject, bool bSetupUndoRedo, bool bPrintPythonCommand)
+{
+	return UE::AnimNext::UncookedOnly::FUtils::GetEditorData(InBlock)->AddParameter(InName, FAnimNextParamType(InValueType, InContainerType, InValueTypeObject), bSetupUndoRedo, bPrintPythonCommand);
+}
+
 UAnimNextParameterBlockGraph* UAnimNextParameterBlockLibrary::AddGraph(UAnimNextParameterBlock* InBlock, FName InName, bool bSetupUndoRedo, bool bPrintPythonCommand)
 {
 	return UE::AnimNext::UncookedOnly::FUtils::GetEditorData(InBlock)->AddGraph(InName, bSetupUndoRedo, bPrintPythonCommand);
@@ -190,12 +180,12 @@ UAnimNextParameterBlockGraph* UAnimNextParameterBlock_EditorData::AddGraph(FName
 }
 
 
-UAnimNextParameterBlockBinding* UAnimNextParameterBlockLibrary::AddBinding(UAnimNextParameterBlock* InBlock, FName InName, UAnimNextParameterLibrary* InLibrary, bool bSetupUndoRedo, bool bPrintPythonCommand)
+UAnimNextParameterBlockBinding* UAnimNextParameterBlockLibrary::AddBinding(UAnimNextParameterBlock* InBlock, FName InName, EPropertyBagPropertyType InValueType, EPropertyBagContainerType InContainerType, const UObject* InValueTypeObject, bool bSetupUndoRedo, bool bPrintPythonCommand)
 {
-	return UE::AnimNext::UncookedOnly::FUtils::GetEditorData(InBlock)->AddBinding(InName, InLibrary, bSetupUndoRedo, bPrintPythonCommand);
+	return UE::AnimNext::UncookedOnly::FUtils::GetEditorData(InBlock)->AddBinding(InName, FAnimNextParamType(InValueType, InContainerType, InValueTypeObject), bSetupUndoRedo, bPrintPythonCommand);
 }
 
-UAnimNextParameterBlockBinding* UAnimNextParameterBlock_EditorData::AddBinding(FName InName, UAnimNextParameterLibrary* InLibrary, bool bSetupUndoRedo, bool bPrintPythonCommand)
+UAnimNextParameterBlockBinding* UAnimNextParameterBlock_EditorData::AddBinding(FName InName, FAnimNextParamType InType, bool bSetupUndoRedo, bool bPrintPythonCommand)
 {
 	if(InName == NAME_None)
 	{
@@ -203,26 +193,12 @@ UAnimNextParameterBlockBinding* UAnimNextParameterBlock_EditorData::AddBinding(F
 		return nullptr;
 	}
 
-	if(InLibrary == nullptr)
-	{
-		ReportError(TEXT("UAnimNextParameterBlock_EditorData::AddBinding: Invalid parameter library supplied."));
-		return nullptr;
-	}
-
-	// Check parameter exists in library
-	UAnimNextParameter* Parameter = InLibrary->FindParameter(InName);
-	if(Parameter == nullptr)
-	{
-		ReportError(TEXT("UAnimNextParameterBlock_EditorData::AddBinding: Parameter does not exist in library."));
-		return nullptr;
-	}
-
 	// Check for duplicate bindings
-	const bool bAlreadyExists = Entries.ContainsByPredicate([InName, InLibrary](const UAnimNextParameterBlockEntry* InEntry)
+	const bool bAlreadyExists = Entries.ContainsByPredicate([InName](const UAnimNextParameterBlockEntry* InEntry)
 	{
 		if(const IAnimNextParameterBlockParameterInterface* Binding = Cast<IAnimNextParameterBlockParameterInterface>(InEntry))
 		{
-			return Binding->GetParameterName() == InName && Binding->GetLibrary() == InLibrary;
+			return Binding->GetParameterName() == InName;
 		}
 		return false;
 	});
@@ -235,7 +211,7 @@ UAnimNextParameterBlockBinding* UAnimNextParameterBlock_EditorData::AddBinding(F
 
 	UAnimNextParameterBlockBinding* NewEntry = UE::AnimNext::Parameters::Private::CreateNewParameterBlockEntry<UAnimNextParameterBlockBinding>(this);
 	NewEntry->ParameterName = InName;
-	NewEntry->Library = InLibrary;
+	NewEntry->Type = InType;
 
 	if(bSetupUndoRedo)
 	{
@@ -254,7 +230,7 @@ UAnimNextParameterBlockBinding* UAnimNextParameterBlock_EditorData::AddBinding(F
 		NewEntry->BindingGraph = NewGraph;
 
 		URigVMController* Controller = RigVMClient.GetController(NewGraph);
-		UE::AnimNext::UncookedOnly::FUtils::SetupBindingGraph(Controller, InName, Parameter->GetType());
+		UE::AnimNext::UncookedOnly::FUtils::SetupBindingGraph(Controller, InName, NewEntry->GetParamType());
 	}
 
 	BroadcastModified();
@@ -262,12 +238,12 @@ UAnimNextParameterBlockBinding* UAnimNextParameterBlock_EditorData::AddBinding(F
 	return NewEntry;
 }
 
-UAnimNextParameterBlockBindingReference* UAnimNextParameterBlockLibrary::AddBindingReference(UAnimNextParameterBlock* InBlock, FName InName, UAnimNextParameterLibrary* InLibrary, UAnimNextParameterBlock* InReferencedBlock, bool bSetupUndoRedo, bool bPrintPythonCommand)
+UAnimNextParameterBlockBindingReference* UAnimNextParameterBlockLibrary::AddBindingReference(UAnimNextParameterBlock* InBlock, FName InName, UAnimNextParameterBlock* InReferencedBlock, bool bSetupUndoRedo, bool bPrintPythonCommand)
 {
-	return UE::AnimNext::UncookedOnly::FUtils::GetEditorData(InBlock)->AddBindingReference(InName, InLibrary, InReferencedBlock, bSetupUndoRedo, bPrintPythonCommand);
+	return UE::AnimNext::UncookedOnly::FUtils::GetEditorData(InBlock)->AddBindingReference(InName, InReferencedBlock, bSetupUndoRedo, bPrintPythonCommand);
 }
 
-UAnimNextParameterBlockBindingReference* UAnimNextParameterBlock_EditorData::AddBindingReference(FName InName, UAnimNextParameterLibrary* InLibrary, UAnimNextParameterBlock* InReferencedBlock, bool bSetupUndoRedo, bool bPrintPythonCommand)
+UAnimNextParameterBlockBindingReference* UAnimNextParameterBlock_EditorData::AddBindingReference(FName InName, UAnimNextParameterBlock* InReferencedBlock, bool bSetupUndoRedo, bool bPrintPythonCommand)
 {
 	if(InName == NAME_None)
 	{
@@ -275,32 +251,25 @@ UAnimNextParameterBlockBindingReference* UAnimNextParameterBlock_EditorData::Add
 		return nullptr;
 	}
 
-	if(InLibrary == nullptr)
-	{
-		ReportError(TEXT("UAnimNextParameterBlock_EditorData::AddBindingReference: Invalid parameter library supplied."));
-		return nullptr;
-	}
-	
 	if(InReferencedBlock == nullptr)
 	{
 		ReportError(TEXT("UAnimNextParameterBlock_EditorData::AddBindingReference: Invalid block supplied."));
 		return nullptr;
 	}
 
-	// Check parameter exists in library
-	UAnimNextParameter* Parameter = InLibrary->FindParameter(InName);
-	if(Parameter == nullptr)
+	const UAnimNextParameterBlockEntry* ReferencedEntry = UE::AnimNext::UncookedOnly::FUtils::GetEditorData(InReferencedBlock)->FindBinding(InName);
+	if(ReferencedEntry == nullptr)
 	{
-		ReportError(TEXT("UAnimNextParameterBlock_EditorData::AddBindingReference: Parameter does not exist in library."));
+		ReportError(TEXT("UAnimNextParameterBlock_EditorData::AddBindingReference: Named parameter to reference not found in block supplied."));
 		return nullptr;
 	}
 
 	// Check for duplicates
-	const bool bAlreadyExists = Entries.ContainsByPredicate([InName, InLibrary](const UAnimNextParameterBlockEntry* InEntry)
+	const bool bAlreadyExists = Entries.ContainsByPredicate([InName](const UAnimNextParameterBlockEntry* InEntry)
 	{
 		if(const IAnimNextParameterBlockParameterInterface* Binding = Cast<IAnimNextParameterBlockParameterInterface>(InEntry))
 		{
-			return Binding->GetParameterName() == InName && Binding->GetLibrary() == InLibrary;
+			return Binding->GetParameterName() == InName;
 		}
 
 		return false;
@@ -314,7 +283,6 @@ UAnimNextParameterBlockBindingReference* UAnimNextParameterBlock_EditorData::Add
 
 	UAnimNextParameterBlockBindingReference* NewEntry = UE::AnimNext::Parameters::Private::CreateNewParameterBlockEntry<UAnimNextParameterBlockBindingReference>(this);
 	NewEntry->ParameterName = InName;
-	NewEntry->Library = InLibrary;
 	NewEntry->Block = InReferencedBlock;
 
 	if(bSetupUndoRedo)
@@ -609,21 +577,28 @@ void UAnimNextParameterBlock_EditorData::GetAssetRegistryTags(TArray<FAssetRegis
 {
 	Super::GetAssetRegistryTags(OutTags);
 
-	FAnimNextParameterBlockAssetRegistryExports Exports;
-	Exports.Bindings.Reserve(Entries.Num());
+	FAnimNextParameterProviderAssetRegistryExports ExportParameters;
+	ExportParameters.Parameters.Reserve(Entries.Num());
 
 	for(const UAnimNextParameterBlockEntry* Entry : Entries)
 	{
 		if(const IAnimNextParameterBlockParameterInterface* Binding = Cast<IAnimNextParameterBlockParameterInterface>(Entry))
 		{
-			Exports.Bindings.Emplace(Binding->GetParameterName(), FSoftObjectPath(Binding->GetLibrary()));
+			ExportParameters.Parameters.Emplace(Binding->GetParameterName(), Binding->GetParamType(), EAnimNextParameterFlags::Bound);
 		}
 	}
 
-	FString TagValue;
-	FAnimNextParameterBlockAssetRegistryExports::StaticStruct()->ExportText(TagValue, &Exports, nullptr, nullptr, PPF_None, nullptr);
+	for (const TObjectPtr<UAnimNextParameterBlock_EdGraph>& Graph : Graphs)
+	{
+		if (const URigVMGraph* RigGraph = Graph->GetModel())
+		{
+			UE::AnimNext::UncookedOnly::FUtils::GetGraphParameters(RigGraph, ExportParameters);
+		}
+	}
 
-	OutTags.Add(FAssetRegistryTag(ExportsAssetRegistryTag, TagValue, FAssetRegistryTag::TT_Hidden));
+	FString TagValue;
+	FAnimNextParameterProviderAssetRegistryExports::StaticStruct()->ExportText(TagValue, &ExportParameters, nullptr, nullptr, PPF_None, nullptr);
+	OutTags.Add(FAssetRegistryTag(UE::AnimNext::ExportsAnimNextAssetRegistryTag, TagValue, FAssetRegistryTag::TT_Hidden));
 }
 
 #if WITH_EDITOR
diff --git a/Engine/Plugins/Experimental/AnimNext/Source/AnimNextUncookedOnly/Private/Param/AnimNextParameterLibrary.cpp b/Engine/Plugins/Experimental/AnimNext/Source/AnimNextUncookedOnly/Private/Param/AnimNextParameterLibrary.cpp
deleted file mode 100644
diff --git a/Engine/Plugins/Experimental/AnimNext/Source/AnimNextUncookedOnly/Private/UncookedOnlyUtils.cpp b/Engine/Plugins/Experimental/AnimNext/Source/AnimNextUncookedOnly/Private/UncookedOnlyUtils.cpp
--- a/Engine/Plugins/Experimental/AnimNext/Source/AnimNextUncookedOnly/Private/UncookedOnlyUtils.cpp
+++ b/Engine/Plugins/Experimental/AnimNext/Source/AnimNextUncookedOnly/Private/UncookedOnlyUtils.cpp
@@ -12,7 +12,6 @@
 #include "Graph/RigUnit_AnimNextGraphEvaluator.h"
 #include "Graph/RigUnit_AnimNextShimRoot.h"
 #include "Graph/AnimNextExecuteContext.h"
-#include "Param/AnimNextParameter.h"
 #include "Param/AnimNextParameterBlock.h"
 #include "Param/AnimNextParameterBlock_EditorData.h"
 #include "Param/AnimNextParameterBlock_EdGraph.h"
@@ -32,10 +31,10 @@
 #include "DecoratorBase/DecoratorRegistry.h"
 #include "DecoratorBase/Decorator.h"
 #include "Graph/RigUnit_AnimNextBeginExecution.h"
-#include "Param/AnimNextParameterLibrary.h"
 #include "Param/Params.h"
 #include "Serialization/MemoryReader.h"
 #include "RigVMRuntimeDataRegistry.h"
+#include "Param/RigVMDispatch_GetLayerParameter.h"
 
 #include "RigVMCompiler/RigVMCompiler.h"
 #include "RigVMCore/RigVM.h"
@@ -742,11 +741,9 @@ void FUtils::CompileStruct(UAnimNextParameterBlock* InParameterBlock)
 	{
 		if(const IAnimNextParameterBlockParameterInterface* Binding = Cast<IAnimNextParameterBlockParameterInterface>(Entry))
 		{
-			if(const UAnimNextParameter* Parameter = Binding->GetParameter())
-			{
-				const FAnimNextParamType& Type = Binding->GetParamType();
-				PropertyDescs.Emplace(Parameter->GetFName(), Type.GetContainerType(), Type.GetValueType(), Type.GetValueTypeObject());
-			}
+			const FAnimNextParamType& Type = Binding->GetParamType();
+			ensure(Type.IsValid());
+			PropertyDescs.Emplace(Binding->GetParameterName(), Type.GetContainerType(), Type.GetValueType(), Type.GetValueTypeObject());
 		}
 	}
 
@@ -1148,12 +1145,6 @@ FText FUtils::GetParameterDisplayNameText(FName InParameterName)
 	return FText::FromString(NameAsString);
 }
 
-bool FUtils::GetExportedParametersForLibrary(const FAssetData& InLibraryAsset, FAnimNextParameterLibraryAssetRegistryExports& OutExports)
-{
-	const FString TagValue = InLibraryAsset.GetTagValueRef<FString>(UAnimNextParameterLibrary::ExportsAssetRegistryTag);
-	return FAnimNextParameterLibraryAssetRegistryExports::StaticStruct()->ImportText(*TagValue, &OutExports, nullptr, PPF_None, nullptr, FAnimNextParameterLibraryAssetRegistryExports::StaticStruct()->GetName()) != nullptr;
-}
-
 FAnimNextParamType FUtils::GetParameterTypeFromName(FName InName)
 {
 	// Check built-in params first as they are cheaper
@@ -1165,26 +1156,103 @@ FAnimNextParamType FUtils::GetParameterTypeFromName(FName InName)
 	// Query the asset registry for other params
 	IAssetRegistry& AssetRegistry = FModuleManager::Get().LoadModuleChecked<FAssetRegistryModule>(TEXT("AssetRegistry")).Get();
 
-	FARFilter ARFilter;
-	ARFilter.ClassPaths = { UAnimNextParameterLibrary::StaticClass()->GetClassPathName() };
+	FAnimNextParameterProviderAssetRegistryExports Exports;
+	GetExportedParametersFromAssetRegistry(Exports);
+	for(const FAnimNextParameterAssetRegistryExportEntry& Export : Exports.Parameters)
+	{
+		if(Export.Name == InName)
+		{
+			return Export.Type;
+		}
+	}
 
-	TArray<FAssetData> LibraryAssets;
-	AssetRegistry.GetAssets(ARFilter, LibraryAssets);
+	return FAnimNextParamType();
+}
+bool FUtils::GetExportedParametersForAsset(const FAssetData& InAsset, FAnimNextParameterProviderAssetRegistryExports& OutExports)
+{
+	const FString TagValue = InAsset.GetTagValueRef<FString>(UE::AnimNext::ExportsAnimNextAssetRegistryTag);
+	return FAnimNextParameterProviderAssetRegistryExports::StaticStruct()->ImportText(*TagValue, &OutExports, nullptr, PPF_None, nullptr, FAnimNextParameterProviderAssetRegistryExports::StaticStruct()->GetName()) != nullptr;
+}
 
-	for(const FAssetData& LibraryAsset : LibraryAssets)
+bool FUtils::GetExportedParametersFromAssetRegistry(FAnimNextParameterProviderAssetRegistryExports& OutExports)
+{
+	TArray<FAssetData> AssetData;
+	IAssetRegistry::GetChecked().GetAssetsByTags({UE::AnimNext::ExportsAnimNextAssetRegistryTag}, AssetData);
+
+	for (const FAssetData& Asset : AssetData)
 	{
-		FAnimNextParameterLibraryAssetRegistryExports Exports;
-		GetExportedParametersForLibrary(LibraryAsset, Exports);
-		for(const FAnimNextParameterLibraryAssetRegistryExportEntry& Export : Exports.Parameters)
+		const FString TagValue = Asset.GetTagValueRef<FString>(UE::AnimNext::ExportsAnimNextAssetRegistryTag);
+		FAnimNextParameterProviderAssetRegistryExports AssetExports;
+		if (FAnimNextParameterProviderAssetRegistryExports::StaticStruct()->ImportText(*TagValue, &AssetExports, nullptr, PPF_None, nullptr, FAnimNextParameterProviderAssetRegistryExports::StaticStruct()->GetName()) != nullptr)
 		{
-			if(Export.Name == InName)
+			for (FAnimNextParameterAssetRegistryExportEntry& Parameter : AssetExports.Parameters)
 			{
-				return Export.Type;
+				if (Parameter.Name != NAME_None && Parameter.Type.IsValid() && Parameter.Type.ValueTypeObject != FRigVMUnknownType::StaticStruct())
+				{
+					FAnimNextParameterAssetRegistryExportEntry* ExistingEntry = OutExports.Parameters.FindByPredicate([Name=Parameter.Name](const FAnimNextParameterAssetRegistryExportEntry& Entry)
+					{
+						return Entry.Name == Name;
+					});
+					
+					if (!ExistingEntry)
+					{
+						Parameter.ReferencingAsset = Asset;
+						OutExports.Parameters.Add(Parameter);
+					}
+					else
+					{
+						ensureMsgf(ExistingEntry->Type == Parameter.Type, TEXT("[%s::%s] %s vs [%s::%s] %s"), *ExistingEntry->ReferencingAsset.ToSoftObjectPath().ToString(), *ExistingEntry->Name.ToString(), *ExistingEntry->Type.ToString(), *Asset.ToSoftObjectPath().ToString(), *Parameter.Name.ToString(), *Parameter.Type.ToString());
+
+						ExistingEntry->Flags |= Parameter.Flags;
+					}
+				}
 			}
 		}
 	}
 
-	return FAnimNextParamType();
+	return OutExports.Parameters.Num() > 0;
 }
 
+void FUtils::GetGraphParameters(const URigVMGraph* Graph, FAnimNextParameterProviderAssetRegistryExports& OutExports)
+{
+	const TArray<URigVMNode*>& Nodes = Graph->GetNodes();
+	for (URigVMNode* Node : Nodes)
+	{
+		if (const URigVMTemplateNode* TemplateNode = Cast<URigVMTemplateNode>(Node))
+		{
+			const FRigVMDispatchFactory* GetParameterFactory = FRigVMRegistry::Get().FindOrAddDispatchFactory(FRigVMDispatch_GetParameter::StaticStruct());
+			const FName GetParameterNotation = GetParameterFactory->GetTemplate()->GetNotation();
+
+			const FRigVMDispatchFactory* GetLayerParameterFactory = FRigVMRegistry::Get().FindOrAddDispatchFactory(FRigVMDispatch_GetLayerParameter::StaticStruct());
+			const FName GetLayerParameterNotation = GetLayerParameterFactory->GetTemplate()->GetNotation();
+
+			const FRigVMDispatchFactory* SetLayerParameterFactory = FRigVMRegistry::Get().FindOrAddDispatchFactory(FRigVMDispatch_SetLayerParameter::StaticStruct());
+			const FName SetLayerParameterNotation = GetLayerParameterFactory->GetTemplate()->GetNotation();
+
+			const bool bReadParameter = TemplateNode->GetNotation() == GetParameterNotation || TemplateNode->GetNotation() == GetLayerParameterNotation;
+			const bool bWriteParameter = TemplateNode->GetNotation() == SetLayerParameterNotation;
+			if (bReadParameter || bWriteParameter)
+			{
+				if (const URigVMPin* NamePin = TemplateNode->FindPin(FRigVMDispatch_GetParameter::ParameterName.ToString()))
+				{							
+					if (const URigVMPin* ValuePin = TemplateNode->FindPin(FRigVMDispatch_GetParameter::ValueName.ToString()))
+					{
+						FAnimNextParamType Type = FAnimNextParamType::FromRigVMTemplateArgument(ValuePin->GetTemplateArgumentType());
+						EAnimNextParameterFlags Flags = EAnimNextParameterFlags::NoFlags;
+						if (bReadParameter)
+						{
+							Flags |= EAnimNextParameterFlags::Read;
+						}
+						
+						if (bWriteParameter)
+						{
+							Flags |= EAnimNextParameterFlags::Write;
+						}
+						OutExports.Parameters.Emplace(FName(NamePin->GetDefaultValue()), Type, Flags);
+					}
+				}
+			}
+		}
+	}
+}
 }
