Fixup static analysis warning (missing null-check)

diff --git a/Engine/Plugins/Experimental/AnimNext/Source/AnimNextEditor/Private/Param/SGraphPinParamName.cpp b/Engine/Plugins/Experimental/AnimNext/Source/AnimNextEditor/Private/Param/SGraphPinParamName.cpp
--- a/Engine/Plugins/Experimental/AnimNext/Source/AnimNextEditor/Private/Param/SGraphPinParamName.cpp
+++ b/Engine/Plugins/Experimental/AnimNext/Source/AnimNextEditor/Private/Param/SGraphPinParamName.cpp
@@ -29,17 +29,20 @@ TSharedRef<SWidget> SGraphPinParamName::GetDefaultValueWidget()
 	Args.bMultiSelect = false;
 
 	// Check whether this is a Set/Get parameter from block node, and if so only show bound parameters
-	if (const URigVMTemplateNode* TemplateNode = Cast<URigVMTemplateNode>(ModelPin->GetOuter()))
+	if (ModelPin)
 	{
-		const FRigVMDispatchFactory* GetLayerParameterFactory = FRigVMRegistry::Get().FindOrAddDispatchFactory(FRigVMDispatch_GetLayerParameter::StaticStruct());
-		const FName GetLayerParameterNotation = GetLayerParameterFactory->GetTemplate()->GetNotation();
+		if (const URigVMTemplateNode* TemplateNode = Cast<URigVMTemplateNode>(ModelPin->GetOuter()))
+		{
+			const FRigVMDispatchFactory* GetLayerParameterFactory = FRigVMRegistry::Get().FindOrAddDispatchFactory(FRigVMDispatch_GetLayerParameter::StaticStruct());
+			const FName GetLayerParameterNotation = GetLayerParameterFactory->GetTemplate()->GetNotation();
 
-		const FRigVMDispatchFactory* SetLayerParameterFactory = FRigVMRegistry::Get().FindOrAddDispatchFactory(FRigVMDispatch_SetLayerParameter::StaticStruct());
-		const FName SetLayerParameterNotation = SetLayerParameterFactory->GetTemplate()->GetNotation();
+			const FRigVMDispatchFactory* SetLayerParameterFactory = FRigVMRegistry::Get().FindOrAddDispatchFactory(FRigVMDispatch_SetLayerParameter::StaticStruct());
+			const FName SetLayerParameterNotation = SetLayerParameterFactory->GetTemplate()->GetNotation();
 
-		if (TemplateNode->GetNotation() == GetLayerParameterNotation || TemplateNode->GetNotation() == SetLayerParameterNotation)
-		{
-			Args.bShowUnboundParameters = false;
+			if (TemplateNode->GetNotation() == GetLayerParameterNotation || TemplateNode->GetNotation() == SetLayerParameterNotation)
+			{
+				Args.bShowUnboundParameters = false;
+			}
 		}
 	}
 	
