Fix AnimNext engine tests following recent traversal changes

diff --git a/Engine/Plugins/Experimental/AnimNext/Source/AnimNextTestSuite/Private/AnimNextDecoratorInterfacesTest.cpp b/Engine/Plugins/Experimental/AnimNext/Source/AnimNextTestSuite/Private/AnimNextDecoratorInterfacesTest.cpp
--- a/Engine/Plugins/Experimental/AnimNext/Source/AnimNextTestSuite/Private/AnimNextDecoratorInterfacesTest.cpp
+++ b/Engine/Plugins/Experimental/AnimNext/Source/AnimNextTestSuite/Private/AnimNextDecoratorInterfacesTest.cpp
@@ -437,7 +437,7 @@ bool FAnimationAnimNextRuntimeTest_IUpdate::RunTest(const FString& InParameters)
 	// Build our graph, it as follow (each node template has a single node instance):
 	// NodeA has no children
 	// NodeB has one child: NodeA (it doesn't update)
-	// NodeD has two children: NodeA and NodeB
+	// NodeC (root) has two children: NodeA and NodeB
 
 	TArray<FNodeHandle> NodeHandles;
 
@@ -447,9 +447,9 @@ bool FAnimationAnimNextRuntimeTest_IUpdate::RunTest(const FString& InParameters)
 	{
 		FDecoratorWriter DecoratorWriter;
 
+		NodeHandles.Add(DecoratorWriter.RegisterNode(*NodeTemplateC));	// Root node
 		NodeHandles.Add(DecoratorWriter.RegisterNode(*NodeTemplateA));
 		NodeHandles.Add(DecoratorWriter.RegisterNode(*NodeTemplateB));
-		NodeHandles.Add(DecoratorWriter.RegisterNode(*NodeTemplateC));
 
 		// We don't have decorator properties
 		TArray<TMap<FString, FString>> DecoratorPropertiesA;
@@ -457,36 +457,36 @@ bool FAnimationAnimNextRuntimeTest_IUpdate::RunTest(const FString& InParameters)
 
 		TArray<TMap<FString, FString>> DecoratorPropertiesB;
 		DecoratorPropertiesB.AddDefaulted(NodeTemplateDecoratorListB.Num());
-		DecoratorPropertiesB[0].Add(TEXT("Child"), ToString<FDecoratorWithOneChild::FSharedData>(TEXT("Child"), FAnimNextDecoratorHandle(NodeHandles[0])));
+		DecoratorPropertiesB[0].Add(TEXT("Child"), ToString<FDecoratorWithOneChild::FSharedData>(TEXT("Child"), FAnimNextDecoratorHandle(NodeHandles[1])));
 
 		TArray<TMap<FString, FString>> DecoratorPropertiesC;
 		DecoratorPropertiesC.AddDefaulted(NodeTemplateDecoratorListC.Num());
-		FAnimNextDecoratorHandle ChildrenHandlesC[2] = { FAnimNextDecoratorHandle(NodeHandles[0]), FAnimNextDecoratorHandle(NodeHandles[1])};
+		FAnimNextDecoratorHandle ChildrenHandlesC[2] = { FAnimNextDecoratorHandle(NodeHandles[1]), FAnimNextDecoratorHandle(NodeHandles[2])};
 		DecoratorPropertiesC[0].Add(TEXT("Children"), ToString<FDecoratorWithChildren::FSharedData>(TEXT("Children"), ChildrenHandlesC));
 
 		DecoratorWriter.BeginNodeWriting();
 		DecoratorWriter.WriteNode(NodeHandles[0],
-			[&DecoratorPropertiesA](uint32 DecoratorIndex, const FString& PropertyName)
+			[&DecoratorPropertiesC](uint32 DecoratorIndex, const FString& PropertyName)
 			{
-				return DecoratorPropertiesA[DecoratorIndex][PropertyName];
+				return DecoratorPropertiesC[DecoratorIndex][PropertyName];
 			},
 			[](uint32 DecoratorIndex, const FString& PropertyName)
 			{
 				return false;
 			});
 		DecoratorWriter.WriteNode(NodeHandles[1],
-			[&DecoratorPropertiesB](uint32 DecoratorIndex, const FString& PropertyName)
+			[&DecoratorPropertiesA](uint32 DecoratorIndex, const FString& PropertyName)
 			{
-				return DecoratorPropertiesB[DecoratorIndex][PropertyName];
+				return DecoratorPropertiesA[DecoratorIndex][PropertyName];
 			},
 			[](uint32 DecoratorIndex, const FString& PropertyName)
 			{
 				return false;
 			});
 		DecoratorWriter.WriteNode(NodeHandles[2],
-			[&DecoratorPropertiesC](uint32 DecoratorIndex, const FString& PropertyName)
+			[&DecoratorPropertiesB](uint32 DecoratorIndex, const FString& PropertyName)
 			{
-				return DecoratorPropertiesC[DecoratorIndex][PropertyName];
+				return DecoratorPropertiesB[DecoratorIndex][PropertyName];
 			},
 			[](uint32 DecoratorIndex, const FString& PropertyName)
 			{
@@ -509,15 +509,8 @@ bool FAnimationAnimNextRuntimeTest_IUpdate::RunTest(const FString& InParameters)
 
 	{
 		TArray<FDecoratorUID> UpdatedDecorators;
-
 		Private::UpdatedDecorators = &UpdatedDecorators;
 
-		FWeakDecoratorPtr NullPtr;								// Empty, no parent
-		FAnimNextDecoratorHandle RootHandle(NodeHandles[2]);	// Point to NodeC, first base decorator
-
-		FDecoratorPtr NodeCPtr = Context.AllocateNodeInstance(NullPtr, RootHandle);
-		AddErrorIfFalse(NodeCPtr.IsValid(), "FAnimationAnimNextRuntimeTest_IUpdate -> Failed to allocate root node instance");
-
 		// Call pre/post update on our graph
 		UpdateGraph(GraphInstance, 0.0333f);
 
@@ -580,7 +573,7 @@ bool FAnimationAnimNextRuntimeTest_IEvaluate::RunTest(const FString& InParameter
 	// Build our graph, it as follow (each node template has a single node instance):
 	// NodeA has no children
 	// NodeB has one child: NodeA (it doesn't evaluate)
-	// NodeD has two children: NodeA and NodeB
+	// NodeC (root) has two children: NodeA and NodeB
 
 	TArray<FNodeHandle> NodeHandles;
 
@@ -590,9 +583,9 @@ bool FAnimationAnimNextRuntimeTest_IEvaluate::RunTest(const FString& InParameter
 	{
 		FDecoratorWriter DecoratorWriter;
 
+		NodeHandles.Add(DecoratorWriter.RegisterNode(*NodeTemplateC));	// Root node
 		NodeHandles.Add(DecoratorWriter.RegisterNode(*NodeTemplateA));
 		NodeHandles.Add(DecoratorWriter.RegisterNode(*NodeTemplateB));
-		NodeHandles.Add(DecoratorWriter.RegisterNode(*NodeTemplateC));
 
 		// We don't have decorator properties
 		TArray<TMap<FString, FString>> DecoratorPropertiesA;
@@ -600,37 +593,37 @@ bool FAnimationAnimNextRuntimeTest_IEvaluate::RunTest(const FString& InParameter
 
 		TArray<TMap<FString, FString>> DecoratorPropertiesB;
 		DecoratorPropertiesB.AddDefaulted(NodeTemplateDecoratorListB.Num());
-		DecoratorPropertiesB[0].Add(TEXT("Child"), ToString<FDecoratorWithOneChild::FSharedData>(TEXT("Child"), FAnimNextDecoratorHandle(NodeHandles[0])));
+		DecoratorPropertiesB[0].Add(TEXT("Child"), ToString<FDecoratorWithOneChild::FSharedData>(TEXT("Child"), FAnimNextDecoratorHandle(NodeHandles[1])));
 
 		TArray<TMap<FString, FString>> DecoratorPropertiesC;
 		DecoratorPropertiesC.AddDefaulted(NodeTemplateDecoratorListC.Num());
 
-		FAnimNextDecoratorHandle ChildrenHandlesC[2] = { FAnimNextDecoratorHandle(NodeHandles[0]), FAnimNextDecoratorHandle(NodeHandles[1])};
+		FAnimNextDecoratorHandle ChildrenHandlesC[2] = { FAnimNextDecoratorHandle(NodeHandles[1]), FAnimNextDecoratorHandle(NodeHandles[2])};
 		DecoratorPropertiesC[0].Add(TEXT("Children"), ToString<FDecoratorWithChildren::FSharedData>(TEXT("Children"), ChildrenHandlesC));
 
 		DecoratorWriter.BeginNodeWriting();
 		DecoratorWriter.WriteNode(NodeHandles[0],
-			[&DecoratorPropertiesA](uint32 DecoratorIndex, const FString& PropertyName)
+			[&DecoratorPropertiesC](uint32 DecoratorIndex, const FString& PropertyName)
 			{
-				return DecoratorPropertiesA[DecoratorIndex][PropertyName];
+				return DecoratorPropertiesC[DecoratorIndex][PropertyName];
 			},
 			[](uint32 DecoratorIndex, const FString& PropertyName)
 			{
 				return false;
 			});
 		DecoratorWriter.WriteNode(NodeHandles[1],
-			[&DecoratorPropertiesB](uint32 DecoratorIndex, const FString& PropertyName)
+			[&DecoratorPropertiesA](uint32 DecoratorIndex, const FString& PropertyName)
 			{
-				return DecoratorPropertiesB[DecoratorIndex][PropertyName];
+				return DecoratorPropertiesA[DecoratorIndex][PropertyName];
 			},
 			[](uint32 DecoratorIndex, const FString& PropertyName)
 			{
 				return false;
 			});
 		DecoratorWriter.WriteNode(NodeHandles[2],
-			[&DecoratorPropertiesC](uint32 DecoratorIndex, const FString& PropertyName)
+			[&DecoratorPropertiesB](uint32 DecoratorIndex, const FString& PropertyName)
 			{
-				return DecoratorPropertiesC[DecoratorIndex][PropertyName];
+				return DecoratorPropertiesB[DecoratorIndex][PropertyName];
 			},
 			[](uint32 DecoratorIndex, const FString& PropertyName)
 			{
@@ -651,16 +644,8 @@ bool FAnimationAnimNextRuntimeTest_IEvaluate::RunTest(const FString& InParameter
 
 	{
 		TArray<FDecoratorUID> EvaluatedDecorators;
-
 		Private::EvaluatedDecorators = &EvaluatedDecorators;
 
-		FWeakDecoratorPtr NullPtr;								// Empty, no parent
-		FAnimNextDecoratorHandle RootHandle(NodeHandles[2]);	// Point to NodeC, first base decorator
-
-		FExecutionContext Context(GraphInstance);
-		FDecoratorPtr NodeCPtr = Context.AllocateNodeInstance(NullPtr, RootHandle);
-		AddErrorIfFalse(NodeCPtr.IsValid(), "FAnimationAnimNextRuntimeTest_IEvaluate -> Failed to allocate root node instance");
-
 		// Call pre/post evaluate on our graph
 		(void)EvaluateGraph(GraphInstance);
 
