AnimNext: Updated how external parameters are accessed

Introduced interface: IParameterSource.
- Implementations of this act as a proxy for the UObjects (and others) that they represent.
- Implementations added include FObjectProxy (for external objects), FPropertyBagProxy (for user params) and FParameterBlockProxy (for parameter blocks).

Introduced interface: IParameterSourceFactory.
- Modular feature allowing external code to extend how parameters are generated.
- Sources are identified by their 'source name' and provide a number of parameters via IParameterSource that is instanced at runtime by the factory.
- Source names are determined at compilation time by the schedule compiler from required parameters.

Introduced concept: Required parameter propagation
- Each task in a schedule that uses parameters now supplies the parameters that it (and eventually its dependents) need
- On compiliation, the earliest use case for these parameters is determined and a task inserted to fetch the parameter via its source
- Tasks run on both game and worker threads depending on individual parameter safety. Sources that provide both thread-safe and non-thread-safe params are split across different tasks.

Added the ability to augment the API of existing UObjects with static 'hoisted' library methods (using UAnimNextClassExtensionLibrary).
Removed lazy initialization of instance data as initialization now has a dependency on instance data (for proxy instantiation). This makes the need to move initialization off the game thread even more pressing.
Removed lazy initialization of parameter source data, as this is now needed during initialization.
Moved schedule compilation into UncookedOnly module.
Removed old adapter system.
Removed 'object' stack layer, as this is now handled by FObjectProxy. Other kinds of stack layers could be similarly removed eventually.
External tasks now only need one parameter to function - they now use a transient 'task binding' struct to find the object & tick function required
Changed logging to structured logs in some places.

Misc other fixes:
- Allowed 'none' to be selected in specified parameter pickers (via property metadata)
- Removed details customization's check() that was firing when selecting old data in a parameter block

diff --git a/Engine/Plugins/Experimental/AnimNext/Config/BaseAnimNext.ini b/Engine/Plugins/Experimental/AnimNext/Config/BaseAnimNext.ini
--- a/Engine/Plugins/Experimental/AnimNext/Config/BaseAnimNext.ini
+++ b/Engine/Plugins/Experimental/AnimNext/Config/BaseAnimNext.ini
@@ -1,4 +1,4 @@
 [/Script/AnimNext.AnimNextConfig]
-+ExposedClasses=(Class="/Script/CoreUObject.Class'/Script/Engine.CharacterMovementComponent'",RootParameter="UE_CharacterMovementComponent",FilterType=AllowAllPropertiesAndFunctions,FilteredFields=,bRegisterTickFunction=True)
-+ExposedClasses=(Class="/Script/CoreUObject.Class'/Script/AnimNext.AnimNextMeshComponent'",RootParameter="UE_AnimNextMeshComponent",FilterType=AllowAllPropertiesAndFunctions,FilteredFields=,bRegisterTickFunction=True)
-
++ExposedClasses=(Class="/Script/Engine.World",AccessorName="UE_World")
++ExposedClasses=(Class="/Script/Engine.CharacterMovementComponent",AccessorName="UE_CharacterMovementComponent")
++ExposedClasses=(Class="/Script/AnimNext.AnimNextMeshComponent",AccessorName="UE_AnimNextMeshComponent")
diff --git a/Engine/Plugins/Experimental/AnimNext/Source/AnimNext/Internal/AnimNextConfig.h b/Engine/Plugins/Experimental/AnimNext/Source/AnimNext/Internal/AnimNextConfig.h
--- a/Engine/Plugins/Experimental/AnimNext/Source/AnimNext/Internal/AnimNextConfig.h
+++ b/Engine/Plugins/Experimental/AnimNext/Source/AnimNext/Internal/AnimNextConfig.h
@@ -2,12 +2,12 @@
 
 #pragma once
 
-#include "Param/AnimNextObjectAdapterConfig.h"
+#include "Param/AnimNextObjectAccessorConfig.h"
 #include "AnimNextConfig.generated.h"
 
 namespace UE::AnimNext
 {
-	struct FParamId;
+	class FObjectProxyFactory;
 }
 
 UCLASS(Config=AnimNext)
@@ -22,9 +22,9 @@ private:
 #endif
 
 private:
-	friend struct UE::AnimNext::FParamId;
+	friend class UE::AnimNext::FObjectProxyFactory;
 
 	/** The classes that are exposed to AnimNext systems */
 	UPROPERTY(Config, EditAnywhere, Category = "Exposed Classes")
-	TArray<FAnimNextObjectAdapterConfig> ExposedClasses;
+	TArray<FAnimNextObjectAccessorConfig> ExposedClasses;
 };
\ No newline at end of file
diff --git a/Engine/Plugins/Experimental/AnimNext/Source/AnimNext/Internal/Graph/AnimGraph/AnimNode_AnimNextParameters.h b/Engine/Plugins/Experimental/AnimNext/Source/AnimNext/Internal/Graph/AnimGraph/AnimNode_AnimNextParameters.h
--- a/Engine/Plugins/Experimental/AnimNext/Source/AnimNext/Internal/Graph/AnimGraph/AnimNode_AnimNextParameters.h
+++ b/Engine/Plugins/Experimental/AnimNext/Source/AnimNext/Internal/Graph/AnimGraph/AnimNode_AnimNextParameters.h
@@ -4,11 +4,14 @@
 
 #include "CoreMinimal.h"
 #include "Animation/AnimNodeBase.h"
-#include "Param/ParamStackLayerHandle.h"
 #include "AnimNode_AnimNextParameters.generated.h"
 
-class IAnimNextParameterSourceInterface;
-class UAnimNextParameters;
+namespace UE::AnimNext
+{
+	struct FParameterBlockProxy;
+}
+
+class UAnimNextParameterBlock;
 
 USTRUCT(BlueprintInternalUseOnly)
 struct FAnimNode_AnimNextParameters : public FAnimNode_Base
@@ -17,25 +20,25 @@ struct FAnimNode_AnimNextParameters : public FAnimNode_Base
 
 	friend class UAnimGraphNode_AnimNextParameters;
 
-	FAnimNode_AnimNextParameters() = default;
+	ANIMNEXT_API FAnimNode_AnimNextParameters();
 	FAnimNode_AnimNextParameters(const FAnimNode_AnimNextParameters& InOther);
 	FAnimNode_AnimNextParameters& operator=(const FAnimNode_AnimNextParameters& InOther);
 	FAnimNode_AnimNextParameters(FAnimNode_AnimNextParameters&& InOther) noexcept;
 	FAnimNode_AnimNextParameters& operator=(FAnimNode_AnimNextParameters&& InOther) noexcept;
-	virtual ~FAnimNode_AnimNextParameters() override = default;
+	ANIMNEXT_API virtual ~FAnimNode_AnimNextParameters() override;
 
 private:
 	UPROPERTY(EditAnywhere, Category = Links)
 	FPoseLink Source;
 
 	UPROPERTY(EditAnywhere, Category = Settings)
-	TScriptInterface<IAnimNextParameterSourceInterface> Parameters;
+	TObjectPtr<UAnimNextParameterBlock> Parameters;
 
 	// Cache previous param block so we know when it changes via pin
-	IAnimNextParameterSourceInterface* PreviousParameters = nullptr;
+	TObjectPtr<UAnimNextParameterBlock> PreviousParameters;
 
-	// Cached layer
-	UE::AnimNext::FParamStackLayerHandle ParamLayerHandle;
+	// Cached proxy
+	TUniquePtr<UE::AnimNext::FParameterBlockProxy> ParametersProxy;
 
 private:
 	// FAnimNode_Base
diff --git a/Engine/Plugins/Experimental/AnimNext/Source/AnimNext/Internal/Graph/AnimNextGraph.h b/Engine/Plugins/Experimental/AnimNext/Source/AnimNext/Internal/Graph/AnimNextGraph.h
--- a/Engine/Plugins/Experimental/AnimNext/Source/AnimNext/Internal/Graph/AnimNextGraph.h
+++ b/Engine/Plugins/Experimental/AnimNext/Source/AnimNext/Internal/Graph/AnimNextGraph.h
@@ -20,6 +20,7 @@ class UAnimNextGraph;
 class UAnimGraphNode_AnimNextGraph;
 struct FAnimNode_AnimNextGraph;
 struct FRigUnit_AnimNextGraphEvaluator;
+class UAnimNextSchedule;
 
 namespace UE::AnimNext
 {
@@ -151,6 +152,9 @@ class ANIMNEXT_API UAnimNextGraph :  public URigVMHost, public IAnimNextSchedule
 	GENERATED_BODY()
 
 public:
+	static const UE::AnimNext::FParamId DefaultReferencePoseId;
+	static const UE::AnimNext::FParamId DefaultCurrentLODId;
+
 	UAnimNextGraph(const FObjectInitializer& ObjectInitializer);
 
 	// UObject interface
@@ -192,7 +196,8 @@ protected:
 	friend struct FAnimNextGraphInstance;
 	friend class UAnimGraphNode_AnimNextGraph;
 	friend UE::AnimNext::FExecutionContext;
-
+	friend class UAnimNextSchedule;
+	
 #if WITH_EDITORONLY_DATA
 	mutable FRWLock GraphInstancesLock;
 
@@ -235,11 +240,11 @@ protected:
 
 	// The parameter to use to access the reference pose
 	UPROPERTY(EditAnywhere, Category = "Graph", meta=(CustomWidget = "ParamName", AllowedParamType = "FAnimNextGraphReferencePose"))
-	FName ReferencePose = TEXT("UE_AnimNextMeshComponent_ReferencePose");
+	FName ReferencePose = DefaultReferencePoseId.GetName();
 
 	// The parameter to use to access the current LOD
 	UPROPERTY(EditAnywhere, Category = "Graph", meta=(CustomWidget = "ParamName", AllowedParamType = "int32"))
-	FName CurrentLOD = TEXT("UE_AnimNextMeshComponent_PredictedLODLevel");
+	FName CurrentLOD = DefaultCurrentLODId.GetName();
 
 	UE::AnimNext::FParamId ReferencePoseId = UE::AnimNext::FParamId(ReferencePose);
 	UE::AnimNext::FParamId CurrentLODId = UE::AnimNext::FParamId(CurrentLOD);
diff --git a/Engine/Plugins/Experimental/AnimNext/Source/AnimNext/Internal/Param/AnimNextObjectAccessorConfig.h b/Engine/Plugins/Experimental/AnimNext/Source/AnimNext/Internal/Param/AnimNextObjectAccessorConfig.h
new file mode 100644
--- /dev/null
+++ b/Engine/Plugins/Experimental/AnimNext/Source/AnimNext/Internal/Param/AnimNextObjectAccessorConfig.h
@@ -0,0 +1,25 @@
+// Copyright Epic Games, Inc. All Rights Reserved.
+
+#pragma once
+
+#include "CoreMinimal.h"
+#include "Templates/SubclassOf.h"
+#include "AnimNextObjectAccessorConfig.generated.h"
+
+USTRUCT()
+struct FAnimNextObjectAccessorConfig
+{
+	GENERATED_BODY()
+
+	/** The class to expose to AnimNext */
+	UPROPERTY(EditAnywhere, Config, Category="Adapter")
+	TSoftClassPtr<UObject> Class;
+
+	/**
+	 * The 'namespace' for the object when exposed to AnimNext
+	 * If Accessor Name is set to 'MyObject', then all properties and functions derived from this class will have the
+	 * form 'MyObject.Param'.
+	 */
+	UPROPERTY(EditAnywhere, Config, Category="Adapter", meta = (CustomWidget = "ParamName"))
+	FName AccessorName;
+};
diff --git a/Engine/Plugins/Experimental/AnimNext/Source/AnimNext/Internal/Param/AnimNextObjectAdapterConfig.h b/Engine/Plugins/Experimental/AnimNext/Source/AnimNext/Internal/Param/AnimNextObjectAdapterConfig.h
deleted file mode 100644
diff --git a/Engine/Plugins/Experimental/AnimNext/Source/AnimNext/Internal/Param/AnimNextParameterBlock.h b/Engine/Plugins/Experimental/AnimNext/Source/AnimNext/Internal/Param/AnimNextParameterBlock.h
--- a/Engine/Plugins/Experimental/AnimNext/Source/AnimNext/Internal/Param/AnimNextParameterBlock.h
+++ b/Engine/Plugins/Experimental/AnimNext/Source/AnimNext/Internal/Param/AnimNextParameterBlock.h
@@ -5,7 +5,7 @@
 #include "CoreMinimal.h"
 #include "PropertyBag.h"
 #include "RigVMCore/RigVMExecuteContext.h"
-#include "Param/IAnimNextParameterSourceInterface.h"
+#include "Param/IParameterSource.h"
 #include "RigVMHost.h"
 #include "AnimNextParameterBlock.generated.h"
 
@@ -16,6 +16,7 @@ struct FAnimNextScheduleParamScopeTask;
 namespace UE::AnimNext
 {
 	struct FContext;
+	struct FParameterBlockProxy;
 }
 
 namespace UE::AnimNext::UncookedOnly
@@ -34,7 +35,7 @@ namespace UE::AnimNext::Editor
 
 /** An asset used to define AnimNext parameters and their bindings */
 UCLASS(MinimalAPI, BlueprintType)
-class UAnimNextParameterBlock : public URigVMHost, public IAnimNextParameterSourceInterface
+class UAnimNextParameterBlock : public URigVMHost
 {
 	GENERATED_BODY()
 
@@ -51,11 +52,9 @@ class UAnimNextParameterBlock : public URigVMHost, public IAnimNextParameterSour
 	friend struct FAnimNextScheduleParamScopeEntryTask;
 	friend class UE::AnimNext::Editor::SParameterBlockViewRow;
 	friend class UE::AnimNext::Editor::FParameterBlockParameterCustomization;
+	friend struct UE::AnimNext::FParameterBlockProxy;
 
-	// IAnimNextParameterSourceInterface interface
-	virtual void UpdateLayer(UE::AnimNext::FParamStackLayerHandle& InHandle) const override;
-	virtual UE::AnimNext::FParamStackLayerHandle CacheLayer() const override;
-	virtual bool ShouldCacheLayer(const UE::AnimNext::FParamStackLayerHandle& InHandle) const override;
+	void UpdateLayer(UE::AnimNext::FParamStackLayerHandle& InHandle) const;
 
 	// UObject interface
 	virtual void BeginDestroy() override;
@@ -64,13 +63,11 @@ class UAnimNextParameterBlock : public URigVMHost, public IAnimNextParameterSour
 
 	FInstancedPropertyBag& GetPropertyBag() { return PropertyBag; }
 
-
 	FRigVMExtendedExecuteContext BaseRigVMContext;
 
 	UPROPERTY()
 	TObjectPtr<URigVM> RigVM;
 
-
 	UPROPERTY()
 	FInstancedPropertyBag PropertyBag;
 
diff --git a/Engine/Plugins/Experimental/AnimNext/Source/AnimNext/Internal/Param/AnimNextParameterCollection.h b/Engine/Plugins/Experimental/AnimNext/Source/AnimNext/Internal/Param/AnimNextParameterCollection.h
deleted file mode 100644
diff --git a/Engine/Plugins/Experimental/AnimNext/Source/AnimNext/Internal/Param/AnimNextParameterSourceRef.h b/Engine/Plugins/Experimental/AnimNext/Source/AnimNext/Internal/Param/AnimNextParameterSourceRef.h
deleted file mode 100644
diff --git a/Engine/Plugins/Experimental/AnimNext/Source/AnimNext/Internal/Param/ExternalParameterRegistry.h b/Engine/Plugins/Experimental/AnimNext/Source/AnimNext/Internal/Param/ExternalParameterRegistry.h
new file mode 100644
--- /dev/null
+++ b/Engine/Plugins/Experimental/AnimNext/Source/AnimNext/Internal/Param/ExternalParameterRegistry.h
@@ -0,0 +1,74 @@
+﻿// Copyright Epic Games, Inc. All Rights Reserved.
+
+#pragma once
+
+#include "Param/IParameterSource.h"
+#include "Features/IModularFeature.h"
+#include "Param/IParameterSourceFactory.h"
+
+class UAnimNextSchedule;
+class UAnimNextParameterBlockParameter;
+class AActor;
+class UActorComponent;
+
+namespace UE::AnimNext
+{
+	struct FObjectAccessor;
+	struct FExternalParameterContext;
+}
+
+namespace UE::AnimNext::Editor
+{
+	class SParameterPicker;
+}
+
+namespace UE::AnimNext::UncookedOnly
+{
+	struct FUtils;
+}
+
+namespace UE::AnimNext
+{
+
+// Registry of all object accessors 
+struct FExternalParameterRegistry
+{
+	// Register all external parameter sources
+	static void Init();
+
+	// Unregister all external parameter sources
+	static void Destroy();
+
+	// Factory method used to create a parameter source of the specified name, with a set of parameters that are initially required
+	// @param   InContext              Context used to set up the parameter source
+	// @param   InSourceName           The name of the parameter source. This is conventionally the 'root' parameter of a set of parameters (e.g.
+	//                                 Game_MyComponent, which contains Game_MyComponent_Param etc.)
+	// @param   InRequiredParameters   Any required parameters that the source should initially supply, can be empty
+	// @return a new parameter source, or nullptr if the source was not found
+	static TUniquePtr<IParameterSource> CreateParameterSource(const FExternalParameterContext& InContext, FName InSourceName, TConstArrayView<FName> InRequiredParameters);
+
+private:
+	friend class ::UAnimNextSchedule;
+	friend class ::UAnimNextParameterBlockParameter;
+	friend class UE::AnimNext::Editor::SParameterPicker;
+	friend struct UE::AnimNext::UncookedOnly::FUtils;
+
+#if WITH_EDITOR
+	// Given a parameter name, find associated info for that parameter
+	// @param   InParameterName        The parameter name to find
+	// @param   OutInfo                The parameter's info
+	// @return true if the parameter info was found
+	static ANIMNEXT_API bool FindParameterInfo(FName InParameterName, IParameterSourceFactory::FParameterInfo& OutInfo);
+
+	// Iterate over all the known parameters
+	// @param   InFunction             Function called for each known parameter
+	static ANIMNEXT_API void ForEachParameter(TFunctionRef<void(FName, const IParameterSourceFactory::FParameterInfo&)> InFunction);
+
+	// Given a parameter name, find the parameter source name for any registered source factories
+	// @param   InParameterName        The parameter name to find
+	// @return the name of the parameter source for the parameter name, or NAME_None if none could be found
+	static ANIMNEXT_API FName FindSourceForParameter(FName InParameterName);
+#endif
+};
+
+}
diff --git a/Engine/Plugins/Experimental/AnimNext/Source/AnimNext/Internal/Param/ParamDefinition.h b/Engine/Plugins/Experimental/AnimNext/Source/AnimNext/Internal/Param/ParamDefinition.h
--- a/Engine/Plugins/Experimental/AnimNext/Source/AnimNext/Internal/Param/ParamDefinition.h
+++ b/Engine/Plugins/Experimental/AnimNext/Source/AnimNext/Internal/Param/ParamDefinition.h
@@ -7,12 +7,6 @@
 #include "Param/ParamType.h"
 #include "Param/ParamTypeHandle.h"
 
-namespace UE::AnimNext
-{
-	struct FParams;
-	struct FParamAdapter;
-}
-
 namespace UE::AnimNext::UncookedOnly
 {
 	struct FUtils;
@@ -107,8 +101,6 @@ struct FParamDefinition
 
 private:
 	friend struct FParamId;
-	friend struct FParams;
-	friend struct FParamAdapter;
 	friend struct UncookedOnly::FUtils;
 
 	FParamId Id;
diff --git a/Engine/Plugins/Experimental/AnimNext/Source/AnimNext/Internal/Param/Params.h b/Engine/Plugins/Experimental/AnimNext/Source/AnimNext/Internal/Param/Params.h
deleted file mode 100644
diff --git a/Engine/Plugins/Experimental/AnimNext/Source/AnimNext/Internal/Scheduler/AnimNextSchedule.h b/Engine/Plugins/Experimental/AnimNext/Source/AnimNext/Internal/Scheduler/AnimNextSchedule.h
--- a/Engine/Plugins/Experimental/AnimNext/Source/AnimNext/Internal/Scheduler/AnimNextSchedule.h
+++ b/Engine/Plugins/Experimental/AnimNext/Source/AnimNext/Internal/Scheduler/AnimNextSchedule.h
@@ -7,6 +7,7 @@
 #include "Scheduler/AnimNextSchedulePortTask.h"
 #include "Scheduler/AnimNextScheduleExternalTask.h"
 #include "Scheduler/AnimNextScheduleParamScopeTask.h"
+#include "Scheduler/AnimNextScheduleExternalParamTask.h"
 #include "Tasks/Task.h"
 #include "Scheduler/IAnimNextScheduleTermInterface.h"
 #include "AnimNextSchedule.generated.h"
@@ -28,6 +29,12 @@ namespace UE::AnimNext
 	struct FScheduleTickFunction;
 }
 
+namespace UE::AnimNext::UncookedOnly
+{
+	class FModule;
+	struct FUtils;
+}
+
 USTRUCT()
 struct FAnimNextScheduleEntryTerm
 {
@@ -51,47 +58,48 @@ struct FAnimNextScheduleEntryTerm
 	EScheduleTermDirection Direction = EScheduleTermDirection::Input;
 };
 
-UCLASS(EditInlineNew, Abstract)
+UCLASS(MinimalAPI, EditInlineNew, Abstract)
 class UAnimNextScheduleEntry : public UObject
 {
 	GENERATED_BODY()
-
-private:
-	friend class UAnimNextSchedule;
 };
 
-UCLASS(DisplayName="Graph")
+UCLASS(MinimalAPI, DisplayName="Graph")
 class UAnimNextScheduleEntry_AnimNextGraph : public UAnimNextScheduleEntry
 {
 	GENERATED_BODY()
 
 private:
-	friend class UAnimNextSchedule;
+	friend struct UE::AnimNext::UncookedOnly::FUtils;
 
 	// The graph to run by default
 	UPROPERTY(EditAnywhere, Category = "Graph")
 	TObjectPtr<UAnimNextGraph> Graph = nullptr;
 
 	// Parameter to get the graph from dynamically
-	UPROPERTY(EditAnywhere, Category = "Graph", meta = (CustomWidget = "ParamName", AllowedParamType = "TObjectPtr<UAnimNextGraph>"))
+	UPROPERTY(EditAnywhere, Category = "Graph", meta = (CustomWidget = "ParamName", AllowedParamType = "TObjectPtr<UAnimNextGraph>", AllowNone))
 	FName DynamicGraph;
 
 	// An optional entry point to use when running the supplied graph
-	UPROPERTY(EditAnywhere, Category = "Graph", meta = (CustomWidget = "ParamName", AllowedParamType = "FName"))
+	UPROPERTY(EditAnywhere, Category = "Graph", meta = (CustomWidget = "ParamName", AllowedParamType = "FName", AllowNone))
 	FName EntryPoint;
 
+	// All parameters that are required by this graph to run (only required if dynamic as static graph params can be discovered by the compiler)
+	UPROPERTY(EditAnywhere, Category = "Graph", meta = (CustomWidget = "ParamName"))
+	TArray<FName> RequiredParameters;
+
 	// The intermediate terms used by the graph
 	UPROPERTY(EditAnywhere, Category = "Graph")
 	TArray<FAnimNextScheduleEntryTerm> Terms;
 };
 
-UCLASS(DisplayName="Port")
+UCLASS(MinimalAPI, DisplayName="Port")
 class UAnimNextScheduleEntry_Port : public UAnimNextScheduleEntry
 {
 	GENERATED_BODY()
 
 private:
-	friend class UAnimNextSchedule;
+	friend struct UE::AnimNext::UncookedOnly::FUtils;
 
 	// The type of the port to use
 	UPROPERTY(EditAnywhere, Category = "Port", meta = (ShowDisplayNames))
@@ -102,30 +110,26 @@ private:
 	TArray<FAnimNextScheduleEntryTerm> Terms;
 };
 
-UCLASS(DisplayName="External")
+UCLASS(MinimalAPI, DisplayName="External")
 class UAnimNextScheduleEntry_ExternalTask : public UAnimNextScheduleEntry
 {
 	GENERATED_BODY()
 
 private:
-	friend class UAnimNextSchedule;
+	friend struct UE::AnimNext::UncookedOnly::FUtils;
 
-	// The tick function that this external task should wrap
-	UPROPERTY(EditAnywhere, Category = "External Task", meta = (CustomWidget = "ParamName", AllowedParamType = "FTickFunction"))
-	FName TickFunction;
-
-	// The object that the tick function is present on
-	UPROPERTY(EditAnywhere, Category = "External Task", meta = (CustomWidget = "ParamName", AllowedParamType = "TObjectPtr<UObject>"))
-	FName Object;
+	// The external task we wrap
+	UPROPERTY(EditAnywhere, Category = "External Task", meta = (CustomWidget = "ParamName", AllowedParamType = "FAnimNextExternalTaskBinding"))
+	FName ExternalTask;
 };
 
-UCLASS(DisplayName="Scope")
+UCLASS(MinimalAPI, DisplayName="Scope")
 class UAnimNextScheduleEntry_ParamScope : public UAnimNextScheduleEntry
 {
 	GENERATED_BODY()
 
 private:
-	friend class UAnimNextSchedule;
+	friend struct UE::AnimNext::UncookedOnly::FUtils;
 
 	// The scope to use
 	UPROPERTY(EditAnywhere, Category = "Parameters", meta = (CustomWidget = "ParamName", AllowedParamType = "FAnimNextScope"))
@@ -140,22 +144,43 @@ private:
 	TArray<TObjectPtr<UAnimNextScheduleEntry>> SubEntries;
 };
 
+// This entry is inserted only by the compilation machinery
+UCLASS(MinimalAPI, DisplayName="External Parameters [INTERNAL]")
+class UAnimNextScheduleEntry_ExternalParams : public UAnimNextScheduleEntry
+{
+	GENERATED_BODY()
+
+private:
+	friend struct UE::AnimNext::UncookedOnly::FUtils;
+	
+	// Parameter sources to use
+	UPROPERTY()
+	TArray<FAnimNextScheduleExternalParameterSource> ParameterSources;
+
+	// Whether the task can run on a worker thread
+	UPROPERTY()
+	bool bThreadSafe = false;
+
+};
+
 // TEMP: opcode in schedule
 UENUM()
 enum class EAnimNextScheduleScheduleOpcode : uint8
 {
 	None,
-	RunTask,				// Operand = Task Index
+	RunGraphTask,			// Operand = Task Index
 	BeginRunExternalTask,	// Operand = Task Index
 	EndRunExternalTask,		// Operand = Task Index
 	RunPort,				// Operand = Port Index
 	RunParamScopeEntry,		// Operand = Scope Index
 	RunParamScopeExit,		// Operand = Scope Index
+	RunExternalParamTask,	// Operand = Task Index
 	PrerequisiteTask,		// Operand = Task Index
 	PrerequisiteBeginExternalTask, // Operand = Task Index
 	PrerequisiteEndExternalTask, // Operand = Task Index
 	PrerequisiteScopeEntry,	// Operand = Scope Index
 	PrerequisiteScopeExit,	// Operand = Scope Index
+	PrerequisiteExternalParamTask, // Operand = Task Index
 	Exit,					// Operand = 0
 };
 
@@ -202,6 +227,8 @@ private:
 	friend struct FAnimNextSchedulerEntry;
 	friend class UAnimNextComponent;
 	friend class UAnimNextSchedulerWorldSubsystem;
+	friend class UE::AnimNext::UncookedOnly::FModule;
+	friend struct UE::AnimNext::UncookedOnly::FUtils;
 
 	// UObject interface
 	virtual void PostLoad() override;
@@ -214,6 +241,9 @@ private:
 #endif
 
 #if WITH_EDITORONLY_DATA
+	// Function hook used to compile in editor/cooker
+	static TUniqueFunction<void(UAnimNextSchedule*)> CompileFunction;
+
 	// Editor only
 	// TODO: move this into an editor only subobject
 	// TODO: this is currently only a linear list, we want it to be a graph
@@ -227,7 +257,7 @@ private:
 
 	// TEMP: Tasks derived from the entries above
 	UPROPERTY(NonTransactional)
-	TArray<FAnimNextScheduleGraphTask> Tasks;
+	TArray<FAnimNextScheduleGraphTask> GraphTasks;
 
 	// TEMP: Ports derived from the entries above
 	UPROPERTY(NonTransactional)
@@ -245,6 +275,10 @@ private:
 	UPROPERTY(NonTransactional)
 	TArray<FAnimNextScheduleParamScopeExitTask> ParamScopeExitTasks;
 
+	// TEMP: External parameter tasks derived from the entries above
+	UPROPERTY(NonTransactional)
+	TArray<FAnimNextScheduleExternalParamTask> ExternalParamTasks;
+	
 	// TEMP: Data for intermediates, defined as a property bag
 	UPROPERTY(NonTransactional)
 	FInstancedPropertyBag IntermediatesData;
diff --git a/Engine/Plugins/Experimental/AnimNext/Source/AnimNext/Internal/Scheduler/AnimNextScheduleExternalParamTask.h b/Engine/Plugins/Experimental/AnimNext/Source/AnimNext/Internal/Scheduler/AnimNextScheduleExternalParamTask.h
new file mode 100644
--- /dev/null
+++ b/Engine/Plugins/Experimental/AnimNext/Source/AnimNext/Internal/Scheduler/AnimNextScheduleExternalParamTask.h
@@ -0,0 +1,60 @@
+// Copyright Epic Games, Inc. All Rights Reserved.
+
+#pragma once
+
+#include "CoreMinimal.h"
+#include "AnimNextScheduleExternalParamTask.generated.h"
+
+class UAnimNextSchedule;
+struct FAnimNextSchedulerEntry;
+
+namespace UE::AnimNext
+{
+	struct FScheduleContext;
+	struct FScheduleInstanceData;
+	struct FScheduleTickFunction;
+}
+
+namespace UE::AnimNext::UncookedOnly
+{
+	struct FUtils;
+}
+
+USTRUCT()
+struct FAnimNextScheduleExternalParameterSource
+{
+	GENERATED_BODY()
+
+	// Parameter source name to use
+	UPROPERTY()
+	FName ParameterSource;
+
+	// Names of all parameters that are required in this task
+	UPROPERTY()
+	TArray<FName> Parameters;
+};
+
+USTRUCT()
+struct FAnimNextScheduleExternalParamTask
+{
+	GENERATED_BODY()
+
+	FAnimNextScheduleExternalParamTask() = default;
+
+	void UpdateExternalParams(const UE::AnimNext::FScheduleContext& InScheduleContext) const;
+
+private:
+	friend struct UE::AnimNext::UncookedOnly::FUtils;
+	friend struct UE::AnimNext::FScheduleInstanceData;
+	friend struct UE::AnimNext::FScheduleTickFunction;
+	friend struct FAnimNextSchedulerEntry;
+
+	UPROPERTY()
+	uint32 TaskIndex = MAX_uint32;
+
+	UPROPERTY()
+	TArray<FAnimNextScheduleExternalParameterSource> ParameterSources;
+
+	UPROPERTY()
+	bool bThreadSafe = false;
+};
diff --git a/Engine/Plugins/Experimental/AnimNext/Source/AnimNext/Internal/Scheduler/AnimNextScheduleExternalTask.h b/Engine/Plugins/Experimental/AnimNext/Source/AnimNext/Internal/Scheduler/AnimNextScheduleExternalTask.h
--- a/Engine/Plugins/Experimental/AnimNext/Source/AnimNext/Internal/Scheduler/AnimNextScheduleExternalTask.h
+++ b/Engine/Plugins/Experimental/AnimNext/Source/AnimNext/Internal/Scheduler/AnimNextScheduleExternalTask.h
@@ -15,6 +15,11 @@ namespace UE::AnimNext
 	struct FScheduleTickFunction;
 }
 
+namespace UE::AnimNext::UncookedOnly
+{
+	struct FUtils;
+}
+
 USTRUCT()
 struct FAnimNextScheduleExternalTask
 {
@@ -23,7 +28,7 @@ struct FAnimNextScheduleExternalTask
 	FAnimNextScheduleExternalTask() = default;
 
 private:
-	friend class UAnimNextSchedule;
+	friend struct UE::AnimNext::UncookedOnly::FUtils;
 	friend class UAnimNextComponent;
 	friend struct FAnimNextSchedulerEntry;
 	friend struct UE::AnimNext::FScheduleTickFunction;
@@ -38,11 +43,6 @@ private:
 	UPROPERTY()
 	uint32 ParamParentScopeIndex = MAX_uint32;
 
-	/** The name of the external task parameter */
-	UPROPERTY()
-	FName TickFunction;
-
-	/** The name of the external task's object parameter */
 	UPROPERTY()
-	FName Object;
+	FName ExternalTask;
 };
diff --git a/Engine/Plugins/Experimental/AnimNext/Source/AnimNext/Internal/Scheduler/AnimNextScheduleGraphTask.h b/Engine/Plugins/Experimental/AnimNext/Source/AnimNext/Internal/Scheduler/AnimNextScheduleGraphTask.h
--- a/Engine/Plugins/Experimental/AnimNext/Source/AnimNext/Internal/Scheduler/AnimNextScheduleGraphTask.h
+++ b/Engine/Plugins/Experimental/AnimNext/Source/AnimNext/Internal/Scheduler/AnimNextScheduleGraphTask.h
@@ -14,6 +14,12 @@ namespace UE::AnimNext
 	struct FScheduleContext;
 	struct FScheduleInstanceData;
 	struct FScheduleTickFunction;
+	struct FParamStack;
+}
+
+namespace UE::AnimNext::UncookedOnly
+{
+	struct FUtils;
 }
 
 USTRUCT()
@@ -25,12 +31,14 @@ struct FAnimNextScheduleGraphTask
 
 private:
 	friend class UAnimNextComponent;
-	friend class UAnimNextSchedule;
+	friend struct UE::AnimNext::UncookedOnly::FUtils;
 	friend struct UE::AnimNext::FScheduleInstanceData;
 	friend struct UE::AnimNext::FScheduleTickFunction;
 
 	void RunGraph(const UE::AnimNext::FScheduleContext& InContext) const;
 
+	UAnimNextGraph* GetGraphToRun(UE::AnimNext::FParamStack& ParamStack) const;
+
 private:
 	UPROPERTY()
 	uint32 TaskIndex = MAX_uint32;
diff --git a/Engine/Plugins/Experimental/AnimNext/Source/AnimNext/Internal/Scheduler/AnimNextScheduleParamScopeTask.h b/Engine/Plugins/Experimental/AnimNext/Source/AnimNext/Internal/Scheduler/AnimNextScheduleParamScopeTask.h
--- a/Engine/Plugins/Experimental/AnimNext/Source/AnimNext/Internal/Scheduler/AnimNextScheduleParamScopeTask.h
+++ b/Engine/Plugins/Experimental/AnimNext/Source/AnimNext/Internal/Scheduler/AnimNextScheduleParamScopeTask.h
@@ -17,6 +17,11 @@ namespace UE::AnimNext
 	struct FScheduleTickFunction;
 }
 
+namespace UE::AnimNext::UncookedOnly
+{
+	struct FUtils;
+}
+
 USTRUCT()
 struct FAnimNextScheduleParamScopeEntryTask
 {
@@ -28,7 +33,7 @@ private:
 	void RunParamScopeEntry(const UE::AnimNext::FScheduleContext& InScheduleContext) const;
 
 private:
-	friend class UAnimNextSchedule;
+	friend struct UE::AnimNext::UncookedOnly::FUtils;
 	friend class UAnimNextSchedulerWorldSubsystem;
 	friend struct UE::AnimNext::FScheduleInstanceData;
 	friend struct UE::AnimNext::FScheduleTickFunction;
@@ -50,7 +55,7 @@ private:
 	UPROPERTY()
 	FName Scope;
 
-	/** Parameters to apply in this scope */
+	// Parameters to apply in this scope
 	UPROPERTY()
 	TArray<TObjectPtr<UAnimNextParameterBlock>> ParameterBlocks;
 };
@@ -66,7 +71,7 @@ private:
 	void RunParamScopeExit(const UE::AnimNext::FScheduleContext& InScheduleContext) const;
 
 private:
-	friend class UAnimNextSchedule;
+	friend struct UE::AnimNext::UncookedOnly::FUtils;
 	friend struct UE::AnimNext::FScheduleTickFunction;
 
 	UPROPERTY()
diff --git a/Engine/Plugins/Experimental/AnimNext/Source/AnimNext/Internal/Scheduler/AnimNextSchedulePortTask.h b/Engine/Plugins/Experimental/AnimNext/Source/AnimNext/Internal/Scheduler/AnimNextSchedulePortTask.h
--- a/Engine/Plugins/Experimental/AnimNext/Source/AnimNext/Internal/Scheduler/AnimNextSchedulePortTask.h
+++ b/Engine/Plugins/Experimental/AnimNext/Source/AnimNext/Internal/Scheduler/AnimNextSchedulePortTask.h
@@ -18,6 +18,11 @@ namespace UE::AnimNext
 	struct FScheduleInstanceData;
 }
 
+namespace UE::AnimNext::UncookedOnly
+{
+	struct FUtils;
+}
+
 USTRUCT()
 struct FAnimNextSchedulePortTask
 {
@@ -25,7 +30,7 @@ struct FAnimNextSchedulePortTask
 
 	friend struct UE::AnimNext::FScheduleTask;
 	friend struct UE::AnimNext::FScheduleTickFunction;
-	friend class UAnimNextSchedule;
+	friend struct UE::AnimNext::UncookedOnly::FUtils;
 	friend struct UE::AnimNext::FScheduleInstanceData;
 
 	FAnimNextSchedulePortTask() = default;
diff --git a/Engine/Plugins/Experimental/AnimNext/Source/AnimNext/Private/AnimNextConfig.cpp b/Engine/Plugins/Experimental/AnimNext/Source/AnimNext/Private/AnimNextConfig.cpp
--- a/Engine/Plugins/Experimental/AnimNext/Source/AnimNext/Private/AnimNextConfig.cpp
+++ b/Engine/Plugins/Experimental/AnimNext/Source/AnimNext/Private/AnimNextConfig.cpp
@@ -1,7 +1,7 @@
 ﻿// Copyright Epic Games, Inc. All Rights Reserved.
 
 #include "AnimNextConfig.h"
-#include "Param/ParamId.h"
+#include "Param/ObjectProxyFactory.h"
 
 #if WITH_EDITOR
 
@@ -11,7 +11,7 @@ void UAnimNextConfig::PostEditChangeProperty(FPropertyChangedEvent& PropertyChan
 	
 	if(PropertyChangedEvent.Property->GetFName() == GET_MEMBER_NAME_CHECKED(UAnimNextConfig, ExposedClasses))
 	{
-		FParamId::RefreshAdapters();
+		FObjectProxyFactory::Refresh();
 	}
 
 	SaveConfig();
diff --git a/Engine/Plugins/Experimental/AnimNext/Source/AnimNext/Private/Component/AnimNextComponent.cpp b/Engine/Plugins/Experimental/AnimNext/Source/AnimNext/Private/Component/AnimNextComponent.cpp
--- a/Engine/Plugins/Experimental/AnimNext/Source/AnimNext/Private/Component/AnimNextComponent.cpp
+++ b/Engine/Plugins/Experimental/AnimNext/Source/AnimNext/Private/Component/AnimNextComponent.cpp
@@ -3,8 +3,7 @@
 #include "AnimNextComponent.h"
 #include "Scheduler/Scheduler.h"
 #include "Scheduler/ScheduleContext.h"
-#include "Param/AnimNextParameterSourceRef.h"
-#include "Param/Params.h"
+#include "Param/PropertyBagProxy.h"
 
 void UAnimNextComponent::OnRegister()
 {
@@ -109,11 +108,11 @@ DEFINE_FUNCTION(UAnimNextComponent::execSetParameterInScope)
 
 	Stack.StepCompiledIn<FProperty>(nullptr);
 	const FProperty* ValueProp = CastField<FProperty>(Stack.MostRecentProperty);
-	const void* ValuePtr = Stack.MostRecentPropertyAddress;
+	const void* ContainerPtr = Stack.MostRecentPropertyContainer;
 
 	P_FINISH;
 
-	if (!ValueProp || !ValuePtr)
+	if (!ValueProp || !ContainerPtr)
 	{
 		FBlueprintExceptionInfo ExceptionInfo(
 			EBlueprintExceptionType::AbortExecution,
@@ -135,17 +134,19 @@ DEFINE_FUNCTION(UAnimNextComponent::execSetParameterInScope)
 	{
 		P_NATIVE_BEGIN;
 
-		FAnimNextParameterSourceRef ParamSource;
-		ParamSource.Type = EAnimNextParameterSourceRefType::Inline;
-		ParamSource.InlineParameters.AddProperty(Name, ValueProp);
-		ParamSource.InlineParameters.SetValue(Name, ParamSource.InlineParameters.GetPropertyBagStruct()->FindPropertyByName(Name), ValuePtr);
-
-		FScheduler::QueueTask(P_THIS, P_THIS->SchedulerHandle, Scope, [Scope, ParamSource = MoveTemp(ParamSource)](const FScheduleContext& InContext)
+		TUniquePtr<FInstancedPropertyBag> PropertyBag = MakeUnique<FInstancedPropertyBag>();
+		PropertyBag->AddProperty(Name, ValueProp);
+		FProperty* NewProperty = PropertyBag->GetPropertyBagStruct()->FindPropertyByName(Name);
+		PropertyBag->SetValue(Name, NewProperty, ContainerPtr);
+		
+		FScheduler::QueueTask(P_THIS, P_THIS->SchedulerHandle, Scope, [Scope, Name, NewProperty, PropertyBag = MoveTemp(PropertyBag)](const FScheduleContext& InContext) mutable
 		{
-			FAnimNextParameterCollection& Collection = InContext.GetInstanceData().UserScopes.FindOrAdd(Scope);
+			TUniquePtr<FPropertyBagProxy>& ScopeSources = InContext.GetInstanceData().UserScopes.FindOrAdd(Scope);
 
 			// TODO: pre/post scope distinction
-			Collection.Parameters.Insert(ParamSource, 0);
+
+			// Copy property to local bag
+			ScopeSources->AddPropertyAndValue(Name, NewProperty, PropertyBag->GetValue().GetMemory());
 		});
 		
 		P_NATIVE_END;
diff --git a/Engine/Plugins/Experimental/AnimNext/Source/AnimNext/Private/Component/AnimNextMeshComponent.cpp b/Engine/Plugins/Experimental/AnimNext/Source/AnimNext/Private/Component/AnimNextMeshComponent.cpp
--- a/Engine/Plugins/Experimental/AnimNext/Source/AnimNext/Private/Component/AnimNextMeshComponent.cpp
+++ b/Engine/Plugins/Experimental/AnimNext/Source/AnimNext/Private/Component/AnimNextMeshComponent.cpp
@@ -3,6 +3,8 @@
 #include "AnimNextMeshComponent.h"
 #include "GenerationTools.h"
 #include "Engine/World.h"
+#include "Graph/AnimNext_LODPose.h"
+#include "ReferencePose.h"
 
 UAnimNextMeshComponent::UAnimNextMeshComponent()
 {
@@ -37,3 +39,10 @@ void UAnimNextMeshComponent::CompleteAndDispatch(TConstArrayView<FBoneIndexType>
 		SendRenderDynamicData_Concurrent();
 	}
 }
+
+FAnimNextGraphReferencePose UAnimNextMeshComponent::GetReferencePose()
+{
+	UE::AnimNext::FDataHandle RefPoseHandle = UE::AnimNext::FDataRegistry::Get()->GetOrGenerateReferencePose(this);
+	const UE::AnimNext::FReferencePose& RefPose = RefPoseHandle.GetRef<UE::AnimNext::FReferencePose>();
+	return FAnimNextGraphReferencePose(&RefPose);
+}
\ No newline at end of file
diff --git a/Engine/Plugins/Experimental/AnimNext/Source/AnimNext/Private/Component/AnimNextMeshComponent.h b/Engine/Plugins/Experimental/AnimNext/Source/AnimNext/Private/Component/AnimNextMeshComponent.h
--- a/Engine/Plugins/Experimental/AnimNext/Source/AnimNext/Private/Component/AnimNextMeshComponent.h
+++ b/Engine/Plugins/Experimental/AnimNext/Source/AnimNext/Private/Component/AnimNextMeshComponent.h
@@ -5,6 +5,7 @@
 #include "Components/SkeletalMeshComponent.h"
 #include "AnimNextMeshComponent.generated.h"
 
+struct FAnimNextGraphReferencePose;
 class UAnimNextSchedulePort_AnimNextMeshComponentPose;
 
 // Mesh component for use with AnimNext
@@ -21,4 +22,8 @@ class UAnimNextMeshComponent : public USkeletalMeshComponent
 
 	// Converts local to component space, flips buffers, updates bounds and dispatches to renderer
 	void CompleteAndDispatch(TConstArrayView<FBoneIndexType> InParentIndices, TConstArrayView<FBoneIndexType> InRequiredBoneIndices, TConstArrayView<FTransform> InLocalSpaceTransforms);
+
+	// Access the ref pose of the component
+	UFUNCTION(BlueprintCallable, BlueprintInternalUseOnly, meta=(BlueprintThreadSafe))
+	FAnimNextGraphReferencePose GetReferencePose();
 };
\ No newline at end of file
diff --git a/Engine/Plugins/Experimental/AnimNext/Source/AnimNext/Private/Graph/AnimGraph/AnimNode_AnimNextParameters.cpp b/Engine/Plugins/Experimental/AnimNext/Source/AnimNext/Private/Graph/AnimGraph/AnimNode_AnimNextParameters.cpp
--- a/Engine/Plugins/Experimental/AnimNext/Source/AnimNext/Private/Graph/AnimGraph/AnimNode_AnimNextParameters.cpp
+++ b/Engine/Plugins/Experimental/AnimNext/Source/AnimNext/Private/Graph/AnimGraph/AnimNode_AnimNextParameters.cpp
@@ -2,16 +2,19 @@
 
 #include "Graph/AnimGraph/AnimNode_AnimNextParameters.h"
 #include "Param/ParamStack.h"
-#include "Param/IAnimNextParameterSourceInterface.h"
+#include "Param/AnimNextParameterBlock.h"
 #include "AnimGraphParamStackScope.h"
+#include "Param/ParameterBlockProxy.h"
 
 #include UE_INLINE_GENERATED_CPP_BY_NAME(AnimNode_AnimNextParameters)
 
+FAnimNode_AnimNextParameters::FAnimNode_AnimNextParameters() = default;
+
 FAnimNode_AnimNextParameters::FAnimNode_AnimNextParameters(const FAnimNode_AnimNextParameters& InOther)
 	: Source(InOther.Source)
 	, Parameters(InOther.Parameters)
 	, PreviousParameters(nullptr)
-	, ParamLayerHandle()
+	, ParametersProxy(nullptr)
 {
 }
 
@@ -20,7 +23,7 @@ FAnimNode_AnimNextParameters& FAnimNode_AnimNextParameters::operator=(const FAni
 	Source = InOther.Source;
 	Parameters = InOther.Parameters;
 	PreviousParameters = nullptr;
-	ParamLayerHandle.Invalidate();
+	ParametersProxy.Reset();
 	return *this;
 }
 
@@ -29,7 +32,7 @@ FAnimNode_AnimNextParameters::FAnimNode_AnimNextParameters(FAnimNode_AnimNextPar
 	Source = InOther.Source;
 	Parameters = InOther.Parameters;
 	PreviousParameters = nullptr;
-	ParamLayerHandle.Invalidate();
+	ParametersProxy.Reset();
 }
 
 FAnimNode_AnimNextParameters& FAnimNode_AnimNextParameters::operator=(FAnimNode_AnimNextParameters&& InOther) noexcept
@@ -37,10 +40,12 @@ FAnimNode_AnimNextParameters& FAnimNode_AnimNextParameters::operator=(FAnimNode_
 	Source = InOther.Source;
 	Parameters = InOther.Parameters;
 	PreviousParameters = nullptr;
-	ParamLayerHandle.Invalidate();
+	ParametersProxy.Reset();
 	return *this;
 }
 
+FAnimNode_AnimNextParameters::~FAnimNode_AnimNextParameters() = default;
+
 void FAnimNode_AnimNextParameters::Initialize_AnyThread(const FAnimationInitializeContext& Context)
 {
 	Source.Initialize(Context);
@@ -52,16 +57,16 @@ void FAnimNode_AnimNextParameters::Update_AnyThread(const FAnimationUpdateContex
 
 	GetEvaluateGraphExposedInputs().Execute(Context);
 
-	IAnimNextParameterSourceInterface* CurrentParameters = Parameters ? Parameters.GetInterface() : nullptr;
+	UAnimNextParameterBlock* CurrentParameters = Parameters;
 
 	// Reconstruct param block's cached layer if required
-	if (CurrentParameters != PreviousParameters || !ParamLayerHandle.IsValid())
+	if (CurrentParameters != PreviousParameters || !ParametersProxy.IsValid())
 	{
-		ParamLayerHandle.Invalidate();
+		ParametersProxy.Reset();
 
 		if (CurrentParameters)
 		{
-			ParamLayerHandle = CurrentParameters->CacheLayer();
+			ParametersProxy = MakeUnique<FParameterBlockProxy>(CurrentParameters);
 		}
 
 		PreviousParameters = CurrentParameters;
@@ -72,10 +77,10 @@ void FAnimNode_AnimNextParameters::Update_AnyThread(const FAnimationUpdateContex
 
 		FParamStack& ParamStack = FParamStack::Get();
 		FParamStack::FPushedLayerHandle PushedLayerHandle;
-		if (CurrentParameters && ParamLayerHandle.IsValid())
+		if (CurrentParameters && ParametersProxy.IsValid())
 		{
-			CurrentParameters->UpdateLayer(ParamLayerHandle);
-			PushedLayerHandle = ParamStack.PushLayer(ParamLayerHandle);
+			ParametersProxy->Update();
+			PushedLayerHandle = ParamStack.PushLayer(ParametersProxy->GetLayerHandle());
 		}
 
 		Source.Update(Context);
diff --git a/Engine/Plugins/Experimental/AnimNext/Source/AnimNext/Private/Graph/AnimNextGraph.cpp b/Engine/Plugins/Experimental/AnimNext/Source/AnimNext/Private/Graph/AnimNextGraph.cpp
--- a/Engine/Plugins/Experimental/AnimNext/Source/AnimNext/Private/Graph/AnimNextGraph.cpp
+++ b/Engine/Plugins/Experimental/AnimNext/Source/AnimNext/Private/Graph/AnimNextGraph.cpp
@@ -19,6 +19,9 @@ DEFINE_STAT(STAT_AnimNext_Graph_AllocateInstance);
 
 #include UE_INLINE_GENERATED_CPP_BY_NAME(AnimNextGraph)
 
+const UE::AnimNext::FParamId UAnimNextGraph::DefaultReferencePoseId = UE::AnimNext::FParamId("UE_AnimNextMeshComponent_ReferencePose");
+const UE::AnimNext::FParamId UAnimNextGraph::DefaultCurrentLODId = UE::AnimNext::FParamId("UE_AnimNextMeshComponent_PredictedLODLevel");
+
 #if WITH_EDITORONLY_DATA
 FAnimNextGraphInstance::FAnimNextGraphInstance(const FAnimNextGraphInstance& Other)
 {
diff --git a/Engine/Plugins/Experimental/AnimNext/Source/AnimNext/Private/Module.cpp b/Engine/Plugins/Experimental/AnimNext/Source/AnimNext/Private/Module.cpp
--- a/Engine/Plugins/Experimental/AnimNext/Source/AnimNext/Private/Module.cpp
+++ b/Engine/Plugins/Experimental/AnimNext/Source/AnimNext/Private/Module.cpp
@@ -14,8 +14,8 @@
 #include "RigVMRuntimeDataRegistry.h"
 #include "Animation/AnimSequence.h"
 #include "Scheduler/Scheduler.h"
-#include "Param/Params.h"
-#include "Components/SkeletalMeshComponent.h"
+#include "Param/ExternalParameterRegistry.h"
+#include "Param/ObjectProxyFactory.h"
 
 namespace UE::AnimNext
 {
@@ -37,8 +37,8 @@ public:
 
 		FRigVMRegistry::Get().RegisterObjectTypes(AllowedObjectTypes);
 
-		FParamId::Init();
-
+		FObjectProxyFactory::Init();
+		FExternalParameterRegistry::Init();
 		FDataRegistry::Init();
 		FDecoratorRegistry::Init();
 		FNodeTemplateRegistry::Init();
@@ -53,8 +53,8 @@ public:
 		FNodeTemplateRegistry::Destroy();
 		FDecoratorRegistry::Destroy();
 		FDataRegistry::Destroy();
-
-		FParamId::Destroy();
+		FObjectProxyFactory::Destroy();
+		FExternalParameterRegistry::Destroy();
 	}
 };
 
diff --git a/Engine/Plugins/Experimental/AnimNext/Source/AnimNext/Private/Param/AnimNextActorComponentLibrary.cpp b/Engine/Plugins/Experimental/AnimNext/Source/AnimNext/Private/Param/AnimNextActorComponentLibrary.cpp
new file mode 100644
--- /dev/null
+++ b/Engine/Plugins/Experimental/AnimNext/Source/AnimNext/Private/Param/AnimNextActorComponentLibrary.cpp
@@ -0,0 +1,18 @@
+﻿// Copyright Epic Games, Inc. All Rights Reserved.
+
+#include "AnimNextActorComponentLibrary.h"
+#include "Scheduler/AnimNextTickFunctionBinding.h"
+#include "Components/ActorComponent.h"
+
+UClass* UAnimNextActorComponentLibrary::GetSupportedClass() const
+{
+	return UActorComponent::StaticClass();
+}
+
+FAnimNextTickFunctionBinding UAnimNextActorComponentLibrary::GetTick(UActorComponent* InComponent)
+{
+	FAnimNextTickFunctionBinding Binding;
+	Binding.Object = InComponent;
+	Binding.TickFunction = &InComponent->PrimaryComponentTick;
+	return Binding;
+}
diff --git a/Engine/Plugins/Experimental/AnimNext/Source/AnimNext/Private/Param/AnimNextActorComponentLibrary.h b/Engine/Plugins/Experimental/AnimNext/Source/AnimNext/Private/Param/AnimNextActorComponentLibrary.h
new file mode 100644
--- /dev/null
+++ b/Engine/Plugins/Experimental/AnimNext/Source/AnimNext/Private/Param/AnimNextActorComponentLibrary.h
@@ -0,0 +1,24 @@
+﻿// Copyright Epic Games, Inc. All Rights Reserved.
+
+#pragma once
+
+#include "CoreMinimal.h"
+#include "Param/AnimNextClassExtensionLibrary.h"
+#include "Scheduler/AnimNextTickFunctionBinding.h"
+#include "AnimNextActorComponentLibrary.generated.h"
+
+class UActorComponent;
+
+// Access to non-UProperty/UFunction data on UActorComponent
+UCLASS()
+class UAnimNextActorComponentLibrary : public UAnimNextClassExtensionLibrary
+{
+	GENERATED_BODY()
+
+	// UAnimNextClassProxy interface
+	virtual UClass* GetSupportedClass() const override;
+
+	// Returns the component's tick function
+	UFUNCTION(BlueprintCallable, BlueprintInternalUseOnly)
+	static FAnimNextTickFunctionBinding GetTick(UActorComponent* InComponent);
+};
diff --git a/Engine/Plugins/Experimental/AnimNext/Source/AnimNext/Private/Param/AnimNextParameterBlock.cpp b/Engine/Plugins/Experimental/AnimNext/Source/AnimNext/Private/Param/AnimNextParameterBlock.cpp
--- a/Engine/Plugins/Experimental/AnimNext/Source/AnimNext/Private/Param/AnimNextParameterBlock.cpp
+++ b/Engine/Plugins/Experimental/AnimNext/Source/AnimNext/Private/Param/AnimNextParameterBlock.cpp
@@ -39,29 +39,6 @@ void UAnimNextParameterBlock::UpdateLayer(UE::AnimNext::FParamStackLayerHandle&
 	}
 }
 
-UE::AnimNext::FParamStackLayerHandle UAnimNextParameterBlock::CacheLayer() const
-{
-	return UE::AnimNext::FParamStack::MakeValueLayer(PropertyBag);
-}
-
-bool UAnimNextParameterBlock::ShouldCacheLayer(const UE::AnimNext::FParamStackLayerHandle& InHandle) const
-{
-	if(!InHandle.IsValid())
-	{
-		return true;
-	}
-
-#if WITH_EDITOR	// Layout should only be changing in editor
-	const FInstancedPropertyBag* HandlePropertyBag = InHandle.As<FInstancedPropertyBag>();
-	if(HandlePropertyBag == nullptr || HandlePropertyBag->GetPropertyBagStruct() != PropertyBag.GetPropertyBagStruct())
-	{
-		return true;
-	}
-#endif
-
-	return false;
-}
-
 void UAnimNextParameterBlock::BeginDestroy()
 {
 	Super::BeginDestroy();
diff --git a/Engine/Plugins/Experimental/AnimNext/Source/AnimNext/Private/Param/AnimNextParameterSourceRef.cpp b/Engine/Plugins/Experimental/AnimNext/Source/AnimNext/Private/Param/AnimNextParameterSourceRef.cpp
deleted file mode 100644
diff --git a/Engine/Plugins/Experimental/AnimNext/Source/AnimNext/Private/Param/AnimNextWorldLibrary.cpp b/Engine/Plugins/Experimental/AnimNext/Source/AnimNext/Private/Param/AnimNextWorldLibrary.cpp
new file mode 100644
--- /dev/null
+++ b/Engine/Plugins/Experimental/AnimNext/Source/AnimNext/Private/Param/AnimNextWorldLibrary.cpp
@@ -0,0 +1,38 @@
+﻿// Copyright Epic Games, Inc. All Rights Reserved.
+
+#include "AnimNextWorldLibrary.h"
+#include "Engine/World.h"
+#include "Scheduler/AnimNextTickFunctionBinding.h"
+
+UClass* UAnimNextWorldLibrary::GetSupportedClass() const
+{
+	return UWorld::StaticClass();
+}
+
+double UAnimNextWorldLibrary::GetTimeSeconds(UWorld* InWorld)
+{
+	return InWorld->GetTimeSeconds();
+}
+
+double UAnimNextWorldLibrary::GetUnpausedTimeSeconds(UWorld* InWorld)
+{
+	return InWorld->GetUnpausedTimeSeconds();
+}
+
+double UAnimNextWorldLibrary::GetRealTimeSeconds(UWorld* InWorld)
+{
+	return InWorld->GetRealTimeSeconds();
+}
+
+float UAnimNextWorldLibrary::GetDeltaSeconds(UWorld* InWorld)
+{
+	return InWorld->GetDeltaSeconds();
+}
+
+FAnimNextTickFunctionBinding UAnimNextWorldLibrary::GetPhysicsTick(UWorld* InWorld)
+{
+	FAnimNextTickFunctionBinding Binding;
+	Binding.Object = InWorld;
+	Binding.TickFunction = &InWorld->EndPhysicsTickFunction;
+	return Binding;
+}
diff --git a/Engine/Plugins/Experimental/AnimNext/Source/AnimNext/Private/Param/AnimNextWorldLibrary.h b/Engine/Plugins/Experimental/AnimNext/Source/AnimNext/Private/Param/AnimNextWorldLibrary.h
new file mode 100644
--- /dev/null
+++ b/Engine/Plugins/Experimental/AnimNext/Source/AnimNext/Private/Param/AnimNextWorldLibrary.h
@@ -0,0 +1,40 @@
+﻿// Copyright Epic Games, Inc. All Rights Reserved.
+
+#pragma once
+
+#include "CoreMinimal.h"
+#include "Param/AnimNextClassExtensionLibrary.h"
+#include "AnimNextWorldLibrary.generated.h"
+
+class UWorld;
+struct FAnimNextTickFunctionBinding;
+
+// Access to non-UProperty/UFunction data on UWorld
+UCLASS()
+class UAnimNextWorldLibrary : public UAnimNextClassExtensionLibrary
+{
+	GENERATED_BODY()
+
+	// UAnimNextClassProxy interface
+	virtual UClass* GetSupportedClass() const override;
+
+	// Returns time in seconds since world was brought up for play, IS stopped when game pauses, IS dilated/clamped
+	UFUNCTION(BlueprintCallable, BlueprintInternalUseOnly)
+	static double GetTimeSeconds(UWorld* InWorld);
+
+	// Returns time in seconds since world was brought up for play, IS NOT stopped when game pauses, IS dilated/clamped
+	UFUNCTION(BlueprintCallable, BlueprintInternalUseOnly)
+	static double GetUnpausedTimeSeconds(UWorld* InWorld);
+
+	// Returns time in seconds since world was brought up for play, does NOT stop when game pauses, NOT dilated/clamped 
+	UFUNCTION(BlueprintCallable, BlueprintInternalUseOnly)
+	static double GetRealTimeSeconds(UWorld* InWorld);
+
+	// Returns the frame delta time in seconds adjusted by e.g. time dilation.
+	UFUNCTION(BlueprintCallable, BlueprintInternalUseOnly)
+	static float GetDeltaSeconds(UWorld* InWorld);
+
+	// Returns the physics tick function
+	UFUNCTION(BlueprintCallable, BlueprintInternalUseOnly)
+	static FAnimNextTickFunctionBinding GetPhysicsTick(UWorld* InWorld);
+};
diff --git a/Engine/Plugins/Experimental/AnimNext/Source/AnimNext/Private/Param/ClassProxy.cpp b/Engine/Plugins/Experimental/AnimNext/Source/AnimNext/Private/Param/ClassProxy.cpp
new file mode 100644
--- /dev/null
+++ b/Engine/Plugins/Experimental/AnimNext/Source/AnimNext/Private/Param/ClassProxy.cpp
@@ -0,0 +1,177 @@
+﻿// Copyright Epic Games, Inc. All Rights Reserved.
+
+#include "ClassProxy.h"
+#include "Logging/StructuredLog.h"
+#include "Param/AnimNextClassExtensionLibrary.h"
+#include "Param/ParamTypeHandle.h"
+#include "EngineLogs.h"
+
+namespace UE::AnimNext
+{
+
+FClassProxy::FClassProxy(const UClass* InClass)
+	: Class(InClass)
+{
+	for(TFieldIterator<UFunction> It(InClass, EFieldIterationFlags::IncludeSuper | EFieldIterationFlags::IncludeInterfaces); It; ++It)
+	{
+		UFunction* Function = *It;
+
+		// We add only 'accessor' functions (no params apart from the return value) that have valid return types
+		const FProperty* ReturnProperty = Function->GetReturnProperty();
+		if(ReturnProperty != nullptr && Function->NumParms == 1 && Function->HasAnyFunctionFlags(FUNC_BlueprintCallable))
+		{
+			FParamTypeHandle TypeHandle = FParamTypeHandle::FromProperty(ReturnProperty);
+			if(TypeHandle.IsValid())
+			{
+				FString FunctionName = Function->GetName();
+
+				// Remove common prefixes.
+				// TODO: explicit mapping or regex transformations?
+				FunctionName.RemoveFromStart(TEXT("K2_"));
+				FunctionName.RemoveFromStart(TEXT("Get"));
+
+				FClassProxyParameter Parameter;
+				Parameter.AccessType = EClassProxyParameterAccessType::AccessorFunction;
+				Parameter.ClassParameterName = *FunctionName;
+				Parameter.Function = Function;
+				Parameter.Type = TypeHandle.GetType();
+#if WITH_EDITOR
+				Parameter.Tooltip = Function->GetToolTipText();
+				Parameter.bThreadSafe = Function->HasMetaData("BlueprintThreadSafe");
+#endif
+				ParameterNameMap.Add(Parameter.ClassParameterName, Parameters.Add(Parameter));
+			}
+		}
+	}
+
+	// Add properties after functions, checking for duplicate names
+	for(TFieldIterator<FProperty> It(InClass, EFieldIterationFlags::IncludeSuper | EFieldIterationFlags::IncludeInterfaces); It; ++It)
+	{
+		FProperty* Property = *It;
+		if(Property->HasAnyPropertyFlags(CPF_Edit | CPF_BlueprintVisible))
+		{
+			FParamTypeHandle TypeHandle = FParamTypeHandle::FromProperty(Property);
+			if(TypeHandle.IsValid())
+			{
+				FString PropertyName = Property->GetName();
+				if(TypeHandle.GetParameterType() == FParamTypeHandle::EParamType::Bool)
+				{
+					PropertyName.RemoveFromStart(TEXT("b"));
+				}
+
+				FName ParameterName = *PropertyName;
+				if(!ParameterNameMap.Contains(ParameterName))
+				{
+					FClassProxyParameter Parameter;
+					Parameter.AccessType = EClassProxyParameterAccessType::Property;
+					Parameter.ClassParameterName = ParameterName;
+					Parameter.Property = Property;
+					Parameter.Type = TypeHandle.GetType();
+#if WITH_EDITOR
+					Parameter.Tooltip = Property->GetToolTipText();
+					Parameter.bThreadSafe = false;
+#endif
+					ParameterNameMap.Add(Parameter.ClassParameterName, Parameters.Add(Parameter));
+				}
+				else
+				{
+					UE_LOGFMT(LogAnimation, Warning, "Could not add parameter '{ParameterName}' for class '{ClassName}' as it already exists", ParameterName, InClass->GetFName());
+				}
+			}
+		}
+	}
+
+	// We allow only native static 'accessor wrapper' functions (one object param and the return value)
+	auto CheckSignature = [](const UFunction* InFunction, const UClass* InProxiedClass)
+	{
+		if(!InFunction->HasAllFunctionFlags(FUNC_BlueprintCallable | FUNC_Static | FUNC_Native))
+		{
+			return false;
+		}
+
+		if(InFunction->NumParms != 2)
+		{
+			return false;
+		}
+		
+		int32 ParamIndex = 0;
+		for(TFieldIterator<FProperty> It(InFunction); It && (It->PropertyFlags & CPF_Parm); ++It, ++ParamIndex)
+		{
+			// Check parameter is an object of the expected class
+			if(ParamIndex == 0)
+			{
+				FObjectProperty* ObjectProperty = CastField<FObjectProperty>(*It);
+				if(ObjectProperty == nullptr)
+				{ 
+					return false;
+				}
+				
+				if(!ObjectProperty->PropertyClass->IsChildOf(InProxiedClass))
+				{
+					return false;
+				}
+			}
+
+			// Check return value
+			if(ParamIndex == 1)
+			{
+				if((It->PropertyFlags & CPF_ReturnParm) == 0)
+				{
+					return false;
+				}
+			}
+		}
+		return true;
+	};
+
+	// Add any additional extension libraries that extend this class
+	TArray<UClass*> Classes;
+	GetDerivedClasses(UAnimNextClassExtensionLibrary::StaticClass(), Classes);
+	for(UClass* ProxyClass : Classes)
+	{
+		UAnimNextClassExtensionLibrary* CDO = ProxyClass->GetDefaultObject<UAnimNextClassExtensionLibrary>();
+		const UClass* ExtendedClass = CDO->GetSupportedClass();
+		if(InClass->IsChildOf(ExtendedClass))
+		{
+			for(TFieldIterator<UFunction> It(ProxyClass); It; ++It)
+			{
+				UFunction* Function = *It;
+				if(CheckSignature(Function, ExtendedClass))
+				{
+					const FProperty* ReturnProperty = Function->GetReturnProperty();
+					FParamTypeHandle TypeHandle = FParamTypeHandle::FromProperty(ReturnProperty);
+					if(TypeHandle.IsValid())
+					{
+						FString FunctionName = Function->GetName();
+
+						// Remove common prefixes.
+						// TODO: explicit mapping or regex transformations?
+						FunctionName.RemoveFromStart(TEXT("K2_"));
+						FunctionName.RemoveFromStart(TEXT("Get"));
+
+						FName ParameterName = *FunctionName;
+						if(!ParameterNameMap.Contains(ParameterName))
+						{
+							FClassProxyParameter Parameter;
+							Parameter.AccessType = EClassProxyParameterAccessType::HoistedFunction;
+							Parameter.ClassParameterName = *FunctionName;
+							Parameter.Function = Function;
+							Parameter.Type = TypeHandle.GetType();
+#if WITH_EDITOR
+							Parameter.Tooltip = Function->GetToolTipText();
+							Parameter.bThreadSafe = Function->HasMetaData("BlueprintThreadSafe");
+#endif
+							ParameterNameMap.Add(Parameter.ClassParameterName, Parameters.Add(Parameter));
+						}
+						else
+						{
+							UE_LOGFMT(LogAnimation, Warning, "Could not add parameter '{ParameterName}' for class proxy '{ClassName}' as it already exists", ParameterName, InClass->GetFName());
+						}
+					}
+				}
+			}
+		}
+	}
+}
+
+}
diff --git a/Engine/Plugins/Experimental/AnimNext/Source/AnimNext/Private/Param/ClassProxy.h b/Engine/Plugins/Experimental/AnimNext/Source/AnimNext/Private/Param/ClassProxy.h
new file mode 100644
--- /dev/null
+++ b/Engine/Plugins/Experimental/AnimNext/Source/AnimNext/Private/Param/ClassProxy.h
@@ -0,0 +1,74 @@
+﻿// Copyright Epic Games, Inc. All Rights Reserved.
+
+#pragma once
+
+#include "CoreMinimal.h"
+#include "Param/ParamType.h"
+
+namespace UE::AnimNext
+{
+
+enum class EClassProxyParameterAccessType : int32
+{
+	// Access via FProperty
+	Property,
+
+	// Access via UFunction with accessor signature Val = Obj.Func() 
+	AccessorFunction,
+
+	// Access via UFunction with hoisted signature Val = Func(Obj) 
+	HoistedFunction,
+};
+
+// Cached info about how a class property or function maps to a parameter
+struct FClassProxyParameter
+{
+	// The parameter name that this proxy corresponds to (within this class - object accessors act as prefixes of this name)
+	FName ClassParameterName;
+
+	// The function to call to get this parameter
+	TWeakObjectPtr<UFunction> Function;
+
+	// The property to copy to get this parameter
+	// TFieldPath to accomodate potential reinstancing
+	TFieldPath<FProperty> Property;
+
+	// The type of the property
+	FAnimNextParamType Type;
+
+	// How this parameter ies accessed
+	EClassProxyParameterAccessType AccessType;
+
+#if WITH_EDITOR
+	// Tooltip to display in editor
+	FText Tooltip;
+
+	// Whether this parameter is safe to access on worker threads
+	bool bThreadSafe = false;
+#endif
+};
+
+// Proxy struct used to hold data about a kind of UClass to cache data for
+struct FClassProxy
+{
+	explicit FClassProxy(const UClass* InClass);
+
+	// The class that this proxy wraps
+	TWeakObjectPtr<const UClass> Class;
+
+	// Cache of properties, fetched from Class
+	TArray<FClassProxyParameter> Parameters;
+
+	// Map of parameter name to index in Parameters array
+	TMap<FName, int32> ParameterNameMap;
+
+#if WITH_EDITOR
+	// Whether any parameters are safe to access on worker threads
+	bool bHasThreadSafeParameters = false;
+
+	// Whether any parameters are not safe to access on worker threads
+	bool bHasNonThreadSafeParameters = true;
+#endif
+};
+
+}
\ No newline at end of file
diff --git a/Engine/Plugins/Experimental/AnimNext/Source/AnimNext/Private/Param/ExternalParameterRegistry.cpp b/Engine/Plugins/Experimental/AnimNext/Source/AnimNext/Private/Param/ExternalParameterRegistry.cpp
new file mode 100644
--- /dev/null
+++ b/Engine/Plugins/Experimental/AnimNext/Source/AnimNext/Private/Param/ExternalParameterRegistry.cpp
@@ -0,0 +1,167 @@
+﻿// Copyright Epic Games, Inc. All Rights Reserved.
+
+#include "Param/ExternalParameterRegistry.h"
+#include "Param/ObjectAccessor.h"
+#include "Param/ObjectProxy.h"
+#include "Logging/StructuredLog.h"
+#include "Param/IParameterSource.h"
+#include "Param/IParameterSourceFactory.h"
+#include "UObject/ObjectKey.h"
+#include "UObject/Package.h"
+#include "Features/IModularFeatures.h"
+#include "EngineLogs.h"
+
+namespace UE::AnimNext
+{
+
+struct FExternalParameterRegistryImpl
+{
+	void ResetFactories()
+	{
+		Factories.Reset();
+		SourceFactoryMap.Reset();
+	}
+
+	void RefreshFactories()
+	{
+		ResetFactories();
+
+		TArray<const IParameterSourceFactory*> LocalFactories = IModularFeatures::Get().GetModularFeatureImplementations<const IParameterSourceFactory>(IParameterSourceFactory::FeatureName);
+		for(const IParameterSourceFactory* Factory : LocalFactories)
+		{
+			AddFactory(Factory);
+		}
+	}
+
+	void AddFactory(const IParameterSourceFactory* InFactory)
+	{
+		Factories.Add(InFactory);
+		InFactory->ForEachSource([this, InFactory](FName InSourceName)
+		{
+			// Skip duplicate names
+			if(SourceFactoryMap.Contains(InSourceName))
+			{
+				UE_LOGFMT(LogAnimation, Error, "Duplicate external parameter source found: {Name}", InSourceName);
+			}
+			else
+			{
+				SourceFactoryMap.Add(InSourceName, InFactory);
+			}
+
+#if WITH_EDITOR
+			InFactory->ForEachParameter(InSourceName, [this, InFactory, InSourceName](FName InParameterName, const IParameterSourceFactory::FParameterInfo& InInfo)
+			{
+				ParameterSourceMap.Add(InParameterName, InSourceName);
+			});
+#endif
+		});
+	}
+
+	void RemoveFactory(const IParameterSourceFactory* InFactory)
+	{
+		Factories.Remove(InFactory);
+		for(auto It = SourceFactoryMap.CreateIterator(); It; ++It)
+		{
+			if(It->Value == InFactory)
+			{
+				It.RemoveCurrent();
+			}
+		}
+#if WITH_EDITOR
+		InFactory->ForEachSource([this, InFactory](FName InSourceName)
+		{
+			InFactory->ForEachParameter(InSourceName, [this](FName InParameterName, const IParameterSourceFactory::FParameterInfo& InInfo)
+			{
+				ParameterSourceMap.Remove(InParameterName);
+			});
+		});
+#endif
+	}
+
+	// All known factories
+	TArray<const IParameterSourceFactory*> Factories;
+
+	// Map of source name -> factory
+	TMap<FName, const IParameterSourceFactory*> SourceFactoryMap;
+
+#if WITH_EDITOR
+	// Map of parameter name -> parameter source
+	TMap<FName, FName> ParameterSourceMap;
+#endif
+
+	FDelegateHandle OnModularFeatureRegisteredHandle;
+	FDelegateHandle OnModularFeatureUnregisteredHandle;
+};
+
+static FExternalParameterRegistryImpl GExternalParameterRegistry;
+
+void FExternalParameterRegistry::Init()
+{
+	GExternalParameterRegistry.OnModularFeatureRegisteredHandle = IModularFeatures::Get().OnModularFeatureRegistered().AddLambda([](const FName& InType, IModularFeature* InModularFeature)
+	{
+		if(InType == IParameterSourceFactory::FeatureName)
+		{
+			GExternalParameterRegistry.AddFactory(static_cast<IParameterSourceFactory*>(InModularFeature));
+		}
+	});
+	GExternalParameterRegistry.OnModularFeatureUnregisteredHandle = IModularFeatures::Get().OnModularFeatureRegistered().AddLambda([](const FName& InType, IModularFeature* InModularFeature)
+	{
+		if(InType == IParameterSourceFactory::FeatureName)
+		{
+			GExternalParameterRegistry.RemoveFactory(static_cast<IParameterSourceFactory*>(InModularFeature));
+		}
+	});
+	GExternalParameterRegistry.RefreshFactories();
+}
+
+void FExternalParameterRegistry::Destroy()
+{
+	IModularFeatures::Get().OnModularFeatureRegistered().Remove(GExternalParameterRegistry.OnModularFeatureRegisteredHandle);
+	IModularFeatures::Get().OnModularFeatureUnregistered().Remove(GExternalParameterRegistry.OnModularFeatureUnregisteredHandle);
+	GExternalParameterRegistry.ResetFactories();
+}
+
+TUniquePtr<IParameterSource> FExternalParameterRegistry::CreateParameterSource(const FExternalParameterContext& InContext, FName InSourceName, TConstArrayView<FName> InRequiredParameters)
+{
+	using namespace UE::AnimNext;
+
+	if(const IParameterSourceFactory** FoundFactoryPtr = GExternalParameterRegistry.SourceFactoryMap.Find(InSourceName))
+	{
+		return (*FoundFactoryPtr)->CreateParameterSource(InContext, InSourceName, InRequiredParameters);
+	}
+	return nullptr;
+}
+
+#if WITH_EDITOR
+
+bool FExternalParameterRegistry::FindParameterInfo(FName InParameterName, IParameterSourceFactory::FParameterInfo& OutInfo)
+{
+	if(const FName* SourceNamePtr = GExternalParameterRegistry.ParameterSourceMap.Find(InParameterName))
+	{
+		if(const IParameterSourceFactory** Factory = GExternalParameterRegistry.SourceFactoryMap.Find(*SourceNamePtr))
+		{
+			return (*Factory)->FindParameterInfo(InParameterName, OutInfo);
+		}
+	}
+	return false;
+}
+
+void FExternalParameterRegistry::ForEachParameter(TFunctionRef<void(FName, const IParameterSourceFactory::FParameterInfo&)> InFunction)
+{
+	for(const IParameterSourceFactory* Factory : GExternalParameterRegistry.Factories)
+	{
+		Factory->ForEachSource([&InFunction, Factory](FName InSourceName)
+		{
+			Factory->ForEachParameter(InSourceName, InFunction);
+		});
+	}
+}
+
+FName FExternalParameterRegistry::FindSourceForParameter(FName InParameterName)
+{
+	return GExternalParameterRegistry.ParameterSourceMap.FindRef(InParameterName);
+}
+
+#endif
+
+}
diff --git a/Engine/Plugins/Experimental/AnimNext/Source/AnimNext/Private/Param/ObjectAccessor.cpp b/Engine/Plugins/Experimental/AnimNext/Source/AnimNext/Private/Param/ObjectAccessor.cpp
new file mode 100644
--- /dev/null
+++ b/Engine/Plugins/Experimental/AnimNext/Source/AnimNext/Private/Param/ObjectAccessor.cpp
@@ -0,0 +1,31 @@
+﻿// Copyright Epic Games, Inc. All Rights Reserved.
+
+#include "ObjectAccessor.h"
+#include "Param/ClassProxy.h"
+#include "EngineLogs.h"
+
+namespace UE::AnimNext
+{
+
+FObjectAccessor::FObjectAccessor(FName InAccessorName, FObjectAccessorFunction&& InFunction, TSharedRef<FClassProxy> InClassProxy)
+	: AccessorName(InAccessorName)
+	, Function(MoveTemp(InFunction))
+	, ClassProxy(InClassProxy)
+{
+	// Prefix class proxy names into RemappedParameters
+	RemappedParameters.Reserve(ClassProxy->Parameters.Num());
+	for(const FClassProxyParameter& ClassProxyParameter : ClassProxy->Parameters)
+	{
+		TStringBuilder<128> RemappedNameBuilder;
+		AccessorName.AppendString(RemappedNameBuilder);
+		RemappedNameBuilder.Append(TEXT("_"));
+		ClassProxyParameter.ClassParameterName.AppendString(RemappedNameBuilder);
+
+		FName RemappedName = RemappedNameBuilder.ToString();
+		int32 Index = RemappedParameters.Add(RemappedName);
+		
+		RemappedParametersMap.Add(RemappedName, Index);
+	}
+}
+	
+}
diff --git a/Engine/Plugins/Experimental/AnimNext/Source/AnimNext/Private/Param/ObjectAccessor.h b/Engine/Plugins/Experimental/AnimNext/Source/AnimNext/Private/Param/ObjectAccessor.h
new file mode 100644
--- /dev/null
+++ b/Engine/Plugins/Experimental/AnimNext/Source/AnimNext/Private/Param/ObjectAccessor.h
@@ -0,0 +1,37 @@
+﻿// Copyright Epic Games, Inc. All Rights Reserved.
+
+#pragma once
+
+#include "CoreMinimal.h"
+#include "ObjectProxyFactory.h"
+
+namespace UE::AnimNext
+{
+	struct FClassProxy;
+}
+
+namespace UE::AnimNext
+{
+
+// Cached information about how to access an object's data and map it to AnimNext parameters 
+struct FObjectAccessor
+{
+	FObjectAccessor(FName InAccessorName, FObjectAccessorFunction&& InFunction, TSharedRef<FClassProxy> InClassProxy);
+
+	// The name of the accessor
+	FName AccessorName;
+
+	// Function called to access the object given a context
+	FObjectAccessorFunction Function;
+
+	// Class proxy defining the 'layout' of the object
+	TSharedRef<FClassProxy> ClassProxy;
+
+	// All parameters mapped from class -> accessor's 'namespace'
+	TArray<FName> RemappedParameters;
+
+	// Index lookup for accessor's 'namespace' parameters
+	TMap<FName, int32> RemappedParametersMap;
+};
+
+}
\ No newline at end of file
diff --git a/Engine/Plugins/Experimental/AnimNext/Source/AnimNext/Private/Param/ObjectProxy.cpp b/Engine/Plugins/Experimental/AnimNext/Source/AnimNext/Private/Param/ObjectProxy.cpp
new file mode 100644
--- /dev/null
+++ b/Engine/Plugins/Experimental/AnimNext/Source/AnimNext/Private/Param/ObjectProxy.cpp
@@ -0,0 +1,122 @@
+﻿// Copyright Epic Games, Inc. All Rights Reserved.
+
+#include "ObjectProxy.h"
+
+#include "ClassProxy.h"
+#include "ObjectAccessor.h"
+#include "Param/ParamStack.h"
+
+namespace UE::AnimNext
+{
+
+FObjectProxy::FObjectProxy(UObject* InObject, const TSharedRef<FObjectAccessor>& InObjectAccessor)
+	: Object(InObject)
+	, ObjectAccessor(InObjectAccessor)
+	, RootParameterName(ObjectAccessor->AccessorName)
+{
+	// Always supply the root parameter in index 0
+	ParameterCache.AddProperty(RootParameterName, EPropertyBagPropertyType::Object, InObject->GetClass());
+	ParameterCache.SetValueObject(RootParameterName, InObject);
+}
+
+void FObjectProxy::Update()
+{
+	const UPropertyBag* PropertyBag = ParameterCache.GetPropertyBagStruct();
+	TConstArrayView<FPropertyBagPropertyDesc> PropertyDescs = PropertyBag->GetPropertyDescs();
+	uint8* StructData = ParameterCache.GetMutableValue().GetMemory();
+	UObject* ResolvedObject = Object.Get();
+	PropertyDescs[0].CachedProperty->SetValue_InContainer(StructData, &ResolvedObject);
+	if(ResolvedObject)
+	{
+		for(int32 ParameterIndex = 0; ParameterIndex < ParametersToUpdate.Num(); ++ParameterIndex)
+		{
+			const FAnimNextObjectProxyParameter& ParameterToUpdate = ParametersToUpdate[ParameterIndex];
+			const FProperty* ResultProperty = PropertyDescs[ParameterToUpdate.ValueParamIndex].CachedProperty;
+			void* ResultBuffer = ResultProperty->ContainerPtrToValuePtr<void>(StructData);
+
+			switch(ParameterToUpdate.AccessType)
+			{
+			case EClassProxyParameterAccessType::Property:
+				{
+					const FProperty* SourceProperty = ParameterToUpdate.GetProperty();
+					checkSlow(SourceProperty);
+					checkSlow(SourceProperty->GetClass() == ResultProperty->GetClass());
+
+					const void* SourceBuffer = SourceProperty->ContainerPtrToValuePtr<const void>(ResolvedObject);
+					SourceProperty->CopyCompleteValue(ResultBuffer, SourceBuffer);
+					break;
+				}
+			case EClassProxyParameterAccessType::AccessorFunction:
+				{
+					UFunction* Function = ParameterToUpdate.GetFunction();
+					checkSlow(Function);
+					checkSlow(ResolvedObject->GetClass()->IsChildOf(Function->GetOuterUClass()));
+
+					FFrame Stack(ResolvedObject, Function, nullptr, nullptr, Function->ChildProperties);
+					Function->Invoke(ResolvedObject, Stack, ResultBuffer);
+					break;
+				}
+			case EClassProxyParameterAccessType::HoistedFunction:
+				{
+					UFunction* Function = ParameterToUpdate.GetFunction();
+					checkSlow(Function);
+
+					const FObjectProperty* ObjectProperty = CastFieldChecked<FObjectProperty>(PropertyDescs[0].CachedProperty);
+					check(ResolvedObject->GetClass()->IsChildOf(ObjectProperty->PropertyClass));
+					UObject** ObjectBuffer = ObjectProperty->ContainerPtrToValuePtr<UObject*>(StructData);
+					*ObjectBuffer = ResolvedObject;
+
+					FFrame Stack(ResolvedObject, Function, ObjectBuffer, nullptr, Function->ChildProperties);
+					Function->Invoke(ResolvedObject, Stack, ResultBuffer);
+					break;
+				}
+			default:
+				checkNoEntry();
+				break;
+			}
+		}
+	}
+}
+
+void FObjectProxy::RequestParameterCache(TConstArrayView<FName> InParameterNames)
+{
+	using namespace UE::AnimNext;
+
+	const int32 NumExistingProperties = ParameterCache.GetNumPropertiesInBag();
+
+	TArray<FPropertyBagPropertyDesc> PropertyDescsToAdd;
+	PropertyDescsToAdd.Reserve(InParameterNames.Num());
+
+	for(FName ParameterName : InParameterNames)
+	{
+		if(!ParameterNameMap.Contains(ParameterName))
+		{
+			if(const int32* ParameterIndexPtr = ObjectAccessor->RemappedParametersMap.Find(ParameterName))
+			{
+				const FClassProxyParameter& ClassProxyParameter = ObjectAccessor->ClassProxy->Parameters[*ParameterIndexPtr];
+
+				FAnimNextObjectProxyParameter& NewParameterToUpdate = ParametersToUpdate.AddDefaulted_GetRef();
+				NewParameterToUpdate.AccessType = ClassProxyParameter.AccessType;
+				NewParameterToUpdate.Function = ClassProxyParameter.Function.Get();
+				NewParameterToUpdate.Property = ClassProxyParameter.Property.Get();
+
+				int32 ValueParamIndex = PropertyDescsToAdd.Emplace(ParameterName, ClassProxyParameter.Type.GetContainerType(), ClassProxyParameter.Type.GetValueType(), ClassProxyParameter.Type.GetValueTypeObject());
+				NewParameterToUpdate.ValueParamIndex = NumExistingProperties + ValueParamIndex;
+				ParameterNameMap.Add(ParameterName, NewParameterToUpdate.ValueParamIndex);
+			}
+		}
+	}
+
+	// Update parameter bag struct
+	ParameterCache.AddProperties(PropertyDescsToAdd);
+
+	// Recreate layer handle as layout has changed
+	LayerHandle = FParamStack::MakeReferenceLayer(ParameterCache);
+}
+
+void FObjectProxy::AddReferencedObjects(FReferenceCollector& Collector)
+{
+	ParameterCache.AddStructReferencedObjects(Collector);
+}
+
+}
\ No newline at end of file
diff --git a/Engine/Plugins/Experimental/AnimNext/Source/AnimNext/Private/Param/ObjectProxy.h b/Engine/Plugins/Experimental/AnimNext/Source/AnimNext/Private/Param/ObjectProxy.h
new file mode 100644
--- /dev/null
+++ b/Engine/Plugins/Experimental/AnimNext/Source/AnimNext/Private/Param/ObjectProxy.h
@@ -0,0 +1,98 @@
+﻿// Copyright Epic Games, Inc. All Rights Reserved.
+
+#pragma once
+
+#include "Param/IParameterSource.h"
+#include "PropertyBag.h"
+
+namespace UE::AnimNext
+{
+	struct FObjectAccessor;
+	enum class EClassProxyParameterAccessType : int32;
+}
+
+namespace UE::AnimNext
+{
+
+// Cached info about a parameter to update
+struct FAnimNextObjectProxyParameter
+{
+	// How to access this parameter
+	EClassProxyParameterAccessType AccessType;
+
+	// The function to call
+#if WITH_EDITOR
+	// TWeakObjectPtr to accomodate potential reinstancing in editor
+	TWeakObjectPtr<UFunction> Function;
+#else
+	UFunction* Function = nullptr;
+#endif
+
+	// The property to copy
+#if WITH_EDITOR
+	// TFieldPath to accomodate potential reinstancing in editor
+	TFieldPath<FProperty> Property;
+#else
+	FProperty* Property = nullptr;
+#endif
+
+	// Index into the ParameterCache property bag's PropertyDescs array
+	int32 ValueParamIndex = 0;
+
+	FProperty* GetProperty() const
+	{
+#if WITH_EDITOR
+		return Property.Get();
+#else
+		return Property;
+#endif
+	}
+
+	UFunction* GetFunction() const
+	{
+#if WITH_EDITOR
+		return Function.Get();
+#else
+		return Function;
+#endif
+	}
+};
+
+// Proxy struct used to fetch and cache external UObject data
+struct FObjectProxy : public IParameterSource
+{
+	FObjectProxy() = delete;
+
+	FObjectProxy(UObject* InObject, const TSharedRef<FObjectAccessor>& InObjectAccessor);
+
+	// IParameterSource interface
+	virtual void Update() override;
+	virtual const FParamStackLayerHandle& GetLayerHandle() const override { return LayerHandle; }
+	virtual void AddReferencedObjects(FReferenceCollector& Collector) override;
+
+	// Adds a set of parameters to cache each time the layer is updated
+	void RequestParameterCache(TConstArrayView<FName> InParameterNames);
+
+	// The object that this proxy wraps
+	TWeakObjectPtr<UObject> Object;
+
+	// Cache of properties, fetched from Object
+	FInstancedPropertyBag ParameterCache;
+
+	// Layer handle - must be updated if ParameterCache changes layout
+	FParamStackLayerHandle LayerHandle;
+
+	// Properties to update each call to UpdateLayer
+	TArray<FAnimNextObjectProxyParameter> ParametersToUpdate;
+
+	// Map of parameter name to index in ParametersToUpdate array
+	TMap<FName, int32> ParameterNameMap;
+
+	// Object accessor used to lookup parameters
+	TSharedRef<FObjectAccessor> ObjectAccessor;
+
+	// The name of the root parameter
+	FName RootParameterName;
+};
+
+}
\ No newline at end of file
diff --git a/Engine/Plugins/Experimental/AnimNext/Source/AnimNext/Private/Param/ObjectProxyFactory.cpp b/Engine/Plugins/Experimental/AnimNext/Source/AnimNext/Private/Param/ObjectProxyFactory.cpp
new file mode 100644
--- /dev/null
+++ b/Engine/Plugins/Experimental/AnimNext/Source/AnimNext/Private/Param/ObjectProxyFactory.cpp
@@ -0,0 +1,246 @@
+﻿// Copyright Epic Games, Inc. All Rights Reserved.
+
+#include "Param/ObjectProxyFactory.h"
+
+#include "AnimNextConfig.h"
+#include "ClassProxy.h"
+#include "ObjectAccessor.h"
+#include "ObjectProxy.h"
+#include "Components/SkeletalMeshComponent.h"
+#include "Logging/StructuredLog.h"
+#include "Param/AnimNextObjectAccessorConfig.h"
+#include "RigVMCore/RigVMRegistry.h"
+#include "Features/IModularFeatures.h"
+#include "Engine/World.h"
+
+namespace UE::AnimNext
+{
+
+static FObjectProxyFactory GObjectProxyFactory;
+
+void FObjectProxyFactory::Init()
+{
+	Refresh();
+}
+
+void FObjectProxyFactory::Destroy()
+{
+	GObjectProxyFactory.Reset();
+
+	IModularFeatures::Get().UnregisterModularFeature(IParameterSourceFactory::FeatureName, &GObjectProxyFactory);
+}
+
+void FObjectProxyFactory::Reset()
+{
+	ClassMap.Reset();
+	ObjectAccessors.Reset();
+	ParameterMap.Reset();
+}
+
+void FObjectProxyFactory::Refresh()
+{
+	IModularFeatures::Get().UnregisterModularFeature(IParameterSourceFactory::FeatureName, &GObjectProxyFactory);
+
+	GObjectProxyFactory.Reset();
+
+	for(const FAnimNextObjectAccessorConfig& AccessorConfig : GetDefault<UAnimNextConfig>()->ExposedClasses)
+	{
+		UClass* Class = AccessorConfig.Class.LoadSynchronous();
+		if(Class && AccessorConfig.AccessorName != NAME_None)
+		{
+			GObjectProxyFactory.RegisterObjectAccessor(AccessorConfig.AccessorName, Class);
+		}
+	}
+
+	// Refresh RigVM registry as we have updated allowed types
+	FRigVMRegistry::Get().RefreshEngineTypes();
+
+	IModularFeatures::Get().RegisterModularFeature(IParameterSourceFactory::FeatureName, &GObjectProxyFactory);
+}
+
+void FObjectProxyFactory::ForEachSource(TFunctionRef<void(FName)> InFunction) const
+{
+	for(const TPair<FName, TSharedPtr<FObjectAccessor>>& ObjectAccessorPair : ObjectAccessors)
+	{
+		InFunction(ObjectAccessorPair.Key);
+	}
+}
+
+TUniquePtr<IParameterSource> FObjectProxyFactory::CreateParameterSource(const FExternalParameterContext& InContext, FName InSourceName, TConstArrayView<FName> InRequiredParameters) const
+{
+	using namespace UE::AnimNext;
+
+	if(TSharedPtr<FObjectAccessor> ObjectAccessor = FindObjectAccessor(InSourceName))
+	{
+		if(UObject* Object = ObjectAccessor->Function(InContext))
+		{
+			TUniquePtr<FObjectProxy> ObjectProxy = MakeUnique<FObjectProxy>(Object, ObjectAccessor.ToSharedRef());
+			ObjectProxy->RequestParameterCache(InRequiredParameters);
+
+			return MoveTemp(ObjectProxy);
+		}
+	}
+	return nullptr;
+}
+
+#if WITH_EDITOR
+
+bool FObjectProxyFactory::FindParameterInfo(FName InParameterName, FParameterInfo& OutInfo) const
+{
+	if(const TWeakPtr<FObjectAccessor>* WeakAccessor = ParameterMap.Find(InParameterName))
+	{
+		if(TSharedPtr<FObjectAccessor> ObjectAccessor = WeakAccessor->Pin())
+		{
+			if(int32* ParamIndexPtr = ObjectAccessor->RemappedParametersMap.Find(InParameterName))
+			{
+				OutInfo.Type = ObjectAccessor->ClassProxy->Parameters[*ParamIndexPtr].Type;
+				OutInfo.Tooltip = ObjectAccessor->ClassProxy->Parameters[*ParamIndexPtr].Tooltip;
+				OutInfo.bThreadSafe = ObjectAccessor->ClassProxy->Parameters[*ParamIndexPtr].bThreadSafe;
+				return true;
+			}
+		}
+	}
+	return false;
+}
+
+void FObjectProxyFactory::ForEachParameter(FName InSourceName, TFunctionRef<void(FName, const FParameterInfo&)> InFunction) const
+{
+	if(TSharedPtr<FObjectAccessor> ObjectAccessor = FindObjectAccessor(InSourceName))
+	{
+		TSharedPtr<FClassProxy> ClassProxy = ObjectAccessor->ClassProxy;
+		check(ClassProxy->Parameters.Num() == ObjectAccessor->RemappedParameters.Num());
+		for(int32 ParameterIndex = 0, ParameterCount = ObjectAccessor->RemappedParameters.Num(); ParameterIndex < ParameterCount; ++ParameterIndex)
+		{
+			FParameterInfo ParameterInfo;
+			ParameterInfo.Type = ClassProxy->Parameters[ParameterIndex].Type;
+			ParameterInfo.Tooltip = ClassProxy->Parameters[ParameterIndex].Tooltip;
+			ParameterInfo.bThreadSafe = ClassProxy->Parameters[ParameterIndex].bThreadSafe;
+			InFunction(ObjectAccessor->RemappedParameters[ParameterIndex], ParameterInfo);
+		}
+	}
+}
+
+#endif
+
+void FObjectProxyFactory::RegisterObjectAccessor(FName InAccessorName, const UClass* InTargetClass)
+{
+	check(IsInGameThread());
+	check(InAccessorName != NAME_None);
+	check(InTargetClass != nullptr);
+
+	if(InTargetClass->IsChildOf<UActorComponent>())
+	{
+		RegisterActorComponentAccessor(InAccessorName, const_cast<UClass*>(InTargetClass));
+	}
+	else if(InTargetClass->IsChildOf<AActor>())
+	{
+		RegisterActorAccessor(InAccessorName, const_cast<UClass*>(InTargetClass));
+	}
+	else if(InTargetClass->IsChildOf<UWorld>())
+	{
+		RegisterObjectAccessor(InAccessorName, const_cast<UClass*>(InTargetClass), [](const FExternalParameterContext& InContext) -> UObject*
+		{
+			return InContext.Object->GetWorld();
+		});
+	}
+	else 
+	{
+		UE_LOGFMT(LogAnimation, Warning, "Class {ClassName} passed to FObjectProxyFactory::RegisterObjectAccessor that does not have a built-in accessor, please use the custom accessor overload.", InTargetClass->GetFName());
+	}
+}
+
+void FObjectProxyFactory::RegisterActorAccessor(FName InAccessorName, TSubclassOf<AActor> InTargetClass)
+{
+	RegisterObjectAccessor(InAccessorName, InTargetClass.Get(), [InTargetClass](const FExternalParameterContext& InContext) -> UObject*
+	{
+		if(UActorComponent* ContextComponent = Cast<UActorComponent>(InContext.Object))
+		{
+			if(ContextComponent->GetOwner()->GetClass()->IsChildOf(InTargetClass.Get()))
+			{
+				return ContextComponent->GetOwner();
+			}
+		}
+
+		return nullptr;
+	});
+}
+
+void FObjectProxyFactory::RegisterActorComponentAccessor(FName InAccessorName, TSubclassOf<UActorComponent> InTargetClass)
+{
+	RegisterObjectAccessor(InAccessorName, InTargetClass.Get(), [InTargetClass](const FExternalParameterContext& InContext) -> UObject*
+	{
+		if(UActorComponent* ContextComponent = Cast<UActorComponent>(InContext.Object))
+		{
+			return ContextComponent->GetOwner()->FindComponentByClass(InTargetClass.Get());
+		}
+
+		return nullptr;
+	});
+}
+
+void FObjectProxyFactory::RegisterObjectAccessor(FName InAccessorName, const UClass* InTargetClass, FObjectAccessorFunction&& InFunction)
+{
+	check(IsInGameThread());
+	check(InAccessorName != NAME_None);
+	check(InTargetClass != nullptr);
+
+	// Ensure we can use the type as a pin in RigVM graphs 
+	FRigVMRegistry::Get().RegisterObjectTypes( {{ const_cast<UClass*>(InTargetClass), FRigVMRegistry::ERegisterObjectOperation::Class } });
+
+	UE_MT_SCOPED_WRITE_ACCESS(ObjectAccessorsAccessDetector);
+	if(!ObjectAccessors.Contains(InAccessorName))
+	{
+		TSharedRef<FClassProxy> ClassProxy = FindOrCreateClassProxy(InTargetClass);
+		TSharedRef<FObjectAccessor> ObjectAccessor = MakeShared<FObjectAccessor>(InAccessorName, MoveTemp(InFunction), ClassProxy);
+		ObjectAccessors.Add(InAccessorName, ObjectAccessor);
+
+		// Add object accessor to parameter lookup
+		ParameterMap.Add(InAccessorName, ObjectAccessor);
+
+		// Add sub-parameters to parameter lookup
+		for(FName RemappedParameter : ObjectAccessor->RemappedParameters)
+		{
+			ParameterMap.Add(RemappedParameter, ObjectAccessor);
+		}
+	}
+}
+
+void FObjectProxyFactory::UnregisterObjectAccessor(FName InAccessorName)
+{
+	check(IsInGameThread());
+
+	UE_MT_SCOPED_WRITE_ACCESS(ObjectAccessorsAccessDetector);
+	if(const TSharedPtr<FObjectAccessor>* FoundAccessor = ObjectAccessors.Find(InAccessorName))
+	{
+		// Remove sub-parameters
+		for(FName RemappedParameter : (*FoundAccessor)->RemappedParameters)
+		{
+			ParameterMap.Remove(RemappedParameter);
+		}
+
+		// Remove object accessor
+		ParameterMap.Remove(InAccessorName);
+
+		// Remove accessor itself
+		ObjectAccessors.Remove(InAccessorName);
+	}
+}
+
+TSharedRef<FClassProxy> FObjectProxyFactory::FindOrCreateClassProxy(const UClass* InClass)
+{
+	UE_MT_SCOPED_WRITE_ACCESS(ObjectAccessorsAccessDetector);
+	if(TSharedPtr<FClassProxy>* FoundProxy = ClassMap.Find(InClass))
+	{
+		return FoundProxy->ToSharedRef();
+	}
+
+	return ClassMap.Add(InClass, MakeShared<FClassProxy>(InClass)).ToSharedRef();
+}
+
+TSharedPtr<FObjectAccessor> FObjectProxyFactory::FindObjectAccessor(FName InAccessorName) const
+{
+	UE_MT_SCOPED_READ_ACCESS(ObjectAccessorsAccessDetector);
+	return ObjectAccessors.FindRef(InAccessorName);
+}
+
+}
diff --git a/Engine/Plugins/Experimental/AnimNext/Source/AnimNext/Private/Param/ObjectProxyFactory.h b/Engine/Plugins/Experimental/AnimNext/Source/AnimNext/Private/Param/ObjectProxyFactory.h
new file mode 100644
--- /dev/null
+++ b/Engine/Plugins/Experimental/AnimNext/Source/AnimNext/Private/Param/ObjectProxyFactory.h
@@ -0,0 +1,97 @@
+﻿// Copyright Epic Games, Inc. All Rights Reserved.
+
+#pragma once
+
+#include "CoreMinimal.h"
+#include "Param/IParameterSourceFactory.h"
+#include "Features/IModularFeature.h"
+#include "Param/ExternalParameterRegistry.h"
+#include "UObject/ObjectKey.h"
+
+class AActor;
+class UActorComponent;
+class UAnimNextConfig;
+
+namespace UE::AnimNext
+{
+	struct FObjectAccessor;
+	struct FClassProxy;
+	class FModule;
+}
+
+namespace UE::AnimNext
+{
+
+using FObjectAccessorFunction = TUniqueFunction<UObject*(const FExternalParameterContext&)>;
+
+// Factory for object proxies that supply 'external' parameters
+class FObjectProxyFactory : public IParameterSourceFactory
+{
+	friend class ::UAnimNextConfig;
+	friend class FModule;
+
+	// Register built-in accessors and modular feature
+	static void Init();
+
+	// Unregister built-in accessors and modular feature
+	static void Destroy();
+
+	// Refresh built-in accessors and modular feature
+	static void Refresh();
+
+	// Clear internal data
+	void Reset();
+	
+	// IParameterSourceFactory interface
+	virtual void ForEachSource(TFunctionRef<void(FName)> InFunction) const override;
+	virtual TUniquePtr<IParameterSource> CreateParameterSource(const FExternalParameterContext& InContext, FName InSourceName, TConstArrayView<FName> InRequiredParameters) const override;
+#if WITH_EDITOR
+	virtual bool FindParameterInfo(FName InParameterName, FParameterInfo& OutInfo) const override;
+	virtual void ForEachParameter(FName InSourceName, TFunctionRef<void(FName, const FParameterInfo&)> InFunction) const override;
+#endif
+
+	// Register an accessor function to be called on initialization to retrieve an object of a specified class given the object held by the schedule's
+	// entry (e.g. InContextObject is of type UAnimNextComponent, class is of type USkeletalMeshComponent).
+	// Should only be called from the game thread
+	// @param	InAccessorName	The name of the accessor.
+	// @param	InTargetClass	The desired class of the object that will be returned by the accessor function
+	// @param	InFunction		The accessor function that will be called.
+	void RegisterObjectAccessor(FName InAccessorName, const UClass* InTargetClass, FObjectAccessorFunction&& InFunction);
+
+	// Register an accessor function to be called on initialization to retrieve an object of a specified class given the object held by the schedule's
+	// entry (e.g. InContextObject is of type UAnimNextComponent, class is of type USkeletalMeshComponent). If the object's class matches known types,
+	// then a built-in accessor function will be used. If the class does not match a known type, then the function will log a warning.
+	// Should only be called from the game thread
+	// @param	InAccessorName	The name of the accessor.
+	// @param	InTargetClass	The desired class of the object that will be returned by the accessor function
+	void RegisterObjectAccessor(FName InAccessorName, const UClass* InTargetClass);
+	
+	// Helper wrapper for RegisterObjectAccessor. Assumes that the accessed object is the owning actor of the context object.
+	void RegisterActorAccessor(FName InAccessorName, TSubclassOf<AActor> InTargetClass);
+
+	// Helper wrapper for RegisterObjectAccessor. Assumes that the accessed object is an actor component on the same actor as the context object.
+	void RegisterActorComponentAccessor(FName InAccessorName, TSubclassOf<UActorComponent> InTargetClass);
+
+	// Unregister an accessor function that was previously registered to RegisterObjectAccessor
+	// Should only be called from the game thread
+	// @param	InAccessorName	The name of the accessor.
+	void UnregisterObjectAccessor(FName InAccessorName);
+	
+	TSharedRef<FClassProxy> FindOrCreateClassProxy(const UClass* InClass);
+
+	TSharedPtr<FObjectAccessor> FindObjectAccessor(FName InAccessorName) const;
+
+	// Map of classes -> proxy
+	TMap<TObjectKey<UClass>, TSharedPtr<FClassProxy>> ClassMap;
+
+	// Map of accessor name -> accessor
+	TMap<FName, TSharedPtr<FObjectAccessor>> ObjectAccessors;
+
+	// Map of parameter name -> accessor
+	TMap<FName, TWeakPtr<FObjectAccessor>> ParameterMap;
+
+	// Detect concurrent access for ObjectAccessors
+	UE_MT_DECLARE_RW_RECURSIVE_ACCESS_DETECTOR(ObjectAccessorsAccessDetector);
+};
+
+}
\ No newline at end of file
diff --git a/Engine/Plugins/Experimental/AnimNext/Source/AnimNext/Private/Param/ParamAdapter.cpp b/Engine/Plugins/Experimental/AnimNext/Source/AnimNext/Private/Param/ParamAdapter.cpp
deleted file mode 100644
diff --git a/Engine/Plugins/Experimental/AnimNext/Source/AnimNext/Private/Param/ParamAdapter.h b/Engine/Plugins/Experimental/AnimNext/Source/AnimNext/Private/Param/ParamAdapter.h
deleted file mode 100644
diff --git a/Engine/Plugins/Experimental/AnimNext/Source/AnimNext/Private/Param/ParamId.cpp b/Engine/Plugins/Experimental/AnimNext/Source/AnimNext/Private/Param/ParamId.cpp
deleted file mode 100644
diff --git a/Engine/Plugins/Experimental/AnimNext/Source/AnimNext/Private/Param/ParamStack.cpp b/Engine/Plugins/Experimental/AnimNext/Source/AnimNext/Private/Param/ParamStack.cpp
--- a/Engine/Plugins/Experimental/AnimNext/Source/AnimNext/Private/Param/ParamStack.cpp
+++ b/Engine/Plugins/Experimental/AnimNext/Source/AnimNext/Private/Param/ParamStack.cpp
@@ -5,7 +5,6 @@
 #include "Param/ParamHelpers.h"
 #include "PropertyBag.h"
 #include "EngineLogs.h"
-#include "Param/ParamAdapter.h"
 #include "UObject/UObjectGlobals.h"
 #include "UObject/Package.h"
 #include "Param/ParamStackLayer.h"
@@ -20,41 +19,6 @@ DEFINE_STAT(STAT_AnimNext_ParamStack_Decoalesce);
 
 namespace UE::AnimNext
 {
-	
-// Stack layer that can own its own data as a UObject or reference an externally owned object
-struct FUObjectLayer : FParamStackLayer
-{
-	FUObjectLayer() = delete;
-
-	explicit FUObjectLayer(UObject* InObject, bool bInMutable)
-		: FParamStackLayer(32)
-	{
-		Object = InObject;
-
-		const UClass* Class = InObject->GetClass();
-		
-		for (TFieldIterator<FProperty> PropIt(Class, EFieldIteratorFlags::IncludeSuper); PropIt; ++PropIt)
-		{
-			FParamTypeHandle ParamTypeHandle = FParamTypeHandle::FromProperty(*PropIt);
-			if(ParamTypeHandle.IsValid())
-			{
-				const FProperty* Property = *PropIt;
-				FParamId ParamId(PropIt->GetFName());
-				uint8* DataPtr = Property->ContainerPtrToValuePtr<uint8>(InObject);
-				uint32 ParamIndex = Params.Emplace(ParamId, ParamTypeHandle, TArrayView<uint8>(DataPtr, Property->GetSize()), true, bInMutable);
-				HashTable.Add(ParamId.GetHash(), ParamIndex);
-			}
-		}
-	}
-
-	// FParamStackLayer interface
-	virtual UObject* AsUObject() override
-	{
-		return Object.Get();
-	}
-
-	TWeakObjectPtr<UObject> Object;
-};
 
 // Stack layer that can own its own data or reference an external FInstancedPropertyBag
 struct FInstancedPropertyBagLayer : FParamStackLayer, FGCObject
@@ -338,23 +302,6 @@ void FParamStack::PopLayer(FPushedLayerHandle InHandle)
 	}
 }
 
-FParamStackLayerHandle FParamStack::MakeValueLayer(const UClass* InClass)
-{
-	check(InClass != nullptr);
-
-	UObject* OwnedObject = NewObject<UObject>(GetTransientPackage(), InClass);
-	TUniquePtr<FParamStackLayer> Layer = MakeUnique<FUObjectLayer>(OwnedObject, true);
-	return FParamStackLayerHandle(MoveTemp(Layer));
-}
-
-FParamStackLayerHandle FParamStack::MakeReferenceLayer(UObject* InObject)
-{
-	check(InObject != nullptr);
-
-	TUniquePtr<FParamStackLayer> Layer = MakeUnique<FUObjectLayer>(InObject, true);
-	return FParamStackLayerHandle(MoveTemp(Layer));
-}
-
 FParamStackLayerHandle FParamStack::MakeValueLayer(const FInstancedPropertyBag& InPropertyBag)
 {
 	FInstancedPropertyBag OwnedPropertyBag = InPropertyBag;
@@ -389,7 +336,7 @@ FParamResult FParamStack::GetParamData(FParamId InId, FParamTypeHandle InTypeHan
 FParamResult FParamStack::GetParamData(FParamId InId, FParamTypeHandle InTypeHandle, TConstArrayView<uint8>& OutParamData, FParamTypeHandle& OutParamTypeHandle, FParamCompatibility InRequiredCompatibility) const
 {
 	SCOPE_CYCLE_COUNTER(STAT_AnimNext_ParamStack_GetParam);
-	
+
 	const FParamResult Result = GetParamDataInternal(InId, InTypeHandle, OutParamData, OutParamTypeHandle, InRequiredCompatibility);
 	if (Result.IsInScope())
 	{
@@ -403,14 +350,6 @@ FParamResult FParamStack::GetParamData(FParamId InId, FParamTypeHandle InTypeHan
 			return ParentStack->GetParamData(InId, InTypeHandle, OutParamData, OutParamTypeHandle, InRequiredCompatibility);
 		}
 	}
-	
-	// Check if we have a built in adapter to fall back on
-	if(const FParamAdapter* Adapter = InId.GetAdapter(InId))
-	{
-		SCOPE_CYCLE_COUNTER(STAT_AnimNext_ParamStack_Adapter);
-		
-		return Adapter->GetParamData(InTypeHandle, OutParamData, OutParamTypeHandle, InRequiredCompatibility);
-	}
 
 	return EParamResult::NotInScope;
 }
@@ -422,7 +361,7 @@ FParamResult FParamStack::GetParamDataInternal(FParamId InId, FParamTypeHandle I
 	{
 		return EParamResult::NotInScope;
 	}
-	
+
 	const FParamResult Result = ParamPtr->GetParamData(InTypeHandle, OutParamData, OutParamTypeHandle, InRequiredCompatibility);
 	if(Result.IsInScope())
 	{
@@ -463,14 +402,6 @@ FParamResult FParamStack::GetMutableParamData(FParamId InId, FParamTypeHandle In
 		}
 	}
 
-	// Check if we have a built in adapter to fall back on
-	if(const FParamAdapter* Adapter = InId.GetAdapter(InId))
-	{
-		SCOPE_CYCLE_COUNTER(STAT_AnimNext_ParamStack_Adapter);
-		
-		return Adapter->GetMutableParamData(InTypeHandle, OutParamData, OutParamTypeHandle, InRequiredCompatibility);
-	}
-
 	return EParamResult::NotInScope;
 }
 
@@ -622,6 +553,7 @@ void FParamStack::Coalesce()
 		}
 
 		CoalesceLayerHandles.Reserve(NumCoalescedLayers);
+		Layers.Reserve(Layers.Num() + NumCoalescedLayers);
 		for(int32 StackIndex = Stacks.Num() - 1; StackIndex >= 0; --StackIndex)
 		{
 			const FParamStack* Stack = Stacks[StackIndex];
diff --git a/Engine/Plugins/Experimental/AnimNext/Source/AnimNext/Private/Param/ParameterBlockProxy.cpp b/Engine/Plugins/Experimental/AnimNext/Source/AnimNext/Private/Param/ParameterBlockProxy.cpp
new file mode 100644
--- /dev/null
+++ b/Engine/Plugins/Experimental/AnimNext/Source/AnimNext/Private/Param/ParameterBlockProxy.cpp
@@ -0,0 +1,38 @@
+﻿// Copyright Epic Games, Inc. All Rights Reserved.
+
+#include "Param/ParameterBlockProxy.h"
+#include "Param/AnimNextParameterBlock.h"
+#include "Param/ParamStack.h"
+
+namespace UE::AnimNext
+{
+
+FParameterBlockProxy::FParameterBlockProxy(UAnimNextParameterBlock* InParameterBlock)
+	: ParameterBlock(InParameterBlock)
+	, PropertyBag(ParameterBlock->PropertyBag)
+	, LayerHandle(FParamStack::MakeReferenceLayer(PropertyBag))
+{
+	check(ParameterBlock);
+}
+
+void FParameterBlockProxy::Update()
+{
+#if WITH_EDITOR	// Layout should only be changing in editor
+	const FInstancedPropertyBag* HandlePropertyBag = LayerHandle.As<FInstancedPropertyBag>();
+	if(HandlePropertyBag == nullptr || HandlePropertyBag->GetPropertyBagStruct() != ParameterBlock->PropertyBag.GetPropertyBagStruct())
+	{
+		PropertyBag = ParameterBlock->PropertyBag;
+		LayerHandle = FParamStack::MakeReferenceLayer(PropertyBag);
+	}
+#endif
+
+	ParameterBlock->UpdateLayer(LayerHandle);
+}
+
+void FParameterBlockProxy::AddReferencedObjects(FReferenceCollector& Collector)
+{
+	Collector.AddReferencedObject(ParameterBlock);
+	PropertyBag.AddStructReferencedObjects(Collector);
+}
+
+}
diff --git a/Engine/Plugins/Experimental/AnimNext/Source/AnimNext/Private/Param/ParameterBlockProxy.h b/Engine/Plugins/Experimental/AnimNext/Source/AnimNext/Private/Param/ParameterBlockProxy.h
new file mode 100644
--- /dev/null
+++ b/Engine/Plugins/Experimental/AnimNext/Source/AnimNext/Private/Param/ParameterBlockProxy.h
@@ -0,0 +1,37 @@
+﻿// Copyright Epic Games, Inc. All Rights Reserved.
+
+#pragma once
+
+#include "UObject/WeakObjectPtrTemplates.h"
+#include "Param/IParameterSource.h"
+#include "PropertyBag.h"
+#include "UObject/StrongObjectPtr.h"
+
+class UAnimNextParameterBlock;
+
+namespace UE::AnimNext
+{
+
+// Proxy struct used to reference parameter block instance data
+struct FParameterBlockProxy : public IParameterSource
+{
+	FParameterBlockProxy() = delete;
+
+	FParameterBlockProxy(UAnimNextParameterBlock* InParameterBlock);
+
+	// IParameterSource interface
+	virtual void Update() override;
+	virtual const FParamStackLayerHandle& GetLayerHandle() const override { return LayerHandle; }
+	virtual void AddReferencedObjects(FReferenceCollector& Collector) override;
+
+	// The object that this proxy wraps
+	TObjectPtr<UAnimNextParameterBlock> ParameterBlock;
+
+	// Copy of the parameter block's data
+	FInstancedPropertyBag PropertyBag;
+
+	// Layer handle - must be updated if PropertyBag changes layout
+	FParamStackLayerHandle LayerHandle;
+};
+
+}
\ No newline at end of file
diff --git a/Engine/Plugins/Experimental/AnimNext/Source/AnimNext/Private/Param/Params.cpp b/Engine/Plugins/Experimental/AnimNext/Source/AnimNext/Private/Param/Params.cpp
deleted file mode 100644
diff --git a/Engine/Plugins/Experimental/AnimNext/Source/AnimNext/Private/Param/PropertyBagProxy.cpp b/Engine/Plugins/Experimental/AnimNext/Source/AnimNext/Private/Param/PropertyBagProxy.cpp
new file mode 100644
--- /dev/null
+++ b/Engine/Plugins/Experimental/AnimNext/Source/AnimNext/Private/Param/PropertyBagProxy.cpp
@@ -0,0 +1,25 @@
+﻿// Copyright Epic Games, Inc. All Rights Reserved.
+
+#include "Param/PropertyBagProxy.h"
+#include "Param/ParamStack.h"
+
+namespace UE::AnimNext
+{
+
+FPropertyBagProxy::FPropertyBagProxy()
+{
+	LayerHandle = FParamStack::MakeReferenceLayer(PropertyBag);
+}
+
+void FPropertyBagProxy::AddReferencedObjects(FReferenceCollector& Collector)
+{
+	PropertyBag.AddStructReferencedObjects(Collector);
+}
+
+void FPropertyBagProxy::AddPropertyAndValue(FName InName, const FProperty* InProperty, const void* InContainerPtr)
+{
+	PropertyBag.AddProperty(InName, InProperty);
+	PropertyBag.SetValue(InName, InProperty, InContainerPtr);
+}
+
+}
diff --git a/Engine/Plugins/Experimental/AnimNext/Source/AnimNext/Private/Param/PropertyBagProxy.h b/Engine/Plugins/Experimental/AnimNext/Source/AnimNext/Private/Param/PropertyBagProxy.h
new file mode 100644
--- /dev/null
+++ b/Engine/Plugins/Experimental/AnimNext/Source/AnimNext/Private/Param/PropertyBagProxy.h
@@ -0,0 +1,32 @@
+﻿// Copyright Epic Games, Inc. All Rights Reserved.
+
+#pragma once
+
+#include "Param/IParameterSource.h"
+#include "PropertyBag.h"
+
+namespace UE::AnimNext
+{
+
+// Proxy struct used to hold data in a property bag
+struct FPropertyBagProxy : public IParameterSource
+{
+	FPropertyBagProxy();
+
+	// IParameterSource interface
+	virtual void Update() override {};
+	virtual const FParamStackLayerHandle& GetLayerHandle() const override { return LayerHandle; }
+	virtual void AddReferencedObjects(FReferenceCollector& Collector) override;
+
+	// Adds (or replaces) a property with the specified name and sets its value using the specified container
+	void AddPropertyAndValue(FName InName, const FProperty* InProperty, const void* InContainerPtr);
+
+private:
+	// Property bag that this wraps
+	FInstancedPropertyBag PropertyBag;
+
+	// Layer handle - must be updated if PropertyBag is changed
+	FParamStackLayerHandle LayerHandle;
+};
+
+}
\ No newline at end of file
diff --git a/Engine/Plugins/Experimental/AnimNext/Source/AnimNext/Private/Scheduler/AnimNextExternalTaskBinding.h b/Engine/Plugins/Experimental/AnimNext/Source/AnimNext/Private/Scheduler/AnimNextExternalTaskBinding.h
new file mode 100644
--- /dev/null
+++ b/Engine/Plugins/Experimental/AnimNext/Source/AnimNext/Private/Scheduler/AnimNextExternalTaskBinding.h
@@ -0,0 +1,13 @@
+﻿// Copyright Epic Games, Inc. All Rights Reserved.
+
+#pragma once
+
+#include "CoreMinimal.h"
+#include "AnimNextExternalTaskBinding.generated.h"
+
+// Transient base structure used to communicate binding information
+USTRUCT()
+struct FAnimNextExternalTaskBinding
+{
+	GENERATED_BODY()
+};
diff --git a/Engine/Plugins/Experimental/AnimNext/Source/AnimNext/Private/Scheduler/AnimNextSchedule.cpp b/Engine/Plugins/Experimental/AnimNext/Source/AnimNext/Private/Scheduler/AnimNextSchedule.cpp
--- a/Engine/Plugins/Experimental/AnimNext/Source/AnimNext/Private/Scheduler/AnimNextSchedule.cpp
+++ b/Engine/Plugins/Experimental/AnimNext/Source/AnimNext/Private/Scheduler/AnimNextSchedule.cpp
@@ -4,8 +4,10 @@
 #include "Tasks/Task.h"
 #include "Async/TaskGraphInterfaces.h"
 #include "EngineLogs.h"
-#include "Scheduler/AnimNextSchedulePort.h"
-#include "Graph/AnimNextGraph.h"
+
+#if WITH_EDITOR
+TUniqueFunction<void(UAnimNextSchedule*)> UAnimNextSchedule::CompileFunction;
+#endif
 
 void UAnimNextSchedule::PostLoad()
 {
@@ -34,351 +36,9 @@ void UAnimNextSchedule::PostEditUndo()
 
 void UAnimNextSchedule::CompileSchedule()
 {
-	using namespace UE::AnimNext;
-
-	Instructions.Empty();
-	Tasks.Empty();
-	Ports.Empty();
-	ExternalTasks.Empty();
-	ParamScopeEntryTasks.Empty();
-	ParamScopeExitTasks.Empty();
-	IntermediatesData.Reset();
-	NumParameterScopes = 0;
-	NumTickFunctions = 0;
-
-	EAnimNextScheduleScheduleOpcode LastOpCode = EAnimNextScheduleScheduleOpcode::None;
-
-	auto Emit = [this, &LastOpCode](EAnimNextScheduleScheduleOpcode InOpCode, int32 InOperand = 0)
-	{
-		FAnimNextScheduleInstruction Instruction;
-		Instruction.Opcode = InOpCode;
-		Instruction.Operand = InOperand;
-		Instructions.Add(Instruction);
-
-		LastOpCode = InOpCode;
-	};
-
-	auto EmitPrerequisite = [this, &Emit, &LastOpCode]()
-	{
-		switch (LastOpCode)
-		{
-		case EAnimNextScheduleScheduleOpcode::RunTask:
-			Emit(EAnimNextScheduleScheduleOpcode::PrerequisiteTask, NumTickFunctions - 1);
-			break;
-		case EAnimNextScheduleScheduleOpcode::BeginRunExternalTask:
-			Emit(EAnimNextScheduleScheduleOpcode::PrerequisiteBeginExternalTask, NumTickFunctions - 1);
-			break;
-		case EAnimNextScheduleScheduleOpcode::EndRunExternalTask:
-			Emit(EAnimNextScheduleScheduleOpcode::PrerequisiteEndExternalTask, NumTickFunctions - 1);
-			break;
-		case EAnimNextScheduleScheduleOpcode::RunParamScopeEntry:
-			Emit(EAnimNextScheduleScheduleOpcode::PrerequisiteScopeEntry, NumTickFunctions - 1);
-			break;
-		case EAnimNextScheduleScheduleOpcode::RunParamScopeExit:
-			Emit(EAnimNextScheduleScheduleOpcode::PrerequisiteScopeExit, NumTickFunctions - 1);
-			break;
-		default:
-			break;
-		}
-	};
-
-	// MAX_uint32 means 'global scope' in this context
-	uint32 ParentScopeIndex = MAX_uint32;
-
-	TArray<FAnimNextScheduleEntryTerm> IntermediateTerms;
-	TMap<FName, uint32> IntermediateMap;
-
-	TFunction<void(const TArray<TObjectPtr<UAnimNextScheduleEntry>>&)> EmitEntries;
-
-	// Iterate over all entries, recursing into scopes
-	EmitEntries = [this, &EmitEntries, &Emit, &EmitPrerequisite, &ParentScopeIndex, &IntermediateTerms, &IntermediateMap](const TArray<TObjectPtr<UAnimNextScheduleEntry>>& InEntries)
-	{
-		for (int32 EntryIndex = 0; EntryIndex < InEntries.Num(); ++EntryIndex)
-		{
-			UAnimNextScheduleEntry* Entry = InEntries[EntryIndex];
-
-			auto CheckTermDirectionCompatibility = [](FName InName, EScheduleTermDirection InExistingDirection, EScheduleTermDirection InNewDirection)
-			{
-				switch(InExistingDirection)
-				{
-				case EScheduleTermDirection::Input:
-					// Input before output: error
-					UE_LOG(LogAnimation, Error, TEXT("Term '%s' was used as an input before it was output"), *InName.ToString());
-					return false;
-				case EScheduleTermDirection::Output:
-					return true;
-				}
-
-				return false;
-			};
-
-			if (UAnimNextScheduleEntry_Port* PortEntry = Cast<UAnimNextScheduleEntry_Port>(Entry))
-			{
-				bool bValid = true;
-
-				if(PortEntry->Port == nullptr)
-				{
-					UE_LOG(LogAnimation, Error, TEXT("AnimNext: Invalid port class found"));
-					bValid = false;
-				}
-				else
-				{
-					UAnimNextSchedulePort* CDO = PortEntry->Port->GetDefaultObject<UAnimNextSchedulePort>();
-					check(CDO);
-
-					TConstArrayView<FScheduleTerm> Terms = CDO->GetTerms();
-					if(PortEntry->Terms.Num() != Terms.Num())
-					{
-						UE_LOG(LogAnimation, Error, TEXT("AnimNext: Incorrect term count for port: %d"), PortEntry->Terms.Num());
-						bValid = false;
-					}
-
-					for(int32 TermIndex = 0; TermIndex < PortEntry->Terms.Num(); ++TermIndex)
-					{
-						FName TermName = PortEntry->Terms[TermIndex].Name;
-						if(!PortEntry->Terms[TermIndex].Type.IsValid())
-						{
-							UE_LOG(LogAnimation, Error, TEXT("AnimNext: Invalid type when processing port term, ignored: '%s'"), *TermName.ToString());
-							bValid = false;
-						}
-						else
-						{
-							const uint32* ExistingIntermediateIndexPtr = IntermediateMap.Find(TermName);
-							if(ExistingIntermediateIndexPtr != nullptr)
-							{
-								const FAnimNextScheduleEntryTerm& IntermediateTerm = IntermediateTerms[*ExistingIntermediateIndexPtr];
-								if(IntermediateTerm.Type != Terms[TermIndex].GetType())
-								{
-									UE_LOG(LogAnimation, Error, TEXT("AnimNext: Mismatched types when processing port term, ignored: '%s'"), *TermName.ToString());
-									bValid = false;
-								}
-
-								if(!CheckTermDirectionCompatibility(TermName, IntermediateTerm.Direction, Terms[TermIndex].Direction))
-								{
-									bValid = false;
-								}
-							}
-						}
-					}
-				}
-
-				if(bValid)
-				{
-					EmitPrerequisite();
-
-					FAnimNextSchedulePortTask PortTask;
-					PortTask.TaskIndex = Ports.Num();
-					PortTask.ParamScopeIndex = ParentScopeIndex;
-					PortTask.Port = PortEntry->Port;
-
-					for(int32 TermIndex = 0; TermIndex < PortEntry->Terms.Num(); ++TermIndex)
-					{
-						FName TermName = PortEntry->Terms[TermIndex].Name;
-						const uint32* ExistingIntermediateIndexPtr = IntermediateMap.Find(TermName);
-						if(ExistingIntermediateIndexPtr == nullptr)
-						{
-							uint32 IntermediateIndex = IntermediateTerms.Emplace(TermName, PortEntry->Terms[TermIndex].Type, PortEntry->Terms[TermIndex].Direction);
-							IntermediateMap.Add(TermName, IntermediateIndex);
-							PortTask.Terms.Add(IntermediateIndex);
-						}
-						else
-						{
-							PortTask.Terms.Add(*ExistingIntermediateIndexPtr);
-						}
-					}
-
-					int32 PortIndex = Ports.Add(PortTask);
-
-					NumTickFunctions++;
-
-					Emit(EAnimNextScheduleScheduleOpcode::RunPort, PortIndex);
-				}
-			}
-			else if (UAnimNextScheduleEntry_AnimNextGraph* GraphEntry = Cast<UAnimNextScheduleEntry_AnimNextGraph>(Entry))
-			{
-				bool bValid = true;
-
-				if(GraphEntry->Graph == nullptr && GraphEntry->DynamicGraph == NAME_None)
-				{
-					UE_LOG(LogAnimation, Error, TEXT("AnimNext: Invalid graph or no parameter supplied"));
-					bValid = false;
-				}
-				else if(GraphEntry->Graph != nullptr)
-				{
-					TConstArrayView<FScheduleTerm> Terms = GraphEntry->Graph->GetTerms();
-					if(GraphEntry->Terms.Num() != Terms.Num())
-					{
-						UE_LOG(LogAnimation, Error, TEXT("AnimNext: Incorrect term count for graph: %d"), GraphEntry->Terms.Num());
-						bValid = false;
-					}
-					else
-					{
-						// Validate graph terms match schedule-expected terms
-						for(int32 TermIndex = 0; TermIndex < GraphEntry->Terms.Num(); ++TermIndex)
-						{
-							FName TermName = GraphEntry->Terms[TermIndex].Name;
-							if(Terms[TermIndex].Direction != GraphEntry->Terms[TermIndex].Direction)
-							{
-								UE_LOG(LogAnimation, Error, TEXT("AnimNext: Mismatched direction when processing graph term, ignored: '%s'"), *TermName.ToString());
-								bValid = false;
-							}
-							
-							if(Terms[TermIndex].GetType() != GraphEntry->Terms[TermIndex].Type)
-							{
-								UE_LOG(LogAnimation, Error, TEXT("AnimNext: Mismatched types when processing graph term, ignored: '%s'"), *TermName.ToString());
-								bValid = false;
-							}
-						}
-					}
-				}
-
-				// Validate terms and check against priors
-				for(int32 TermIndex = 0; TermIndex < GraphEntry->Terms.Num(); ++TermIndex)
-				{
-					FName TermName = GraphEntry->Terms[TermIndex].Name;
-					if(!GraphEntry->Terms[TermIndex].Type.IsValid())
-					{
-						UE_LOG(LogAnimation, Error, TEXT("AnimNext: Invalid type when processing graph term, ignored: '%s'"), *TermName.ToString());
-						bValid = false;
-					}
-					else
-					{
-						const uint32* ExistingIntermediateIndexPtr = IntermediateMap.Find(TermName);
-						if(ExistingIntermediateIndexPtr != nullptr)
-						{
-							const FAnimNextScheduleEntryTerm& IntermediateTerm = IntermediateTerms[*ExistingIntermediateIndexPtr];
-							if(IntermediateTerm.Type != GraphEntry->Terms[TermIndex].Type)
-							{
-								UE_LOG(LogAnimation, Error, TEXT("AnimNext: Mismatched types when processing graph term, ignored: '%s'"), *TermName.ToString());
-								bValid = false;
-							}
-								
-							if(!CheckTermDirectionCompatibility(TermName, IntermediateTerm.Direction, GraphEntry->Terms[TermIndex].Direction))
-							{
-								bValid = false;
-							}
-						}
-					}
-				}
-
-				if(bValid)
-				{
-					EmitPrerequisite();
-
-					FAnimNextScheduleGraphTask Task;
-					Task.TaskIndex = Tasks.Num();
-					Task.ParamScopeIndex = NumParameterScopes++;
-					Task.ParamParentScopeIndex = ParentScopeIndex;
-					Task.EntryPoint = GraphEntry->EntryPoint;
-					Task.Graph = GraphEntry->Graph;
-					Task.DynamicGraph = GraphEntry->DynamicGraph;
-
-					for(int32 TermIndex = 0; TermIndex < GraphEntry->Terms.Num(); ++TermIndex)
-					{
-						FName TermName = GraphEntry->Terms[TermIndex].Name;
-						const uint32* ExistingIntermediateIndexPtr = IntermediateMap.Find(TermName);
-						if(ExistingIntermediateIndexPtr == nullptr)
-						{
-							uint32 IntermediateIndex = IntermediateTerms.Emplace(TermName, GraphEntry->Terms[TermIndex].Type, GraphEntry->Terms[TermIndex].Direction);
-							IntermediateMap.Add(TermName, IntermediateIndex);
-							Task.Terms.Add(IntermediateIndex);
-						}
-						else
-						{
-							Task.Terms.Add(*ExistingIntermediateIndexPtr);
-						}
-					}
-
-					int32 TaskIndex = Tasks.Add(Task);
-
-					NumTickFunctions++;
-
-					Emit(EAnimNextScheduleScheduleOpcode::RunTask, TaskIndex);
-				}
-			}
-			else if (UAnimNextScheduleEntry_ExternalTask* ExternalTaskEntry = Cast<UAnimNextScheduleEntry_ExternalTask>(Entry))
-			{
-				EmitPrerequisite();
-
-				FAnimNextScheduleExternalTask ExternalTask;
-				ExternalTask.TaskIndex = ExternalTasks.Num();
-				ExternalTask.ParamScopeIndex = NumParameterScopes++;
-				ExternalTask.ParamParentScopeIndex = ParentScopeIndex;
-				ExternalTask.TickFunction = ExternalTaskEntry->TickFunction;
-				ExternalTask.Object = ExternalTaskEntry->Object;
-				int32 ExternalTaskIndex = ExternalTasks.Add(ExternalTask);
-
-				// Emit the external task
-				Emit(EAnimNextScheduleScheduleOpcode::BeginRunExternalTask, ExternalTaskIndex);
-				NumTickFunctions++;
-
-				EmitPrerequisite();
-
-				Emit(EAnimNextScheduleScheduleOpcode::EndRunExternalTask, ExternalTaskIndex);
-				NumTickFunctions++;
-			}
-			else if (UAnimNextScheduleEntry_ParamScope* ParamScopeTaskEntry = Cast<UAnimNextScheduleEntry_ParamScope>(Entry))
-			{
-				EmitPrerequisite();
-
-				FAnimNextScheduleParamScopeEntryTask ParamScopeEntryTask;
-				ParamScopeEntryTask.TaskIndex = ParamScopeEntryTasks.Num();
-				const uint32 ParamScopeIndex = NumParameterScopes++;
-				ParamScopeEntryTask.ParamScopeIndex = ParamScopeIndex;
-				ParamScopeEntryTask.ParamParentScopeIndex = ParentScopeIndex;
-				ParamScopeEntryTask.TickFunctionIndex = NumTickFunctions;
-				ParamScopeEntryTask.Scope = ParamScopeTaskEntry->Scope;
-				ParamScopeEntryTask.ParameterBlocks = ParamScopeTaskEntry->ParameterBlocks;
-				int32 ParamScopeTaskEntryIndex = ParamScopeEntryTasks.Add(ParamScopeEntryTask);
-
-				Emit(EAnimNextScheduleScheduleOpcode::RunParamScopeEntry, ParamScopeTaskEntryIndex);
-				NumTickFunctions++;
-
-				// Enter new scope
-				uint32 PreviousParentScope = ParentScopeIndex;
-				ParentScopeIndex = ParamScopeIndex;
-
-				// Emit the subentries
-				EmitEntries(ParamScopeTaskEntry->SubEntries);
-
-				// Exit scope
-				ParentScopeIndex = PreviousParentScope;
-
-				EmitPrerequisite();
-
-				FAnimNextScheduleParamScopeExitTask ParamScopeExitTask;
-				ParamScopeExitTask.TaskIndex = ParamScopeExitTasks.Num();
-				ParamScopeExitTask.ParamScopeIndex = ParamScopeIndex;
-				ParamScopeExitTask.Scope = ParamScopeTaskEntry->Scope;
-				int32 ParamScopeExitTaskIndex = ParamScopeExitTasks.Add(ParamScopeExitTask);
-
-				Emit(EAnimNextScheduleScheduleOpcode::RunParamScopeExit, ParamScopeExitTaskIndex);
-				NumTickFunctions++;
-			}
-		}
-	};
-
-	EmitEntries(Entries);
-
-	Emit(EAnimNextScheduleScheduleOpcode::Exit);
-
-	// Process intermediates
-	if(IntermediateMap.Num() > 0)
-	{
-		check(IntermediateMap.Num() == IntermediateTerms.Num());
-		
-		TArray<FPropertyBagPropertyDesc> PropertyDescs;
-		PropertyDescs.Reserve(IntermediateTerms.Num());
-		 
-		for(const TPair<FName, uint32>& IntermediatePair : IntermediateMap)
-		{
-			const FAnimNextParamType& IntermediateType = IntermediateTerms[IntermediatePair.Value].Type;
-			check(IntermediateType.IsValid());
-			PropertyDescs.Emplace(IntermediatePair.Key, IntermediateType.GetContainerType(), IntermediateType.GetValueType(), IntermediateType.GetValueTypeObject());
-		}
+	check(CompileFunction);
 
-		IntermediatesData.AddProperties(PropertyDescs);
-	}
+	CompileFunction(this);
 }
 
 #endif // #if WITH_EDITOR
\ No newline at end of file
diff --git a/Engine/Plugins/Experimental/AnimNext/Source/AnimNext/Private/Scheduler/AnimNextScheduleExternalParamTask.cpp b/Engine/Plugins/Experimental/AnimNext/Source/AnimNext/Private/Scheduler/AnimNextScheduleExternalParamTask.cpp
new file mode 100644
--- /dev/null
+++ b/Engine/Plugins/Experimental/AnimNext/Source/AnimNext/Private/Scheduler/AnimNextScheduleExternalParamTask.cpp
@@ -0,0 +1,24 @@
+// Copyright Epic Games, Inc. All Rights Reserved.
+
+#include "Scheduler/AnimNextScheduleExternalParamTask.h"
+#include "AnimNextStats.h"
+#include "ScheduleContext.h"
+#include "ScheduleInstanceData.h"
+#include "Param/IParameterSource.h"
+
+DEFINE_STAT(STAT_AnimNext_Task_ExternalParams);
+
+void FAnimNextScheduleExternalParamTask::UpdateExternalParams(const UE::AnimNext::FScheduleContext& InScheduleContext) const
+{
+	SCOPE_CYCLE_COUNTER(STAT_AnimNext_Task_ExternalParams);
+
+	using namespace UE::AnimNext;
+	
+	FScheduleInstanceData& InstanceData = InScheduleContext.GetInstanceData();
+	FScheduleInstanceData::FExternalParamCache& ExternalParamCache = InstanceData.ExternalParamCaches[TaskIndex];
+	
+	for (TUniquePtr<IParameterSource>& ParameterSource : ExternalParamCache.ParameterSources)
+	{
+		ParameterSource->Update();
+	}
+}
diff --git a/Engine/Plugins/Experimental/AnimNext/Source/AnimNext/Private/Scheduler/AnimNextScheduleGraphTask.cpp b/Engine/Plugins/Experimental/AnimNext/Source/AnimNext/Private/Scheduler/AnimNextScheduleGraphTask.cpp
--- a/Engine/Plugins/Experimental/AnimNext/Source/AnimNext/Private/Scheduler/AnimNextScheduleGraphTask.cpp
+++ b/Engine/Plugins/Experimental/AnimNext/Source/AnimNext/Private/Scheduler/AnimNextScheduleGraphTask.cpp
@@ -14,14 +14,8 @@
 
 DEFINE_STAT(STAT_AnimNext_Task_Graph);
 
-void FAnimNextScheduleGraphTask::RunGraph(const UE::AnimNext::FScheduleContext& InContext) const
+UAnimNextGraph* FAnimNextScheduleGraphTask::GetGraphToRun(UE::AnimNext::FParamStack& ParamStack) const
 {
-	SCOPE_CYCLE_COUNTER(STAT_AnimNext_Task_Graph);
-
-	using namespace UE::AnimNext;
-
-	FParamStack& ParamStack = FParamStack::Get();
-
 	UAnimNextGraph* GraphToRun = Graph;
 	if (GraphToRun == nullptr && DynamicGraph != NAME_None)
 	{
@@ -31,24 +25,37 @@ void FAnimNextScheduleGraphTask::RunGraph(const UE::AnimNext::FScheduleContext&
 		}
 	}
 
+	return GraphToRun;
+}
+
+void FAnimNextScheduleGraphTask::RunGraph(const UE::AnimNext::FScheduleContext& InContext) const
+{
+	SCOPE_CYCLE_COUNTER(STAT_AnimNext_Task_Graph);
+
+	using namespace UE::AnimNext;
+
+	FParamStack& ParamStack = FParamStack::Get();
+
+	UAnimNextGraph* GraphToRun = GetGraphToRun(ParamStack);
 	if(GraphToRun == nullptr)
 	{
 		return;
 	}
 
 	FScheduleInstanceData& InstanceData = InContext.GetInstanceData();
+	FScheduleInstanceData::FGraphCache& GraphCache = InstanceData.GraphCaches[TaskIndex];
 
 	// Check if we are running the correct graph and release (and any term mapping layers) it if not
-	if(InstanceData.GraphInstanceData[TaskIndex].IsValid() && !InstanceData.GraphInstanceData[TaskIndex].UsesGraph(GraphToRun))
+	if(GraphCache.GraphInstanceData.IsValid() && !GraphCache.GraphInstanceData.UsesGraph(GraphToRun))
 	{
-		InstanceData.GraphInstanceData[TaskIndex].Release();
-		InstanceData.GraphTermLayers[TaskIndex].Invalidate();
+		GraphCache.GraphInstanceData.Release();
+		GraphCache.GraphTermLayer.Invalidate();
 	}
 
 	// Allocate our graph instance data
-	if (!InstanceData.GraphInstanceData[TaskIndex].IsValid())
+	if (!GraphCache.GraphInstanceData.IsValid())
 	{
-		GraphToRun->AllocateInstance(InstanceData.GraphInstanceData[TaskIndex]);
+		GraphToRun->AllocateInstance(GraphCache.GraphInstanceData);
 	}
 
 	const FAnimNextGraphReferencePose* GraphReferencePose = ParamStack.GetParamPtr<FAnimNextGraphReferencePose>(GraphToRun->GetReferencePoseParam());
@@ -64,8 +71,7 @@ void FAnimNextScheduleGraphTask::RunGraph(const UE::AnimNext::FScheduleContext&
 	}
 
 	// Check and allocate remapped term layer
-	FParamStackLayerHandle& TermLayerHandle = InstanceData.GraphTermLayers[TaskIndex];
-	if(!TermLayerHandle.IsValid())
+	if(!GraphCache.GraphTermLayer.IsValid())
 	{
 		TConstArrayView<FScheduleTerm> GraphTerms = GraphToRun->GetTerms();
 		check(Terms.Num() == GraphTerms.Num());
@@ -79,11 +85,11 @@ void FAnimNextScheduleGraphTask::RunGraph(const UE::AnimNext::FScheduleContext&
 			Mapping.Add(PropertyDesc.Name, GraphTerms[TermIndex].GetName());
 		}
 
-		TermLayerHandle = FParamStack::MakeRemappedLayer(InstanceData.IntermediatesLayer, Mapping);
+		GraphCache.GraphTermLayer = FParamStack::MakeRemappedLayer(InstanceData.IntermediatesLayer, Mapping);
 	}
 
 	// TODO: This should not be fixed at arg 0, we should define this in the graph asset
-	FAnimNextGraphLODPose* OutputPose = TermLayerHandle.GetMutableParamPtr<FAnimNextGraphLODPose>(GraphToRun->GetTerms()[0].GetId());
+	FAnimNextGraphLODPose* OutputPose = GraphCache.GraphTermLayer.GetMutableParamPtr<FAnimNextGraphLODPose>(GraphToRun->GetTerms()[0].GetId());
 	if(OutputPose == nullptr)
 	{
 		return;
@@ -98,9 +104,6 @@ void FAnimNextScheduleGraphTask::RunGraph(const UE::AnimNext::FScheduleContext&
 
 	check(OutputPose->LODPose.LODLevel == *GraphLODLevel);
 
-	// Push parameter layers that translate schedule data to graph inputs
-	FParamStack::FPushedLayerHandle LayerHandle = ParamStack.PushLayer(InstanceData.GraphInputLayers[TaskIndex]);
-
 	// Internally we use memstack allocation, so we need a mark here
 	FMemStack& MemStack = FMemStack::Get();
 	FMemMark MemMark(MemStack);
@@ -109,10 +112,10 @@ void FAnimNextScheduleGraphTask::RunGraph(const UE::AnimNext::FScheduleContext&
 	// This reduces churn internally by avoiding a chunk to be repeatedly allocated and freed as we push/pop marks
 	MemStack.Alloc(size_t(FPageAllocator::SmallPageSize) + 1, 16);
 
-	UE::AnimNext::UpdateGraph(InstanceData.GraphInstanceData[TaskIndex], InContext.GetDeltaTime());
+	UE::AnimNext::UpdateGraph(GraphCache.GraphInstanceData, InContext.GetDeltaTime());
 
 	{
-		const FEvaluationProgram EvaluationProgram = UE::AnimNext::EvaluateGraph(InstanceData.GraphInstanceData[TaskIndex]);
+		const FEvaluationProgram EvaluationProgram = UE::AnimNext::EvaluateGraph(GraphCache.GraphInstanceData);
 
 		FEvaluationVM EvaluationVM(EEvaluationFlags::All, *GraphReferencePose->ReferencePose, *GraphLODLevel);
 		bool bHasValidOutput = false;
@@ -136,6 +139,4 @@ void FAnimNextScheduleGraphTask::RunGraph(const UE::AnimNext::FScheduleContext&
 			OutputPose->LODPose.CopyFrom(ReferenceKeyframe.Pose);
 		}
 	}
-
-	ParamStack.PopLayer(LayerHandle);
 }
diff --git a/Engine/Plugins/Experimental/AnimNext/Source/AnimNext/Private/Scheduler/AnimNextScheduleParamScopeTask.cpp b/Engine/Plugins/Experimental/AnimNext/Source/AnimNext/Private/Scheduler/AnimNextScheduleParamScopeTask.cpp
--- a/Engine/Plugins/Experimental/AnimNext/Source/AnimNext/Private/Scheduler/AnimNextScheduleParamScopeTask.cpp
+++ b/Engine/Plugins/Experimental/AnimNext/Source/AnimNext/Private/Scheduler/AnimNextScheduleParamScopeTask.cpp
@@ -1,11 +1,10 @@
 // Copyright Epic Games, Inc. All Rights Reserved.
 
 #include "Scheduler/AnimNextScheduleParamScopeTask.h"
-#include "Param/AnimNextParameterBlock.h"
 #include "Scheduler/ScheduleContext.h"
-#include "Scheduler/AnimNextSchedulerEntry.h"
-#include "Param/AnimNextParameterSourceRef.h"
+#include "Param/IParameterSource.h"
 #include "AnimNextStats.h"
+#include "Param/PropertyBagProxy.h"
 
 DEFINE_STAT(STAT_AnimNext_Task_ScopeEntry);
 DEFINE_STAT(STAT_AnimNext_Task_ScopeExit);
@@ -20,76 +19,25 @@ void FAnimNextScheduleParamScopeEntryTask::RunParamScopeEntry(const UE::AnimNext
 
 	FScheduleInstanceData& InstanceData = InScheduleContext.GetInstanceData();
 	FScheduleInstanceData::FScopeCache& ScopeCache = InstanceData.ScopeCaches[TaskIndex];
-	ScopeCache.PushedLayers.Reset();
-
-	const FAnimNextParameterCollection* FoundUserScope = InstanceData.UserScopes.Find(Scope);
-	if(FoundUserScope)
-	{
-		ScopeCache.PushedLayers.Reserve(ParameterBlocks.Num() + FoundUserScope->Parameters.Num());
-	}
-	else
-	{
-		ScopeCache.PushedLayers.Reserve(ParameterBlocks.Num());
-	}
+	check(ScopeCache.PushedLayers.Num() == 0);
 
-	// Resize layer handles appropriately & cache layer
-	if (ScopeCache.StaticHandles.Num() != ParameterBlocks.Num())
+	// Update & push user params
+	if (TUniquePtr<FPropertyBagProxy>* FoundUserScope = InstanceData.UserScopes.Find(Scope))
 	{
-		ScopeCache.StaticHandles.SetNum(ParameterBlocks.Num());
-		for (int32 LayerHandleIndex = 0; LayerHandleIndex < ScopeCache.StaticHandles.Num(); ++LayerHandleIndex)
-		{
-			if(ParameterBlocks[LayerHandleIndex])
-			{
-				ScopeCache.StaticHandles[LayerHandleIndex] = ParameterBlocks[LayerHandleIndex]->CacheLayer();
-			}
-		}
-	}
-
-	if (FoundUserScope)
-	{
-		const UObject* ObjectContext = InstanceData.Entry->ResolvedObject;
-
-		// Resize user scope data appropriately
-		if (ScopeCache.UserHandles.Num() != FoundUserScope->Parameters.Num())
-		{
-			ScopeCache.UserHandles.SetNum(FoundUserScope->Parameters.Num());
-		}
-
-		// Cache any layers if required
-		for (int32 LayerHandleIndex = 0; LayerHandleIndex < ScopeCache.UserHandles.Num(); ++LayerHandleIndex)
-		{
-			if(const IAnimNextParameterSourceInterface* ParameterSource = FoundUserScope->Parameters[LayerHandleIndex].Get(ObjectContext))
-			{
-				if(ParameterSource->ShouldCacheLayer(ScopeCache.UserHandles[LayerHandleIndex]))
-				{
-					ScopeCache.UserHandles[LayerHandleIndex] = ParameterSource->CacheLayer();
-				}
-			}
-		}
-
-		for (int32 ParamBlockIndex = 0; ParamBlockIndex < FoundUserScope->Parameters.Num(); ++ParamBlockIndex)
-		{
-			if(const IAnimNextParameterSourceInterface* ParameterSource = FoundUserScope->Parameters[ParamBlockIndex].Get(ObjectContext))
-			{
-				FParamStackLayerHandle& LayerHandle = ScopeCache.UserHandles[ParamBlockIndex];
-				ParameterSource->UpdateLayer(LayerHandle);
-				ScopeCache.PushedLayers.Add(ParamStack.PushLayer(LayerHandle));
-			}
-		}
+		FPropertyBagProxy& UserParameterSource = *FoundUserScope->Get();
+		UserParameterSource.Update();
+		ScopeCache.PushedLayers.Add(ParamStack.PushLayer(UserParameterSource.GetLayerHandle()));
 	}
 
 	// TODO: Pre/post scope support
 
 	// Update & push static params
-	for (int32 ParamBlockIndex = 0; ParamBlockIndex < ParameterBlocks.Num(); ++ParamBlockIndex)
+	for (int32 ParameterSourceIndex = 0; ParameterSourceIndex < ScopeCache.ParameterSources.Num(); ++ParameterSourceIndex)
 	{
-		if(const UAnimNextParameterBlock* ParameterBlock = ParameterBlocks[ParamBlockIndex])
-		{
-			FParamStackLayerHandle& LayerHandle = ScopeCache.StaticHandles[ParamBlockIndex];
+		TUniquePtr<IParameterSource>& StaticParameterSource = ScopeCache.ParameterSources[ParameterSourceIndex];
 
-			ParameterBlock->UpdateLayer(LayerHandle);
-			ScopeCache.PushedLayers.Add(ParamStack.PushLayer(LayerHandle));
-		}
+		StaticParameterSource->Update();
+		ScopeCache.PushedLayers.Add(ParamStack.PushLayer(StaticParameterSource->GetLayerHandle()));
 	}
 }
 
diff --git a/Engine/Plugins/Experimental/AnimNext/Source/AnimNext/Private/Scheduler/AnimNextSchedulerEntry.cpp b/Engine/Plugins/Experimental/AnimNext/Source/AnimNext/Private/Scheduler/AnimNextSchedulerEntry.cpp
--- a/Engine/Plugins/Experimental/AnimNext/Source/AnimNext/Private/Scheduler/AnimNextSchedulerEntry.cpp
+++ b/Engine/Plugins/Experimental/AnimNext/Source/AnimNext/Private/Scheduler/AnimNextSchedulerEntry.cpp
@@ -6,6 +6,8 @@
 #include "Scheduler/AnimNextSchedule.h"
 #include "Engine/World.h"
 #include "AnimNextStats.h"
+#include "AnimNextTickFunctionBinding.h"
+#include "Logging/StructuredLog.h"
 
 DEFINE_STAT(STAT_AnimNext_InitializeEntry);
 
@@ -49,6 +51,9 @@ void FAnimNextSchedulerEntry::Initialize(TUniqueFunction<void(const UE::AnimNext
 
 		RootParamStack = MakeShared<FParamStack>();
 
+		// Allocate instance data
+		Context.InstanceData = MakeUnique<FScheduleInstanceData>(Context, Schedule, Handle, this);
+
 		FParamStack::AttachToCurrentThread(RootParamStack);
 		FScheduleContext::AttachToCurrentThread(Context);
 
@@ -92,26 +97,26 @@ void FAnimNextSchedulerEntry::Initialize(TUniqueFunction<void(const UE::AnimNext
 						TConstArrayView<FAnimNextScheduleInstruction>(&Instruction, 1),
 						TConstArrayView<TWeakObjectPtr<UObject>>(&TargetObject, 1)));
 
-					const FName TickFunctionName = Schedule->ExternalTasks[Instruction.Operand].TickFunction;
-					const FName ObjectName = Schedule->ExternalTasks[Instruction.Operand].Object;
-					if(TickFunctionName != NAME_None && ObjectName != NAME_None)
+					const FName ExternalTaskName = Schedule->ExternalTasks[Instruction.Operand].ExternalTask;
+					if(ExternalTaskName != NAME_None)
 					{
-						FTickFunction* FoundFunction = ParamStack.GetMutableParamPtr<FTickFunction>(TickFunctionName);
-						const TObjectPtr<UObject>* FoundObject = ParamStack.GetParamPtr<TObjectPtr<UObject>>(ObjectName);
-						if(FoundFunction && FoundObject && FoundObject->Get())
+						FAnimNextTickFunctionBinding* FoundBinding = ParamStack.GetMutableParamPtr<FAnimNextTickFunctionBinding>(ExternalTaskName);
+						if(FoundBinding)
 						{
-							FoundFunction->AddPrerequisite(ResolvedObject, *TickFunctions.Last().Get());
-							TargetObject = FoundObject->Get();
-							TickFunctions.Last()->Subsequents.Emplace(FoundObject->Get(), *FoundFunction);
+							check(FoundBinding->Object.Get());
+							check(FoundBinding->TickFunction);
+							FoundBinding->TickFunction->AddPrerequisite(ResolvedObject, *TickFunctions.Last().Get());
+							TargetObject = FoundBinding->Object.Get();
+							TickFunctions.Last()->Subsequents.Emplace(FoundBinding->Object.Get(), *FoundBinding->TickFunction);
 						}
 						else
 						{
-							UE_LOG(LogAnimation, Warning, TEXT("AnimNext: Could not bind to external tick function using binding parameters (TickFunctionName=%s, ObjectName=%s)"), *TickFunctionName.ToString(), *ObjectName.ToString());
+							UE_LOGFMT(LogAnimation, Warning, "AnimNext: Could not bind to external tick function using binding parameters (ExternalTaskName={ExternalTaskName})", ExternalTaskName);
 						}
 					}
 					else
 					{
-						UE_LOG(LogAnimation, Warning, TEXT("AnimNext: Invalid external tick function binding parameters (TickFunctionName=%s, ObjectName=%s)"), *TickFunctionName.ToString(), *ObjectName.ToString());
+						UE_LOGFMT(LogAnimation, Warning, "AnimNext: Invalid external tick function binding parameters (ExternalTaskName={ExternalTaskName})", ExternalTaskName);
 					}
 					break;
 				}
@@ -122,21 +127,19 @@ void FAnimNextSchedulerEntry::Initialize(TUniqueFunction<void(const UE::AnimNext
 						TConstArrayView<FAnimNextScheduleInstruction>(&Instruction, 1),
 						TConstArrayView<TWeakObjectPtr<UObject>>(&TargetObject, 1)));
 
-					const FName TickFunctionName = Schedule->ExternalTasks[Instruction.Operand].TickFunction;
-					const FName ObjectName = Schedule->ExternalTasks[Instruction.Operand].Object;
-					if(TickFunctionName != NAME_None && ObjectName != NAME_None)
+					const FName ExternalTaskName = Schedule->ExternalTasks[Instruction.Operand].ExternalTask;
+					if(ExternalTaskName != NAME_None)
 					{
-						FTickFunction* FoundFunction = ParamStack.GetMutableParamPtr<FTickFunction>(TickFunctionName);
-						const TObjectPtr<UObject>* FoundObject = ParamStack.GetParamPtr<TObjectPtr<UObject>>(ObjectName);
-						if(FoundFunction && FoundObject && FoundObject->Get())
+						FAnimNextTickFunctionBinding* FoundBinding = ParamStack.GetMutableParamPtr<FAnimNextTickFunctionBinding>(ExternalTaskName);
+						if(FoundBinding)
 						{
-							TargetObject = FoundObject->Get();
-							TickFunctions.Last()->AddPrerequisite(FoundObject->Get(), *FoundFunction);
+							TargetObject = FoundBinding->Object.Get();
+							TickFunctions.Last()->AddPrerequisite(FoundBinding->Object.Get(), *FoundBinding->TickFunction);
 						}
 					}
 					break;
 				}
-			case EAnimNextScheduleScheduleOpcode::RunTask:
+			case EAnimNextScheduleScheduleOpcode::RunGraphTask:
 			case EAnimNextScheduleScheduleOpcode::RunPort:
 			case EAnimNextScheduleScheduleOpcode::RunParamScopeEntry:
 			case EAnimNextScheduleScheduleOpcode::RunParamScopeExit:
@@ -145,11 +148,22 @@ void FAnimNextSchedulerEntry::Initialize(TUniqueFunction<void(const UE::AnimNext
 					TConstArrayView<FAnimNextScheduleInstruction>(&Instruction, 1),
 					TConstArrayView<TWeakObjectPtr<UObject>>(&TargetObject, 1)));
 				break;
+			case EAnimNextScheduleScheduleOpcode::RunExternalParamTask:
+				{
+					TUniquePtr<FScheduleTickFunction> NewTickFunction = MakeUnique<FScheduleTickFunction>(
+						Context,
+						TConstArrayView<FAnimNextScheduleInstruction>(&Instruction, 1),
+						TConstArrayView<TWeakObjectPtr<UObject>>(&TargetObject, 1));
+					NewTickFunction->bRunOnAnyThread = Schedule->ExternalParamTasks[Instruction.Operand].bThreadSafe;
+					AddTickFunctionAndPrerequisites(MoveTemp(NewTickFunction));
+					break;
+				}
 			case EAnimNextScheduleScheduleOpcode::PrerequisiteTask:
 			case EAnimNextScheduleScheduleOpcode::PrerequisiteBeginExternalTask:
 			case EAnimNextScheduleScheduleOpcode::PrerequisiteEndExternalTask:
 			case EAnimNextScheduleScheduleOpcode::PrerequisiteScopeEntry:
 			case EAnimNextScheduleScheduleOpcode::PrerequisiteScopeExit:
+			case EAnimNextScheduleScheduleOpcode::PrerequisiteExternalParamTask:
 				check(TickFunctions[Instruction.Operand].IsValid());
 				Prerequisites.Add(&TickFunctions[Instruction.Operand]);
 				break;
@@ -159,7 +173,7 @@ void FAnimNextSchedulerEntry::Initialize(TUniqueFunction<void(const UE::AnimNext
 				TickFunctions.Last()->Subsequents.Emplace(ResolvedObject, *EndTickFunction.Get());
 				break;
 			default:
-				check(false);
+				checkNoEntry();
 				break;
 			}
 		}
@@ -180,8 +194,6 @@ void FAnimNextSchedulerEntry::Initialize(TUniqueFunction<void(const UE::AnimNext
 			TickFunction->RegisterTickFunction(Level);
 		}
 
-		LazyAllocateInstanceData();
-
 		RunState = ERunState::RunningInitialUpdate;
 
 		if(InInitializeCallback)
@@ -279,14 +291,3 @@ void FAnimNextSchedulerEntry::Enable(bool bInEnabled)
 
 	RunState = bInEnabled ? ERunState::Running : ERunState::Paused;
 }
-
-void FAnimNextSchedulerEntry::LazyAllocateInstanceData()
-{
-	using namespace UE::AnimNext;
-
-	// Allocate instance data if required
-	if (!Context.InstanceData.IsValid())
-	{
-		Context.InstanceData = MakeUnique<FScheduleInstanceData>(Context, Schedule, Handle, this);
-	}
-}
\ No newline at end of file
diff --git a/Engine/Plugins/Experimental/AnimNext/Source/AnimNext/Private/Scheduler/AnimNextSchedulerEntry.h b/Engine/Plugins/Experimental/AnimNext/Source/AnimNext/Private/Scheduler/AnimNextSchedulerEntry.h
--- a/Engine/Plugins/Experimental/AnimNext/Source/AnimNext/Private/Scheduler/AnimNextSchedulerEntry.h
+++ b/Engine/Plugins/Experimental/AnimNext/Source/AnimNext/Private/Scheduler/AnimNextSchedulerEntry.h
@@ -38,9 +38,6 @@ struct FAnimNextSchedulerEntry
 	// Clears the bTickEvenWhenPaused flags of the entries' tick functions
 	void ClearTickFunctionPauseFlags();
 
-	// Allocate instance data if required
-	void LazyAllocateInstanceData();
-
 	UPROPERTY(Transient)
 	TObjectPtr<const UAnimNextSchedule> Schedule = nullptr;
 
diff --git a/Engine/Plugins/Experimental/AnimNext/Source/AnimNext/Private/Scheduler/AnimNextSchedulerWorldSubsystem.cpp b/Engine/Plugins/Experimental/AnimNext/Source/AnimNext/Private/Scheduler/AnimNextSchedulerWorldSubsystem.cpp
--- a/Engine/Plugins/Experimental/AnimNext/Source/AnimNext/Private/Scheduler/AnimNextSchedulerWorldSubsystem.cpp
+++ b/Engine/Plugins/Experimental/AnimNext/Source/AnimNext/Private/Scheduler/AnimNextSchedulerWorldSubsystem.cpp
@@ -4,7 +4,6 @@
 #include "ScheduleInstanceData.h"
 #include "Scheduler/AnimNextSchedule.h"
 #include "Param/ParamDefinition.h"
-#include "Param/Params.h"
 #include "Engine/World.h"
 
 namespace UE::AnimNext
diff --git a/Engine/Plugins/Experimental/AnimNext/Source/AnimNext/Private/Scheduler/AnimNextTickFunctionBinding.h b/Engine/Plugins/Experimental/AnimNext/Source/AnimNext/Private/Scheduler/AnimNextTickFunctionBinding.h
new file mode 100644
--- /dev/null
+++ b/Engine/Plugins/Experimental/AnimNext/Source/AnimNext/Private/Scheduler/AnimNextTickFunctionBinding.h
@@ -0,0 +1,23 @@
+﻿// Copyright Epic Games, Inc. All Rights Reserved.
+
+#pragma once
+
+#include "UObject/WeakObjectPtrTemplates.h"
+#include "AnimNextExternalTaskBinding.h"
+#include "AnimNextTickFunctionBinding.generated.h"
+
+struct FTickFunction;
+
+// Wrapper around an object and a tick function
+// Transient structure used to communicate binding information
+USTRUCT(BlueprintInternalUseOnly)
+struct FAnimNextTickFunctionBinding : public FAnimNextExternalTaskBinding
+{
+	GENERATED_BODY()
+
+	// The object the tick function exists on
+	TWeakObjectPtr<UObject> Object;
+
+	// The tick function itself
+	FTickFunction* TickFunction = nullptr;
+};
diff --git a/Engine/Plugins/Experimental/AnimNext/Source/AnimNext/Private/Scheduler/Ports/AnimNextSchedulePort_AnimNextMeshComponentPose.cpp b/Engine/Plugins/Experimental/AnimNext/Source/AnimNext/Private/Scheduler/Ports/AnimNextSchedulePort_AnimNextMeshComponentPose.cpp
--- a/Engine/Plugins/Experimental/AnimNext/Source/AnimNext/Private/Scheduler/Ports/AnimNextSchedulePort_AnimNextMeshComponentPose.cpp
+++ b/Engine/Plugins/Experimental/AnimNext/Source/AnimNext/Private/Scheduler/Ports/AnimNextSchedulePort_AnimNextMeshComponentPose.cpp
@@ -13,6 +13,9 @@
 
 DEFINE_STAT(STAT_AnimNext_Port_SkeletalMeshComponent);
 
+UE::AnimNext::FParamId UAnimNextSchedulePort_AnimNextMeshComponentPose::ComponentParamId("UE_AnimNextMeshComponent");
+UE::AnimNext::FParamId UAnimNextSchedulePort_AnimNextMeshComponentPose::ReferencePoseParamId("UE_AnimNextMeshComponent_ReferencePose");
+
 void UAnimNextSchedulePort_AnimNextMeshComponentPose::Run(const UE::AnimNext::FScheduleTermContext& InContext) const
 {
 	SCOPE_CYCLE_COUNTER(STAT_AnimNext_Port_SkeletalMeshComponent);
@@ -21,7 +24,6 @@ void UAnimNextSchedulePort_AnimNextMeshComponentPose::Run(const UE::AnimNext::FS
 
 	const FParamStack& ParamStack = FParamStack::Get();
 
-	static FParamId ComponentParamId("UE_AnimNextMeshComponent");
 	const TObjectPtr<UAnimNextMeshComponent>* ComponentPtr = ParamStack.GetParamPtr<TObjectPtr<UAnimNextMeshComponent>>(ComponentParamId);
 	if(ComponentPtr == nullptr)
 	{
@@ -39,7 +41,6 @@ void UAnimNextSchedulePort_AnimNextMeshComponentPose::Run(const UE::AnimNext::FS
 		return;
 	}
 
-	static const FParamId ReferencePoseParamId("UE_AnimNextMeshComponent_ReferencePose");
 	const FAnimNextGraphReferencePose* GraphReferencePose = ParamStack.GetParamPtr<FAnimNextGraphReferencePose>(ReferencePoseParamId);
 	if(GraphReferencePose == nullptr)
 	{
@@ -82,3 +83,16 @@ TConstArrayView<UE::AnimNext::FScheduleTerm> UAnimNextSchedulePort_AnimNextMeshC
 
 	return Terms;
 }
+
+TConstArrayView<FName> UAnimNextSchedulePort_AnimNextMeshComponentPose::GetRequiredParameters() const
+{
+	using namespace UE::AnimNext;
+
+	static const FName ParamNames[] =
+	{
+		ComponentParamId.GetName(),
+		ReferencePoseParamId.GetName()
+	};
+
+	return ParamNames;
+}
diff --git a/Engine/Plugins/Experimental/AnimNext/Source/AnimNext/Private/Scheduler/Ports/AnimNextSchedulePort_AnimNextMeshComponentPose.h b/Engine/Plugins/Experimental/AnimNext/Source/AnimNext/Private/Scheduler/Ports/AnimNextSchedulePort_AnimNextMeshComponentPose.h
--- a/Engine/Plugins/Experimental/AnimNext/Source/AnimNext/Private/Scheduler/Ports/AnimNextSchedulePort_AnimNextMeshComponentPose.h
+++ b/Engine/Plugins/Experimental/AnimNext/Source/AnimNext/Private/Scheduler/Ports/AnimNextSchedulePort_AnimNextMeshComponentPose.h
@@ -12,5 +12,11 @@ class UAnimNextSchedulePort_AnimNextMeshComponentPose : public UAnimNextSchedule
 
 	// UAnimNextSchedulePort interface
 	virtual void Run(const UE::AnimNext::FScheduleTermContext& InContext) const override;
+	virtual TConstArrayView<FName> GetRequiredParameters() const override;
+	
+	// IAnimNextScheduleTermInterface interface
 	virtual TConstArrayView<UE::AnimNext::FScheduleTerm> GetTerms() const override;
+
+	static UE::AnimNext::FParamId ComponentParamId;
+	static UE::AnimNext::FParamId ReferencePoseParamId;
 };
\ No newline at end of file
diff --git a/Engine/Plugins/Experimental/AnimNext/Source/AnimNext/Private/Scheduler/ScheduleInstanceData.cpp b/Engine/Plugins/Experimental/AnimNext/Source/AnimNext/Private/Scheduler/ScheduleInstanceData.cpp
--- a/Engine/Plugins/Experimental/AnimNext/Source/AnimNext/Private/Scheduler/ScheduleInstanceData.cpp
+++ b/Engine/Plugins/Experimental/AnimNext/Source/AnimNext/Private/Scheduler/ScheduleInstanceData.cpp
@@ -7,6 +7,11 @@
 #include "Scheduler/AnimNextSchedulerEntry.h"
 #include "Scheduler/AnimNextSchedulePort.h"
 #include "AnimNextStats.h"
+#include "Param/ExternalParameterRegistry.h"
+#include "Param/IParameterSourceFactory.h"
+#include "Param/ParameterBlockProxy.h"
+#include "Param/PropertyBagProxy.h"
+#include "Scheduler/AnimNextScheduleExternalParamTask.h"
 
 DEFINE_STAT(STAT_AnimNext_CreateInstanceData);
 
@@ -21,9 +26,24 @@ FScheduleInstanceData::FScheduleInstanceData(const FScheduleContext& InScheduleC
 
 	// Preallocate data for all scopes & graphs in the schedule
 	ScopeCaches.SetNum(InSchedule->ParamScopeEntryTasks.Num());
-	GraphInstanceData.SetNum(InSchedule->Tasks.Num());
-	GraphInputLayers.SetNum(InSchedule->Tasks.Num());
-
+	GraphCaches.SetNum(InSchedule->GraphTasks.Num());
+	ExternalParamCaches.SetNum(InSchedule->ExternalParamTasks.Num());
+	
+	for(int32 ScopeIndex = 0; ScopeIndex < ScopeCaches.Num(); ++ScopeIndex)
+	{
+		FScopeCache& ScopeCache = ScopeCaches[ScopeIndex];
+		ScopeCache.ParameterSources.Reserve(InSchedule->ParamScopeEntryTasks[ScopeIndex].ParameterBlocks.Num());
+		for(UAnimNextParameterBlock* ParameterBlock : InSchedule->ParamScopeEntryTasks[ScopeIndex].ParameterBlocks)
+		{
+			if(ParameterBlock)
+			{
+				ScopeCache.ParameterSources.Emplace(MakeUnique<FParameterBlockProxy>(ParameterBlock));
+			}
+		}
+
+		ScopeCache.PushedLayers.Reserve(ScopeCache.ParameterSources.Num() + 1); // +1 for any user handles added dynamically
+	}
+	
 	// Setup param stack graph
 	RootParamStack = InCurrentEntry->RootParamStack;
 	ParamStacks.SetNum(InSchedule->NumParameterScopes);
@@ -32,7 +52,7 @@ FScheduleInstanceData::FScheduleInstanceData(const FScheduleContext& InScheduleC
 		ParamStack = MakeShared<FParamStack>();
 	}
 
-	for (const FAnimNextScheduleGraphTask& Task : InSchedule->Tasks)
+	for (const FAnimNextScheduleGraphTask& Task : InSchedule->GraphTasks)
 	{
 		const TSharedPtr<FParamStack> ParentStack = Task.ParamParentScopeIndex != MAX_uint32 ? ParamStacks[Task.ParamParentScopeIndex] : RootParamStack;
 		ParamStacks[Task.ParamScopeIndex]->SetParent(ParentStack);
@@ -50,27 +70,51 @@ FScheduleInstanceData::FScheduleInstanceData(const FScheduleContext& InScheduleC
 		ParamStacks[ScopeEntryTask.ParamScopeIndex]->SetParent(ParentStack);
 	}
 
+	// Set up external parameters
+	FExternalParameterContext ExternalParameterContext;
+	ExternalParameterContext.Object = Entry->WeakObject.Get();
+
+	for(int32 ExternalParamSourceIndex = 0; ExternalParamSourceIndex < ExternalParamCaches.Num(); ++ExternalParamSourceIndex)
+	{
+		FExternalParamCache& ExternalParamCache = ExternalParamCaches[ExternalParamSourceIndex];
+
+		for(const FAnimNextScheduleExternalParameterSource& ParameterSource : InSchedule->ExternalParamTasks[ExternalParamSourceIndex].ParameterSources)
+		{
+			if(TUniquePtr<IParameterSource> NewParameterSource = FExternalParameterRegistry::CreateParameterSource(ExternalParameterContext, ParameterSource.ParameterSource, ParameterSource.Parameters))
+			{
+				// Initial update is required to populate the cache
+				// TODO: This needs to move outside this function once we run initialization off the game thread, depending on thread-safety
+				NewParameterSource->Update();
+
+				// External parameter layer is always pushed
+				RootParamStack->PushLayer(NewParameterSource->GetLayerHandle());
+				ExternalParamCache.ParameterSources.Add(MoveTemp(NewParameterSource));
+			}
+		}
+	}
+
 	// Duplicate intermediate data area
 	IntermediatesData = InSchedule->IntermediatesData;
 
 	// Make a hosting layer for the intermediates
 	IntermediatesLayer = FParamStack::MakeReferenceLayer(IntermediatesData);
 
-	// Resize remapped intermediate data layers for graph and port tasks, they will be allocated lazily later
-	GraphTermLayers.SetNum(InSchedule->Tasks.Num());
+	// Resize remapped intermediate data layers for port tasks, they will be allocated lazily later
 	PortTermLayers.SetNum(InSchedule->Ports.Num());
 }
 
+FScheduleInstanceData::~FScheduleInstanceData() = default;
+
 void FScheduleInstanceData::AddReferencedObjects(FReferenceCollector& Collector)
 {
-	for (TPair<FName, FAnimNextParameterCollection>& ParamPair : UserScopes)
+	for (TPair<FName, TUniquePtr<FPropertyBagProxy>>& ParamPair : UserScopes)
 	{
-		Collector.AddPropertyReferencesWithStructARO(FAnimNextParameterCollection::StaticStruct(), &ParamPair.Value);
+		ParamPair.Value->AddReferencedObjects(Collector);
 	}
 
-	for (FAnimNextGraphInstance& GraphInstance : GraphInstanceData)
+	for (FGraphCache& GraphCache : GraphCaches)
 	{
-		Collector.AddPropertyReferencesWithStructARO(FAnimNextGraphInstance::StaticStruct(), &GraphInstance);
+		Collector.AddPropertyReferencesWithStructARO(FAnimNextGraphInstance::StaticStruct(), &GraphCache.GraphInstanceData);
 	}
 }
 
diff --git a/Engine/Plugins/Experimental/AnimNext/Source/AnimNext/Private/Scheduler/ScheduleInstanceData.h b/Engine/Plugins/Experimental/AnimNext/Source/AnimNext/Private/Scheduler/ScheduleInstanceData.h
--- a/Engine/Plugins/Experimental/AnimNext/Source/AnimNext/Private/Scheduler/ScheduleInstanceData.h
+++ b/Engine/Plugins/Experimental/AnimNext/Source/AnimNext/Private/Scheduler/ScheduleInstanceData.h
@@ -3,10 +3,9 @@
 #pragma once
 
 #include "CoreMinimal.h"
-#include "LODPose.h"
 #include "ScheduleHandle.h"
-#include "Param/AnimNextParameterCollection.h"
 #include "Graph/AnimNextGraph.h"
+#include "Param/ParamStack.h"
 
 struct FAnimNextGraphInstance;
 class UAnimNextSchedule;
@@ -16,6 +15,8 @@ namespace UE::AnimNext
 {
 	struct FParamStack;
 	struct FParamStackLayerHandle;
+	class IParameterSource;
+	struct FPropertyBagProxy;
 }
 
 namespace UE::AnimNext
@@ -26,6 +27,8 @@ struct FScheduleInstanceData : public FGCObject
 {
 	FScheduleInstanceData(const FScheduleContext& InScheduleContext, const UAnimNextSchedule* InSchedule, FScheduleHandle InHandle, FAnimNextSchedulerEntry* InCurrentEntry);
 
+	~FScheduleInstanceData();
+	
 	// FGCObject interface
 	virtual void AddReferencedObjects(FReferenceCollector& Collector) override;
 	virtual FString GetReferencerName() const override;
@@ -40,16 +43,13 @@ struct FScheduleInstanceData : public FGCObject
 	FAnimNextSchedulerEntry* Entry = nullptr;
 
 	// Set of dynamic parameter scopes supplied by the user
-	TMap<FName, FAnimNextParameterCollection> UserScopes;
+	TMap<FName, TUniquePtr<FPropertyBagProxy>> UserScopes;
 
 	// Cached data for each parameter scope
 	struct FScopeCache
 	{
-		// Cached handles for scheduled parameter blocks
-		TArray<FParamStackLayerHandle> StaticHandles;
-
-		// Cached handles for user-defined parameter blocks hooked to a scope
-		TArray<FParamStackLayerHandle> UserHandles;
+		// Parameter sources at this scope
+		TArray<TUniquePtr<IParameterSource>> ParameterSources;
 
 		// Pushed layers, popped when scope exits
 		TArray<FParamStack::FPushedLayerHandle> PushedLayers;
@@ -67,23 +67,35 @@ struct FScheduleInstanceData : public FGCObject
 	// User handles initialized at startup, always pushed
 	TArray<FParamStackLayerHandle> StaticUserHandles;
 
-	// Graph instance data for each graph task
-	TArray<FAnimNextGraphInstance> GraphInstanceData;
-
-	// Layer handles for translating schedule terms to graph inputs
-	TArray<FParamStackLayerHandle> GraphInputLayers;
-
 	// Intermediate data area
 	FInstancedPropertyBag IntermediatesData;
 
 	// Layer for intermediates data
 	FParamStackLayerHandle IntermediatesLayer;
 
-	// Remapped data layers for each graph
-	TArray<FParamStackLayerHandle> GraphTermLayers;
-
 	// Remapped data layers for each port
 	TArray<FParamStackLayerHandle> PortTermLayers;
+
+	// Cached data for each graph
+	struct FGraphCache
+	{
+		// Graph instance data
+		FAnimNextGraphInstance GraphInstanceData;
+
+		// Remapped data layers for input terms (from schedule)
+		FParamStackLayerHandle GraphTermLayer;
+	};
+
+	TArray<FGraphCache> GraphCaches;
+
+	// Cached data for each external param task
+	struct FExternalParamCache
+	{
+		// Parameter sources at this scope
+		TArray<TUniquePtr<IParameterSource>> ParameterSources;
+	};
+
+	TArray<FExternalParamCache> ExternalParamCaches;
 };
 
 }
\ No newline at end of file
diff --git a/Engine/Plugins/Experimental/AnimNext/Source/AnimNext/Private/Scheduler/ScheduleTickFunction.cpp b/Engine/Plugins/Experimental/AnimNext/Source/AnimNext/Private/Scheduler/ScheduleTickFunction.cpp
--- a/Engine/Plugins/Experimental/AnimNext/Source/AnimNext/Private/Scheduler/ScheduleTickFunction.cpp
+++ b/Engine/Plugins/Experimental/AnimNext/Source/AnimNext/Private/Scheduler/ScheduleTickFunction.cpp
@@ -15,8 +15,6 @@ namespace UE::AnimNext
 
 void FScheduleBeginTickFunction::ExecuteTick(float DeltaTime, ELevelTick TickType, ENamedThreads::Type CurrentThread, const FGraphEventRef& MyCompletionGraphEvent)
 {
-	Entry.LazyAllocateInstanceData();
-
 	Entry.ResolvedObject = Entry.WeakObject.Get();
 	Entry.DeltaTime = DeltaTime;
 }
@@ -116,13 +114,13 @@ void FScheduleTickFunction::RunSchedule(TConstArrayView<FAnimNextScheduleInstruc
 		const FAnimNextScheduleInstruction& Instruction = InInstructions[InstructionIndex];
 		switch (Instruction.Opcode)
 		{
-		case EAnimNextScheduleScheduleOpcode::RunTask:
+		case EAnimNextScheduleScheduleOpcode::RunGraphTask:
 			{
-				uint32 TaskIndex = Instruction.Operand;
+				uint32 GraphTaskIndex = Instruction.Operand;
 				FScheduleInstanceData& InstanceData = ScheduleContext.GetInstanceData();
-				FParamStack::AttachToCurrentThread(InstanceData.GetParamStack(Schedule->Tasks[TaskIndex].ParamScopeIndex), FParamStack::ECoalesce::Coalesce);
+				FParamStack::AttachToCurrentThread(InstanceData.GetParamStack(Schedule->GraphTasks[GraphTaskIndex].ParamScopeIndex), FParamStack::ECoalesce::Coalesce);
 
-				Schedule->Tasks[TaskIndex].RunGraph(ScheduleContext);
+				Schedule->GraphTasks[GraphTaskIndex].RunGraph(ScheduleContext);
 
 				FParamStack::DetachFromCurrentThread(FParamStack::EDecoalesce::Decoalesce);
 				break;
@@ -189,6 +187,17 @@ void FScheduleTickFunction::RunSchedule(TConstArrayView<FAnimNextScheduleInstruc
 				FParamStack::DetachFromCurrentThread(FParamStack::EDecoalesce::Decoalesce);
 				break;
 			}
+		case EAnimNextScheduleScheduleOpcode::RunExternalParamTask:
+			{
+				uint32 ExternalParamIndex = Instruction.Operand;
+				FScheduleInstanceData& InstanceData = ScheduleContext.GetInstanceData();
+				FParamStack::AttachToCurrentThread(InstanceData.RootParamStack);
+
+				Schedule->ExternalParamTasks[ExternalParamIndex].UpdateExternalParams(ScheduleContext);
+
+				FParamStack::DetachFromCurrentThread();
+				break;
+			}
 		default:
 			break;
 		}
diff --git a/Engine/Plugins/Experimental/AnimNext/Source/AnimNext/Private/Scheduler/Scheduler.cpp b/Engine/Plugins/Experimental/AnimNext/Source/AnimNext/Private/Scheduler/Scheduler.cpp
--- a/Engine/Plugins/Experimental/AnimNext/Source/AnimNext/Private/Scheduler/Scheduler.cpp
+++ b/Engine/Plugins/Experimental/AnimNext/Source/AnimNext/Private/Scheduler/Scheduler.cpp
@@ -9,7 +9,6 @@
 #include "Tasks/Task.h"
 #include "UObject/GCObject.h"
 #include "LODPose.h"
-#include "Param/Params.h"
 #include "Engine/World.h"
 #include "Components/SkeletalMeshComponent.h"
 #include "UObject/ObjectKey.h"
diff --git a/Engine/Plugins/Experimental/AnimNext/Source/AnimNext/Private/Scheduler/Scheduler.h b/Engine/Plugins/Experimental/AnimNext/Source/AnimNext/Private/Scheduler/Scheduler.h
--- a/Engine/Plugins/Experimental/AnimNext/Source/AnimNext/Private/Scheduler/Scheduler.h
+++ b/Engine/Plugins/Experimental/AnimNext/Source/AnimNext/Private/Scheduler/Scheduler.h
@@ -10,7 +10,6 @@ enum class EAnimNextScheduleInitMethod : uint8;
 class UAnimNextGraph;
 class UAnimNextParameterBlock;
 class UAnimNextSchedule;
-struct FAnimNextParameterCollection;
 
 namespace UE::AnimNext
 {
diff --git a/Engine/Plugins/Experimental/AnimNext/Source/AnimNext/Public/AnimNextStats.h b/Engine/Plugins/Experimental/AnimNext/Source/AnimNext/Public/AnimNextStats.h
--- a/Engine/Plugins/Experimental/AnimNext/Source/AnimNext/Public/AnimNextStats.h
+++ b/Engine/Plugins/Experimental/AnimNext/Source/AnimNext/Public/AnimNextStats.h
@@ -10,6 +10,7 @@ DECLARE_CYCLE_STAT_EXTERN(TEXT("AnimNext: Initialize Entry"), STAT_AnimNext_Init
 DECLARE_CYCLE_STAT_EXTERN(TEXT("AnimNext: Create Instance Data"), STAT_AnimNext_CreateInstanceData, STATGROUP_AnimNext, ANIMNEXT_API);
 
 DECLARE_CYCLE_STAT_EXTERN(TEXT("AnimNext Task: Graph"), STAT_AnimNext_Task_Graph, STATGROUP_AnimNext, ANIMNEXT_API);
+DECLARE_CYCLE_STAT_EXTERN(TEXT("AnimNext Task: External Params"), STAT_AnimNext_Task_ExternalParams, STATGROUP_AnimNext, ANIMNEXT_API);
 DECLARE_CYCLE_STAT_EXTERN(TEXT("AnimNext Task: Port"), STAT_AnimNext_Task_Port, STATGROUP_AnimNext, ANIMNEXT_API);
 DECLARE_CYCLE_STAT_EXTERN(TEXT("AnimNext Task: Scope Entry"), STAT_AnimNext_Task_ScopeEntry, STATGROUP_AnimNext, ANIMNEXT_API);
 DECLARE_CYCLE_STAT_EXTERN(TEXT("AnimNext Task: Scope Exit"), STAT_AnimNext_Task_ScopeExit, STATGROUP_AnimNext, ANIMNEXT_API);
@@ -25,7 +26,6 @@ DECLARE_CYCLE_STAT_EXTERN(TEXT("AnimNext: Blend Overwrite (SoA)"), STAT_AnimNext
 DECLARE_CYCLE_STAT_EXTERN(TEXT("AnimNext: Blend Accumulate (SoA)"), STAT_AnimNext_BlendAccumulate_SoA, STATGROUP_AnimNext, ANIMNEXT_API);
 
 DECLARE_CYCLE_STAT_EXTERN(TEXT("AnimNext: Param Type Handle Lock"), STAT_AnimNext_ParamTypeHandle_Lock, STATGROUP_AnimNext, ANIMNEXT_API);
-DECLARE_CYCLE_STAT_EXTERN(TEXT("AnimNext: Param ID Lock"), STAT_AnimNext_ParamIdLock, STATGROUP_AnimNext, ANIMNEXT_API);
 DECLARE_CYCLE_STAT_EXTERN(TEXT("AnimNext: Param Stack Get Param"), STAT_AnimNext_ParamStack_GetParam, STATGROUP_AnimNext, ANIMNEXT_API);
 DECLARE_CYCLE_STAT_EXTERN(TEXT("AnimNext: Param Stack Adapter"), STAT_AnimNext_ParamStack_Adapter, STATGROUP_AnimNext, ANIMNEXT_API);
 DECLARE_CYCLE_STAT_EXTERN(TEXT("AnimNext: Param Stack Coalesce"), STAT_AnimNext_ParamStack_Coalesce, STATGROUP_AnimNext, ANIMNEXT_API);
diff --git a/Engine/Plugins/Experimental/AnimNext/Source/AnimNext/Public/Param/AnimNextClassExtensionLibrary.h b/Engine/Plugins/Experimental/AnimNext/Source/AnimNext/Public/Param/AnimNextClassExtensionLibrary.h
new file mode 100644
--- /dev/null
+++ b/Engine/Plugins/Experimental/AnimNext/Source/AnimNext/Public/Param/AnimNextClassExtensionLibrary.h
@@ -0,0 +1,20 @@
+﻿// Copyright Epic Games, Inc. All Rights Reserved.
+
+#pragma once
+
+#include "CoreMinimal.h"
+#include "AnimNextClassExtensionLibrary.generated.h"
+
+// Base class for function library-like classes that act to adapt UObjects of a specified type, exposing them to AnimNext parameters.
+// To do this, implement GetSupportedClass, then provide static UFUNCTION members that take the type as an input, e.g.
+// static float GetDeltaSeconds(UWorld* InWorld); 
+UCLASS(Abstract)
+class UAnimNextClassExtensionLibrary : public UObject
+{
+	GENERATED_BODY()
+
+public:
+	// Get the class that is supported by this extension library
+	virtual UClass* GetSupportedClass() const PURE_VIRTUAL(UAnimNextClassExtensionLibrary::GetSupportedClass, return nullptr; )
+};
+
diff --git a/Engine/Plugins/Experimental/AnimNext/Source/AnimNext/Public/Param/IAnimNextParameterSourceInterface.h b/Engine/Plugins/Experimental/AnimNext/Source/AnimNext/Public/Param/IAnimNextParameterSourceInterface.h
deleted file mode 100644
diff --git a/Engine/Plugins/Experimental/AnimNext/Source/AnimNext/Public/Param/IParameterSource.h b/Engine/Plugins/Experimental/AnimNext/Source/AnimNext/Public/Param/IParameterSource.h
new file mode 100644
--- /dev/null
+++ b/Engine/Plugins/Experimental/AnimNext/Source/AnimNext/Public/Param/IParameterSource.h
@@ -0,0 +1,32 @@
+// Copyright Epic Games, Inc. All Rights Reserved.
+
+#pragma once
+
+#include "CoreMinimal.h"
+#include "Param/ParamStackLayerHandle.h"
+
+namespace UE::AnimNext
+{
+	struct FParamContext;
+}
+
+namespace UE::AnimNext
+{
+
+/** An interface used to abstract AnimNext parameter sources */
+class IParameterSource
+{
+public:
+	virtual ~IParameterSource() = default;
+	
+	/** Update this parameter source */
+	virtual void Update() = 0;
+
+	/** Get a layer handle to allow this source to be used in the parameter stack */
+	virtual const FParamStackLayerHandle& GetLayerHandle() const = 0; 
+
+	/** Allow GC subscription - this should be called by the owner of the parameter source */
+	virtual void AddReferencedObjects(FReferenceCollector& Collector) {}
+};
+
+}
\ No newline at end of file
diff --git a/Engine/Plugins/Experimental/AnimNext/Source/AnimNext/Public/Param/IParameterSourceFactory.h b/Engine/Plugins/Experimental/AnimNext/Source/AnimNext/Public/Param/IParameterSourceFactory.h
new file mode 100644
--- /dev/null
+++ b/Engine/Plugins/Experimental/AnimNext/Source/AnimNext/Public/Param/IParameterSourceFactory.h
@@ -0,0 +1,66 @@
+﻿// Copyright Epic Games, Inc. All Rights Reserved.
+
+#pragma once
+
+#include "Param/IParameterSource.h"
+#include "Features/IModularFeature.h"
+
+namespace UE::AnimNext
+{
+	class IParameterSource;
+}
+
+namespace UE::AnimNext
+{
+
+// Context passed to object accessor functions registered to RegisterObjectAccessor
+struct FExternalParameterContext
+{
+	// The object that the AnimNext entry is bound to (e.g. a UAnimNextComponent)
+	UObject* Object = nullptr;
+};
+
+// Modular feature allowing other modules to extend and add to external parameter system
+class IParameterSourceFactory : public IModularFeature
+{
+public:
+	inline static const FName FeatureName = TEXT("AnimNextParameterSourceFactory");
+
+	virtual ~IParameterSourceFactory() = default;
+
+	// Iterate over all the source names that this factory provides
+	virtual void ForEachSource(TFunctionRef<void(FName)> InFunction) const = 0;
+
+	// Factory method used to create a parameter source of the specified name, with a set of parameters that are initially required
+	// @param   InContext              Context used to set up the parameter source
+	// @param   InSourceName           The name of the parameter source. This is conventionally the 'root' parameter of a set of parameters (e.g.
+	//                                 Game_MyComponent, which contains Game_MyComponent_Param etc.)
+	// @param   InRequiredParameters   Any required parameters that the source should initially supply, can be empty
+	// @return a new parameter source, or nullptr if the source was not found
+	virtual TUniquePtr<IParameterSource> CreateParameterSource(const FExternalParameterContext& InContext, FName InSourceName, TConstArrayView<FName> InRequiredParameters) const = 0;
+
+#if WITH_EDITOR
+	// Info about parameters gleaned from FindParameterInfo and ForEachParameter
+	struct FParameterInfo
+	{
+		FAnimNextParamType Type;
+		FText Tooltip;
+		bool bThreadSafe = false;
+	};
+
+	// Given a parameter name, find associated info for that parameter
+	// @param   InParameterName        The parameter name to find
+	// @param   OutType                The parameter type
+	// @param   OutTooltip             Optional: The parameter's tooltip
+	// @return true if the parameter info was found
+	virtual bool FindParameterInfo(FName InParameterName, FParameterInfo& OutInfo) const = 0;
+
+	// Iterate over all the known parameters for a particular source
+	// @param   InSourceName           The name of the parameter source to iterate
+	// @param   InFunction             Function called for each known parameter
+	// @return an array of parameter names
+	virtual void ForEachParameter(FName InSourceName, TFunctionRef<void(FName, const FParameterInfo&)> InFunction) const = 0;
+#endif
+};
+
+}
\ No newline at end of file
diff --git a/Engine/Plugins/Experimental/AnimNext/Source/AnimNext/Public/Param/ParamId.h b/Engine/Plugins/Experimental/AnimNext/Source/AnimNext/Public/Param/ParamId.h
--- a/Engine/Plugins/Experimental/AnimNext/Source/AnimNext/Public/Param/ParamId.h
+++ b/Engine/Plugins/Experimental/AnimNext/Source/AnimNext/Public/Param/ParamId.h
@@ -11,7 +11,6 @@ struct FRigVMDispatch_GetParameter;
 
 namespace UE::AnimNext
 {
-	struct FParamAdapter;
 	struct FParamDefinition;
 	class FModule;
 	struct FRemappedLayer;
@@ -33,7 +32,6 @@ struct ANIMNEXT_API FParamId
 	friend struct ::FRigVMDispatch_SetLayerParameter;
 	friend struct ::FRigVMDispatch_GetLayerParameter;
 	friend struct ::FRigVMDispatch_GetParameter;
-	friend struct FParamAdapter;
 	friend class ::UAnimNextConfig;
 	friend struct FParamDefinition;
 	friend class FModule;
@@ -44,10 +42,19 @@ struct ANIMNEXT_API FParamId
 	FParamId& operator=(const FParamId& InName) = default;
 
 	// Make a parameter ID from an FName, generating the hash
-	explicit FParamId(FName InName);
+	explicit FParamId(FName InName)
+		: Name(InName)
+		, Hash(GetTypeHash(InName))
+	{
+	}
 
 	// Make a parameter ID from a name and hash
-	explicit FParamId(FName InName, uint32 InHash);
+	explicit FParamId(FName InName, uint32 InHash)
+		: Name(InName)
+		, Hash(InHash)
+	{
+		checkSlow(Hash == GetTypeHash(Name));
+	}
 
 	// Get the name of this param
 	FName GetName() const
@@ -67,38 +74,6 @@ struct ANIMNEXT_API FParamId
 		return Hash != 0;
 	}
 
-private:
-	// Initialize the parameter ID system
-	static void Init();
-
-	// Shut down the parameter ID system
-	static void Destroy();
-
-	// Refresh all adapters, built in and per-config
-	static void RefreshAdapters();
-
-	// Resets/clears all adapters
-	static void ResetAdapters();
-
-	// Register all built-in adapters
-	static void RegisterBuiltInAdapters();
-
-	// Refresh all adapters defined in config
-	static void RefreshConfigAdapters();
-
-	// Get any built in adapter for the supplied parameter
-	static const FParamAdapter* GetAdapter(FParamId InId);
-
-	// Get the thread-local scratch area for an adapted parameter id
-	static uint8* GetScratchAreaForParamIdAdapter(FParamId InId);
-
-	// Get the thread-local return value for an adapted parameter id
-	template<typename ReturnType>
-	static ReturnType* GetAdapterReturnValue(FParamId InId)
-	{
-		return reinterpret_cast<ReturnType*>(GetScratchAreaForParamIdAdapter(InId));
-	}
-
 private:
 	// Parameter name
 	FName Name;
diff --git a/Engine/Plugins/Experimental/AnimNext/Source/AnimNext/Public/Param/ParamStack.h b/Engine/Plugins/Experimental/AnimNext/Source/AnimNext/Public/Param/ParamStack.h
--- a/Engine/Plugins/Experimental/AnimNext/Source/AnimNext/Public/Param/ParamStack.h
+++ b/Engine/Plugins/Experimental/AnimNext/Source/AnimNext/Public/Param/ParamStack.h
@@ -194,16 +194,10 @@ public:
 	// Push a layer
 	ANIMNEXT_API FPushedLayerHandle PushLayer(const FParamStackLayerHandle& InLayerHandle);
 
-	// Create a cached parameter layer from a class. This layer will own an object used to represent the layer.
-	static ANIMNEXT_API FParamStackLayerHandle MakeValueLayer(const UClass* InClass);
-
-	// Create a cached parameter layer from an object. This layer will weakly reference the supplied object.
-	static ANIMNEXT_API FParamStackLayerHandle MakeReferenceLayer(UObject* InObject);
-
 	// Create a cached parameter layer from an instanced property bag. This layer will own a copy of the property bag used to represent the layer.
 	static ANIMNEXT_API FParamStackLayerHandle MakeValueLayer(const FInstancedPropertyBag& InInstancedPropertyBag);
 
-	// Create a cached parameter layer from an instanced property bag. This layer will reference the suppled propery bag and does not transfer ownership.
+	// Create a cached parameter layer from an instanced property bag. This layer will reference the supplied property bag and does not transfer ownership.
 	static ANIMNEXT_API FParamStackLayerHandle MakeReferenceLayer(FInstancedPropertyBag& InInstancedPropertyBag);
 
 	// Create a cached parameter layer by remapping the entries from another layer
diff --git a/Engine/Plugins/Experimental/AnimNext/Source/AnimNext/Public/Scheduler/AnimNextSchedulePort.h b/Engine/Plugins/Experimental/AnimNext/Source/AnimNext/Public/Scheduler/AnimNextSchedulePort.h
--- a/Engine/Plugins/Experimental/AnimNext/Source/AnimNext/Public/Scheduler/AnimNextSchedulePort.h
+++ b/Engine/Plugins/Experimental/AnimNext/Source/AnimNext/Public/Scheduler/AnimNextSchedulePort.h
@@ -6,6 +6,7 @@
 #include "IAnimNextScheduleTermInterface.h"
 #include "AnimNextSchedulePort.generated.h"
 
+class UAnimNextSchedule;
 struct FAnimNextSchedulePortTask;
 
 namespace UE::AnimNext
@@ -19,8 +20,12 @@ class ANIMNEXT_API UAnimNextSchedulePort : public UObject, public IAnimNextSched
 	GENERATED_BODY()
 
 private:
+	friend struct UE::AnimNext::UncookedOnly::FUtils;
 	friend struct FAnimNextSchedulePortTask;
 
 	// Run the port's logic
 	virtual void Run(const UE::AnimNext::FScheduleTermContext& InContext) const PURE_VIRTUAL(UAnimNextSchedulePort::Run, )
+
+	// Get any required parameters for this port
+	virtual TConstArrayView<FName> GetRequiredParameters() const PURE_VIRTUAL(UAnimNextSchedulePort::GetRequiredParameters, return TConstArrayView<FName>(); )
 };
diff --git a/Engine/Plugins/Experimental/AnimNext/Source/AnimNext/Public/Scheduler/IAnimNextScheduleTermInterface.h b/Engine/Plugins/Experimental/AnimNext/Source/AnimNext/Public/Scheduler/IAnimNextScheduleTermInterface.h
--- a/Engine/Plugins/Experimental/AnimNext/Source/AnimNext/Public/Scheduler/IAnimNextScheduleTermInterface.h
+++ b/Engine/Plugins/Experimental/AnimNext/Source/AnimNext/Public/Scheduler/IAnimNextScheduleTermInterface.h
@@ -14,6 +14,11 @@ namespace UE::AnimNext
 	struct FScheduleInstanceData;
 }
 
+namespace UE::AnimNext::UncookedOnly
+{
+	struct FUtils;
+}
+
 UENUM()
 enum class EScheduleTermDirection : uint8
 {
@@ -52,7 +57,7 @@ class IAnimNextScheduleTermInterface
 	GENERATED_BODY()
 
 private:
-	friend class UAnimNextSchedule;
+	friend struct UE::AnimNext::UncookedOnly::FUtils;
 	friend struct UE::AnimNext::FScheduleInstanceData;
 	friend struct FAnimNextSchedulePortTask;
 
diff --git a/Engine/Plugins/Experimental/AnimNext/Source/AnimNextEditor/Internal/Param/ParameterPickerArgs.h b/Engine/Plugins/Experimental/AnimNext/Source/AnimNextEditor/Internal/Param/ParameterPickerArgs.h
--- a/Engine/Plugins/Experimental/AnimNext/Source/AnimNextEditor/Internal/Param/ParameterPickerArgs.h
+++ b/Engine/Plugins/Experimental/AnimNext/Source/AnimNextEditor/Internal/Param/ParameterPickerArgs.h
@@ -119,6 +119,9 @@ struct FParameterPickerArgs
 
 	// Whether the search box should be focussed on widget creation
 	bool bFocusSearchWidget = true;
+
+	// Whether 'none' can be selected
+	bool bAllowNone = false;
 };
 
 }
diff --git a/Engine/Plugins/Experimental/AnimNext/Source/AnimNextEditor/Private/EditorUtils.cpp b/Engine/Plugins/Experimental/AnimNext/Source/AnimNextEditor/Private/EditorUtils.cpp
--- a/Engine/Plugins/Experimental/AnimNext/Source/AnimNextEditor/Private/EditorUtils.cpp
+++ b/Engine/Plugins/Experimental/AnimNext/Source/AnimNextEditor/Private/EditorUtils.cpp
@@ -12,7 +12,6 @@
 #include "AssetRegistry/IAssetRegistry.h"
 #include "Param/IAnimNextParameterBlockGraphInterface.h"
 #include "Param/IAnimNextParameterBlockParameterInterface.h"
-#include "Param/Params.h"
 #include "Workspace/AnimNextWorkspace.h"
 #include "UncookedOnlyUtils.h"
 #include "Param/RigVMDispatch_GetLayerParameter.h"
diff --git a/Engine/Plugins/Experimental/AnimNext/Source/AnimNextEditor/Private/Param/ParamNamePropertyCustomization.cpp b/Engine/Plugins/Experimental/AnimNext/Source/AnimNextEditor/Private/Param/ParamNamePropertyCustomization.cpp
--- a/Engine/Plugins/Experimental/AnimNext/Source/AnimNextEditor/Private/Param/ParamNamePropertyCustomization.cpp
+++ b/Engine/Plugins/Experimental/AnimNext/Source/AnimNextEditor/Private/Param/ParamNamePropertyCustomization.cpp
@@ -49,6 +49,7 @@ void FParamNamePropertyTypeCustomization::CustomizeHeader(TSharedRef<IPropertyHa
 	
 	const FProperty* Property = GetMetadataProperty(InPropertyHandle->GetProperty());
 	const FString ParamTypeString = Property->GetMetaData("AllowedParamType");
+	const bool bAllowNone = Property->HasMetaData("AllowNone");
 	FAnimNextParamType FilterType = FAnimNextParamType::FromString(ParamTypeString);
 	
 	FParameterPickerArgs PickerArgs;
@@ -73,6 +74,7 @@ void FParamNamePropertyTypeCustomization::CustomizeHeader(TSharedRef<IPropertyHa
 			Refresh();
 		}
 	});
+	PickerArgs.bAllowNone = bAllowNone;
 	
 	InHeaderRow
 	.NameContent()
diff --git a/Engine/Plugins/Experimental/AnimNext/Source/AnimNextEditor/Private/Param/ParamTypePropertyCustomization.cpp b/Engine/Plugins/Experimental/AnimNext/Source/AnimNextEditor/Private/Param/ParamTypePropertyCustomization.cpp
--- a/Engine/Plugins/Experimental/AnimNext/Source/AnimNextEditor/Private/Param/ParamTypePropertyCustomization.cpp
+++ b/Engine/Plugins/Experimental/AnimNext/Source/AnimNextEditor/Private/Param/ParamTypePropertyCustomization.cpp
@@ -26,9 +26,8 @@ void FParamTypePropertyTypeCustomization::CustomizeHeader(TSharedRef<IPropertyHa
 			TArray<const void*> RawData;
 			PropertyHandle->AccessRawData(RawData);
 
-			if(RawData.Num() > 0)
+			if(RawData.Num() > 0 && RawData[0] != nullptr)
 			{
-				check(RawData[0] != nullptr);
 				const FAnimNextParamType& FirstType = *static_cast<const FAnimNextParamType*>(RawData[0]);
 				bool bMultipleValues = false;
 
diff --git a/Engine/Plugins/Experimental/AnimNext/Source/AnimNextEditor/Private/Param/SParameterPicker.cpp b/Engine/Plugins/Experimental/AnimNext/Source/AnimNextEditor/Private/Param/SParameterPicker.cpp
--- a/Engine/Plugins/Experimental/AnimNext/Source/AnimNextEditor/Private/Param/SParameterPicker.cpp
+++ b/Engine/Plugins/Experimental/AnimNext/Source/AnimNextEditor/Private/Param/SParameterPicker.cpp
@@ -16,11 +16,11 @@
 #include "String/ParseTokens.h"
 #include "Widgets/Layout/SBox.h"
 #include "Widgets/Images/SImage.h"
-#include "Param/Params.h"
 #include "Widgets/Input/SButton.h"
 #include "Framework/Application/SlateApplication.h"
 #include "ScopedTransaction.h"
 #include "SSimpleButton.h"
+#include "Param/ExternalParameterRegistry.h"
 
 #define LOCTEXT_NAMESPACE "SParameterPicker"
 
@@ -302,6 +302,11 @@ void SParameterPicker::RefreshEntries()
 
 	TSet<TTuple<FName, FAssetData>> BoundParameters;
 
+	if(Args.bAllowNone)
+	{
+		Entries.Add(MakeShared<FParameterPickerEntry>(FParameterBindingReference(NAME_None), FAnimNextParamType()));
+	}
+	
 	// Find all blocks and their bound parameters
 	if(Args.bShowBoundParameters)
 	{
@@ -365,14 +370,14 @@ void SParameterPicker::RefreshEntries()
 
 	if (Args.bShowBuiltInParameters)
 	{
-		FParams::ForEachBuiltInParameter([this](const FParamDefinition& InDefinition)
+		FExternalParameterRegistry::ForEachParameter([this](FName InParameterName, const IParameterSourceFactory::FParameterInfo& InInfo)
 		{
-			FParameterBindingReference NewReference(InDefinition.GetName());
+			FParameterBindingReference NewReference(InParameterName);
 			if (!Args.OnFilterParameter.IsBound() || Args.OnFilterParameter.Execute(NewReference) == EFilterParameterResult::Include)
 			{
-				if (!Args.OnFilterParameterType.IsBound() || Args.OnFilterParameterType.Execute(InDefinition.GetType()) == EFilterParameterResult::Include)
+				if (!Args.OnFilterParameterType.IsBound() || Args.OnFilterParameterType.Execute(InInfo.Type) == EFilterParameterResult::Include)
 				{
-					Entries.Add(MakeShared<FParameterPickerEntry>(NewReference, InDefinition.GetType()));
+					Entries.Add(MakeShared<FParameterPickerEntry>(NewReference, InInfo.Type));
 				}
 			}
 		});
@@ -642,7 +647,7 @@ void SParameterPicker::HandleSelectionChanged(TSharedPtr<FParameterPickerEntry>
 		TArray<TSharedRef<FParameterPickerEntry>> SelectedEntries;
 		EntriesList->GetSelectedItems(SelectedEntries);
 
-		if(SelectedEntries[0]->ParamType.IsValid())
+		if(SelectedEntries[0]->ParamType.IsValid() || Args.bAllowNone)
 		{
 			Args.OnParameterPicked.ExecuteIfBound(SelectedEntries[0]->Binding);
 		}
diff --git a/Engine/Plugins/Experimental/AnimNext/Source/AnimNextEditor/Private/Param/SParameterPickerCombo.cpp b/Engine/Plugins/Experimental/AnimNext/Source/AnimNextEditor/Private/Param/SParameterPickerCombo.cpp
--- a/Engine/Plugins/Experimental/AnimNext/Source/AnimNextEditor/Private/Param/SParameterPickerCombo.cpp
+++ b/Engine/Plugins/Experimental/AnimNext/Source/AnimNextEditor/Private/Param/SParameterPickerCombo.cpp
@@ -1,4 +1,4 @@
-﻿// Copyright Epic Games, Inc. All Rights Reserved.
+﻿	// Copyright Epic Games, Inc. All Rights Reserved.
 
 #include "SParameterPickerCombo.h"
 #include "SParameterPicker.h"
@@ -34,8 +34,14 @@ void SParameterPickerCombo::Construct(const FArguments& InArgs)
 		SNew(SComboButton)
 		.OnGetMenuContent_Lambda([PickerArgs, this]()
 		{
-			return SNew(SParameterPicker)
-				.Args(PickerArgs);
+			return
+				SNew(SBox)
+				.WidthOverride(300.0f)
+				.HeightOverride(300.0f)
+				[
+					SNew(SParameterPicker)
+					.Args(PickerArgs)
+				];
 		})
 		.ButtonContent()
 		[
diff --git a/Engine/Plugins/Experimental/AnimNext/Source/AnimNextUncookedOnly/Internal/UncookedOnlyUtils.h b/Engine/Plugins/Experimental/AnimNext/Source/AnimNextUncookedOnly/Internal/UncookedOnlyUtils.h
--- a/Engine/Plugins/Experimental/AnimNext/Source/AnimNextUncookedOnly/Internal/UncookedOnlyUtils.h
+++ b/Engine/Plugins/Experimental/AnimNext/Source/AnimNextUncookedOnly/Internal/UncookedOnlyUtils.h
@@ -10,6 +10,7 @@
 
 #include "UncookedOnlyUtils.generated.h"
 
+class UAnimNextSchedule;
 class UAnimNextGraph;
 class UAnimNextGraph_EditorData;
 class UAnimNextGraph_EdGraph;
@@ -180,6 +181,9 @@ struct ANIMNEXTUNCOOKEDONLY_API FUtils
 	// If the name cannot be found, the returned type will be invalid
 	// Note that this is expensive and can query the asset registry
 	static FAnimNextParamType GetParameterTypeFromName(FName InName);
+
+	// Compiles a schedule
+	static void CompileSchedule(UAnimNextSchedule* InSchedule);
 };
 
 }
\ No newline at end of file
diff --git a/Engine/Plugins/Experimental/AnimNext/Source/AnimNextUncookedOnly/Private/AnimNextUncookedOnlyModule.cpp b/Engine/Plugins/Experimental/AnimNext/Source/AnimNextUncookedOnly/Private/AnimNextUncookedOnlyModule.cpp
--- a/Engine/Plugins/Experimental/AnimNext/Source/AnimNextUncookedOnly/Private/AnimNextUncookedOnlyModule.cpp
+++ b/Engine/Plugins/Experimental/AnimNext/Source/AnimNextUncookedOnly/Private/AnimNextUncookedOnlyModule.cpp
@@ -1,21 +1,29 @@
 // Copyright Epic Games, Inc. All Rights Reserved.
 
 #include "CoreMinimal.h"
+#include "UncookedOnlyUtils.h"
 #include "Modules/ModuleManager.h"
+#include "Scheduler/AnimNextSchedule.h"
 
 namespace UE::AnimNext::UncookedOnly
 {
 	class FModule : public IModuleInterface
 	{
-	public:
-
 	private:
-		virtual void StartupModule() override;
-	};
+		virtual void StartupModule() override
+		{
+			// TEMP: Bind the compilation function for schedules
+			UAnimNextSchedule::CompileFunction = [](UAnimNextSchedule* InSchedule)
+			{
+				FUtils::CompileSchedule(InSchedule);
+			};
+		}
 
-	void FModule::StartupModule()
-	{
-	}
+		virtual void ShutdownModule() override
+		{
+			UAnimNextSchedule::CompileFunction = nullptr;
+		}
+	};
 }
 
 IMPLEMENT_MODULE(UE::AnimNext::UncookedOnly::FModule, AnimNextUncookedOnly);
diff --git a/Engine/Plugins/Experimental/AnimNext/Source/AnimNextUncookedOnly/Private/Param/AnimNextParameterBlockParameter.cpp b/Engine/Plugins/Experimental/AnimNext/Source/AnimNextUncookedOnly/Private/Param/AnimNextParameterBlockParameter.cpp
--- a/Engine/Plugins/Experimental/AnimNext/Source/AnimNextUncookedOnly/Private/Param/AnimNextParameterBlockParameter.cpp
+++ b/Engine/Plugins/Experimental/AnimNext/Source/AnimNextUncookedOnly/Private/Param/AnimNextParameterBlockParameter.cpp
@@ -2,16 +2,17 @@
 
 #include "Param/AnimNextParameterBlockParameter.h"
 #include "Param/AnimNextParameterBlock_EditorData.h"
-#include "Param/Params.h"
+#include "Param/ExternalParameterRegistry.h"
 
 FAnimNextParamType UAnimNextParameterBlockParameter::GetParamType() const
 {
 	using namespace UE::AnimNext;
 
 	// Look in built-in parameters
-	if(const FParamDefinition* ParamDefinition = FParams::FindBuiltInParameter(ParameterName))
+	IParameterSourceFactory::FParameterInfo Info;
+	if(FExternalParameterRegistry::FindParameterInfo(ParameterName, Info))
 	{
-		return ParamDefinition->GetType();
+		return Info.Type;
 	}
 
 	return Type;
@@ -54,5 +55,13 @@ FText UAnimNextParameterBlockParameter::GetDisplayName() const
 
 FText UAnimNextParameterBlockParameter::GetDisplayNameTooltip() const
 {
+	using namespace UE::AnimNext;
+
+	IParameterSourceFactory::FParameterInfo Info;
+	if(FExternalParameterRegistry::FindParameterInfo(ParameterName, Info))
+	{
+		return Info.Tooltip;
+	}
+	
 	return FText::FromName(ParameterName);
 }
\ No newline at end of file
diff --git a/Engine/Plugins/Experimental/AnimNext/Source/AnimNextUncookedOnly/Private/UncookedOnlyUtils.cpp b/Engine/Plugins/Experimental/AnimNext/Source/AnimNextUncookedOnly/Private/UncookedOnlyUtils.cpp
--- a/Engine/Plugins/Experimental/AnimNext/Source/AnimNextUncookedOnly/Private/UncookedOnlyUtils.cpp
+++ b/Engine/Plugins/Experimental/AnimNext/Source/AnimNextUncookedOnly/Private/UncookedOnlyUtils.cpp
@@ -31,13 +31,14 @@
 #include "DecoratorBase/DecoratorRegistry.h"
 #include "DecoratorBase/Decorator.h"
 #include "Graph/RigUnit_AnimNextBeginExecution.h"
-#include "Param/Params.h"
 #include "Serialization/MemoryReader.h"
 #include "RigVMRuntimeDataRegistry.h"
 #include "Param/RigVMDispatch_GetLayerParameter.h"
-
 #include "RigVMCompiler/RigVMCompiler.h"
 #include "RigVMCore/RigVM.h"
+#include "Param/ExternalParameterRegistry.h"
+#include "Scheduler/AnimNextSchedule.h"
+#include "Scheduler/AnimNextSchedulePort.h"
 
 namespace UE::AnimNext::UncookedOnly
 {
@@ -1148,9 +1149,10 @@ FText FUtils::GetParameterDisplayNameText(FName InParameterName)
 FAnimNextParamType FUtils::GetParameterTypeFromName(FName InName)
 {
 	// Check built-in params first as they are cheaper
-	if(const FParamDefinition* FoundDefinition = FParams::FindBuiltInParameter(InName))
+	IParameterSourceFactory::FParameterInfo Info;
+	if(FExternalParameterRegistry::FindParameterInfo(InName, Info))
 	{
-		return FoundDefinition->Type;
+		return Info.Type;
 	}
 
 	// Query the asset registry for other params
@@ -1255,4 +1257,563 @@ void FUtils::GetGraphParameters(const URigVMGraph* Graph, FAnimNextParameterProv
 		}
 	}
 }
+
+void FUtils::CompileSchedule(UAnimNextSchedule* InSchedule)
+{
+	using namespace UE::AnimNext;
+
+	InSchedule->Instructions.Empty();
+	InSchedule->GraphTasks.Empty();
+	InSchedule->Ports.Empty();
+	InSchedule->ExternalTasks.Empty();
+	InSchedule->ParamScopeEntryTasks.Empty();
+	InSchedule->ParamScopeExitTasks.Empty();
+	InSchedule->ExternalParamTasks.Empty();
+	InSchedule->IntermediatesData.Reset();
+	InSchedule->NumParameterScopes = 0;
+	InSchedule->NumTickFunctions = 0;
+
+	EAnimNextScheduleScheduleOpcode LastOpCode = EAnimNextScheduleScheduleOpcode::None;
+
+	auto Emit = [InSchedule, &LastOpCode](EAnimNextScheduleScheduleOpcode InOpCode, int32 InOperand = 0)
+	{
+		FAnimNextScheduleInstruction Instruction;
+		Instruction.Opcode = InOpCode;
+		Instruction.Operand = InOperand;
+		InSchedule->Instructions.Add(Instruction);
+
+		LastOpCode = InOpCode;
+	};
+
+	auto EmitPrerequisite = [InSchedule, &Emit, &LastOpCode]()
+	{
+		switch (LastOpCode)
+		{
+		case EAnimNextScheduleScheduleOpcode::RunGraphTask:
+			Emit(EAnimNextScheduleScheduleOpcode::PrerequisiteTask, InSchedule->NumTickFunctions - 1);
+			break;
+		case EAnimNextScheduleScheduleOpcode::BeginRunExternalTask:
+			Emit(EAnimNextScheduleScheduleOpcode::PrerequisiteBeginExternalTask, InSchedule->NumTickFunctions - 1);
+			break;
+		case EAnimNextScheduleScheduleOpcode::EndRunExternalTask:
+			Emit(EAnimNextScheduleScheduleOpcode::PrerequisiteEndExternalTask, InSchedule->NumTickFunctions - 1);
+			break;
+		case EAnimNextScheduleScheduleOpcode::RunParamScopeEntry:
+			Emit(EAnimNextScheduleScheduleOpcode::PrerequisiteScopeEntry, InSchedule->NumTickFunctions - 1);
+			break;
+		case EAnimNextScheduleScheduleOpcode::RunParamScopeExit:
+			Emit(EAnimNextScheduleScheduleOpcode::PrerequisiteScopeExit, InSchedule->NumTickFunctions - 1);
+			break;
+		case EAnimNextScheduleScheduleOpcode::RunExternalParamTask:
+			Emit(EAnimNextScheduleScheduleOpcode::PrerequisiteExternalParamTask, InSchedule->NumTickFunctions - 1);
+			break;
+		case EAnimNextScheduleScheduleOpcode::None:
+			break;
+		default:
+			checkNoEntry();
+			break;
+		}
+	};
+
+	// MAX_uint32 means 'global scope' in this context
+	uint32 ParentScopeIndex = MAX_uint32;
+
+	TArray<FAnimNextScheduleEntryTerm> IntermediateTerms;
+	TMap<FName, uint32> IntermediateMap;
+
+	TFunction<void(TArrayView<TObjectPtr<UAnimNextScheduleEntry>>)> EmitEntries;
+
+	// Iterate over all entries, recursing into scopes
+	EmitEntries = [InSchedule, &EmitEntries, &Emit, &EmitPrerequisite, &ParentScopeIndex, &IntermediateTerms, &IntermediateMap](TArrayView<TObjectPtr<UAnimNextScheduleEntry>> InEntries)
+	{
+		for (int32 EntryIndex = 0; EntryIndex < InEntries.Num(); ++EntryIndex)
+		{
+			UAnimNextScheduleEntry* Entry = InEntries[EntryIndex];
+
+			auto CheckTermDirectionCompatibility = [](FName InName, EScheduleTermDirection InExistingDirection, EScheduleTermDirection InNewDirection)
+			{
+				switch(InExistingDirection)
+				{
+				case EScheduleTermDirection::Input:
+					// Input before output: error
+					UE_LOG(LogAnimation, Error, TEXT("Term '%s' was used as an input before it was output"), *InName.ToString());
+					return false;
+				case EScheduleTermDirection::Output:
+					return true;
+				}
+
+				return false;
+			};
+
+			if (UAnimNextScheduleEntry_Port* PortEntry = Cast<UAnimNextScheduleEntry_Port>(Entry))
+			{
+				bool bValid = true;
+
+				if(PortEntry->Port == nullptr)
+				{
+					UE_LOG(LogAnimation, Error, TEXT("AnimNext: Invalid port class found"));
+					bValid = false;
+				}
+				else
+				{
+					UAnimNextSchedulePort* CDO = PortEntry->Port->GetDefaultObject<UAnimNextSchedulePort>();
+					check(CDO);
+
+					TConstArrayView<FScheduleTerm> Terms = CDO->GetTerms();
+					if(PortEntry->Terms.Num() != Terms.Num())
+					{
+						UE_LOG(LogAnimation, Error, TEXT("AnimNext: Incorrect term count for port: %d"), PortEntry->Terms.Num());
+						bValid = false;
+					}
+
+					for(int32 TermIndex = 0; TermIndex < PortEntry->Terms.Num(); ++TermIndex)
+					{
+						FName TermName = PortEntry->Terms[TermIndex].Name;
+						if(!PortEntry->Terms[TermIndex].Type.IsValid())
+						{
+							UE_LOG(LogAnimation, Error, TEXT("AnimNext: Invalid type when processing port term, ignored: '%s'"), *TermName.ToString());
+							bValid = false;
+						}
+						else
+						{
+							const uint32* ExistingIntermediateIndexPtr = IntermediateMap.Find(TermName);
+							if(ExistingIntermediateIndexPtr != nullptr)
+							{
+								const FAnimNextScheduleEntryTerm& IntermediateTerm = IntermediateTerms[*ExistingIntermediateIndexPtr];
+								if(IntermediateTerm.Type != Terms[TermIndex].GetType())
+								{
+									UE_LOG(LogAnimation, Error, TEXT("AnimNext: Mismatched types when processing port term, ignored: '%s'"), *TermName.ToString());
+									bValid = false;
+								}
+
+								if(!CheckTermDirectionCompatibility(TermName, IntermediateTerm.Direction, Terms[TermIndex].Direction))
+								{
+									bValid = false;
+								}
+							}
+						}
+					}
+				}
+
+				if(bValid)
+				{
+					EmitPrerequisite();
+
+					FAnimNextSchedulePortTask PortTask;
+					PortTask.TaskIndex = InSchedule->Ports.Num();
+					PortTask.ParamScopeIndex = ParentScopeIndex;
+					PortTask.Port = PortEntry->Port;
+
+					for(int32 TermIndex = 0; TermIndex < PortEntry->Terms.Num(); ++TermIndex)
+					{
+						FName TermName = PortEntry->Terms[TermIndex].Name;
+						const uint32* ExistingIntermediateIndexPtr = IntermediateMap.Find(TermName);
+						if(ExistingIntermediateIndexPtr == nullptr)
+						{
+							uint32 IntermediateIndex = IntermediateTerms.Emplace(TermName, PortEntry->Terms[TermIndex].Type, PortEntry->Terms[TermIndex].Direction);
+							IntermediateMap.Add(TermName, IntermediateIndex);
+							PortTask.Terms.Add(IntermediateIndex);
+						}
+						else
+						{
+							PortTask.Terms.Add(*ExistingIntermediateIndexPtr);
+						}
+					}
+
+					int32 PortIndex = InSchedule->Ports.Add(PortTask);
+
+					InSchedule->NumTickFunctions++;
+
+					Emit(EAnimNextScheduleScheduleOpcode::RunPort, PortIndex);
+				}
+			}
+			else if (UAnimNextScheduleEntry_AnimNextGraph* GraphEntry = Cast<UAnimNextScheduleEntry_AnimNextGraph>(Entry))
+			{
+				bool bValid = true;
+
+				if(GraphEntry->Graph == nullptr && GraphEntry->DynamicGraph == NAME_None)
+				{
+					UE_LOG(LogAnimation, Error, TEXT("AnimNext: Invalid graph or no parameter supplied"));
+					bValid = false;
+				}
+				else if(GraphEntry->Graph != nullptr)
+				{
+					TConstArrayView<FScheduleTerm> Terms = GraphEntry->Graph->GetTerms();
+					if(GraphEntry->Terms.Num() != Terms.Num())
+					{
+						UE_LOG(LogAnimation, Error, TEXT("AnimNext: Incorrect term count for graph: %d"), GraphEntry->Terms.Num());
+						bValid = false;
+					}
+					else
+					{
+						// Validate graph terms match schedule-expected terms
+						for(int32 TermIndex = 0; TermIndex < GraphEntry->Terms.Num(); ++TermIndex)
+						{
+							FName TermName = GraphEntry->Terms[TermIndex].Name;
+							if(Terms[TermIndex].Direction != GraphEntry->Terms[TermIndex].Direction)
+							{
+								UE_LOG(LogAnimation, Error, TEXT("AnimNext: Mismatched direction when processing graph term, ignored: '%s'"), *TermName.ToString());
+								bValid = false;
+							}
+							
+							if(Terms[TermIndex].GetType() != GraphEntry->Terms[TermIndex].Type)
+							{
+								UE_LOG(LogAnimation, Error, TEXT("AnimNext: Mismatched types when processing graph term, ignored: '%s'"), *TermName.ToString());
+								bValid = false;
+							}
+						}
+					}
+				}
+
+				// Validate terms and check against priors
+				for(int32 TermIndex = 0; TermIndex < GraphEntry->Terms.Num(); ++TermIndex)
+				{
+					FName TermName = GraphEntry->Terms[TermIndex].Name;
+					if(!GraphEntry->Terms[TermIndex].Type.IsValid())
+					{
+						UE_LOG(LogAnimation, Error, TEXT("AnimNext: Invalid type when processing graph term, ignored: '%s'"), *TermName.ToString());
+						bValid = false;
+					}
+					else
+					{
+						const uint32* ExistingIntermediateIndexPtr = IntermediateMap.Find(TermName);
+						if(ExistingIntermediateIndexPtr != nullptr)
+						{
+							const FAnimNextScheduleEntryTerm& IntermediateTerm = IntermediateTerms[*ExistingIntermediateIndexPtr];
+							if(IntermediateTerm.Type != GraphEntry->Terms[TermIndex].Type)
+							{
+								UE_LOG(LogAnimation, Error, TEXT("AnimNext: Mismatched types when processing graph term, ignored: '%s'"), *TermName.ToString());
+								bValid = false;
+							}
+
+							if(!CheckTermDirectionCompatibility(TermName, IntermediateTerm.Direction, GraphEntry->Terms[TermIndex].Direction))
+							{
+								bValid = false;
+							}
+						}
+					}
+				}
+
+				if(bValid)
+				{
+					EmitPrerequisite();
+
+					FAnimNextScheduleGraphTask GraphTask;
+					GraphTask.TaskIndex = InSchedule->GraphTasks.Num();
+					GraphTask.ParamScopeIndex = InSchedule->NumParameterScopes++;
+					GraphTask.ParamParentScopeIndex = ParentScopeIndex;
+					GraphTask.EntryPoint = GraphEntry->EntryPoint;
+					GraphTask.Graph = GraphEntry->Graph;
+					GraphTask.DynamicGraph = GraphEntry->DynamicGraph;
+
+					for(int32 TermIndex = 0; TermIndex < GraphEntry->Terms.Num(); ++TermIndex)
+					{
+						FName TermName = GraphEntry->Terms[TermIndex].Name;
+						const uint32* ExistingIntermediateIndexPtr = IntermediateMap.Find(TermName);
+						if(ExistingIntermediateIndexPtr == nullptr)
+						{
+							uint32 IntermediateIndex = IntermediateTerms.Emplace(TermName, GraphEntry->Terms[TermIndex].Type, GraphEntry->Terms[TermIndex].Direction);
+							IntermediateMap.Add(TermName, IntermediateIndex);
+							GraphTask.Terms.Add(IntermediateIndex);
+						}
+						else
+						{
+							GraphTask.Terms.Add(*ExistingIntermediateIndexPtr);
+						}
+					}
+
+					int32 TaskIndex = InSchedule->GraphTasks.Add(GraphTask);
+
+					InSchedule->NumTickFunctions++;
+
+					Emit(EAnimNextScheduleScheduleOpcode::RunGraphTask, TaskIndex);
+				}
+			}
+			else if (UAnimNextScheduleEntry_ExternalTask* ExternalTaskEntry = Cast<UAnimNextScheduleEntry_ExternalTask>(Entry))
+			{
+				EmitPrerequisite();
+
+				FAnimNextScheduleExternalTask ExternalTask;
+				ExternalTask.TaskIndex = InSchedule->ExternalTasks.Num();
+				ExternalTask.ParamScopeIndex = InSchedule->NumParameterScopes++;
+				ExternalTask.ParamParentScopeIndex = ParentScopeIndex;
+				ExternalTask.ExternalTask = ExternalTaskEntry->ExternalTask;
+				int32 ExternalTaskIndex = InSchedule->ExternalTasks.Add(ExternalTask);
+
+				// Emit the external task
+				Emit(EAnimNextScheduleScheduleOpcode::BeginRunExternalTask, ExternalTaskIndex);
+				InSchedule->NumTickFunctions++;
+
+				EmitPrerequisite();
+
+				Emit(EAnimNextScheduleScheduleOpcode::EndRunExternalTask, ExternalTaskIndex);
+				InSchedule->NumTickFunctions++;
+			}
+			else if (UAnimNextScheduleEntry_ParamScope* ParamScopeTaskEntry = Cast<UAnimNextScheduleEntry_ParamScope>(Entry))
+			{
+				EmitPrerequisite();
+
+				FAnimNextScheduleParamScopeEntryTask ParamScopeEntryTask;
+				ParamScopeEntryTask.TaskIndex = InSchedule->ParamScopeEntryTasks.Num();
+				const uint32 ParamScopeIndex = InSchedule->NumParameterScopes++;
+				ParamScopeEntryTask.ParamScopeIndex = ParamScopeIndex;
+				ParamScopeEntryTask.ParamParentScopeIndex = ParentScopeIndex;
+				ParamScopeEntryTask.TickFunctionIndex = InSchedule->NumTickFunctions;
+				ParamScopeEntryTask.Scope = ParamScopeTaskEntry->Scope;
+				ParamScopeEntryTask.ParameterBlocks = ParamScopeTaskEntry->ParameterBlocks;
+				int32 ParamScopeTaskEntryIndex = InSchedule->ParamScopeEntryTasks.Add(ParamScopeEntryTask);
+
+				Emit(EAnimNextScheduleScheduleOpcode::RunParamScopeEntry, ParamScopeTaskEntryIndex);
+				InSchedule->NumTickFunctions++;
+
+				// Enter new scope
+				uint32 PreviousParentScope = ParentScopeIndex;
+				ParentScopeIndex = ParamScopeIndex;
+
+				// Emit the subentries
+				EmitEntries(ParamScopeTaskEntry->SubEntries);
+
+				// Exit scope
+				ParentScopeIndex = PreviousParentScope;
+
+				EmitPrerequisite();
+
+				FAnimNextScheduleParamScopeExitTask ParamScopeExitTask;
+				ParamScopeExitTask.TaskIndex = InSchedule->ParamScopeExitTasks.Num();
+				ParamScopeExitTask.ParamScopeIndex = ParamScopeIndex;
+				ParamScopeExitTask.Scope = ParamScopeTaskEntry->Scope;
+				int32 ParamScopeExitTaskIndex = InSchedule->ParamScopeExitTasks.Add(ParamScopeExitTask);
+
+				Emit(EAnimNextScheduleScheduleOpcode::RunParamScopeExit, ParamScopeExitTaskIndex);
+				InSchedule->NumTickFunctions++;
+			}
+			else if(UAnimNextScheduleEntry_ExternalParams* ExternalParamsTaskEntry = Cast<UAnimNextScheduleEntry_ExternalParams>(Entry))
+			{
+				EmitPrerequisite();
+
+				FAnimNextScheduleExternalParamTask ExternalParamTask;
+				ExternalParamTask.TaskIndex = InSchedule->ExternalParamTasks.Num();
+				ExternalParamTask.ParameterSources = ExternalParamsTaskEntry->ParameterSources;
+				ExternalParamTask.bThreadSafe = ExternalParamsTaskEntry->bThreadSafe;
+				int32 ExternalParamTaskEntryIndex = InSchedule->ExternalParamTasks.Add(ExternalParamTask);
+
+				Emit(EAnimNextScheduleScheduleOpcode::RunExternalParamTask, ExternalParamTaskEntryIndex);
+				InSchedule->NumTickFunctions++;
+			}
+		}
+	};
+
+	auto GenerateExternalParameters = [](TArray<TObjectPtr<UAnimNextScheduleEntry>>& InEntries)
+	{
+		struct FParameterTracker
+		{
+			FName SourceName;
+			TArray<TObjectPtr<UAnimNextScheduleEntry>>* BestContainer = nullptr;
+			TSet<FName> ThreadSafeParameters;
+			TSet<FName> NonThreadSafeParameters;
+			uint32 BestDistance = MAX_uint32;
+			uint32 BestArrayIndex = MAX_uint32;
+		};
+
+		// We need to find the task that is 'earliest' in the DAG for each external parameter, so we track that with this map
+		TMap<FName, FParameterTracker> TrackerMap;
+		auto TrackExternalParameters = [&TrackerMap](TConstArrayView<FName> InParameterNames, uint32 InDistance, UAnimNextScheduleEntry* InEntry, TArray<TObjectPtr<UAnimNextScheduleEntry>>* InContainer, int32 InArrayIndex, bool bInUpdateDependent)
+		{
+			bool bHasExternal = false;
+			for(FName ParameterName : InParameterNames)
+			{
+				// Only add if the parameter name is 'external'
+				FName ParameterSourceName = FExternalParameterRegistry::FindSourceForParameter(ParameterName);
+				if(ParameterSourceName != NAME_None)
+				{
+					bHasExternal = true;
+					FParameterTracker& Tracker = TrackerMap.FindOrAdd(ParameterSourceName);
+
+					Tracker.SourceName = ParameterSourceName;
+
+					// Only track array index if this param is update dependent
+					if(bInUpdateDependent && InDistance < Tracker.BestDistance)
+					{
+						Tracker.BestDistance = InDistance;
+						Tracker.BestContainer = InContainer;
+						Tracker.BestArrayIndex = InArrayIndex == 0 ? 0 : InArrayIndex - 1;
+					}
+
+					IParameterSourceFactory::FParameterInfo Info;
+					ensure(FExternalParameterRegistry::FindParameterInfo(ParameterName, Info));
+
+					if(Info.bThreadSafe)
+					{
+						Tracker.ThreadSafeParameters.Add(ParameterName);
+					}
+					else
+					{
+						Tracker.NonThreadSafeParameters.Add(ParameterName);
+					}
+				}
+			}
+			return bHasExternal;
+		};
+
+		TFunction<void(TArray<TObjectPtr<UAnimNextScheduleEntry>>&, uint32)> PopulateExternalParameters;
+		PopulateExternalParameters = [&PopulateExternalParameters, &TrackExternalParameters](TArray<TObjectPtr<UAnimNextScheduleEntry>>& InEntries, uint32 InDistance)
+		{
+			// First populate internal parameters for all those tasks entries that reference them
+			int32 ArrayIndex = 0;
+			for(UAnimNextScheduleEntry* Entry : InEntries)
+			{
+				if (UAnimNextScheduleEntry_AnimNextGraph* GraphEntry = Cast<UAnimNextScheduleEntry_AnimNextGraph>(Entry))
+				{
+					if(GraphEntry->Graph)
+					{
+						FAnimNextParameterProviderAssetRegistryExports Exports;
+						if(UncookedOnly::FUtils::GetExportedParametersForAsset(FAssetData(GraphEntry->Graph), Exports))
+						{
+							TArray<FName> RequiredParameters;
+							RequiredParameters.Reserve(Exports.Parameters.Num());
+							for(const FAnimNextParameterAssetRegistryExportEntry& ExportedParameter : Exports.Parameters)
+							{
+								RequiredParameters.Add(ExportedParameter.Name);
+							}
+							TrackExternalParameters(RequiredParameters, InDistance, Entry, &InEntries, ArrayIndex, true);
+						}
+					}
+
+					TrackExternalParameters(GraphEntry->RequiredParameters, InDistance, Entry, &InEntries, ArrayIndex, true);
+
+					// All graphs require access to LOD and ref pose
+					// TODO: need to not require defaults here - use static graph params. In fact, these params should probably be defined at the schedule level.
+					TrackExternalParameters({ UAnimNextGraph::DefaultCurrentLODId.GetName(), UAnimNextGraph::DefaultReferencePoseId.GetName() }, InDistance, Entry, &InEntries, ArrayIndex, true);
+				}
+				else if (UAnimNextScheduleEntry_Port* PortEntry = Cast<UAnimNextScheduleEntry_Port>(Entry))
+				{
+					if(PortEntry->Port)
+					{
+						UAnimNextSchedulePort* CDO = PortEntry->Port->GetDefaultObject<UAnimNextSchedulePort>();
+						TConstArrayView<FName> RequiredParameters = CDO->GetRequiredParameters();
+						TrackExternalParameters(RequiredParameters, InDistance, Entry, &InEntries, ArrayIndex, true);
+					}
+				}
+				else if (UAnimNextScheduleEntry_ExternalTask* ExternalTaskEntry = Cast<UAnimNextScheduleEntry_ExternalTask>(Entry))
+				{
+					// Note: external task params are not update dependent as this would cause external param updates to occur before tick functions
+					TrackExternalParameters({ ExternalTaskEntry->ExternalTask }, InDistance, Entry, &InEntries, ArrayIndex, false);
+				}
+				else if (UAnimNextScheduleEntry_ParamScope* ParamScopeTaskEntry = Cast<UAnimNextScheduleEntry_ParamScope>(Entry))
+				{
+					for(UAnimNextParameterBlock* ParameterBlock : ParamScopeTaskEntry->ParameterBlocks)
+					{
+						if(ParameterBlock)
+						{
+							FAnimNextParameterProviderAssetRegistryExports Exports;
+							if(UncookedOnly::FUtils::GetExportedParametersForAsset(FAssetData(ParameterBlock), Exports))
+							{
+								TArray<FName> RequiredParameters;
+								RequiredParameters.Reserve(Exports.Parameters.Num());
+								for(const FAnimNextParameterAssetRegistryExportEntry& ExportedParameter : Exports.Parameters)
+								{
+									RequiredParameters.Add(ExportedParameter.Name);
+								}
+								TrackExternalParameters(RequiredParameters, InDistance, Entry, &InEntries, ArrayIndex, true);
+							}
+						}
+					}
+
+					// Recurse into sub-entries
+					PopulateExternalParameters(ParamScopeTaskEntry->SubEntries, InDistance);
+				}
+
+				InDistance++;
+				ArrayIndex++;
+			}
+		};
+
+		// First populate internal parameter lists and tracker map
+		PopulateExternalParameters(InEntries, 0);
+
+		// Unique location: container, index and thread safe flag
+		using FInsertionLocation = TTuple<TArray<TObjectPtr<UAnimNextScheduleEntry>>*, uint32, bool>;
+
+		// Build sources that need to run (thread-safe or not) at each index
+		TMap<FInsertionLocation, TArray<FAnimNextScheduleExternalParameterSource>> InsertionMap;
+		for(const TPair<FName, FParameterTracker>& TrackedParameterPair : TrackerMap)
+		{
+			// If an index/container was not set up, then the update of the set of parameters is not a pre-requisite of a task, so we just insert the task at the
+			// start of the schedule as the only requirement is that the parameters exist 
+			TArray<TObjectPtr<UAnimNextScheduleEntry>>* Container = TrackedParameterPair.Value.BestContainer != nullptr ? TrackedParameterPair.Value.BestContainer : &InEntries;
+			uint32 ArrayIndex = TrackedParameterPair.Value.BestArrayIndex != MAX_uint32 ? TrackedParameterPair.Value.BestArrayIndex : 0;
+			
+			if(TrackedParameterPair.Value.ThreadSafeParameters.Num() > 0)
+			{
+				TArray<FAnimNextScheduleExternalParameterSource>& ParameterSources = InsertionMap.FindOrAdd({ Container, ArrayIndex, true });
+				FAnimNextScheduleExternalParameterSource& NewSource = ParameterSources.AddDefaulted_GetRef();
+				NewSource.ParameterSource = TrackedParameterPair.Value.SourceName;
+				NewSource.Parameters = TrackedParameterPair.Value.ThreadSafeParameters.Array();
+			}
+
+			if(TrackedParameterPair.Value.NonThreadSafeParameters.Num() > 0)
+			{
+				TArray<FAnimNextScheduleExternalParameterSource>& ParameterSources = InsertionMap.FindOrAdd({ Container, ArrayIndex, false });
+				FAnimNextScheduleExternalParameterSource& NewSource = ParameterSources.AddDefaulted_GetRef();
+				NewSource.ParameterSource = TrackedParameterPair.Value.SourceName;
+				NewSource.Parameters = TrackedParameterPair.Value.NonThreadSafeParameters.Array();
+			}
+		}
+
+		// Sort sources map by insertion index
+		InsertionMap.KeySort([](const FInsertionLocation& InLHS, const FInsertionLocation& InRHS)
+		{
+			return InLHS.Get<1>() > InRHS.Get<1>();
+		});
+
+		// Now insert a task to fetch the parameters at the recorded index/container
+		// NOTE: here we just insert the task before the earliest usage of the external parameter source we found, but in the case of a full DAG
+		// schedule, we would need to add a prerequisite for ALL tasks that use the external parameters
+		for(const TPair<FInsertionLocation, TArray<FAnimNextScheduleExternalParameterSource>>& InsertionLocationPair : InsertionMap)
+		{
+			// Add a new external param entry
+			UAnimNextScheduleEntry_ExternalParams* NewParamEntry = NewObject<UAnimNextScheduleEntry_ExternalParams>();
+			NewParamEntry->bThreadSafe = InsertionLocationPair.Key.Get<2>();
+			NewParamEntry->ParameterSources = InsertionLocationPair.Value;
+			InsertionLocationPair.Key.Get<0>()->Insert(NewParamEntry, InsertionLocationPair.Key.Get<1>());
+		}
+	};
+
+	// Duplicate the entries, we are going to rewrite them
+	TArray<TObjectPtr<UAnimNextScheduleEntry>> NewEntries;
+	for(UAnimNextScheduleEntry* Entry : InSchedule->Entries)
+	{
+		if(Entry)
+		{
+			NewEntries.Add(CastChecked<UAnimNextScheduleEntry>(StaticDuplicateObject(Entry, GetTransientPackage())));
+		}
+	}
+
+	// Push required parameters up scopes
+	GenerateExternalParameters(NewEntries);
+
+	// Emit the schedule 'bytecode'
+	EmitEntries(NewEntries);
+
+	Emit(EAnimNextScheduleScheduleOpcode::Exit);
+
+	// Process intermediates
+	if(IntermediateMap.Num() > 0)
+	{
+		check(IntermediateMap.Num() == IntermediateTerms.Num());
+		
+		TArray<FPropertyBagPropertyDesc> PropertyDescs;
+		PropertyDescs.Reserve(IntermediateTerms.Num());
+		 
+		for(const TPair<FName, uint32>& IntermediatePair : IntermediateMap)
+		{
+			const FAnimNextParamType& IntermediateType = IntermediateTerms[IntermediatePair.Value].Type;
+			check(IntermediateType.IsValid());
+			PropertyDescs.Emplace(IntermediatePair.Key, IntermediateType.GetContainerType(), IntermediateType.GetValueType(), IntermediateType.GetValueTypeObject());
+		}
+
+		InSchedule->IntermediatesData.AddProperties(PropertyDescs);
+	}
+}
+
 }
