Fix for scheduler thumbnails not getting properly initialized when duplicating AnimNext graphs

diff --git a/Engine/Plugins/Experimental/AnimNext/Source/AnimNextUncookedOnly/Internal/Graph/AnimNextGraph_EditorData.h b/Engine/Plugins/Experimental/AnimNext/Source/AnimNextUncookedOnly/Internal/Graph/AnimNextGraph_EditorData.h
--- a/Engine/Plugins/Experimental/AnimNext/Source/AnimNextUncookedOnly/Internal/Graph/AnimNextGraph_EditorData.h
+++ b/Engine/Plugins/Experimental/AnimNext/Source/AnimNextUncookedOnly/Internal/Graph/AnimNextGraph_EditorData.h
@@ -96,6 +96,7 @@ public:
 	// UObject interface
 	virtual void Serialize(FArchive& Ar) override;
 	virtual void PostLoad() override;
+	virtual void PostDuplicate(EDuplicateMode::Type DuplicateMode) override;
 	virtual bool IsEditorOnly() const override { return true; }
 	virtual void GetAssetRegistryTags(TArray<FAssetRegistryTag>& OutTags) const override;
 #if WITH_EDITOR
diff --git a/Engine/Plugins/Experimental/AnimNext/Source/AnimNextUncookedOnly/Internal/Param/AnimNextParameterBlock_EditorData.h b/Engine/Plugins/Experimental/AnimNext/Source/AnimNextUncookedOnly/Internal/Param/AnimNextParameterBlock_EditorData.h
--- a/Engine/Plugins/Experimental/AnimNext/Source/AnimNextUncookedOnly/Internal/Param/AnimNextParameterBlock_EditorData.h
+++ b/Engine/Plugins/Experimental/AnimNext/Source/AnimNextUncookedOnly/Internal/Param/AnimNextParameterBlock_EditorData.h
@@ -132,6 +132,7 @@ class UAnimNextParameterBlock_EditorData : public UObject, public IRigVMClientHo
 	// UObject interface
 	virtual void Serialize(FArchive& Ar) override;
 	virtual void PostLoad() override;
+	virtual void PostDuplicate(EDuplicateMode::Type DuplicateMode) override;
 	virtual void PostTransacted(const FTransactionObjectEvent& TransactionEvent) override;
 	virtual bool IsEditorOnly() const override { return true; }
 	virtual void GetAssetRegistryTags(TArray<FAssetRegistryTag>& OutTags) const override;
diff --git a/Engine/Plugins/Experimental/AnimNext/Source/AnimNextUncookedOnly/Private/Graph/AnimNextGraph_EditorData.cpp b/Engine/Plugins/Experimental/AnimNext/Source/AnimNextUncookedOnly/Private/Graph/AnimNextGraph_EditorData.cpp
--- a/Engine/Plugins/Experimental/AnimNext/Source/AnimNextUncookedOnly/Private/Graph/AnimNextGraph_EditorData.cpp
+++ b/Engine/Plugins/Experimental/AnimNext/Source/AnimNextUncookedOnly/Private/Graph/AnimNextGraph_EditorData.cpp
@@ -167,6 +167,11 @@ void UAnimNextGraph_EditorData::PostLoad()
 #endif // WITH_EDITOR
 }
 
+void UAnimNextGraph_EditorData::PostDuplicate(EDuplicateMode::Type DuplicateMode)
+{
+	Initialize(/*bRecompileVM*/true);
+}
+
 #if WITH_EDITOR
 void UAnimNextGraph_EditorData::GetAssetRegistryTags(TArray<FAssetRegistryTag>& OutTags) const
 {
diff --git a/Engine/Plugins/Experimental/AnimNext/Source/AnimNextUncookedOnly/Private/Param/AnimNextParameterBlock_EditorData.cpp b/Engine/Plugins/Experimental/AnimNext/Source/AnimNextUncookedOnly/Private/Param/AnimNextParameterBlock_EditorData.cpp
--- a/Engine/Plugins/Experimental/AnimNext/Source/AnimNextUncookedOnly/Private/Param/AnimNextParameterBlock_EditorData.cpp
+++ b/Engine/Plugins/Experimental/AnimNext/Source/AnimNextUncookedOnly/Private/Param/AnimNextParameterBlock_EditorData.cpp
@@ -563,6 +563,11 @@ void UAnimNextParameterBlock_EditorData::PostLoad()
 #endif // WITH_EDITOR
 }
 
+void UAnimNextParameterBlock_EditorData::PostDuplicate(EDuplicateMode::Type DuplicateMode)
+{
+	Initialize(/*bRecompileVM*/true);
+}
+
 void UAnimNextParameterBlock_EditorData::PostTransacted(const FTransactionObjectEvent& TransactionEvent)
 {
 	Super::PostTransacted(TransactionEvent);
diff --git a/Engine/Plugins/Experimental/AnimNext/Source/AnimNextUncookedOnly/Private/UncookedOnlyUtils.cpp b/Engine/Plugins/Experimental/AnimNext/Source/AnimNextUncookedOnly/Private/UncookedOnlyUtils.cpp
--- a/Engine/Plugins/Experimental/AnimNext/Source/AnimNextUncookedOnly/Private/UncookedOnlyUtils.cpp
+++ b/Engine/Plugins/Experimental/AnimNext/Source/AnimNextUncookedOnly/Private/UncookedOnlyUtils.cpp
@@ -537,7 +537,10 @@ void FUtils::Compile(UAnimNextGraph* InGraph)
 
 void FUtils::RecreateVM(UAnimNextGraph* InGraph)
 {
-	InGraph->VM = NewObject<URigVM>(InGraph, TEXT("VM"), RF_NoFlags);
+	if (InGraph->VM == nullptr)
+	{
+		InGraph->VM = NewObject<URigVM>(InGraph, TEXT("VM"), RF_NoFlags);
+	}
 	InGraph->VM->Reset(InGraph->ExtendedExecuteContext);
 	InGraph->RigVM = InGraph->VM; // Local serialization
 }
@@ -788,7 +791,10 @@ void FUtils::Compile(UAnimNextParameterBlock* InParameterBlock)
 
 void FUtils::RecreateVM(UAnimNextParameterBlock* InParameterBlock)
 {
-	InParameterBlock->VM = NewObject<URigVM>(InParameterBlock, TEXT("VM"), RF_NoFlags);
+	if(InParameterBlock->VM == nullptr)
+	{
+		InParameterBlock->VM = NewObject<URigVM>(InParameterBlock, TEXT("VM"), RF_NoFlags);
+	}
 	InParameterBlock->VM->Reset(InParameterBlock->GetRigVMExtendedExecuteContext());
 	InParameterBlock->RigVM = InParameterBlock->VM; // Local serialization
 }
