Fix for AnimNext node details panel not refreshing when adding decorators and details getting cleared when selecting an inline value inside a node

#rnx

diff --git a/Engine/Plugins/Experimental/AnimNext/Source/AnimNextEditor/Private/Workspace/AnimNextWorkspaceEditor.cpp b/Engine/Plugins/Experimental/AnimNext/Source/AnimNextEditor/Private/Workspace/AnimNextWorkspaceEditor.cpp
--- a/Engine/Plugins/Experimental/AnimNext/Source/AnimNextEditor/Private/Workspace/AnimNextWorkspaceEditor.cpp
+++ b/Engine/Plugins/Experimental/AnimNext/Source/AnimNextEditor/Private/Workspace/AnimNextWorkspaceEditor.cpp
@@ -325,7 +325,7 @@ void FWorkspaceEditor::SaveAsset_Execute()
 
 void FWorkspaceEditor::OnGraphModified(ERigVMGraphNotifType Type, URigVMGraph* Graph, UObject* Subject)
 {
-	if (Type == ERigVMGraphNotifType::PinDefaultValueChanged)
+	if (Type == ERigVMGraphNotifType::InteractionBracketClosed)
 	{
 		if (DetailsView.IsValid())
 		{
@@ -337,9 +337,12 @@ void FWorkspaceEditor::OnGraphModified(ERigVMGraphNotifType Type, URigVMGraph* G
 void FWorkspaceEditor::SetFocusedGraphEditor(TSharedPtr<SGraphEditor> InGraphEditor)
 {
 	// Update the graph editor that is currently focused
-	FocusedGraphEdPtr = InGraphEditor;
+	if (FocusedGraphEdPtr != InGraphEditor)
+	{
+		FocusedGraphEdPtr = InGraphEditor;
 
-	SetSelectedObjects({});
+		SetSelectedObjects({});
+	}
 }
 
 UEdGraph* FWorkspaceEditor::GetFocusedGraph() const
