Fix AnimNext CIS warnings

Reorder class proxy generation to formalize priority (Extensions -> Accessors -> Properties) and remove warning
Fix uninitialized variable warning

diff --git a/Engine/Plugins/Experimental/AnimNext/Source/AnimNext/Private/Param/ClassProxy.cpp b/Engine/Plugins/Experimental/AnimNext/Source/AnimNext/Private/Param/ClassProxy.cpp
--- a/Engine/Plugins/Experimental/AnimNext/Source/AnimNext/Private/Param/ClassProxy.cpp
+++ b/Engine/Plugins/Experimental/AnimNext/Source/AnimNext/Private/Param/ClassProxy.cpp
@@ -12,6 +12,93 @@ namespace UE::AnimNext
 FClassProxy::FClassProxy(const UClass* InClass)
 	: Class(InClass)
 {
+	// Add any additional extension libraries that extend this class first
+	{
+		// We allow only native static 'accessor wrapper' functions (one object param and the return value)
+		auto CheckSignature = [](const UFunction* InFunction, const UClass* InProxiedClass)
+		{
+			if(!InFunction->HasAllFunctionFlags(FUNC_BlueprintCallable | FUNC_Static | FUNC_Native))
+			{
+				return false;
+			}
+
+			if(InFunction->NumParms != 2)
+			{
+				return false;
+			}
+		
+			int32 ParamIndex = 0;
+			for(TFieldIterator<FProperty> It(InFunction); It && (It->PropertyFlags & CPF_Parm); ++It, ++ParamIndex)
+			{
+				// Check parameter is an object of the expected class
+				if(ParamIndex == 0)
+				{
+					FObjectProperty* ObjectProperty = CastField<FObjectProperty>(*It);
+					if(ObjectProperty == nullptr)
+					{ 
+						return false;
+					}
+				
+					if(!ObjectProperty->PropertyClass->IsChildOf(InProxiedClass))
+					{
+						return false;
+					}
+				}
+
+				// Check return value
+				if(ParamIndex == 1)
+				{
+					if((It->PropertyFlags & CPF_ReturnParm) == 0)
+					{
+						return false;
+					}
+				}
+			}
+			return true;
+		};
+
+		TArray<UClass*> Classes;
+		GetDerivedClasses(UAnimNextClassExtensionLibrary::StaticClass(), Classes);
+		for(UClass* ProxyClass : Classes)
+		{
+			UAnimNextClassExtensionLibrary* CDO = ProxyClass->GetDefaultObject<UAnimNextClassExtensionLibrary>();
+			const UClass* ExtendedClass = CDO->GetSupportedClass();
+			if(InClass->IsChildOf(ExtendedClass))
+			{
+				for(TFieldIterator<UFunction> It(ProxyClass); It; ++It)
+				{
+					UFunction* Function = *It;
+					if(CheckSignature(Function, ExtendedClass))
+					{
+						const FProperty* ReturnProperty = Function->GetReturnProperty();
+						FParamTypeHandle TypeHandle = FParamTypeHandle::FromProperty(ReturnProperty);
+						if(TypeHandle.IsValid())
+						{
+							FString FunctionName = Function->GetName();
+
+							// Remove common prefixes.
+							// TODO: explicit mapping or regex transformations?
+							FunctionName.RemoveFromStart(TEXT("K2_"));
+							FunctionName.RemoveFromStart(TEXT("Get"));
+
+							FClassProxyParameter Parameter;
+							Parameter.AccessType = EClassProxyParameterAccessType::HoistedFunction;
+							Parameter.ClassParameterName = *FunctionName;
+							Parameter.Function = Function;
+							Parameter.Type = TypeHandle.GetType();
+#if WITH_EDITOR
+							Parameter.Tooltip = Function->GetToolTipText();
+							Parameter.bThreadSafe = Function->HasMetaData("BlueprintThreadSafe");
+#endif
+							ParameterNameMap.Add(Parameter.ClassParameterName, Parameters.Add(Parameter));
+						}
+					}
+				}
+			}
+		}
+	}
+
+	// Add functions as the next priority (extensions have already been added above so will take priority with duplicate names)
 	for(TFieldIterator<UFunction> It(InClass, EFieldIterationFlags::IncludeSuper | EFieldIterationFlags::IncludeInterfaces); It; ++It)
 	{
 		UFunction* Function = *It;
@@ -30,21 +117,25 @@ FClassProxy::FClassProxy(const UClass* InClass)
 				FunctionName.RemoveFromStart(TEXT("K2_"));
 				FunctionName.RemoveFromStart(TEXT("Get"));
 
-				FClassProxyParameter Parameter;
-				Parameter.AccessType = EClassProxyParameterAccessType::AccessorFunction;
-				Parameter.ClassParameterName = *FunctionName;
-				Parameter.Function = Function;
-				Parameter.Type = TypeHandle.GetType();
+				FName ParameterName = *FunctionName;
+				if(!ParameterNameMap.Contains(ParameterName))
+				{
+					FClassProxyParameter Parameter;
+					Parameter.AccessType = EClassProxyParameterAccessType::AccessorFunction;
+					Parameter.ClassParameterName = *FunctionName;
+					Parameter.Function = Function;
+					Parameter.Type = TypeHandle.GetType();
 #if WITH_EDITOR
-				Parameter.Tooltip = Function->GetToolTipText();
-				Parameter.bThreadSafe = Function->HasMetaData("BlueprintThreadSafe");
+					Parameter.Tooltip = Function->GetToolTipText();
+					Parameter.bThreadSafe = Function->HasMetaData("BlueprintThreadSafe");
 #endif
-				ParameterNameMap.Add(Parameter.ClassParameterName, Parameters.Add(Parameter));
+					ParameterNameMap.Add(Parameter.ClassParameterName, Parameters.Add(Parameter));
+				}
 			}
 		}
 	}
 
-	// Add properties after functions, checking for duplicate names
+	// Finally add properties (accessors and extensions have already been added above so will take priority with duplicate names)
 	for(TFieldIterator<FProperty> It(InClass, EFieldIterationFlags::IncludeSuper | EFieldIterationFlags::IncludeInterfaces); It; ++It)
 	{
 		FProperty* Property = *It;
@@ -73,102 +164,6 @@ FClassProxy::FClassProxy(const UClass* InClass)
 #endif
 					ParameterNameMap.Add(Parameter.ClassParameterName, Parameters.Add(Parameter));
 				}
-				else
-				{
-					UE_LOGFMT(LogAnimation, Warning, "Could not add parameter '{ParameterName}' for class '{ClassName}' as it already exists", ParameterName, InClass->GetFName());
-				}
-			}
-		}
-	}
-
-	// We allow only native static 'accessor wrapper' functions (one object param and the return value)
-	auto CheckSignature = [](const UFunction* InFunction, const UClass* InProxiedClass)
-	{
-		if(!InFunction->HasAllFunctionFlags(FUNC_BlueprintCallable | FUNC_Static | FUNC_Native))
-		{
-			return false;
-		}
-
-		if(InFunction->NumParms != 2)
-		{
-			return false;
-		}
-		
-		int32 ParamIndex = 0;
-		for(TFieldIterator<FProperty> It(InFunction); It && (It->PropertyFlags & CPF_Parm); ++It, ++ParamIndex)
-		{
-			// Check parameter is an object of the expected class
-			if(ParamIndex == 0)
-			{
-				FObjectProperty* ObjectProperty = CastField<FObjectProperty>(*It);
-				if(ObjectProperty == nullptr)
-				{ 
-					return false;
-				}
-				
-				if(!ObjectProperty->PropertyClass->IsChildOf(InProxiedClass))
-				{
-					return false;
-				}
-			}
-
-			// Check return value
-			if(ParamIndex == 1)
-			{
-				if((It->PropertyFlags & CPF_ReturnParm) == 0)
-				{
-					return false;
-				}
-			}
-		}
-		return true;
-	};
-
-	// Add any additional extension libraries that extend this class
-	TArray<UClass*> Classes;
-	GetDerivedClasses(UAnimNextClassExtensionLibrary::StaticClass(), Classes);
-	for(UClass* ProxyClass : Classes)
-	{
-		UAnimNextClassExtensionLibrary* CDO = ProxyClass->GetDefaultObject<UAnimNextClassExtensionLibrary>();
-		const UClass* ExtendedClass = CDO->GetSupportedClass();
-		if(InClass->IsChildOf(ExtendedClass))
-		{
-			for(TFieldIterator<UFunction> It(ProxyClass); It; ++It)
-			{
-				UFunction* Function = *It;
-				if(CheckSignature(Function, ExtendedClass))
-				{
-					const FProperty* ReturnProperty = Function->GetReturnProperty();
-					FParamTypeHandle TypeHandle = FParamTypeHandle::FromProperty(ReturnProperty);
-					if(TypeHandle.IsValid())
-					{
-						FString FunctionName = Function->GetName();
-
-						// Remove common prefixes.
-						// TODO: explicit mapping or regex transformations?
-						FunctionName.RemoveFromStart(TEXT("K2_"));
-						FunctionName.RemoveFromStart(TEXT("Get"));
-
-						FName ParameterName = *FunctionName;
-						if(!ParameterNameMap.Contains(ParameterName))
-						{
-							FClassProxyParameter Parameter;
-							Parameter.AccessType = EClassProxyParameterAccessType::HoistedFunction;
-							Parameter.ClassParameterName = *FunctionName;
-							Parameter.Function = Function;
-							Parameter.Type = TypeHandle.GetType();
-#if WITH_EDITOR
-							Parameter.Tooltip = Function->GetToolTipText();
-							Parameter.bThreadSafe = Function->HasMetaData("BlueprintThreadSafe");
-#endif
-							ParameterNameMap.Add(Parameter.ClassParameterName, Parameters.Add(Parameter));
-						}
-						else
-						{
-							UE_LOGFMT(LogAnimation, Warning, "Could not add parameter '{ParameterName}' for class proxy '{ClassName}' as it already exists", ParameterName, InClass->GetFName());
-						}
-					}
-				}
 			}
 		}
 	}
diff --git a/Engine/Plugins/Experimental/AnimNext/Source/AnimNextUncookedOnly/Internal/UncookedOnlyUtils.h b/Engine/Plugins/Experimental/AnimNext/Source/AnimNextUncookedOnly/Internal/UncookedOnlyUtils.h
--- a/Engine/Plugins/Experimental/AnimNext/Source/AnimNextUncookedOnly/Internal/UncookedOnlyUtils.h
+++ b/Engine/Plugins/Experimental/AnimNext/Source/AnimNextUncookedOnly/Internal/UncookedOnlyUtils.h
@@ -65,7 +65,7 @@ struct FAnimNextParameterAssetRegistryExportEntry
 	FAssetData ReferencingAsset;
 
 	UPROPERTY()
-	EAnimNextParameterFlags Flags;
+	EAnimNextParameterFlags Flags = EAnimNextParameterFlags::NoFlags;
 };
 
 USTRUCT()
