AnimNext: Misc parameter/schedule improvements

Improve SetParameterInScope BP support
- Fix nested scopes using incorrect indices
- Adds before/after scope option
- Merged with initial parameter setup code
- Added existing pin customization to regular Blueprint nodes
Exposed parameters used in schedules and Blueprints to the asset registry.
Adds a combined parameter name/type struct FAnimNextParam. This is convenient to use when the type is not adjacent to the name already. Some other use cases could still be converted over.
Implemented dynamic graph incompatibility checks/warnings.
Schedules can now be live-recompiled.
- Verified that schedule state changes are using valid transitions.
- Reworked schedule entry initialization a little.
Pickers now correclty pre-populate types when creating new parameters.
New parameter dialog now verifies and sanitizes parameter names.
Removed dead code:
- Schedule context thread attachment.
- 'Processor' param block entries
Fixed incorrect node notation use in FUtils::GetGraphParameters.

diff --git a/Engine/Plugins/Experimental/AnimNext/Source/AnimNext/Internal/Component/AnimNextComponent.h b/Engine/Plugins/Experimental/AnimNext/Source/AnimNext/Internal/Component/AnimNextComponent.h
new file mode 100644
--- /dev/null
+++ b/Engine/Plugins/Experimental/AnimNext/Source/AnimNext/Internal/Component/AnimNextComponent.h
@@ -0,0 +1,71 @@
+// Copyright Epic Games, Inc. All Rights Reserved.
+
+#pragma once
+
+#include "AnimNextComponentParameter.h"
+#include "Components/ActorComponent.h"
+#include "Scheduler/AnimNextSchedule.h"
+#include "Scheduler/ScheduleHandle.h"
+#include "AnimNextComponent.generated.h"
+
+class UAnimNextSchedule;
+struct FAnimNextComponentInstanceData;
+
+namespace UE::AnimNext::UncookedOnly
+{
+	struct FUtils;
+}
+
+UENUM()
+enum class EAnimNextParameterScopeOrdering : uint8
+{
+	// Value will be pushed before the scope, allowing the static scope to potentially override the value
+	Before,
+
+	// Value will be pushed after the scope, potentially overriding the static scope
+	After,
+};
+
+UCLASS(MinimalAPI, meta = (BlueprintSpawnableComponent))
+class UAnimNextComponent : public UActorComponent
+{
+	GENERATED_BODY()
+
+	// UActorComponent interface
+	virtual void OnRegister() override;
+	virtual void OnUnregister() override;
+
+public:
+	// Sets a parameter's value in the supplied scope.
+	// @param    Scope    Scopes corresponding to an existing scope in a schedule, or "None". Passing "None" will apply the parameter to the whole schedule.
+	// @param    Ordering Where to apply the parameter in relation to the supplied scope. Ignored for scope "None".
+	// @param    Name     The name of the parameter to apply
+	// @param    Value    The value to set the parameter to
+	UFUNCTION(BlueprintCallable, Category = "AnimNext", CustomThunk, meta = (CustomStructureParam = Value, UnsafeDuringActorConstruction))
+	void SetParameterInScope(UPARAM(meta = (CustomWidget = "ParamName", AllowedParamType = "FAnimNextScope")) FName Scope, EAnimNextParameterScopeOrdering Ordering, UPARAM(meta = (CustomWidget = "ParamName")) FName Name, int32 Value);
+
+	// Enable or disable this component's update
+	UFUNCTION(BlueprintCallable, Category = "AnimNext")
+	void Enable(bool bEnabled);
+	
+private:
+	DECLARE_FUNCTION(execSetParameterInScope);
+
+private:
+	friend struct UE::AnimNext::UncookedOnly::FUtils;
+
+	// The execution schedule that this component will run
+	UPROPERTY(EditAnywhere, Category="Schedule")
+	TObjectPtr<UAnimNextSchedule> Schedule = nullptr;
+
+	// Parameters to apply on schedule/component registration
+	UPROPERTY(Instanced, EditAnywhere, Category = "Parameters")
+	TArray<TObjectPtr<UAnimNextComponentParameter>> Parameters;
+
+	// How to initialize the schedule
+	UPROPERTY(EditAnywhere, Category="Schedule")
+	EAnimNextScheduleInitMethod InitMethod = EAnimNextScheduleInitMethod::InitializeAndPauseInEditor;
+
+	// Handle to the registered results/schedule
+	UE::AnimNext::FScheduleHandle SchedulerHandle;
+};
diff --git a/Engine/Plugins/Experimental/AnimNext/Source/AnimNext/Internal/Component/AnimNextComponentParameter.h b/Engine/Plugins/Experimental/AnimNext/Source/AnimNext/Internal/Component/AnimNextComponentParameter.h
new file mode 100644
--- /dev/null
+++ b/Engine/Plugins/Experimental/AnimNext/Source/AnimNext/Internal/Component/AnimNextComponentParameter.h
@@ -0,0 +1,77 @@
+﻿// Copyright Epic Games, Inc. All Rights Reserved.
+
+#pragma once
+
+#include "Param/ParamId.h"
+#include "Param/ParamType.h"
+#include "Param/ParamTypeHandle.h"
+#include "AnimNextComponentParameter.generated.h"
+
+class UAnimNextGraph;
+
+// Boilerplate macro that all derived types of UAnimNextComponentParameter should use.
+// Usage is: IMPLEMENT_ANIMNEXT_COMPONENT_PARAMETER(UAnimNextComponentParameter_DerivedType)
+#define IMPLEMENT_ANIMNEXT_COMPONENT_PARAMETER(Type) \
+	virtual void CacheParamInfo() const override \
+	{ \
+		using namespace UE::AnimNext; \
+		if(Name == NAME_None || ValueProperty == nullptr) \
+		{ \
+			Name = Parameter; \
+			ValueProperty = StaticClass()->FindPropertyByName(GET_MEMBER_NAME_CHECKED(Type, Value)); \
+		} \
+	} \
+
+// Base class for static parameters that can be inserted into a schedule via a component
+// Each parameter type needs its own derived type of this object
+UCLASS(DefaultToInstanced, abstract, editinlinenew)
+class UAnimNextComponentParameter : public UObject
+{
+	GENERATED_BODY()
+
+public:
+	// The scope to apply the parameter to. If this is "None", it is applied at the root scope of the schedule.
+	UPROPERTY(EditAnywhere, Category = "Parameter", AdvancedDisplay, meta = (CustomWidget = "ParamName", AllowedParamType = "FAnimNextScope", AllowNone))
+	FName Scope;
+
+	// Check validity
+	bool IsValid() const
+	{
+		CacheParamInfo();
+
+		return ValueProperty != nullptr && Name != NAME_None;
+	}
+	
+	// Get the name, type and value for this parameter
+	void GetParamInfo(FName& OutName, const FProperty*& OutProperty) const
+	{
+		CacheParamInfo();
+
+		OutName = Name;
+		OutProperty = ValueProperty;
+	}
+
+private:
+	virtual void CacheParamInfo() const PURE_VIRTUAL(UAnimNextComponentParameter::CacheParamInfo, )
+
+protected:
+	mutable FName Name;
+	mutable FProperty* ValueProperty = nullptr;
+};
+
+// An object parameter
+UCLASS(MinimalAPI, meta = (DisplayName = "Anim Next Graph"))
+class UAnimNextComponentParameter_AnimNextGraph : public UAnimNextComponentParameter
+{
+	GENERATED_BODY()
+
+	IMPLEMENT_ANIMNEXT_COMPONENT_PARAMETER(UAnimNextComponentParameter_AnimNextGraph);
+
+	// The parameter to set the value to
+	UPROPERTY(EditAnywhere, Category = "Parameter", meta = (CustomWidget = "ParamName", AllowedParamType = "TObjectPtr<UAnimNextGraph>"))
+	FName Parameter;
+
+	// The value to set
+	UPROPERTY(EditAnywhere, Category = "Parameter")
+	TObjectPtr<UAnimNextGraph> Value;
+};
diff --git a/Engine/Plugins/Experimental/AnimNext/Source/AnimNext/Internal/Graph/AnimNextGraph.h b/Engine/Plugins/Experimental/AnimNext/Source/AnimNext/Internal/Graph/AnimNextGraph.h
--- a/Engine/Plugins/Experimental/AnimNext/Source/AnimNext/Internal/Graph/AnimNextGraph.h
+++ b/Engine/Plugins/Experimental/AnimNext/Source/AnimNext/Internal/Graph/AnimNextGraph.h
@@ -21,6 +21,8 @@ class UAnimGraphNode_AnimNextGraph;
 struct FAnimNode_AnimNextGraph;
 struct FRigUnit_AnimNextGraphEvaluator;
 class UAnimNextSchedule;
+struct FAnimNextScheduleGraphTask;
+struct FAnimNextParam;
 
 namespace UE::AnimNext
 {
@@ -197,6 +199,7 @@ protected:
 	friend class UAnimGraphNode_AnimNextGraph;
 	friend UE::AnimNext::FExecutionContext;
 	friend class UAnimNextSchedule;
+	friend struct FAnimNextScheduleGraphTask;
 	
 #if WITH_EDITORONLY_DATA
 	mutable FRWLock GraphInstancesLock;
@@ -249,6 +252,14 @@ protected:
 	UE::AnimNext::FParamId ReferencePoseId = UE::AnimNext::FParamId(ReferencePose);
 	UE::AnimNext::FParamId CurrentLODId = UE::AnimNext::FParamId(CurrentLOD);
 
+	// Hash of required parameters
+	UPROPERTY()
+	uint64 RequiredParametersHash = 0;
+
+	// All the parameters that are required for this graph to run
+	UPROPERTY()
+	TArray<FAnimNextParam> RequiredParameters;
+
 #if WITH_EDITORONLY_DATA
 	UPROPERTY(VisibleAnywhere, Instanced, Category = "Graph", meta = (ShowInnerProperties))
 	TObjectPtr<UObject> EditorData;
diff --git a/Engine/Plugins/Experimental/AnimNext/Source/AnimNext/Internal/Param/AnimNextParam.h b/Engine/Plugins/Experimental/AnimNext/Source/AnimNext/Internal/Param/AnimNextParam.h
new file mode 100644
--- /dev/null
+++ b/Engine/Plugins/Experimental/AnimNext/Source/AnimNext/Internal/Param/AnimNextParam.h
@@ -0,0 +1,26 @@
+﻿// Copyright Epic Games, Inc. All Rights Reserved.
+
+#pragma once
+
+#include "CoreMinimal.h"
+#include "Param/ParamType.h"
+#include "AnimNextParam.generated.h"
+
+USTRUCT()
+struct FAnimNextParam
+{
+	GENERATED_BODY()
+
+	FAnimNextParam() = default;
+
+	FAnimNextParam(FName InName, const FAnimNextParamType& InType)
+		: Name(InName)
+		, Type(InType)
+	{}
+
+	UPROPERTY(EditAnywhere, Category = "Parameter")
+	FName Name;
+
+	UPROPERTY(EditAnywhere, Category = "Parameter")
+	FAnimNextParamType Type;
+};
\ No newline at end of file
diff --git a/Engine/Plugins/Experimental/AnimNext/Source/AnimNext/Internal/Scheduler/AnimNextExternalTaskBinding.h b/Engine/Plugins/Experimental/AnimNext/Source/AnimNext/Internal/Scheduler/AnimNextExternalTaskBinding.h
new file mode 100644
--- /dev/null
+++ b/Engine/Plugins/Experimental/AnimNext/Source/AnimNext/Internal/Scheduler/AnimNextExternalTaskBinding.h
@@ -0,0 +1,13 @@
+﻿// Copyright Epic Games, Inc. All Rights Reserved.
+
+#pragma once
+
+#include "CoreMinimal.h"
+#include "AnimNextExternalTaskBinding.generated.h"
+
+// Transient base structure used to communicate binding information
+USTRUCT()
+struct FAnimNextExternalTaskBinding
+{
+	GENERATED_BODY()
+};
diff --git a/Engine/Plugins/Experimental/AnimNext/Source/AnimNext/Internal/Scheduler/AnimNextSchedule.h b/Engine/Plugins/Experimental/AnimNext/Source/AnimNext/Internal/Scheduler/AnimNextSchedule.h
--- a/Engine/Plugins/Experimental/AnimNext/Source/AnimNext/Internal/Scheduler/AnimNextSchedule.h
+++ b/Engine/Plugins/Experimental/AnimNext/Source/AnimNext/Internal/Scheduler/AnimNextSchedule.h
@@ -18,6 +18,7 @@ class UAnimNextSchedule;
 class UAnimNextSchedulerWorldSubsystem;
 class UAnimNextComponent;
 struct FAnimNextSchedulerEntry;
+struct FAnimNextParam;
 
 namespace UE::AnimNext
 {
@@ -86,7 +87,7 @@ private:
 
 	// All parameters that are required by this graph to run (only required if dynamic as static graph params can be discovered by the compiler)
 	UPROPERTY(EditAnywhere, Category = "Graph", meta = (CustomWidget = "ParamName"))
-	TArray<FName> RequiredParameters;
+	TArray<FAnimNextParam> RequiredParameters;
 
 	// The intermediate terms used by the graph
 	UPROPERTY(EditAnywhere, Category = "Graph")
@@ -235,6 +236,7 @@ private:
 #if WITH_EDITOR
 	virtual void PostEditChangeProperty(FPropertyChangedEvent& PropertyChangedEvent) override;
 	virtual void PostEditUndo() override;
+	virtual void GetAssetRegistryTags(TArray<FAssetRegistryTag>& OutTags) const override;
 
 	// Compile the editor data into a compact runtime representation
 	void CompileSchedule();
@@ -244,6 +246,9 @@ private:
 	// Function hook used to compile in editor/cooker
 	static TUniqueFunction<void(UAnimNextSchedule*)> CompileFunction;
 
+	// Function hook used to expose to asset registry
+	static TUniqueFunction<void(const UAnimNextSchedule*, TArray<FAssetRegistryTag>&)> GetAssetRegistryTagsFunction;
+	
 	// Editor only
 	// TODO: move this into an editor only subobject
 	// TODO: this is currently only a linear list, we want it to be a graph
diff --git a/Engine/Plugins/Experimental/AnimNext/Source/AnimNext/Internal/Scheduler/AnimNextScheduleGraphTask.h b/Engine/Plugins/Experimental/AnimNext/Source/AnimNext/Internal/Scheduler/AnimNextScheduleGraphTask.h
--- a/Engine/Plugins/Experimental/AnimNext/Source/AnimNext/Internal/Scheduler/AnimNextScheduleGraphTask.h
+++ b/Engine/Plugins/Experimental/AnimNext/Source/AnimNext/Internal/Scheduler/AnimNextScheduleGraphTask.h
@@ -8,6 +8,7 @@
 
 class UAnimNextGraph;
 class UAnimNextParameterBlock;
+struct FAnimNextParam;
 
 namespace UE::AnimNext
 {
@@ -39,6 +40,9 @@ private:
 
 	UAnimNextGraph* GetGraphToRun(UE::AnimNext::FParamStack& ParamStack) const;
 
+	// Verify graph's required parameters are satisfied by this task's supplied parameters
+	void VerifyRequiredParameters(UAnimNextGraph* InGraphToRun) const;
+
 private:
 	UPROPERTY()
 	uint32 TaskIndex = MAX_uint32;
@@ -61,4 +65,11 @@ private:
 	// Index of each term in the schedule intermediates
 	UPROPERTY()
 	TArray<uint32> Terms;
+
+	UPROPERTY()
+	uint64 SuppliedParametersHash = 0;
+
+	// All supplied parameters for any dynamic graphs slotted here
+	UPROPERTY()
+	TArray<FAnimNextParam> SuppliedParameters;
 };
diff --git a/Engine/Plugins/Experimental/AnimNext/Source/AnimNext/Internal/Scheduler/ScheduleHandle.h b/Engine/Plugins/Experimental/AnimNext/Source/AnimNext/Internal/Scheduler/ScheduleHandle.h
new file mode 100644
--- /dev/null
+++ b/Engine/Plugins/Experimental/AnimNext/Source/AnimNext/Internal/Scheduler/ScheduleHandle.h
@@ -0,0 +1,50 @@
+// Copyright Epic Games, Inc. All Rights Reserved.
+
+#pragma once
+
+#include "CoreMinimal.h"
+
+class UAnimNextSchedulerWorldSubsystem;
+
+namespace UE::AnimNext
+{
+	struct FScheduler;
+	struct FSchedulerImpl;
+}
+
+namespace UE::AnimNext
+{
+
+// Opaque handle representing a binding between a schedule and parameters
+struct FScheduleHandle
+{
+	FScheduleHandle() = default;
+
+	bool IsValid() const
+	{
+		return Index != MAX_uint32 && SerialNumber != 0;
+	}
+
+	void Invalidate()
+	{
+		Index = MAX_uint32;
+		SerialNumber = 0;
+	}
+
+	FString ToString() const
+	{
+		return FString::Printf(TEXT("(%d, %d)"), Index, SerialNumber);
+	}
+
+private:
+	friend struct FScheduler;
+	friend struct FSchedulerImpl;
+	friend struct FSchedulerEntry;
+	friend struct FScheduleContext;
+	friend class ::UAnimNextSchedulerWorldSubsystem;
+
+	uint32 Index = MAX_uint32;
+	uint32 SerialNumber = 0;
+};
+
+}
\ No newline at end of file
diff --git a/Engine/Plugins/Experimental/AnimNext/Source/AnimNext/Internal/Scheduler/Scheduler.h b/Engine/Plugins/Experimental/AnimNext/Source/AnimNext/Internal/Scheduler/Scheduler.h
new file mode 100644
--- /dev/null
+++ b/Engine/Plugins/Experimental/AnimNext/Source/AnimNext/Internal/Scheduler/Scheduler.h
@@ -0,0 +1,77 @@
+// Copyright Epic Games, Inc. All Rights Reserved.
+
+#pragma once
+
+#include "CoreMinimal.h"
+#include "Scheduler/ScheduleHandle.h"
+
+enum class EAnimNextScheduleInitMethod : uint8;
+class UAnimNextGraph;
+class UAnimNextParameterBlock;
+class UAnimNextSchedule;
+
+namespace UE::AnimNext
+{
+	struct FScheduler;
+	struct FSchedulerImpl;
+	struct FSchedulePortDefinition;
+	struct FParamStackLayerHandle;
+	struct FScheduleContext;
+}
+
+namespace UE::AnimNext::UncookedOnly
+{
+	struct FUtils;
+}
+
+namespace UE::AnimNext
+{
+
+// Main interface into the AnimNext scheduling system
+struct FScheduler
+{
+	// Start up the scheduler system
+	static void Init();
+
+	// Shut down the scheduler system
+	static void Destroy();
+
+	// Acquire a handle that binds a schedule
+	// Initial parameter binding can be achieved via InitializeCallback, which is called once the schedule's data
+	// structures have been set up, but before it is first run.
+	static FScheduleHandle AcquireHandle(UObject* InObject, UAnimNextSchedule* InSchedule, EAnimNextScheduleInitMethod InInitMethod, TUniqueFunction<void(const FScheduleContext&)>&& InitializeCallback = nullptr);
+
+	// Release an already acquired handle
+	// The full release of the binding referenced by the handle map be deferred after this call is made
+	static void ReleaseHandle(UObject* InObject, FScheduleHandle& InHandle);
+
+	// Enables or disables the schedule parameterization represented by the supplied handle
+	// This operation is deferred until the next time the schedule ticks
+	static void EnableHandle(UObject* InObject, FScheduleHandle InHandle, bool bInEnabled);
+
+	enum class ETaskRunLocation : int32
+	{
+		// Run the task before the specified task
+		Before,
+
+		// Run the task after the specified task
+		After,
+	};
+
+	// Queue a task to run at a particular point in a schedule
+	// @param	InHandle		The handle to queue the task to
+	// @param	InTaskName		The name of the task in the schedule to run the supplied task relative to
+	// @param	InTaskFunction	The function to run
+	// @param	InLocation		Where to run the task, before or after
+	static void QueueTask(UObject* InObject, FScheduleHandle InHandle, FName InScheduleTaskName, TUniqueFunction<void(const FScheduleContext&)>&& InTaskFunction, ETaskRunLocation InLocation = ETaskRunLocation::Before);
+
+private:
+	friend struct UE::AnimNext::UncookedOnly::FUtils;
+
+#if WITH_EDITOR
+	// Refresh any entries that use the provided schedule as it has been recompiled.
+	static ANIMNEXT_API void OnScheduleCompiled(UAnimNextSchedule* InSchedule);
+#endif
+};
+
+}
\ No newline at end of file
diff --git a/Engine/Plugins/Experimental/AnimNext/Source/AnimNext/Private/Component/AnimNextComponent.cpp b/Engine/Plugins/Experimental/AnimNext/Source/AnimNext/Private/Component/AnimNextComponent.cpp
--- a/Engine/Plugins/Experimental/AnimNext/Source/AnimNext/Private/Component/AnimNextComponent.cpp
+++ b/Engine/Plugins/Experimental/AnimNext/Source/AnimNext/Private/Component/AnimNextComponent.cpp
@@ -1,10 +1,58 @@
 // Copyright Epic Games, Inc. All Rights Reserved.
 
-#include "AnimNextComponent.h"
+#include "Component/AnimNextComponent.h"
+
+#include "Param/ParamUtils.h"
 #include "Scheduler/Scheduler.h"
 #include "Scheduler/ScheduleContext.h"
 #include "Param/PropertyBagProxy.h"
 
+namespace UE::AnimNext::Private
+{
+
+static void SetValuesInScopeHelper(FScheduleInstanceData& InInstanceData, FName InScope, EAnimNextParameterScopeOrdering InOrdering, TConstArrayView<FPropertyBagProxy::FPropertyAndValue> InPropertiesAndValues)
+{
+	// apply to root
+	if(InScope == NAME_None)
+	{
+		if(!InInstanceData.RootUserScope.IsValid())
+		{
+			InInstanceData.RootUserScope = MakeUnique<FPropertyBagProxy>();
+		}
+
+		InInstanceData.RootUserScope->AddPropertiesAndValues(InPropertiesAndValues);
+	}
+	else // apply to specified scope
+	{
+		FScheduleInstanceData::FUserScope& ScopeSource = InInstanceData.UserScopes.FindOrAdd(InScope);
+		TUniquePtr<FPropertyBagProxy>* ProxyToUse = nullptr;
+		switch(InOrdering)
+		{
+		default:
+		case EAnimNextParameterScopeOrdering::Before:
+			if(!ScopeSource.BeforeSource.IsValid())
+			{
+				ScopeSource.BeforeSource = MakeUnique<FPropertyBagProxy>();
+			}
+			ProxyToUse = &ScopeSource.BeforeSource;
+			break;
+		case EAnimNextParameterScopeOrdering::After:
+			if(!ScopeSource.AfterSource.IsValid())
+			{
+				ScopeSource.AfterSource = MakeUnique<FPropertyBagProxy>();
+			}
+			ProxyToUse = &ScopeSource.AfterSource;
+			break;
+		}
+
+		check(ProxyToUse && ProxyToUse->IsValid());
+
+		(*ProxyToUse)->AddPropertiesAndValues(InPropertiesAndValues);
+	}
+}
+
+}
+
 void UAnimNextComponent::OnRegister()
 {
 	using namespace UE::AnimNext;
@@ -28,50 +76,20 @@ void UAnimNextComponent::OnRegister()
 				}
 			}
 
+			// Now apply to each scope
 			for(const TPair<FName, TArray<UAnimNextComponentParameter*, TInlineAllocator<4>>>& ParamPair : ParamsByScope)
 			{
-				FName Scope = ParamPair.Key;
-				TSharedPtr<FParamStack> StackToUse;
-
-				if(Scope == NAME_None)
-				{
-					StackToUse = InstanceData.RootParamStack;
-				}
-				else
+				TArray<FPropertyBagProxy::FPropertyAndValue, TInlineAllocator<16>> PropertiesAndValues;
+				PropertiesAndValues.Reserve(ParamPair.Value.Num());
+				for(UAnimNextComponentParameter* Parameter : ParamPair.Value)
 				{
-					const FAnimNextScheduleParamScopeEntryTask* FoundTask = InContext.Schedule->ParamScopeEntryTasks.FindByPredicate([Scope](const FAnimNextScheduleParamScopeEntryTask& InTask)
-					{
-						return InTask.Scope == Scope;
-					});
-
-					if(FoundTask)
-					{
-						StackToUse = InstanceData.ParamStacks[FoundTask->ParamScopeIndex];
-					}
+					FPropertyBagProxy::FPropertyAndValue& PropertyAndValue = PropertiesAndValues.AddDefaulted_GetRef();
+					PropertyAndValue.ContainerPtr = Parameter;
+					Parameter->GetParamInfo(PropertyAndValue.Name, PropertyAndValue.Property);
 				}
 
-				if(StackToUse.IsValid())
-				{
-					TArray<Private::FParamEntry, TInlineAllocator<4>> Params;
-					for(UAnimNextComponentParameter* Parameter : ParamPair.Value)
-					{
-						FParamId ParamId;
-						FAnimNextParamType Type;
-						uint8* Value = nullptr;
-						Parameter->GetParamInfo(ParamId, Type, Value);
-						check(Type.IsValid() && Value != nullptr);
-
-						constexpr bool bIsReference = true;
-						constexpr bool bIsMutable = false;
-						Params.Emplace(Private::FParamEntry(ParamId, Type.GetHandle(), TArrayView<uint8>(Value, 1), bIsReference, bIsMutable));
-					}
-
-					FParamStackLayerHandle NewLayer = FParamStack::MakeLayer(Params);
-					InstanceData.StaticUserHandles.Add(MoveTemp(NewLayer));
-
-					// Layer is never popped as this is happening at the very start of execution
-					StackToUse->PushLayer(InstanceData.StaticUserHandles.Last());
-				}
+				// NOTE: Layer is always applied 'before' currently. If we have a use case for 'After' we can add it to UAnimNextComponentParameter
+				Private::SetValuesInScopeHelper(InContext.GetInstanceData(), ParamPair.Key, EAnimNextParameterScopeOrdering::Before, PropertiesAndValues);
 			}
 		};
 
@@ -90,7 +108,7 @@ void UAnimNextComponent::OnUnregister()
 	SchedulerHandle.Invalidate();
 }
 
-void UAnimNextComponent::SetParameterInScope(FName Scope, FName Name, int32 Value)
+void UAnimNextComponent::SetParameterInScope(FName Scope, EAnimNextParameterScopeOrdering Ordering, FName Name, int32 Value)
 {
 	checkNoEntry();
 }
@@ -103,8 +121,9 @@ DEFINE_FUNCTION(UAnimNextComponent::execSetParameterInScope)
 	Stack.MostRecentPropertyAddress = nullptr;
 	Stack.MostRecentPropertyContainer = nullptr;
 
-	PARAM_PASSED_BY_VAL(Scope, FNameProperty, FName);
-	PARAM_PASSED_BY_VAL(Name, FNameProperty, FName);
+	P_GET_PROPERTY(FNameProperty, Scope);
+	P_GET_ENUM(EAnimNextParameterScopeOrdering, Ordering);
+	P_GET_PROPERTY(FNameProperty, Name);
 
 	Stack.StepCompiledIn<FProperty>(nullptr);
 	const FProperty* ValueProp = CastField<FProperty>(Stack.MostRecentProperty);
@@ -120,37 +139,38 @@ DEFINE_FUNCTION(UAnimNextComponent::execSetParameterInScope)
 		);
 
 		FBlueprintCoreDelegates::ThrowScriptException(P_THIS, Stack, ExceptionInfo);
+		return;
 	}
-	else if (Scope == NAME_None || Name == NAME_None)
+
+	if (Name == NAME_None)
 	{
 		FBlueprintExceptionInfo ExceptionInfo(
 			EBlueprintExceptionType::NonFatalError,
-			NSLOCTEXT("AnimNextComponent", "AnimNextComponent_SetParameterInScopeWarning", "Invalid scope or parameter name supplied to Set Parameter In Scope")
+			NSLOCTEXT("AnimNextComponent", "AnimNextComponent_SetParameterInScopeWarning", "Invalid parameter name supplied to Set Parameter In Scope")
 		);
 
 		FBlueprintCoreDelegates::ThrowScriptException(P_THIS, Stack, ExceptionInfo);
+		return;
 	}
-	else
-	{
-		P_NATIVE_BEGIN;
-
-		TUniquePtr<FInstancedPropertyBag> PropertyBag = MakeUnique<FInstancedPropertyBag>();
-		PropertyBag->AddProperty(Name, ValueProp);
-		FProperty* NewProperty = PropertyBag->GetPropertyBagStruct()->FindPropertyByName(Name);
-		PropertyBag->SetValue(Name, NewProperty, ContainerPtr);
-		
-		FScheduler::QueueTask(P_THIS, P_THIS->SchedulerHandle, Scope, [Scope, Name, NewProperty, PropertyBag = MoveTemp(PropertyBag)](const FScheduleContext& InContext) mutable
-		{
-			TUniquePtr<FPropertyBagProxy>& ScopeSources = InContext.GetInstanceData().UserScopes.FindOrAdd(Scope);
 
-			// TODO: pre/post scope distinction
+	P_NATIVE_BEGIN;
 
-			// Copy property to local bag
-			ScopeSources->AddPropertyAndValue(Name, NewProperty, PropertyBag->GetValue().GetMemory());
-		});
-		
-		P_NATIVE_END;
-	}
+	TUniquePtr<FInstancedPropertyBag> PropertyBag = MakeUnique<FInstancedPropertyBag>();
+	PropertyBag->AddProperty(Name, ValueProp);
+	const FProperty* NewProperty = PropertyBag->GetPropertyBagStruct()->GetPropertyDescs()[0].CachedProperty;
+	const void* ValuePtr = ValueProp->ContainerPtrToValuePtr<void>(ContainerPtr);
+	NewProperty->SetValue_InContainer(PropertyBag->GetMutableValue().GetMemory(), ValuePtr);
+
+	FScheduler::QueueTask(P_THIS, P_THIS->SchedulerHandle, Scope, [Scope, Name, NewProperty, PropertyBag = MoveTemp(PropertyBag), Ordering](const FScheduleContext& InContext) mutable
+	{
+		FPropertyBagProxy::FPropertyAndValue PropertyAndValue;
+		PropertyAndValue.Name = Name;
+		PropertyAndValue.Property = NewProperty;
+		PropertyAndValue.ContainerPtr = PropertyBag->GetValue().GetMemory();
+		Private::SetValuesInScopeHelper(InContext.GetInstanceData(), Scope, Ordering, { PropertyAndValue });
+	});
+
+	P_NATIVE_END;
 }
 
 void UAnimNextComponent::Enable(bool bEnabled)
diff --git a/Engine/Plugins/Experimental/AnimNext/Source/AnimNext/Private/Component/AnimNextComponent.h b/Engine/Plugins/Experimental/AnimNext/Source/AnimNext/Private/Component/AnimNextComponent.h
deleted file mode 100644
diff --git a/Engine/Plugins/Experimental/AnimNext/Source/AnimNext/Private/Component/AnimNextComponentParameter.h b/Engine/Plugins/Experimental/AnimNext/Source/AnimNext/Private/Component/AnimNextComponentParameter.h
deleted file mode 100644
diff --git a/Engine/Plugins/Experimental/AnimNext/Source/AnimNext/Private/Graph/AnimGraph/AnimNode_AnimNextGraph.cpp b/Engine/Plugins/Experimental/AnimNext/Source/AnimNext/Private/Graph/AnimGraph/AnimNode_AnimNextGraph.cpp
--- a/Engine/Plugins/Experimental/AnimNext/Source/AnimNext/Private/Graph/AnimGraph/AnimNode_AnimNextGraph.cpp
+++ b/Engine/Plugins/Experimental/AnimNext/Source/AnimNext/Private/Graph/AnimGraph/AnimNode_AnimNextGraph.cpp
@@ -142,9 +142,6 @@ void FAnimNode_AnimNextGraph::Evaluate_AnyThread(FPoseContext& Output)
 
 		FAnimNextGraphLODPose ResultPose(FLODPoseHeap(RefPose, LODLevel, true, Output.ExpectsAdditivePose()));
 
-		FScheduleContext ScheduleContext(SkeletalMeshComponent);
-		FScheduleContext::AttachToCurrentThread(ScheduleContext);
-
 		FAnimGraphParamStackScope Scope(Output);
 		FParamStack& ParamStack = FParamStack::Get();
 		FParamStack::FPushedLayerHandle LayerHandle = ParamStack.PushValues(
@@ -181,8 +178,6 @@ void FAnimNode_AnimNextGraph::Evaluate_AnyThread(FPoseContext& Output)
 		FGenerationTools::RemapPose(ResultPose.LODPose, Output);
 
 		ParamStack.PopLayer(LayerHandle);
-
-		FScheduleContext::DetachFromCurrentThread();
 	}
 	else
 	{
diff --git a/Engine/Plugins/Experimental/AnimNext/Source/AnimNext/Private/Param/ParamType.cpp b/Engine/Plugins/Experimental/AnimNext/Source/AnimNext/Private/Param/ParamType.cpp
--- a/Engine/Plugins/Experimental/AnimNext/Source/AnimNext/Private/Param/ParamType.cpp
+++ b/Engine/Plugins/Experimental/AnimNext/Source/AnimNext/Private/Param/ParamType.cpp
@@ -464,9 +464,9 @@ size_t FAnimNextParamType::GetValueTypeAlignment() const
 	return 0;
 }
 
-FString FAnimNextParamType::ToString() const
+void FAnimNextParamType::ToString(FStringBuilderBase& InStringBuilder) const
 {
-	auto GetTypeString = [this](TStringBuilder<128>& InStringBuilder)
+	auto GetTypeString = [this, &InStringBuilder]()
 	{
 		switch(ValueType)
 		{
@@ -576,25 +576,28 @@ FString FAnimNextParamType::ToString() const
 		}
 	};
 
-	TStringBuilder<128> StringBuilder;
-
 	switch(ContainerType)
 	{
 	case EContainerType::None:
-		GetTypeString(StringBuilder);
+		GetTypeString();
 		break;
 	case EContainerType::Array:
 		{
-			StringBuilder.Append(TEXT("TArray<"));
-			GetTypeString(StringBuilder);
-			StringBuilder.Append(TEXT(">"));
+			InStringBuilder.Append(TEXT("TArray<"));
+			GetTypeString();
+			InStringBuilder.Append(TEXT(">"));
 		}
 		break;
 	default:
-		StringBuilder.Append(TEXT("Error: Unknown container type"));
+		InStringBuilder.Append(TEXT("Error: Unknown container type"));
 		break;
 	}
+}
 
+FString FAnimNextParamType::ToString() const
+{
+	TStringBuilder<128> StringBuilder;
+	ToString(StringBuilder);
 	return StringBuilder.ToString();
 }
 
diff --git a/Engine/Plugins/Experimental/AnimNext/Source/AnimNext/Private/Param/PropertyBagProxy.cpp b/Engine/Plugins/Experimental/AnimNext/Source/AnimNext/Private/Param/PropertyBagProxy.cpp
--- a/Engine/Plugins/Experimental/AnimNext/Source/AnimNext/Private/Param/PropertyBagProxy.cpp
+++ b/Engine/Plugins/Experimental/AnimNext/Source/AnimNext/Private/Param/PropertyBagProxy.cpp
@@ -20,6 +20,32 @@ void FPropertyBagProxy::AddPropertyAndValue(FName InName, const FProperty* InPro
 {
 	PropertyBag.AddProperty(InName, InProperty);
 	PropertyBag.SetValue(InName, InProperty, InContainerPtr);
+
+	// Recreate the layer handle as the bag layout has changed
+	LayerHandle = FParamStack::MakeReferenceLayer(PropertyBag);
+}
+
+void FPropertyBagProxy::AddPropertiesAndValues(TConstArrayView<FPropertyAndValue> InPropertiesAndValues)
+{
+	TArray<FPropertyBagPropertyDesc, TInlineAllocator<16>> NewDescs;
+	NewDescs.Reserve(InPropertiesAndValues.Num());
+
+	// Add all new properties
+	for(const FPropertyAndValue& PropertyAndValue : InPropertiesAndValues)
+	{
+		NewDescs.Emplace(PropertyAndValue.Name, PropertyAndValue.Property);
+	}
+
+	PropertyBag.AddProperties(NewDescs);
+
+	// Set each value - note that SetValue uses a linear search internally so this is quite slow
+	for(const FPropertyAndValue& PropertyAndValue : InPropertiesAndValues)
+	{
+		PropertyBag.SetValue(PropertyAndValue.Name, PropertyAndValue.Property, PropertyAndValue.ContainerPtr);
+	}
+
+	// Recreate the layer handle as the bag layout has changed
+	LayerHandle = FParamStack::MakeReferenceLayer(PropertyBag);
 }
 
 }
diff --git a/Engine/Plugins/Experimental/AnimNext/Source/AnimNext/Private/Param/PropertyBagProxy.h b/Engine/Plugins/Experimental/AnimNext/Source/AnimNext/Private/Param/PropertyBagProxy.h
--- a/Engine/Plugins/Experimental/AnimNext/Source/AnimNext/Private/Param/PropertyBagProxy.h
+++ b/Engine/Plugins/Experimental/AnimNext/Source/AnimNext/Private/Param/PropertyBagProxy.h
@@ -21,6 +21,17 @@ struct FPropertyBagProxy : public IParameterSource
 	// Adds (or replaces) a property with the specified name and sets its value using the specified container
 	void AddPropertyAndValue(FName InName, const FProperty* InProperty, const void* InContainerPtr);
 
+	// Arg to AddPropertiesAndValues
+	struct FPropertyAndValue
+	{
+		FName Name;
+		const FProperty* Property = nullptr;
+		const void* ContainerPtr = nullptr;
+	};
+
+	// Adds (or replaces) a group of properties with the specified names and sets their values using the specified containers
+	void AddPropertiesAndValues(TConstArrayView<FPropertyAndValue> InPropertiesAndValues);
+
 private:
 	// Property bag that this wraps
 	FInstancedPropertyBag PropertyBag;
diff --git a/Engine/Plugins/Experimental/AnimNext/Source/AnimNext/Private/Scheduler/AnimNextExternalTaskBinding.h b/Engine/Plugins/Experimental/AnimNext/Source/AnimNext/Private/Scheduler/AnimNextExternalTaskBinding.h
deleted file mode 100644
diff --git a/Engine/Plugins/Experimental/AnimNext/Source/AnimNext/Private/Scheduler/AnimNextSchedule.cpp b/Engine/Plugins/Experimental/AnimNext/Source/AnimNext/Private/Scheduler/AnimNextSchedule.cpp
--- a/Engine/Plugins/Experimental/AnimNext/Source/AnimNext/Private/Scheduler/AnimNextSchedule.cpp
+++ b/Engine/Plugins/Experimental/AnimNext/Source/AnimNext/Private/Scheduler/AnimNextSchedule.cpp
@@ -7,6 +7,7 @@
 
 #if WITH_EDITOR
 TUniqueFunction<void(UAnimNextSchedule*)> UAnimNextSchedule::CompileFunction;
+TUniqueFunction<void(const UAnimNextSchedule*, TArray<UObject::FAssetRegistryTag>&)> UAnimNextSchedule::GetAssetRegistryTagsFunction;
 #endif
 
 void UAnimNextSchedule::PostLoad()
@@ -41,4 +42,13 @@ void UAnimNextSchedule::CompileSchedule()
 	CompileFunction(this);
 }
 
+void UAnimNextSchedule::GetAssetRegistryTags(TArray<FAssetRegistryTag>& OutTags) const
+{
+	Super::GetAssetRegistryTags(OutTags);
+
+	check(GetAssetRegistryTagsFunction);
+
+	GetAssetRegistryTagsFunction(this, OutTags);
+}
+
 #endif // #if WITH_EDITOR
\ No newline at end of file
diff --git a/Engine/Plugins/Experimental/AnimNext/Source/AnimNext/Private/Scheduler/AnimNextScheduleGraphTask.cpp b/Engine/Plugins/Experimental/AnimNext/Source/AnimNext/Private/Scheduler/AnimNextScheduleGraphTask.cpp
--- a/Engine/Plugins/Experimental/AnimNext/Source/AnimNext/Private/Scheduler/AnimNextScheduleGraphTask.cpp
+++ b/Engine/Plugins/Experimental/AnimNext/Source/AnimNext/Private/Scheduler/AnimNextScheduleGraphTask.cpp
@@ -11,6 +11,8 @@
 #include "Graph/AnimNextExecuteContext.h"
 #include "Graph/AnimNext_LODPose.h"
 #include "AnimNextStats.h"
+#include "Logging/StructuredLog.h"
+#include "Param/AnimNextParam.h"
 
 DEFINE_STAT(STAT_AnimNext_Task_Graph);
 
@@ -28,6 +30,49 @@ UAnimNextGraph* FAnimNextScheduleGraphTask::GetGraphToRun(UE::AnimNext::FParamSt
 	return GraphToRun;
 }
 
+void FAnimNextScheduleGraphTask::VerifyRequiredParameters(UAnimNextGraph* InGraphToRun) const
+{
+	if(SuppliedParametersHash != InGraphToRun->RequiredParametersHash)
+	{
+		bool bWarningOutput = false;
+
+		for(const FAnimNextParam& RequiredParameter : InGraphToRun->RequiredParameters)
+		{
+			bool bFound = false;
+			bool bFoundCorrectType = true;
+			FAnimNextParamType SuppliedParameterType;
+			for(const FAnimNextParam& SuppliedParameter : SuppliedParameters)
+			{
+				if(RequiredParameter.Name == SuppliedParameter.Name)
+				{
+					if(RequiredParameter.Type != SuppliedParameter.Type)
+					{
+						SuppliedParameterType = SuppliedParameter.Type;
+						bFoundCorrectType = false;
+					}
+					bFound = true;
+					break;
+				}
+			}
+
+			if(!bWarningOutput && (!bFound || !bFoundCorrectType))
+			{
+				UE_LOGFMT(LogAnimation, Warning, "AnimNext: Graph {GraphToRun} has different required parameters, it may not run correctly.", InGraphToRun->GetName());
+				bWarningOutput = true;
+			}
+			
+			if(!bFound)
+			{
+				UE_LOGFMT(LogAnimation, Warning, "    Not Found: {Name}", RequiredParameter.Name);
+			}
+			else if(!bFoundCorrectType)
+			{
+				UE_LOGFMT(LogAnimation, Warning, "    Incorrect Type: {Name} ({RequiredType} vs {SuppliedType})", RequiredParameter.Name, RequiredParameter.Type.ToString(), SuppliedParameterType.ToString());
+			}
+		}
+	}
+}
+
 void FAnimNextScheduleGraphTask::RunGraph(const UE::AnimNext::FScheduleContext& InContext) const
 {
 	SCOPE_CYCLE_COUNTER(STAT_AnimNext_Task_Graph);
@@ -56,6 +101,12 @@ void FAnimNextScheduleGraphTask::RunGraph(const UE::AnimNext::FScheduleContext&
 	if (!GraphCache.GraphInstanceData.IsValid())
 	{
 		GraphToRun->AllocateInstance(GraphCache.GraphInstanceData);
+
+		// Only do dynamic verification for dynamic graphs. Static graphs get verified at compile time. 
+		if (Graph == nullptr && DynamicGraph != NAME_None)
+		{
+			VerifyRequiredParameters(GraphToRun);
+		}
 	}
 
 	const FAnimNextGraphReferencePose* GraphReferencePose = ParamStack.GetParamPtr<FAnimNextGraphReferencePose>(GraphToRun->GetReferencePoseParam());
diff --git a/Engine/Plugins/Experimental/AnimNext/Source/AnimNext/Private/Scheduler/AnimNextScheduleParamScopeTask.cpp b/Engine/Plugins/Experimental/AnimNext/Source/AnimNext/Private/Scheduler/AnimNextScheduleParamScopeTask.cpp
--- a/Engine/Plugins/Experimental/AnimNext/Source/AnimNext/Private/Scheduler/AnimNextScheduleParamScopeTask.cpp
+++ b/Engine/Plugins/Experimental/AnimNext/Source/AnimNext/Private/Scheduler/AnimNextScheduleParamScopeTask.cpp
@@ -18,19 +18,18 @@ void FAnimNextScheduleParamScopeEntryTask::RunParamScopeEntry(const UE::AnimNext
 	FParamStack& ParamStack = FParamStack::Get();
 
 	FScheduleInstanceData& InstanceData = InScheduleContext.GetInstanceData();
-	FScheduleInstanceData::FScopeCache& ScopeCache = InstanceData.ScopeCaches[TaskIndex];
+	FScheduleInstanceData::FScopeCache& ScopeCache = InstanceData.ScopeCaches[ParamScopeIndex];
 	check(ScopeCache.PushedLayers.Num() == 0);
 
-	// Update & push user params
-	if (TUniquePtr<FPropertyBagProxy>* FoundUserScope = InstanceData.UserScopes.Find(Scope))
+	// Update & push user params (before scope)
+	FScheduleInstanceData::FUserScope* FoundUserScope = InstanceData.UserScopes.Find(Scope);
+	if (FoundUserScope && FoundUserScope->BeforeSource.IsValid())
 	{
-		FPropertyBagProxy& UserParameterSource = *FoundUserScope->Get();
+		FPropertyBagProxy& UserParameterSource = *FoundUserScope->BeforeSource.Get();
 		UserParameterSource.Update();
 		ScopeCache.PushedLayers.Add(ParamStack.PushLayer(UserParameterSource.GetLayerHandle()));
 	}
 
-	// TODO: Pre/post scope support
-
 	// Update & push static params
 	for (int32 ParameterSourceIndex = 0; ParameterSourceIndex < ScopeCache.ParameterSources.Num(); ++ParameterSourceIndex)
 	{
@@ -39,6 +38,14 @@ void FAnimNextScheduleParamScopeEntryTask::RunParamScopeEntry(const UE::AnimNext
 		StaticParameterSource->Update();
 		ScopeCache.PushedLayers.Add(ParamStack.PushLayer(StaticParameterSource->GetLayerHandle()));
 	}
+
+	// Update & push user params (after scope)
+	if (FoundUserScope && FoundUserScope->AfterSource.IsValid())
+	{
+		FPropertyBagProxy& UserParameterSource = *FoundUserScope->AfterSource.Get();
+		UserParameterSource.Update();
+		ScopeCache.PushedLayers.Add(ParamStack.PushLayer(UserParameterSource.GetLayerHandle()));
+	}
 }
 
 void FAnimNextScheduleParamScopeExitTask::RunParamScopeExit(const UE::AnimNext::FScheduleContext& InScheduleContext) const
@@ -50,7 +57,7 @@ void FAnimNextScheduleParamScopeExitTask::RunParamScopeExit(const UE::AnimNext::
 	FParamStack& ParamStack = FParamStack::Get();
 
 	FScheduleInstanceData& InstanceData = InScheduleContext.GetInstanceData();
-	FScheduleInstanceData::FScopeCache& ScopeCache = InstanceData.ScopeCaches[TaskIndex];
+	FScheduleInstanceData::FScopeCache& ScopeCache = InstanceData.ScopeCaches[ParamScopeIndex];
 	for (int32 LayerIndex = ScopeCache.PushedLayers.Num() - 1; LayerIndex >= 0; --LayerIndex)
 	{
 		ParamStack.PopLayer(ScopeCache.PushedLayers[LayerIndex]);
diff --git a/Engine/Plugins/Experimental/AnimNext/Source/AnimNext/Private/Scheduler/AnimNextSchedulerEntry.cpp b/Engine/Plugins/Experimental/AnimNext/Source/AnimNext/Private/Scheduler/AnimNextSchedulerEntry.cpp
--- a/Engine/Plugins/Experimental/AnimNext/Source/AnimNext/Private/Scheduler/AnimNextSchedulerEntry.cpp
+++ b/Engine/Plugins/Experimental/AnimNext/Source/AnimNext/Private/Scheduler/AnimNextSchedulerEntry.cpp
@@ -8,14 +8,16 @@
 #include "AnimNextStats.h"
 #include "AnimNextTickFunctionBinding.h"
 #include "Logging/StructuredLog.h"
+#include "Param/PropertyBagProxy.h"
 
 DEFINE_STAT(STAT_AnimNext_InitializeEntry);
 
-FAnimNextSchedulerEntry::FAnimNextSchedulerEntry(const UAnimNextSchedule* InSchedule, UObject* InObject, UE::AnimNext::FScheduleHandle InHandle, EAnimNextScheduleInitMethod InInitMethod)
+FAnimNextSchedulerEntry::FAnimNextSchedulerEntry(const UAnimNextSchedule* InSchedule, UObject* InObject, UE::AnimNext::FScheduleHandle InHandle, EAnimNextScheduleInitMethod InInitMethod, TUniqueFunction<void(const UE::AnimNext::FScheduleContext&)>&& InInitializeCallback)
 	: Schedule(InSchedule)
 	, WeakObject(InObject)
 	, Handle(InHandle)
 	, Context(InSchedule, this)
+	, InitializeCallback(MoveTemp(InInitializeCallback))
 	, RunState(ERunState::None)
 	, InitMethod(InInitMethod)
 {
@@ -26,7 +28,7 @@ FAnimNextSchedulerEntry::~FAnimNextSchedulerEntry()
 	Invalidate();
 }
 
-void FAnimNextSchedulerEntry::Initialize(TUniqueFunction<void(const UE::AnimNext::FScheduleContext&)>&& InInitializeCallback)
+void FAnimNextSchedulerEntry::Initialize()
 {
 	SCOPE_CYCLE_COUNTER(STAT_AnimNext_InitializeEntry);
 	
@@ -39,15 +41,15 @@ void FAnimNextSchedulerEntry::Initialize(TUniqueFunction<void(const UE::AnimNext
 	check(Handle.IsValid());
 
 	ResolvedObject = WeakObject.Get();
-	bIsEditor = ResolvedObject && ResolvedObject->GetWorld()->WorldType == EWorldType::Editor;
+	UWorld* World = ResolvedObject->GetWorld();
+	bIsEditor = ResolvedObject && World->WorldType == EWorldType::Editor;
 
 	// Setup tick function graph
 	if (Schedule->Instructions.Num() > 0)
 	{
 		// TODO: split allocation and registration into separate async tasks to reduce GT overheads of spawning
 		// This will require correctly dealing with ReleaseHandle in the schedule as the two could then overlap
-		check(RunState == ERunState::None);
-		RunState = ERunState::CreatingTasks;
+		TransitionToRunState(ERunState::CreatingTasks);
 
 		RootParamStack = MakeShared<FParamStack>();
 
@@ -55,7 +57,6 @@ void FAnimNextSchedulerEntry::Initialize(TUniqueFunction<void(const UE::AnimNext
 		Context.InstanceData = MakeUnique<FScheduleInstanceData>(Context, Schedule, Handle, this);
 
 		FParamStack::AttachToCurrentThread(RootParamStack);
-		FScheduleContext::AttachToCurrentThread(Context);
 
 		FParamStack& ParamStack = FParamStack::Get();
 
@@ -169,23 +170,28 @@ void FAnimNextSchedulerEntry::Initialize(TUniqueFunction<void(const UE::AnimNext
 				break;
 			case EAnimNextScheduleScheduleOpcode::Exit:
 				check(Prerequisites.Num() == 0);
-				EndTickFunction->AddPrerequisite(ResolvedObject, *TickFunctions.Last().Get());
-				TickFunctions.Last()->Subsequents.Emplace(ResolvedObject, *EndTickFunction.Get());
+				if(TickFunctions.Num() > 0)
+				{
+					EndTickFunction->AddPrerequisite(ResolvedObject, *TickFunctions.Last().Get());
+					TickFunctions.Last()->Subsequents.Emplace(ResolvedObject, *EndTickFunction.Get());
+				}
+				else
+				{
+					EndTickFunction->AddPrerequisite(ResolvedObject, *BeginTickFunction.Get());
+					BeginTickFunction->Subsequent = FTickPrerequisite(ResolvedObject, *EndTickFunction.Get());
+				}
 				break;
 			default:
 				checkNoEntry();
 				break;
 			}
 		}
-
-		FScheduleContext::DetachFromCurrentThread();
+		
 		FParamStack::DetachFromCurrentThread();
 
-		check(RunState == ERunState::CreatingTasks);
-		RunState = ERunState::BindingTasks;
+		TransitionToRunState(ERunState::BindingTasks);
 
 		// Register our tick functions
-		UWorld* World = ResolvedObject->GetWorld();
 		ULevel* Level = World->PersistentLevel;
 		BeginTickFunction->RegisterTickFunction(Level);
 		EndTickFunction->RegisterTickFunction(Level);
@@ -194,11 +200,11 @@ void FAnimNextSchedulerEntry::Initialize(TUniqueFunction<void(const UE::AnimNext
 			TickFunction->RegisterTickFunction(Level);
 		}
 
-		RunState = ERunState::RunningInitialUpdate;
+		TransitionToRunState(ERunState::PendingInitialUpdate);
 
-		if(InInitializeCallback)
+		if(InitializeCallback)
 		{
-			InInitializeCallback(Context);
+			InitializeCallback(Context);
 		}
 
 		// Just pause now if we arent needing an initial update
@@ -212,9 +218,7 @@ void FAnimNextSchedulerEntry::Initialize(TUniqueFunction<void(const UE::AnimNext
 			// generate an output pose, as these worlds never tick
 			if(World->WorldType == EWorldType::EditorPreview)
 			{
-				FScheduleContext::AttachToCurrentThread(Context);
-				FScheduleTickFunction::RunSchedule(Schedule->Instructions);
-				FScheduleContext::DetachFromCurrentThread();
+				FScheduleTickFunction::RunSchedule(*this);
 			}
 		}
 	}
@@ -222,12 +226,14 @@ void FAnimNextSchedulerEntry::Initialize(TUniqueFunction<void(const UE::AnimNext
 	ResolvedObject = nullptr;
 }
 
-void FAnimNextSchedulerEntry::Invalidate()
+void FAnimNextSchedulerEntry::ResetBindingsAndInstanceData()
 {
 	using namespace UE::AnimNext;
 
 	check(IsInGameThread());
 
+	TransitionToRunState(ERunState::None);
+
 	if(BeginTickFunction)
 	{
 		BeginTickFunction->Subsequent.PrerequisiteTickFunction->RemovePrerequisite(WeakObject.Get(), *BeginTickFunction.Get());
@@ -255,11 +261,20 @@ void FAnimNextSchedulerEntry::Invalidate()
 	EndTickFunction.Reset();
 	TickFunctions.Reset();
 
+	Context.InstanceData.Reset();
+}
+
+void FAnimNextSchedulerEntry::Invalidate()
+{
+	using namespace UE::AnimNext;
+
+	ResetBindingsAndInstanceData();
+
+	InitializeCallback = nullptr;
 	Schedule = nullptr;
 	WeakObject = nullptr;
 	TargetObjects.Reset();
 	Handle.Invalidate();
-	Context.InstanceData.Reset();
 }
 
 void FAnimNextSchedulerEntry::ClearTickFunctionPauseFlags()
@@ -289,5 +304,52 @@ void FAnimNextSchedulerEntry::Enable(bool bInEnabled)
 	}
 	EndTickFunction->SetTickFunctionEnable(bInEnabled);
 
-	RunState = bInEnabled ? ERunState::Running : ERunState::Paused;
+	TransitionToRunState(bInEnabled ? ERunState::Running : ERunState::Paused);
+}
+
+void FAnimNextSchedulerEntry::TransitionToRunState(ERunState InNewState)
+{
+	switch(InNewState)
+	{
+	case ERunState::None:
+		check(RunState == ERunState::None || RunState == ERunState::PendingInitialUpdate || RunState == ERunState::Paused || RunState == ERunState::Running);
+		break;
+	case ERunState::CreatingTasks:
+		check(RunState == ERunState::None);
+		break;
+	case ERunState::BindingTasks:
+		check(RunState == ERunState::CreatingTasks);
+		break;
+	case ERunState::PendingInitialUpdate:
+		check(RunState == ERunState::BindingTasks);
+		break;
+	case ERunState::Running:
+		check(RunState == ERunState::PendingInitialUpdate || RunState == ERunState::Paused || RunState == ERunState::Running);
+		break;
+	case ERunState::Paused:
+		check(RunState == ERunState::PendingInitialUpdate || RunState == ERunState::Paused || RunState == ERunState::Running);
+		break;
+	default:
+		checkNoEntry();
+	}
+
+	RunState = InNewState;
+}
+
+#if WITH_EDITOR
+void FAnimNextSchedulerEntry::OnScheduleCompiled()
+{
+	using namespace UE::AnimNext;
+
+	// Store any user-defined scopes, as the instance data will be going away
+	TUniquePtr<FPropertyBagProxy> RootUserScope = MoveTemp(Context.InstanceData->RootUserScope);
+	TMap<FName, FScheduleInstanceData::FUserScope> UserScopes = MoveTemp(Context.InstanceData->UserScopes);
+
+	ResetBindingsAndInstanceData();
+	Initialize();
+
+	// Restore any user scopes to the recreated instance data
+	Context.InstanceData->RootUserScope = MoveTemp(RootUserScope);
+	Context.InstanceData->UserScopes = MoveTemp(UserScopes);
 }
+#endif
\ No newline at end of file
diff --git a/Engine/Plugins/Experimental/AnimNext/Source/AnimNext/Private/Scheduler/AnimNextSchedulerEntry.h b/Engine/Plugins/Experimental/AnimNext/Source/AnimNext/Private/Scheduler/AnimNextSchedulerEntry.h
--- a/Engine/Plugins/Experimental/AnimNext/Source/AnimNext/Private/Scheduler/AnimNextSchedulerEntry.h
+++ b/Engine/Plugins/Experimental/AnimNext/Source/AnimNext/Private/Scheduler/AnimNextSchedulerEntry.h
@@ -23,11 +23,14 @@ struct FAnimNextSchedulerEntry
 	GENERATED_BODY()
 
 	FAnimNextSchedulerEntry() = default;
-	FAnimNextSchedulerEntry(const UAnimNextSchedule* InSchedule, UObject* InObject, UE::AnimNext::FScheduleHandle InHandle, EAnimNextScheduleInitMethod InInitMethod);
+	FAnimNextSchedulerEntry(const UAnimNextSchedule* InSchedule, UObject* InObject, UE::AnimNext::FScheduleHandle InHandle, EAnimNextScheduleInitMethod InInitMethod, TUniqueFunction<void(const UE::AnimNext::FScheduleContext&)>&& InInitializeCallback);
 	~FAnimNextSchedulerEntry();
 
 	// Setup the entry
-	void Initialize(TUniqueFunction<void(const UE::AnimNext::FScheduleContext&)>&& InInitializeCallback);
+	void Initialize();
+
+	// Clear data that binds the schedule to a runtime (e.g. tick functions) and any instance data
+	void ResetBindingsAndInstanceData();
 
 	// Used for pooling
 	void Invalidate();
@@ -38,6 +41,11 @@ struct FAnimNextSchedulerEntry
 	// Clears the bTickEvenWhenPaused flags of the entries' tick functions
 	void ClearTickFunctionPauseFlags();
 
+#if WITH_EDITOR
+	// Resets internal sata if the schedule we are bound to is recompiled in editor
+	void OnScheduleCompiled();
+#endif
+
 	UPROPERTY(Transient)
 	TObjectPtr<const UAnimNextSchedule> Schedule = nullptr;
 
@@ -67,6 +75,9 @@ struct FAnimNextSchedulerEntry
 	// Pre-allocated graph of tick functions
 	TArray<TUniquePtr<UE::AnimNext::FScheduleTickFunction>> TickFunctions;
 
+	// Callback to run on intialization
+	TUniqueFunction<void(const UE::AnimNext::FScheduleContext&)> InitializeCallback;
+
 	// Current delta time, updated each time the schedule runs
 	float DeltaTime = 0.0f;
 
@@ -78,7 +89,7 @@ struct FAnimNextSchedulerEntry
 
 		BindingTasks,
 
-		RunningInitialUpdate,
+		PendingInitialUpdate,
 
 		Running,
 
@@ -88,6 +99,9 @@ struct FAnimNextSchedulerEntry
 	// Current running state
 	ERunState RunState = ERunState::None;
 
+	// Transition to the specified run state, verifying that the current state is valid
+	void TransitionToRunState(ERunState InNewState);
+	
 	// How this entry initializes
 	EAnimNextScheduleInitMethod InitMethod = EAnimNextScheduleInitMethod::InitializeAndPauseInEditor;
 
diff --git a/Engine/Plugins/Experimental/AnimNext/Source/AnimNext/Private/Scheduler/AnimNextSchedulerWorldSubsystem.cpp b/Engine/Plugins/Experimental/AnimNext/Source/AnimNext/Private/Scheduler/AnimNextSchedulerWorldSubsystem.cpp
--- a/Engine/Plugins/Experimental/AnimNext/Source/AnimNext/Private/Scheduler/AnimNextSchedulerWorldSubsystem.cpp
+++ b/Engine/Plugins/Experimental/AnimNext/Source/AnimNext/Private/Scheduler/AnimNextSchedulerWorldSubsystem.cpp
@@ -5,6 +5,7 @@
 #include "Scheduler/AnimNextSchedule.h"
 #include "Param/ParamDefinition.h"
 #include "Engine/World.h"
+#include "Logging/StructuredLog.h"
 
 namespace UE::AnimNext
 {
@@ -110,17 +111,17 @@ UE::AnimNext::FScheduleHandle UAnimNextSchedulerWorldSubsystem::AcquireHandle(UO
 		Handle.Index = FreeEntryIndices.Last();
 		Handle.SerialNumber = ++GEntrySerialNumber;
 		FreeEntryIndices.Pop(false);
-		new (Entries[Handle.Index].Get()) FAnimNextSchedulerEntry(InSchedule, InObject, Handle, InInitMethod);
+		new (Entries[Handle.Index].Get()) FAnimNextSchedulerEntry(InSchedule, InObject, Handle, InInitMethod, MoveTemp(InInitializeCallback));
 	}
 	// Otherwise append a new entry
 	else
 	{
 		Handle.Index = Entries.Num();
 		Handle.SerialNumber = ++GEntrySerialNumber;
-		Entries.Emplace(MakeUnique<FAnimNextSchedulerEntry>(InSchedule, InObject, Handle, InInitMethod));
+		Entries.Emplace(MakeUnique<FAnimNextSchedulerEntry>(InSchedule, InObject, Handle, InInitMethod, MoveTemp(InInitializeCallback)));
 	}
 
-	Entries[Handle.Index]->Initialize(MoveTemp(InInitializeCallback));
+	Entries[Handle.Index]->Initialize();
 
 	// Skip 'invalid' 0 serial number
 	if (GEntrySerialNumber == 0)
@@ -175,27 +176,58 @@ void UAnimNextSchedulerWorldSubsystem::QueueTask(UE::AnimNext::FScheduleHandle I
 	{
 		TUniquePtr<FAnimNextSchedulerEntry>& Entry = Entries[InHandle.Index];
 
-		// TODO: Only supporting scope tasks for now
-		const FAnimNextScheduleParamScopeEntryTask* FoundScope = Entry->Schedule->ParamScopeEntryTasks.FindByPredicate([&InScheduleTaskName](const FAnimNextScheduleParamScopeEntryTask& InTask)
+		// TODO: Only supporting scope tasks or "None" for root for now
+		TSpscQueue<TUniqueFunction<void(const UE::AnimNext::FScheduleContext&)>>* Queue = nullptr;
+		if(InScheduleTaskName == NAME_None)
 		{
-			return InTask.Scope == InScheduleTaskName;
-		});
-
-		if (FoundScope)
+			Queue = &Entry->BeginTickFunction->PreExecuteTasks;
+		}
+		else
 		{
-			switch (InLocation)
+			const FAnimNextScheduleParamScopeEntryTask* FoundScope = Entry->Schedule->ParamScopeEntryTasks.FindByPredicate([&InScheduleTaskName](const FAnimNextScheduleParamScopeEntryTask& InTask)
 			{
-			case FScheduler::ETaskRunLocation::Before:
-				Entry->TickFunctions[FoundScope->TickFunctionIndex]->PreExecuteTasks.Enqueue(MoveTemp(InTaskFunction));
-				break;
-			case FScheduler::ETaskRunLocation::After:
-				Entry->TickFunctions[FoundScope->TickFunctionIndex]->PostExecuteTasks.Enqueue(MoveTemp(InTaskFunction));
-				break;
+				return InTask.Scope == InScheduleTaskName;
+			});
+
+			if (FoundScope)
+			{
+				switch (InLocation)
+				{
+				case FScheduler::ETaskRunLocation::Before:
+					Queue = &Entry->TickFunctions[FoundScope->TickFunctionIndex]->PreExecuteTasks;
+					break;
+				case FScheduler::ETaskRunLocation::After:
+					Queue = &Entry->TickFunctions[FoundScope->TickFunctionIndex]->PostExecuteTasks;
+					break;
+				}
 			}
 		}
+	
+		if (Queue)
+		{
+			Queue->Enqueue(MoveTemp(InTaskFunction));
+		}
 		else
 		{
-			UE_LOG(LogAnimation, Warning, TEXT("QueueTask: Could not find scope '%s' in schedule '%s'"), *InScheduleTaskName.ToString(), *Entry->Schedule.GetName());
+			UE_LOGFMT(LogAnimation, Warning, "QueueTask: Could not find scope '{ScopeName}' in schedule '{ScheduleName}'", InScheduleTaskName, Entry->Schedule.GetName());
 		}
 	}
-}
\ No newline at end of file
+}
+
+#if WITH_EDITOR
+
+void UAnimNextSchedulerWorldSubsystem::OnScheduleCompiled(UAnimNextSchedule* InSchedule)
+{
+	// Cant do this while we are running in a world tick
+	check(!GetWorld()->bInTick); 
+
+	for(TUniquePtr<FAnimNextSchedulerEntry>& Entry : Entries)
+	{
+		if(Entry->Schedule == InSchedule)
+		{
+			Entry->OnScheduleCompiled();
+		}
+	}
+}
+
+#endif
\ No newline at end of file
diff --git a/Engine/Plugins/Experimental/AnimNext/Source/AnimNext/Private/Scheduler/AnimNextSchedulerWorldSubsystem.h b/Engine/Plugins/Experimental/AnimNext/Source/AnimNext/Private/Scheduler/AnimNextSchedulerWorldSubsystem.h
--- a/Engine/Plugins/Experimental/AnimNext/Source/AnimNext/Private/Scheduler/AnimNextSchedulerWorldSubsystem.h
+++ b/Engine/Plugins/Experimental/AnimNext/Source/AnimNext/Private/Scheduler/AnimNextSchedulerWorldSubsystem.h
@@ -4,7 +4,7 @@
 
 #include "CoreMinimal.h"
 #include "Subsystems/WorldSubsystem.h"
-#include "ScheduleHandle.h"
+#include "Scheduler/ScheduleHandle.h"
 #include "ScheduleContext.h"
 #include "Param/ParamStack.h"
 #include "Param/ParamStackLayerHandle.h"
@@ -81,6 +81,11 @@ class UAnimNextSchedulerWorldSubsystem : public UWorldSubsystem
 	// @param	InLocation		Where to run the task, before or after
 	void QueueTask(UE::AnimNext::FScheduleHandle InHandle, FName InScheduleTaskName, TUniqueFunction<void(const UE::AnimNext::FScheduleContext&)>&& InTaskFunction, UE::AnimNext::FScheduler::ETaskRunLocation InLocation);
 
+#if WITH_EDITOR
+	// Refresh any entries that use the provided schedule as it has been recompiled.
+	void OnScheduleCompiled(UAnimNextSchedule* InSchedule);
+#endif
+
 	// Currently running entries, pooled
 	TArray<TUniquePtr<FAnimNextSchedulerEntry>> Entries;
 
diff --git a/Engine/Plugins/Experimental/AnimNext/Source/AnimNext/Private/Scheduler/AnimNextTickFunctionBinding.h b/Engine/Plugins/Experimental/AnimNext/Source/AnimNext/Private/Scheduler/AnimNextTickFunctionBinding.h
--- a/Engine/Plugins/Experimental/AnimNext/Source/AnimNext/Private/Scheduler/AnimNextTickFunctionBinding.h
+++ b/Engine/Plugins/Experimental/AnimNext/Source/AnimNext/Private/Scheduler/AnimNextTickFunctionBinding.h
@@ -3,7 +3,7 @@
 #pragma once
 
 #include "UObject/WeakObjectPtrTemplates.h"
-#include "AnimNextExternalTaskBinding.h"
+#include "Scheduler/AnimNextExternalTaskBinding.h"
 #include "AnimNextTickFunctionBinding.generated.h"
 
 struct FTickFunction;
diff --git a/Engine/Plugins/Experimental/AnimNext/Source/AnimNext/Private/Scheduler/Ports/AnimNextSchedulePort_AnimNextMeshComponentPose.cpp b/Engine/Plugins/Experimental/AnimNext/Source/AnimNext/Private/Scheduler/Ports/AnimNextSchedulePort_AnimNextMeshComponentPose.cpp
--- a/Engine/Plugins/Experimental/AnimNext/Source/AnimNext/Private/Scheduler/Ports/AnimNextSchedulePort_AnimNextMeshComponentPose.cpp
+++ b/Engine/Plugins/Experimental/AnimNext/Source/AnimNext/Private/Scheduler/Ports/AnimNextSchedulePort_AnimNextMeshComponentPose.cpp
@@ -10,6 +10,7 @@
 #include "Engine/SkeletalMesh.h"
 #include "AnimNextStats.h"
 #include "Component/AnimNextMeshComponent.h"
+#include "Param/AnimNextParam.h"
 
 DEFINE_STAT(STAT_AnimNext_Port_SkeletalMeshComponent);
 
@@ -84,15 +85,15 @@ TConstArrayView<UE::AnimNext::FScheduleTerm> UAnimNextSchedulePort_AnimNextMeshC
 	return Terms;
 }
 
-TConstArrayView<FName> UAnimNextSchedulePort_AnimNextMeshComponentPose::GetRequiredParameters() const
+TConstArrayView<FAnimNextParam> UAnimNextSchedulePort_AnimNextMeshComponentPose::GetRequiredParameters() const
 {
 	using namespace UE::AnimNext;
 
-	static const FName ParamNames[] =
+	static const FAnimNextParam Params[] =
 	{
-		ComponentParamId.GetName(),
-		ReferencePoseParamId.GetName()
+		FAnimNextParam(ComponentParamId.GetName(), FAnimNextParamType::GetType<TObjectPtr<UAnimNextMeshComponent>>()),
+		FAnimNextParam(ReferencePoseParamId.GetName(), FAnimNextParamType::GetType<FAnimNextGraphLODPose>()),
 	};
 
-	return ParamNames;
+	return Params;
 }
diff --git a/Engine/Plugins/Experimental/AnimNext/Source/AnimNext/Private/Scheduler/Ports/AnimNextSchedulePort_AnimNextMeshComponentPose.h b/Engine/Plugins/Experimental/AnimNext/Source/AnimNext/Private/Scheduler/Ports/AnimNextSchedulePort_AnimNextMeshComponentPose.h
--- a/Engine/Plugins/Experimental/AnimNext/Source/AnimNext/Private/Scheduler/Ports/AnimNextSchedulePort_AnimNextMeshComponentPose.h
+++ b/Engine/Plugins/Experimental/AnimNext/Source/AnimNext/Private/Scheduler/Ports/AnimNextSchedulePort_AnimNextMeshComponentPose.h
@@ -12,7 +12,7 @@ class UAnimNextSchedulePort_AnimNextMeshComponentPose : public UAnimNextSchedule
 
 	// UAnimNextSchedulePort interface
 	virtual void Run(const UE::AnimNext::FScheduleTermContext& InContext) const override;
-	virtual TConstArrayView<FName> GetRequiredParameters() const override;
+	virtual TConstArrayView<FAnimNextParam> GetRequiredParameters() const override;
 	
 	// IAnimNextScheduleTermInterface interface
 	virtual TConstArrayView<UE::AnimNext::FScheduleTerm> GetTerms() const override;
diff --git a/Engine/Plugins/Experimental/AnimNext/Source/AnimNext/Private/Scheduler/ScheduleContext.cpp b/Engine/Plugins/Experimental/AnimNext/Source/AnimNext/Private/Scheduler/ScheduleContext.cpp
--- a/Engine/Plugins/Experimental/AnimNext/Source/AnimNext/Private/Scheduler/ScheduleContext.cpp
+++ b/Engine/Plugins/Experimental/AnimNext/Source/AnimNext/Private/Scheduler/ScheduleContext.cpp
@@ -2,32 +2,10 @@
 
 #include "ScheduleContext.h"
 #include "AnimNextSchedulerEntry.h"
-#include "Scheduler/AnimNextSchedule.h"
 
 namespace UE::AnimNext
 {
 
-// The currently attached schedule context for this thread
-static thread_local const FScheduleContext* GScheduleContext = nullptr;
-
-void FScheduleContext::AttachToCurrentThread(const FScheduleContext& InContext)
-{
-	check(GScheduleContext == nullptr);
-	GScheduleContext = &InContext;
-}
-
-void FScheduleContext::DetachFromCurrentThread()
-{
-	check(GScheduleContext != nullptr);
-	GScheduleContext = nullptr;
-}
-
-const FScheduleContext& FScheduleContext::Get()
-{
-	check(GScheduleContext != nullptr);
-	return *GScheduleContext;
-}
-
 UObject* FScheduleContext::GetContextObject() const
 {
 	if(Entry)
diff --git a/Engine/Plugins/Experimental/AnimNext/Source/AnimNext/Private/Scheduler/ScheduleContext.h b/Engine/Plugins/Experimental/AnimNext/Source/AnimNext/Private/Scheduler/ScheduleContext.h
--- a/Engine/Plugins/Experimental/AnimNext/Source/AnimNext/Private/Scheduler/ScheduleContext.h
+++ b/Engine/Plugins/Experimental/AnimNext/Source/AnimNext/Private/Scheduler/ScheduleContext.h
@@ -15,15 +15,6 @@ namespace UE::AnimNext
 struct FScheduleContext
 {
 public:
-	// Attach a context to the current thread
-	static void AttachToCurrentThread(const FScheduleContext& InContext);
-
-	// Detach a context from the current thread
-	static void DetachFromCurrentThread();
-
-	// Get the context that is attached to the current thread
-	static const FScheduleContext& Get();
-
 	FScheduleContext() = default;
 
 	FScheduleContext(FScheduleContext&& InOther) = default;
diff --git a/Engine/Plugins/Experimental/AnimNext/Source/AnimNext/Private/Scheduler/ScheduleHandle.h b/Engine/Plugins/Experimental/AnimNext/Source/AnimNext/Private/Scheduler/ScheduleHandle.h
deleted file mode 100644
diff --git a/Engine/Plugins/Experimental/AnimNext/Source/AnimNext/Private/Scheduler/ScheduleInstanceData.cpp b/Engine/Plugins/Experimental/AnimNext/Source/AnimNext/Private/Scheduler/ScheduleInstanceData.cpp
--- a/Engine/Plugins/Experimental/AnimNext/Source/AnimNext/Private/Scheduler/ScheduleInstanceData.cpp
+++ b/Engine/Plugins/Experimental/AnimNext/Source/AnimNext/Private/Scheduler/ScheduleInstanceData.cpp
@@ -2,7 +2,7 @@
 
 #include "ScheduleInstanceData.h"
 #include "Scheduler/AnimNextSchedule.h"
-#include "ScheduleHandle.h"
+#include "Scheduler/ScheduleHandle.h"
 #include "Scheduler/AnimNextSchedulerWorldSubsystem.h"
 #include "Scheduler/AnimNextSchedulerEntry.h"
 #include "Scheduler/AnimNextSchedulePort.h"
@@ -25,15 +25,16 @@ FScheduleInstanceData::FScheduleInstanceData(const FScheduleContext& InScheduleC
 	SCOPE_CYCLE_COUNTER(STAT_AnimNext_CreateInstanceData);
 
 	// Preallocate data for all scopes & graphs in the schedule
-	ScopeCaches.SetNum(InSchedule->ParamScopeEntryTasks.Num());
+	ScopeCaches.SetNum(InSchedule->NumParameterScopes);
 	GraphCaches.SetNum(InSchedule->GraphTasks.Num());
 	ExternalParamCaches.SetNum(InSchedule->ExternalParamTasks.Num());
-	
-	for(int32 ScopeIndex = 0; ScopeIndex < ScopeCaches.Num(); ++ScopeIndex)
+
+	for(int32 EntryIndex = 0; EntryIndex < InSchedule->ParamScopeEntryTasks.Num(); ++EntryIndex)
 	{
-		FScopeCache& ScopeCache = ScopeCaches[ScopeIndex];
-		ScopeCache.ParameterSources.Reserve(InSchedule->ParamScopeEntryTasks[ScopeIndex].ParameterBlocks.Num());
-		for(UAnimNextParameterBlock* ParameterBlock : InSchedule->ParamScopeEntryTasks[ScopeIndex].ParameterBlocks)
+		const FAnimNextScheduleParamScopeEntryTask& EntryTask = InSchedule->ParamScopeEntryTasks[EntryIndex];
+		FScopeCache& ScopeCache = ScopeCaches[EntryTask.ParamScopeIndex];
+		ScopeCache.ParameterSources.Reserve(EntryTask.ParameterBlocks.Num());
+		for(UAnimNextParameterBlock* ParameterBlock : EntryTask.ParameterBlocks)
 		{
 			if(ParameterBlock)
 			{
@@ -43,7 +44,7 @@ FScheduleInstanceData::FScheduleInstanceData(const FScheduleContext& InScheduleC
 
 		ScopeCache.PushedLayers.Reserve(ScopeCache.ParameterSources.Num() + 1); // +1 for any user handles added dynamically
 	}
-	
+
 	// Setup param stack graph
 	RootParamStack = InCurrentEntry->RootParamStack;
 	ParamStacks.SetNum(InSchedule->NumParameterScopes);
@@ -107,9 +108,16 @@ FScheduleInstanceData::~FScheduleInstanceData() = default;
 
 void FScheduleInstanceData::AddReferencedObjects(FReferenceCollector& Collector)
 {
-	for (TPair<FName, TUniquePtr<FPropertyBagProxy>>& ParamPair : UserScopes)
+	for (TPair<FName, FUserScope>& ParamPair : UserScopes)
 	{
-		ParamPair.Value->AddReferencedObjects(Collector);
+		if(ParamPair.Value.AfterSource.IsValid())
+		{
+			ParamPair.Value.AfterSource->AddReferencedObjects(Collector);
+		}
+		if(ParamPair.Value.BeforeSource.IsValid())
+		{
+			ParamPair.Value.BeforeSource->AddReferencedObjects(Collector);
+		}
 	}
 
 	for (FGraphCache& GraphCache : GraphCaches)
diff --git a/Engine/Plugins/Experimental/AnimNext/Source/AnimNext/Private/Scheduler/ScheduleInstanceData.h b/Engine/Plugins/Experimental/AnimNext/Source/AnimNext/Private/Scheduler/ScheduleInstanceData.h
--- a/Engine/Plugins/Experimental/AnimNext/Source/AnimNext/Private/Scheduler/ScheduleInstanceData.h
+++ b/Engine/Plugins/Experimental/AnimNext/Source/AnimNext/Private/Scheduler/ScheduleInstanceData.h
@@ -3,7 +3,7 @@
 #pragma once
 
 #include "CoreMinimal.h"
-#include "ScheduleHandle.h"
+#include "Scheduler/ScheduleHandle.h"
 #include "Graph/AnimNextGraph.h"
 #include "Param/ParamStack.h"
 
@@ -42,8 +42,23 @@ struct FScheduleInstanceData : public FGCObject
 	// Schedule entry that owns this instance
 	FAnimNextSchedulerEntry* Entry = nullptr;
 
+	// Scope for user parameters to be applied at the root of the schedule
+	TUniquePtr<FPropertyBagProxy> RootUserScope;
+
+	// Pushed layer for the root scope
+	FParamStack::FPushedLayerHandle PushedRootUserLayer;
+
+	struct FUserScope
+	{
+		// Layer that will be pushed before the scope, allowing the static scope to override the layer
+		TUniquePtr<FPropertyBagProxy> BeforeSource;
+
+		// Layer that will be pushed after the scope, overriding the static scope
+		TUniquePtr<FPropertyBagProxy> AfterSource;
+	};
+	
 	// Set of dynamic parameter scopes supplied by the user
-	TMap<FName, TUniquePtr<FPropertyBagProxy>> UserScopes;
+	TMap<FName, FUserScope> UserScopes;
 
 	// Cached data for each parameter scope
 	struct FScopeCache
@@ -64,9 +79,6 @@ struct FScheduleInstanceData : public FGCObject
 	// Param stacks required to run the schedule (one per task that requires a stack)
 	TArray<TSharedPtr<FParamStack>> ParamStacks;
 
-	// User handles initialized at startup, always pushed
-	TArray<FParamStackLayerHandle> StaticUserHandles;
-
 	// Intermediate data area
 	FInstancedPropertyBag IntermediatesData;
 
diff --git a/Engine/Plugins/Experimental/AnimNext/Source/AnimNext/Private/Scheduler/ScheduleTickFunction.cpp b/Engine/Plugins/Experimental/AnimNext/Source/AnimNext/Private/Scheduler/ScheduleTickFunction.cpp
--- a/Engine/Plugins/Experimental/AnimNext/Source/AnimNext/Private/Scheduler/ScheduleTickFunction.cpp
+++ b/Engine/Plugins/Experimental/AnimNext/Source/AnimNext/Private/Scheduler/ScheduleTickFunction.cpp
@@ -1,6 +1,8 @@
 // Copyright Epic Games, Inc. All Rights Reserved.
 
 #include "ScheduleTickFunction.h"
+
+#include "Param/PropertyBagProxy.h"
 #include "Scheduler/ScheduleContext.h"
 #include "Scheduler/AnimNextSchedule.h"
 #include "Scheduler/AnimNextScheduleGraphTask.h"
@@ -15,6 +17,25 @@ namespace UE::AnimNext
 
 void FScheduleBeginTickFunction::ExecuteTick(float DeltaTime, ELevelTick TickType, ENamedThreads::Type CurrentThread, const FGraphEventRef& MyCompletionGraphEvent)
 {
+	Run(DeltaTime);
+}
+
+void FScheduleBeginTickFunction::Run(float DeltaTime)
+{
+	while (!PreExecuteTasks.IsEmpty())
+	{
+		TOptional<TUniqueFunction<void(const UE::AnimNext::FScheduleContext&)>> Function = PreExecuteTasks.Dequeue();
+		check(Function.IsSet());
+		Function.GetValue()(Entry.Context);
+	}
+
+	// Push any user layer we have at the root
+	FScheduleInstanceData& InstanceData = Entry.Context.GetInstanceData();
+	if(InstanceData.RootUserScope.IsValid())
+	{
+		InstanceData.PushedRootUserLayer = InstanceData.RootParamStack->PushLayer(InstanceData.RootUserScope->GetLayerHandle());
+	}
+
 	Entry.ResolvedObject = Entry.WeakObject.Get();
 	Entry.DeltaTime = DeltaTime;
 }
@@ -26,7 +47,28 @@ FString FScheduleBeginTickFunction::DiagnosticMessage()
 
 void FScheduleEndTickFunction::ExecuteTick(float DeltaTime, ELevelTick TickType, ENamedThreads::Type CurrentThread, const FGraphEventRef& MyCompletionGraphEvent)
 {
-	if(Entry.RunState == FAnimNextSchedulerEntry::ERunState::RunningInitialUpdate)
+	Run();
+}
+
+void FScheduleEndTickFunction::Run()
+{
+	// Pop any user layer we have at the root
+	FScheduleInstanceData& InstanceData = Entry.Context.GetInstanceData();
+	InstanceData.RootParamStack->PopLayer(InstanceData.PushedRootUserLayer);
+
+	auto RunTaskOnGameThread = [](TUniqueFunction<void(void)>&& InFunction)
+	{
+		if(IsInGameThread())
+		{
+			InFunction();
+		}
+		else
+		{
+			FFunctionGraphTask::CreateAndDispatchWhenReady(MoveTemp(InFunction), TStatId(), nullptr, ENamedThreads::GameThread);
+		}
+	};
+	
+	if(Entry.RunState == FAnimNextSchedulerEntry::ERunState::PendingInitialUpdate)
 	{
 		if( Entry.InitMethod == EAnimNextScheduleInitMethod::InitializeAndPause
 #if WITH_EDITOR
@@ -35,31 +77,20 @@ void FScheduleEndTickFunction::ExecuteTick(float DeltaTime, ELevelTick TickType,
 			)
 		{
 			// Queue task to disable our tick functions now we have performed our initial update
-			FFunctionGraphTask::CreateAndDispatchWhenReady(
-			[this]()
+			RunTaskOnGameThread([this]()
 			{
 				check(IsInGameThread());
-
 				Entry.Enable(false);
-				Entry.RunState = FAnimNextSchedulerEntry::ERunState::Paused;
-			},
-			TStatId(),
-			nullptr,
-			ENamedThreads::GameThread);
+			});
 		}
 	}
 	else
 	{
-		FFunctionGraphTask::CreateAndDispatchWhenReady(
-		[this]()
+		RunTaskOnGameThread([this]()
 		{
 			check(IsInGameThread());
-
-			Entry.RunState = FAnimNextSchedulerEntry::ERunState::Running;
-		},
-		TStatId(),
-		nullptr,
-		ENamedThreads::GameThread);
+			Entry.TransitionToRunState(FAnimNextSchedulerEntry::ERunState::Running);
+		});
 	}
 
 	Entry.ResolvedObject = nullptr;
@@ -72,14 +103,17 @@ FString FScheduleEndTickFunction::DiagnosticMessage()
 
 void FScheduleTickFunction::ExecuteTick(float DeltaTime, ELevelTick TickType, ENamedThreads::Type CurrentThread, const FGraphEventRef& MyCompletionGraphEvent)
 {
-	FScheduleContext::AttachToCurrentThread(ScheduleContext);
+	Run();
+}
 
-	RunSchedule(Instructions, TargetObjects,
+void FScheduleTickFunction::Run()
+{
+	RunScheduleHelper(ScheduleContext, Instructions, TargetObjects,
 		[this]()
 		{
 			while (!PreExecuteTasks.IsEmpty())
 			{
-				TOptional<TUniqueFunction<void(const UE::AnimNext::FScheduleContext&)>> Function = PreExecuteTasks.Dequeue();
+				TOptional<TUniqueFunction<void(const FScheduleContext&)>> Function = PreExecuteTasks.Dequeue();
 				check(Function.IsSet());
 				Function.GetValue()(ScheduleContext);
 			}
@@ -88,25 +122,26 @@ void FScheduleTickFunction::ExecuteTick(float DeltaTime, ELevelTick TickType, EN
 		{
 			while (!PostExecuteTasks.IsEmpty())
 			{
-				TOptional<TUniqueFunction<void(const UE::AnimNext::FScheduleContext&)>> Function = PostExecuteTasks.Dequeue();
+				TOptional<TUniqueFunction<void(const FScheduleContext&)>> Function = PostExecuteTasks.Dequeue();
 				check(Function.IsSet());
 				Function.GetValue()(ScheduleContext);
 			}
 		});
-
-	FScheduleContext::DetachFromCurrentThread();
 }
 
-void FScheduleTickFunction::RunSchedule(TConstArrayView<FAnimNextScheduleInstruction> InInstructions)
+void FScheduleTickFunction::RunSchedule(const FAnimNextSchedulerEntry& InEntry)
 {
-	RunSchedule(InInstructions, {}, [](){}, [](){});
+	InEntry.BeginTickFunction->Run(0.0f);
+	for(const TUniquePtr<FScheduleTickFunction>& TickFunction : InEntry.TickFunctions)
+	{
+		TickFunction->Run();
+	}
+	InEntry.EndTickFunction->Run();
 }
 
-void FScheduleTickFunction::RunSchedule(TConstArrayView<FAnimNextScheduleInstruction> InInstructions, TConstArrayView<TWeakObjectPtr<UObject>> InTargetObjects, TFunctionRef<void(void)> InPreExecuteScope, TFunctionRef<void(void)> InPostExecuteScope)
+void FScheduleTickFunction::RunScheduleHelper(const FScheduleContext& InScheduleContext, TConstArrayView<FAnimNextScheduleInstruction> InInstructions, TConstArrayView<TWeakObjectPtr<UObject>> InTargetObjects, TFunctionRef<void(void)> InPreExecuteScope, TFunctionRef<void(void)> InPostExecuteScope)
 {
-	const FScheduleContext& ScheduleContext = FScheduleContext::Get();
-
-	const UAnimNextSchedule* Schedule = ScheduleContext.Schedule;
+	const UAnimNextSchedule* Schedule = InScheduleContext.Schedule;
 
 	int32 InstructionIndex = 0;
 	while(InstructionIndex < InInstructions.Num())
@@ -117,10 +152,10 @@ void FScheduleTickFunction::RunSchedule(TConstArrayView<FAnimNextScheduleInstruc
 		case EAnimNextScheduleScheduleOpcode::RunGraphTask:
 			{
 				uint32 GraphTaskIndex = Instruction.Operand;
-				FScheduleInstanceData& InstanceData = ScheduleContext.GetInstanceData();
+				FScheduleInstanceData& InstanceData = InScheduleContext.GetInstanceData();
 				FParamStack::AttachToCurrentThread(InstanceData.GetParamStack(Schedule->GraphTasks[GraphTaskIndex].ParamScopeIndex), FParamStack::ECoalesce::Coalesce);
 
-				Schedule->GraphTasks[GraphTaskIndex].RunGraph(ScheduleContext);
+				Schedule->GraphTasks[GraphTaskIndex].RunGraph(InScheduleContext);
 
 				FParamStack::DetachFromCurrentThread(FParamStack::EDecoalesce::Decoalesce);
 				break;
@@ -131,7 +166,7 @@ void FScheduleTickFunction::RunSchedule(TConstArrayView<FAnimNextScheduleInstruc
 				{
 					if (const UObject* Object = InTargetObjects[InstructionIndex].Get())
 					{
-						FScheduleInstanceData& InstanceData = ScheduleContext.GetInstanceData();
+						FScheduleInstanceData& InstanceData = InScheduleContext.GetInstanceData();
 						uint32 ExternalTaskIndex = Instruction.Operand;
 						FParamStack::AddForPendingObject(Object, InstanceData.GetParamStack(Schedule->ExternalTasks[ExternalTaskIndex].ParamScopeIndex));
 					}
@@ -152,10 +187,10 @@ void FScheduleTickFunction::RunSchedule(TConstArrayView<FAnimNextScheduleInstruc
 		case EAnimNextScheduleScheduleOpcode::RunPort:
 			{
 				uint32 PortIndex = Instruction.Operand;
-				FScheduleInstanceData& InstanceData = ScheduleContext.GetInstanceData();
+				FScheduleInstanceData& InstanceData = InScheduleContext.GetInstanceData();
 				FParamStack::AttachToCurrentThread(InstanceData.GetParamStack(Schedule->Ports[PortIndex].ParamScopeIndex));
 
-				Schedule->Ports[PortIndex].RunPort(ScheduleContext);
+				Schedule->Ports[PortIndex].RunPort(InScheduleContext);
 
 				FParamStack::DetachFromCurrentThread();
 				break;
@@ -165,10 +200,10 @@ void FScheduleTickFunction::RunSchedule(TConstArrayView<FAnimNextScheduleInstruc
 				InPreExecuteScope();
 
 				uint32 ScopeEntryIndex = Instruction.Operand;
-				FScheduleInstanceData& InstanceData = ScheduleContext.GetInstanceData();
+				FScheduleInstanceData& InstanceData = InScheduleContext.GetInstanceData();
 				FParamStack::AttachToCurrentThread(InstanceData.GetParamStack(Schedule->ParamScopeEntryTasks[ScopeEntryIndex].ParamScopeIndex), FParamStack::ECoalesce::Coalesce);
 
-				Schedule->ParamScopeEntryTasks[ScopeEntryIndex].RunParamScopeEntry(ScheduleContext);
+				Schedule->ParamScopeEntryTasks[ScopeEntryIndex].RunParamScopeEntry(InScheduleContext);
 
 				FParamStack::DetachFromCurrentThread();
 
@@ -179,10 +214,10 @@ void FScheduleTickFunction::RunSchedule(TConstArrayView<FAnimNextScheduleInstruc
 		case EAnimNextScheduleScheduleOpcode::RunParamScopeExit:
 			{
 				uint32 ScopeExitIndex = Instruction.Operand;
-				FScheduleInstanceData& InstanceData = ScheduleContext.GetInstanceData();
+				FScheduleInstanceData& InstanceData = InScheduleContext.GetInstanceData();
 				FParamStack::AttachToCurrentThread(InstanceData.GetParamStack(Schedule->ParamScopeExitTasks[ScopeExitIndex].ParamScopeIndex));
 
-				Schedule->ParamScopeExitTasks[ScopeExitIndex].RunParamScopeExit(ScheduleContext);
+				Schedule->ParamScopeExitTasks[ScopeExitIndex].RunParamScopeExit(InScheduleContext);
 
 				FParamStack::DetachFromCurrentThread(FParamStack::EDecoalesce::Decoalesce);
 				break;
@@ -190,10 +225,10 @@ void FScheduleTickFunction::RunSchedule(TConstArrayView<FAnimNextScheduleInstruc
 		case EAnimNextScheduleScheduleOpcode::RunExternalParamTask:
 			{
 				uint32 ExternalParamIndex = Instruction.Operand;
-				FScheduleInstanceData& InstanceData = ScheduleContext.GetInstanceData();
+				FScheduleInstanceData& InstanceData = InScheduleContext.GetInstanceData();
 				FParamStack::AttachToCurrentThread(InstanceData.RootParamStack);
 
-				Schedule->ExternalParamTasks[ExternalParamIndex].UpdateExternalParams(ScheduleContext);
+				Schedule->ExternalParamTasks[ExternalParamIndex].UpdateExternalParams(InScheduleContext);
 
 				FParamStack::DetachFromCurrentThread();
 				break;
diff --git a/Engine/Plugins/Experimental/AnimNext/Source/AnimNext/Private/Scheduler/ScheduleTickFunction.h b/Engine/Plugins/Experimental/AnimNext/Source/AnimNext/Private/Scheduler/ScheduleTickFunction.h
--- a/Engine/Plugins/Experimental/AnimNext/Source/AnimNext/Private/Scheduler/ScheduleTickFunction.h
+++ b/Engine/Plugins/Experimental/AnimNext/Source/AnimNext/Private/Scheduler/ScheduleTickFunction.h
@@ -32,6 +32,9 @@ struct FScheduleBeginTickFunction : public FTickFunction
 	virtual void ExecuteTick(float DeltaTime, ELevelTick TickType, ENamedThreads::Type CurrentThread, const FGraphEventRef& MyCompletionGraphEvent) override;
 	virtual FString DiagnosticMessage() override;
 
+	void Run(float DeltaTime);
+
+	TSpscQueue<TUniqueFunction<void(const UE::AnimNext::FScheduleContext&)>> PreExecuteTasks;
 	FAnimNextSchedulerEntry& Entry;
 	FTickPrerequisite Subsequent;
 };
@@ -50,6 +53,8 @@ struct FScheduleEndTickFunction : public FTickFunction
 	virtual void ExecuteTick(float DeltaTime, ELevelTick TickType, ENamedThreads::Type CurrentThread, const FGraphEventRef& MyCompletionGraphEvent) override;
 	virtual FString DiagnosticMessage() override;
 
+	void Run();
+
 	FAnimNextSchedulerEntry& Entry;
 };
 
@@ -69,10 +74,13 @@ struct FScheduleTickFunction : public FTickFunction
 	virtual void ExecuteTick(float DeltaTime, ELevelTick TickType, ENamedThreads::Type CurrentThread, const FGraphEventRef& MyCompletionGraphEvent) override;
 	virtual FString DiagnosticMessage() override;
 
-	// Static helpers for running a slice of instructions
-	// This can be called standalone to run a whole schedule in one call
-	static void RunSchedule(TConstArrayView<FAnimNextScheduleInstruction> InInstructions);
-	static void RunSchedule(TConstArrayView<FAnimNextScheduleInstruction> InInstructions, TConstArrayView<TWeakObjectPtr<UObject>> InTargetObjects, TFunctionRef<void(void)> InPreExecuteScope, TFunctionRef<void(void)> InPostExecuteScope);
+	void Run();
+
+	// Called standalone to run a whole schedule in one call
+	static void RunSchedule(const FAnimNextSchedulerEntry& InEntry);
+
+	// Static helper for running a slice of instructions
+	static void RunScheduleHelper(const FScheduleContext& InScheduleContext, TConstArrayView<FAnimNextScheduleInstruction> InInstructions, TConstArrayView<TWeakObjectPtr<UObject>> InTargetObjects, TFunctionRef<void(void)> InPreExecuteScope, TFunctionRef<void(void)> InPostExecuteScope);
 
 	const FScheduleContext& ScheduleContext;
 	TConstArrayView<FAnimNextScheduleInstruction> Instructions;
diff --git a/Engine/Plugins/Experimental/AnimNext/Source/AnimNext/Private/Scheduler/Scheduler.cpp b/Engine/Plugins/Experimental/AnimNext/Source/AnimNext/Private/Scheduler/Scheduler.cpp
--- a/Engine/Plugins/Experimental/AnimNext/Source/AnimNext/Private/Scheduler/Scheduler.cpp
+++ b/Engine/Plugins/Experimental/AnimNext/Source/AnimNext/Private/Scheduler/Scheduler.cpp
@@ -12,6 +12,7 @@
 #include "Engine/World.h"
 #include "Components/SkeletalMeshComponent.h"
 #include "UObject/ObjectKey.h"
+#include "UObject/UObjectIterator.h"
 
 namespace UE::AnimNext
 {
@@ -148,4 +149,16 @@ void FScheduler::QueueTask(UObject* InObject, FScheduleHandle InHandle, FName In
 	Subsystem->QueueTask(InHandle, InScheduleTaskName, MoveTemp(InTaskFunction), InLocation);
 }
 
+#if WITH_EDITOR
+
+void FScheduler::OnScheduleCompiled(UAnimNextSchedule* InSchedule)
+{
+	for(TObjectIterator<UAnimNextSchedulerWorldSubsystem> It; It; ++It)
+	{
+		It->OnScheduleCompiled(InSchedule);
+	}
+}
+
+#endif
+
 }
\ No newline at end of file
diff --git a/Engine/Plugins/Experimental/AnimNext/Source/AnimNext/Private/Scheduler/Scheduler.h b/Engine/Plugins/Experimental/AnimNext/Source/AnimNext/Private/Scheduler/Scheduler.h
deleted file mode 100644
diff --git a/Engine/Plugins/Experimental/AnimNext/Source/AnimNext/Public/Param/ParamType.h b/Engine/Plugins/Experimental/AnimNext/Source/AnimNext/Public/Param/ParamType.h
--- a/Engine/Plugins/Experimental/AnimNext/Source/AnimNext/Public/Param/ParamType.h
+++ b/Engine/Plugins/Experimental/AnimNext/Source/AnimNext/Public/Param/ParamType.h
@@ -275,6 +275,9 @@ public:
 	 */	
 	size_t GetValueTypeAlignment() const;
 
+	/** Append a string representing this type to the supplied string builder */
+	void ToString(FStringBuilderBase& InStringBuilder) const;
+
 	/** Get a string representing this type */
 	FString ToString() const;
 
diff --git a/Engine/Plugins/Experimental/AnimNext/Source/AnimNext/Public/Scheduler/AnimNextSchedulePort.h b/Engine/Plugins/Experimental/AnimNext/Source/AnimNext/Public/Scheduler/AnimNextSchedulePort.h
--- a/Engine/Plugins/Experimental/AnimNext/Source/AnimNext/Public/Scheduler/AnimNextSchedulePort.h
+++ b/Engine/Plugins/Experimental/AnimNext/Source/AnimNext/Public/Scheduler/AnimNextSchedulePort.h
@@ -8,6 +8,7 @@
 
 class UAnimNextSchedule;
 struct FAnimNextSchedulePortTask;
+struct FAnimNextParam;
 
 namespace UE::AnimNext
 {
@@ -27,5 +28,5 @@ private:
 	virtual void Run(const UE::AnimNext::FScheduleTermContext& InContext) const PURE_VIRTUAL(UAnimNextSchedulePort::Run, )
 
 	// Get any required parameters for this port
-	virtual TConstArrayView<FName> GetRequiredParameters() const PURE_VIRTUAL(UAnimNextSchedulePort::GetRequiredParameters, return TConstArrayView<FName>(); )
+	virtual TConstArrayView<FAnimNextParam> GetRequiredParameters() const PURE_VIRTUAL(UAnimNextSchedulePort::GetRequiredParameters, return TConstArrayView<FAnimNextParam>(); )
 };
diff --git a/Engine/Plugins/Experimental/AnimNext/Source/AnimNextEditor/Internal/Param/ParameterPickerArgs.h b/Engine/Plugins/Experimental/AnimNext/Source/AnimNextEditor/Internal/Param/ParameterPickerArgs.h
--- a/Engine/Plugins/Experimental/AnimNext/Source/AnimNextEditor/Internal/Param/ParameterPickerArgs.h
+++ b/Engine/Plugins/Experimental/AnimNext/Source/AnimNextEditor/Internal/Param/ParameterPickerArgs.h
@@ -38,8 +38,9 @@ struct FParameterBindingReference
 {
 	FParameterBindingReference() = default;
 
-	FParameterBindingReference(const FName& InParameter, const FAssetData& InBlock = FAssetData())
+	FParameterBindingReference(FName InParameter, const FAnimNextParamType& InType, const FAssetData& InBlock = FAssetData())
 		: Parameter(InParameter)
+		, Type(InType)
 		, Block(InBlock)
 	{
 	}
@@ -47,6 +48,9 @@ struct FParameterBindingReference
 	// Parameter name
 	FName Parameter;
 
+	// Parameter type
+	FAnimNextParamType Type;
+
 	// Asset (first found in asset registry) that the parameter is used in
 	FAssetData Asset;
 
@@ -99,6 +103,9 @@ struct FParameterPickerArgs
 	// Delegate called to filter parameters by type for display to the user
 	FOnFilterParameterType OnFilterParameterType;
 
+	// Type to use for any new parameters generated through the picker
+	FAnimNextParamType NewParameterType;
+
 	// Whether we allow selecting multiple parameters or just one
 	bool bMultiSelect = true;
 
diff --git a/Engine/Plugins/Experimental/AnimNext/Source/AnimNextEditor/Private/AnimNextEditorModule.cpp b/Engine/Plugins/Experimental/AnimNext/Source/AnimNextEditor/Private/AnimNextEditorModule.cpp
--- a/Engine/Plugins/Experimental/AnimNext/Source/AnimNextEditor/Private/AnimNextEditorModule.cpp
+++ b/Engine/Plugins/Experimental/AnimNext/Source/AnimNextEditor/Private/AnimNextEditorModule.cpp
@@ -20,9 +20,8 @@
 #include "Scheduler/AnimNextSchedule.h"
 #include "Workspace/AnimNextWorkspaceEditor.h"
 #include "Param/ParameterBlockParameterCustomization.h"
-#include "Param/AnimNextParameterBlockParameter.h"
-#include "Graph/AnimNextGraph_EdGraphNode.h"
 #include "Graph/AnimNextGraph_EdGraphNodeCustomization.h"
+#include "Param/ParamPropertyCustomization.h"
 
 #define LOCTEXT_NAMESPACE "AnimNextEditorModule"
 
@@ -47,16 +46,20 @@ class FModule : public IModule
 			"AnimNextParamType",
 			FOnGetPropertyTypeCustomizationInstance::CreateLambda([] { return MakeShared<FParamTypePropertyTypeCustomization>(); }));
 
+		PropertyModule.RegisterCustomPropertyTypeLayout(
+			"AnimNextParam",
+			FOnGetPropertyTypeCustomizationInstance::CreateLambda([] { return MakeShared<FParamPropertyTypeCustomization>(); }));
+
 		Identifier = MakeShared<FParamNamePropertyTypeIdentifier>();
 		PropertyModule.RegisterCustomPropertyTypeLayout(
 			FNameProperty::StaticClass()->GetFName(),
 			FOnGetPropertyTypeCustomizationInstance::CreateLambda([] { return MakeShared<FParamNamePropertyTypeCustomization>(); }),
 			Identifier);
 
-		PropertyModule.RegisterCustomClassLayout(UAnimNextParameterBlockParameter::StaticClass()->GetFName(), 
+		PropertyModule.RegisterCustomClassLayout("AnimNextParameterBlockParameter", 
 			FOnGetDetailCustomizationInstance::CreateLambda([] { return MakeShared<FParameterBlockParameterCustomization>(); }));
 
-		PropertyModule.RegisterCustomClassLayout(UAnimNextGraph_EdGraphNode::StaticClass()->GetFName(),
+		PropertyModule.RegisterCustomClassLayout("AnimNextGraph_EdGraphNode",
 			FOnGetDetailCustomizationInstance::CreateLambda([] { return MakeShared<FAnimNextGraph_EdGraphNodeCustomization>(); }));
 
 		AnimNextGraphPanelNodeFactory = MakeShared<FAnimNextGraphPanelNodeFactory>();
@@ -160,6 +163,10 @@ class FModule : public IModule
 		{
 			FPropertyEditorModule& PropertyModule = FModuleManager::LoadModuleChecked<FPropertyEditorModule>("PropertyEditor");
 			PropertyModule.UnregisterCustomPropertyTypeLayout("AnimNextParamType");
+			PropertyModule.UnregisterCustomPropertyTypeLayout("AnimNextParam");
+			PropertyModule.UnregisterCustomPropertyTypeLayout(FNameProperty::StaticClass()->GetFName());
+			PropertyModule.UnregisterCustomClassLayout("AnimNextParameterBlockParameter");
+			PropertyModule.UnregisterCustomClassLayout("AnimNextGraph_EdGraphNode");
 		}
 
 		FEdGraphUtilities::UnregisterVisualNodeFactory(AnimNextGraphPanelNodeFactory);
diff --git a/Engine/Plugins/Experimental/AnimNext/Source/AnimNextEditor/Private/EditorUtils.cpp b/Engine/Plugins/Experimental/AnimNext/Source/AnimNextEditor/Private/EditorUtils.cpp
--- a/Engine/Plugins/Experimental/AnimNext/Source/AnimNextEditor/Private/EditorUtils.cpp
+++ b/Engine/Plugins/Experimental/AnimNext/Source/AnimNextEditor/Private/EditorUtils.cpp
@@ -18,6 +18,8 @@
 #include "Param/RigVMDispatch_GetParameter.h"
 #include "Param/RigVMDispatch_SetLayerParameter.h"
 
+#define LOCTEXT_NAMESPACE "AnimNextEditorUtils"
+
 namespace UE::AnimNext::Editor
 {
 
@@ -483,6 +485,43 @@ FName FUtils::GetNewParameterName(const TCHAR* InBaseName, TArrayView<FName> InA
 	return NAME_None;
 }
 
+bool FUtils::IsValidParameterName(const FName InName, FText& OutErrorText)
+{
+	const FString NewString = InName.ToString();
+
+	if(NewString.Len() == 0)
+	{
+		OutErrorText = LOCTEXT("Error_EmptyName", "Empty parameter names are not allowed");
+		return false;
+	}
+
+	// Check start
+	if (NewString[0] == TEXT('.') ||
+		FChar::IsUnderscore(NewString[0]) ||
+		FChar::IsDigit(NewString[0]))
+	{
+		OutErrorText = LOCTEXT("Error_Start", "Name cannot start with an underscore, period or digit");
+		return false;
+	}
+
+	bool bAllowed = true;
+	for (int32 CharIndex = 0; bAllowed && CharIndex < NewString.Len(); ++CharIndex)
+	{
+		bAllowed &= FChar::IsAlnum(NewString[CharIndex]) ||
+					FChar::IsUnderscore(NewString[CharIndex]) ||
+					NewString[CharIndex] == TEXT('.');
+	}
+
+	// Make sure the new name only contains valid characters
+	if (!bAllowed)
+	{
+		OutErrorText = LOCTEXT("Error_CharacterNotAllowed", "Only alpha-numerical, underscore or period characters are allowed");
+		return false;
+	}
+
+	return true;
+}
+
 bool FUtils::DoesParameterNameExist(const FName InName)
 {
 	FAnimNextParameterProviderAssetRegistryExports Exports;
@@ -504,3 +543,5 @@ bool FUtils::GetExportedAssetsForWorkspace(const FAssetData& InWorkspaceAsset, F
 }
 
 }
+
+#undef LOCTEXT_NAMESPACE
\ No newline at end of file
diff --git a/Engine/Plugins/Experimental/AnimNext/Source/AnimNextEditor/Private/EditorUtils.h b/Engine/Plugins/Experimental/AnimNext/Source/AnimNextEditor/Private/EditorUtils.h
--- a/Engine/Plugins/Experimental/AnimNext/Source/AnimNextEditor/Private/EditorUtils.h
+++ b/Engine/Plugins/Experimental/AnimNext/Source/AnimNextEditor/Private/EditorUtils.h
@@ -38,7 +38,10 @@ struct FUtils
 
 	static FName GetNewParameterName(const TCHAR* InBaseName, TArrayView<FName> InAdditionalExistingNames);
 
+	static bool IsValidParameterName(const FName InName, FText& OutErrorText);
+
 	static bool DoesParameterNameExist(const FName InName);
+	
 	static bool DoesParameterNameExistInAsset(const FName InName, const FAssetData& InAsset);
 
 	static bool GetExportedAssetsForWorkspace(const FAssetData& InWorkspaceAsset, FAnimNextWorkspaceAssetRegistryExports& OutExports);
diff --git a/Engine/Plugins/Experimental/AnimNext/Source/AnimNextEditor/Private/Param/ParamNamePropertyCustomization.cpp b/Engine/Plugins/Experimental/AnimNext/Source/AnimNextEditor/Private/Param/ParamNamePropertyCustomization.cpp
--- a/Engine/Plugins/Experimental/AnimNext/Source/AnimNextEditor/Private/Param/ParamNamePropertyCustomization.cpp
+++ b/Engine/Plugins/Experimental/AnimNext/Source/AnimNextEditor/Private/Param/ParamNamePropertyCustomization.cpp
@@ -19,7 +19,7 @@ namespace UE::AnimNext::Editor
 
 static const FProperty* GetMetadataProperty(const FProperty* Property)
 {
-	if (FProperty* OuterProperty = Property->GetOwner<FProperty>())
+	if (const FProperty* OuterProperty = Property->GetOwner<FProperty>())
 	{
 		if (OuterProperty->IsA<FArrayProperty>()
 			|| OuterProperty->IsA<FSetProperty>()
@@ -63,6 +63,7 @@ void FParamNamePropertyTypeCustomization::CustomizeHeader(TSharedRef<IPropertyHa
 		}
 		return EFilterParameterResult::Exclude;
 	});
+	PickerArgs.NewParameterType = FilterType;
 	PickerArgs.OnParameterPicked = FOnParameterPicked::CreateLambda([this, WeakPropertyHandle = WeakPropertyHandle](const FParameterBindingReference& InParameterBinding)
 	{
 		if(TSharedPtr<IPropertyHandle> PropertyHandle = WeakPropertyHandle.Pin())
diff --git a/Engine/Plugins/Experimental/AnimNext/Source/AnimNextEditor/Private/Param/ParamPropertyCustomization.cpp b/Engine/Plugins/Experimental/AnimNext/Source/AnimNextEditor/Private/Param/ParamPropertyCustomization.cpp
new file mode 100644
--- /dev/null
+++ b/Engine/Plugins/Experimental/AnimNext/Source/AnimNextEditor/Private/Param/ParamPropertyCustomization.cpp
@@ -0,0 +1,116 @@
+﻿// Copyright Epic Games, Inc. All Rights Reserved.
+
+#include "ParamPropertyCustomization.h"
+
+#include "DetailWidgetRow.h"
+#include "EdGraphSchema_K2.h"
+#include "EditorUtils.h"
+#include "PropertyHandle.h"
+#include "SParameterPickerCombo.h"
+#include "Param/AnimNextParam.h"
+#include "Param/ParamCompatibility.h"
+#include "Param/ParamUtils.h"
+
+#define LOCTEXT_NAMESPACE "ParamPropertyCustomization"
+
+namespace UE::AnimNext::Editor
+{
+
+void FParamPropertyTypeCustomization::CustomizeHeader(TSharedRef<IPropertyHandle> InPropertyHandle, FDetailWidgetRow& InHeaderRow, IPropertyTypeCustomizationUtils& InCustomizationUtils)
+{
+	auto GetMetadataProperty = [](const FProperty* InProperty)
+	{
+		if (const FProperty* OuterProperty = InProperty->GetOwner<FProperty>())
+		{
+			if (OuterProperty->IsA<FArrayProperty>()
+				|| OuterProperty->IsA<FSetProperty>()
+				|| OuterProperty->IsA<FMapProperty>())
+			{
+				return OuterProperty;
+			}
+		}
+
+		return InProperty;
+	};
+
+	PropertyHandle = InPropertyHandle;
+	NamePropertyHandle = InPropertyHandle->GetChildHandle(GET_MEMBER_NAME_CHECKED(FAnimNextParam, Name));
+	TypePropertyHandle = InPropertyHandle->GetChildHandle(GET_MEMBER_NAME_CHECKED(FAnimNextParam, Type));
+
+	const FProperty* Property = GetMetadataProperty(InPropertyHandle->GetProperty());
+	const FString ParamTypeString = Property->GetMetaData("AllowedParamType");
+	const bool bAllowNone = Property->HasMetaData("AllowNone");
+	FAnimNextParamType FilterType = FAnimNextParamType::FromString(ParamTypeString);
+
+	FParameterPickerArgs PickerArgs;
+	PickerArgs.bShowBlocks = false;
+	PickerArgs.bMultiSelect = false;
+	PickerArgs.OnFilterParameterType = FOnFilterParameterType::CreateLambda([FilterType](const FAnimNextParamType& InParameterType)
+	{
+		if(!FilterType.IsValid() || FParamUtils::GetCompatibility(FilterType, InParameterType).IsCompatible())
+		{
+			return EFilterParameterResult::Include;
+		}
+		return EFilterParameterResult::Exclude;
+	});
+	PickerArgs.NewParameterType = FilterType;
+	PickerArgs.OnParameterPicked = FOnParameterPicked::CreateLambda([this](const FParameterBindingReference& InParameterBinding)
+	{
+		if(PropertyHandle && NamePropertyHandle && TypePropertyHandle)
+		{
+			PropertyHandle->NotifyPreChange();
+			NamePropertyHandle->SetValue(InParameterBinding.Parameter);
+			TArray<void*> TypeValues;
+			TypePropertyHandle->AccessRawData(TypeValues);
+			for(void* TypePtr : TypeValues)
+			{
+				FAnimNextParamType& Type = *static_cast<FAnimNextParamType*>(TypePtr);
+				Type = InParameterBinding.Type;
+			}
+			PropertyHandle->NotifyPostChange(EPropertyChangeType::ValueSet);
+
+			Refresh();
+		}
+	});
+	PickerArgs.bAllowNone = bAllowNone;
+	
+	InHeaderRow
+	.NameContent()
+	[
+		InPropertyHandle->CreatePropertyNameWidget()
+	]
+	.ValueContent()
+	[
+		SNew(SParameterPickerCombo)
+		.PickerArgs(PickerArgs)
+		.OnGetParameterName_Lambda([this]()
+		{
+			return CachedName;
+		})
+		.OnGetParameterType_Lambda([this]()
+		{
+			return CachedType;
+		})
+	];
+
+	Refresh();
+}
+
+void FParamPropertyTypeCustomization::CustomizeChildren(TSharedRef<IPropertyHandle> InPropertyHandle, IDetailChildrenBuilder& ChildBuilder, IPropertyTypeCustomizationUtils& CustomizationUtils)
+{
+}
+
+void FParamPropertyTypeCustomization::Refresh()
+{
+	TArray<void*> TypeValues;
+	TypePropertyHandle->AccessRawData(TypeValues);
+	if(TypeValues.Num() == 1 && TypeValues[0] != nullptr)
+	{
+		CachedType = *static_cast<FAnimNextParamType*>(TypeValues[0]);
+	}
+	NamePropertyHandle->GetValue(CachedName);
+}
+
+}
+
+#undef LOCTEXT_NAMESPACE
\ No newline at end of file
diff --git a/Engine/Plugins/Experimental/AnimNext/Source/AnimNextEditor/Private/Param/ParamPropertyCustomization.h b/Engine/Plugins/Experimental/AnimNext/Source/AnimNextEditor/Private/Param/ParamPropertyCustomization.h
new file mode 100644
--- /dev/null
+++ b/Engine/Plugins/Experimental/AnimNext/Source/AnimNextEditor/Private/Param/ParamPropertyCustomization.h
@@ -0,0 +1,30 @@
+﻿// Copyright Epic Games, Inc. All Rights Reserved.
+
+#pragma once
+
+#include "IPropertyTypeCustomization.h"
+#include "PropertyEditorModule.h"
+#include "Param/ParamType.h"
+
+namespace UE::AnimNext::Editor
+{
+
+class FParamPropertyTypeCustomization : public IPropertyTypeCustomization
+{
+private:
+	// IPropertyTypeCustomization interface
+	virtual void CustomizeHeader(TSharedRef<IPropertyHandle> InPropertyHandle, FDetailWidgetRow& HeaderRow, IPropertyTypeCustomizationUtils& CustomizationUtils) override;
+	virtual void CustomizeChildren(TSharedRef<IPropertyHandle> InPropertyHandle, IDetailChildrenBuilder& ChildBuilder, IPropertyTypeCustomizationUtils& CustomizationUtils) override;
+
+	void Refresh();
+
+private:
+	TSharedPtr<IPropertyHandle> PropertyHandle;
+	TSharedPtr<IPropertyHandle> NamePropertyHandle;
+	TSharedPtr<IPropertyHandle> TypePropertyHandle;
+	FName CachedName = NAME_None;
+	FText CachedNameText;
+	FAnimNextParamType CachedType;
+};
+
+}
diff --git a/Engine/Plugins/Experimental/AnimNext/Source/AnimNextEditor/Private/Param/ParametersGraphPanelPinFactory.cpp b/Engine/Plugins/Experimental/AnimNext/Source/AnimNextEditor/Private/Param/ParametersGraphPanelPinFactory.cpp
--- a/Engine/Plugins/Experimental/AnimNext/Source/AnimNextEditor/Private/Param/ParametersGraphPanelPinFactory.cpp
+++ b/Engine/Plugins/Experimental/AnimNext/Source/AnimNextEditor/Private/Param/ParametersGraphPanelPinFactory.cpp
@@ -31,6 +31,16 @@ TSharedPtr<SGraphPin> FParametersGraphPanelPinFactory::CreatePin_Internal(UEdGra
 			}
 		}
 	}
+	else if(UEdGraphNode* EdGraphNode = InPin->GetOwningNode())
+	{
+		if(EdGraphNode->GetPinMetaData(InPin->GetFName(), "CustomWidget") == "ParamName")
+		{
+			const FString ParamTypeString = EdGraphNode->GetPinMetaData(InPin->GetFName(), "AllowedParamType");
+			FAnimNextParamType FilterType = FAnimNextParamType::FromString(ParamTypeString);
+			return SNew(SGraphPinParamName, InPin)
+				.FilterType(FilterType);
+		}
+	}
 
 	return FRigVMEdGraphPanelPinFactory::CreatePin_Internal(InPin);
 }
diff --git a/Engine/Plugins/Experimental/AnimNext/Source/AnimNextEditor/Private/Param/SAddParametersDialog.cpp b/Engine/Plugins/Experimental/AnimNext/Source/AnimNextEditor/Private/Param/SAddParametersDialog.cpp
--- a/Engine/Plugins/Experimental/AnimNext/Source/AnimNextEditor/Private/Param/SAddParametersDialog.cpp
+++ b/Engine/Plugins/Experimental/AnimNext/Source/AnimNextEditor/Private/Param/SAddParametersDialog.cpp
@@ -23,6 +23,7 @@
 #include "SSimpleButton.h"
 #include "SSimpleComboButton.h"
 #include "ToolMenus.h"
+#include "String/ParseTokens.h"
 
 #define LOCTEXT_NAMESPACE "SAddParametersDialog"
 
@@ -56,6 +57,7 @@ void SAddParametersDialog::Construct(const FArguments& InArgs)
 	using namespace AddParametersDialog;
 
 	TargetBlock = InArgs._Block;
+	OnFilterParameterType = InArgs._OnFilterParameterType;
 
 	SWindow::Construct(SWindow::FArguments()
 		.Title(LOCTEXT("WindowTitle", "Add Parameters"))
@@ -172,10 +174,10 @@ void SAddParametersDialog::Construct(const FArguments& InArgs)
 		]);
 
 	// Add an initial item
-	AddEntry();
+	AddEntry(InArgs._InitialParamType);
 }
 
-void SAddParametersDialog::AddEntry()
+void SAddParametersDialog::AddEntry(const FAnimNextParamType& InParamType)
 {
 	const UAnimNextParameterSettings* Settings = GetDefault<UAnimNextParameterSettings>();
 	
@@ -186,7 +188,7 @@ void SAddParametersDialog::AddEntry()
 		PendingNames.Add(QueuedAdd->Name);
 	}
 	FName ParameterName = FUtils::GetNewParameterName(TEXT("NewParameter"), PendingNames);
-	Entries.Add(MakeShared<FParameterToAdd>(Settings->GetLastParameterType(), ParameterName));
+	Entries.Add(MakeShared<FParameterToAdd>(InParamType.IsValid() ? InParamType : Settings->GetLastParameterType(), ParameterName));
 
 	RefreshEntries();
 }
@@ -226,23 +228,39 @@ class SParameterToAdd : public SMultiColumnTableRow<TSharedRef<FParameterToAdd>>
 					.ToolTipText(LOCTEXT("NameTooltip", "The name of the new parameter"))
 					.Text_Lambda([this]()
 					{
-						return FText::FromName(Entry->Name);
+						return UncookedOnly::FUtils::GetParameterDisplayNameText(Entry->Name);
 					})
 					.OnTextCommitted_Lambda([this](const FText& InText, ETextCommit::Type InCommitType)
 					{
-						Entry->Name = *InText.ToString();
+						const FString UserInput = InText.ToString();
+						// Parse out segments to collapse adjacent delimiters
+						TStringBuilder<128> RebuiltInput;
+						UE::String::ParseTokensMultiple(UserInput, { TEXT('.'), TEXT('_') }, [&RebuiltInput](const FStringView InToken)
+						{
+							if(RebuiltInput.Len() != 0)
+							{
+								RebuiltInput.Append(TEXT("_"));
+							}
+							RebuiltInput.Append(InToken);
+						}, String::EParseTokensOptions::SkipEmpty);
+						Entry->Name = RebuiltInput.ToString();
 					})
 					.OnVerifyTextChanged_Lambda([this](const FText& InNewText, FText& OutErrorText)
 					{
 						const FString NewString = InNewText.ToString();
 
-						// Make sure the new name only contains valid characters
-						if (!FName::IsValidXName(NewString, INVALID_OBJECTNAME_CHARACTERS INVALID_LONGPACKAGE_CHARACTERS, &OutErrorText))
+						// See if this can be represented as an FName
+						if(!FName::IsValidXName(NewString, INVALID_NAME_CHARACTERS, &OutErrorText))
 						{
 							return false;
 						}
 
 						const FName Name(*NewString);
+						if(!FUtils::IsValidParameterName(Name, OutErrorText))
+						{
+							return false;
+						}
+
 						if(FUtils::DoesParameterNameExist(Name))
 						{
 							OutErrorText = LOCTEXT("Error_NameExists", "This name already exists in the project");
@@ -267,13 +285,62 @@ class SParameterToAdd : public SMultiColumnTableRow<TSharedRef<FParameterToAdd>>
 				UAnimNextParameterSettings* Settings = GetMutableDefault<UAnimNextParameterSettings>();
 				Settings->SetLastParameterType(Entry->Type);
 			};
+			
+			auto GetFilteredVariableTypeTree = [this](TArray<TSharedPtr<UEdGraphSchema_K2::FPinTypeTreeInfo>>& TypeTree, ETypeTreeFilter TypeTreeFilter)
+			{
+				FUtils::GetFilteredVariableTypeTree(TypeTree, TypeTreeFilter);
+
+				if(TSharedPtr<SAddParametersDialog> Dialog = WeakDialog.Pin())
+				{
+					if(Dialog->OnFilterParameterType.IsBound())
+					{
+						auto IsPinTypeAllowed = [&Dialog](const FEdGraphPinType& InType)
+						{
+							FAnimNextParamType Type = UncookedOnly::FUtils::GetParamTypeFromPinType(InType);
+							if(Type.IsValid())
+							{
+								return Dialog->OnFilterParameterType.Execute(Type) == EFilterParameterResult::Include;
+							}
+							return false;
+						};
+
+						// Additionally filter by allowed types
+						for (int32 Index = 0; Index < TypeTree.Num(); )
+						{
+							TSharedPtr<UEdGraphSchema_K2::FPinTypeTreeInfo>& PinType = TypeTree[Index];
+
+							if (PinType->Children.Num() == 0 && !IsPinTypeAllowed(PinType->GetPinType(/*bForceLoadSubCategoryObject*/false)))
+							{
+								TypeTree.RemoveAt(Index);
+								continue;
+							}
+
+							for (int32 ChildIndex = 0; ChildIndex < PinType->Children.Num(); )
+							{
+								TSharedPtr<UEdGraphSchema_K2::FPinTypeTreeInfo> Child = PinType->Children[ChildIndex];
+								if (Child.IsValid())
+								{
+									if (!IsPinTypeAllowed(Child->GetPinType(/*bForceLoadSubCategoryObject*/false)))
+									{
+										PinType->Children.RemoveAt(ChildIndex);
+										continue;
+									}
+								}
+								++ChildIndex;
+							}
+
+							++Index;
+						}
+					}
+				}
+			};
 
 			return
 				SNew(SBox)
 				.HAlign(HAlign_Left)
 				.VAlign(VAlign_Center)
 				[
-					SNew(SPinTypeSelector, FGetPinTypeTree::CreateStatic(&Editor::FUtils::GetFilteredVariableTypeTree))
+					SNew(SPinTypeSelector, FGetPinTypeTree::CreateLambda(GetFilteredVariableTypeTree))
 						.TargetPinType_Lambda(GetPinInfo)
 						.OnPinTypeChanged_Lambda(PinInfoChanged)
 						.Schema(GetDefault<UPropertyBagSchema>())
diff --git a/Engine/Plugins/Experimental/AnimNext/Source/AnimNextEditor/Private/Param/SAddParametersDialog.h b/Engine/Plugins/Experimental/AnimNext/Source/AnimNextEditor/Private/Param/SAddParametersDialog.h
--- a/Engine/Plugins/Experimental/AnimNext/Source/AnimNextEditor/Private/Param/SAddParametersDialog.h
+++ b/Engine/Plugins/Experimental/AnimNext/Source/AnimNextEditor/Private/Param/SAddParametersDialog.h
@@ -6,6 +6,7 @@
 #include "Widgets/SWindow.h"
 #include "AssetRegistry/AssetData.h"
 #include "Widgets/Views/SListView.h"
+#include "Param/ParameterPickerArgs.h"
 
 class SWrapBox;
 class UAnimNextParameterLibrary;
@@ -28,7 +29,13 @@ public:
 
 	/** Whether we allow multiple parameters to be added or just one at a time */
 	SLATE_ARGUMENT(bool, AllowMultiple)
-	
+
+	/** Delegate called to filter parameters by type for display to the user */
+	SLATE_EVENT(FOnFilterParameterType, OnFilterParameterType)
+
+	/** Initial parameter type to use */
+	SLATE_ARGUMENT(FAnimNextParamType, InitialParamType)
+
 	SLATE_END_ARGS()
 
 	void Construct(const FArguments& InArgs);
@@ -36,7 +43,7 @@ public:
 	bool ShowModal(TArray<FParameterToAdd>& OutParameters);
 
 private:
-	void AddEntry();
+	void AddEntry(const FAnimNextParamType& InParamType = FAnimNextParamType());
 
 	void RefreshEntries();
 
@@ -55,6 +62,8 @@ private:
 
 	UAnimNextParameterBlock_EditorData* TargetBlock = nullptr;
 	
+	FOnFilterParameterType OnFilterParameterType;
+
 	bool bCancelPressed = false;
 };
 
diff --git a/Engine/Plugins/Experimental/AnimNext/Source/AnimNextEditor/Private/Param/SGraphPinParamName.cpp b/Engine/Plugins/Experimental/AnimNext/Source/AnimNextEditor/Private/Param/SGraphPinParamName.cpp
--- a/Engine/Plugins/Experimental/AnimNext/Source/AnimNextEditor/Private/Param/SGraphPinParamName.cpp
+++ b/Engine/Plugins/Experimental/AnimNext/Source/AnimNextEditor/Private/Param/SGraphPinParamName.cpp
@@ -5,6 +5,8 @@
 #include "SParameterPickerCombo.h"
 #include "ScopedTransaction.h"
 #include "UncookedOnlyUtils.h"
+#include "Param/ParamCompatibility.h"
+#include "Param/ParamUtils.h"
 #include "Param/RigVMDispatch_GetLayerParameter.h"
 #include "Param/RigVMDispatch_GetParameter.h"
 #include "Param/RigVMDispatch_SetLayerParameter.h"
@@ -18,6 +20,7 @@ void SGraphPinParamName::Construct(const FArguments& InArgs, UEdGraphPin* InPin)
 {
 	ModelPin = InArgs._ModelPin;
 	Node = InArgs._GraphNode;
+	FilterType = InArgs._FilterType;
 
 	SGraphPin::Construct(SGraphPin::FArguments(), InPin);
 }
@@ -28,6 +31,8 @@ TSharedRef<SWidget> SGraphPinParamName::GetDefaultValueWidget()
 	Args.bShowBlocks = false;
 	Args.bMultiSelect = false;
 
+	CachedType = UncookedOnly::FUtils::GetParameterTypeFromName(FName(*GraphPinObj->DefaultValue));
+
 	// Check whether this is a Set/Get parameter from block node, and if so only show bound parameters
 	if (ModelPin)
 	{
@@ -48,24 +53,28 @@ TSharedRef<SWidget> SGraphPinParamName::GetDefaultValueWidget()
 	
 	Args.OnParameterPicked = FOnParameterPicked::CreateLambda([this](const FParameterBindingReference& InParameterBinding)
 	{
-		if(ModelPin)
-		{
-			FScopedTransaction Transaction(LOCTEXT("SelectParameter", "Select Parameter"));
-			GraphPinObj->Modify();
-			GraphPinObj->GetSchema()->TrySetDefaultValue(*GraphPinObj, InParameterBinding.Parameter.ToString());
-		}
+		FScopedTransaction Transaction(LOCTEXT("SelectParameter", "Select Parameter"));
+		GraphPinObj->Modify();
+		GraphPinObj->GetSchema()->TrySetDefaultValue(*GraphPinObj, InParameterBinding.Parameter.ToString());
+
+		CachedType = UncookedOnly::FUtils::GetParameterTypeFromName(FName(*GraphPinObj->DefaultValue));
 	});
 	
 	Args.OnFilterParameterType = FOnFilterParameterType::CreateLambda([this](const FAnimNextParamType& InParamType)-> EFilterParameterResult
 	{
-		if(Node && ModelPin->IsLinked())
+		if(FilterType.IsValid())
+		{
+			return FParamUtils::GetCompatibility(FilterType, InParamType).IsCompatible() ? EFilterParameterResult::Include : EFilterParameterResult::Exclude;
+		}
+		else if(Node && ModelPin->IsLinked())
 		{
 			const FAnimNextParamType Type = FAnimNextParamType::FromRigVMTemplateArgument(ModelPin->GetTemplateArgumentType());
-			return Type.IsValid() && Type == InParamType ? EFilterParameterResult::Include : EFilterParameterResult::Exclude;
+			return Type.IsValid() && FParamUtils::GetCompatibility(Type, InParamType).IsCompatible() ? EFilterParameterResult::Include : EFilterParameterResult::Exclude;
 		}
 
 		return EFilterParameterResult::Include;
 	});
+	Args.NewParameterType = FilterType;
 	
 	return SNew(SParameterPickerCombo)
 		.PickerArgs(Args)
@@ -75,7 +84,7 @@ TSharedRef<SWidget> SGraphPinParamName::GetDefaultValueWidget()
 		})
 		.OnGetParameterType_Lambda([this]()
 		{
-			return UncookedOnly::FUtils::GetParameterTypeFromName(FName(*GraphPinObj->DefaultValue));
+			return CachedType;
 		});
 }
 
diff --git a/Engine/Plugins/Experimental/AnimNext/Source/AnimNextEditor/Private/Param/SGraphPinParamName.h b/Engine/Plugins/Experimental/AnimNext/Source/AnimNextEditor/Private/Param/SGraphPinParamName.h
--- a/Engine/Plugins/Experimental/AnimNext/Source/AnimNextEditor/Private/Param/SGraphPinParamName.h
+++ b/Engine/Plugins/Experimental/AnimNext/Source/AnimNextEditor/Private/Param/SGraphPinParamName.h
@@ -3,6 +3,7 @@
 #pragma once
 
 #include "SGraphPin.h"
+#include "Param/ParamType.h"
 
 class URigVMPin;
 class URigVMEdGraphNode;
@@ -13,11 +14,17 @@ namespace UE::AnimNext::Editor
 // A pin widget that allows picking using an AnimNext parameter picker
 class SGraphPinParamName : public SGraphPin
 {
-	SLATE_BEGIN_ARGS(SGraphPinParamName) {}
+	SLATE_BEGIN_ARGS(SGraphPinParamName)
+		: _ModelPin(nullptr)
+		, _GraphNode(nullptr)
+	{}
 
 	SLATE_ARGUMENT(URigVMPin*, ModelPin)
+
 	SLATE_ARGUMENT(URigVMEdGraphNode*, GraphNode)
 
+	SLATE_ARGUMENT(FAnimNextParamType, FilterType)
+
 	SLATE_END_ARGS()
 
 	void Construct(const FArguments& InArgs, UEdGraphPin* InPin);
@@ -26,7 +33,12 @@ private:
 	virtual TSharedRef<SWidget>	GetDefaultValueWidget() override;
 
 	URigVMPin* ModelPin = nullptr;
+
 	URigVMEdGraphNode* Node = nullptr;
+
+	FAnimNextParamType FilterType;
+
+	FAnimNextParamType CachedType;
 };
 
 }
\ No newline at end of file
diff --git a/Engine/Plugins/Experimental/AnimNext/Source/AnimNextEditor/Private/Param/SParameterPicker.cpp b/Engine/Plugins/Experimental/AnimNext/Source/AnimNextEditor/Private/Param/SParameterPicker.cpp
--- a/Engine/Plugins/Experimental/AnimNext/Source/AnimNextEditor/Private/Param/SParameterPicker.cpp
+++ b/Engine/Plugins/Experimental/AnimNext/Source/AnimNextEditor/Private/Param/SParameterPicker.cpp
@@ -48,11 +48,10 @@ struct FParameterPickerEntry
 
 	FParameterPickerEntry() = default;
 
-	FParameterPickerEntry(const FParameterBindingReference& InBinding, const FAnimNextParamType& InParamType)
+	FParameterPickerEntry(const FParameterBindingReference& InBinding)
 		: Binding(InBinding)
-		, ParamType(InParamType)
 	{
-		PinType = UE::AnimNext::UncookedOnly::FUtils::GetPinTypeFromParamType(ParamType);
+		PinType = UE::AnimNext::UncookedOnly::FUtils::GetPinTypeFromParamType(InBinding.Type);
 		PinIcon = FBlueprintEditorUtils::GetIconFromPin(PinType, /* bIsLarge = */true);
 		PinColor = GetDefault<UEdGraphSchema_K2>()->GetPinTypeColor(PinType);
 		FString ParameterString = Binding.Parameter.ToString();
@@ -85,8 +84,6 @@ struct FParameterPickerEntry
 
 	FParameterBindingReference Binding;
 
-	FAnimNextParamType ParamType;
-
 	FEdGraphPinType PinType;
 
 	const FSlateBrush* PinIcon = nullptr;
@@ -146,7 +143,9 @@ void SParameterPicker::Construct(const FArguments& InArgs)
 			FSlateApplication::Get().DismissAllMenus();
 			TSharedRef<SAddParametersDialog> AddParametersDialog =
 				SNew(SAddParametersDialog)
-				.AllowMultiple(false);
+				.AllowMultiple(false)
+				.OnFilterParameterType(Args.OnFilterParameterType)
+				.InitialParamType(Args.NewParameterType);
 			TArray<FParameterToAdd> ParametersToAdd;
 			if(AddParametersDialog->ShowModal(ParametersToAdd))
 			{
@@ -267,8 +266,7 @@ void SParameterPicker::Construct(const FArguments& InArgs)
 						{
 							FScopedTransaction Transaction(LOCTEXT("AddParameter", "Add parameter"));
 
-							FParameterBindingReference Reference;
-							Reference.Parameter = ParametersToAdd[0].Name;
+							FParameterBindingReference Reference(ParametersToAdd[0].Name, ParametersToAdd[0].Type);
 							Args.OnParameterPicked.ExecuteIfBound(Reference);
 						}
 					}
@@ -304,7 +302,7 @@ void SParameterPicker::RefreshEntries()
 
 	if(Args.bAllowNone)
 	{
-		Entries.Add(MakeShared<FParameterPickerEntry>(FParameterBindingReference(NAME_None), FAnimNextParamType()));
+		Entries.Add(MakeShared<FParameterPickerEntry>(FParameterBindingReference(NAME_None, FAnimNextParamType())));
 	}
 	
 	// Find all blocks and their bound parameters
@@ -323,7 +321,7 @@ void SParameterPicker::RefreshEntries()
 				for(const FAnimNextParameterAssetRegistryExportEntry& Export : Exports.Parameters)
 				{
 					BoundParameters.Add({ Export.Name, BlockAsset });
-					FParameterBindingReference NewReference(Export.Name, BlockAsset);
+					FParameterBindingReference NewReference(Export.Name, Export.Type, BlockAsset);
 					if(!Args.OnFilterParameter.IsBound() || Args.OnFilterParameter.Execute(NewReference) == EFilterParameterResult::Include)
 					{
 						FAnimNextParamType ParamType = UE::AnimNext::UncookedOnly::FUtils::GetParameterTypeFromName(Export.Name);
@@ -331,7 +329,7 @@ void SParameterPicker::RefreshEntries()
 						{
 							if (Args.bShowBoundParameters && EnumHasAnyFlags(Export.Flags, EAnimNextParameterFlags::Bound))
 							{
-								TSharedRef<FParameterPickerEntry> NewEntry = MakeShared<FParameterPickerEntry>(NewReference, ParamType);
+								TSharedRef<FParameterPickerEntry> NewEntry = MakeShared<FParameterPickerEntry>(NewReference);
 								Entries.Add(NewEntry);
 							}
 						}
@@ -351,15 +349,13 @@ void SParameterPicker::RefreshEntries()
 			{
 				if(!BoundParameters.Contains( { ExportEntry.Name, ExportEntry.ReferencingAsset } ))
 				{
-					FParameterBindingReference NewReference;
-					NewReference.Asset = ExportEntry.ReferencingAsset;
-					NewReference.Parameter = ExportEntry.Name;
+					FParameterBindingReference NewReference(ExportEntry.Name, ExportEntry.Type, ExportEntry.ReferencingAsset);
 
 					if(!Args.OnFilterParameter.IsBound() || Args.OnFilterParameter.Execute(NewReference) == EFilterParameterResult::Include)
 					{
 						if (!Args.OnFilterParameterType.IsBound() || Args.OnFilterParameterType.Execute(ExportEntry.Type) == EFilterParameterResult::Include)
 						{
-							TSharedRef<FParameterPickerEntry> NewEntry = MakeShared<FParameterPickerEntry>(NewReference, ExportEntry.Type);
+							TSharedRef<FParameterPickerEntry> NewEntry = MakeShared<FParameterPickerEntry>(NewReference);
 							Entries.Add(NewEntry);
 						}
 					}
@@ -372,12 +368,12 @@ void SParameterPicker::RefreshEntries()
 	{
 		FExternalParameterRegistry::ForEachParameter([this](FName InParameterName, const IParameterSourceFactory::FParameterInfo& InInfo)
 		{
-			FParameterBindingReference NewReference(InParameterName);
+			FParameterBindingReference NewReference(InParameterName, InInfo.Type);
 			if (!Args.OnFilterParameter.IsBound() || Args.OnFilterParameter.Execute(NewReference) == EFilterParameterResult::Include)
 			{
 				if (!Args.OnFilterParameterType.IsBound() || Args.OnFilterParameterType.Execute(InInfo.Type) == EFilterParameterResult::Include)
 				{
-					Entries.Add(MakeShared<FParameterPickerEntry>(NewReference, InInfo.Type));
+					Entries.Add(MakeShared<FParameterPickerEntry>(NewReference));
 				}
 			}
 		});
@@ -561,7 +557,7 @@ class SParameterPickerRow : public SMultiColumnTableRow<TSharedRef<FParameterPic
 		{
 			return
 				SNew(SHorizontalBox)
-				.Visibility(Entry->ParamType.IsValid() ? EVisibility::Visible : EVisibility::Hidden)
+				.Visibility(Entry->Binding.Type.IsValid() ? EVisibility::Visible : EVisibility::Hidden)
 				+ SHorizontalBox::Slot()
 				.HAlign(HAlign_Center)
 				.VAlign(VAlign_Center)
@@ -573,7 +569,7 @@ class SParameterPickerRow : public SMultiColumnTableRow<TSharedRef<FParameterPic
 					.ColorAndOpacity(Entry->PinColor)
 					.ToolTipText_Lambda([this]() -> FText
 					{
-						return FText::FromString(Entry->ParamType.ToString());	
+						return FText::FromString(Entry->Binding.Type.ToString());
 					})
 				];
 		}
@@ -647,7 +643,7 @@ void SParameterPicker::HandleSelectionChanged(TSharedPtr<FParameterPickerEntry>
 		TArray<TSharedRef<FParameterPickerEntry>> SelectedEntries;
 		EntriesList->GetSelectedItems(SelectedEntries);
 
-		if(SelectedEntries[0]->ParamType.IsValid() || Args.bAllowNone)
+		if(SelectedEntries[0]->Binding.Type.IsValid() || Args.bAllowNone)
 		{
 			Args.OnParameterPicked.ExecuteIfBound(SelectedEntries[0]->Binding);
 		}
@@ -668,7 +664,7 @@ void SParameterPicker::HandleGetParameterBindings(TArray<FParameterBindingRefere
 bool SParameterPicker::HandleIsSelectableOrNavigable(TSharedRef<FParameterPickerEntry> InEntry) const
 {
 	// Only allow selecting items that have valid parameters
-	return InEntry->ParamType.IsValid();
+	return InEntry->Binding.Type.IsValid();
 }
 
 }
diff --git a/Engine/Plugins/Experimental/AnimNext/Source/AnimNextUncookedOnly/Internal/Param/AnimNextParameterBlockProcessor.h b/Engine/Plugins/Experimental/AnimNext/Source/AnimNextUncookedOnly/Internal/Param/AnimNextParameterBlockProcessor.h
deleted file mode 100644
diff --git a/Engine/Plugins/Experimental/AnimNext/Source/AnimNextUncookedOnly/Internal/UncookedOnlyUtils.h b/Engine/Plugins/Experimental/AnimNext/Source/AnimNextUncookedOnly/Internal/UncookedOnlyUtils.h
--- a/Engine/Plugins/Experimental/AnimNext/Source/AnimNextUncookedOnly/Internal/UncookedOnlyUtils.h
+++ b/Engine/Plugins/Experimental/AnimNext/Source/AnimNextUncookedOnly/Internal/UncookedOnlyUtils.h
@@ -10,6 +10,7 @@
 
 #include "UncookedOnlyUtils.generated.h"
 
+struct FAnimNextParam;
 class UAnimNextSchedule;
 class UAnimNextGraph;
 class UAnimNextGraph_EditorData;
@@ -171,12 +172,22 @@ struct ANIMNEXTUNCOOKEDONLY_API FUtils
 			OutNodes.Append(GraphNodes);
 		}
 	}
-	
+
+	// Gets the parameters that are exported to the asset registry for an asset
 	static bool GetExportedParametersForAsset(const FAssetData& InAsset, FAnimNextParameterProviderAssetRegistryExports& OutExports);
+
+	// Gets all the parameters that are exported to the asset registry
 	static bool GetExportedParametersFromAssetRegistry(FAnimNextParameterProviderAssetRegistryExports& OutExports);
 
-	static void GetGraphParameters(const URigVMGraph* Graph,FAnimNextParameterProviderAssetRegistryExports& OutExports);
-	
+	// Gets the parameters that are used by a RigVM graph
+	static void GetGraphParameters(const URigVMGraph* Graph, FAnimNextParameterProviderAssetRegistryExports& OutExports);
+
+	// Gets the parameters that are exported to the asset registry by a schedule
+	static void GetScheduleParameters(const UAnimNextSchedule* InSchedule, FAnimNextParameterProviderAssetRegistryExports& OutExports);
+
+	// Gets the parameters that are exported to the asset registry by a blueprint
+	static void GetBlueprintParameters(const UBlueprint* InBlueprint, FAnimNextParameterProviderAssetRegistryExports& OutExports);
+
 	// Attempts to determine the type from a parameter name
 	// If the name cannot be found, the returned type will be invalid
 	// Note that this is expensive and can query the asset registry
@@ -184,6 +195,9 @@ struct ANIMNEXTUNCOOKEDONLY_API FUtils
 
 	// Compiles a schedule
 	static void CompileSchedule(UAnimNextSchedule* InSchedule);
+
+	// Sorts the incoming array of parameters, then generates a hash and returns it.
+	static uint64 SortAndHashParameters(TArray<FAnimNextParam>& InParameters);
 };
 
 }
\ No newline at end of file
diff --git a/Engine/Plugins/Experimental/AnimNext/Source/AnimNextUncookedOnly/Private/AnimNextUncookedOnlyModule.cpp b/Engine/Plugins/Experimental/AnimNext/Source/AnimNextUncookedOnly/Private/AnimNextUncookedOnlyModule.cpp
--- a/Engine/Plugins/Experimental/AnimNext/Source/AnimNextUncookedOnly/Private/AnimNextUncookedOnlyModule.cpp
+++ b/Engine/Plugins/Experimental/AnimNext/Source/AnimNextUncookedOnly/Private/AnimNextUncookedOnlyModule.cpp
@@ -1,8 +1,13 @@
 // Copyright Epic Games, Inc. All Rights Reserved.
 
 #include "CoreMinimal.h"
+#include "K2Node_CallFunction.h"
 #include "UncookedOnlyUtils.h"
+#include "Component/AnimNextComponentParameter.h"
+#include "Component/AnimNextComponent.h"
+#include "Engine/SCS_Node.h"
 #include "Modules/ModuleManager.h"
+#include "Param/AnimNextTag.h"
 #include "Scheduler/AnimNextSchedule.h"
 
 namespace UE::AnimNext::UncookedOnly
@@ -17,12 +22,42 @@ namespace UE::AnimNext::UncookedOnly
 			{
 				FUtils::CompileSchedule(InSchedule);
 			};
+
+			// TEMP: Bind the asset registry tags function for schedules
+			UAnimNextSchedule::GetAssetRegistryTagsFunction = [](const UAnimNextSchedule* InSchedule, TArray<UObject::FAssetRegistryTag>& OutTags)
+			{
+				FAnimNextParameterProviderAssetRegistryExports Exports;
+				FUtils::GetScheduleParameters(InSchedule, Exports);
+				
+				FString TagValue;
+				FAnimNextParameterProviderAssetRegistryExports::StaticStruct()->ExportText(TagValue, &Exports, nullptr, nullptr, PPF_None, nullptr);
+				OutTags.Add(UObject::FAssetRegistryTag(UE::AnimNext::ExportsAnimNextAssetRegistryTag, TagValue, UObject::FAssetRegistryTag::TT_Hidden));
+			};
+			
+			// Ensure that any BP components that we care about contribute to the parameter pool
+			OnGetExtraObjectTagsHandle = UObject::FAssetRegistryTag::OnGetExtraObjectTags.AddLambda([](const UObject* InObject, TArray<UObject::FAssetRegistryTag>& OutTags)
+			{
+				if(const UBlueprint* Blueprint = Cast<UBlueprint>(InObject))
+				{
+					FAnimNextParameterProviderAssetRegistryExports Exports;
+					FUtils::GetBlueprintParameters(Blueprint, Exports);
+
+					FString TagValue;
+					FAnimNextParameterProviderAssetRegistryExports::StaticStruct()->ExportText(TagValue, &Exports, nullptr, nullptr, PPF_None, nullptr);
+					OutTags.Add(UObject::FAssetRegistryTag(UE::AnimNext::ExportsAnimNextAssetRegistryTag, TagValue, UObject::FAssetRegistryTag::TT_Hidden));
+				}
+			});
 		}
 
 		virtual void ShutdownModule() override
 		{
+			UAnimNextSchedule::GetAssetRegistryTagsFunction = nullptr;
 			UAnimNextSchedule::CompileFunction = nullptr;
+
+			UObject::FAssetRegistryTag::OnGetExtraObjectTags.Remove(OnGetExtraObjectTagsHandle);
 		}
+
+		FDelegateHandle OnGetExtraObjectTagsHandle;
 	};
 }
 
diff --git a/Engine/Plugins/Experimental/AnimNext/Source/AnimNextUncookedOnly/Private/Graph/AnimGraph/AnimGraphNodeBinding_AnimNextParameters.cpp b/Engine/Plugins/Experimental/AnimNext/Source/AnimNextUncookedOnly/Private/Graph/AnimGraph/AnimGraphNodeBinding_AnimNextParameters.cpp
--- a/Engine/Plugins/Experimental/AnimNext/Source/AnimNextUncookedOnly/Private/Graph/AnimGraph/AnimGraphNodeBinding_AnimNextParameters.cpp
+++ b/Engine/Plugins/Experimental/AnimNext/Source/AnimNextUncookedOnly/Private/Graph/AnimGraph/AnimGraphNodeBinding_AnimNextParameters.cpp
@@ -548,6 +548,7 @@ TSharedRef<SWidget> UAnimGraphNodeBinding_AnimNextParameters::MakePropertyBindin
 				{
 					return (FParamUtils::GetCompatibility(ParamTypeHandle, InParamType.GetHandle()).IsCompatibleWithDataLoss()) ? EFilterParameterResult::Include : EFilterParameterResult::Exclude;
 				});
+				Args.NewParameterType = FParamTypeHandle::FromProperty(PropertyToBindTo).GetType();
 				Args.bMultiSelect = false;
 				Args.bShowBlocks = false;
 			
diff --git a/Engine/Plugins/Experimental/AnimNext/Source/AnimNextUncookedOnly/Private/UncookedOnlyUtils.cpp b/Engine/Plugins/Experimental/AnimNext/Source/AnimNextUncookedOnly/Private/UncookedOnlyUtils.cpp
--- a/Engine/Plugins/Experimental/AnimNext/Source/AnimNextUncookedOnly/Private/UncookedOnlyUtils.cpp
+++ b/Engine/Plugins/Experimental/AnimNext/Source/AnimNextUncookedOnly/Private/UncookedOnlyUtils.cpp
@@ -2,8 +2,11 @@
 
 #include "UncookedOnlyUtils.h"
 
+#include "Scheduler/Scheduler.h"
+#include "K2Node_CallFunction.h"
 #include "AssetRegistry/AssetRegistryModule.h"
 #include "AssetRegistry/IAssetRegistry.h"
+#include "Component/AnimNextComponentParameter.h"
 #include "Graph/AnimNextGraph.h"
 #include "Graph/AnimNextGraph_Controller.h"
 #include "Graph/AnimNextGraph_EditorData.h"
@@ -33,10 +36,17 @@
 #include "Graph/RigUnit_AnimNextBeginExecution.h"
 #include "Serialization/MemoryReader.h"
 #include "RigVMRuntimeDataRegistry.h"
+#include "Component/AnimNextComponent.h"
+#include "Engine/SCS_Node.h"
+#include "Engine/SimpleConstructionScript.h"
+#include "Param/AnimNextParam.h"
+#include "Param/AnimNextTag.h"
 #include "Param/RigVMDispatch_GetLayerParameter.h"
 #include "RigVMCompiler/RigVMCompiler.h"
 #include "RigVMCore/RigVM.h"
 #include "Param/ExternalParameterRegistry.h"
+#include "Param/ParamUtils.h"
+#include "Scheduler/AnimNextExternalTaskBinding.h"
 #include "Scheduler/AnimNextSchedule.h"
 #include "Scheduler/AnimNextSchedulePort.h"
 
@@ -530,6 +540,20 @@ void FUtils::Compile(UAnimNextGraph* InGraph)
 		InGraph->AllocateInstance(*GraphInstance);
 	}
 
+	// Gather and store the required parameters
+	InGraph->RequiredParameters.Empty();
+
+	FAssetData AssetData(InGraph);
+	FAnimNextParameterProviderAssetRegistryExports Exports;
+	GetExportedParametersForAsset(AssetData, Exports);
+
+	for(FAnimNextParameterAssetRegistryExportEntry& Entry : Exports.Parameters)
+	{
+		InGraph->RequiredParameters.Emplace(Entry.Name, Entry.Type);
+	}
+
+	InGraph->RequiredParametersHash = SortAndHashParameters(InGraph->RequiredParameters);
+
 #if WITH_EDITOR
 //	RefreshBreakpoints(EditorData);
 #endif
@@ -892,13 +916,11 @@ FAnimNextParamType FUtils::GetParamTypeFromPinType(const FEdGraphPinType& InPinT
 	{
 		ValueType = FAnimNextParamType::EValueType::Struct;
 		ValueTypeObject = Cast<UScriptStruct>(InPinType.PinSubCategoryObject.Get());
-		ensure(ValueTypeObject);
 	}
-	else if (InPinType.PinCategory == UEdGraphSchema_K2::PC_Object)
+	else if (InPinType.PinCategory == UEdGraphSchema_K2::PC_Object || InPinType.PinCategory == UEdGraphSchema_K2::AllObjectTypes)
 	{
 		ValueType = FAnimNextParamType::EValueType::Object;
 		ValueTypeObject = Cast<UClass>(InPinType.PinSubCategoryObject.Get());
-		ensure(ValueTypeObject);
 	}
 	else if (InPinType.PinCategory == UEdGraphSchema_K2::PC_SoftObject)
 	{
@@ -1235,14 +1257,14 @@ void FUtils::GetGraphParameters(const URigVMGraph* Graph, FAnimNextParameterProv
 			const FName GetLayerParameterNotation = GetLayerParameterFactory->GetTemplate()->GetNotation();
 
 			const FRigVMDispatchFactory* SetLayerParameterFactory = FRigVMRegistry::Get().FindOrAddDispatchFactory(FRigVMDispatch_SetLayerParameter::StaticStruct());
-			const FName SetLayerParameterNotation = GetLayerParameterFactory->GetTemplate()->GetNotation();
+			const FName SetLayerParameterNotation = SetLayerParameterFactory->GetTemplate()->GetNotation();
 
 			const bool bReadParameter = TemplateNode->GetNotation() == GetParameterNotation || TemplateNode->GetNotation() == GetLayerParameterNotation;
 			const bool bWriteParameter = TemplateNode->GetNotation() == SetLayerParameterNotation;
 			if (bReadParameter || bWriteParameter)
 			{
 				if (const URigVMPin* NamePin = TemplateNode->FindPin(FRigVMDispatch_GetParameter::ParameterName.ToString()))
-				{							
+				{
 					if (const URigVMPin* ValuePin = TemplateNode->FindPin(FRigVMDispatch_GetParameter::ValueName.ToString()))
 					{
 						FAnimNextParamType Type = FAnimNextParamType::FromRigVMTemplateArgument(ValuePin->GetTemplateArgumentType());
@@ -1264,6 +1286,121 @@ void FUtils::GetGraphParameters(const URigVMGraph* Graph, FAnimNextParameterProv
 	}
 }
 
+void FUtils::GetScheduleParameters(const UAnimNextSchedule* InSchedule, FAnimNextParameterProviderAssetRegistryExports& OutExports)
+{
+	for(UAnimNextScheduleEntry* Entry : InSchedule->Entries)
+	{
+		if (UAnimNextScheduleEntry_Port* PortEntry = Cast<UAnimNextScheduleEntry_Port>(Entry))
+		{
+			if(PortEntry->Port)
+			{
+				UAnimNextSchedulePort* CDO = PortEntry->Port->GetDefaultObject<UAnimNextSchedulePort>();
+				TConstArrayView<FAnimNextParam> RequiredParameters = CDO->GetRequiredParameters();
+				for(const FAnimNextParam& RequiredParameter : RequiredParameters)
+				{
+					if(!RequiredParameter.Name.IsNone() && RequiredParameter.Type.IsValid())
+					{
+						OutExports.Parameters.Emplace(RequiredParameter.Name, RequiredParameter.Type, EAnimNextParameterFlags::Read);
+					}
+				}
+			}
+		}
+		else if (UAnimNextScheduleEntry_AnimNextGraph* GraphEntry = Cast<UAnimNextScheduleEntry_AnimNextGraph>(Entry))
+		{
+			if(!GraphEntry->DynamicGraph.IsNone())
+			{
+				OutExports.Parameters.Emplace(GraphEntry->DynamicGraph, FAnimNextParamType::GetType<TObjectPtr<UAnimNextGraph>>(), EAnimNextParameterFlags::Read);
+			}
+			for(const FAnimNextParam& ParameterName : GraphEntry->RequiredParameters)
+			{
+				OutExports.Parameters.Emplace(ParameterName.Name, ParameterName.Type, EAnimNextParameterFlags::Read);
+			}
+		}
+		else if (UAnimNextScheduleEntry_ExternalTask* ExternalTaskEntry = Cast<UAnimNextScheduleEntry_ExternalTask>(Entry))
+		{
+			if(!ExternalTaskEntry->ExternalTask.IsNone())
+			{
+				OutExports.Parameters.Emplace(ExternalTaskEntry->ExternalTask, FAnimNextParamType::GetType<FAnimNextExternalTaskBinding>(), EAnimNextParameterFlags::Read);
+			}
+		}
+		else if (UAnimNextScheduleEntry_ParamScope* ParamScopeTaskEntry = Cast<UAnimNextScheduleEntry_ParamScope>(Entry))
+		{
+			if(!ParamScopeTaskEntry->Scope.IsNone())
+			{
+				OutExports.Parameters.Emplace(ParamScopeTaskEntry->Scope, FAnimNextParamType::GetType<FAnimNextScope>(), EAnimNextParameterFlags::Read);
+			}
+		}
+	}
+}
+
+void FUtils::GetBlueprintParameters(const UBlueprint* InBlueprint, FAnimNextParameterProviderAssetRegistryExports& OutExports)
+{
+	// Add 'static' params held on components
+	if(InBlueprint->SimpleConstructionScript)
+	{
+		for(USCS_Node* SCSNode : InBlueprint->SimpleConstructionScript->GetAllNodes())
+		{
+			if(UAnimNextComponent* Template = Cast<UAnimNextComponent>(SCSNode->ComponentTemplate))
+			{
+				for(UAnimNextComponentParameter* Parameter : Template->Parameters)
+				{
+					if(!Parameter->Scope.IsNone())
+					{
+						OutExports.Parameters.Emplace(Parameter->Scope, FAnimNextParamType::GetType<FAnimNextScope>(), EAnimNextParameterFlags::Read);
+					}
+
+					FName Name;
+					const FProperty* Property = nullptr;
+					Parameter->GetParamInfo(Name, Property);
+					if(!Name.IsNone())
+					{
+						FAnimNextParamType Type = FParamTypeHandle::FromProperty(Property).GetType();
+						check(Type.IsValid());
+						OutExports.Parameters.Emplace(Name, Type, EAnimNextParameterFlags::Write);
+					}
+				}
+			}
+		}
+	}
+
+	// Add any dynamic params held on graph nodes
+	const UFunction* SetParameterInScopeFunc = UAnimNextComponent::StaticClass()->FindFunctionByName(GET_FUNCTION_NAME_CHECKED(UAnimNextComponent, SetParameterInScope));
+
+	check(SetParameterInScopeFunc);
+	TArray<UEdGraph*> AllGraphs;
+	InBlueprint->GetAllGraphs(AllGraphs);
+	for(UEdGraph* EdGraph : AllGraphs)
+	{
+		TArray<UK2Node_CallFunction*> FunctionNodes;
+		EdGraph->GetNodesOfClass(FunctionNodes);
+		for(UK2Node_CallFunction* FunctionNode : FunctionNodes)
+		{
+			UFunction* Function = FunctionNode->FunctionReference.ResolveMember<UFunction>(FunctionNode->GetBlueprintClassFromNode());
+			if(Function == SetParameterInScopeFunc)
+			{
+				UEdGraphPin* ScopePin = FunctionNode->FindPinChecked(TEXT("Scope"));
+				FName ScopeName(*ScopePin->GetDefaultAsString());
+				if(!ScopeName.IsNone())
+				{
+					OutExports.Parameters.Emplace(ScopeName, FAnimNextParamType::GetType<FAnimNextScope>(), EAnimNextParameterFlags::Read);
+				}
+
+				UEdGraphPin* NamePin = FunctionNode->FindPinChecked(TEXT("Name"));
+				FName ParamName(*NamePin->GetDefaultAsString());
+				if(!ParamName.IsNone())
+				{
+					UEdGraphPin* ValuePin = FunctionNode->FindPinChecked(TEXT("Value"));
+					FAnimNextParamType Type = UncookedOnly::FUtils::GetParamTypeFromPinType(ValuePin->PinType);
+					if(Type.IsValid())
+					{
+						OutExports.Parameters.Emplace(ParamName, Type, EAnimNextParameterFlags::Write);
+					}
+				}
+			}
+		}
+	}
+}
+
 void FUtils::CompileSchedule(UAnimNextSchedule* InSchedule)
 {
 	using namespace UE::AnimNext;
@@ -1511,6 +1648,11 @@ void FUtils::CompileSchedule(UAnimNextSchedule* InSchedule)
 					GraphTask.EntryPoint = GraphEntry->EntryPoint;
 					GraphTask.Graph = GraphEntry->Graph;
 					GraphTask.DynamicGraph = GraphEntry->DynamicGraph;
+					if(GraphEntry->Graph == nullptr && GraphEntry->DynamicGraph != NAME_None)
+					{
+						GraphTask.SuppliedParameters = GraphEntry->RequiredParameters;
+						GraphTask.SuppliedParametersHash = SortAndHashParameters(GraphTask.SuppliedParameters);
+					}
 
 					for(int32 TermIndex = 0; TermIndex < GraphEntry->Terms.Num(); ++TermIndex)
 					{
@@ -1623,13 +1765,13 @@ void FUtils::CompileSchedule(UAnimNextSchedule* InSchedule)
 
 		// We need to find the task that is 'earliest' in the DAG for each external parameter, so we track that with this map
 		TMap<FName, FParameterTracker> TrackerMap;
-		auto TrackExternalParameters = [&TrackerMap](TConstArrayView<FName> InParameterNames, uint32 InDistance, UAnimNextScheduleEntry* InEntry, TArray<TObjectPtr<UAnimNextScheduleEntry>>* InContainer, int32 InArrayIndex, bool bInUpdateDependent)
+		auto TrackExternalParameters = [&TrackerMap](TConstArrayView<FAnimNextParam> InParameters, uint32 InDistance, UAnimNextScheduleEntry* InEntry, TArray<TObjectPtr<UAnimNextScheduleEntry>>* InContainer, int32 InArrayIndex, bool bInUpdateDependent)
 		{
 			bool bHasExternal = false;
-			for(FName ParameterName : InParameterNames)
+			for(const FAnimNextParam& Parameter : InParameters)
 			{
 				// Only add if the parameter name is 'external'
-				FName ParameterSourceName = FExternalParameterRegistry::FindSourceForParameter(ParameterName);
+				FName ParameterSourceName = FExternalParameterRegistry::FindSourceForParameter(Parameter.Name);
 				if(ParameterSourceName != NAME_None)
 				{
 					bHasExternal = true;
@@ -1646,15 +1788,15 @@ void FUtils::CompileSchedule(UAnimNextSchedule* InSchedule)
 					}
 
 					IParameterSourceFactory::FParameterInfo Info;
-					ensure(FExternalParameterRegistry::FindParameterInfo(ParameterName, Info));
+					ensure(FExternalParameterRegistry::FindParameterInfo(Parameter.Name, Info));
 
 					if(Info.bThreadSafe)
 					{
-						Tracker.ThreadSafeParameters.Add(ParameterName);
+						Tracker.ThreadSafeParameters.Add(Parameter.Name);
 					}
 					else
 					{
-						Tracker.NonThreadSafeParameters.Add(ParameterName);
+						Tracker.NonThreadSafeParameters.Add(Parameter.Name);
 					}
 				}
 			}
@@ -1675,11 +1817,11 @@ void FUtils::CompileSchedule(UAnimNextSchedule* InSchedule)
 						FAnimNextParameterProviderAssetRegistryExports Exports;
 						if(UncookedOnly::FUtils::GetExportedParametersForAsset(FAssetData(GraphEntry->Graph), Exports))
 						{
-							TArray<FName> RequiredParameters;
+							TArray<FAnimNextParam> RequiredParameters;
 							RequiredParameters.Reserve(Exports.Parameters.Num());
 							for(const FAnimNextParameterAssetRegistryExportEntry& ExportedParameter : Exports.Parameters)
 							{
-								RequiredParameters.Add(ExportedParameter.Name);
+								RequiredParameters.Emplace(ExportedParameter.Name, ExportedParameter.Type);
 							}
 							TrackExternalParameters(RequiredParameters, InDistance, Entry, &InEntries, ArrayIndex, true);
 						}
@@ -1689,21 +1831,26 @@ void FUtils::CompileSchedule(UAnimNextSchedule* InSchedule)
 
 					// All graphs require access to LOD and ref pose
 					// TODO: need to not require defaults here - use static graph params. In fact, these params should probably be defined at the schedule level.
-					TrackExternalParameters({ UAnimNextGraph::DefaultCurrentLODId.GetName(), UAnimNextGraph::DefaultReferencePoseId.GetName() }, InDistance, Entry, &InEntries, ArrayIndex, true);
+					const FAnimNextParam DefaultParams[] =
+					{
+						FAnimNextParam(UAnimNextGraph::DefaultCurrentLODId.GetName(), FAnimNextParamType::GetType<int32>()),
+						FAnimNextParam(UAnimNextGraph::DefaultReferencePoseId.GetName(), FAnimNextParamType::GetType<FAnimNextGraphReferencePose>())
+					};
+					TrackExternalParameters(DefaultParams, InDistance, Entry, &InEntries, ArrayIndex, true);
 				}
 				else if (UAnimNextScheduleEntry_Port* PortEntry = Cast<UAnimNextScheduleEntry_Port>(Entry))
 				{
 					if(PortEntry->Port)
 					{
 						UAnimNextSchedulePort* CDO = PortEntry->Port->GetDefaultObject<UAnimNextSchedulePort>();
-						TConstArrayView<FName> RequiredParameters = CDO->GetRequiredParameters();
+						TConstArrayView<FAnimNextParam> RequiredParameters = CDO->GetRequiredParameters();
 						TrackExternalParameters(RequiredParameters, InDistance, Entry, &InEntries, ArrayIndex, true);
 					}
 				}
 				else if (UAnimNextScheduleEntry_ExternalTask* ExternalTaskEntry = Cast<UAnimNextScheduleEntry_ExternalTask>(Entry))
 				{
 					// Note: external task params are not update dependent as this would cause external param updates to occur before tick functions
-					TrackExternalParameters({ ExternalTaskEntry->ExternalTask }, InDistance, Entry, &InEntries, ArrayIndex, false);
+					TrackExternalParameters({ FAnimNextParam(ExternalTaskEntry->ExternalTask, FAnimNextParamType::GetType<FAnimNextExternalTaskBinding>() ) }, InDistance, Entry, &InEntries, ArrayIndex, false);
 				}
 				else if (UAnimNextScheduleEntry_ParamScope* ParamScopeTaskEntry = Cast<UAnimNextScheduleEntry_ParamScope>(Entry))
 				{
@@ -1714,11 +1861,11 @@ void FUtils::CompileSchedule(UAnimNextSchedule* InSchedule)
 							FAnimNextParameterProviderAssetRegistryExports Exports;
 							if(UncookedOnly::FUtils::GetExportedParametersForAsset(FAssetData(ParameterBlock), Exports))
 							{
-								TArray<FName> RequiredParameters;
+								TArray<FAnimNextParam> RequiredParameters;
 								RequiredParameters.Reserve(Exports.Parameters.Num());
 								for(const FAnimNextParameterAssetRegistryExportEntry& ExportedParameter : Exports.Parameters)
 								{
-									RequiredParameters.Add(ExportedParameter.Name);
+									RequiredParameters.Emplace(ExportedParameter.Name, ExportedParameter.Type);
 								}
 								TrackExternalParameters(RequiredParameters, InDistance, Entry, &InEntries, ArrayIndex, true);
 							}
@@ -1820,6 +1967,26 @@ void FUtils::CompileSchedule(UAnimNextSchedule* InSchedule)
 
 		InSchedule->IntermediatesData.AddProperties(PropertyDescs);
 	}
+
+	FScheduler::OnScheduleCompiled(InSchedule);
+}
+
+uint64 FUtils::SortAndHashParameters(TArray<FAnimNextParam>& InParameters)
+{
+	Algo::Sort(InParameters, [](const FAnimNextParam& InLHS, const FAnimNextParam& InRHS)
+	{
+		return InLHS.Name.LexicalLess(InRHS.Name);
+	});
+
+	uint64 Hash = 0;
+	for(const FAnimNextParam& Parameter : InParameters)
+	{
+		FString ExportedString;
+		FAnimNextParam::StaticStruct()->ExportText(ExportedString, &Parameter, nullptr, nullptr, PPF_None, nullptr);
+		Hash = CityHash64WithSeed(reinterpret_cast<const char*>(*ExportedString), ExportedString.Len() * sizeof(TCHAR), Hash);
+	}
+
+	return Hash;
 }
 
 }
