Add AnimNext proxy graph decorator that runs a child graph within a parent graph
Reworked graph traversal to be more flexible and allow child graphs

diff --git a/Engine/Plugins/Experimental/AnimNext/Source/AnimNext/Internal/Graph/AnimGraph/AnimNode_AnimNextGraph.h b/Engine/Plugins/Experimental/AnimNext/Source/AnimNext/Internal/Graph/AnimGraph/AnimNode_AnimNextGraph.h
--- a/Engine/Plugins/Experimental/AnimNext/Source/AnimNext/Internal/Graph/AnimGraph/AnimNode_AnimNextGraph.h
+++ b/Engine/Plugins/Experimental/AnimNext/Source/AnimNext/Internal/Graph/AnimGraph/AnimNode_AnimNextGraph.h
@@ -4,6 +4,7 @@
 
 #include "Animation/AnimNode_CustomProperty.h"
 #include "Graph/AnimNextGraph.h"
+#include "Graph/AnimNextGraphInstancePtr.h"
 #include "DecoratorBase/DecoratorPtr.h"
 #include "Context.h"
 #include "AnimNode_AnimNextGraph.generated.h"
@@ -45,7 +46,7 @@ private:
 
 	// Our graph instance, we own it
 	UPROPERTY()
-	FAnimNextGraphInstance GraphInstance;
+	FAnimNextGraphInstancePtr GraphInstance;
 
 	/*
 	 * Max LOD that this node is allowed to run
@@ -71,6 +72,7 @@ struct TStructOpsTypeTraits<FAnimNode_AnimNextGraph> : public TStructOpsTypeTrai
 {
 	enum
 	{
+		WithCopy = false,
 		WithPostSerialize = true,
 	};
 };
\ No newline at end of file
diff --git a/Engine/Plugins/Experimental/AnimNext/Source/AnimNext/Internal/Graph/AnimNextGraph.h b/Engine/Plugins/Experimental/AnimNext/Source/AnimNext/Internal/Graph/AnimNextGraph.h
--- a/Engine/Plugins/Experimental/AnimNext/Source/AnimNext/Internal/Graph/AnimNextGraph.h
+++ b/Engine/Plugins/Experimental/AnimNext/Source/AnimNext/Internal/Graph/AnimNextGraph.h
@@ -20,6 +20,8 @@ class UAnimNextGraph;
 class UAnimGraphNode_AnimNextGraph;
 struct FAnimNode_AnimNextGraph;
 struct FRigUnit_AnimNextGraphEvaluator;
+struct FAnimNextGraphInstancePtr;
+struct FAnimNextGraphInstance;
 class UAnimNextSchedule;
 struct FAnimNextScheduleGraphTask;
 struct FAnimNextParam;
@@ -47,106 +49,6 @@ namespace UE::AnimNext::Graph
 	extern ANIMNEXT_API const FName ResultName;
 }
 
-using GraphInstanceComponentMapType = TMap<FName, TSharedPtr<UE::AnimNext::FGraphInstanceComponent>>;
-
-// Represents an instance of an AnimNext graph
-// This struct uses UE reflection because we wish for the GC to keep the graph
-// alive while we own a reference to it. It is not intended to be serialized on disk with a live instance.
-USTRUCT()
-struct ANIMNEXT_API FAnimNextGraphInstance
-{
-	GENERATED_BODY()
-
-	// Creates an empty graph instance that doesn't reference anything
-	FAnimNextGraphInstance() = default;
-
-#if WITH_EDITORONLY_DATA
-	// In editor, we need custom copy/move semantics to update the live instance tracking
-	FAnimNextGraphInstance(const FAnimNextGraphInstance& Other);
-	FAnimNextGraphInstance(FAnimNextGraphInstance&& Other);
-	FAnimNextGraphInstance& operator=(const FAnimNextGraphInstance& Other);
-	FAnimNextGraphInstance& operator=(FAnimNextGraphInstance&& Other);
-#else
-	FAnimNextGraphInstance(const FAnimNextGraphInstance&) = default;
-	FAnimNextGraphInstance& operator=(const FAnimNextGraphInstance&) = default;
-	FAnimNextGraphInstance(FAnimNextGraphInstance&&) = default;
-	FAnimNextGraphInstance& operator=(FAnimNextGraphInstance&&) = default;
-#endif
-
-	// If the graph instance is allocated, we release it during destruction
-	~FAnimNextGraphInstance();
-
-	// Releases the graph instance and frees all corresponding memory
-	void Release();
-
-	// Returns true if we have a live graph instance, false otherwise
-	bool IsValid() const;
-
-	// Returns the graph used by this instance or nullptr if the instance is invalid
-	const UAnimNextGraph* GetGraph() const;
-
-	// Returns a weak handle to the root decorator instance
-	UE::AnimNext::FWeakDecoratorPtr GetGraphRootPtr() const;
-
-	// Check to see if this instance data matches the provided graph
-	bool UsesGraph(const UAnimNextGraph* InGraph) const;
-
-	// Adds strong/hard object references during GC
-	void AddStructReferencedObjects(class FReferenceCollector& Collector);
-
-	// Returns a typed graph instance component, creating it lazily the first time it is queried
-	template<class ComponentType>
-	ComponentType& GetComponent();
-
-	// Returns a typed graph instance component pointer if found or nullptr otherwise
-	template<class ComponentType>
-	ComponentType* TryGetComponent();
-
-	// Returns a typed graph instance component pointer if found or nullptr otherwise
-	template<class ComponentType>
-	const ComponentType* TryGetComponent() const;
-
-	// Returns const iterators to the graph instance component container
-	GraphInstanceComponentMapType::TConstIterator GetComponentIterator() const;
-
-private:
-	// Returns a pointer to the specified component, or nullptr if not found
-	UE::AnimNext::FGraphInstanceComponent* TryGetComponent(int32 ComponentNameHash, FName ComponentName) const;
-
-	// Adds the specified component and returns a reference to it
-	UE::AnimNext::FGraphInstanceComponent& AddComponent(int32 ComponentNameHash, FName ComponentName, TSharedPtr<UE::AnimNext::FGraphInstanceComponent>&& Component);
-
-	// Executes a latent RigVM pin and writes the result into the destination pointer
-	void ExecuteLatentPin(int32 LatentPinIndex, void* DestinationPtr);
-
-	// Hard reference to the graph used to create this instance to ensure we can release it safely
-	UPROPERTY()
-	TObjectPtr<const UAnimNextGraph> Graph;
-
-	// Hard reference to the graph instance data, we own it
-	UE::AnimNext::FDecoratorPtr GraphInstancePtr;
-
-	// Extended execute context instance for this graph instance, we own it
-	UPROPERTY()
-	FRigVMExtendedExecuteContext ExtendedExecuteContext;
-
-	// Graph instance components that persist from update to update
-	GraphInstanceComponentMapType Components;
-
-	friend UAnimNextGraph;					// The graph is the one that allocates instances
-	friend FRigUnit_AnimNextGraphEvaluator;	// We evaluate the instance
-	friend UE::AnimNext::FExecutionContext;
-};
-
-template<>
-struct TStructOpsTypeTraits<FAnimNextGraphInstance> : public TStructOpsTypeTraitsBase2<FAnimNextGraphInstance>
-{
-	enum
-	{
-		WithAddStructReferencedObjects = true,
-	};
-};
-
 // A user-created graph of logic used to supply data
 UCLASS(BlueprintType)
 class ANIMNEXT_API UAnimNextGraph :  public URigVMHost, public IAnimNextScheduleTermInterface
@@ -167,7 +69,10 @@ public:
 	virtual TConstArrayView<UE::AnimNext::FScheduleTerm> GetTerms() const override;
 
 	// Allocates an instance of the graph
-	void AllocateInstance(FAnimNextGraphInstance& Instance) const;
+	void AllocateInstance(FAnimNextGraphInstancePtr& Instance) const;
+
+	// Allocates an instance of the graph with the specified parent graph instance
+	void AllocateInstance(FAnimNextGraphInstance& ParentGraphInstance, FAnimNextGraphInstancePtr& Instance) const;
 
 	// Get the parameter to use to access the reference pose
 	UE::AnimNext::FParamId GetReferencePoseParam() const { return ReferencePoseId; }
@@ -185,9 +90,16 @@ protected:
 	// Loads the graph data from the provided archive buffer and returns true on success, false otherwise
 	bool LoadFromArchiveBuffer(const TArray<uint8>& SharedDataArchiveBuffer);
 
+	// Allocates an instance of the graph with an optional parent graph instance
+	void AllocateInstanceImpl(FAnimNextGraphInstance* ParentGraphInstance, FAnimNextGraphInstancePtr& Instance) const;
+
 #if WITH_EDITORONLY_DATA
-	// Releases all live graph instances and returns a list of the instances that were released
-	TArray<FAnimNextGraphInstance*> ReleaseAllInstances();
+	// During graph compilation, if we have existing graph instances, we freeze them by releasing their memory before thawing them
+	// Freezing is a partial release of resources that retains the necessary information to re-create things safely
+	void FreezeGraphInstances();
+
+	// During graph compilation, once compilation is done we thaw existing graph instances to reallocate their memory
+	void ThawGraphInstances();
 #endif
 
 	friend class UAnimNextGraphFactory;
@@ -195,7 +107,8 @@ protected:
 	friend struct UE::AnimNext::UncookedOnly::FUtils;
 	friend class UE::AnimNext::Editor::FGraphEditor;
 	friend struct UE::AnimNext::FTestUtils;
-	friend struct FAnimNextGraphInstance;
+	friend FAnimNextGraphInstancePtr;
+	friend FAnimNextGraphInstance;
 	friend class UAnimGraphNode_AnimNextGraph;
 	friend UE::AnimNext::FExecutionContext;
 	friend class UAnimNextSchedule;
@@ -269,37 +182,3 @@ protected:
 	TArray<uint8> SharedDataArchiveBuffer;
 #endif
 };
-
-//////////////////////////////////////////////////////////////////////////
-
-template<class ComponentType>
-ComponentType& FAnimNextGraphInstance::GetComponent()
-{
-	const FName ComponentName = ComponentType::StaticComponentName();
-	const int32 ComponentNameHash = GetTypeHash(ComponentName);
-
-	if (UE::AnimNext::FGraphInstanceComponent* Component = TryGetComponent(ComponentNameHash, ComponentName))
-	{
-		return *static_cast<ComponentType*>(Component);
-	}
-
-	return static_cast<ComponentType&>(AddComponent(ComponentNameHash, ComponentName, MakeShared<ComponentType>()));
-}
-
-template<class ComponentType>
-ComponentType* FAnimNextGraphInstance::TryGetComponent()
-{
-	const FName ComponentName = ComponentType::StaticComponentName();
-	const int32 ComponentNameHash = GetTypeHash(ComponentName);
-
-	return static_cast<ComponentType*>(TryGetComponent(ComponentNameHash, ComponentName));
-}
-
-template<class ComponentType>
-const ComponentType* FAnimNextGraphInstance::TryGetComponent() const
-{
-	const FName ComponentName = ComponentType::StaticComponentName();
-	const int32 ComponentNameHash = GetTypeHash(ComponentName);
-
-	return static_cast<ComponentType*>(TryGetComponent(ComponentNameHash, ComponentName));
-}
diff --git a/Engine/Plugins/Experimental/AnimNext/Source/AnimNext/Private/DecoratorBase/ExecutionContext.cpp b/Engine/Plugins/Experimental/AnimNext/Source/AnimNext/Private/DecoratorBase/ExecutionContext.cpp
--- a/Engine/Plugins/Experimental/AnimNext/Source/AnimNext/Private/DecoratorBase/ExecutionContext.cpp
+++ b/Engine/Plugins/Experimental/AnimNext/Source/AnimNext/Private/DecoratorBase/ExecutionContext.cpp
@@ -9,6 +9,8 @@
 #include "DecoratorBase/NodeInstance.h"
 #include "DecoratorBase/NodeTemplate.h"
 #include "DecoratorBase/NodeTemplateRegistry.h"
+#include "Graph/AnimNextGraph.h"
+#include "Graph/AnimNextGraphInstance.h"
 #include "RigVMCore/RigVMExecuteContext.h"
 
 namespace UE::AnimNext
@@ -20,12 +22,26 @@ namespace UE::AnimNext
 	{
 	}
 
+	FExecutionContext::FExecutionContext(FAnimNextGraphInstancePtr& InGraphInstance)
+		: FExecutionContext()
+	{
+		BindTo(InGraphInstance);
+	}
+
 	FExecutionContext::FExecutionContext(FAnimNextGraphInstance& InGraphInstance)
 		: FExecutionContext()
 	{
 		BindTo(InGraphInstance);
 	}
 
+	void FExecutionContext::BindTo(FAnimNextGraphInstancePtr& InGraphInstance)
+	{
+		if (FAnimNextGraphInstance* Impl = InGraphInstance.GetImpl())
+		{
+			BindTo(*Impl);
+		}
+	}
+
 	void FExecutionContext::BindTo(FAnimNextGraphInstance& InGraphInstance)
 	{
 		if (GraphInstance == &InGraphInstance)
@@ -33,8 +49,11 @@ namespace UE::AnimNext
 			return;	// Already bound to this graph instance, nothing to do
 		}
 
-		GraphInstance = &InGraphInstance;
-		GraphSharedData = InGraphInstance.GetGraph()->SharedDataBuffer;
+		if (const UAnimNextGraph* Graph = InGraphInstance.GetGraph())
+		{
+			GraphInstance = &InGraphInstance;
+			GraphSharedData = Graph->SharedDataBuffer;
+		}
 	}
 
 	void FExecutionContext::BindTo(const FWeakDecoratorPtr& DecoratorPtr)
@@ -50,6 +69,11 @@ namespace UE::AnimNext
 		return GraphInstance != nullptr;
 	}
 
+	bool FExecutionContext::IsBoundTo(const FAnimNextGraphInstancePtr& InGraphInstance) const
+	{
+		return GraphInstance == InGraphInstance.GetImpl();
+	}
+
 	bool FExecutionContext::IsBoundTo(const FAnimNextGraphInstance& InGraphInstance) const
 	{
 		return GraphInstance == &InGraphInstance;
@@ -365,6 +389,21 @@ namespace UE::AnimNext
 		GraphInstance->ExecuteLatentPin(LatentPropertyHandle.GetLatentPropertyIndex(), DestinationPtr);
 	}
 
+	FGraphInstanceComponent* FExecutionContext::TryGetComponent(int32 ComponentNameHash, FName ComponentName) const
+	{
+		return GraphInstance->TryGetComponent(ComponentNameHash, ComponentName);
+	}
+
+	FGraphInstanceComponent& FExecutionContext::AddComponent(int32 ComponentNameHash, FName ComponentName, TSharedPtr<FGraphInstanceComponent>&& Component) const
+	{
+		return GraphInstance->AddComponent(ComponentNameHash, ComponentName, MoveTemp(Component));
+	}
+
+	GraphInstanceComponentMapType::TConstIterator FExecutionContext::GetComponentIterator() const
+	{
+		return GraphInstance->GetComponentIterator();
+	}
+
 	const FNodeDescription& FExecutionContext::GetNodeDescription(FNodeHandle NodeHandle) const
 	{
 		check(NodeHandle.IsValid());
diff --git a/Engine/Plugins/Experimental/AnimNext/Source/AnimNext/Private/DecoratorInterfaces/IEvaluate.cpp b/Engine/Plugins/Experimental/AnimNext/Source/AnimNext/Private/DecoratorInterfaces/IEvaluate.cpp
--- a/Engine/Plugins/Experimental/AnimNext/Source/AnimNext/Private/DecoratorInterfaces/IEvaluate.cpp
+++ b/Engine/Plugins/Experimental/AnimNext/Source/AnimNext/Private/DecoratorInterfaces/IEvaluate.cpp
@@ -71,7 +71,7 @@ namespace UE::AnimNext
 		}
 	};
 
-	FEvaluationProgram EvaluateGraph(FAnimNextGraphInstance& GraphInstance)
+	FEvaluationProgram EvaluateGraph(FAnimNextGraphInstancePtr& GraphInstance)
 	{
 		return EvaluateGraph(GraphInstance.GetGraphRootPtr());
 	}
diff --git a/Engine/Plugins/Experimental/AnimNext/Source/AnimNext/Private/DecoratorInterfaces/IUpdate.cpp b/Engine/Plugins/Experimental/AnimNext/Source/AnimNext/Private/DecoratorInterfaces/IUpdate.cpp
--- a/Engine/Plugins/Experimental/AnimNext/Source/AnimNext/Private/DecoratorInterfaces/IUpdate.cpp
+++ b/Engine/Plugins/Experimental/AnimNext/Source/AnimNext/Private/DecoratorInterfaces/IUpdate.cpp
@@ -3,6 +3,7 @@
 #include "DecoratorInterfaces/IUpdate.h"
 
 #include "DecoratorInterfaces/IHierarchy.h"
+#include "Graph/GraphInstanceComponent.h"
 #include "AnimNextStats.h"
 
 DEFINE_STAT(STAT_AnimNext_UpdateGraph);
@@ -207,7 +208,7 @@ namespace UE::AnimNext
 	// perform as much useful work as possible while waiting for memory, hiding its slow latency by fully
 	// leveraging out-of-order CPU execution.
 
-	void UpdateGraph(FAnimNextGraphInstance& GraphInstance, float DeltaTime)
+	void UpdateGraph(FAnimNextGraphInstancePtr& GraphInstance, float DeltaTime)
 	{
 		SCOPE_CYCLE_COUNTER(STAT_AnimNext_UpdateGraph);
 		
@@ -216,6 +217,11 @@ namespace UE::AnimNext
 			return;	// Nothing to update
 		}
 
+		if (!ensure(GraphInstance.IsRoot()))
+		{
+			return;	// We can only update starting at the root
+		}
+
 		FMemStack& MemStack = FMemStack::Get();
 		FMemMark Mark(MemStack);
 
diff --git a/Engine/Plugins/Experimental/AnimNext/Source/AnimNext/Private/Decorators/BlendByBool.cpp b/Engine/Plugins/Experimental/AnimNext/Source/AnimNext/Private/Decorators/BlendByBool.cpp
--- a/Engine/Plugins/Experimental/AnimNext/Source/AnimNext/Private/Decorators/BlendByBool.cpp
+++ b/Engine/Plugins/Experimental/AnimNext/Source/AnimNext/Private/Decorators/BlendByBool.cpp
@@ -2,6 +2,7 @@
 
 #include "Decorators/BlendByBool.h"
 
+#include "Animation/AnimTypes.h"
 #include "DecoratorBase/ExecutionContext.h"
 #include "EvaluationVM/Tasks/BlendKeyframes.h"
 
@@ -11,7 +12,6 @@ namespace UE::AnimNext
 
 	DEFINE_ANIM_DECORATOR_BEGIN(FBlendByBoolDecorator)
 		DEFINE_ANIM_DECORATOR_IMPLEMENTS_INTERFACE(IDiscreteBlend)
-		DEFINE_ANIM_DECORATOR_IMPLEMENTS_INTERFACE(IEvaluate)
 		DEFINE_ANIM_DECORATOR_IMPLEMENTS_INTERFACE(IHierarchy)
 		DEFINE_ANIM_DECORATOR_IMPLEMENTS_INTERFACE(IUpdate)
 	DEFINE_ANIM_DECORATOR_END(FBlendByBoolDecorator)
@@ -19,26 +19,6 @@ namespace UE::AnimNext
 	static constexpr int32 TRUE_CHILD_INDEX = 0;
 	static constexpr int32 FALSE_CHILD_INDEX = 1;
 
-	void FBlendByBoolDecorator::PostEvaluate(FEvaluateTraversalContext& Context, const TDecoratorBinding<IEvaluate>& Binding) const
-	{
-		const FInstanceData* InstanceData = Binding.GetInstanceData<FInstanceData>();
-
-		if (InstanceData->TrueChild.IsValid() && InstanceData->FalseChild.IsValid())
-		{
-			// We have two children, interpolate them
-
-			TDecoratorBinding<IDiscreteBlend> DiscreteBlendDecorator;
-			Context.GetInterface(Binding, DiscreteBlendDecorator);
-
-			const float BlendWeight = DiscreteBlendDecorator.GetBlendWeight(Context, FALSE_CHILD_INDEX);
-			Context.AppendTask(FAnimNextBlendTwoKeyframesTask::Make(BlendWeight));
-		}
-		else
-		{
-			// We have only one child that is active, do nothing
-		}
-	}
-
 	void FBlendByBoolDecorator::PreUpdate(FUpdateTraversalContext& Context, const TDecoratorBinding<IUpdate>& Binding, const FDecoratorUpdateState& DecoratorState) const
 	{
 		const FSharedData* SharedData = Binding.GetSharedData<FSharedData>();
diff --git a/Engine/Plugins/Experimental/AnimNext/Source/AnimNext/Private/Decorators/BlendByBool.h b/Engine/Plugins/Experimental/AnimNext/Source/AnimNext/Private/Decorators/BlendByBool.h
--- a/Engine/Plugins/Experimental/AnimNext/Source/AnimNext/Private/Decorators/BlendByBool.h
+++ b/Engine/Plugins/Experimental/AnimNext/Source/AnimNext/Private/Decorators/BlendByBool.h
@@ -6,7 +6,6 @@
 
 #include "DecoratorBase/Decorator.h"
 #include "DecoratorInterfaces/IDiscreteBlend.h"
-#include "DecoratorInterfaces/IEvaluate.h"
 #include "DecoratorInterfaces/IHierarchy.h"
 #include "DecoratorInterfaces/IUpdate.h"
 
@@ -44,7 +43,7 @@ namespace UE::AnimNext
 	 * 
 	 * A decorator that can blend two discrete inputs through a boolean.
 	 */
-	struct FBlendByBoolDecorator : FBaseDecorator, IEvaluate, IUpdate, IHierarchy, IDiscreteBlend
+	struct FBlendByBoolDecorator : FBaseDecorator, IUpdate, IHierarchy, IDiscreteBlend
 	{
 		DECLARE_ANIM_DECORATOR(FBlendByBoolDecorator, 0xc6d8c9ea, FBaseDecorator)
 
@@ -58,9 +57,6 @@ namespace UE::AnimNext
 			int32 PreviousChildIndex = INDEX_NONE;
 		};
 
-		// IEvaluate impl
-		virtual void PostEvaluate(FEvaluateTraversalContext& Context, const TDecoratorBinding<IEvaluate>& Binding) const override;
-
 		// IUpdate impl
 		virtual void PreUpdate(FUpdateTraversalContext& Context, const TDecoratorBinding<IUpdate>& Binding, const FDecoratorUpdateState& DecoratorState) const override;
 		virtual void QueueChildrenForTraversal(FUpdateTraversalContext& Context, const TDecoratorBinding<IUpdate>& Binding, const FDecoratorUpdateState& DecoratorState, FUpdateTraversalQueue& TraversalQueue) const override;
diff --git a/Engine/Plugins/Experimental/AnimNext/Source/AnimNext/Private/Decorators/BlendSmoother.cpp b/Engine/Plugins/Experimental/AnimNext/Source/AnimNext/Private/Decorators/BlendSmoother.cpp
--- a/Engine/Plugins/Experimental/AnimNext/Source/AnimNext/Private/Decorators/BlendSmoother.cpp
+++ b/Engine/Plugins/Experimental/AnimNext/Source/AnimNext/Private/Decorators/BlendSmoother.cpp
@@ -6,17 +6,76 @@
 #include "DecoratorBase/ExecutionContext.h"
 #include "DecoratorInterfaces/IHierarchy.h"
 #include "EvaluationVM/Tasks/BlendKeyframes.h"
+#include "EvaluationVM/Tasks/NormalizeRotations.h"
 
 namespace UE::AnimNext
 {
 	AUTO_REGISTER_ANIM_DECORATOR(FBlendSmootherDecorator)
 
 	DEFINE_ANIM_DECORATOR_BEGIN(FBlendSmootherDecorator)
-		DEFINE_ANIM_DECORATOR_IMPLEMENTS_INTERFACE(IUpdate)
 		DEFINE_ANIM_DECORATOR_IMPLEMENTS_INTERFACE(IDiscreteBlend)
+		DEFINE_ANIM_DECORATOR_IMPLEMENTS_INTERFACE(IEvaluate)
 		DEFINE_ANIM_DECORATOR_IMPLEMENTS_INTERFACE(ISmoothBlend)
+		DEFINE_ANIM_DECORATOR_IMPLEMENTS_INTERFACE(IUpdate)
 	DEFINE_ANIM_DECORATOR_END(FBlendSmootherDecorator)
 
+	void FBlendSmootherDecorator::PostEvaluate(FEvaluateTraversalContext& Context, const TDecoratorBinding<IEvaluate>& Binding) const
+	{
+		const FInstanceData* InstanceData = Binding.GetInstanceData<FInstanceData>();
+
+		// We override the default behavior since we need to blend over time
+
+		int32 NumBlending = 0;
+		for (const FBlendData& ChildBlendData : InstanceData->PerChildBlendData)
+		{
+			NumBlending += ChildBlendData.bIsBlending ? 1 : 0;
+		}
+
+		if (NumBlending < 2)
+		{
+			return;	// If we don't have at least 2 children blending, there is nothing to do
+		}
+
+		// Children are visited depth first, in the order returned
+		// As such, when we evaluate the task program, the keyframe of the last child will be
+		// on top of the keyframe stack
+		// We thus process children in reverse order
+
+		// The last child override the top keyframe and scales it
+		int32 ChildIndex = InstanceData->PerChildBlendData.Num() - 1;
+		for (; ChildIndex >= 0; --ChildIndex)
+		{
+			const FBlendData& ChildBlendData = InstanceData->PerChildBlendData[ChildIndex];
+			if (!ChildBlendData.bIsBlending)
+			{
+				continue;	// Skip this inactive child
+			}
+
+			// This decorator controls the blend weight and owns it
+			Context.AppendTask(FAnimNextBlendOverwriteKeyframeWithScaleTask::Make(ChildBlendData.Weight));
+
+			// We found the last child to blend
+			break;
+		}
+
+		// Other children accumulate with scale
+		ChildIndex--;
+		for (; ChildIndex >= 0; --ChildIndex)
+		{
+			const FBlendData& ChildBlendData = InstanceData->PerChildBlendData[ChildIndex];
+			if (!ChildBlendData.bIsBlending)
+			{
+				continue;	// Skip this inactive child
+			}
+
+			// This decorator controls the blend weight and owns it
+			Context.AppendTask(FAnimNextBlendAddKeyframeWithScaleTask::Make(ChildBlendData.Weight));
+		}
+
+		// Once we are done, we normalize rotations
+		Context.AppendTask(FAnimNextNormalizeKeyframeRotationsTask());
+	}
+
 	void FBlendSmootherDecorator::PreUpdate(FUpdateTraversalContext& Context, const TDecoratorBinding<IUpdate>& Binding, const FDecoratorUpdateState& DecoratorState) const
 	{
 		const FSharedData* SharedData = Binding.GetSharedData<FSharedData>();
@@ -113,6 +172,18 @@ namespace UE::AnimNext
 		const FSharedData* SharedData = Binding.GetSharedData<FSharedData>();
 		FInstanceData* InstanceData = Binding.GetInstanceData<FInstanceData>();
 
+		const int32 NumChildren = InstanceData->PerChildBlendData.Num();
+		if (NewChildIndex >= NumChildren)
+		{
+			// We have a new child
+			check(NewChildIndex == NumChildren);
+
+			FBlendData& ChildBlendData = InstanceData->PerChildBlendData.AddDefaulted_GetRef();
+
+			ChildBlendData.Blend.SetBlendOption(SharedData->BlendType);
+			ChildBlendData.Blend.SetCustomCurve(SharedData->CustomBlendCurve);
+		}
+
 		// scale by the weight difference since we want consistency:
 		// - if you're moving from 0 to full weight 1, it will use the normal blend time
 		// - if you're moving from 0.5 to full weight 1, it will get there in half the time
@@ -160,25 +231,34 @@ namespace UE::AnimNext
 	float FBlendSmootherDecorator::GetBlendTime(const FExecutionContext& Context, const TDecoratorBinding<ISmoothBlend>& Binding, int32 ChildIndex) const
 	{
 		const FSharedData* SharedData = Binding.GetSharedData<FSharedData>();
-		return SharedData->BlendTimes.IsValidIndex(ChildIndex) ? SharedData->BlendTimes[ChildIndex] : 0.0f;
+
+		if (SharedData->BlendTimes.IsValidIndex(ChildIndex))
+		{
+			return SharedData->BlendTimes[ChildIndex];
+		}
+		else if (!SharedData->BlendTimes.IsEmpty())
+		{
+			// If we index outside the array of values we have, use the last value
+			// Allows a user to specify a single blend time to be used with all children
+			return SharedData->BlendTimes.Last();
+		}
+		else
+		{
+			// No blend time has been specified, we snap
+			return 0.0f;
+		}
 	}
 
 	void FBlendSmootherDecorator::InitializeInstanceData(const FExecutionContext& Context, const FDecoratorBinding& Binding, const FSharedData* SharedData, FInstanceData* InstanceData)
 	{
 		check(InstanceData->PerChildBlendData.IsEmpty());
 
-		uint32 NumExpectedChildren = SharedData->BlendTimes.Num();
-
-#if DO_CHECK
 		TDecoratorBinding<IHierarchy> HierarchyDecorator;
 		Context.GetInterface(Binding, HierarchyDecorator);
 
-		const uint32 NumActualChildren = HierarchyDecorator.GetNumChildren(Context);
-		ensureMsgf(NumActualChildren != 0, TEXT("BlendSmootherDecorator has %u blend times for %u children"), NumExpectedChildren, NumActualChildren);
-		NumExpectedChildren = NumActualChildren;
-#endif
+		const uint32 NumChildren = HierarchyDecorator.GetNumChildren(Context);
 
-		InstanceData->PerChildBlendData.SetNum(NumExpectedChildren);
+		InstanceData->PerChildBlendData.SetNum(NumChildren);
 
 		for (FBlendData& ChildBlendData : InstanceData->PerChildBlendData)
 		{
diff --git a/Engine/Plugins/Experimental/AnimNext/Source/AnimNext/Private/Decorators/BlendSmoother.h b/Engine/Plugins/Experimental/AnimNext/Source/AnimNext/Private/Decorators/BlendSmoother.h
--- a/Engine/Plugins/Experimental/AnimNext/Source/AnimNext/Private/Decorators/BlendSmoother.h
+++ b/Engine/Plugins/Experimental/AnimNext/Source/AnimNext/Private/Decorators/BlendSmoother.h
@@ -7,6 +7,7 @@
 #include "AlphaBlend.h"
 #include "DecoratorBase/Decorator.h"
 #include "DecoratorInterfaces/IDiscreteBlend.h"
+#include "DecoratorInterfaces/IEvaluate.h"
 #include "DecoratorInterfaces/ISmoothBlend.h"
 #include "DecoratorInterfaces/IUpdate.h"
 
@@ -39,7 +40,7 @@ namespace UE::AnimNext
 	 * 
 	 * A decorator that smoothly blends between discrete states over time.
 	 */
-	struct FBlendSmootherDecorator : FAdditiveDecorator, IUpdate, IDiscreteBlend, ISmoothBlend
+	struct FBlendSmootherDecorator : FAdditiveDecorator, IEvaluate, IUpdate, IDiscreteBlend, ISmoothBlend
 	{
 		DECLARE_ANIM_DECORATOR(FBlendSmootherDecorator, 0x7b6c3d2e, FAdditiveDecorator)
 
@@ -64,6 +65,9 @@ namespace UE::AnimNext
 			TArray<FBlendData> PerChildBlendData;
 		};
 
+		// IEvaluate impl
+		virtual void PostEvaluate(FEvaluateTraversalContext& Context, const TDecoratorBinding<IEvaluate>& Binding) const override;
+
 		// IUpdate impl
 		virtual void PreUpdate(FUpdateTraversalContext& Context, const TDecoratorBinding<IUpdate>& Binding, const FDecoratorUpdateState& DecoratorState) const override;
 
diff --git a/Engine/Plugins/Experimental/AnimNext/Source/AnimNext/Private/Decorators/BlendSmootherPerBone.cpp b/Engine/Plugins/Experimental/AnimNext/Source/AnimNext/Private/Decorators/BlendSmootherPerBone.cpp
--- a/Engine/Plugins/Experimental/AnimNext/Source/AnimNext/Private/Decorators/BlendSmootherPerBone.cpp
+++ b/Engine/Plugins/Experimental/AnimNext/Source/AnimNext/Private/Decorators/BlendSmootherPerBone.cpp
@@ -120,10 +120,23 @@ namespace UE::AnimNext
 			return;	// No blend profile set, nothing to do
 		}
 
+		const int32 NumChildren = InstanceData->PerChildBlendData.Num();
+		if (NewChildIndex >= NumChildren)
+		{
+			// We have a new child
+			check(NewChildIndex == NumChildren + 1);
+
+			InstanceData->PerChildBlendData.AddDefaulted();
+			FBlendSampleData& SampleData = InstanceData->PerBoneSampleData.AddDefaulted_GetRef();
+
+			const uint32 NumBlendEntries = SharedData->BlendProfile->GetNumBlendEntries();
+			SampleData.SampleDataIndex = NewChildIndex;
+			SampleData.PerBoneBlendData.AddZeroed(NumBlendEntries);
+		}
+
 		TDecoratorBinding<IDiscreteBlend> DiscreteBlendDecorator;
 		Context.GetInterface(Binding, DiscreteBlendDecorator);
 
-		const int32 NumChildren = InstanceData->PerChildBlendData.Num();
 		for (int32 ChildIndex = 0; ChildIndex < NumChildren; ++ChildIndex)
 		{
 			FBlendData& ChildBlendData = InstanceData->PerChildBlendData[ChildIndex];
@@ -148,15 +161,13 @@ namespace UE::AnimNext
 		const uint32 NumChildren = HierarchyDecorator.GetNumChildren(Context);
 
 		InstanceData->PerChildBlendData.SetNum(NumChildren);
-
-		// Initialise per-bone data
 		InstanceData->PerBoneSampleData.SetNum(NumChildren);
 
 		const uint32 NumBlendEntries = SharedData->BlendProfile->GetNumBlendEntries();
-		for (uint32 Idx = 0; Idx < NumChildren; ++Idx)
+		for (uint32 ChildIndex = 0; ChildIndex < NumChildren; ++ChildIndex)
 		{
-			FBlendSampleData& SampleData = InstanceData->PerBoneSampleData[Idx];
-			SampleData.SampleDataIndex = Idx;
+			FBlendSampleData& SampleData = InstanceData->PerBoneSampleData[ChildIndex];
+			SampleData.SampleDataIndex = ChildIndex;
 			SampleData.PerBoneBlendData.AddZeroed(NumBlendEntries);
 		}
 	}
diff --git a/Engine/Plugins/Experimental/AnimNext/Source/AnimNext/Private/Decorators/BlendTwoWay.cpp b/Engine/Plugins/Experimental/AnimNext/Source/AnimNext/Private/Decorators/BlendTwoWay.cpp
--- a/Engine/Plugins/Experimental/AnimNext/Source/AnimNext/Private/Decorators/BlendTwoWay.cpp
+++ b/Engine/Plugins/Experimental/AnimNext/Source/AnimNext/Private/Decorators/BlendTwoWay.cpp
@@ -116,14 +116,15 @@ namespace UE::AnimNext
 		const FSharedData* SharedData = Binding.GetSharedData<FSharedData>();
 
 		const float BlendWeight = SharedData->GetBlendWeight(Context, Binding);
+		const float ClampedWeight = FMath::Clamp(BlendWeight, 0.0f, 1.0f);
 
 		if (ChildIndex == 0)
 		{
-			return 1.0f - BlendWeight;
+			return 1.0f - ClampedWeight;
 		}
 		else if (ChildIndex == 1)
 		{
-			return BlendWeight;
+			return ClampedWeight;
 		}
 		else
 		{
diff --git a/Engine/Plugins/Experimental/AnimNext/Source/AnimNext/Private/Decorators/SubGraphHost.cpp b/Engine/Plugins/Experimental/AnimNext/Source/AnimNext/Private/Decorators/SubGraphHost.cpp
new file mode 100644
--- /dev/null
+++ b/Engine/Plugins/Experimental/AnimNext/Source/AnimNext/Private/Decorators/SubGraphHost.cpp
@@ -0,0 +1,215 @@
+// Copyright Epic Games, Inc. All Rights Reserved.
+
+#include "Decorators/SubGraphHost.h"
+
+#include "DecoratorBase/ExecutionContext.h"
+#include "EvaluationVM/Tasks/PushReferenceKeyframe.h"
+
+namespace UE::AnimNext
+{
+	AUTO_REGISTER_ANIM_DECORATOR(FSubGraphHostDecorator)
+
+	DEFINE_ANIM_DECORATOR_BEGIN(FSubGraphHostDecorator)
+		DEFINE_ANIM_DECORATOR_IMPLEMENTS_INTERFACE(IDiscreteBlend)
+		DEFINE_ANIM_DECORATOR_IMPLEMENTS_INTERFACE(IEvaluate)
+		DEFINE_ANIM_DECORATOR_IMPLEMENTS_INTERFACE(IGarbageCollection)
+		DEFINE_ANIM_DECORATOR_IMPLEMENTS_INTERFACE(IHierarchy)
+		DEFINE_ANIM_DECORATOR_IMPLEMENTS_INTERFACE(IUpdate)
+	DEFINE_ANIM_DECORATOR_END(FSubGraphHostDecorator)
+
+	void FSubGraphHostDecorator::FInstanceData::Construct(const FExecutionContext& Context, const FDecoratorBinding& Binding)
+	{
+		FDecorator::FInstanceData::Construct(Context, Binding);
+
+		IGarbageCollection::RegisterWithGC(Context, Binding);
+	}
+
+	void FSubGraphHostDecorator::FInstanceData::Destruct(const FExecutionContext& Context, const FDecoratorBinding& Binding)
+	{
+		FDecorator::FInstanceData::Destruct(Context, Binding);
+
+		IGarbageCollection::UnregisterWithGC(Context, Binding);
+	}
+
+	uint32 FSubGraphHostDecorator::GetNumChildren(const FExecutionContext& Context, const TDecoratorBinding<IHierarchy>& Binding) const
+	{
+		const FInstanceData* InstanceData = Binding.GetInstanceData<FInstanceData>();
+
+		return InstanceData->SubGraphSlots.Num();
+	}
+
+	void FSubGraphHostDecorator::GetChildren(const FExecutionContext& Context, const TDecoratorBinding<IHierarchy>& Binding, FChildrenArray& Children) const
+	{
+		const FInstanceData* InstanceData = Binding.GetInstanceData<FInstanceData>();
+
+		for (const FInstanceData::FSubGraphSlot& SubGraphEntry : InstanceData->SubGraphSlots)
+		{
+			Children.Add(SubGraphEntry.GraphInstance.GetGraphRootPtr());
+		}
+	}
+
+	void FSubGraphHostDecorator::PostEvaluate(FEvaluateTraversalContext& Context, const TDecoratorBinding<IEvaluate>& Binding) const
+	{
+		const FInstanceData* InstanceData = Binding.GetInstanceData<FInstanceData>();
+
+		if (InstanceData->SubGraphSlots.IsEmpty())
+		{
+			// We have no children, output a non-additive reference pose
+			Context.AppendTask(FAnimNextPushReferenceKeyframeTask::MakeFromSkeleton());
+		}
+		else
+		{
+			// We only have one child or another decorator handles this, do nothing
+		}
+	}
+
+	void FSubGraphHostDecorator::PreUpdate(FUpdateTraversalContext& Context, const TDecoratorBinding<IUpdate>& Binding, const FDecoratorUpdateState& DecoratorState) const
+	{
+		const FSharedData* SharedData = Binding.GetSharedData<FSharedData>();
+		FInstanceData* InstanceData = Binding.GetInstanceData<FInstanceData>();
+
+		TObjectPtr<const UAnimNextGraph> CurrentActiveSubGraph;
+		if (InstanceData->CurrentlyActiveSubGraphIndex != INDEX_NONE)
+		{
+			CurrentActiveSubGraph = InstanceData->SubGraphSlots[InstanceData->CurrentlyActiveSubGraphIndex].SubGraph;
+		}
+
+		const TObjectPtr<const UAnimNextGraph> SubGraph = SharedData->GetSubGraph(Context, Binding);
+		if (CurrentActiveSubGraph != SubGraph)
+		{
+			// We can't blend from a sub-graph to an empty pose right now
+			check(SubGraph);
+
+			// Find an empty slot we can use
+			int32 FreeSlotIndex = INDEX_NONE;
+
+			const int32 NumSubGraphSlots = InstanceData->SubGraphSlots.Num();
+			for (int32 SlotIndex = 0; SlotIndex < NumSubGraphSlots; ++SlotIndex)
+			{
+				if (!InstanceData->SubGraphSlots[SlotIndex].GraphInstance.IsValid())
+				{
+					// This graph instance is invalid, we can re-use it
+					FreeSlotIndex = SlotIndex;
+					break;
+				}
+			}
+
+			if (FreeSlotIndex == INDEX_NONE)
+			{
+				// All slots are in use, add a new one
+				FreeSlotIndex = InstanceData->SubGraphSlots.AddDefaulted();
+			}
+
+			FInstanceData::FSubGraphSlot& SubGraphSlot = InstanceData->SubGraphSlots[FreeSlotIndex];
+			SubGraphSlot.SubGraph = SubGraph;
+
+			const int32 OldChildIndex = InstanceData->CurrentlyActiveSubGraphIndex;
+			const int32 NewChildIndex = FreeSlotIndex;
+
+			InstanceData->CurrentlyActiveSubGraphIndex = FreeSlotIndex;
+
+			TDecoratorBinding<IDiscreteBlend> DiscreteBlendDecorator;
+			Context.GetInterface(Binding, DiscreteBlendDecorator);
+
+			DiscreteBlendDecorator.OnBlendTransition(Context, OldChildIndex, NewChildIndex);
+		}
+	}
+
+	void FSubGraphHostDecorator::QueueChildrenForTraversal(FUpdateTraversalContext& Context, const TDecoratorBinding<IUpdate>& Binding, const FDecoratorUpdateState& DecoratorState, FUpdateTraversalQueue& TraversalQueue) const
+	{
+		const FInstanceData* InstanceData = Binding.GetInstanceData<FInstanceData>();
+
+		const int32 NumSubGraphs = InstanceData->SubGraphSlots.Num();
+		if (NumSubGraphs == 0)
+		{
+			return;
+		}
+
+		TDecoratorBinding<IDiscreteBlend> DiscreteBlendDecorator;
+		Context.GetInterface(Binding, DiscreteBlendDecorator);
+
+		for (int32 SubGraphIndex = 0; SubGraphIndex < NumSubGraphs; ++SubGraphIndex)
+		{
+			const float BlendWeight = DiscreteBlendDecorator.GetBlendWeight(Context, SubGraphIndex);
+
+			FDecoratorUpdateState SubGraphDecoratorState = DecoratorState.WithWeight(BlendWeight);
+			if (SubGraphIndex != InstanceData->CurrentlyActiveSubGraphIndex)
+			{
+				SubGraphDecoratorState = SubGraphDecoratorState.AsBlendingOut();
+			}
+
+			TraversalQueue.Push(InstanceData->SubGraphSlots[SubGraphIndex].GraphInstance.GetGraphRootPtr(), SubGraphDecoratorState);
+		}
+	}
+
+	float FSubGraphHostDecorator::GetBlendWeight(const FExecutionContext& Context, const TDecoratorBinding<IDiscreteBlend>& Binding, int32 ChildIndex) const
+	{
+		const FInstanceData* InstanceData = Binding.GetInstanceData<FInstanceData>();
+
+		if (ChildIndex == InstanceData->CurrentlyActiveSubGraphIndex)
+		{
+			return 1.0f;	// Active child has full weight
+		}
+		else if (InstanceData->SubGraphSlots.IsValidIndex(ChildIndex))
+		{
+			return 0.0f;	// Other children have no weight
+		}
+		else
+		{
+			// Invalid child index
+			return -1.0f;
+		}
+	}
+
+	int32 FSubGraphHostDecorator::GetBlendDestinationChildIndex(const FExecutionContext& Context, const TDecoratorBinding<IDiscreteBlend>& Binding) const
+	{
+		const FInstanceData* InstanceData = Binding.GetInstanceData<FInstanceData>();
+
+		return InstanceData->CurrentlyActiveSubGraphIndex;
+	}
+
+	void FSubGraphHostDecorator::OnBlendTransition(const FExecutionContext& Context, const TDecoratorBinding<IDiscreteBlend>& Binding, int32 OldChildIndex, int32 NewChildIndex) const
+	{
+		TDecoratorBinding<IDiscreteBlend> DiscreteBlendDecorator;
+		Context.GetInterface(Binding, DiscreteBlendDecorator);
+
+		// We initiate immediately when we transition
+		DiscreteBlendDecorator.OnBlendInitiated(Context, NewChildIndex);
+
+		// We terminate immediately when we transition
+		DiscreteBlendDecorator.OnBlendTerminated(Context, OldChildIndex);
+	}
+
+	void FSubGraphHostDecorator::OnBlendInitiated(const FExecutionContext& Context, const TDecoratorBinding<IDiscreteBlend>& Binding, int32 ChildIndex) const
+	{
+		FInstanceData* InstanceData = Binding.GetInstanceData<FInstanceData>();
+
+		if (InstanceData->SubGraphSlots.IsValidIndex(ChildIndex))
+		{
+			// Allocate our new sub-graph instance
+			FInstanceData::FSubGraphSlot& SubGraphEntry = InstanceData->SubGraphSlots[ChildIndex];
+			SubGraphEntry.SubGraph->AllocateInstance(Context.GetGraphInstance(), SubGraphEntry.GraphInstance);
+		}
+	}
+
+	void FSubGraphHostDecorator::OnBlendTerminated(const FExecutionContext& Context, const TDecoratorBinding<IDiscreteBlend>& Binding, int32 ChildIndex) const
+	{
+		FInstanceData* InstanceData = Binding.GetInstanceData<FInstanceData>();
+
+		if (InstanceData->SubGraphSlots.IsValidIndex(ChildIndex))
+		{
+			// Deallocate our sub-graph instance
+			InstanceData->SubGraphSlots[ChildIndex].GraphInstance.Release();
+		}
+	}
+
+	void FSubGraphHostDecorator::AddReferencedObjects(const FExecutionContext& Context, const TDecoratorBinding<IGarbageCollection>& Binding, FReferenceCollector& Collector) const
+	{
+		FInstanceData* InstanceData = Binding.GetInstanceData<FInstanceData>();
+
+		for (FInstanceData::FSubGraphSlot& SubGraphEntry : InstanceData->SubGraphSlots)
+		{
+			Collector.AddPropertyReferencesWithStructARO(FAnimNextGraphInstancePtr::StaticStruct(), &SubGraphEntry.GraphInstance);
+		}
+	}
+}
diff --git a/Engine/Plugins/Experimental/AnimNext/Source/AnimNext/Private/Decorators/SubGraphHost.h b/Engine/Plugins/Experimental/AnimNext/Source/AnimNext/Private/Decorators/SubGraphHost.h
new file mode 100644
--- /dev/null
+++ b/Engine/Plugins/Experimental/AnimNext/Source/AnimNext/Private/Decorators/SubGraphHost.h
@@ -0,0 +1,91 @@
+// Copyright Epic Games, Inc. All Rights Reserved.
+
+#pragma once
+
+#include "CoreMinimal.h"
+
+#include "DecoratorBase/Decorator.h"
+#include "DecoratorInterfaces/IDiscreteBlend.h"
+#include "DecoratorInterfaces/IEvaluate.h"
+#include "DecoratorInterfaces/IGarbageCollection.h"
+#include "DecoratorInterfaces/IHierarchy.h"
+#include "DecoratorInterfaces/IUpdate.h"
+#include "Graph/AnimNextGraph.h"
+#include "Graph/AnimNextGraphInstancePtr.h"
+
+#include "SubGraphHost.generated.h"
+
+USTRUCT(meta = (DisplayName = "SubGraph"))
+struct FAnimNextSubGraphHostDecoratorSharedData : public FAnimNextDecoratorSharedData
+{
+	GENERATED_BODY()
+
+	/** Subgraph that we host and manage. */
+	UPROPERTY(EditAnywhere, Category = "Default")
+	TObjectPtr<const UAnimNextGraph> SubGraph;
+
+	// Latent pin support boilerplate
+	#define DECORATOR_LATENT_PROPERTIES_ENUMERATOR(GeneratorMacro) \
+		GeneratorMacro(SubGraph) \
+
+	GENERATE_DECORATOR_LATENT_PROPERTIES(FAnimNextSubGraphHostDecoratorSharedData, DECORATOR_LATENT_PROPERTIES_ENUMERATOR)
+	#undef DECORATOR_LATENT_PROPERTIES_ENUMERATOR
+};
+
+namespace UE::AnimNext
+{
+	/**
+	 * FSubGraphHostDecorator
+	 * 
+	 * A decorator that hosts and manages a sub-graph instance.
+	 */
+	struct FSubGraphHostDecorator : FBaseDecorator, IEvaluate, IUpdate, IHierarchy, IDiscreteBlend, IGarbageCollection
+	{
+		DECLARE_ANIM_DECORATOR(FSubGraphHostDecorator, 0xb0cfe72e, FBaseDecorator)
+
+		using FSharedData = FAnimNextSubGraphHostDecoratorSharedData;
+
+		struct FInstanceData : FDecorator::FInstanceData
+		{
+			void Construct(const FExecutionContext& Context, const FDecoratorBinding& Binding);
+			void Destruct(const FExecutionContext& Context, const FDecoratorBinding& Binding);
+
+			struct FSubGraphSlot
+			{
+				// The sub-graph to use
+				TObjectPtr<const UAnimNextGraph> SubGraph;
+
+				// The graph instance
+				FAnimNextGraphInstancePtr GraphInstance;
+			};
+
+			// List of sub-graph slots
+			TArray<FSubGraphSlot> SubGraphSlots;
+
+			// The index of the currently active sub-graph slot
+			// All other sub-graphs are blending out
+			int32 CurrentlyActiveSubGraphIndex = INDEX_NONE;
+		};
+
+		// IEvaluate impl
+		virtual void PostEvaluate(FEvaluateTraversalContext& Context, const TDecoratorBinding<IEvaluate>& Binding) const override;
+
+		// IUpdate impl
+		virtual void PreUpdate(FUpdateTraversalContext& Context, const TDecoratorBinding<IUpdate>& Binding, const FDecoratorUpdateState& DecoratorState) const override;
+		virtual void QueueChildrenForTraversal(FUpdateTraversalContext& Context, const TDecoratorBinding<IUpdate>& Binding, const FDecoratorUpdateState& DecoratorState, FUpdateTraversalQueue& TraversalQueue) const override;
+
+		// IHierarchy impl
+		virtual uint32 GetNumChildren(const FExecutionContext& Context, const TDecoratorBinding<IHierarchy>& Binding) const override;
+		virtual void GetChildren(const FExecutionContext& Context, const TDecoratorBinding<IHierarchy>& Binding, FChildrenArray& Children) const override;
+
+		// IDiscreteBlend impl
+		virtual float GetBlendWeight(const FExecutionContext& Context, const TDecoratorBinding<IDiscreteBlend>& Binding, int32 ChildIndex) const override;
+		virtual int32 GetBlendDestinationChildIndex(const FExecutionContext& Context, const TDecoratorBinding<IDiscreteBlend>& Binding) const override;
+		virtual void OnBlendTransition(const FExecutionContext& Context, const TDecoratorBinding<IDiscreteBlend>& Binding, int32 OldChildIndex, int32 NewChildIndex) const override;
+		virtual void OnBlendInitiated(const FExecutionContext& Context, const TDecoratorBinding<IDiscreteBlend>& Binding, int32 ChildIndex) const override;
+		virtual void OnBlendTerminated(const FExecutionContext& Context, const TDecoratorBinding<IDiscreteBlend>& Binding, int32 ChildIndex) const override;
+
+		// IGarbageCollection impl
+		virtual void AddReferencedObjects(const FExecutionContext& Context, const TDecoratorBinding<IGarbageCollection>& Binding, FReferenceCollector& Collector) const override;
+	};
+}
diff --git a/Engine/Plugins/Experimental/AnimNext/Source/AnimNext/Private/Graph/AnimNextGraph.cpp b/Engine/Plugins/Experimental/AnimNext/Source/AnimNext/Private/Graph/AnimNextGraph.cpp
--- a/Engine/Plugins/Experimental/AnimNext/Source/AnimNext/Private/Graph/AnimNextGraph.cpp
+++ b/Engine/Plugins/Experimental/AnimNext/Source/AnimNext/Private/Graph/AnimNextGraph.cpp
@@ -11,10 +11,10 @@
 #include "DecoratorBase/ExecutionContext.h"
 #include "Graph/AnimNext_LODPose.h"
 #include "Graph/GC_GraphInstanceComponent.h"
+#include "Graph/AnimNextGraphInstance.h"
 #include "Serialization/MemoryReader.h"
 #include "AnimNextStats.h"
 
-DEFINE_STAT(STAT_AnimNext_Graph_RigVM);
 DEFINE_STAT(STAT_AnimNext_Graph_AllocateInstance);
 
 #include UE_INLINE_GENERATED_CPP_BY_NAME(AnimNextGraph)
@@ -22,246 +22,67 @@ DEFINE_STAT(STAT_AnimNext_Graph_AllocateInstance);
 const UE::AnimNext::FParamId UAnimNextGraph::DefaultReferencePoseId = UE::AnimNext::FParamId("UE_AnimNextMeshComponent_ReferencePose");
 const UE::AnimNext::FParamId UAnimNextGraph::DefaultCurrentLODId = UE::AnimNext::FParamId("UE_AnimNextMeshComponent_PredictedLODLevel");
 
-#if WITH_EDITORONLY_DATA
-FAnimNextGraphInstance::FAnimNextGraphInstance(const FAnimNextGraphInstance& Other)
-{
-	*this = Other;
-}
-
-FAnimNextGraphInstance::FAnimNextGraphInstance(FAnimNextGraphInstance&& Other)
-{
-	*this = MoveTemp(Other);
-}
-
-FAnimNextGraphInstance& FAnimNextGraphInstance::operator=(const FAnimNextGraphInstance& Other)
-{
-	if (GraphInstancePtr.IsValid())
-	{
-		if (!Other.GraphInstancePtr.IsValid())
-		{
-			// We were valid but we won't be anymore, unregister ourself
-
-			FRWScopeLock Lock(Graph->GraphInstancesLock, SLT_Write);
-
-			check(Graph->GraphInstances.Contains(this));
-			Graph->GraphInstances.Remove(this);
-		}
-		else
-		{
-			// Both instances remain valid, nothing to do
-		}
-	}
-	else
-	{
-		if (Other.GraphInstancePtr.IsValid())
-		{
-			// We were invalid but we will become valid, register ourself
-
-			FRWScopeLock Lock(Graph->GraphInstancesLock, SLT_Write);
-
-			check(!Graph->GraphInstances.Contains(this));
-			Graph->GraphInstances.Add(this);
-		}
-		else
-		{
-			// Both instances were invalid, nothing to do
-		}
-	}
-
-	Graph = Other.Graph;
-	GraphInstancePtr =  Other.GraphInstancePtr;
-	ExtendedExecuteContext = Other.ExtendedExecuteContext;
-
-	return *this;
-}
-
-FAnimNextGraphInstance& FAnimNextGraphInstance::operator=(FAnimNextGraphInstance&& Other)
-{
-	if (GraphInstancePtr.IsValid())
-	{
-		if (!Other.GraphInstancePtr.IsValid())
-		{
-			// We were valid but we won't be anymore, unregister ourself
-			// Other will become valid, register it
-
-			FRWScopeLock Lock(Graph->GraphInstancesLock, SLT_Write);
-
-			check(Graph->GraphInstances.Contains(this));
-			Graph->GraphInstances.Remove(this);
-
-			check(!Graph->GraphInstances.Contains(&Other));
-			Graph->GraphInstances.Add(&Other);
-		}
-		else
-		{
-			// Both instances remain valid, nothing to do
-		}
-	}
-	else
-	{
-		if (Other.GraphInstancePtr.IsValid())
-		{
-			// We were invalid but we will become valid, register ourself
-			// Other will become invalid, unregister it
-
-			FRWScopeLock Lock(Graph->GraphInstancesLock, SLT_Write);
-
-			check(!Graph->GraphInstances.Contains(this));
-			Graph->GraphInstances.Add(this);
-
-			check(Graph->GraphInstances.Contains(&Other));
-			Graph->GraphInstances.Remove(&Other);
-		}
-		else
-		{
-			// Both instances were invalid, nothing to do
-		}
-	}
-
-	Swap(Graph, Other.Graph);
-	Swap(GraphInstancePtr, Other.GraphInstancePtr);
-	Swap(ExtendedExecuteContext, Other.ExtendedExecuteContext);
-
-	return *this;
-}
-#endif
-
-FAnimNextGraphInstance::~FAnimNextGraphInstance()
-{
-	Release();
-}
-
-void FAnimNextGraphInstance::Release()
-{
-	if (!GraphInstancePtr.IsValid())
-	{
-		return;
-	}
-
-	GraphInstancePtr.Reset();
-	ExtendedExecuteContext.Reset();
-	Components.Empty();
-
-#if WITH_EDITORONLY_DATA
-	FRWScopeLock Lock(Graph->GraphInstancesLock, SLT_Write);
-	check(Graph->GraphInstances.Contains(this));
-	Graph->GraphInstances.Remove(this);
-#endif
-
-	Graph = nullptr;
-}
-
-bool FAnimNextGraphInstance::IsValid() const
-{
-	return GraphInstancePtr.IsValid();
-}
-
-const UAnimNextGraph* FAnimNextGraphInstance::GetGraph() const
-{
-	return Graph;
-}
-
-UE::AnimNext::FWeakDecoratorPtr FAnimNextGraphInstance::GetGraphRootPtr() const
-{
-	return GraphInstancePtr;
-}
-
-bool FAnimNextGraphInstance::UsesGraph(const UAnimNextGraph* InGraph) const
-{
-	return Graph == InGraph;
-}
-
-void FAnimNextGraphInstance::AddStructReferencedObjects(FReferenceCollector& Collector)
+UAnimNextGraph::UAnimNextGraph(const FObjectInitializer& ObjectInitializer)
+	: Super(ObjectInitializer)
 {
-	if (const UE::AnimNext::FGCGraphInstanceComponent* Component = TryGetComponent<UE::AnimNext::FGCGraphInstanceComponent>())
-	{
-		Component->AddReferencedObjects(Collector);
-	}
+	ExtendedExecuteContext.SetContextPublicDataStruct(FAnimNextExecuteContext::StaticStruct());
+	SetRigVMExtendedExecuteContext(&ExtendedExecuteContext);
 }
 
-UE::AnimNext::FGraphInstanceComponent* FAnimNextGraphInstance::TryGetComponent(int32 ComponentNameHash, FName ComponentName) const
+void UAnimNextGraph::AllocateInstance(FAnimNextGraphInstancePtr& Instance) const
 {
-	if (const TSharedPtr<UE::AnimNext::FGraphInstanceComponent>* Component = Components.FindByHash(ComponentNameHash, ComponentName))
-	{
-		return Component->Get();
-	}
-
-	return nullptr;
+	AllocateInstanceImpl(nullptr, Instance);
 }
 
-UE::AnimNext::FGraphInstanceComponent& FAnimNextGraphInstance::AddComponent(int32 ComponentNameHash, FName ComponentName, TSharedPtr<UE::AnimNext::FGraphInstanceComponent>&& Component)
+void UAnimNextGraph::AllocateInstance(FAnimNextGraphInstance& ParentGraphInstance, FAnimNextGraphInstancePtr& Instance) const
 {
-	return *Components.AddByHash(ComponentNameHash, ComponentName, MoveTemp(Component)).Get();
+	AllocateInstanceImpl(&ParentGraphInstance, Instance);
 }
 
-GraphInstanceComponentMapType::TConstIterator FAnimNextGraphInstance::GetComponentIterator() const
+void UAnimNextGraph::AllocateInstanceImpl(FAnimNextGraphInstance* ParentGraphInstance, FAnimNextGraphInstancePtr& Instance) const
 {
-	return Components.CreateConstIterator();
-}
+	SCOPE_CYCLE_COUNTER(STAT_AnimNext_Graph_AllocateInstance);
 
-void FAnimNextGraphInstance::ExecuteLatentPin(int32 LatentPinIndex, void* DestinationPtr)
-{
-	SCOPE_CYCLE_COUNTER(STAT_AnimNext_Graph_RigVM);
+	Instance.Release();
 
-	if (!IsValid())
+	if (!ResolvedRootDecoratorHandle.IsValid())
 	{
 		return;
 	}
 
-	if (URigVM* VM = Graph->VM)
 	{
-		FAnimNextExecuteContext& AnimNextContext = ExtendedExecuteContext.GetPublicDataSafe<FAnimNextExecuteContext>();
-		AnimNextContext.SetupForExecution(LatentPinIndex, DestinationPtr);
+		TUniquePtr<FAnimNextGraphInstance> InstanceImpl = MakeUnique<FAnimNextGraphInstance>();
 
-		VM->ExecuteVM(ExtendedExecuteContext, FRigUnit_AnimNextShimRoot::EventName);
+		InstanceImpl->Graph = this;
+		InstanceImpl->ParentGraphInstance = ParentGraphInstance;
 
-		// Reset the context to avoid issues if we forget to reset it the next time we use it
-		AnimNextContext.DebugReset();
-	}
-}
+		// If we have a parent graph, use its root since we share the same root, otherwise if we have no parent, we are the root
+		InstanceImpl->RootGraphInstance = ParentGraphInstance != nullptr ? ParentGraphInstance->GetRootGraphInstance() : InstanceImpl.Get();
 
-UAnimNextGraph::UAnimNextGraph(const FObjectInitializer& ObjectInitializer)
-	: Super(ObjectInitializer)
-{
-	ExtendedExecuteContext.SetContextPublicDataStruct(FAnimNextExecuteContext::StaticStruct());
-	SetRigVMExtendedExecuteContext(&ExtendedExecuteContext);
-}
+		InstanceImpl->ExtendedExecuteContext.CopyMemoryStorage(ExtendedExecuteContext);
+		VM->InitializeInstance(InstanceImpl->ExtendedExecuteContext);
 
-void UAnimNextGraph::AllocateInstance(FAnimNextGraphInstance& Instance) const
-{
-	SCOPE_CYCLE_COUNTER(STAT_AnimNext_Graph_AllocateInstance);
-
-	Instance.Release();
-
-	if (!ResolvedRootDecoratorHandle.IsValid())
-	{
-		return;
+		// Move our implementation so that we can use the instance to allocate our root node
+		Instance.Impl = MoveTemp(InstanceImpl);
 	}
 
-	Instance.Graph = this;
-
-	Instance.ExtendedExecuteContext.CopyMemoryStorage(ExtendedExecuteContext);
-	VM->InitializeInstance(Instance.ExtendedExecuteContext);
-
 	{
 		UE::AnimNext::FExecutionContext Context(Instance);
-		Instance.GraphInstancePtr = Context.AllocateNodeInstance(UE::AnimNext::FWeakDecoratorPtr(), ResolvedRootDecoratorHandle);
+		Instance.Impl->GraphInstancePtr = Context.AllocateNodeInstance(UE::AnimNext::FWeakDecoratorPtr(), ResolvedRootDecoratorHandle);
 	}
 
-	if (!Instance.GraphInstancePtr.IsValid())
+	if (!Instance.IsValid())
 	{
 		// We failed to allocate our instance, clear everything
-		Instance.Graph = nullptr;
-		Instance.ExtendedExecuteContext.Reset();
-		Instance.Components.Empty();
+		Instance.Release();
 	}
 
 #if WITH_EDITORONLY_DATA
-	if (Instance.GraphInstancePtr.IsValid())
+	if (Instance.IsValid())
 	{
 		FRWScopeLock Lock(GraphInstancesLock, SLT_Write);
-		check(!GraphInstances.Contains(&Instance));
-		GraphInstances.Add(&Instance);
+		check(!GraphInstances.Contains(Instance.Impl.Get()));
+		GraphInstances.Add(Instance.Impl.Get());
 	}
 #endif
 }
@@ -395,23 +216,23 @@ bool UAnimNextGraph::LoadFromArchiveBuffer(const TArray<uint8>& InSharedDataArch
 }
 
 #if WITH_EDITORONLY_DATA
-TArray<FAnimNextGraphInstance*> UAnimNextGraph::ReleaseAllInstances()
+void UAnimNextGraph::FreezeGraphInstances()
 {
-	// Make a copy of our live instances since we'll release them all
-	TArray<FAnimNextGraphInstance*> TempGraphInstances;
+	FRWScopeLock Lock(GraphInstancesLock, SLT_ReadOnly);
+
+	for (FAnimNextGraphInstance* GraphInstance : GraphInstances)
 	{
-		// Lock shouldn't be necessary here since this should be called before graph compilation
-		// If a graph is executing and we attempt to release it, things are likely to crash
-		// Compilation must happen at a point in the frame where animation isn't executing
-		FRWScopeLock Lock(GraphInstancesLock, SLT_Write);
-		TempGraphInstances = GraphInstances.Array();
+		GraphInstance->Freeze();
 	}
+}
 
-	for (FAnimNextGraphInstance* GraphInstance : TempGraphInstances)
+void UAnimNextGraph::ThawGraphInstances()
+{
+	FRWScopeLock Lock(GraphInstancesLock, SLT_ReadOnly);
+
+	for (FAnimNextGraphInstance* GraphInstance : GraphInstances)
 	{
-		GraphInstance->Release();
+		GraphInstance->Thaw();
 	}
-
-	return TempGraphInstances;
 }
 #endif
diff --git a/Engine/Plugins/Experimental/AnimNext/Source/AnimNext/Private/Graph/AnimNextGraphInstance.cpp b/Engine/Plugins/Experimental/AnimNext/Source/AnimNext/Private/Graph/AnimNextGraphInstance.cpp
new file mode 100644
--- /dev/null
+++ b/Engine/Plugins/Experimental/AnimNext/Source/AnimNext/Private/Graph/AnimNextGraphInstance.cpp
@@ -0,0 +1,152 @@
+// Copyright Epic Games, Inc. All Rights Reserved.
+
+#include "Graph/AnimNextGraphInstance.h"
+
+#include "AnimNextStats.h"
+#include "DecoratorBase/ExecutionContext.h"
+#include "Graph/AnimNextGraph.h"
+#include "Graph/GC_GraphInstanceComponent.h"
+#include "Graph/RigUnit_AnimNextShimRoot.h"
+
+DEFINE_STAT(STAT_AnimNext_Graph_RigVM);
+
+FAnimNextGraphInstance::~FAnimNextGraphInstance()
+{
+	Release();
+}
+
+void FAnimNextGraphInstance::Release()
+{
+	if (!GraphInstancePtr.IsValid())
+	{
+		return;
+	}
+
+	GraphInstancePtr.Reset();
+	ParentGraphInstance = nullptr;
+	RootGraphInstance = nullptr;
+	ExtendedExecuteContext.Reset();
+	Components.Empty();
+	Graph = nullptr;
+}
+
+bool FAnimNextGraphInstance::IsValid() const
+{
+	return GraphInstancePtr.IsValid();
+}
+
+const UAnimNextGraph* FAnimNextGraphInstance::GetGraph() const
+{
+	return Graph;
+}
+
+UE::AnimNext::FWeakDecoratorPtr FAnimNextGraphInstance::GetGraphRootPtr() const
+{
+	return GraphInstancePtr;
+}
+
+FAnimNextGraphInstance* FAnimNextGraphInstance::GetParentGraphInstance() const
+{
+	return ParentGraphInstance;
+}
+
+FAnimNextGraphInstance* FAnimNextGraphInstance::GetRootGraphInstance() const
+{
+	return RootGraphInstance;
+}
+
+bool FAnimNextGraphInstance::UsesGraph(const UAnimNextGraph* InGraph) const
+{
+	return Graph == InGraph;
+}
+
+bool FAnimNextGraphInstance::IsRoot() const
+{
+	return this == RootGraphInstance;
+}
+
+void FAnimNextGraphInstance::AddStructReferencedObjects(FReferenceCollector& Collector)
+{
+	if (!IsRoot())
+	{
+		return;	// If we aren't the root graph instance, we don't own the components
+	}
+
+	if (const UE::AnimNext::FGCGraphInstanceComponent* Component = TryGetComponent<UE::AnimNext::FGCGraphInstanceComponent>())
+	{
+		Component->AddReferencedObjects(Collector);
+	}
+}
+
+UE::AnimNext::FGraphInstanceComponent* FAnimNextGraphInstance::TryGetComponent(int32 ComponentNameHash, FName ComponentName) const
+{
+	if (const TSharedPtr<UE::AnimNext::FGraphInstanceComponent>* Component = RootGraphInstance->Components.FindByHash(ComponentNameHash, ComponentName))
+	{
+		return Component->Get();
+	}
+
+	return nullptr;
+}
+
+UE::AnimNext::FGraphInstanceComponent& FAnimNextGraphInstance::AddComponent(int32 ComponentNameHash, FName ComponentName, TSharedPtr<UE::AnimNext::FGraphInstanceComponent>&& Component)
+{
+	return *RootGraphInstance->Components.AddByHash(ComponentNameHash, ComponentName, MoveTemp(Component)).Get();
+}
+
+GraphInstanceComponentMapType::TConstIterator FAnimNextGraphInstance::GetComponentIterator() const
+{
+	return RootGraphInstance->Components.CreateConstIterator();
+}
+
+void FAnimNextGraphInstance::ExecuteLatentPin(int32 LatentPinIndex, void* DestinationPtr)
+{
+	SCOPE_CYCLE_COUNTER(STAT_AnimNext_Graph_RigVM);
+
+	if (!IsValid())
+	{
+		return;
+	}
+
+	if (URigVM* VM = Graph->VM)
+	{
+		FAnimNextExecuteContext& AnimNextContext = ExtendedExecuteContext.GetPublicDataSafe<FAnimNextExecuteContext>();
+		AnimNextContext.SetupForExecution(LatentPinIndex, DestinationPtr);
+
+		VM->ExecuteVM(ExtendedExecuteContext, FRigUnit_AnimNextShimRoot::EventName);
+
+		// Reset the context to avoid issues if we forget to reset it the next time we use it
+		AnimNextContext.DebugReset();
+	}
+}
+
+void FAnimNextGraphInstance::Freeze()
+{
+	if (!IsValid())
+	{
+		return;
+	}
+
+	GraphInstancePtr.Reset();
+	ExtendedExecuteContext.Reset();
+	Components.Empty();
+}
+
+void FAnimNextGraphInstance::Thaw()
+{
+	if (const UAnimNextGraph* GraphPtr = Graph)
+	{
+		ExtendedExecuteContext.CopyMemoryStorage(GraphPtr->ExtendedExecuteContext);
+		GraphPtr->VM->InitializeInstance(ExtendedExecuteContext);
+
+		{
+			UE::AnimNext::FExecutionContext Context(*this);
+			GraphInstancePtr = Context.AllocateNodeInstance(UE::AnimNext::FWeakDecoratorPtr(), GraphPtr->ResolvedRootDecoratorHandle);
+		}
+
+		if (!IsValid())
+		{
+			// We failed to allocate our instance, clear everything
+			Release();
+		}
+	}
+}
diff --git a/Engine/Plugins/Experimental/AnimNext/Source/AnimNext/Private/Graph/AnimNextGraphInstancePtr.cpp b/Engine/Plugins/Experimental/AnimNext/Source/AnimNext/Private/Graph/AnimNextGraphInstancePtr.cpp
new file mode 100644
--- /dev/null
+++ b/Engine/Plugins/Experimental/AnimNext/Source/AnimNext/Private/Graph/AnimNextGraphInstancePtr.cpp
@@ -0,0 +1,89 @@
+// Copyright Epic Games, Inc. All Rights Reserved.
+
+#include "Graph/AnimNextGraphInstancePtr.h"
+
+#include "Graph/AnimNextGraph.h"
+#include "Graph/AnimNextGraphInstance.h"
+
+FAnimNextGraphInstancePtr::FAnimNextGraphInstancePtr() = default;
+FAnimNextGraphInstancePtr::FAnimNextGraphInstancePtr(FAnimNextGraphInstancePtr&&) = default;
+FAnimNextGraphInstancePtr& FAnimNextGraphInstancePtr::operator=(FAnimNextGraphInstancePtr&&) = default;
+
+FAnimNextGraphInstancePtr::~FAnimNextGraphInstancePtr()
+{
+	Release();
+}
+
+void FAnimNextGraphInstancePtr::Release()
+{
+	if (Impl)
+	{
+#if WITH_EDITORONLY_DATA
+		const UAnimNextGraph* Graph = Impl->GetGraph();
+		FRWScopeLock Lock(Graph->GraphInstancesLock, SLT_Write);
+		check(Graph->GraphInstances.Contains(Impl.Get()));
+		Graph->GraphInstances.Remove(Impl.Get());
+#endif
+
+		// Destroy the graph instance
+		Impl->Release();
+
+		// Reset our unique ptr
+		Impl.Reset();
+	}
+}
+
+bool FAnimNextGraphInstancePtr::IsValid() const
+{
+	return Impl && Impl->IsValid();
+}
+
+const UAnimNextGraph* FAnimNextGraphInstancePtr::GetGraph() const
+{
+	return Impl ? Impl->GetGraph() : nullptr;
+}
+
+UE::AnimNext::FWeakDecoratorPtr FAnimNextGraphInstancePtr::GetGraphRootPtr() const
+{
+	return Impl ? Impl->GetGraphRootPtr() : UE::AnimNext::FWeakDecoratorPtr();
+}
+
+FAnimNextGraphInstance* FAnimNextGraphInstancePtr::GetImpl() const
+{
+	return Impl.Get();
+}
+
+bool FAnimNextGraphInstancePtr::UsesGraph(const UAnimNextGraph* InGraph) const
+{
+	return Impl ? Impl->UsesGraph(InGraph) : false;
+}
+
+bool FAnimNextGraphInstancePtr::IsRoot() const
+{
+	return Impl ? Impl->IsRoot() : true;
+}
+
+void FAnimNextGraphInstancePtr::AddStructReferencedObjects(FReferenceCollector& Collector)
+{
+	if (Impl)
+	{
+		Impl->AddStructReferencedObjects(Collector);
+	}
+}
+
+UE::AnimNext::FGraphInstanceComponent* FAnimNextGraphInstancePtr::TryGetComponent(int32 ComponentNameHash, FName ComponentName) const
+{
+	return Impl ? Impl->TryGetComponent(ComponentNameHash, ComponentName) : nullptr;
+}
+
+UE::AnimNext::FGraphInstanceComponent& FAnimNextGraphInstancePtr::AddComponent(int32 ComponentNameHash, FName ComponentName, TSharedPtr<UE::AnimNext::FGraphInstanceComponent>&& Component)
+{
+	check(Impl);
+	return Impl->AddComponent(ComponentNameHash, ComponentName, MoveTemp(Component));
+}
+
+GraphInstanceComponentMapType::TConstIterator FAnimNextGraphInstancePtr::GetComponentIterator() const
+{
+	check(Impl);
+	return Impl->GetComponentIterator();
+}
diff --git a/Engine/Plugins/Experimental/AnimNext/Source/AnimNext/Private/Module.cpp b/Engine/Plugins/Experimental/AnimNext/Source/AnimNext/Private/Module.cpp
--- a/Engine/Plugins/Experimental/AnimNext/Source/AnimNext/Private/Module.cpp
+++ b/Engine/Plugins/Experimental/AnimNext/Source/AnimNext/Private/Module.cpp
@@ -10,6 +10,7 @@
 #include "DataRegistry.h"
 #include "DecoratorBase/DecoratorRegistry.h"
 #include "DecoratorBase/NodeTemplateRegistry.h"
+#include "Graph/AnimNextGraph.h"
 #include "RigVMCore/RigVMRegistry.h"
 #include "RigVMRuntimeDataRegistry.h"
 #include "Animation/AnimSequence.h"
@@ -33,6 +34,7 @@ public:
 			{ UScriptStruct::StaticClass(), FRigVMRegistry::ERegisterObjectOperation::Class },
 			{ UBlendProfile::StaticClass(), FRigVMRegistry::ERegisterObjectOperation::Class },
 			{ UCurveFloat::StaticClass(), FRigVMRegistry::ERegisterObjectOperation::Class },
+			{ UAnimNextGraph::StaticClass(), FRigVMRegistry::ERegisterObjectOperation::Class },
 		};
 
 		FRigVMRegistry::Get().RegisterObjectTypes(AllowedObjectTypes);
diff --git a/Engine/Plugins/Experimental/AnimNext/Source/AnimNext/Private/Scheduler/ScheduleInstanceData.cpp b/Engine/Plugins/Experimental/AnimNext/Source/AnimNext/Private/Scheduler/ScheduleInstanceData.cpp
--- a/Engine/Plugins/Experimental/AnimNext/Source/AnimNext/Private/Scheduler/ScheduleInstanceData.cpp
+++ b/Engine/Plugins/Experimental/AnimNext/Source/AnimNext/Private/Scheduler/ScheduleInstanceData.cpp
@@ -122,7 +122,7 @@ void FScheduleInstanceData::AddReferencedObjects(FReferenceCollector& Collector)
 
 	for (FGraphCache& GraphCache : GraphCaches)
 	{
-		Collector.AddPropertyReferencesWithStructARO(FAnimNextGraphInstance::StaticStruct(), &GraphCache.GraphInstanceData);
+		Collector.AddPropertyReferencesWithStructARO(FAnimNextGraphInstancePtr::StaticStruct(), &GraphCache.GraphInstanceData);
 	}
 }
 
diff --git a/Engine/Plugins/Experimental/AnimNext/Source/AnimNext/Private/Scheduler/ScheduleInstanceData.h b/Engine/Plugins/Experimental/AnimNext/Source/AnimNext/Private/Scheduler/ScheduleInstanceData.h
--- a/Engine/Plugins/Experimental/AnimNext/Source/AnimNext/Private/Scheduler/ScheduleInstanceData.h
+++ b/Engine/Plugins/Experimental/AnimNext/Source/AnimNext/Private/Scheduler/ScheduleInstanceData.h
@@ -5,9 +5,9 @@
 #include "CoreMinimal.h"
 #include "Scheduler/ScheduleHandle.h"
 #include "Graph/AnimNextGraph.h"
+#include "Graph/AnimNextGraphInstancePtr.h"
 #include "Param/ParamStack.h"
 
-struct FAnimNextGraphInstance;
 class UAnimNextSchedule;
 struct FAnimNextSchedulerEntry;
 
@@ -92,7 +92,7 @@ struct FScheduleInstanceData : public FGCObject
 	struct FGraphCache
 	{
 		// Graph instance data
-		FAnimNextGraphInstance GraphInstanceData;
+		FAnimNextGraphInstancePtr GraphInstanceData;
 
 		// Remapped data layers for input terms (from schedule)
 		FParamStackLayerHandle GraphTermLayer;
diff --git a/Engine/Plugins/Experimental/AnimNext/Source/AnimNext/Public/DecoratorBase/ExecutionContext.h b/Engine/Plugins/Experimental/AnimNext/Source/AnimNext/Public/DecoratorBase/ExecutionContext.h
--- a/Engine/Plugins/Experimental/AnimNext/Source/AnimNext/Public/DecoratorBase/ExecutionContext.h
+++ b/Engine/Plugins/Experimental/AnimNext/Source/AnimNext/Public/DecoratorBase/ExecutionContext.h
@@ -9,7 +9,9 @@
 #include "DecoratorBase/DecoratorInterfaceUID.h"
 #include "DecoratorBase/LatentPropertyHandle.h"
 #include "DecoratorBase/NodeHandle.h"
-#include "Graph/AnimNextGraph.h"
+#include "Graph/AnimNextGraphInstancePtr.h"
+
+struct FAnimNextGraphInstance;
 
 namespace UE::AnimNext
 {
@@ -32,9 +34,15 @@ namespace UE::AnimNext
 		// Creates an unbound execution context
 		FExecutionContext();
 
+		// Creates an execution context and binds it to the specified graph instance
+		explicit FExecutionContext(FAnimNextGraphInstancePtr& InGraphInstance);
+
 		// Creates an execution context and binds it to the specified graph instance
 		explicit FExecutionContext(FAnimNextGraphInstance& InGraphInstance);
 
+		// Binds the execution context to the specified graph instance if it differs from the currently bound instance
+		void BindTo(FAnimNextGraphInstancePtr& InGraphInstance);
+
 		// Binds the execution context to the specified graph instance if it differs from the currently bound instance
 		void BindTo(FAnimNextGraphInstance& InGraphInstance);
 
@@ -44,6 +52,9 @@ namespace UE::AnimNext
 		// Returns whether or not this execution context is bound to a graph instance
 		bool IsBound() const;
 
+		// Returns whether or not this execution context is bound to the specified graph instance
+		bool IsBoundTo(const FAnimNextGraphInstancePtr& InGraphInstance) const;
+
 		// Returns whether or not this execution context is bound to the specified graph instance
 		bool IsBoundTo(const FAnimNextGraphInstance& InGraphInstance) const;
 
@@ -92,7 +103,7 @@ namespace UE::AnimNext
 		ComponentType* TryGetComponent() const;
 
 		// Returns const iterators to the graph instance component container
-		GraphInstanceComponentMapType::TConstIterator GetComponentIterator() const { return GraphInstance->GetComponentIterator(); }
+		GraphInstanceComponentMapType::TConstIterator GetComponentIterator() const;
 
 		// Returns the bound graph instance
 		FAnimNextGraphInstance& GetGraphInstance() const { return *GraphInstance; }
@@ -105,6 +116,8 @@ namespace UE::AnimNext
 		bool GetInterfaceImpl(FDecoratorInterfaceUID InterfaceUID, const FWeakDecoratorPtr& DecoratorPtr, FDecoratorBinding& InterfaceBinding) const;
 		bool GetInterfaceSuperImpl(FDecoratorInterfaceUID InterfaceUID, const FWeakDecoratorPtr& DecoratorPtr, FDecoratorBinding& SuperBinding) const;
 		void EvaluateLatentPinImpl(FLatentPropertyHandle LatentPropertyHandle, void* DestinationPtr) const;
+		FGraphInstanceComponent* TryGetComponent(int32 ComponentNameHash, FName ComponentName) const;
+		FGraphInstanceComponent& AddComponent(int32 ComponentNameHash, FName ComponentName, TSharedPtr<FGraphInstanceComponent>&& Component) const;
 
 		const FNodeDescription& GetNodeDescription(FNodeHandle NodeHandle) const;
 		const FNodeTemplate* GetNodeTemplate(const FNodeDescription& NodeDesc) const;
@@ -170,12 +183,23 @@ namespace UE::AnimNext
 	template<class ComponentType>
 	ComponentType& FExecutionContext::GetComponent() const
 	{
-		return GraphInstance->GetComponent<ComponentType>();
+		const FName ComponentName = ComponentType::StaticComponentName();
+		const int32 ComponentNameHash = GetTypeHash(ComponentName);
+
+		if (FGraphInstanceComponent* Component = TryGetComponent(ComponentNameHash, ComponentName))
+		{
+			return *static_cast<ComponentType*>(Component);
+		}
+
+		return static_cast<ComponentType&>(AddComponent(ComponentNameHash, ComponentName, MakeShared<ComponentType>()));
 	}
 
 	template<class ComponentType>
 	ComponentType* FExecutionContext::TryGetComponent() const
 	{
-		return GraphInstance->TryGetComponent<ComponentType>();
+		const FName ComponentName = ComponentType::StaticComponentName();
+		const int32 ComponentNameHash = GetTypeHash(ComponentName);
+
+		return static_cast<ComponentType*>(TryGetComponent(ComponentNameHash, ComponentName));
 	}
 }
diff --git a/Engine/Plugins/Experimental/AnimNext/Source/AnimNext/Public/DecoratorBase/ExecutionContextProxy.h b/Engine/Plugins/Experimental/AnimNext/Source/AnimNext/Public/DecoratorBase/ExecutionContextProxy.h
--- a/Engine/Plugins/Experimental/AnimNext/Source/AnimNext/Public/DecoratorBase/ExecutionContextProxy.h
+++ b/Engine/Plugins/Experimental/AnimNext/Source/AnimNext/Public/DecoratorBase/ExecutionContextProxy.h
@@ -31,6 +31,12 @@ namespace UE::AnimNext
 			return ExecutionContext.IsBound();
 		}
 
+		// Returns whether or not this execution context is bound to the specified graph instance
+		bool IsBoundTo(const FAnimNextGraphInstancePtr& InGraphInstance) const
+		{
+			return ExecutionContext.IsBoundTo(InGraphInstance);
+		}
+
 		// Returns whether or not this execution context is bound to the specified graph instance
 		bool IsBoundTo(const FAnimNextGraphInstance& InGraphInstance) const
 		{
diff --git a/Engine/Plugins/Experimental/AnimNext/Source/AnimNext/Public/DecoratorInterfaces/IDiscreteBlend.h b/Engine/Plugins/Experimental/AnimNext/Source/AnimNext/Public/DecoratorInterfaces/IDiscreteBlend.h
--- a/Engine/Plugins/Experimental/AnimNext/Source/AnimNext/Public/DecoratorInterfaces/IDiscreteBlend.h
+++ b/Engine/Plugins/Experimental/AnimNext/Source/AnimNext/Public/DecoratorInterfaces/IDiscreteBlend.h
@@ -33,6 +33,10 @@ namespace UE::AnimNext
 		virtual int32 GetBlendDestinationChildIndex(const FExecutionContext& Context, const TDecoratorBinding<IDiscreteBlend>& Binding) const;
 
 		// Called when a blend transition between children occurs
+		// OldChildIndex can be INDEX_NONE if there was no previously active child
+		// NewChildIndex can be larger than the current number of known children to support a dynamic number of children at runtime
+		// When this occurs, the number of children increments by one
+		// The number of children never shrinks
 		virtual void OnBlendTransition(const FExecutionContext& Context, const TDecoratorBinding<IDiscreteBlend>& Binding, int32 OldChildIndex, int32 NewChildIndex) const;
 
 		// Called when the blend for specified child is initiated
diff --git a/Engine/Plugins/Experimental/AnimNext/Source/AnimNext/Public/DecoratorInterfaces/IEvaluate.h b/Engine/Plugins/Experimental/AnimNext/Source/AnimNext/Public/DecoratorInterfaces/IEvaluate.h
--- a/Engine/Plugins/Experimental/AnimNext/Source/AnimNext/Public/DecoratorInterfaces/IEvaluate.h
+++ b/Engine/Plugins/Experimental/AnimNext/Source/AnimNext/Public/DecoratorInterfaces/IEvaluate.h
@@ -91,7 +91,7 @@ namespace UE::AnimNext
 	 *
 	 * @see IEvaluate::PreEvaluate, IEvaluate::PostEvaluate, IHierarchy::GetChildren
 	 */
-	[[nodiscard]] ANIMNEXT_API FEvaluationProgram EvaluateGraph(FAnimNextGraphInstance& GraphInstance);
+	[[nodiscard]] ANIMNEXT_API FEvaluationProgram EvaluateGraph(FAnimNextGraphInstancePtr& GraphInstance);
 
 	/**
 	 * Evaluates a sub-graph starting at its root and produces an evaluation program.
diff --git a/Engine/Plugins/Experimental/AnimNext/Source/AnimNext/Public/DecoratorInterfaces/IUpdate.h b/Engine/Plugins/Experimental/AnimNext/Source/AnimNext/Public/DecoratorInterfaces/IUpdate.h
--- a/Engine/Plugins/Experimental/AnimNext/Source/AnimNext/Public/DecoratorInterfaces/IUpdate.h
+++ b/Engine/Plugins/Experimental/AnimNext/Source/AnimNext/Public/DecoratorInterfaces/IUpdate.h
@@ -6,7 +6,7 @@
 #include "DecoratorBase/ExecutionContextProxy.h"
 #include "DecoratorBase/IDecoratorInterface.h"
 
-struct FAnimNextGraphInstance;
+struct FAnimNextGraphInstancePtr;
 class FMemStack;
 
 namespace UE::AnimNext
@@ -127,7 +127,7 @@ namespace UE::AnimNext
 		// be warmer in the CPU cache
 		FUpdateEntry* FreeEntryStackHead = nullptr;
 
-		friend ANIMNEXT_API void UpdateGraph(FAnimNextGraphInstance& GraphInstance, float DeltaTime);
+		friend ANIMNEXT_API void UpdateGraph(FAnimNextGraphInstancePtr& GraphInstance, float DeltaTime);
 		friend FUpdateTraversalQueue;
 	};
 
@@ -155,7 +155,7 @@ namespace UE::AnimNext
 		// emptied and pushed onto the update stack
 		FUpdateEntry* QueuedUpdateStackHead = nullptr;
 
-		friend ANIMNEXT_API void UpdateGraph(FAnimNextGraphInstance& GraphInstance, float DeltaTime);
+		friend ANIMNEXT_API void UpdateGraph(FAnimNextGraphInstancePtr& GraphInstance, float DeltaTime);
 		friend FUpdateTraversalContext;
 	};
 
@@ -226,5 +226,5 @@ namespace UE::AnimNext
 	 *
 	 * @see IUpdate::PreUpdate, IUpdate::PostUpdate, IHierarchy::GetChildren
 	 */
-	ANIMNEXT_API void UpdateGraph(FAnimNextGraphInstance& GraphInstance, float DeltaTime);
+	ANIMNEXT_API void UpdateGraph(FAnimNextGraphInstancePtr& GraphInstance, float DeltaTime);
 }
diff --git a/Engine/Plugins/Experimental/AnimNext/Source/AnimNext/Public/Graph/AnimNextGraphInstance.h b/Engine/Plugins/Experimental/AnimNext/Source/AnimNext/Public/Graph/AnimNextGraphInstance.h
new file mode 100644
--- /dev/null
+++ b/Engine/Plugins/Experimental/AnimNext/Source/AnimNext/Public/Graph/AnimNextGraphInstance.h
@@ -0,0 +1,170 @@
+// Copyright Epic Games, Inc. All Rights Reserved.
+
+#pragma once
+
+#include "CoreMinimal.h"
+#include "DecoratorBase/DecoratorPtr.h"
+#include "RigVMCore/RigVMExecuteContext.h"
+
+#include "AnimNextGraphInstance.generated.h"
+
+class FReferenceCollector;
+
+struct FAnimNextGraphInstancePtr;
+struct FRigUnit_AnimNextGraphEvaluator;
+class UAnimNextGraph;
+
+namespace UE::AnimNext
+{
+	struct FExecutionContext;
+	struct FGraphInstanceComponent;
+}
+
+using GraphInstanceComponentMapType = TMap<FName, TSharedPtr<UE::AnimNext::FGraphInstanceComponent>>;
+
+// Represents an instance of an AnimNext graph
+// This struct uses UE reflection because we wish for the GC to keep the graph
+// alive while we own a reference to it. It is not intended to be serialized on disk with a live instance.
+USTRUCT()
+struct FAnimNextGraphInstance
+{
+	GENERATED_BODY()
+
+	// Creates an empty graph instance that doesn't reference anything
+	FAnimNextGraphInstance() = default;
+
+	// No copying, no moving
+	FAnimNextGraphInstance(const FAnimNextGraphInstance&) = delete;
+	FAnimNextGraphInstance& operator=(const FAnimNextGraphInstance&) = delete;
+
+	// If the graph instance is allocated, we release it during destruction
+	~FAnimNextGraphInstance();
+
+	// Releases the graph instance and frees all corresponding memory
+	void Release();
+
+	// Returns true if we have a live graph instance, false otherwise
+	bool IsValid() const;
+
+	// Returns the graph used by this instance or nullptr if the instance is invalid
+	const UAnimNextGraph* GetGraph() const;
+
+	// Returns a weak handle to the root decorator instance
+	UE::AnimNext::FWeakDecoratorPtr GetGraphRootPtr() const;
+
+	// Returns the parent graph instance that owns us or nullptr for the root graph instance or if we are invalid
+	FAnimNextGraphInstance* GetParentGraphInstance() const;
+
+	// Returns the root graph instance that owns us and the components or nullptr if we are invalid
+	FAnimNextGraphInstance* GetRootGraphInstance() const;
+
+	// Check to see if this instance data matches the provided graph
+	bool UsesGraph(const UAnimNextGraph* InGraph) const;
+
+	// Returns whether or not this graph instance is the root graph instance or false otherwise
+	bool IsRoot() const;
+
+	// Adds strong/hard object references during GC
+	void AddStructReferencedObjects(class FReferenceCollector& Collector);
+
+	// Returns a typed graph instance component, creating it lazily the first time it is queried
+	template<class ComponentType>
+	ComponentType& GetComponent();
+
+	// Returns a typed graph instance component pointer if found or nullptr otherwise
+	template<class ComponentType>
+	ComponentType* TryGetComponent();
+
+	// Returns a typed graph instance component pointer if found or nullptr otherwise
+	template<class ComponentType>
+	const ComponentType* TryGetComponent() const;
+
+	// Returns const iterators to the graph instance component container
+	GraphInstanceComponentMapType::TConstIterator GetComponentIterator() const;
+
+private:
+	// Returns a pointer to the specified component, or nullptr if not found
+	UE::AnimNext::FGraphInstanceComponent* TryGetComponent(int32 ComponentNameHash, FName ComponentName) const;
+
+	// Adds the specified component and returns a reference to it
+	UE::AnimNext::FGraphInstanceComponent& AddComponent(int32 ComponentNameHash, FName ComponentName, TSharedPtr<UE::AnimNext::FGraphInstanceComponent>&& Component);
+
+	// Executes a latent RigVM pin and writes the result into the destination pointer
+	void ExecuteLatentPin(int32 LatentPinIndex, void* DestinationPtr);
+
+	// During graph compilation, if we have existing graph instances, we freeze them by releasing their memory before thawing them
+	// Freezing is a partial release of resources that retains the necessary information to re-create things safely
+	void Freeze();
+
+	// During graph compilation, once compilation is done we thaw existing graph instances to reallocate their memory
+	void Thaw();
+
+	// Hard reference to the graph used to create this instance to ensure we can release it safely
+	UPROPERTY()
+	TObjectPtr<const UAnimNextGraph> Graph;
+
+	// Hard reference to the graph instance data, we own it
+	UE::AnimNext::FDecoratorPtr GraphInstancePtr;
+
+	// The graph instance that owns us
+	FAnimNextGraphInstance* ParentGraphInstance = nullptr;
+
+	// The root graph instance that owns us and the components
+	FAnimNextGraphInstance* RootGraphInstance = nullptr;
+
+	// Extended execute context instance for this graph instance, we own it
+	UPROPERTY()
+	FRigVMExtendedExecuteContext ExtendedExecuteContext;
+
+	// Graph instance components that persist from update to update
+	GraphInstanceComponentMapType Components;
+
+	friend UAnimNextGraph;					// The graph is the one that allocates instances
+	friend FRigUnit_AnimNextGraphEvaluator;	// We evaluate the instance
+	friend UE::AnimNext::FExecutionContext;
+	friend FAnimNextGraphInstancePtr;
+};
+
+template<>
+struct TStructOpsTypeTraits<FAnimNextGraphInstance> : public TStructOpsTypeTraitsBase2<FAnimNextGraphInstance>
+{
+	enum
+	{
+		WithAddStructReferencedObjects = true,
+		WithCopy = false,
+	};
+};
+
+//////////////////////////////////////////////////////////////////////////
+
+template<class ComponentType>
+ComponentType& FAnimNextGraphInstance::GetComponent()
+{
+	const FName ComponentName = ComponentType::StaticComponentName();
+	const int32 ComponentNameHash = GetTypeHash(ComponentName);
+
+	if (UE::AnimNext::FGraphInstanceComponent* Component = TryGetComponent(ComponentNameHash, ComponentName))
+	{
+		return *static_cast<ComponentType*>(Component);
+	}
+
+	return static_cast<ComponentType&>(AddComponent(ComponentNameHash, ComponentName, MakeShared<ComponentType>()));
+}
+
+template<class ComponentType>
+ComponentType* FAnimNextGraphInstance::TryGetComponent()
+{
+	const FName ComponentName = ComponentType::StaticComponentName();
+	const int32 ComponentNameHash = GetTypeHash(ComponentName);
+
+	return static_cast<ComponentType*>(TryGetComponent(ComponentNameHash, ComponentName));
+}
+
+template<class ComponentType>
+const ComponentType* FAnimNextGraphInstance::TryGetComponent() const
+{
+	const FName ComponentName = ComponentType::StaticComponentName();
+	const int32 ComponentNameHash = GetTypeHash(ComponentName);
+
+	return static_cast<ComponentType*>(TryGetComponent(ComponentNameHash, ComponentName));
+}
diff --git a/Engine/Plugins/Experimental/AnimNext/Source/AnimNext/Public/Graph/AnimNextGraphInstancePtr.h b/Engine/Plugins/Experimental/AnimNext/Source/AnimNext/Public/Graph/AnimNextGraphInstancePtr.h
new file mode 100644
--- /dev/null
+++ b/Engine/Plugins/Experimental/AnimNext/Source/AnimNext/Public/Graph/AnimNextGraphInstancePtr.h
@@ -0,0 +1,140 @@
+// Copyright Epic Games, Inc. All Rights Reserved.
+
+#pragma once
+
+#include "CoreMinimal.h"
+#include "DecoratorBase/DecoratorPtr.h"
+
+#include "AnimNextGraphInstancePtr.generated.h"
+
+class FReferenceCollector;
+
+struct FAnimNextGraphInstance;
+struct FRigUnit_AnimNextGraphEvaluator;
+class UAnimNextGraph;
+
+namespace UE::AnimNext
+{
+	struct FExecutionContext;
+	struct FGraphInstanceComponent;
+}
+
+using GraphInstanceComponentMapType = TMap<FName, TSharedPtr<UE::AnimNext::FGraphInstanceComponent>>;
+
+// Represents an instance of an AnimNext graph
+// This struct uses UE reflection because we wish for the GC to keep the graph
+// alive while we own a reference to it. It is not intended to be serialized on disk with a live instance.
+USTRUCT()
+struct ANIMNEXT_API FAnimNextGraphInstancePtr
+{
+	GENERATED_BODY()
+
+	// Creates an empty graph instance that doesn't reference anything
+	FAnimNextGraphInstancePtr();
+
+	// No copying, only moving
+	FAnimNextGraphInstancePtr(const FAnimNextGraphInstancePtr&) = delete;
+	FAnimNextGraphInstancePtr& operator=(const FAnimNextGraphInstancePtr&) = delete;
+	FAnimNextGraphInstancePtr(FAnimNextGraphInstancePtr&&);
+	FAnimNextGraphInstancePtr& operator=(FAnimNextGraphInstancePtr&&);
+
+	// If the graph instance is allocated, we release it during destruction
+	~FAnimNextGraphInstancePtr();
+
+	// Releases the graph instance and frees all corresponding memory
+	void Release();
+
+	// Returns true if we have a live graph instance, false otherwise
+	bool IsValid() const;
+
+	// Returns the graph used by this instance or nullptr if the instance is invalid
+	const UAnimNextGraph* GetGraph() const;
+
+	// Returns a weak handle to the root decorator instance
+	UE::AnimNext::FWeakDecoratorPtr GetGraphRootPtr() const;
+
+	// Returns the graph instance implementation
+	FAnimNextGraphInstance* GetImpl() const;
+
+	// Check to see if this instance data matches the provided graph
+	bool UsesGraph(const UAnimNextGraph* InGraph) const;
+
+	// Returns whether or not this graph instance is the root graph instance or false otherwise
+	bool IsRoot() const;
+
+	// Adds strong/hard object references during GC
+	void AddStructReferencedObjects(class FReferenceCollector& Collector);
+
+	// Returns a typed graph instance component, creating it lazily the first time it is queried
+	template<class ComponentType>
+	ComponentType& GetComponent();
+
+	// Returns a typed graph instance component pointer if found or nullptr otherwise
+	template<class ComponentType>
+	ComponentType* TryGetComponent();
+
+	// Returns a typed graph instance component pointer if found or nullptr otherwise
+	template<class ComponentType>
+	const ComponentType* TryGetComponent() const;
+
+	// Returns const iterators to the graph instance component container
+	GraphInstanceComponentMapType::TConstIterator GetComponentIterator() const;
+
+private:
+	// Returns a pointer to the specified component, or nullptr if not found
+	UE::AnimNext::FGraphInstanceComponent* TryGetComponent(int32 ComponentNameHash, FName ComponentName) const;
+
+	// Adds the specified component and returns a reference to it
+	UE::AnimNext::FGraphInstanceComponent& AddComponent(int32 ComponentNameHash, FName ComponentName, TSharedPtr<UE::AnimNext::FGraphInstanceComponent>&& Component);
+
+	// Indirection to hide implementation details and to fix the graph instance into a single memory location
+	TUniquePtr<FAnimNextGraphInstance> Impl;
+
+	friend UAnimNextGraph;					// The graph is the one that allocates instances
+	friend FRigUnit_AnimNextGraphEvaluator;	// We evaluate the instance
+	friend UE::AnimNext::FExecutionContext;
+};
+
+template<>
+struct TStructOpsTypeTraits<FAnimNextGraphInstancePtr> : public TStructOpsTypeTraitsBase2<FAnimNextGraphInstancePtr>
+{
+	enum
+	{
+		WithAddStructReferencedObjects = true,
+		WithCopy = false,
+	};
+};
+
+//////////////////////////////////////////////////////////////////////////
+
+template<class ComponentType>
+ComponentType& FAnimNextGraphInstancePtr::GetComponent()
+{
+	const FName ComponentName = ComponentType::StaticComponentName();
+	const int32 ComponentNameHash = GetTypeHash(ComponentName);
+
+	if (UE::AnimNext::FGraphInstanceComponent* Component = TryGetComponent(ComponentNameHash, ComponentName))
+	{
+		return *static_cast<ComponentType*>(Component);
+	}
+
+	return static_cast<ComponentType&>(AddComponent(ComponentNameHash, ComponentName, MakeShared<ComponentType>()));
+}
+
+template<class ComponentType>
+ComponentType* FAnimNextGraphInstancePtr::TryGetComponent()
+{
+	const FName ComponentName = ComponentType::StaticComponentName();
+	const int32 ComponentNameHash = GetTypeHash(ComponentName);
+
+	return static_cast<ComponentType*>(TryGetComponent(ComponentNameHash, ComponentName));
+}
+
+template<class ComponentType>
+const ComponentType* FAnimNextGraphInstancePtr::TryGetComponent() const
+{
+	const FName ComponentName = ComponentType::StaticComponentName();
+	const int32 ComponentNameHash = GetTypeHash(ComponentName);
+
+	return static_cast<ComponentType*>(TryGetComponent(ComponentNameHash, ComponentName));
+}
diff --git a/Engine/Plugins/Experimental/AnimNext/Source/AnimNextTestSuite/Private/AnimNextDecoratorBaseTest.cpp b/Engine/Plugins/Experimental/AnimNext/Source/AnimNextTestSuite/Private/AnimNextDecoratorBaseTest.cpp
--- a/Engine/Plugins/Experimental/AnimNext/Source/AnimNextTestSuite/Private/AnimNextDecoratorBaseTest.cpp
+++ b/Engine/Plugins/Experimental/AnimNext/Source/AnimNextTestSuite/Private/AnimNextDecoratorBaseTest.cpp
@@ -18,6 +18,7 @@
 #include "DecoratorBase/NodeInstance.h"
 #include "DecoratorBase/NodeTemplateBuilder.h"
 #include "DecoratorBase/NodeTemplateRegistry.h"
+#include "Graph/AnimNextGraph.h"
 #include "Graph/GraphFactory.h"
 
 //****************************************************************************
@@ -564,7 +565,7 @@ bool FAnimationAnimNextRuntimeTest_NodeLifetime::RunTest(const FString& InParame
 	// Read our graph
 	FTestUtils::LoadFromArchiveBuffer(*AnimNextGraph, NodeHandles, GraphSharedDataArchiveBuffer);
 
-	FAnimNextGraphInstance GraphInstance;
+	FAnimNextGraphInstancePtr GraphInstance;
 	AnimNextGraph->AllocateInstance(GraphInstance);
 
 	FExecutionContext Context(GraphInstance);
@@ -763,7 +764,7 @@ bool FAnimationAnimNextRuntimeTest_GetDecoratorInterface::RunTest(const FString&
 	// Read our graph
 	FTestUtils::LoadFromArchiveBuffer(*AnimNextGraph, NodeHandles, GraphSharedDataArchiveBuffer);
 
-	FAnimNextGraphInstance GraphInstance;
+	FAnimNextGraphInstancePtr GraphInstance;
 	AnimNextGraph->AllocateInstance(GraphInstance);
 
 	FExecutionContext Context(GraphInstance);
@@ -1032,7 +1033,7 @@ bool FAnimationAnimNextRuntimeTest_GetDecoratorInterfaceSuper::RunTest(const FSt
 	// Read our graph
 	FTestUtils::LoadFromArchiveBuffer(*AnimNextGraph, NodeHandles, GraphSharedDataArchiveBuffer);
 
-	FAnimNextGraphInstance GraphInstance;
+	FAnimNextGraphInstancePtr GraphInstance;
 	AnimNextGraph->AllocateInstance(GraphInstance);
 
 	FExecutionContext Context(GraphInstance);
@@ -1387,7 +1388,7 @@ bool FAnimationAnimNextRuntimeTest_DecoratorSerialization::RunTest(const FString
 		// Read our graph
 		FTestUtils::LoadFromArchiveBuffer(*AnimNextGraph, NodeHandles, GraphSharedDataArchiveBuffer);
 
-		FAnimNextGraphInstance GraphInstance;
+		FAnimNextGraphInstancePtr GraphInstance;
 		AnimNextGraph->AllocateInstance(GraphInstance);
 
 		FExecutionContext Context(GraphInstance);
diff --git a/Engine/Plugins/Experimental/AnimNext/Source/AnimNextTestSuite/Private/AnimNextDecoratorGraphTest.cpp b/Engine/Plugins/Experimental/AnimNext/Source/AnimNextTestSuite/Private/AnimNextDecoratorGraphTest.cpp
--- a/Engine/Plugins/Experimental/AnimNext/Source/AnimNextTestSuite/Private/AnimNextDecoratorGraphTest.cpp
+++ b/Engine/Plugins/Experimental/AnimNext/Source/AnimNextTestSuite/Private/AnimNextDecoratorGraphTest.cpp
@@ -249,7 +249,7 @@ bool FAnimationAnimNextRuntimeTest_GraphExecute::RunTest(const FString& InParame
 	TSharedRef<FParamStack> ParamStack = MakeShared<FParamStack>();
 	FParamStack::AttachToCurrentThread(ParamStack);
 
-	FAnimNextGraphInstance GraphInstance;
+	FAnimNextGraphInstancePtr GraphInstance;
 	AnimNextGraph->AllocateInstance(GraphInstance);
 
 	FParamStack::FPushedLayerHandle LayerHandle = ParamStack->PushValues(
@@ -391,7 +391,7 @@ bool FAnimationAnimNextRuntimeTest_GraphExecuteLatent::RunTest(const FString& In
 	TSharedRef<FParamStack> ParamStack = MakeShared<FParamStack>();
 	FParamStack::AttachToCurrentThread(ParamStack);
 
-	FAnimNextGraphInstance GraphInstance;
+	FAnimNextGraphInstancePtr GraphInstance;
 	AnimNextGraph->AllocateInstance(GraphInstance);
 
 	FParamStack::FPushedLayerHandle LayerHandle = ParamStack->PushValues(
diff --git a/Engine/Plugins/Experimental/AnimNext/Source/AnimNextTestSuite/Private/AnimNextDecoratorInterfacesTest.cpp b/Engine/Plugins/Experimental/AnimNext/Source/AnimNextTestSuite/Private/AnimNextDecoratorInterfacesTest.cpp
--- a/Engine/Plugins/Experimental/AnimNext/Source/AnimNextTestSuite/Private/AnimNextDecoratorInterfacesTest.cpp
+++ b/Engine/Plugins/Experimental/AnimNext/Source/AnimNextTestSuite/Private/AnimNextDecoratorInterfacesTest.cpp
@@ -21,6 +21,7 @@
 #include "DecoratorInterfaces/IEvaluate.h"
 #include "DecoratorInterfaces/IHierarchy.h"
 #include "DecoratorInterfaces/IUpdate.h"
+#include "Graph/AnimNextGraph.h"
 #include "Graph/GraphFactory.h"
 
 //****************************************************************************
@@ -341,7 +342,7 @@ bool FAnimationAnimNextRuntimeTest_IHierarchy::RunTest(const FString& InParamete
 	// Read our graph
 	FTestUtils::LoadFromArchiveBuffer(*AnimNextGraph, NodeHandles, GraphSharedDataArchiveBuffer);
 
-	FAnimNextGraphInstance GraphInstance;
+	FAnimNextGraphInstancePtr GraphInstance;
 	AnimNextGraph->AllocateInstance(GraphInstance);
 
 	FExecutionContext Context(GraphInstance);
@@ -502,7 +503,7 @@ bool FAnimationAnimNextRuntimeTest_IUpdate::RunTest(const FString& InParameters)
 	// Read our graph
 	FTestUtils::LoadFromArchiveBuffer(*AnimNextGraph, NodeHandles, GraphSharedDataArchiveBuffer);
 
-	FAnimNextGraphInstance GraphInstance;
+	FAnimNextGraphInstancePtr GraphInstance;
 	AnimNextGraph->AllocateInstance(GraphInstance);
 
 	FExecutionContext Context(GraphInstance);
@@ -639,7 +640,7 @@ bool FAnimationAnimNextRuntimeTest_IEvaluate::RunTest(const FString& InParameter
 	// Read our graph
 	FTestUtils::LoadFromArchiveBuffer(*AnimNextGraph, NodeHandles, GraphSharedDataArchiveBuffer);
 
-	FAnimNextGraphInstance GraphInstance;
+	FAnimNextGraphInstancePtr GraphInstance;
 	AnimNextGraph->AllocateInstance(GraphInstance);
 
 	{
diff --git a/Engine/Plugins/Experimental/AnimNext/Source/AnimNextUncookedOnly/Private/UncookedOnlyUtils.cpp b/Engine/Plugins/Experimental/AnimNext/Source/AnimNextUncookedOnly/Private/UncookedOnlyUtils.cpp
--- a/Engine/Plugins/Experimental/AnimNext/Source/AnimNextUncookedOnly/Private/UncookedOnlyUtils.cpp
+++ b/Engine/Plugins/Experimental/AnimNext/Source/AnimNextUncookedOnly/Private/UncookedOnlyUtils.cpp
@@ -406,7 +406,7 @@ void FUtils::Compile(UAnimNextGraph* InGraph)
 
 	// Before we re-compile a graph, we need to release and live instances since we need the metadata we are about to replace
 	// to call decorator destructors etc
-	const TArray<FAnimNextGraphInstance*> PreviousLiveGraphInstances = InGraph->ReleaseAllInstances();
+	InGraph->FreezeGraphInstances();
 
 	EditorData->bErrorsDuringCompilation = false;
 
@@ -535,10 +535,7 @@ void FUtils::Compile(UAnimNextGraph* InGraph)
 	VMClient->RemoveController(VMTempGraph);
 
 	// Now that the graph has been re-compiled, re-allocate the previous live instances
-	for (FAnimNextGraphInstance* GraphInstance : PreviousLiveGraphInstances)
-	{
-		InGraph->AllocateInstance(*GraphInstance);
-	}
+	InGraph->ThawGraphInstances();
 
 	// Gather and store the required parameters
 	InGraph->RequiredParameters.Empty();
