Add delta time to parameter block execution so we can plug it into RigVM

diff --git a/Engine/Plugins/Experimental/AnimNext/Source/AnimNext/Internal/Param/AnimNextParameterBlock.h b/Engine/Plugins/Experimental/AnimNext/Source/AnimNext/Internal/Param/AnimNextParameterBlock.h
--- a/Engine/Plugins/Experimental/AnimNext/Source/AnimNext/Internal/Param/AnimNextParameterBlock.h
+++ b/Engine/Plugins/Experimental/AnimNext/Source/AnimNext/Internal/Param/AnimNextParameterBlock.h
@@ -54,7 +54,7 @@ class UAnimNextParameterBlock : public URigVMHost
 	friend class UE::AnimNext::Editor::FParameterBlockParameterCustomization;
 	friend struct UE::AnimNext::FParameterBlockProxy;
 
-	void UpdateLayer(UE::AnimNext::FParamStackLayerHandle& InHandle) const;
+	void UpdateLayer(UE::AnimNext::FParamStackLayerHandle& InHandle, float InDeltaTime) const;
 
 	// UObject interface
 	virtual void BeginDestroy() override;
diff --git a/Engine/Plugins/Experimental/AnimNext/Source/AnimNext/Private/Graph/AnimGraph/AnimNode_AnimNextParameters.cpp b/Engine/Plugins/Experimental/AnimNext/Source/AnimNext/Private/Graph/AnimGraph/AnimNode_AnimNextParameters.cpp
--- a/Engine/Plugins/Experimental/AnimNext/Source/AnimNext/Private/Graph/AnimGraph/AnimNode_AnimNextParameters.cpp
+++ b/Engine/Plugins/Experimental/AnimNext/Source/AnimNext/Private/Graph/AnimGraph/AnimNode_AnimNextParameters.cpp
@@ -79,7 +79,7 @@ void FAnimNode_AnimNextParameters::Update_AnyThread(const FAnimationUpdateContex
 		FParamStack::FPushedLayerHandle PushedLayerHandle;
 		if (CurrentParameters && ParametersProxy.IsValid())
 		{
-			ParametersProxy->Update();
+			ParametersProxy->Update(Context.GetDeltaTime());
 			PushedLayerHandle = ParamStack.PushLayer(ParametersProxy->GetLayerHandle());
 		}
 
diff --git a/Engine/Plugins/Experimental/AnimNext/Source/AnimNext/Private/Param/AnimNextParameterBlock.cpp b/Engine/Plugins/Experimental/AnimNext/Source/AnimNext/Private/Param/AnimNextParameterBlock.cpp
--- a/Engine/Plugins/Experimental/AnimNext/Source/AnimNext/Private/Param/AnimNextParameterBlock.cpp
+++ b/Engine/Plugins/Experimental/AnimNext/Source/AnimNext/Private/Param/AnimNextParameterBlock.cpp
@@ -20,7 +20,7 @@ UAnimNextParameterBlock::UAnimNextParameterBlock(const FObjectInitializer& Objec
 	SetRigVMExtendedExecuteContext(&BaseRigVMContext);
 }
 
-void UAnimNextParameterBlock::UpdateLayer(UE::AnimNext::FParamStackLayerHandle& InHandle) const
+void UAnimNextParameterBlock::UpdateLayer(UE::AnimNext::FParamStackLayerHandle& InHandle, float InDeltaTime) const
 {
 	SCOPE_CYCLE_COUNTER(STAT_AnimNext_ParamBlock_UpdateLayer);
 	
@@ -33,7 +33,13 @@ void UAnimNextParameterBlock::UpdateLayer(UE::AnimNext::FParamStackLayerHandle&
 			check(Context.VMHash == VM->GetVMHash());
 
 			FAnimNextParameterExecuteContext& AnimNextParameterContext = Context.GetPublicDataSafe<FAnimNextParameterExecuteContext>();
+
+			// Param block setup
 			AnimNextParameterContext.SetParamContextData(InHandle);
+
+			// RigVM setup
+			AnimNextParameterContext.SetDeltaTime(InDeltaTime);
+
 			VM->ExecuteVM(Context, FRigUnit_AnimNextBeginExecution::EventName);
 		}
 	}
diff --git a/Engine/Plugins/Experimental/AnimNext/Source/AnimNext/Private/Param/ObjectProxy.cpp b/Engine/Plugins/Experimental/AnimNext/Source/AnimNext/Private/Param/ObjectProxy.cpp
--- a/Engine/Plugins/Experimental/AnimNext/Source/AnimNext/Private/Param/ObjectProxy.cpp
+++ b/Engine/Plugins/Experimental/AnimNext/Source/AnimNext/Private/Param/ObjectProxy.cpp
@@ -1,4 +1,4 @@
-﻿// Copyright Epic Games, Inc. All Rights Reserved.
+// Copyright Epic Games, Inc. All Rights Reserved.
 
 #include "ObjectProxy.h"
 
@@ -19,7 +19,7 @@ FObjectProxy::FObjectProxy(UObject* InObject, const TSharedRef<FObjectAccessor>&
 	ParameterCache.SetValueObject(RootParameterName, InObject);
 }
 
-void FObjectProxy::Update()
+void FObjectProxy::Update(float DeltaTime)
 {
 	const UPropertyBag* PropertyBag = ParameterCache.GetPropertyBagStruct();
 	TConstArrayView<FPropertyBagPropertyDesc> PropertyDescs = PropertyBag->GetPropertyDescs();
diff --git a/Engine/Plugins/Experimental/AnimNext/Source/AnimNext/Private/Param/ObjectProxy.h b/Engine/Plugins/Experimental/AnimNext/Source/AnimNext/Private/Param/ObjectProxy.h
--- a/Engine/Plugins/Experimental/AnimNext/Source/AnimNext/Private/Param/ObjectProxy.h
+++ b/Engine/Plugins/Experimental/AnimNext/Source/AnimNext/Private/Param/ObjectProxy.h
@@ -1,4 +1,4 @@
-﻿// Copyright Epic Games, Inc. All Rights Reserved.
+// Copyright Epic Games, Inc. All Rights Reserved.
 
 #pragma once
 
@@ -66,7 +66,7 @@ struct FObjectProxy : public IParameterSource
 	FObjectProxy(UObject* InObject, const TSharedRef<FObjectAccessor>& InObjectAccessor);
 
 	// IParameterSource interface
-	virtual void Update() override;
+	virtual void Update(float DeltaTime) override;
 	virtual const FParamStackLayerHandle& GetLayerHandle() const override { return LayerHandle; }
 	virtual void AddReferencedObjects(FReferenceCollector& Collector) override;
 
diff --git a/Engine/Plugins/Experimental/AnimNext/Source/AnimNext/Private/Param/ParameterBlockProxy.cpp b/Engine/Plugins/Experimental/AnimNext/Source/AnimNext/Private/Param/ParameterBlockProxy.cpp
--- a/Engine/Plugins/Experimental/AnimNext/Source/AnimNext/Private/Param/ParameterBlockProxy.cpp
+++ b/Engine/Plugins/Experimental/AnimNext/Source/AnimNext/Private/Param/ParameterBlockProxy.cpp
@@ -1,4 +1,4 @@
-﻿// Copyright Epic Games, Inc. All Rights Reserved.
+// Copyright Epic Games, Inc. All Rights Reserved.
 
 #include "Param/ParameterBlockProxy.h"
 #include "Param/AnimNextParameterBlock.h"
@@ -15,7 +15,7 @@ FParameterBlockProxy::FParameterBlockProxy(UAnimNextParameterBlock* InParameterB
 	check(ParameterBlock);
 }
 
-void FParameterBlockProxy::Update()
+void FParameterBlockProxy::Update(float DeltaTime)
 {
 #if WITH_EDITOR	// Layout should only be changing in editor
 	const FInstancedPropertyBag* HandlePropertyBag = LayerHandle.As<FInstancedPropertyBag>();
@@ -26,7 +26,7 @@ void FParameterBlockProxy::Update()
 	}
 #endif
 
-	ParameterBlock->UpdateLayer(LayerHandle);
+	ParameterBlock->UpdateLayer(LayerHandle, DeltaTime);
 }
 
 void FParameterBlockProxy::AddReferencedObjects(FReferenceCollector& Collector)
diff --git a/Engine/Plugins/Experimental/AnimNext/Source/AnimNext/Private/Param/ParameterBlockProxy.h b/Engine/Plugins/Experimental/AnimNext/Source/AnimNext/Private/Param/ParameterBlockProxy.h
--- a/Engine/Plugins/Experimental/AnimNext/Source/AnimNext/Private/Param/ParameterBlockProxy.h
+++ b/Engine/Plugins/Experimental/AnimNext/Source/AnimNext/Private/Param/ParameterBlockProxy.h
@@ -1,4 +1,4 @@
-﻿// Copyright Epic Games, Inc. All Rights Reserved.
+// Copyright Epic Games, Inc. All Rights Reserved.
 
 #pragma once
 
@@ -20,7 +20,7 @@ struct FParameterBlockProxy : public IParameterSource
 	FParameterBlockProxy(UAnimNextParameterBlock* InParameterBlock);
 
 	// IParameterSource interface
-	virtual void Update() override;
+	virtual void Update(float DeltaTime) override;
 	virtual const FParamStackLayerHandle& GetLayerHandle() const override { return LayerHandle; }
 	virtual void AddReferencedObjects(FReferenceCollector& Collector) override;
 
diff --git a/Engine/Plugins/Experimental/AnimNext/Source/AnimNext/Private/Param/PropertyBagProxy.h b/Engine/Plugins/Experimental/AnimNext/Source/AnimNext/Private/Param/PropertyBagProxy.h
--- a/Engine/Plugins/Experimental/AnimNext/Source/AnimNext/Private/Param/PropertyBagProxy.h
+++ b/Engine/Plugins/Experimental/AnimNext/Source/AnimNext/Private/Param/PropertyBagProxy.h
@@ -1,4 +1,4 @@
-﻿// Copyright Epic Games, Inc. All Rights Reserved.
+// Copyright Epic Games, Inc. All Rights Reserved.
 
 #pragma once
 
@@ -14,7 +14,7 @@ struct FPropertyBagProxy : public IParameterSource
 	FPropertyBagProxy();
 
 	// IParameterSource interface
-	virtual void Update() override {};
+	virtual void Update(float DeltaTime) override {};
 	virtual const FParamStackLayerHandle& GetLayerHandle() const override { return LayerHandle; }
 	virtual void AddReferencedObjects(FReferenceCollector& Collector) override;
 
diff --git a/Engine/Plugins/Experimental/AnimNext/Source/AnimNext/Private/Scheduler/AnimNextScheduleExternalParamTask.cpp b/Engine/Plugins/Experimental/AnimNext/Source/AnimNext/Private/Scheduler/AnimNextScheduleExternalParamTask.cpp
--- a/Engine/Plugins/Experimental/AnimNext/Source/AnimNext/Private/Scheduler/AnimNextScheduleExternalParamTask.cpp
+++ b/Engine/Plugins/Experimental/AnimNext/Source/AnimNext/Private/Scheduler/AnimNextScheduleExternalParamTask.cpp
@@ -16,9 +16,10 @@ void FAnimNextScheduleExternalParamTask::UpdateExternalParams(const UE::AnimNext
 	
 	FScheduleInstanceData& InstanceData = InScheduleContext.GetInstanceData();
 	FScheduleInstanceData::FExternalParamCache& ExternalParamCache = InstanceData.ExternalParamCaches[TaskIndex];
+	const float DeltaTime = InScheduleContext.GetDeltaTime();
 	
 	for (TUniquePtr<IParameterSource>& ParameterSource : ExternalParamCache.ParameterSources)
 	{
-		ParameterSource->Update();
+		ParameterSource->Update(DeltaTime);
 	}
 }
diff --git a/Engine/Plugins/Experimental/AnimNext/Source/AnimNext/Private/Scheduler/AnimNextScheduleParamScopeTask.cpp b/Engine/Plugins/Experimental/AnimNext/Source/AnimNext/Private/Scheduler/AnimNextScheduleParamScopeTask.cpp
--- a/Engine/Plugins/Experimental/AnimNext/Source/AnimNext/Private/Scheduler/AnimNextScheduleParamScopeTask.cpp
+++ b/Engine/Plugins/Experimental/AnimNext/Source/AnimNext/Private/Scheduler/AnimNextScheduleParamScopeTask.cpp
@@ -17,6 +17,7 @@ void FAnimNextScheduleParamScopeEntryTask::RunParamScopeEntry(const UE::AnimNext
 
 	FParamStack& ParamStack = FParamStack::Get();
 
+	const float DeltaTime = InScheduleContext.GetDeltaTime();
 	FScheduleInstanceData& InstanceData = InScheduleContext.GetInstanceData();
 	FScheduleInstanceData::FScopeCache& ScopeCache = InstanceData.ScopeCaches[ParamScopeIndex];
 	check(ScopeCache.PushedLayers.Num() == 0);
@@ -26,7 +27,7 @@ void FAnimNextScheduleParamScopeEntryTask::RunParamScopeEntry(const UE::AnimNext
 	if (FoundUserScope && FoundUserScope->BeforeSource.IsValid())
 	{
 		FPropertyBagProxy& UserParameterSource = *FoundUserScope->BeforeSource.Get();
-		UserParameterSource.Update();
+		UserParameterSource.Update(DeltaTime);
 		ScopeCache.PushedLayers.Add(ParamStack.PushLayer(UserParameterSource.GetLayerHandle()));
 	}
 
@@ -35,7 +36,7 @@ void FAnimNextScheduleParamScopeEntryTask::RunParamScopeEntry(const UE::AnimNext
 	{
 		TUniquePtr<IParameterSource>& StaticParameterSource = ScopeCache.ParameterSources[ParameterSourceIndex];
 
-		StaticParameterSource->Update();
+		StaticParameterSource->Update(DeltaTime);
 		ScopeCache.PushedLayers.Add(ParamStack.PushLayer(StaticParameterSource->GetLayerHandle()));
 	}
 
@@ -43,7 +44,7 @@ void FAnimNextScheduleParamScopeEntryTask::RunParamScopeEntry(const UE::AnimNext
 	if (FoundUserScope && FoundUserScope->AfterSource.IsValid())
 	{
 		FPropertyBagProxy& UserParameterSource = *FoundUserScope->AfterSource.Get();
-		UserParameterSource.Update();
+		UserParameterSource.Update(DeltaTime);
 		ScopeCache.PushedLayers.Add(ParamStack.PushLayer(UserParameterSource.GetLayerHandle()));
 	}
 }
diff --git a/Engine/Plugins/Experimental/AnimNext/Source/AnimNext/Private/Scheduler/ScheduleInstanceData.cpp b/Engine/Plugins/Experimental/AnimNext/Source/AnimNext/Private/Scheduler/ScheduleInstanceData.cpp
--- a/Engine/Plugins/Experimental/AnimNext/Source/AnimNext/Private/Scheduler/ScheduleInstanceData.cpp
+++ b/Engine/Plugins/Experimental/AnimNext/Source/AnimNext/Private/Scheduler/ScheduleInstanceData.cpp
@@ -75,6 +75,7 @@ FScheduleInstanceData::FScheduleInstanceData(const FScheduleContext& InScheduleC
 	FExternalParameterContext ExternalParameterContext;
 	ExternalParameterContext.Object = Entry->WeakObject.Get();
 
+	const float DeltaTime = InScheduleContext.GetDeltaTime();
 	for(int32 ExternalParamSourceIndex = 0; ExternalParamSourceIndex < ExternalParamCaches.Num(); ++ExternalParamSourceIndex)
 	{
 		FExternalParamCache& ExternalParamCache = ExternalParamCaches[ExternalParamSourceIndex];
@@ -85,7 +86,7 @@ FScheduleInstanceData::FScheduleInstanceData(const FScheduleContext& InScheduleC
 			{
 				// Initial update is required to populate the cache
 				// TODO: This needs to move outside this function once we run initialization off the game thread, depending on thread-safety
-				NewParameterSource->Update();
+				NewParameterSource->Update(DeltaTime);
 
 				// External parameter layer is always pushed
 				RootParamStack->PushLayer(NewParameterSource->GetLayerHandle());
diff --git a/Engine/Plugins/Experimental/AnimNext/Source/AnimNext/Public/Param/IParameterSource.h b/Engine/Plugins/Experimental/AnimNext/Source/AnimNext/Public/Param/IParameterSource.h
--- a/Engine/Plugins/Experimental/AnimNext/Source/AnimNext/Public/Param/IParameterSource.h
+++ b/Engine/Plugins/Experimental/AnimNext/Source/AnimNext/Public/Param/IParameterSource.h
@@ -20,7 +20,7 @@ public:
 	virtual ~IParameterSource() = default;
 	
 	/** Update this parameter source */
-	virtual void Update() = 0;
+	virtual void Update(float DeltaTime) = 0;
 
 	/** Get a layer handle to allow this source to be used in the parameter stack */
 	virtual const FParamStackLayerHandle& GetLayerHandle() const = 0; 
