AnimNext: Only external parameters are needed to specified as 'required'

diff --git a/Engine/Plugins/Experimental/AnimNext/Source/AnimNext/Private/Scheduler/AnimNextScheduleGraphTask.cpp b/Engine/Plugins/Experimental/AnimNext/Source/AnimNext/Private/Scheduler/AnimNextScheduleGraphTask.cpp
--- a/Engine/Plugins/Experimental/AnimNext/Source/AnimNext/Private/Scheduler/AnimNextScheduleGraphTask.cpp
+++ b/Engine/Plugins/Experimental/AnimNext/Source/AnimNext/Private/Scheduler/AnimNextScheduleGraphTask.cpp
@@ -57,7 +57,7 @@ void FAnimNextScheduleGraphTask::VerifyRequiredParameters(UAnimNextGraph* InGrap
 
 			if(!bWarningOutput && (!bFound || !bFoundCorrectType))
 			{
-				UE_LOGFMT(LogAnimation, Warning, "AnimNext: Graph {GraphToRun} has different required parameters, it may not run correctly.", InGraphToRun->GetName());
+				UE_LOGFMT(LogAnimation, Warning, "AnimNext: Graph {GraphToRun} has different required parameters, it may not run correctly.", InGraphToRun->GetFName());
 				bWarningOutput = true;
 			}
 			
diff --git a/Engine/Plugins/Experimental/AnimNext/Source/AnimNextUncookedOnly/Private/UncookedOnlyUtils.cpp b/Engine/Plugins/Experimental/AnimNext/Source/AnimNextUncookedOnly/Private/UncookedOnlyUtils.cpp
--- a/Engine/Plugins/Experimental/AnimNext/Source/AnimNextUncookedOnly/Private/UncookedOnlyUtils.cpp
+++ b/Engine/Plugins/Experimental/AnimNext/Source/AnimNextUncookedOnly/Private/UncookedOnlyUtils.cpp
@@ -537,7 +537,7 @@ void FUtils::Compile(UAnimNextGraph* InGraph)
 	// Now that the graph has been re-compiled, re-allocate the previous live instances
 	InGraph->ThawGraphInstances();
 
-	// Gather and store the required parameters
+	// Gather and store the required external parameters
 	InGraph->RequiredParameters.Empty();
 
 	FAssetData AssetData(InGraph);
@@ -546,7 +546,11 @@ void FUtils::Compile(UAnimNextGraph* InGraph)
 
 	for(FAnimNextParameterAssetRegistryExportEntry& Entry : Exports.Parameters)
 	{
-		InGraph->RequiredParameters.Emplace(Entry.Name, Entry.Type);
+		FName ParameterSourceName = FExternalParameterRegistry::FindSourceForParameter(Entry.Name);
+		if(ParameterSourceName != NAME_None)
+		{
+			InGraph->RequiredParameters.Emplace(Entry.Name, Entry.Type);
+		}
 	}
 
 	InGraph->RequiredParametersHash = SortAndHashParameters(InGraph->RequiredParameters);
