AnimNext: Fix deletion of nodes not removing underlying RigVM nodes

diff --git a/Engine/Plugins/Experimental/AnimNext/Source/AnimNextEditor/Private/Workspace/AnimNextWorkspaceEditor.cpp b/Engine/Plugins/Experimental/AnimNext/Source/AnimNextEditor/Private/Workspace/AnimNextWorkspaceEditor.cpp
--- a/Engine/Plugins/Experimental/AnimNext/Source/AnimNextEditor/Private/Workspace/AnimNextWorkspaceEditor.cpp
+++ b/Engine/Plugins/Experimental/AnimNext/Source/AnimNextEditor/Private/Workspace/AnimNextWorkspaceEditor.cpp
@@ -422,69 +422,53 @@ FGraphPanelSelectionSet FWorkspaceEditor::GetSelectedNodes() const
 
 void FWorkspaceEditor::DeleteSelectedNodes()
 {
-	TSharedPtr<SGraphEditor> FocusedGraphEd = FocusedGraphEdPtr.Pin();
-	if (!FocusedGraphEd.IsValid())
-	{
-		return;
-	}
-
-	const FScopedTransaction Transaction(FGenericCommands::Get().Delete->GetDescription());
-	FocusedGraphEd->GetCurrentGraph()->Modify();
-	
 	const FGraphPanelSelectionSet SelectedNodes = GetSelectedNodes();
-	
-	if(FocusedGraphEd)
-	{
-		FocusedGraphEd->ClearSelectionSet();
-	}
 
-	// Some nodes have sub-objects that are represented as other tabs.
-	// Close them here as a pre-pass before we remove their nodes. If the documents are left open they
-	// may reference dangling data and function incorrectly in cases such as FindBlueprintforNodeChecked
-	for (FGraphPanelSelectionSet::TConstIterator NodeIt( SelectedNodes ); NodeIt; ++NodeIt)
+	bool bRelinkPins = false;
+	TArray<URigVMNode*> NodesToRemove;
+
+	for (FGraphPanelSelectionSet::TConstIterator NodeIt(SelectedNodes); NodeIt; ++NodeIt)
 	{
 		if (UEdGraphNode* Node = Cast<UEdGraphNode>(*NodeIt))
 		{
 			if (Node->CanUserDeleteNode())
 			{
-				auto CloseAllDocumentsTab = [this](const UEdGraphNode* InNode)
+				if (const URigVMEdGraphNode* RigVMEdGraphNode = Cast<URigVMEdGraphNode>(Node))
 				{
-					TArray<UObject*> NodesToClose;
-					GetObjectsWithOuter(InNode, NodesToClose);
-					for (UObject* Node : NodesToClose)
+					bRelinkPins = bRelinkPins || FSlateApplication::Get().GetModifierKeys().IsShiftDown();
+
+					if(URigVMNode* ModelNode = GetFocusedVMController()->GetGraph()->FindNodeByName(*RigVMEdGraphNode->GetModelNodePath()))
 					{
-						UEdGraph* NodeGraph = Cast<UEdGraph>(Node);
-						if (NodeGraph)
-						{
-							CloseDocumentTab(NodeGraph);
-						}
+						NodesToRemove.Add(ModelNode);
 					}
-				};
-				
-				if (Node->GetSubGraphs().Num() > 0)
+				}
+				else if (const UEdGraphNode_Comment* CommentNode = Cast<UEdGraphNode_Comment>(Node))
 				{
-					CloseAllDocumentsTab(Node);
+					if(URigVMNode* ModelNode = GetFocusedVMController()->GetGraph()->FindNodeByName(CommentNode->GetFName()))
+					{
+						NodesToRemove.Add(ModelNode);
+					}
+				}
+				else
+				{
+					Node->GetGraph()->RemoveNode(Node);
 				}
 			}
 		}
 	}
 
-	// Now remove the selected nodes
-	for (FGraphPanelSelectionSet::TConstIterator NodeIt( SelectedNodes ); NodeIt; ++NodeIt)
+	if(NodesToRemove.IsEmpty())
 	{
-		if (UEdGraphNode* Node = Cast<UEdGraphNode>(*NodeIt))
-		{
-			if (Node->CanUserDeleteNode())
-			{
-				if (Node->GetSubGraphs().Num() > 0)
-				{
-					DocumentManager->CleanInvalidTabs();
-				}
+		return;
+	}
 
-				FBlueprintEditorUtils::RemoveNode(nullptr, Node);
-			}
-		}
+	GetFocusedVMController()->OpenUndoBracket(TEXT("Delete selected nodes"));
+	if(bRelinkPins && NodesToRemove.Num() == 1)
+	{
+		GetFocusedVMController()->RelinkSourceAndTargetPins(NodesToRemove[0], true);;
 	}
+	GetFocusedVMController()->RemoveNodes(NodesToRemove, true);
+	GetFocusedVMController()->CloseUndoBracket();
 }
 
 bool FWorkspaceEditor::CanDeleteSelectedNodes()
