AnimNext: Fix incorrect parameter reporting in mesh port

diff --git a/Engine/Plugins/Experimental/AnimNext/Source/AnimNext/Private/Scheduler/Ports/AnimNextSchedulePort_AnimNextMeshComponentPose.cpp b/Engine/Plugins/Experimental/AnimNext/Source/AnimNext/Private/Scheduler/Ports/AnimNextSchedulePort_AnimNextMeshComponentPose.cpp
--- a/Engine/Plugins/Experimental/AnimNext/Source/AnimNext/Private/Scheduler/Ports/AnimNextSchedulePort_AnimNextMeshComponentPose.cpp
+++ b/Engine/Plugins/Experimental/AnimNext/Source/AnimNext/Private/Scheduler/Ports/AnimNextSchedulePort_AnimNextMeshComponentPose.cpp
@@ -92,7 +92,7 @@ TConstArrayView<FAnimNextParam> UAnimNextSchedulePort_AnimNextMeshComponentPose:
 	static const FAnimNextParam Params[] =
 	{
 		FAnimNextParam(ComponentParamId.GetName(), FAnimNextParamType::GetType<TObjectPtr<UAnimNextMeshComponent>>()),
-		FAnimNextParam(ReferencePoseParamId.GetName(), FAnimNextParamType::GetType<FAnimNextGraphLODPose>()),
+		FAnimNextParam(ReferencePoseParamId.GetName(), FAnimNextParamType::GetType<FAnimNextGraphReferencePose>()),
 	};
 
 	return Params;
