AnimNext: Begin merging parameter blocks and graphs

Moved a lot of code from UAnimNextGraph and UAnimNextParameterBlock (and their respective _EditorData classes) into a common base: UAnimNextRigVMAsset + UAnimNextRigVMAssetEditorData.
Refactored SParameterBlockView to work with the new classes and displayed the new view for UAnimNextGraph assets instead of the old 'open graph' button.
Updated tests with additional scopes and GC/transaction cleanup code to prevent UObject re-use on subsequent test runs.

Misc fixes:
- Fixed workspace editor not calling the base class OnClose, which was causing it to miss saving open documents correctly
- Fixed parameter name validation to use the FUtils verifier all the time

RigVM stuff:
- Made URigVMEdGraph::bIsSelecting public so AnimNext widgets can properly participate in node selection
- Added (opt-in) option to FRigVMClient to allow deletion of the default model
- Removed restriction around matching packages in FRigVMClient::GetModel. UEdGraphs can now be held in different OFPA packages.
- Removed ensure in URigVM::GetHostCDO as this is no longer valid for AnimNext VMs

diff --git a/Engine/Plugins/Experimental/AnimNext/Source/AnimNext/Internal/AnimNextRigVMAsset.h b/Engine/Plugins/Experimental/AnimNext/Source/AnimNext/Internal/AnimNextRigVMAsset.h
new file mode 100644
--- /dev/null
+++ b/Engine/Plugins/Experimental/AnimNext/Source/AnimNext/Internal/AnimNextRigVMAsset.h
@@ -0,0 +1,48 @@
+﻿// Copyright Epic Games, Inc. All Rights Reserved.
+
+#pragma once
+
+#include "CoreMinimal.h"
+#include "RigVMCore/RigVMExecuteContext.h"
+#include "RigVMHost.h"
+#include "AnimNextRigVMAsset.generated.h"
+
+class UAnimNextRigVMAssetEditorData;
+
+namespace UE::AnimNext::UncookedOnly
+{
+	struct FUtils;
+}
+
+/** Base class for all AnimNext assets that can host RigVM logic */
+UCLASS(MinimalAPI, Abstract)
+class UAnimNextRigVMAsset : public URigVMHost
+{
+	GENERATED_BODY()
+
+protected:
+	friend class UAnimNextRigVMAssetEditorData;
+	friend struct UE::AnimNext::UncookedOnly::FUtils;
+
+	UAnimNextRigVMAsset(const FObjectInitializer& ObjectInitializer);
+
+	// UObject interface
+	virtual void BeginDestroy() override;
+	virtual void PostLoad() override;
+	virtual void GetAssetRegistryTags(FAssetRegistryTagsContext Context) const override; 
+
+	// The ExtendedExecuteContext object holds the common work data used by the RigVM internals. It is populated during the initial VM initialization.
+	// Each instance of an AnimGraph requires a copy of this context and a call to initialize the VM instance with the context copy, 
+	// so the cached memory handles are updated to the correct memory addresses.
+	// This context is used as a reference to copy the common data for all instances created.
+	UPROPERTY(Transient)
+	FRigVMExtendedExecuteContext ExtendedExecuteContext;
+
+	UPROPERTY()
+	TObjectPtr<URigVM> RigVM;
+
+#if WITH_EDITORONLY_DATA
+	UPROPERTY(VisibleAnywhere, Instanced, Category = "Editor Data", meta = (ShowInnerProperties))
+	TObjectPtr<UObject> EditorData;
+#endif
+};
diff --git a/Engine/Plugins/Experimental/AnimNext/Source/AnimNext/Internal/Graph/AnimNextGraph.h b/Engine/Plugins/Experimental/AnimNext/Source/AnimNext/Internal/Graph/AnimNextGraph.h
--- a/Engine/Plugins/Experimental/AnimNext/Source/AnimNext/Internal/Graph/AnimNextGraph.h
+++ b/Engine/Plugins/Experimental/AnimNext/Source/AnimNext/Internal/Graph/AnimNextGraph.h
@@ -3,6 +3,7 @@
 #pragma once
 
 #include "CoreMinimal.h"
+#include "AnimNextRigVMAsset.h"
 #include "RigVMCore/RigVM.h"
 #include "DecoratorBase/DecoratorPtr.h"
 #include "DecoratorBase/DecoratorHandle.h"
@@ -51,7 +52,7 @@ namespace UE::AnimNext::Graph
 
 // A user-created graph of logic used to supply data
 UCLASS(BlueprintType)
-class ANIMNEXT_API UAnimNextGraph :  public URigVMHost, public IAnimNextScheduleTermInterface
+class ANIMNEXT_API UAnimNextGraph : public UAnimNextRigVMAsset, public IAnimNextScheduleTermInterface
 {
 	GENERATED_BODY()
 
@@ -80,14 +81,7 @@ public:
 	// Get the parameter to use to access the current LOD
 	UE::AnimNext::FParamId GetCurrentLODParam() const { return CurrentLODId; }
 
-#if WITH_EDITOR
-	virtual void GetAssetRegistryTags(FAssetRegistryTagsContext Context) const override;
-	UE_DEPRECATED(5.4, "Implement the version that takes FAssetRegistryTagsContext instead.")
-	virtual void GetAssetRegistryTags(TArray<FAssetRegistryTag>& OutTags) const override;
-#endif	
 protected:
-	// Support rig VM execution
-	TArray<FRigVMExternalVariable> GetRigVMExternalVariables();
 
 	// Loads the graph data from the provided archive buffer and returns true on success, false otherwise
 	bool LoadFromArchiveBuffer(const TArray<uint8>& SharedDataArchiveBuffer);
@@ -145,17 +139,6 @@ protected:
 	UPROPERTY()
 	TArray<TObjectPtr<UObject>> GraphReferencedObjects;
 
-	// The RigVM object holds the bytecode, literals, etc used by the RigVM internals, there is a single instance along with the UAnimNextGraph (1:1 mapping)
-	UPROPERTY()
-	TObjectPtr<URigVM> RigVM;
-
-	// The ExtendedExecuteContext object holds the common work data used by the RigVM internals. It is populated during the initial VM initialization.
-	// Each instance of an AnimGraph requires a copy of this context and a call to initialize the VM instance with the context copy, 
-	// so the cached memory handles are updated to the correct memory addresses.
-	// This context is used as a reference to copy the common data for all instances created.
-	UPROPERTY(Transient)
-	FRigVMExtendedExecuteContext ExtendedExecuteContext;
-
 	// The parameter to use to access the reference pose
 	UPROPERTY(EditAnywhere, Category = "Graph", meta=(CustomWidget = "ParamName", AllowedParamType = "FAnimNextGraphReferencePose"))
 	FName ReferencePose = DefaultReferencePoseId.GetName();
@@ -176,9 +159,6 @@ protected:
 	TArray<FAnimNextParam> RequiredParameters;
 
 #if WITH_EDITORONLY_DATA
-	UPROPERTY(VisibleAnywhere, Instanced, Category = "Graph", meta = (ShowInnerProperties))
-	TObjectPtr<UObject> EditorData;
-
 	// This buffer holds the output of the FDecoratorWriter post compilation
 	// We serialize it manually and it is discarded at runtime
 	TArray<uint8> SharedDataArchiveBuffer;
diff --git a/Engine/Plugins/Experimental/AnimNext/Source/AnimNext/Internal/Param/AnimNextParameterBlock.h b/Engine/Plugins/Experimental/AnimNext/Source/AnimNext/Internal/Param/AnimNextParameterBlock.h
--- a/Engine/Plugins/Experimental/AnimNext/Source/AnimNext/Internal/Param/AnimNextParameterBlock.h
+++ b/Engine/Plugins/Experimental/AnimNext/Source/AnimNext/Internal/Param/AnimNextParameterBlock.h
@@ -3,15 +3,15 @@
 #pragma once
 
 #include "CoreMinimal.h"
+#include "AnimNextRigVMAsset.h"
 #include "PropertyBag.h"
-#include "RigVMCore/RigVMExecuteContext.h"
 #include "Param/IParameterSource.h"
-#include "RigVMHost.h"
 #include "AnimNextParameterBlock.generated.h"
 
 class UEdGraph;
 struct FAnimNextScheduleGraphTask;
 struct FAnimNextScheduleParamScopeTask;
+class UAnimNextParameterBlockParameter;
 
 namespace UE::AnimNext
 {
@@ -29,13 +29,13 @@ namespace UE::AnimNext::Editor
 {
 	class FParametersEditor;
 	struct FUtils;
-	class SParameterBlockViewRow;
+	class SRigVMAssetViewRow;
 	class FParameterBlockParameterCustomization;
 }
 
 /** An asset used to define AnimNext parameters and their bindings */
 UCLASS(MinimalAPI, BlueprintType)
-class UAnimNextParameterBlock : public URigVMHost
+class UAnimNextParameterBlock : public UAnimNextRigVMAsset
 {
 	GENERATED_BODY()
 
@@ -50,31 +50,13 @@ class UAnimNextParameterBlock : public URigVMHost
 	friend struct FAnimNode_AnimNextParameters;
 	friend struct FAnimNextScheduleGraphTask;
 	friend struct FAnimNextScheduleParamScopeEntryTask;
-	friend class UE::AnimNext::Editor::SParameterBlockViewRow;
+	friend class UE::AnimNext::Editor::SRigVMAssetViewRow;
 	friend class UE::AnimNext::Editor::FParameterBlockParameterCustomization;
 	friend struct UE::AnimNext::FParameterBlockProxy;
+	friend class UAnimNextParameterBlockParameter;
 
 	void UpdateLayer(UE::AnimNext::FParamStackLayerHandle& InHandle, float InDeltaTime) const;
 
-	// UObject interface
-	virtual void BeginDestroy() override;
-	virtual void PostLoad() override;
-	virtual void GetAssetRegistryTags(FAssetRegistryTagsContext Context) const override;
-	UE_DEPRECATED(5.4, "Implement the version that takes FAssetRegistryTagsContext instead.")
-	virtual void GetAssetRegistryTags(TArray<FAssetRegistryTag>& OutTags) const override;
-
-	FInstancedPropertyBag& GetPropertyBag() { return PropertyBag; }
-
-	FRigVMExtendedExecuteContext BaseRigVMContext;
-
-	UPROPERTY()
-	TObjectPtr<URigVM> RigVM;
-
 	UPROPERTY()
 	FInstancedPropertyBag PropertyBag;
-
-#if WITH_EDITORONLY_DATA
-	UPROPERTY(VisibleAnywhere, Instanced, Category = "Parameters", meta = (ShowInnerProperties))
-	TObjectPtr<UObject> EditorData;
-#endif
 };
diff --git a/Engine/Plugins/Experimental/AnimNext/Source/AnimNext/Private/AnimNextRigVMAsset.cpp b/Engine/Plugins/Experimental/AnimNext/Source/AnimNext/Private/AnimNextRigVMAsset.cpp
new file mode 100644
--- /dev/null
+++ b/Engine/Plugins/Experimental/AnimNext/Source/AnimNext/Private/AnimNextRigVMAsset.cpp
@@ -0,0 +1,54 @@
+﻿// Copyright Epic Games, Inc. All Rights Reserved.
+
+#include "AnimNextRigVMAsset.h"
+#include "RigVMRuntimeDataRegistry.h"
+#include "Graph/RigUnit_AnimNextBeginExecution.h"
+#include "UObject/AssetRegistryTagsContext.h"
+
+UAnimNextRigVMAsset::UAnimNextRigVMAsset(const FObjectInitializer& ObjectInitializer)
+	: Super(ObjectInitializer)
+{
+	SetRigVMExtendedExecuteContext(&ExtendedExecuteContext);
+}
+
+void UAnimNextRigVMAsset::BeginDestroy()
+{
+	Super::BeginDestroy();
+
+	if (VM)
+	{
+		UE::AnimNext::FRigVMRuntimeDataRegistry::ReleaseAllVMRuntimeData(VM);
+	}
+}
+
+void UAnimNextRigVMAsset::PostLoad()
+{
+	Super::PostLoad();
+
+	ExtendedExecuteContext.InvalidateCachedMemory();
+
+	VM = RigVM;
+
+	// In packaged builds, initialize the VM
+	// In editor, the VM will be recompiled and initialized at UAnimNextRigVMAssetEditorData::HandlePackageDone::RecompileVM
+#if !WITH_EDITOR
+	if(VM != nullptr)
+	{
+		VM->ClearExternalVariables(ExtendedExecuteContext);
+		VM->Initialize(ExtendedExecuteContext);
+		InitializeVM(FRigUnit_AnimNextBeginExecution::EventName);
+	}
+#endif
+}
+
+void UAnimNextRigVMAsset::GetAssetRegistryTags(FAssetRegistryTagsContext Context) const
+{
+	Super::GetAssetRegistryTags(Context);
+
+#if WITH_EDITORONLY_DATA
+	if(EditorData)
+	{
+		EditorData->GetAssetRegistryTags(Context);
+	}
+#endif
+}
diff --git a/Engine/Plugins/Experimental/AnimNext/Source/AnimNext/Private/Graph/AnimNextGraph.cpp b/Engine/Plugins/Experimental/AnimNext/Source/AnimNext/Private/Graph/AnimNextGraph.cpp
--- a/Engine/Plugins/Experimental/AnimNext/Source/AnimNext/Private/Graph/AnimNextGraph.cpp
+++ b/Engine/Plugins/Experimental/AnimNext/Source/AnimNext/Private/Graph/AnimNextGraph.cpp
@@ -27,7 +27,6 @@ UAnimNextGraph::UAnimNextGraph(const FObjectInitializer& ObjectInitializer)
 	: Super(ObjectInitializer)
 {
 	ExtendedExecuteContext.SetContextPublicDataStruct(FAnimNextExecuteContext::StaticStruct());
-	SetRigVMExtendedExecuteContext(&ExtendedExecuteContext);
 }
 
 void UAnimNextGraph::AllocateInstance(FAnimNextGraphInstancePtr& Instance) const
@@ -88,50 +87,11 @@ void UAnimNextGraph::AllocateInstanceImpl(FAnimNextGraphInstance* ParentGraphIns
 #endif
 }
 
-
-#if WITH_EDITOR
-void UAnimNextGraph::GetAssetRegistryTags(TArray<FAssetRegistryTag>& OutTags) const
-{
-	PRAGMA_DISABLE_DEPRECATION_WARNINGS;
-	Super::GetAssetRegistryTags(OutTags);
-	PRAGMA_ENABLE_DEPRECATION_WARNINGS;
-}
-
-void UAnimNextGraph::GetAssetRegistryTags(FAssetRegistryTagsContext Context) const
-{
-	Super::GetAssetRegistryTags(Context);
-
-	if (EditorData)
-	{		
-		EditorData->GetAssetRegistryTags(Context);
-	}
-}
-#endif
-
-TArray<FRigVMExternalVariable> UAnimNextGraph::GetRigVMExternalVariables()
-{
-	return TArray<FRigVMExternalVariable>(); 
-}
-
 void UAnimNextGraph::PostLoad()
 {
 	using namespace UE::AnimNext;
-	
-	Super::PostLoad();
-
-	VM = RigVM;
 
-	ExtendedExecuteContext.InvalidateCachedMemory();
-
-	// In packaged builds, initialize the VM
-	// In editor, the VM will be recompiled and initialized at RecompileVM
-#if !WITH_EDITOR
-	if (VM)
-	{
-		VM->ClearExternalVariables(ExtendedExecuteContext);
-		VM->Initialize(ExtendedExecuteContext);
-	}
-#endif
+	Super::PostLoad();
 
 	ReferencePoseId = FParamId(ReferencePose);
 	CurrentLODId = FParamId(CurrentLOD);
diff --git a/Engine/Plugins/Experimental/AnimNext/Source/AnimNext/Private/Param/AnimNextParameterBlock.cpp b/Engine/Plugins/Experimental/AnimNext/Source/AnimNext/Private/Param/AnimNextParameterBlock.cpp
--- a/Engine/Plugins/Experimental/AnimNext/Source/AnimNext/Private/Param/AnimNextParameterBlock.cpp
+++ b/Engine/Plugins/Experimental/AnimNext/Source/AnimNext/Private/Param/AnimNextParameterBlock.cpp
@@ -17,8 +17,7 @@ DEFINE_STAT(STAT_AnimNext_ParamBlock_UpdateLayer);
 UAnimNextParameterBlock::UAnimNextParameterBlock(const FObjectInitializer& ObjectInitializer)
 	: Super(ObjectInitializer)
 {
-	BaseRigVMContext.SetContextPublicDataStruct(FAnimNextParameterExecuteContext::StaticStruct());
-	SetRigVMExtendedExecuteContext(&BaseRigVMContext);
+	ExtendedExecuteContext.SetContextPublicDataStruct(FAnimNextParameterExecuteContext::StaticStruct());
 }
 
 void UAnimNextParameterBlock::UpdateLayer(UE::AnimNext::FParamStackLayerHandle& InHandle, float InDeltaTime) const
@@ -45,49 +44,3 @@ void UAnimNextParameterBlock::UpdateLayer(UE::AnimNext::FParamStackLayerHandle&
 		}
 	}
 }
-
-void UAnimNextParameterBlock::BeginDestroy()
-{
-	Super::BeginDestroy();
-
-	if (VM)
-	{
-		UE::AnimNext::FRigVMRuntimeDataRegistry::ReleaseAllVMRuntimeData(VM);
-	}
-}
-
-void UAnimNextParameterBlock::PostLoad()
-{
-	Super::PostLoad();
-
-	VM = RigVM;
-
-	// In packaged builds, initialize the VM
-	// In editor, the VM will be recompiled and initialized at UAnimNextParameterBlock_EditorData::HandlePackageDone::RecompileVM
-#if !WITH_EDITOR
-	if(VM != nullptr)
-	{
-		VM->Initialize(BaseRigVMContext);
-		InitializeVM(FRigUnit_AnimNextBeginExecution::EventName);
-	}
-#endif
-}
-
-void UAnimNextParameterBlock::GetAssetRegistryTags(TArray<FAssetRegistryTag>& OutTags) const
-{
-	PRAGMA_DISABLE_DEPRECATION_WARNINGS;
-	Super::GetAssetRegistryTags(OutTags);
-	PRAGMA_ENABLE_DEPRECATION_WARNINGS;
-}
-
-void UAnimNextParameterBlock::GetAssetRegistryTags(FAssetRegistryTagsContext Context) const
-{
-	Super::GetAssetRegistryTags(Context);
-	
-#if WITH_EDITORONLY_DATA
-	if(EditorData)
-	{
-		EditorData->GetAssetRegistryTags(Context);
-	}
-#endif
-}
\ No newline at end of file
diff --git a/Engine/Plugins/Experimental/AnimNext/Source/AnimNextEditor/Private/AnimNextEditorModule.cpp b/Engine/Plugins/Experimental/AnimNext/Source/AnimNextEditor/Private/AnimNextEditorModule.cpp
--- a/Engine/Plugins/Experimental/AnimNext/Source/AnimNextEditor/Private/AnimNextEditorModule.cpp
+++ b/Engine/Plugins/Experimental/AnimNext/Source/AnimNextEditor/Private/AnimNextEditorModule.cpp
@@ -3,25 +3,26 @@
 #include "AnimNextEditorModule.h"
 
 #include "AnimNextConfig.h"
+#include "ISettingsModule.h"
+#include "ScopedTransaction.h"
+#include "SSimpleButton.h"
+#include "SSimpleComboButton.h"
+#include "UncookedOnlyUtils.h"
+#include "Common/SRigVMAssetView.h"
+#include "Framework/Application/SlateApplication.h"
+#include "Graph/AnimNextGraph.h"
 #include "Graph/AnimNextGraphPanelNodeFactory.h"
-#include "Param/ParamTypePropertyCustomization.h"
+#include "Graph/AnimNextGraph_EdGraphNodeCustomization.h"
+#include "Param/AnimNextParameterBlock.h"
+#include "Param/ParameterBlockParameterCustomization.h"
 #include "Param/ParameterPickerArgs.h"
 #include "Param/ParametersGraphPanelPinFactory.h"
 #include "Param/ParamNamePropertyCustomization.h"
+#include "Param/ParamPropertyCustomization.h"
+#include "Param/ParamTypePropertyCustomization.h"
 #include "Param/SParameterPicker.h"
-#include "ISettingsModule.h"
-#include "UncookedOnlyUtils.h"
-#include "Graph/AnimNextGraph.h"
-#include "Param/AnimNextParameterBlock.h"
-#include "Param/AnimNextParameterBlockEntry.h"
-#include "Param/IAnimNextParameterBlockGraphInterface.h"
-#include "Param/SParameterBlockView.h"
-#include "Graph/SAnimNextGraphView.h"
 #include "Scheduler/AnimNextSchedule.h"
 #include "Workspace/AnimNextWorkspaceEditor.h"
-#include "Param/ParameterBlockParameterCustomization.h"
-#include "Graph/AnimNextGraph_EdGraphNodeCustomization.h"
-#include "Param/ParamPropertyCustomization.h"
 
 #define LOCTEXT_NAMESPACE "AnimNextEditorModule"
 
@@ -72,55 +73,12 @@ class FModule : public IModule
 		{
 			UAnimNextParameterBlock* ParameterBlock = CastChecked<UAnimNextParameterBlock>(InAsset);
 			UAnimNextParameterBlock_EditorData* EditorData = UncookedOnly::FUtils::GetEditorData(ParameterBlock);
-			return SNew(SParameterBlockView, EditorData)
-				.OnSelectionChanged_Lambda([WeakEditor = TWeakPtr<FWorkspaceEditor>(InEditor)](const TArray<UObject*>& InObjects)
-				{
-					if(TSharedPtr<FWorkspaceEditor> Editor = WeakEditor.Pin())
-					{
-						Editor->SetSelectedObjects(InObjects);
-					}
-				})
-				.OnOpenGraph_Lambda([WeakEditor = TWeakPtr<FWorkspaceEditor>(InEditor)](URigVMGraph* InGraph)
-				{
-					if(TSharedPtr<FWorkspaceEditor> Editor = WeakEditor.Pin())
-					{
-						if(IRigVMClientHost* RigVMClientHost = InGraph->GetImplementingOuter<IRigVMClientHost>())
-						{
-							if(UObject* EditorObject = RigVMClientHost->GetEditorObjectForRigVMGraph(InGraph))
-							{
-								Editor->OpenDocument(EditorObject, FDocumentTracker::EOpenDocumentCause::OpenNewDocument);
-							}
-						}
-					}
-				})
-				.OnDeleteEntries_Lambda([WeakEditor = TWeakPtr<FWorkspaceEditor>(InEditor)](const TArray<UAnimNextParameterBlockEntry*>& InEntries)
-				{
-					if(InEntries.Num() > 0)
-					{
-						if(TSharedPtr<FWorkspaceEditor> Editor = WeakEditor.Pin())
-						{
-							if(IRigVMClientHost* RigVMClientHost = InEntries[0]->GetImplementingOuter<IRigVMClientHost>())
-							{
-								for(UAnimNextParameterBlockEntry* Entry : InEntries)
-								{
-									if(IAnimNextParameterBlockGraphInterface* GraphInterface = Cast<IAnimNextParameterBlockGraphInterface>(Entry))
-									{
-										if(URigVMGraph* RigVMGraph = GraphInterface->GetGraph())
-										{
-											if (UObject* EditorObject = RigVMClientHost->GetEditorObjectForRigVMGraph(RigVMGraph))
-											{
-												Editor->CloseDocumentTab(EditorObject);
-											}
-										}
-									}
-								}
-							}
-						}
-					}
-				});
+			return SNew(SRigVMAssetView, EditorData)
+				.OnSelectionChanged(&InEditor.Get(), &FWorkspaceEditor::SetSelectedObjects)
+				.OnOpenGraph(&InEditor.Get(), &FWorkspaceEditor::OnOpenGraph)
+				.OnDeleteEntries(&InEditor.Get(), &FWorkspaceEditor::OnDeleteEntries);
 		});
 
-
 		FWorkspaceEditor::RegisterAssetDocumentWidget(UAnimNextSchedule::StaticClass()->GetFName(), [](TSharedRef<FWorkspaceEditor> InEditor, UObject* InAsset)
 		{
 			UAnimNextSchedule* Schedule = CastChecked<UAnimNextSchedule>(InAsset);
@@ -136,23 +94,101 @@ class FModule : public IModule
 		{
 			UAnimNextGraph* Graph = CastChecked<UAnimNextGraph>(InAsset);
 			UAnimNextGraph_EditorData* EditorData = UncookedOnly::FUtils::GetEditorData(Graph);
-			
-			EditorData->GetRigVMGraphModifiedEvent().RemoveAll(&InEditor.Get());
-			EditorData->GetRigVMGraphModifiedEvent().AddSP(InEditor, &FWorkspaceEditor::OnGraphModified);
 
-			return SNew(SAnimNextGraphView, EditorData)
-				.OnOpenGraph_Lambda([WeakEditor = TWeakPtr<FWorkspaceEditor>(InEditor)](URigVMGraph* InGraph)
+			EditorData->RigVMGraphModifiedEvent.RemoveAll(&InEditor.Get());
+			EditorData->RigVMGraphModifiedEvent.AddSP(InEditor, &FWorkspaceEditor::OnGraphModified);
+
+			return SNew(SRigVMAssetView, EditorData)
+				.OnSelectionChanged(&InEditor.Get(), &FWorkspaceEditor::SetSelectedObjects)
+				.OnOpenGraph(&InEditor.Get(), &FWorkspaceEditor::OnOpenGraph)
+				.OnDeleteEntries(&InEditor.Get(), &FWorkspaceEditor::OnDeleteEntries);
+		});
+
+		SRigVMAssetView::RegisterCategoryFactory("Parameters", [](UAnimNextRigVMAssetEditorData* InEditorData)
+		{
+			UAnimNextParameterBlock_EditorData* EditorData = CastChecked<UAnimNextParameterBlock_EditorData>(InEditorData);
+			return SNew(SSimpleComboButton)
+				.Text(LOCTEXT("AddParameterButton", "Add Parameter"))
+				.Icon(FAppStyle::Get().GetBrush("Icons.Plus"))
+				.HasDownArrow(true)
+				.OnGetMenuContent_Lambda([EditorData]()
 				{
-					if(TSharedPtr<FWorkspaceEditor> Editor = WeakEditor.Pin())
+					FAssetData AssetData(UncookedOnly::FUtils::GetAsset(EditorData));
+					
+					FParameterPickerArgs Args;
+					Args.bMultiSelect = false;
+					Args.bShowBlocks = false;
+					Args.bShowBoundParameters = false;
+					Args.bShowBuiltInParameters = false; // Built-In parameters disabled for MVP
+					Args.OnFilterParameter = FOnFilterParameter::CreateLambda([EditorData, AssetData](const FParameterBindingReference& InParameterBinding)
+					{
+						// Skip params that are already bound in this block
+						if(InParameterBinding.Block == AssetData)
+						{
+							return EFilterParameterResult::Exclude;
+						}
+						
+						return EFilterParameterResult::Include;
+					});
+
+					Args.OnAddParameter = FOnAddParameter::CreateLambda([EditorData](const FParameterToAdd& ParameterToAdd)
 					{
-						if(IRigVMClientHost* RigVMClientHost = InGraph->GetImplementingOuter<IRigVMClientHost>())
+						FSlateApplication::Get().DismissAllMenus();
+
+						check(EditorData->FindEntry(ParameterToAdd.Name) == nullptr);
+						FScopedTransaction Transaction(LOCTEXT("AddParameter", "Add parameter"));
+						EditorData->AddParameter(ParameterToAdd.Name, ParameterToAdd.Type);
+					});
+					Args.OnParameterPicked = FOnParameterPicked::CreateLambda([EditorData](const FParameterBindingReference& InParameterBinding)
+					{
+						FSlateApplication::Get().DismissAllMenus();
+
+						if (EditorData->FindEntry(InParameterBinding.Parameter) == nullptr)
 						{
-							if(UObject* EditorObject = RigVMClientHost->GetEditorObjectForRigVMGraph(InGraph))
+							const FAnimNextParamType Type = UncookedOnly::FUtils::GetParameterTypeFromName(InParameterBinding.Parameter);
+							if (Type.IsValid())
 							{
-								Editor->OpenDocument(EditorObject, FDocumentTracker::EOpenDocumentCause::OpenNewDocument);
-							}
+								EditorData->AddParameter(InParameterBinding.Parameter, Type);
+							};
 						}
-					}
+					});
+					
+					return SNew(SParameterPicker)
+						.Args(Args);
+				});
+		});
+
+		SRigVMAssetView::RegisterCategoryFactory("Parameter Graphs", [](UAnimNextRigVMAssetEditorData* InEditorData)
+		{
+			UAnimNextParameterBlock_EditorData* EditorData = CastChecked<UAnimNextParameterBlock_EditorData>(InEditorData);
+			return SNew(SSimpleButton)
+				.Text(LOCTEXT("AddGraphButton", "Add Graph"))
+				.Icon(FAppStyle::Get().GetBrush("Icons.Plus"))
+				.OnClicked_Lambda([EditorData]()
+				{
+					FScopedTransaction Transaction(LOCTEXT("AddGraph", "Add Graph"));
+
+					// Create a new entry for the graph
+					EditorData->AddGraph(TEXT("NewGraph"));
+
+					return FReply::Handled();
+				});
+		});
+
+		SRigVMAssetView::RegisterCategoryFactory("Animation Graphs", [](UAnimNextRigVMAssetEditorData* InEditorData)
+		{
+			UAnimNextGraph_EditorData* EditorData = CastChecked<UAnimNextGraph_EditorData>(InEditorData);
+			return SNew(SSimpleButton)
+				.Text(LOCTEXT("AddGraphButton", "Add Graph"))
+				.Icon(FAppStyle::Get().GetBrush("Icons.Plus"))
+				.OnClicked_Lambda([EditorData]()
+				{
+					FScopedTransaction Transaction(LOCTEXT("AddGraph", "Add Graph"));
+
+					// Create a new entry for the graph
+					EditorData->AddGraph(TEXT("NewGraph"));
+
+					return FReply::Handled();
 				});
 		});
 	}
@@ -172,6 +208,13 @@ class FModule : public IModule
 		FEdGraphUtilities::UnregisterVisualNodeFactory(AnimNextGraphPanelNodeFactory);
 
 		FEdGraphUtilities::UnregisterVisualPinFactory(ParametersGraphPanelPinFactory);
+
+		FWorkspaceEditor::UnregisterAssetDocumentWidget(UAnimNextSchedule::StaticClass()->GetFName());
+		FWorkspaceEditor::UnregisterAssetDocumentWidget(UAnimNextParameterBlock::StaticClass()->GetFName());
+		FWorkspaceEditor::UnregisterAssetDocumentWidget(UAnimNextGraph::StaticClass()->GetFName());
+		
+		SRigVMAssetView::UnregisterCategoryFactory("Parameters");
+		SRigVMAssetView::UnregisterCategoryFactory("Parameter Graphs");
 	}
 
 	virtual TSharedRef<SWidget> CreateParameterPicker(const FParameterPickerArgs& InArgs) override
diff --git a/Engine/Plugins/Experimental/AnimNext/Source/AnimNextEditor/Private/Common/RigVMAssetViewMenuContext.h b/Engine/Plugins/Experimental/AnimNext/Source/AnimNextEditor/Private/Common/RigVMAssetViewMenuContext.h
new file mode 100644
--- /dev/null
+++ b/Engine/Plugins/Experimental/AnimNext/Source/AnimNextEditor/Private/Common/RigVMAssetViewMenuContext.h
@@ -0,0 +1,21 @@
+﻿// Copyright Epic Games, Inc. All Rights Reserved.
+
+#pragma once
+
+#include "RigVMAssetViewMenuContext.generated.h"
+
+namespace UE::AnimNext::Editor
+{
+	class SRigVMAssetView;
+}
+
+UCLASS()
+class URigVMAssetViewMenuContext : public UObject
+{
+	GENERATED_BODY()
+
+	friend class UE::AnimNext::Editor::SRigVMAssetView;
+	
+	// The RigVM asset view that we are editing
+	TWeakPtr<UE::AnimNext::Editor::SRigVMAssetView> RigVMAssetView;
+};
\ No newline at end of file
diff --git a/Engine/Plugins/Experimental/AnimNext/Source/AnimNextEditor/Private/Common/SRigVMAssetView.cpp b/Engine/Plugins/Experimental/AnimNext/Source/AnimNextEditor/Private/Common/SRigVMAssetView.cpp
new file mode 100644
--- /dev/null
+++ b/Engine/Plugins/Experimental/AnimNext/Source/AnimNextEditor/Private/Common/SRigVMAssetView.cpp
@@ -0,0 +1,993 @@
+// Copyright Epic Games, Inc. All Rights Reserved.
+
+#include "SRigVMAssetView.h"
+
+#include "AnimNextRigVMAssetEntry.h"
+#include "Param/AnimNextParameterBlock.h"
+#include "Param/AnimNextParameterBlock_EditorData.h"
+#include "Widgets/Layout/SBox.h"
+#include "Widgets/Input/SSearchBox.h"
+#include "DetailLayoutBuilder.h"
+#include "ISourceControlModule.h"
+#include "ISourceControlProvider.h"
+#include "EditorUtils.h"
+#include "InstancedPropertyBagStructureDataProvider.h"
+#include "UncookedOnlyUtils.h"
+#include "RigVMAssetViewMenuContext.h"
+#include "PropertyBagDetails.h"
+#include "SourceControlOperations.h"
+#include "SPinTypeSelector.h"
+#include "Framework/Commands/GenericCommands.h"
+#include "RevisionControlStyle/RevisionControlStyle.h"
+#include "Widgets/Text/SInlineEditableTextBlock.h"
+#include "ToolMenus.h"
+#include "ScopedTransaction.h"
+#include "Param/AnimNextParameterBlockParameter.h"
+#include "Param/AnimNextParameterBlockGraph.h"
+#include "Framework/Application/SlateApplication.h"
+#include "PropertyEditorModule.h"
+#include "ISinglePropertyView.h"
+#include "ObjectEditorUtils.h"
+#include "Misc/NotifyHook.h"
+
+#define LOCTEXT_NAMESPACE "SRigVMAssetView"
+
+namespace UE::AnimNext::Editor
+{
+
+namespace ParameterBlockView
+{
+
+static FName ContextMenuName(TEXT("AnimNext.RigVMAssetView.ContextMenu"));
+static FName Column_RevisionControl(TEXT("RevisionControl"));
+static FName Column_ModifiedStatus(TEXT("ModifiedStatus"));
+static FName Column_Name(TEXT("Name"));
+static FName Column_Type(TEXT("Type"));
+static FName Column_Value(TEXT("Value"));
+
+static const FName NAME_Category("Category");
+
+}
+
+TMap<FName, SRigVMAssetView::FCategoryWidgetFactoryFunction> SRigVMAssetView::CategoryFactories;
+
+// An entry category displayed in the parameters table
+template<typename ItemType>
+class SCategoryHeaderTableRow : public STableRow<ItemType>
+{
+public:
+	SLATE_BEGIN_ARGS(SCategoryHeaderTableRow)
+		{}
+		SLATE_DEFAULT_SLOT(typename SCategoryHeaderTableRow::FArguments, Content)
+	SLATE_END_ARGS()
+
+	void Construct(const FArguments& InArgs, const TSharedRef<STableViewBase>& InOwnerTableView)
+	{
+		STableRow<ItemType>::ChildSlot
+			.Padding(0.0f, 2.0f, .0f, 0.0f)
+			[
+				SAssignNew(ContentBorder, SBorder)
+					.BorderImage(this, &SCategoryHeaderTableRow::GetBackgroundImage)
+					.Padding(FMargin(3.0f, 5.0f))
+					[
+						SNew(SHorizontalBox)
+							+ SHorizontalBox::Slot()
+							.VAlign(VAlign_Center)
+							.Padding(5.0f)
+							.AutoWidth()
+							[
+								SNew(SExpanderArrow, STableRow< ItemType >::SharedThis(this))
+							]
+							+ SHorizontalBox::Slot()
+							.VAlign(VAlign_Center)
+							.AutoWidth()
+							[
+								InArgs._Content.Widget
+							]
+					]
+			];
+
+		STableRow < ItemType >::ConstructInternal(
+			typename STableRow< ItemType >::FArguments()
+			.Style(FAppStyle::Get(), "DetailsView.TreeView.TableRow")
+			.ShowSelection(false),
+			InOwnerTableView
+		);
+	}
+
+	const FSlateBrush* GetBackgroundImage() const
+	{
+		if (STableRow<ItemType>::IsHovered())
+		{
+			return FAppStyle::Get().GetBrush("Brushes.Secondary");
+		}
+		else
+		{
+			return FAppStyle::Get().GetBrush("Brushes.Header");
+		}
+	}
+
+	virtual void SetContent(TSharedRef< SWidget > InContent) override
+	{
+		ContentBorder->SetContent(InContent);
+	}
+
+	virtual void SetRowContent(TSharedRef< SWidget > InContent) override
+	{
+		ContentBorder->SetContent(InContent);
+	}
+
+	virtual const FSlateBrush* GetBorder() const
+	{
+		return nullptr;
+	}
+
+	FReply OnMouseButtonDown(const FGeometry& MyGeometry, const FPointerEvent& MouseEvent)
+	{
+		if (MouseEvent.GetEffectingButton() == EKeys::LeftMouseButton)
+		{
+			STableRow<ItemType>::ToggleExpansion();
+			return FReply::Handled();
+		}
+		else
+		{
+			return FReply::Unhandled();
+		}
+	}
+private:
+	TSharedPtr<SBorder> ContentBorder;
+};
+
+// An entry displayed in the parameters view
+struct FRigVMAssetViewEntry
+{
+
+	FRigVMAssetViewEntry() = default;
+
+	FRigVMAssetViewEntry(UAnimNextRigVMAssetEntry* InEntry)
+		: WeakEntry(InEntry)
+	{}
+
+	FRigVMAssetViewEntry(FName InCategoryType, const FText& InCategoryName)
+		: CategoryTypeName(InCategoryType)
+		, CategoryText(InCategoryName)
+	{}
+
+	bool PassesFilter(const FString& InFilterText) const
+	{
+		if(UAnimNextRigVMAssetEntry* Entry = WeakEntry.Get())
+		{
+			return Entry->GetDisplayName().ToString().Contains(InFilterText);
+		}
+
+		return false;
+	}
+
+	void ResetChildren(int32 NewSize = 0)
+	{
+		Children.Reset(NewSize);
+	}
+
+	void GetChildrenRecursive(TArray< TSharedRef<FRigVMAssetViewEntry> >& OutChildren)
+	{
+		for (TSharedRef<FRigVMAssetViewEntry>& Entry : Children)
+		{
+			OutChildren.Add(Entry);
+			Entry->GetChildrenRecursive(OutChildren);
+		}
+	}
+
+	// Ptr to the underlying entry
+	TWeakObjectPtr<UAnimNextRigVMAssetEntry> WeakEntry;
+
+	// Widget used to rename items
+	TWeakPtr<SInlineEditableTextBlock> NameWidget;
+
+	FName CategoryTypeName;
+	FText CategoryText;
+
+	// Children when entry is a category
+	TArray< TSharedRef<FRigVMAssetViewEntry> > Children;
+
+	// Flag to indicate a rename was requested
+	bool bRenameWhenScrolledIntoView = false;
+};
+
+void SRigVMAssetView::RegisterCategoryFactory(FName InCategory, FCategoryWidgetFactoryFunction&& InFunction)
+{
+	CategoryFactories.Add(InCategory, MoveTemp(InFunction));
+}
+
+void SRigVMAssetView::UnregisterCategoryFactory(FName InCategory)
+{
+	CategoryFactories.Remove(InCategory);
+}
+
+void SRigVMAssetView::Construct(const FArguments& InArgs, UAnimNextRigVMAssetEditorData* InEditorData)
+{
+	using namespace ParameterBlockView;
+	
+	check(InEditorData);
+
+	EditorData = InEditorData;
+	EditorData->SetFlags(RF_Transactional);
+
+	TMap<FName, FText> AllCategories;
+	for (TSubclassOf<UAnimNextRigVMAssetEntry> EntryClass : EditorData->GetEntryClasses())
+	{
+		if(const FString* CategoryMetaData = EntryClass->FindMetaData(NAME_Category))
+		{
+			AllCategories.Add(**CategoryMetaData, FObjectEditorUtils::GetCategoryText(EntryClass));
+		}
+	}
+
+	for(const TPair<FName, FText>& CategoryPair : AllCategories)
+	{
+		Categories.Add(MakeShared<FRigVMAssetViewEntry>(CategoryPair.Key, CategoryPair.Value));
+	}
+
+	// Cache asset data for comparisons/filtering
+	AssetData = FAssetData(UncookedOnly::FUtils::GetAsset(EditorData));
+
+	OnSelectionChangedDelegate = InArgs._OnSelectionChanged;
+	OnOpenGraphDelegate = InArgs._OnOpenGraph;
+	OnDeleteEntriesDelegate = InArgs._OnDeleteEntries;
+
+	EditorData->ModifiedDelegate.AddSP(this, &SRigVMAssetView::HandleBlockModified);
+
+	ChildSlot
+	[
+		SNew(SVerticalBox)
+		+SVerticalBox::Slot()
+		.AutoHeight()
+		[
+			SNew(SHorizontalBox)
+			+ SHorizontalBox::Slot()
+			.FillWidth(1.0f)
+			.Padding(2.0f)
+			[
+				SNew(SSearchBox)
+				.OnTextChanged_Lambda([this](FText InText)
+				{
+					FilterText = InText;
+					RefreshFilter();
+				})
+			]
+		]
+		+SVerticalBox::Slot()
+		.FillHeight(1.0f)
+		[
+			SAssignNew(EntriesList, STreeView<TSharedRef<FRigVMAssetViewEntry>>)
+			.TreeItemsSource(&FilteredEntries)
+			.OnGenerateRow(this, &SRigVMAssetView::HandleGenerateRow)
+			.OnGetChildren(this, &SRigVMAssetView::HandleGetChildren)
+			.OnItemScrolledIntoView(this, &SRigVMAssetView::HandleItemScrolledIntoView)
+			.OnSelectionChanged(this, &SRigVMAssetView::HandleSelectionChanged)
+			.ItemHeight(20.0f)
+			.HeaderRow(
+				SNew(SHeaderRow)
+				+SHeaderRow::Column(Column_RevisionControl)
+				.DefaultLabel(FText::GetEmpty())
+				.FixedWidth(24.0f)
+				.HeaderContent()
+				[
+					SNew(SBox)
+					.WidthOverride(16.0f)
+					.HeightOverride(16.0f)
+					.VAlign(VAlign_Center)
+					.HAlign(HAlign_Center)
+					[
+						SNew(SImage)
+						.ColorAndOpacity(FSlateColor::UseForeground())
+						.Image(FRevisionControlStyleManager::Get().GetBrush("RevisionControl.Icon"))
+						.ToolTipText(LOCTEXT("RevisionControlStatusTooltip", "Revision control status of this parameter"))
+					]
+				]
+
+				+SHeaderRow::Column(Column_ModifiedStatus)
+				.DefaultLabel(FText::GetEmpty())
+				.FixedWidth(24.0f)
+				.HeaderContent()
+				[
+					SNew(SBox)
+					.WidthOverride(16.0f)
+					.HeightOverride(16.0f)
+					.VAlign(VAlign_Center)
+					.HAlign(HAlign_Center)
+					[
+						SNew(SImage)
+						.ColorAndOpacity(FSlateColor::UseForeground())
+						.Image(FAppStyle::GetBrush("ContentBrowser.ContentDirty"))
+						.ToolTipText(LOCTEXT("ModifiedStatusTooltip", "Modified status of this parameter"))
+					]
+				]
+
+				+SHeaderRow::Column(Column_Name)
+				.DefaultLabel(LOCTEXT("NameColumnHeader", "Name"))
+				.ToolTipText(LOCTEXT("NameColumnHeaderTooltip", "The name of the parameter"))
+				.FillWidth(20.0f)
+
+				+SHeaderRow::Column(Column_Type)
+				.DefaultLabel(LOCTEXT("TypeColumnHeader", "Type"))
+				.ToolTipText(LOCTEXT("TypeColumnHeaderTooltip", "The type of the parameter"))
+				.ManualWidth(145.0f)
+
+				+SHeaderRow::Column(Column_Value)
+				.DefaultLabel(LOCTEXT("ValueColumnHeader", "Value"))
+				.ToolTipText(LOCTEXT("ValueColumnHeaderTooltip", "The value or binding to the parameter"))
+				.FillWidth(10.0f)
+			)
+		]
+	];
+
+	// Update revision control status of all our files
+	if (ISourceControlModule::Get().IsEnabled())
+	{
+		TArray<UPackage*> Packages = EditorData->GetPackage()->GetExternalPackages();
+
+		ISourceControlProvider& SourceControlProvider = ISourceControlModule::Get().GetProvider();
+		SourceControlProvider.Execute(ISourceControlOperation::Create<FUpdateStatus>(), Packages);
+	}
+
+	BindCommands();
+
+	RegisterActiveTimer(1.0f / 60.0f, FWidgetActiveTimerDelegate::CreateLambda(
+		[this](double InCurrentTime, float InDeltaTime)
+		{
+			if(bRefreshRequested)
+			{
+				RefreshEntries();
+				bRefreshRequested = false;
+			}
+
+			return EActiveTimerReturnType::Continue;
+		}));
+
+	RefreshEntries();
+
+	for (TSharedRef<FRigVMAssetViewEntry>& Category : Categories)
+	{
+		EntriesList->SetItemExpansion(Category, true);
+	}
+}
+
+FReply SRigVMAssetView::OnKeyDown(const FGeometry& MyGeometry, const FKeyEvent& InKeyEvent)
+{
+	if (UICommandList.IsValid() && UICommandList->ProcessCommandBindings(InKeyEvent))
+	{
+		return FReply::Handled();
+	}
+	return FReply::Unhandled();
+}
+
+void SRigVMAssetView::RequestRefresh()
+{
+	bRefreshRequested = true;
+}
+
+template<typename ItemType, typename ComparisonType>
+void RestoreExpansionState(TSharedPtr< STreeView<ItemType> > InTree, const TArray<ItemType>& ItemSource, const TSet<ItemType>& OldExpansionState, ComparisonType ComparisonFunction)
+{
+	check(InTree.IsValid());
+
+	// Iterate over new tree items
+	for (int32 ItemIdx = 0; ItemIdx < ItemSource.Num(); ItemIdx++)
+	{
+		ItemType NewItem = ItemSource[ItemIdx];
+
+		// Look through old expansion state
+		for (typename TSet<ItemType>::TConstIterator OldExpansionIter(OldExpansionState); OldExpansionIter; ++OldExpansionIter)
+		{
+			const ItemType OldItem = *OldExpansionIter;
+			// See if this matches this new item
+			if (ComparisonFunction(OldItem, NewItem))
+			{
+				// It does, so expand it
+				InTree->SetItemExpansion(NewItem, true);
+			}
+		}
+	}
+}
+
+static bool CompareGraphActionNode(TSharedRef<FRigVMAssetViewEntry> A, TSharedRef<FRigVMAssetViewEntry> B)
+{
+	if (A->CategoryTypeName != B->CategoryTypeName)
+	{
+		return false;
+	}
+
+	return true;
+}
+
+void SRigVMAssetView::RefreshEntries()
+{
+	// First, save off current expansion state
+	TSet< TSharedRef<FRigVMAssetViewEntry> > OldExpansionState;
+	EntriesList->GetExpandedItems(OldExpansionState);
+
+	Entries.Reset();
+
+	TArray<TSharedRef<FRigVMAssetViewEntry>> EntriesToSelect;
+
+	for (UAnimNextRigVMAssetEntry* SubEntry : EditorData->GetAllEntries())
+	{
+		TSharedRef<FRigVMAssetViewEntry> NewEntry = MakeShared<FRigVMAssetViewEntry>(SubEntry);
+
+		Entries.Add(NewEntry);
+
+		if (PendingSelection.Contains(SubEntry))
+		{
+			EntriesToSelect.Add(NewEntry);
+		}
+	}
+
+	// Restore the expanded items
+	TArray< TSharedRef<FRigVMAssetViewEntry> > AllEntries;
+	for (const TSharedRef<FRigVMAssetViewEntry> & Entry : Entries)
+	{
+		AllEntries.Add(Entry);
+		Entry->GetChildrenRecursive(AllEntries);
+	}
+	RestoreExpansionState< TSharedRef<FRigVMAssetViewEntry> >(EntriesList, AllEntries, OldExpansionState, CompareGraphActionNode);
+
+	PendingSelection.Empty();
+
+	RefreshFilter();
+
+	if(EntriesToSelect.Num() > 0)
+	{
+		EntriesList->SetItemSelection(EntriesToSelect, true);
+	}
+}
+
+void SRigVMAssetView::RefreshFilter()
+{
+	using namespace ParameterBlockView;
+
+	FilteredEntries = Categories;
+
+	for (TSharedRef<FRigVMAssetViewEntry>& Category : Categories)
+	{
+		Category->ResetChildren();
+	}
+
+	// add the entries as categories children
+	const FString FilterTextAsString = FilterText.ToString();
+	for (const TSharedRef<FRigVMAssetViewEntry>& Entry : Entries)
+	{
+		if (Entry->PassesFilter(FilterTextAsString))
+		{
+			const UObject* BlockEntry = Entry->WeakEntry.Get();
+			const UClass* EntryClass = BlockEntry->GetClass();
+			const FString& CategoryMetaData = EntryClass->GetMetaData(NAME_Category);
+			TSharedRef<FRigVMAssetViewEntry> CategoryEntry = GetCategoryEntry(*CategoryMetaData);
+			CategoryEntry->Children.Add(Entry);
+		}
+	}
+
+	EntriesList->RequestListRefresh();
+}
+
+void SRigVMAssetView::BindCommands()
+{
+	UICommandList = MakeShared<FUICommandList>();
+
+	UICommandList->MapAction(FGenericCommands::Get().Delete,
+		FExecuteAction::CreateSP(this, &SRigVMAssetView::HandleDelete),
+		FCanExecuteAction::CreateSP(this, &SRigVMAssetView::HasValidSelection));
+
+	UICommandList->MapAction(FGenericCommands::Get().Rename,
+		FExecuteAction::CreateSP(this, &SRigVMAssetView::HandleRename),
+		FCanExecuteAction::CreateSP(this, &SRigVMAssetView::HasValidSingleSelection));
+}
+
+void SRigVMAssetView::HandleBlockModified(UAnimNextRigVMAssetEditorData* InEditorData)
+{
+	check(InEditorData == EditorData);
+
+	RequestRefresh();
+}
+
+TSharedRef<SWidget> SRigVMAssetView::HandleGetContextContent()
+{
+	using namespace ParameterBlockView;
+	
+	UToolMenus* ToolMenus = UToolMenus::Get();
+
+	if(!ToolMenus->IsMenuRegistered(ContextMenuName))
+	{
+		UToolMenu* Menu = ToolMenus->RegisterMenu(ContextMenuName);
+
+		FToolMenuSection& Section = Menu->AddSection("EntryOperations", LOCTEXT("EntryOperationsMenuSection", "Entry"));
+		Section.AddMenuEntry(FGenericCommands::Get().Delete);
+		Section.AddMenuEntry(FGenericCommands::Get().Rename);
+	}
+
+	URigVMAssetViewMenuContext* MenuContext = NewObject<URigVMAssetViewMenuContext>();
+	MenuContext->RigVMAssetView = SharedThis(this);
+	return ToolMenus->GenerateWidget(ContextMenuName, FToolMenuContext(MenuContext));
+}
+
+void SRigVMAssetView::HandleDelete()
+{
+	auto DeleteSelected = [this](TArray<TSharedRef<FRigVMAssetViewEntry>> SelectedItems)
+	{
+		TArray<UAnimNextRigVMAssetEntry*> EntriesToRemove;
+		Algo::Transform(SelectedItems, EntriesToRemove, [](const TSharedRef<FRigVMAssetViewEntry>& InEntry) { return InEntry->WeakEntry.Get(); });
+
+		bool bAnyEntriesRemoved = false;
+		{
+			FScopedTransaction Transaction(FText::FormatOrdered(LOCTEXT("DeleteEntry", "Delete asset {0}|plural(one=entry,other=entries)"), EntriesList->GetNumItemsSelected()));
+			bAnyEntriesRemoved = EditorData->RemoveEntries(EntriesToRemove);
+		}
+		
+		if(bAnyEntriesRemoved)
+		{
+			// Remove all the entries that werent actually removed
+			EntriesToRemove.RemoveAll([this](UAnimNextRigVMAssetEntry* InEntry) { return EditorData->FindEntry(InEntry->GetEntryName()) != nullptr; });
+			OnDeleteEntriesDelegate.ExecuteIfBound(EntriesToRemove);
+		}
+	};
+
+	if (EntriesList->GetNumItemsSelected() > 0)
+	{
+		TArray<TSharedRef<FRigVMAssetViewEntry>> SelectedItems = EntriesList->GetSelectedItems();
+		DeleteSelected(SelectedItems);
+	}
+}
+
+void SRigVMAssetView::HandleRename()
+{
+	auto RenameSelected = [this](TSharedPtr<STreeView<TSharedRef<FRigVMAssetViewEntry>>>& InEntriesList, const TSharedRef<FRigVMAssetViewEntry>& SelectedItem)
+	{
+		SelectedItem->bRenameWhenScrolledIntoView = true;
+		InEntriesList->RequestScrollIntoView(SelectedItem);
+	};
+
+	if(EntriesList->GetNumItemsSelected() == 1)
+	{
+		TArray<TSharedRef<FRigVMAssetViewEntry>> SelectedItems = EntriesList->GetSelectedItems();
+		RenameSelected(EntriesList, SelectedItems[0]);
+	}
+}
+
+bool SRigVMAssetView::HasValidSelection() const
+{
+	return EntriesList->GetNumItemsSelected() > 0;
+}
+
+bool SRigVMAssetView::HasValidSingleSelection() const
+{
+	return EntriesList->GetNumItemsSelected() == 1;
+}
+
+void SRigVMAssetView::HandleItemScrolledIntoView(TSharedRef<FRigVMAssetViewEntry> InEntry, const TSharedPtr<ITableRow>& InWidget)
+{
+	if(InEntry->bRenameWhenScrolledIntoView)
+	{
+		InEntry->bRenameWhenScrolledIntoView = false;
+
+		if(TSharedPtr<SInlineEditableTextBlock> NameWidget = InEntry->NameWidget.Pin())
+		{
+			NameWidget->EnterEditingMode();
+		}
+	}
+}
+
+class SRigVMAssetViewRow : public SMultiColumnTableRow<TSharedRef<FRigVMAssetViewEntry>>, private FNotifyHook
+{
+	SLATE_BEGIN_ARGS(SRigVMAssetViewRow) {}
+	SLATE_END_ARGS()
+
+	void Construct(const FArguments& InArgs, const TSharedRef<STableViewBase>& InOwnerTableView, TSharedRef<SRigVMAssetView> InView, TSharedRef<FRigVMAssetViewEntry> InEntry)
+	{
+		WeakView = InView;
+		Entry = InEntry;
+
+		SMultiColumnTableRow<TSharedRef<FRigVMAssetViewEntry>>::Construct( SMultiColumnTableRow<TSharedRef<FRigVMAssetViewEntry>>::FArguments(), InOwnerTableView);
+	}
+
+	virtual FReply OnMouseButtonDoubleClick(const FGeometry& InMyGeometry, const FPointerEvent& InMouseEvent) override
+	{
+		if(IAnimNextRigVMGraphInterface* GraphInterface = Cast<IAnimNextRigVMGraphInterface>(Entry->WeakEntry.Get()))
+		{
+			if(TSharedPtr<SRigVMAssetView> View = WeakView.Pin())
+			{
+				View->OnOpenGraphDelegate.ExecuteIfBound(GraphInterface->GetRigVMGraph());
+			}
+			return FReply::Handled();
+		}
+
+		return SMultiColumnTableRow<TSharedRef<FRigVMAssetViewEntry>>::OnMouseButtonDoubleClick(InMyGeometry, InMouseEvent);
+	}
+	
+	//~ Begin FNotifyHook Interface
+	virtual void NotifyPreChange(FProperty* PropertyAboutToChange) override
+	{
+		if (TSharedPtr<SRigVMAssetView> View = WeakView.Pin())
+		{
+			if (UAnimNextRigVMAssetEditorData* EditorData = View->EditorData)
+			{
+				if (UAnimNextRigVMAssetEntry* AssetEntry = EditorData->FindEntry(PropertyAboutToChange->GetFName()))
+				{
+					UAnimNextRigVMAsset* Asset = UE::AnimNext::UncookedOnly::FUtils::GetAsset(EditorData);
+
+					// needed to enable the transaction when we modify the PropertyBag
+					// TODO: remove this by moving defaults into entries
+					Asset->Modify();
+				}
+			}
+		}
+	}
+
+	virtual void NotifyPostChange(const FPropertyChangedEvent& PropertyChangedEvent, FProperty* PropertyThatChanged) override
+	{
+		if (TSharedPtr<SRigVMAssetView> View = WeakView.Pin())
+		{
+			if (UAnimNextRigVMAssetEditorData* EditorData = View->EditorData)
+			{
+				if (UAnimNextRigVMAssetEntry* AssetEntry = EditorData->FindEntry(PropertyChangedEvent.GetMemberPropertyName()))
+				{
+					// Needed to show the changed sign a the table when we modify the PropertyBag
+					// TODO: remove this by moving defaults into entries
+					AssetEntry->MarkPackageDirty(); 
+				}
+			}
+		}
+	}
+	// End of FNotifyHook
+
+	virtual TSharedRef<SWidget> GenerateWidgetForColumn(const FName& InColumnName) override
+	{
+		using namespace ParameterBlockView;
+		
+		if(InColumnName == Column_RevisionControl)
+		{
+			return
+				SNew(SBox)
+				.HAlign(HAlign_Center)
+				.VAlign(VAlign_Center)
+				.HeightOverride(20.0f)
+				[
+					SNew(SImage)
+					.ToolTipText_Lambda([this]()
+					{
+						if(UAnimNextRigVMAssetEntry* AssetEntry = Entry->WeakEntry.Get())
+						{
+							ISourceControlProvider& SourceControlProvider = ISourceControlModule::Get().GetProvider();
+							UPackage* Package = AssetEntry->GetExternalPackage();
+							if(FSourceControlStatePtr State = SourceControlProvider.GetState(Package, EStateCacheUsage::Use))
+							{
+								return FText::Format(LOCTEXT("RevisionControlStatusFormat", "File: {0}\nStatus: {1}"), FText::FromName(Package->GetFName()),  State->GetDisplayTooltip());
+							}
+						}
+
+						return LOCTEXT("RevisionControlStatus", "Revision control status of this parameter");
+					})
+					.Image_Lambda([this]() -> const FSlateBrush*
+					{
+						if(UAnimNextRigVMAssetEntry* AssetEntry = Entry->WeakEntry.Get())
+						{
+							ISourceControlProvider& SourceControlProvider = ISourceControlModule::Get().GetProvider();
+							if(FSourceControlStatePtr State = SourceControlProvider.GetState(AssetEntry->GetExternalPackage(), EStateCacheUsage::Use))
+							{
+								return State->GetIcon().GetSmallIcon();
+							}
+						}
+						return nullptr;
+					})
+				];
+		}
+		else if(InColumnName == Column_ModifiedStatus)
+		{
+			return
+				SNew(SBox)
+				.HAlign(HAlign_Center)
+				.VAlign(VAlign_Center)
+				.HeightOverride(20.0f)
+				[
+					SNew(SImage)
+					.ToolTipText_Lambda([this]()
+					{
+						FTextBuilder TextBuilder;
+						TextBuilder.AppendLine(LOCTEXT("ModifiedTooltip", "Modified Status"));
+
+						if(UAnimNextRigVMAssetEntry* AssetEntry = Entry->WeakEntry.Get())
+						{
+							const UPackage* ExternalPackage = AssetEntry->GetExternalPackage();
+							check(ExternalPackage);
+							if(ExternalPackage->IsDirty())
+							{
+								TextBuilder.AppendLine(FText::FromName(ExternalPackage->GetFName()));
+							}
+						}
+
+						return TextBuilder.ToText();
+					})
+					.Image_Lambda([this]() -> const FSlateBrush*
+					{
+						if(UAnimNextRigVMAssetEntry* AssetEntry = Entry->WeakEntry.Get())
+						{
+							bool bIsDirty = false;
+							const UPackage* ExternalPackage = AssetEntry->GetExternalPackage();
+							check(ExternalPackage);
+							if(ExternalPackage->IsDirty())
+							{
+								bIsDirty = true;
+							}
+
+
+							return bIsDirty ? FAppStyle::GetBrush("ContentBrowser.ContentDirty") : nullptr;
+						}
+						return nullptr;
+					})
+				];
+		}
+		else if(InColumnName == Column_Type)
+		{
+			if(Entry->WeakEntry.Get()->Implements<UAnimNextRigVMParameterInterface>())
+			{
+				return
+					SNew(SBox)
+					.HAlign(HAlign_Left)
+					.VAlign(VAlign_Center)
+					[
+						SNew(SPinTypeSelector, FGetPinTypeTree::CreateStatic(&Editor::FUtils::GetFilteredVariableTypeTree))
+							.TargetPinType_Lambda([this]()
+							{
+								if(const IAnimNextRigVMParameterInterface* Binding = Cast<IAnimNextRigVMParameterInterface>(Entry->WeakEntry.Get()))
+								{
+									return UncookedOnly::FUtils::GetPinTypeFromParamType(Binding->GetParamType());
+								}
+
+								return FEdGraphPinType();
+							})
+							.OnPinTypeChanged_Lambda([this](const FEdGraphPinType& PinType)
+							{
+								if(IAnimNextRigVMParameterInterface* Binding = Cast<IAnimNextRigVMParameterInterface>(Entry->WeakEntry.Get()))
+								{
+									const FAnimNextParamType ParamType = UncookedOnly::FUtils::GetParamTypeFromPinType(PinType);
+									if(ParamType.IsValid())
+									{
+										Binding->SetParamType(ParamType);
+									}
+								}
+							})
+							.Schema(GetDefault<UPropertyBagSchema>())
+							.bAllowArrays(true)
+							.TypeTreeFilter(ETypeTreeFilter::None)
+							.Font(IDetailLayoutBuilder::GetDetailFont())
+					];
+			}
+		}
+		else if(InColumnName == Column_Name)
+		{
+			return
+				SNew(SHorizontalBox)
+				+SHorizontalBox::Slot()
+				.HAlign(HAlign_Left)
+				.VAlign(VAlign_Center)
+				.AutoWidth()
+				[
+					SNew(SExpanderArrow, SharedThis(this))
+					.Visibility_Lambda([this]() { return DoesItemHaveChildren() ? EVisibility::Visible : EVisibility::Collapsed; })
+				]
+				+SHorizontalBox::Slot()
+				.HAlign(HAlign_Left)
+				.VAlign(VAlign_Center)
+				.AutoWidth()
+				[
+					SAssignNew(Entry->NameWidget, SInlineEditableTextBlock)
+					.Font(IDetailLayoutBuilder::GetDetailFont())
+					.IsSelected(this, &SRigVMAssetViewRow::IsSelectedExclusively)
+					.IsReadOnly_Lambda([this]()
+					{
+						if (Cast<IAnimNextRigVMParameterInterface>(Entry->WeakEntry.Get()) 
+							|| Cast<IAnimNextRigVMGraphInterface>(Entry->WeakEntry.Get()))
+						{
+							return false;
+						}
+						return true;
+					})
+					.OnTextCommitted_Lambda([this](const FText& InNewText, ETextCommit::Type InCommitType)
+					{
+						if(InCommitType == ETextCommit::OnEnter)
+						{
+							if(UAnimNextRigVMAssetEntry* AssetEntry = Cast<UAnimNextRigVMAssetEntry>(Entry->WeakEntry.Get()))
+							{
+								FScopedTransaction Transaction(LOCTEXT("SetParameterName", "Set Entry name"));
+								AssetEntry->SetEntryName(*InNewText.ToString());
+							}
+						}
+					})
+					.OnVerifyTextChanged_Lambda([this](const FText& InNewText, FText& OutErrorText)
+					{
+						const FString NewString = InNewText.ToString();
+						if (!FUtils::IsValidEntryNameString(NewString, OutErrorText))
+						{
+							return false;
+						}
+
+						FName Name(*NewString);
+						if (TSharedPtr<SRigVMAssetView> View = WeakView.Pin())
+						{
+							if(UAnimNextRigVMAssetEntry* ExistingEntry = View->EditorData->FindEntry(Name))
+							{
+								OutErrorText = LOCTEXT("Error_NameExists", "This name already exists in this asset");
+								return false;
+							}
+						}
+						return true;
+					})
+					.Text_Lambda([this]()
+					{
+						if(UAnimNextRigVMAssetEntry* AssetEntry = Cast<UAnimNextRigVMAssetEntry>(Entry->WeakEntry.Get()))
+						{
+							if(IAnimNextRigVMParameterInterface* ParameterInterface = Cast<IAnimNextRigVMParameterInterface>(AssetEntry))
+							{
+								return UncookedOnly::FUtils::GetParameterDisplayNameText(AssetEntry->GetEntryName());
+							}
+							else
+							{
+								return FText::FromName(AssetEntry->GetEntryName());
+							}
+						}
+						return FText::GetEmpty();
+					})
+					.ToolTipText_Lambda([this]()
+					{
+						if(UAnimNextRigVMAssetEntry* AssetEntry = Entry->WeakEntry.Get())
+						{
+							return AssetEntry->GetDisplayNameTooltip();
+						}
+						return FText::GetEmpty();
+					})
+				];
+		}
+		else if(InColumnName == Column_Value)
+		{
+			TSharedRef< SWidget > ColumnWidget = SNullWidget::NullWidget;
+
+			if (TSharedPtr<SRigVMAssetView> View = WeakView.Pin())
+			{
+				if(UAnimNextRigVMAssetEntry* AssetEntry = Cast<UAnimNextRigVMAssetEntry>(Entry->WeakEntry.Get()))
+				{
+					if (const IAnimNextRigVMParameterInterface* ParameterInterface = Cast<IAnimNextRigVMParameterInterface>(AssetEntry))
+					{
+						FInstancedPropertyBag& PropertyBag = ParameterInterface->GetPropertyBag();
+						const FName ParameterName = AssetEntry->GetEntryName();
+
+						if (PropertyBag.FindPropertyDescByName(ParameterName))
+						{
+							FSinglePropertyParams SinglePropertyArgs;
+							SinglePropertyArgs.NamePlacement = EPropertyNamePlacement::Hidden;
+							SinglePropertyArgs.NotifyHook = this;
+
+							FPropertyEditorModule& PropertyEditorModule = FModuleManager::Get().LoadModuleChecked<FPropertyEditorModule>("PropertyEditor");
+							const TSharedPtr<ISinglePropertyView> SingleStructPropertyView = PropertyEditorModule.CreateSingleProperty(MakeShared<FInstancePropertyBagStructureDataProvider>(PropertyBag), ParameterName, SinglePropertyArgs);
+							if (SingleStructPropertyView.IsValid())
+							{
+								ColumnWidget = SingleStructPropertyView.ToSharedRef();
+							}
+						}
+					}
+				}
+			}
+
+			return SNew(SHorizontalBox)
+				+ SHorizontalBox::Slot()
+				.HAlign(HAlign_Center)
+				.AutoWidth()
+				[
+					ColumnWidget
+				];
+		}
+
+		return SNullWidget::NullWidget;
+	}
+
+	TWeakPtr<SRigVMAssetView> WeakView;
+	TSharedPtr<FRigVMAssetViewEntry> Entry;
+};
+
+TSharedRef<ITableRow> SRigVMAssetView::HandleGenerateRow(TSharedRef<FRigVMAssetViewEntry> InEntry, const TSharedRef<STableViewBase>& InOwnerTable)
+{
+	if (InEntry->CategoryTypeName != NAME_None)
+	{
+		TSharedPtr< STableRow< TSharedPtr<FRigVMAssetViewEntry> > > TableRow;
+		TableRow = SNew(SCategoryHeaderTableRow< TSharedPtr<FRigVMAssetViewEntry> >, InOwnerTable);
+
+		TSharedPtr<SHorizontalBox> RowContainer;
+		TableRow->SetRowContent
+		(
+			SAssignNew(RowContainer, SHorizontalBox)
+		);
+
+		const FMargin RowPadding = FMargin(0, 2);
+		TSharedPtr<SHorizontalBox> RowContent = SNew(SHorizontalBox)
+			+ SHorizontalBox::Slot()
+			.VAlign(VAlign_Center)
+			[
+				SNew(STextBlock)
+				.Font(FCoreStyle::GetDefaultFontStyle("Bold", 9))
+				.Text(InEntry->CategoryText)
+			];
+
+		RowContainer->AddSlot()
+			.AutoWidth()
+			.VAlign(VAlign_Fill)
+			.HAlign(HAlign_Right)
+			[
+				SNew(SExpanderArrow, TableRow)
+					.BaseIndentLevel(1)
+			];
+
+		RowContainer->AddSlot()
+			.FillWidth(1.0)
+			.Padding(RowPadding)
+			[
+				RowContent.ToSharedRef()
+			];
+
+		if(FCategoryWidgetFactoryFunction* FactoryFunction = CategoryFactories.Find(InEntry->CategoryTypeName))
+		{
+			RowContainer->AddSlot()
+			.AutoWidth()
+			.VAlign(VAlign_Fill)
+			.HAlign(HAlign_Right)
+			[
+				(*FactoryFunction)(EditorData)
+			];
+		}
+
+		return TableRow.ToSharedRef();
+	}
+	else
+	{
+		return SNew(SRigVMAssetViewRow, InOwnerTable, SharedThis(this), InEntry);
+	}
+}
+
+void SRigVMAssetView::HandleGetChildren(TSharedRef<FRigVMAssetViewEntry> InEntry, TArray<TSharedRef<FRigVMAssetViewEntry>>& OutChildren)
+{
+	if (InEntry->CategoryTypeName != NAME_None)
+	{
+		OutChildren = InEntry->Children;
+	}
+}
+
+void SRigVMAssetView::HandleSelectionChanged(TSharedPtr<FRigVMAssetViewEntry> InEntry, ESelectInfo::Type InSelectionType)
+{
+	if(OnSelectionChangedDelegate.IsBound())
+	{
+		TArray<UObject*> SelectedItems;
+		SelectedItems.Reserve(EntriesList->GetNumItemsSelected());
+		for(const TSharedRef<FRigVMAssetViewEntry>& SelectedItem : EntriesList->GetSelectedItems())
+		{
+			if(UAnimNextRigVMAssetEntry* Entry = SelectedItem->WeakEntry.Get())
+			{
+				SelectedItems.Add(Entry);
+			}
+		}
+		OnSelectionChangedDelegate.Execute(SelectedItems);
+	}
+}
+
+TSharedRef<FRigVMAssetViewEntry> SRigVMAssetView::GetCategoryEntry(FName InCategoryName)
+{
+	for (TSharedRef<FRigVMAssetViewEntry>& Category : Categories)
+	{
+		if (Category->CategoryTypeName == InCategoryName)
+		{
+			return Category;
+		}
+	}
+
+	checkNoEntry();	// If we hit here, an entry's class has not been set up with a category
+
+	return TSharedRef<FRigVMAssetViewEntry>();
+}
+
+}
+
+#undef LOCTEXT_NAMESPACE
\ No newline at end of file
diff --git a/Engine/Plugins/Experimental/AnimNext/Source/AnimNextEditor/Private/Common/SRigVMAssetView.h b/Engine/Plugins/Experimental/AnimNext/Source/AnimNextEditor/Private/Common/SRigVMAssetView.h
new file mode 100644
--- /dev/null
+++ b/Engine/Plugins/Experimental/AnimNext/Source/AnimNextEditor/Private/Common/SRigVMAssetView.h
@@ -0,0 +1,118 @@
+// Copyright Epic Games, Inc. All Rights Reserved.
+
+#pragma once
+
+#include "Widgets/SCompoundWidget.h"
+#include "Widgets/Views/STreeView.h"
+#include "AssetRegistry/AssetData.h"
+#include "Templates/SharedPointer.h"
+
+class URigVMGraph;
+class FUICommandList;
+class UAnimNextRigVMAssetEntry;
+class UAnimNextRigVMAssetEditorData;
+
+namespace UE::AnimNext::Editor
+{
+
+struct FRigVMAssetViewEntry; 
+
+class SRigVMAssetView : public SCompoundWidget
+{
+public:
+	DECLARE_DELEGATE_OneParam(FOnSelectionChanged, const TArray<UObject*>& /*InEntries*/);
+
+	DECLARE_DELEGATE_OneParam(FOnOpenGraph, URigVMGraph* /*InGraph*/);
+
+	DECLARE_DELEGATE_OneParam(FOnDeleteEntries, const TArray<UAnimNextRigVMAssetEntry*>& /*InEntries*/);
+	
+	SLATE_BEGIN_ARGS(SRigVMAssetView) {}
+
+	SLATE_EVENT(SRigVMAssetView::FOnSelectionChanged, OnSelectionChanged)
+	
+	SLATE_EVENT(SRigVMAssetView::FOnOpenGraph, OnOpenGraph)
+
+	SLATE_EVENT(SRigVMAssetView::FOnDeleteEntries, OnDeleteEntries)
+
+	SLATE_END_ARGS()
+
+	void Construct(const FArguments& InArgs, UAnimNextRigVMAssetEditorData* InEditorData);
+
+	using FCategoryWidgetFactoryFunction = TUniqueFunction<TSharedRef<SWidget>(UAnimNextRigVMAssetEditorData*)>;
+
+	// Register a factory function used to generate widgets for a category
+	static void RegisterCategoryFactory(FName InCategory, FCategoryWidgetFactoryFunction&& InFunction);
+	static void UnregisterCategoryFactory(FName InCategory);
+
+private:
+	virtual FReply OnKeyDown(const FGeometry& MyGeometry, const FKeyEvent& InKeyEvent) override;
+
+	void RequestRefresh();
+
+	void RefreshEntries();
+
+	void RefreshFilter();
+
+	// Bind input commands
+	void BindCommands();
+
+	// Handle modifications to the block
+	void HandleBlockModified(UAnimNextRigVMAssetEditorData* InEditorData);
+
+	// Get the content for the context menu
+	TSharedRef<SWidget> HandleGetContextContent();
+
+	void HandleDelete();
+
+	void HandleRename();
+
+	bool HasValidSelection() const;
+
+	bool HasValidSingleSelection() const;
+
+	// Generate a row for the list view
+	TSharedRef<ITableRow> HandleGenerateRow(TSharedRef<FRigVMAssetViewEntry> InEntry, const TSharedRef<STableViewBase>& InOwnerTable);
+
+	void HandleGetChildren(TSharedRef<FRigVMAssetViewEntry> InEntry, TArray<TSharedRef<FRigVMAssetViewEntry>>& OutChildren);
+
+	// Handle rename after scrolling into view
+	void HandleItemScrolledIntoView(TSharedRef<FRigVMAssetViewEntry> Entry, const TSharedPtr<ITableRow>& Widget);
+
+	// Handle selection
+	void HandleSelectionChanged(TSharedPtr<FRigVMAssetViewEntry> InEntry, ESelectInfo::Type InSelectionType);
+
+	TSharedRef<FRigVMAssetViewEntry> GetCategoryEntry(FName InCategoryName);
+	
+private:
+	friend class SRigVMAssetViewRow;
+	
+	TArray<TSharedRef<FRigVMAssetViewEntry>> Categories;
+
+	TSharedPtr<STreeView<TSharedRef<FRigVMAssetViewEntry>>> EntriesList;
+
+	TArray<TSharedRef<FRigVMAssetViewEntry>> Entries;
+
+	FText FilterText;
+
+	TArray<TSharedRef<FRigVMAssetViewEntry>> FilteredEntries;
+
+	UAnimNextRigVMAssetEditorData* EditorData = nullptr;
+
+	TSharedPtr<FUICommandList> UICommandList = nullptr;
+
+	FOnSelectionChanged OnSelectionChangedDelegate;
+
+	FOnOpenGraph OnOpenGraphDelegate;
+
+	FOnDeleteEntries OnDeleteEntriesDelegate;
+
+	FAssetData AssetData;
+
+	TArray<UObject*> PendingSelection;
+
+	bool bRefreshRequested = false;
+
+	static TMap<FName, FCategoryWidgetFactoryFunction> CategoryFactories;
+};
+
+}
\ No newline at end of file
diff --git a/Engine/Plugins/Experimental/AnimNext/Source/AnimNextEditor/Private/EditorUtils.cpp b/Engine/Plugins/Experimental/AnimNext/Source/AnimNextEditor/Private/EditorUtils.cpp
--- a/Engine/Plugins/Experimental/AnimNext/Source/AnimNextEditor/Private/EditorUtils.cpp
+++ b/Engine/Plugins/Experimental/AnimNext/Source/AnimNextEditor/Private/EditorUtils.cpp
@@ -8,26 +8,22 @@
 #include "Param/AnimNextParameterBlock.h"
 #include "PropertyBagDetails.h"
 #include "Graph/AnimNextGraphEntry.h"
-#include "AssetRegistry/AssetRegistryModule.h"
-#include "AssetRegistry/IAssetRegistry.h"
-#include "Param/IAnimNextParameterBlockGraphInterface.h"
-#include "Param/IAnimNextParameterBlockParameterInterface.h"
+#include "IAnimNextRigVMGraphInterface.h"
+#include "IAnimNextRigVMParameterInterface.h"
 #include "Workspace/AnimNextWorkspace.h"
 #include "UncookedOnlyUtils.h"
-#include "Param/RigVMDispatch_GetLayerParameter.h"
 #include "Param/RigVMDispatch_GetParameter.h"
-#include "Param/RigVMDispatch_SetLayerParameter.h"
 
 #define LOCTEXT_NAMESPACE "AnimNextEditorUtils"
 
 namespace UE::AnimNext::Editor
 {
 
-void FUtils::GetAllGraphNames(const UAnimNextGraph_EditorData* InEditorData, TSet<FName>& OutNames)
+void FUtils::GetAllEntryNames(const UAnimNextRigVMAssetEditorData* InEditorData, TSet<FName>& OutNames)
 {
-	for(UAnimNextGraphEntry* Entry : InEditorData->Entries)
+	for(const  UAnimNextRigVMAssetEntry* Entry : InEditorData->GetAllEntries())
 	{
-		OutNames.Add(Entry->GraphName);
+		OutNames.Add(Entry->GetEntryName());
 	}
 }
 
@@ -241,22 +237,6 @@ FAnimNextParamType FUtils::GetParameterTypeFromMetaData(const FStringView& InStr
 	return FAnimNextParamType(); 
 }
 
-
-void FUtils::GetAllEntryNames(const UAnimNextParameterBlock_EditorData* InEditorData, TSet<FName>& OutNames)
-{
-	for(UAnimNextParameterBlockEntry* Entry : InEditorData->Entries)
-	{
-		if(IAnimNextParameterBlockParameterInterface* Parameter = Cast<IAnimNextParameterBlockParameterInterface>(Entry))
-		{
-			OutNames.Add(Parameter->GetParameterName());
-		}
-		else if(IAnimNextParameterBlockGraphInterface* Graph = Cast<IAnimNextParameterBlockGraphInterface>(Entry))
-		{
-			OutNames.Add(Graph->GetGraphName());
-		}
-	}
-}
-
 FName FUtils::ValidateName(const UAnimNextParameterBlock_EditorData* InEditorData, const FString& InName)
 {
 	struct FNameValidator : public INameValidatorInterface
@@ -485,13 +465,24 @@ FName FUtils::GetNewParameterName(const TCHAR* InBaseName, TArrayView<FName> InA
 	return NAME_None;
 }
 
-bool FUtils::IsValidParameterName(const FName InName, FText& OutErrorText)
+bool FUtils::IsValidEntryNameString(FStringView InStringView, FText& OutErrorText)
+{
+	// See if this can be represented as an FName
+	if(!FName::IsValidXName(InStringView, INVALID_NAME_CHARACTERS, &OutErrorText))
+	{
+		return false;
+	}
+
+	return IsValidEntryName(FName(InStringView), OutErrorText);
+}
+
+bool FUtils::IsValidEntryName(const FName InName, FText& OutErrorText)
 {
 	const FString NewString = InName.ToString();
 
 	if(NewString.Len() == 0)
 	{
-		OutErrorText = LOCTEXT("Error_EmptyName", "Empty parameter names are not allowed");
+		OutErrorText = LOCTEXT("Error_EmptyName", "Empty names are not allowed");
 		return false;
 	}
 
diff --git a/Engine/Plugins/Experimental/AnimNext/Source/AnimNextEditor/Private/EditorUtils.h b/Engine/Plugins/Experimental/AnimNext/Source/AnimNextEditor/Private/EditorUtils.h
--- a/Engine/Plugins/Experimental/AnimNext/Source/AnimNextEditor/Private/EditorUtils.h
+++ b/Engine/Plugins/Experimental/AnimNext/Source/AnimNextEditor/Private/EditorUtils.h
@@ -3,12 +3,11 @@
 #pragma once
 
 #include "CoreMinimal.h"
-#include "Param/AnimNextParameterBlockEntry.h"
 #include "Param/ParamType.h"
 #include "EdGraphSchema_K2.h"
 
 class UAnimNextGraph;
-class UAnimNextGraph_EditorData;
+class UAnimNextRigVMAssetEditorData;
 struct FAnimNextParamType;
 class UAnimNextParameterBlock;
 class UAnimNextParameterBlockBinding;
@@ -26,19 +25,19 @@ struct FUtils
 {
 	static FName ValidateName(const UObject* InObject, const FString& InName);
 
-	static void GetAllGraphNames(const UAnimNextGraph_EditorData* InEditorData, TSet<FName>& OutNames);
+	static void GetAllEntryNames(const UAnimNextRigVMAssetEditorData* InEditorData, TSet<FName>& OutNames);
 
 	static FAnimNextParamType GetParameterTypeFromMetaData(const FStringView& InStringView);
 
 	static FName ValidateName(const UAnimNextParameterBlock_EditorData* InEditorData, const FString& InName);
 
-	static void GetAllEntryNames(const UAnimNextParameterBlock_EditorData* InEditorData, TSet<FName>& OutNames);
-
 	static void GetFilteredVariableTypeTree(TArray<TSharedPtr<UEdGraphSchema_K2::FPinTypeTreeInfo>>& TypeTree, ETypeTreeFilter TypeTreeFilter);
 
 	static FName GetNewParameterName(const TCHAR* InBaseName, TArrayView<FName> InAdditionalExistingNames);
 
-	static bool IsValidParameterName(const FName InName, FText& OutErrorText);
+	static bool IsValidEntryNameString(FStringView InStringView, FText& OutErrorText);
+
+	static bool IsValidEntryName(const FName InName, FText& OutErrorText);
 
 	static bool DoesParameterNameExist(const FName InName);
 	
diff --git a/Engine/Plugins/Experimental/AnimNext/Source/AnimNextEditor/Private/Graph/AnimNextGraphEditor.cpp b/Engine/Plugins/Experimental/AnimNext/Source/AnimNextEditor/Private/Graph/AnimNextGraphEditor.cpp
--- a/Engine/Plugins/Experimental/AnimNext/Source/AnimNextEditor/Private/Graph/AnimNextGraphEditor.cpp
+++ b/Engine/Plugins/Experimental/AnimNext/Source/AnimNextEditor/Private/Graph/AnimNextGraphEditor.cpp
@@ -60,12 +60,6 @@ void FGraphEditor::InitEditor(const EToolkitMode::Type InMode, const TSharedPtr<
 	ExtendMenu();
 	ExtendToolbar();
 	RegenerateMenusAndToolbars();
-
-	// Open initial document
-	if(AnimNextGraph_EditorData->Graphs.Num() > 0)
-	{
-		DocumentManager->OpenDocument(FTabPayload_UObject::Make(AnimNextGraph_EditorData->Graphs[0]), FDocumentTracker::EOpenDocumentCause::OpenNewDocument);
-	}
 }
 
 void FGraphEditor::BindCommands()
diff --git a/Engine/Plugins/Experimental/AnimNext/Source/AnimNextEditor/Private/Graph/GraphFactory.cpp b/Engine/Plugins/Experimental/AnimNext/Source/AnimNextEditor/Private/Graph/GraphFactory.cpp
--- a/Engine/Plugins/Experimental/AnimNext/Source/AnimNextEditor/Private/Graph/GraphFactory.cpp
+++ b/Engine/Plugins/Experimental/AnimNext/Source/AnimNextEditor/Private/Graph/GraphFactory.cpp
@@ -23,17 +23,22 @@ bool UAnimNextGraphFactory::ConfigureProperties()
 
 UObject* UAnimNextGraphFactory::FactoryCreateNew(UClass* Class, UObject* InParent, FName Name, EObjectFlags Flags, UObject* Context, FFeedbackContext* Warn, FName CallingContext)
 {
-	UAnimNextGraph* NewGraph = NewObject<UAnimNextGraph>(InParent, Class, Name, Flags | RF_Public | RF_Standalone | RF_Transactional | RF_LoadCompleted);
+	EObjectFlags FlagsToUse = Flags | RF_Public | RF_Standalone | RF_Transactional | RF_LoadCompleted;
+	if(InParent == GetTransientPackage())
+	{
+		FlagsToUse &= ~RF_Standalone;
+	}
+
+	UAnimNextGraph* NewGraph = NewObject<UAnimNextGraph>(InParent, Class, Name, FlagsToUse);
 
 	// Create internal editor data
-	UAnimNextGraph_EditorData* EditorData = NewObject<UAnimNextGraph_EditorData>(NewGraph, TEXT("EditorData"));
+	UAnimNextGraph_EditorData* EditorData = NewObject<UAnimNextGraph_EditorData>(NewGraph, TEXT("EditorData"), RF_Transactional);
 	NewGraph->EditorData = EditorData;
+	EditorData->Initialize(/*bRecompileVM*/false);
 
 	// Add root graph
-	EditorData->AddGraph(TEXT("Root"));
-	check(EditorData->Graphs.Num() > 0);
-
-	EditorData->Initialize(/*bRecompileVM*/false);
+	EditorData->AddGraph(TEXT("Root"), false);
+	check(EditorData->Entries.Num() > 0);
 
 	// Compile the initial skeleton
 	UE::AnimNext::UncookedOnly::FUtils::Compile(NewGraph);
diff --git a/Engine/Plugins/Experimental/AnimNext/Source/AnimNextEditor/Private/Param/AnimNextParameterBlockEditor.cpp b/Engine/Plugins/Experimental/AnimNext/Source/AnimNextEditor/Private/Param/AnimNextParameterBlockEditor.cpp
--- a/Engine/Plugins/Experimental/AnimNext/Source/AnimNextEditor/Private/Param/AnimNextParameterBlockEditor.cpp
+++ b/Engine/Plugins/Experimental/AnimNext/Source/AnimNextEditor/Private/Param/AnimNextParameterBlockEditor.cpp
@@ -1,6 +1,7 @@
 ﻿// Copyright Epic Games, Inc. All Rights Reserved.
 
 #include "AnimNextParameterBlockEditor.h"
+
 #include "ParameterBlockEditorMode.h"
 #include "Param/AnimNextParameterBlock.h"
 #include "Param/AnimNextParameterBlock_EditorData.h"
@@ -12,9 +13,7 @@
 #include "RigVMModel/RigVMController.h"
 #include "Kismet2/BlueprintEditorUtils.h"
 #include "ScopedTransaction.h"
-#include "Param/AnimNextParameterBlockEntry.h"
 #include "Param/AnimNextParameterSettings.h"
-#include "ParameterBlockEditorMode.h"
 #include "Param/AnimNextParameterExecuteContext.h"
 #include "Widgets/Docking/SDockTab.h"
 
@@ -383,12 +382,6 @@ void FParameterBlockEditor::GetSaveableObjects(TArray<UObject*>& OutObjects) con
 
 	// Get external objects too
 	FExternalPackageHelper::GetExternalSaveableObjects(EditorData, OutObjects);
-
-	// Get any referenced objects too for convenience 
-	for(UAnimNextParameterBlockEntry* Entry : EditorData->Entries)
-	{
-		Entry->GetEditedObjects(OutObjects);
-	}
 }
 
 void FParameterBlockEditor::HandleSaveGraphState(UEdGraph* InGraph, FVector2D InViewOffset, float InZoomAmount)
diff --git a/Engine/Plugins/Experimental/AnimNext/Source/AnimNextEditor/Private/Param/ParameterBlockEditorMode.cpp b/Engine/Plugins/Experimental/AnimNext/Source/AnimNextEditor/Private/Param/ParameterBlockEditorMode.cpp
--- a/Engine/Plugins/Experimental/AnimNext/Source/AnimNextEditor/Private/Param/ParameterBlockEditorMode.cpp
+++ b/Engine/Plugins/Experimental/AnimNext/Source/AnimNextEditor/Private/Param/ParameterBlockEditorMode.cpp
@@ -2,18 +2,17 @@
 
 #include "ParameterBlockEditorMode.h"
 #include "AnimNextParameterBlockEditor.h"
+#include "AnimNextRigVMAssetEntry.h"
+#include "IAnimNextRigVMGraphInterface.h"
 #include "PropertyEditorModule.h"
-#include "Widgets/Docking/SDockTab.h"
+#include "Common/SRigVMAssetView.h"
 #include "EdGraph/EdGraph.h"
 #include "EdGraph/EdGraphSchema.h"
-#include "Widgets/Layout/SSpacer.h"
 #include "Modules/ModuleManager.h"
-#include "SParameterBlockView.h"
 #include "Param/AnimNextParameterBlock_EditorData.h"
 #include "RigVMModel/RigVMGraph.h"
-#include "Param/IAnimNextParameterBlockGraphInterface.h"
-#include "Param/AnimNextParameterBlockEntry.h"
-#include "Param/AnimNextParameterBlockEditor.h"
+#include "Widgets/Docking/SDockTab.h"
+#include "Widgets/Layout/SSpacer.h"
 
 #define LOCTEXT_NAMESPACE "ParameterBlockEditorMode"
 
@@ -229,7 +228,7 @@ public:
 private:
 	virtual TSharedRef<SWidget> CreateTabBody(const FWorkflowTabSpawnInfo& Info) const override
 	{
-		return SNew(SParameterBlockView, StaticCastSharedPtr<FParameterBlockEditor>(HostingApp.Pin())->EditorData)
+		return SNew(SRigVMAssetView, StaticCastSharedPtr<FParameterBlockEditor>(HostingApp.Pin())->EditorData)
 			.OnSelectionChanged(this, &FParameterBlockTabSummoner::HandleSelectionChanged)
 			.OnOpenGraph(this, &FParameterBlockTabSummoner::HandleOpenGraph)
 			.OnDeleteEntries(this, &FParameterBlockTabSummoner::HandleDeleteEntries);
@@ -256,21 +255,18 @@ private:
 		}
 	}
 
-	void HandleDeleteEntries(const TArray<UAnimNextParameterBlockEntry*>& InEntries) const
+	void HandleDeleteEntries(const TArray<UAnimNextRigVMAssetEntry*>& InEntries) const
 	{
 		TSharedPtr<FParameterBlockEditor> ParameterBlockEditor = StaticCastSharedPtr<FParameterBlockEditor>(HostingApp.Pin());
 		UAnimNextParameterBlock_EditorData* EditorData = StaticCastSharedPtr<FParameterBlockEditor>(HostingApp.Pin())->EditorData;
 
-		for(UAnimNextParameterBlockEntry* Entry : InEntries)
+		for(UAnimNextRigVMAssetEntry* Entry : InEntries)
 		{
-			if(IAnimNextParameterBlockGraphInterface* GraphInterface = Cast<IAnimNextParameterBlockGraphInterface>(Entry))
+			if(IAnimNextRigVMGraphInterface* GraphInterface = Cast<IAnimNextRigVMGraphInterface>(Entry))
 			{
-				if(URigVMGraph* RigVMGraph = GraphInterface->GetGraph())
+				if(URigVMEdGraph* EdGraph = GraphInterface->GetEdGraph())
 				{
-					if (UObject* EditorObject = EditorData->GetEditorObjectForRigVMGraph(RigVMGraph))
-					{
-						ParameterBlockEditor->CloseDocumentTab(EditorObject);
-					}
+					ParameterBlockEditor->CloseDocumentTab(EdGraph);
 				}
 			}
 		}
diff --git a/Engine/Plugins/Experimental/AnimNext/Source/AnimNextEditor/Private/Param/ParameterBlockFactory.cpp b/Engine/Plugins/Experimental/AnimNext/Source/AnimNextEditor/Private/Param/ParameterBlockFactory.cpp
--- a/Engine/Plugins/Experimental/AnimNext/Source/AnimNextEditor/Private/Param/ParameterBlockFactory.cpp
+++ b/Engine/Plugins/Experimental/AnimNext/Source/AnimNextEditor/Private/Param/ParameterBlockFactory.cpp
@@ -20,12 +20,17 @@ bool UAnimNextParameterBlockFactory::ConfigureProperties()
 
 UObject* UAnimNextParameterBlockFactory::FactoryCreateNew(UClass* Class, UObject* InParent, FName Name, EObjectFlags Flags, UObject* Context, FFeedbackContext* Warn, FName CallingContext)
 {
-	UAnimNextParameterBlock* NewBlock = NewObject<UAnimNextParameterBlock>(InParent, Class, Name, Flags | RF_Public | RF_Standalone | RF_Transactional | RF_LoadCompleted);
+	EObjectFlags FlagsToUse = Flags | RF_Public | RF_Standalone | RF_Transactional | RF_LoadCompleted;
+	if(InParent == GetTransientPackage())
+	{
+		FlagsToUse &= ~RF_Standalone;
+	}
+
+	UAnimNextParameterBlock* NewBlock = NewObject<UAnimNextParameterBlock>(InParent, Class, Name, FlagsToUse);
 
 	UAnimNextParameterBlock_EditorData* EditorData = NewObject<UAnimNextParameterBlock_EditorData>(NewBlock, TEXT("EditorData"), RF_Transactional);
 	NewBlock->EditorData = EditorData;
 	EditorData->Initialize(/*bRecompileVM*/false);
-	EditorData->GetRigVMClient()->SetExecuteContextStruct(FAnimNextParameterExecuteContext::StaticStruct());
 
 	// Compile the initial skeleton
 	UE::AnimNext::UncookedOnly::FUtils::Compile(NewBlock);
diff --git a/Engine/Plugins/Experimental/AnimNext/Source/AnimNextEditor/Private/Param/ParameterBlockParameterCustomization.cpp b/Engine/Plugins/Experimental/AnimNext/Source/AnimNextEditor/Private/Param/ParameterBlockParameterCustomization.cpp
--- a/Engine/Plugins/Experimental/AnimNext/Source/AnimNextEditor/Private/Param/ParameterBlockParameterCustomization.cpp
+++ b/Engine/Plugins/Experimental/AnimNext/Source/AnimNextEditor/Private/Param/ParameterBlockParameterCustomization.cpp
@@ -36,24 +36,23 @@ void FParameterBlockParameterCustomization::CustomizeDetails(IDetailLayoutBuilde
 	{
 		IDetailCategoryBuilder& ParameterCategory = DetailBuilder.EditCategory(TEXT("Parameter"), FText::GetEmpty(), ECategoryPriority::Important);
 
-
 		IDetailCategoryBuilder& DefaultValueCategory = DetailBuilder.EditCategory(TEXT("DefaultValue"), FText::GetEmpty(), ECategoryPriority::Default);
 
 		TSharedRef< SWidget > ColumnWidget = SNullWidget::NullWidget;
 
 		if (UAnimNextParameterBlock_EditorData* EditorData = Cast<UAnimNextParameterBlock_EditorData>(BlockParam->GetOuter()))
 		{
-			const FName ParameterName = BlockParam->GetParameterName();
-			if (UAnimNextParameterBlockEntry* BlockEntry = EditorData->FindBinding(ParameterName)) // filter non parameter entries
+			const FName ParameterName = BlockParam->GetEntryName();
+			if (UAnimNextRigVMAssetEntry* BlockEntry = EditorData->FindEntry(ParameterName)) 
 			{
 				if (UAnimNextParameterBlock* ReferencedBlock = UE::AnimNext::UncookedOnly::FUtils::GetBlock(EditorData))
 				{
 					FAddPropertyParams AddPropertyParams;
 					TArray<IDetailPropertyRow*> DetailPropertyRows;
 
-					if (ReferencedBlock->GetPropertyBag().FindPropertyDescByName(ParameterName))
+					if (ReferencedBlock->PropertyBag.FindPropertyDescByName(ParameterName))
 					{
-						IDetailPropertyRow* DetailPropertyRow = DefaultValueCategory.AddExternalStructureProperty(MakeShared<FInstancePropertyBagStructureDataProvider>(ReferencedBlock->GetPropertyBag()), ParameterName, EPropertyLocation::Default, AddPropertyParams);
+						IDetailPropertyRow* DetailPropertyRow = DefaultValueCategory.AddExternalStructureProperty(MakeShared<FInstancePropertyBagStructureDataProvider>(ReferencedBlock->PropertyBag), ParameterName, EPropertyLocation::Default, AddPropertyParams);
 						if (TSharedPtr<IPropertyHandle> Handle = DetailPropertyRow->GetPropertyHandle(); Handle.IsValid())
 						{
 							Handle->SetOnChildPropertyValuePreChange(FSimpleDelegate::CreateLambda([this, ReferencedBlock]()
@@ -64,7 +63,7 @@ void FParameterBlockParameterCustomization::CustomizeDetails(IDetailLayoutBuilde
 							{
 								if (UAnimNextParameterBlock_EditorData* EditorData = Cast<UAnimNextParameterBlock_EditorData>(ReferencedBlock->EditorData))
 								{
-									if (UAnimNextParameterBlockEntry* BlockEntry = EditorData->FindBinding(ParameterName))
+									if (UAnimNextRigVMAssetEntry* BlockEntry = EditorData->FindEntry(ParameterName))
 									{
 										BlockEntry->MarkPackageDirty();
 										ReferencedBlock->MarkPackageDirty(); // show the ParameterBlock has changed
diff --git a/Engine/Plugins/Experimental/AnimNext/Source/AnimNextEditor/Private/Param/ParametersAssetDefinitions.cpp b/Engine/Plugins/Experimental/AnimNext/Source/AnimNextEditor/Private/Param/ParametersAssetDefinitions.cpp
--- a/Engine/Plugins/Experimental/AnimNext/Source/AnimNextEditor/Private/Param/ParametersAssetDefinitions.cpp
+++ b/Engine/Plugins/Experimental/AnimNext/Source/AnimNextEditor/Private/Param/ParametersAssetDefinitions.cpp
@@ -28,10 +28,4 @@ EAssetCommandResult UAssetDefinition_AnimNextParameterBlock::OpenAssets(const FA
 	return EAssetCommandResult::Handled;
 }
 
-FText UAssetDefinition_AnimNextParameterBlockBinding::GetObjectDisplayNameText(UObject* Object) const
-{
-	UAnimNextParameterBlockBinding* Binding = CastChecked<UAnimNextParameterBlockBinding>(Object);
-	return FText::Format(LOCTEXT("BindingDisplayFormat", "{0} Binding"), FText::FromName(Binding->ParameterName));
-}
-
 #undef LOCTEXT_NAMESPACE
\ No newline at end of file
diff --git a/Engine/Plugins/Experimental/AnimNext/Source/AnimNextEditor/Private/Param/ParametersAssetDefinitions.h b/Engine/Plugins/Experimental/AnimNext/Source/AnimNextEditor/Private/Param/ParametersAssetDefinitions.h
--- a/Engine/Plugins/Experimental/AnimNext/Source/AnimNextEditor/Private/Param/ParametersAssetDefinitions.h
+++ b/Engine/Plugins/Experimental/AnimNext/Source/AnimNextEditor/Private/Param/ParametersAssetDefinitions.h
@@ -3,7 +3,6 @@
 #pragma once
 
 #include "AssetDefinitionDefault.h"
-#include "Param/AnimNextParameterBlockBinding.h"
 #include "Param/AnimNextParameterBlock.h"
 #include "ParametersAssetDefinitions.generated.h"
 
@@ -28,17 +27,4 @@ public:
 	}
 };
 
-UCLASS()
-class UAssetDefinition_AnimNextParameterBlockBinding : public UAssetDefinitionDefault
-{
-	GENERATED_BODY()
-
-public:
-	// UAssetDefinition interface
-	virtual FText GetAssetDisplayName() const override { return LOCTEXT("AnimNextParameterBinding", "AnimNext Parameter Binding"); }
-	virtual FLinearColor GetAssetColor() const override { return FLinearColor(FColor(128,64,64)); }
-	virtual TSoftClassPtr<UObject> GetAssetClass() const override { return UAnimNextParameterBlockBinding::StaticClass(); }
-	virtual FText GetObjectDisplayNameText(UObject* Object) const override;
-};
-
 #undef LOCTEXT_NAMESPACE
\ No newline at end of file
diff --git a/Engine/Plugins/Experimental/AnimNext/Source/AnimNextEditor/Private/Param/SAddParametersDialog.cpp b/Engine/Plugins/Experimental/AnimNext/Source/AnimNextEditor/Private/Param/SAddParametersDialog.cpp
--- a/Engine/Plugins/Experimental/AnimNext/Source/AnimNextEditor/Private/Param/SAddParametersDialog.cpp
+++ b/Engine/Plugins/Experimental/AnimNext/Source/AnimNextEditor/Private/Param/SAddParametersDialog.cpp
@@ -249,18 +249,12 @@ class SParameterToAdd : public SMultiColumnTableRow<TSharedRef<FParameterToAdd>>
 					{
 						const FString NewString = InNewText.ToString();
 
-						// See if this can be represented as an FName
-						if(!FName::IsValidXName(NewString, INVALID_NAME_CHARACTERS, &OutErrorText))
+						if(!FUtils::IsValidEntryNameString(NewString, OutErrorText))
 						{
 							return false;
 						}
 
 						const FName Name(*NewString);
-						if(!FUtils::IsValidParameterName(Name, OutErrorText))
-						{
-							return false;
-						}
-
 						if(FUtils::DoesParameterNameExist(Name))
 						{
 							OutErrorText = LOCTEXT("Error_NameExists", "This name already exists in the project");
diff --git a/Engine/Plugins/Experimental/AnimNext/Source/AnimNextEditor/Private/Workspace/AnimNextWorkspaceEditor.cpp b/Engine/Plugins/Experimental/AnimNext/Source/AnimNextEditor/Private/Workspace/AnimNextWorkspaceEditor.cpp
--- a/Engine/Plugins/Experimental/AnimNext/Source/AnimNextEditor/Private/Workspace/AnimNextWorkspaceEditor.cpp
+++ b/Engine/Plugins/Experimental/AnimNext/Source/AnimNextEditor/Private/Workspace/AnimNextWorkspaceEditor.cpp
@@ -2,6 +2,7 @@
 
 #include "AnimNextWorkspaceEditor.h"
 
+#include "AnimNextRigVMAssetEntry.h"
 #include "GraphDocumentSummoner.h"
 #include "AnimNextWorkspaceEditorMode.h"
 #include "AnimNextWorkspace.h"
@@ -27,6 +28,8 @@
 #include "Graph/AnimNextGraph.h"
 #include "Graph/AnimNextGraph_EdGraphNode.h"
 #include "GraphEditAction.h"
+#include "IAnimNextRigVMGraphInterface.h"
+#include "UncookedOnlyUtils.h"
 
 #define LOCTEXT_NAMESPACE "AnimNextWorkspaceEditor"
 
@@ -198,7 +201,7 @@ void FWorkspaceEditor::OpenWorkspaceForAsset(UObject* InAsset, EOpenWorkspaceMet
 	
 	if(WorkspaceEditor)
 	{
-		WorkspaceEditor->OpenDocument(InAsset, FDocumentTracker::OpenNewDocument);
+		WorkspaceEditor->OpenAssets({InAsset});
 	}
 }
 
@@ -252,6 +255,18 @@ void FWorkspaceEditor::OpenAssets(TConstArrayView<FAssetData> InAssets)
 		if(UObject* LoadedAsset = Asset.GetAsset())
 		{
 			OpenDocument(LoadedAsset, FDocumentTracker::EOpenDocumentCause::OpenNewDocument);
+
+			// If its a RigVM asset, open its first graph as well
+			if(UAnimNextRigVMAsset* RigVMAsset = Cast<UAnimNextRigVMAsset>(LoadedAsset))
+			{
+				UAnimNextRigVMAssetEditorData* EditorData = UncookedOnly::FUtils::GetEditorData(RigVMAsset);
+				check(EditorData);
+				EditorData->ForEachEntryOfType<IAnimNextRigVMGraphInterface>([this](IAnimNextRigVMGraphInterface* InEntry)
+				{
+					OpenDocument(InEntry->GetEdGraph(), FDocumentTracker::EOpenDocumentCause::OpenNewDocument);
+					return false;
+				});
+			}
 		}
 	}
 }
@@ -334,6 +349,34 @@ void FWorkspaceEditor::OnGraphModified(ERigVMGraphNotifType Type, URigVMGraph* G
 	}
 }
 
+void FWorkspaceEditor::OnOpenGraph(URigVMGraph* InGraph)
+{
+	if(IRigVMClientHost* RigVMClientHost = InGraph->GetImplementingOuter<IRigVMClientHost>())
+	{
+		if(UObject* EditorObject = RigVMClientHost->GetEditorObjectForRigVMGraph(InGraph))
+		{
+			OpenDocument(EditorObject, FDocumentTracker::EOpenDocumentCause::OpenNewDocument);
+		}
+	}
+}
+
+void FWorkspaceEditor::OnDeleteEntries(const TArray<UAnimNextRigVMAssetEntry*>& InEntries)
+{
+	if(InEntries.Num() > 0)
+	{
+		for(UAnimNextRigVMAssetEntry* Entry : InEntries)
+		{
+			if(IAnimNextRigVMGraphInterface* GraphInterface = Cast<IAnimNextRigVMGraphInterface>(Entry))
+			{
+				if(URigVMEdGraph* EdGraph = GraphInterface->GetEdGraph())
+				{
+					CloseDocumentTab(EdGraph);
+				}
+			}
+		}
+	}
+}
+
 void FWorkspaceEditor::SetFocusedGraphEditor(TSharedPtr<SGraphEditor> InGraphEditor)
 {
 	// Update the graph editor that is currently focused
@@ -494,7 +537,7 @@ bool FWorkspaceEditor::CanDeleteSelectedNodes()
 	return bCanUserDeleteNode;
 }
 
-void FWorkspaceEditor::SetSelectedObjects(TArray<UObject*> InObjects)
+void FWorkspaceEditor::SetSelectedObjects(const TArray<UObject*>& InObjects)
 {
 	if(DetailsView.IsValid())
 	{
@@ -560,6 +603,8 @@ void FWorkspaceEditor::OnClose()
 		DetailsView->SetObject(nullptr);
 		DetailsView.Reset();
 	}
+
+	FWorkflowCentricApplication::OnClose();
 }
 
 }
diff --git a/Engine/Plugins/Experimental/AnimNext/Source/AnimNextEditor/Private/Workspace/AnimNextWorkspaceEditor.h b/Engine/Plugins/Experimental/AnimNext/Source/AnimNextEditor/Private/Workspace/AnimNextWorkspaceEditor.h
--- a/Engine/Plugins/Experimental/AnimNext/Source/AnimNextEditor/Private/Workspace/AnimNextWorkspaceEditor.h
+++ b/Engine/Plugins/Experimental/AnimNext/Source/AnimNextEditor/Private/Workspace/AnimNextWorkspaceEditor.h
@@ -7,6 +7,7 @@
 #include "GraphEditor.h"
 #include "WorkflowOrientedApp/WorkflowTabManager.h"
 
+class UAnimNextRigVMAssetEntry;
 enum class ERigVMGraphNotifType : uint8;
 class UAnimNextWorkspace;
 class FDocumentTracker;
@@ -123,7 +124,7 @@ private:
 	
 	bool CanDeleteSelectedNodes();
 
-	void SetSelectedObjects(TArray<UObject*> InObjects);
+	void SetSelectedObjects(const TArray<UObject*>& InObjects);
 
 	void HandleDetailsViewCreated(TSharedRef<IDetailsView> InDetailsView)
 	{
@@ -145,6 +146,10 @@ private:
 
 	void OnGraphModified(ERigVMGraphNotifType Type, URigVMGraph* Graph, UObject* Subject);
 
+	void OnOpenGraph(URigVMGraph* InGraph);
+
+	void OnDeleteEntries(const TArray<UAnimNextRigVMAssetEntry*>& InEntries);
+
 	// The asset we are editing
 	UAnimNextWorkspace* Workspace = nullptr;
 
diff --git a/Engine/Plugins/Experimental/AnimNext/Source/AnimNextEditor/Private/Workspace/GraphDocumentSummoner.cpp b/Engine/Plugins/Experimental/AnimNext/Source/AnimNextEditor/Private/Workspace/GraphDocumentSummoner.cpp
--- a/Engine/Plugins/Experimental/AnimNext/Source/AnimNextEditor/Private/Workspace/GraphDocumentSummoner.cpp
+++ b/Engine/Plugins/Experimental/AnimNext/Source/AnimNextEditor/Private/Workspace/GraphDocumentSummoner.cpp
@@ -7,6 +7,7 @@
 #include "Widgets/Docking/SDockTab.h"
 #include "Widgets/Layout/SSpacer.h"
 #include "EdGraph/EdGraphSchema.h"
+#include "EdGraph/RigVMEdGraph.h"
 #include "Framework/Commands/GenericCommands.h"
 
 #define LOCTEXT_NAMESPACE "AnimNextGraphDocumentSummoner"
@@ -86,7 +87,7 @@ TSharedRef<SWidget> FGraphDocumentSummoner::CreateTabBodyForObject(const FWorkfl
 {
 	SGraphEditor::FGraphEditorEvents Events;
 	Events.OnCreateActionMenu = SGraphEditor::FOnCreateActionMenu::CreateSP(this, &FGraphDocumentSummoner::OnCreateGraphActionMenu);
-	Events.OnSelectionChanged = SGraphEditor::FOnSelectionChanged::CreateSP(HostingAppPtr.Pin().Get(), &FWorkspaceEditor::OnGraphSelectionChanged);
+	Events.OnSelectionChanged = SGraphEditor::FOnSelectionChanged::CreateSP(this, &FGraphDocumentSummoner::OnGraphSelectionChanged, TWeakObjectPtr<UEdGraph>(DocumentID));
 	
 	return SNew(SGraphEditor)
 		.AdditionalCommands(CommandList)
@@ -96,6 +97,38 @@ TSharedRef<SWidget> FGraphDocumentSummoner::CreateTabBodyForObject(const FWorkfl
 		.AssetEditorToolkit(HostingAppPtr);
 }
 
+void FGraphDocumentSummoner::OnGraphSelectionChanged(const TSet<UObject*>& NewSelection, TWeakObjectPtr<UEdGraph> InGraph) const
+{
+	URigVMEdGraph* RigVMEdGraph = Cast<URigVMEdGraph>(InGraph.Get());
+	if (RigVMEdGraph == nullptr)
+	{
+		return;
+	}
+
+	if (RigVMEdGraph->bIsSelecting || GIsTransacting)
+	{
+		return;
+	}
+
+	TGuardValue<bool> SelectGuard(RigVMEdGraph->bIsSelecting, true);
+
+	TArray<FName> NodeNamesToSelect;
+	for (UObject* Object : NewSelection)
+	{
+		if (URigVMEdGraphNode* RigVMEdGraphNode = Cast<URigVMEdGraphNode>(Object))
+		{
+			NodeNamesToSelect.Add(RigVMEdGraphNode->GetModelNodeName());
+		}
+		else if(UEdGraphNode* Node = Cast<UEdGraphNode>(Object))
+		{
+			NodeNamesToSelect.Add(Node->GetFName());
+		}
+	}
+	RigVMEdGraph->GetController()->SetNodeSelection(NodeNamesToSelect, true, true);
+
+	HostingAppPtr.Pin()->OnGraphSelectionChanged(NewSelection);
+}
+
 bool FGraphDocumentSummoner::IsPayloadSupported(TSharedRef<FTabPayload> Payload) const
 {
 	UObject* Object = Payload->IsValid() ? FTabPayload_UObject::CastChecked<UObject>(Payload) : nullptr;
diff --git a/Engine/Plugins/Experimental/AnimNext/Source/AnimNextEditor/Private/Workspace/GraphDocumentSummoner.h b/Engine/Plugins/Experimental/AnimNext/Source/AnimNextEditor/Private/Workspace/GraphDocumentSummoner.h
--- a/Engine/Plugins/Experimental/AnimNext/Source/AnimNextEditor/Private/Workspace/GraphDocumentSummoner.h
+++ b/Engine/Plugins/Experimental/AnimNext/Source/AnimNextEditor/Private/Workspace/GraphDocumentSummoner.h
@@ -45,6 +45,8 @@ private:
 	// Create the context 'action' menu used for graphs
 	virtual FActionMenuContent OnCreateGraphActionMenu(UEdGraph* InGraph, const FVector2D& InNodePosition, const TArray<UEdGraphPin*>& InDraggedPins, bool bAutoExpand, SGraphEditor::FActionMenuClosed InOnMenuClosed) const = 0;
 
+	void OnGraphSelectionChanged(const TSet<UObject*>& NewSelection, TWeakObjectPtr<UEdGraph> InGraph) const;
+	
 	// The hosting app
 	TWeakPtr<FWorkspaceEditor> HostingAppPtr;
 
