AnimNext: Begin merging parameter blocks and graphs

diff --git a/Engine/Plugins/Experimental/AnimNext/Source/AnimNextTestSuite/Private/AnimNextDecoratorBaseTest.cpp b/Engine/Plugins/Experimental/AnimNext/Source/AnimNextTestSuite/Private/AnimNextDecoratorBaseTest.cpp
--- a/Engine/Plugins/Experimental/AnimNext/Source/AnimNextTestSuite/Private/AnimNextDecoratorBaseTest.cpp
+++ b/Engine/Plugins/Experimental/AnimNext/Source/AnimNextTestSuite/Private/AnimNextDecoratorBaseTest.cpp
@@ -2,6 +2,7 @@
 
 #include "AnimNextDecoratorBaseTest.h"
 #include "AnimNextRuntimeTest.h"
+#include "AnimNextTest.h"
 
 #include "Misc/AutomationTest.h"
 
@@ -497,207 +498,209 @@ IMPLEMENT_SIMPLE_AUTOMATION_TEST(FAnimationAnimNextRuntimeTest_NodeLifetime, "An
 bool FAnimationAnimNextRuntimeTest_NodeLifetime::RunTest(const FString& InParameters)
 {
 	using namespace UE::AnimNext;
+	{
+		AUTO_REGISTER_ANIM_DECORATOR(FDecoratorA_Base)
+		AUTO_REGISTER_ANIM_DECORATOR(FDecoratorAB_Add)
+		AUTO_REGISTER_ANIM_DECORATOR(FDecoratorAC_Add)
 
-	AUTO_REGISTER_ANIM_DECORATOR(FDecoratorA_Base)
-	AUTO_REGISTER_ANIM_DECORATOR(FDecoratorAB_Add)
-	AUTO_REGISTER_ANIM_DECORATOR(FDecoratorAC_Add)
+		UFactory* GraphFactory = NewObject<UAnimNextGraphFactory>();
+		UAnimNextGraph* AnimNextGraph = CastChecked<UAnimNextGraph>(GraphFactory->FactoryCreateNew(UAnimNextGraph::StaticClass(), GetTransientPackage(), TEXT("TestAnimNextGraph"), RF_Transient, nullptr, nullptr, NAME_None));
+		UE_RETURN_ON_ERROR(AnimNextGraph != nullptr, "FAnimationAnimNextRuntimeTest_NodeLifetime -> Failed to create animation graph");
 
-	UFactory* GraphFactory = NewObject<UAnimNextGraphFactory>();
-	UAnimNextGraph* AnimNextGraph = CastChecked<UAnimNextGraph>(GraphFactory->FactoryCreateNew(UAnimNextGraph::StaticClass(), GetTransientPackage(), TEXT("TestAnimNextGraph"), RF_Transient, nullptr, nullptr, NAME_None));
-	UE_RETURN_ON_ERROR(AnimNextGraph != nullptr, "FAnimationAnimNextRuntimeTest_NodeLifetime -> Failed to create animation graph");
+		FScopedClearNodeTemplateRegistry ScopedClearNodeTemplateRegistry;
+		FNodeTemplateRegistry& Registry = FNodeTemplateRegistry::Get();
 
-	FScopedClearNodeTemplateRegistry ScopedClearNodeTemplateRegistry;
-	FNodeTemplateRegistry& Registry = FNodeTemplateRegistry::Get();
+		TArray<FDecoratorUID> NodeTemplateDecoratorList;
+		NodeTemplateDecoratorList.Add(FDecoratorA_Base::DecoratorUID);
+		NodeTemplateDecoratorList.Add(FDecoratorAB_Add::DecoratorUID);
+		NodeTemplateDecoratorList.Add(FDecoratorAC_Add::DecoratorUID);
+		NodeTemplateDecoratorList.Add(FDecoratorA_Base::DecoratorUID);
+		NodeTemplateDecoratorList.Add(FDecoratorAC_Add::DecoratorUID);
 
-	TArray<FDecoratorUID> NodeTemplateDecoratorList;
-	NodeTemplateDecoratorList.Add(FDecoratorA_Base::DecoratorUID);
-	NodeTemplateDecoratorList.Add(FDecoratorAB_Add::DecoratorUID);
-	NodeTemplateDecoratorList.Add(FDecoratorAC_Add::DecoratorUID);
-	NodeTemplateDecoratorList.Add(FDecoratorA_Base::DecoratorUID);
-	NodeTemplateDecoratorList.Add(FDecoratorAC_Add::DecoratorUID);
-
-	// Populate our node template registry
-	TArray<uint8> NodeTemplateBuffer0;
-	const FNodeTemplate* NodeTemplate0 = FNodeTemplateBuilder::BuildNodeTemplate(NodeTemplateDecoratorList, NodeTemplateBuffer0);
+		// Populate our node template registry
+		TArray<uint8> NodeTemplateBuffer0;
+		const FNodeTemplate* NodeTemplate0 = FNodeTemplateBuilder::BuildNodeTemplate(NodeTemplateDecoratorList, NodeTemplateBuffer0);
 
-	FNodeTemplateRegistryHandle TemplateHandle0 = Registry.FindOrAdd(NodeTemplate0);
-	AddErrorIfFalse(TemplateHandle0.IsValid(), "FAnimationAnimNextRuntimeTest_NodeLifetime -> Registry should contain our template");
+		FNodeTemplateRegistryHandle TemplateHandle0 = Registry.FindOrAdd(NodeTemplate0);
+		AddErrorIfFalse(TemplateHandle0.IsValid(), "FAnimationAnimNextRuntimeTest_NodeLifetime -> Registry should contain our template");
 
-	TArray<FNodeHandle> NodeHandles;
+		TArray<FNodeHandle> NodeHandles;
 
-	// Write our graph
-	TArray<uint8> GraphSharedDataArchiveBuffer;
-	TArray<TObjectPtr<UObject>> GraphReferencedObjects;
-	{
-		FDecoratorWriter DecoratorWriter;
+		// Write our graph
+		TArray<uint8> GraphSharedDataArchiveBuffer;
+		TArray<TObjectPtr<UObject>> GraphReferencedObjects;
+		{
+			FDecoratorWriter DecoratorWriter;
 
-		NodeHandles.Add(DecoratorWriter.RegisterNode(*NodeTemplate0));
-		NodeHandles.Add(DecoratorWriter.RegisterNode(*NodeTemplate0));
+			NodeHandles.Add(DecoratorWriter.RegisterNode(*NodeTemplate0));
+			NodeHandles.Add(DecoratorWriter.RegisterNode(*NodeTemplate0));
 
-		// We don't have decorator properties
+			// We don't have decorator properties
 
-		DecoratorWriter.BeginNodeWriting();
-		DecoratorWriter.WriteNode(NodeHandles[0],
-			[](uint32 DecoratorIndex, const FString& PropertyName)
-			{
-				return FString();
-			},
-			[](uint32 DecoratorIndex, const FString& PropertyName)
-			{
-				return false;
-			});
-		DecoratorWriter.WriteNode(NodeHandles[1],
-			[](uint32 DecoratorIndex, const FString& PropertyName)
-			{
-				return FString();
-			},
-			[](uint32 DecoratorIndex, const FString& PropertyName)
-			{
-				return false;
-			});
-		DecoratorWriter.EndNodeWriting();
-
-		AddErrorIfFalse(DecoratorWriter.GetErrorState() == FDecoratorWriter::EErrorState::None, "FAnimationAnimNextRuntimeTest_NodeLifetime -> Failed to write decorators");
-		GraphSharedDataArchiveBuffer = DecoratorWriter.GetGraphSharedData();
-		GraphReferencedObjects = DecoratorWriter.GetGraphReferencedObjects();
-	}
-
-	// Read our graph
-	FTestUtils::LoadFromArchiveBuffer(*AnimNextGraph, NodeHandles, GraphSharedDataArchiveBuffer);
+			DecoratorWriter.BeginNodeWriting();
+			DecoratorWriter.WriteNode(NodeHandles[0],
+				[](uint32 DecoratorIndex, const FString& PropertyName)
+				{
+					return FString();
+				},
+				[](uint32 DecoratorIndex, const FString& PropertyName)
+				{
+					return false;
+				});
+			DecoratorWriter.WriteNode(NodeHandles[1],
+				[](uint32 DecoratorIndex, const FString& PropertyName)
+				{
+					return FString();
+				},
+				[](uint32 DecoratorIndex, const FString& PropertyName)
+				{
+					return false;
+				});
+			DecoratorWriter.EndNodeWriting();
 
-	FAnimNextGraphInstancePtr GraphInstance;
-	AnimNextGraph->AllocateInstance(GraphInstance);
+			AddErrorIfFalse(DecoratorWriter.GetErrorState() == FDecoratorWriter::EErrorState::None, "FAnimationAnimNextRuntimeTest_NodeLifetime -> Failed to write decorators");
+			GraphSharedDataArchiveBuffer = DecoratorWriter.GetGraphSharedData();
+			GraphReferencedObjects = DecoratorWriter.GetGraphReferencedObjects();
+		}
 
-	FExecutionContext Context(GraphInstance);
+		// Read our graph
+		FTestUtils::LoadFromArchiveBuffer(*AnimNextGraph, NodeHandles, GraphSharedDataArchiveBuffer);
 
-	// Validate handle bookkeeping
-	{
-		FDecoratorBinding RootBinding;									// Empty, no parent
-		FAnimNextDecoratorHandle DecoratorHandle00(NodeHandles[0], 0);	// Point to first node, first base decorator
+		FAnimNextGraphInstancePtr GraphInstance;
+		AnimNextGraph->AllocateInstance(GraphInstance);
 
-		// Allocate a node
-		FDecoratorPtr DecoratorPtr00 = Context.AllocateNodeInstance(RootBinding, DecoratorHandle00);
-		AddErrorIfFalse(DecoratorPtr00.IsValid(), "FAnimationAnimNextRuntimeTest_NodeLifetime -> Failed to allocate a node instance");
-		AddErrorIfFalse(DecoratorPtr00.GetDecoratorIndex() == 0, "FAnimationAnimNextRuntimeTest_NodeLifetime -> Allocated decorator pointer should point to root decorator");
-		AddErrorIfFalse(!DecoratorPtr00.IsWeak(), "FAnimationAnimNextRuntimeTest_NodeLifetime -> Allocated decorator pointer should not be weak, we have no parent");
-		AddErrorIfFalse(DecoratorPtr00.GetNodeInstance()->GetNodeHandle() == NodeHandles[0], "FAnimationAnimNextRuntimeTest_NodeLifetime -> Allocated node should point to the provided node handle");
-		AddErrorIfFalse(DecoratorPtr00.GetNodeInstance()->GetReferenceCount() == 1, "FAnimationAnimNextRuntimeTest_NodeLifetime -> Allocated node should have a single reference");
+		FExecutionContext Context(GraphInstance);
 
+		// Validate handle bookkeeping
 		{
-			FWeakDecoratorPtr WeakDecoratorPtr00(DecoratorPtr00);
-			AddErrorIfFalse(WeakDecoratorPtr00.GetNodeInstance() == DecoratorPtr00.GetNodeInstance(), "FAnimationAnimNextRuntimeTest_NodeLifetime -> Weak reference should point to the same node instance");
-			AddErrorIfFalse(WeakDecoratorPtr00.GetDecoratorIndex() == DecoratorPtr00.GetDecoratorIndex(), "FAnimationAnimNextRuntimeTest_NodeLifetime -> Weak reference should point to the same decorator index");
-			AddErrorIfFalse(DecoratorPtr00.GetNodeInstance()->GetReferenceCount() == 1, "FAnimationAnimNextRuntimeTest_NodeLifetime -> Weak references shouldn't increase ref count");
-		}
-
-		AddErrorIfFalse(DecoratorPtr00.GetNodeInstance()->GetReferenceCount() == 1, "FAnimationAnimNextRuntimeTest_NodeLifetime -> Weak references shouldn't decrease ref count");
+			FDecoratorBinding RootBinding;									// Empty, no parent
+			FAnimNextDecoratorHandle DecoratorHandle00(NodeHandles[0], 0);	// Point to first node, first base decorator
 
-		{
-			FWeakDecoratorPtr WeakDecoratorPtr00 = DecoratorPtr00;
-			AddErrorIfFalse(WeakDecoratorPtr00.GetNodeInstance() == DecoratorPtr00.GetNodeInstance(), "FAnimationAnimNextRuntimeTest_NodeLifetime -> Weak reference should point to the same node instance");
-			AddErrorIfFalse(WeakDecoratorPtr00.GetDecoratorIndex() == DecoratorPtr00.GetDecoratorIndex(), "FAnimationAnimNextRuntimeTest_NodeLifetime -> Weak reference should point to the same decorator index");
-			AddErrorIfFalse(DecoratorPtr00.GetNodeInstance()->GetReferenceCount() == 1, "FAnimationAnimNextRuntimeTest_NodeLifetime -> Weak references shouldn't increase ref count");
-		}
+			// Allocate a node
+			FDecoratorPtr DecoratorPtr00 = Context.AllocateNodeInstance(RootBinding, DecoratorHandle00);
+			AddErrorIfFalse(DecoratorPtr00.IsValid(), "FAnimationAnimNextRuntimeTest_NodeLifetime -> Failed to allocate a node instance");
+			AddErrorIfFalse(DecoratorPtr00.GetDecoratorIndex() == 0, "FAnimationAnimNextRuntimeTest_NodeLifetime -> Allocated decorator pointer should point to root decorator");
+			AddErrorIfFalse(!DecoratorPtr00.IsWeak(), "FAnimationAnimNextRuntimeTest_NodeLifetime -> Allocated decorator pointer should not be weak, we have no parent");
+			AddErrorIfFalse(DecoratorPtr00.GetNodeInstance()->GetNodeHandle() == NodeHandles[0], "FAnimationAnimNextRuntimeTest_NodeLifetime -> Allocated node should point to the provided node handle");
+			AddErrorIfFalse(DecoratorPtr00.GetNodeInstance()->GetReferenceCount() == 1, "FAnimationAnimNextRuntimeTest_NodeLifetime -> Allocated node should have a single reference");
 
-		AddErrorIfFalse(DecoratorPtr00.GetNodeInstance()->GetReferenceCount() == 1, "FAnimationAnimNextRuntimeTest_NodeLifetime -> Weak references shouldn't decrease ref count");
+			{
+				FWeakDecoratorPtr WeakDecoratorPtr00(DecoratorPtr00);
+				AddErrorIfFalse(WeakDecoratorPtr00.GetNodeInstance() == DecoratorPtr00.GetNodeInstance(), "FAnimationAnimNextRuntimeTest_NodeLifetime -> Weak reference should point to the same node instance");
+				AddErrorIfFalse(WeakDecoratorPtr00.GetDecoratorIndex() == DecoratorPtr00.GetDecoratorIndex(), "FAnimationAnimNextRuntimeTest_NodeLifetime -> Weak reference should point to the same decorator index");
+				AddErrorIfFalse(DecoratorPtr00.GetNodeInstance()->GetReferenceCount() == 1, "FAnimationAnimNextRuntimeTest_NodeLifetime -> Weak references shouldn't increase ref count");
+			}
 
-		{
-			FDecoratorPtr DecoratorPtr00_1(DecoratorPtr00);
-			AddErrorIfFalse(DecoratorPtr00_1.GetNodeInstance() == DecoratorPtr00.GetNodeInstance(), "FAnimationAnimNextRuntimeTest_NodeLifetime -> Strong reference should point to the same node instance");
-			AddErrorIfFalse(DecoratorPtr00_1.GetDecoratorIndex() == DecoratorPtr00.GetDecoratorIndex(), "FAnimationAnimNextRuntimeTest_NodeLifetime -> Strong reference should point to the same decorator index");
-			AddErrorIfFalse(DecoratorPtr00.GetNodeInstance()->GetReferenceCount() == 2, "FAnimationAnimNextRuntimeTest_NodeLifetime -> Strong references should increase ref count");
-		}
+			AddErrorIfFalse(DecoratorPtr00.GetNodeInstance()->GetReferenceCount() == 1, "FAnimationAnimNextRuntimeTest_NodeLifetime -> Weak references shouldn't decrease ref count");
 
-		AddErrorIfFalse(DecoratorPtr00.GetNodeInstance()->GetReferenceCount() == 1, "FAnimationAnimNextRuntimeTest_NodeLifetime -> Strong references should decrease ref count");
+			{
+				FWeakDecoratorPtr WeakDecoratorPtr00 = DecoratorPtr00;
+				AddErrorIfFalse(WeakDecoratorPtr00.GetNodeInstance() == DecoratorPtr00.GetNodeInstance(), "FAnimationAnimNextRuntimeTest_NodeLifetime -> Weak reference should point to the same node instance");
+				AddErrorIfFalse(WeakDecoratorPtr00.GetDecoratorIndex() == DecoratorPtr00.GetDecoratorIndex(), "FAnimationAnimNextRuntimeTest_NodeLifetime -> Weak reference should point to the same decorator index");
+				AddErrorIfFalse(DecoratorPtr00.GetNodeInstance()->GetReferenceCount() == 1, "FAnimationAnimNextRuntimeTest_NodeLifetime -> Weak references shouldn't increase ref count");
+			}
 
-		{
-			FDecoratorPtr DecoratorPtr00_1 = DecoratorPtr00;
-			AddErrorIfFalse(DecoratorPtr00_1.GetNodeInstance() == DecoratorPtr00.GetNodeInstance(), "FAnimationAnimNextRuntimeTest_NodeLifetime -> Strong reference should point to the same node instance");
-			AddErrorIfFalse(DecoratorPtr00_1.GetDecoratorIndex() == DecoratorPtr00.GetDecoratorIndex(), "FAnimationAnimNextRuntimeTest_NodeLifetime -> Strong reference should point to the same decorator index");
-			AddErrorIfFalse(DecoratorPtr00.GetNodeInstance()->GetReferenceCount() == 2, "FAnimationAnimNextRuntimeTest_NodeLifetime -> Strong references should increase ref count");
-		}
+			AddErrorIfFalse(DecoratorPtr00.GetNodeInstance()->GetReferenceCount() == 1, "FAnimationAnimNextRuntimeTest_NodeLifetime -> Weak references shouldn't decrease ref count");
 
-		AddErrorIfFalse(DecoratorPtr00.GetNodeInstance()->GetReferenceCount() == 1, "FAnimationAnimNextRuntimeTest_NodeLifetime -> Strong references should decrease ref count");
-	}
+			{
+				FDecoratorPtr DecoratorPtr00_1(DecoratorPtr00);
+				AddErrorIfFalse(DecoratorPtr00_1.GetNodeInstance() == DecoratorPtr00.GetNodeInstance(), "FAnimationAnimNextRuntimeTest_NodeLifetime -> Strong reference should point to the same node instance");
+				AddErrorIfFalse(DecoratorPtr00_1.GetDecoratorIndex() == DecoratorPtr00.GetDecoratorIndex(), "FAnimationAnimNextRuntimeTest_NodeLifetime -> Strong reference should point to the same decorator index");
+				AddErrorIfFalse(DecoratorPtr00.GetNodeInstance()->GetReferenceCount() == 2, "FAnimationAnimNextRuntimeTest_NodeLifetime -> Strong references should increase ref count");
+			}
 
-	// Validate parent support
-	{
-		FDecoratorBinding RootBinding;									// Empty, no parent
-		FAnimNextDecoratorHandle DecoratorHandle00(NodeHandles[0], 0);	// Point to first node, first base decorator
-		FAnimNextDecoratorHandle DecoratorHandle03(NodeHandles[0], 3);	// Point to first node, second base decorator
-		FAnimNextDecoratorHandle DecoratorHandle10(NodeHandles[1], 0);	// Point to second node, first base decorator
-
-		// Allocate our first node
-		FDecoratorPtr DecoratorPtr00 = Context.AllocateNodeInstance(RootBinding, DecoratorHandle00);
-		AddErrorIfFalse(DecoratorPtr00.IsValid(), "FAnimationAnimNextRuntimeTest_NodeLifetime -> Failed to allocate a node instance");
-
-		// Allocate a new node, using the first as a parent
-		// Both decorators live on the same node, the returned handle should be weak on the parent
-		FDecoratorPtr DecoratorPtr03 = Context.AllocateNodeInstance(DecoratorPtr00, DecoratorHandle03);
-		AddErrorIfFalse(DecoratorPtr03.IsValid(), "FAnimationAnimNextRuntimeTest_NodeLifetime -> Failed to allocate a node instance");
-		AddErrorIfFalse(DecoratorPtr03.GetDecoratorIndex() == 3, "FAnimationAnimNextRuntimeTest_NodeLifetime -> Allocated decorator pointer should point to fourth decorator");
-		AddErrorIfFalse(DecoratorPtr03.IsWeak(), "FAnimationAnimNextRuntimeTest_NodeLifetime -> Allocated decorator pointer should be weak, we have the same parent");
-		AddErrorIfFalse(DecoratorPtr03.GetNodeInstance()->GetNodeHandle() == NodeHandles[0], "FAnimationAnimNextRuntimeTest_NodeLifetime -> Allocated node should point to the provided node handle");
-		AddErrorIfFalse(DecoratorPtr03.GetNodeInstance() == DecoratorPtr00.GetNodeInstance(), "FAnimationAnimNextRuntimeTest_NodeLifetime -> Handles should point to the same node instance");
-		AddErrorIfFalse(DecoratorPtr03.GetNodeInstance()->GetReferenceCount() == 1, "FAnimationAnimNextRuntimeTest_NodeLifetime -> Allocated node should have one reference");
-
-		// Allocate a new node, using the first as a parent
-		// The second decorator lives on a new node, a new node instance will be allocated
-		FDecoratorPtr DecoratorPtr10 = Context.AllocateNodeInstance(DecoratorPtr00, DecoratorHandle10);
-		AddErrorIfFalse(DecoratorPtr10.IsValid(), "FAnimationAnimNextRuntimeTest_NodeLifetime -> Failed to allocate a node instance");
-		AddErrorIfFalse(DecoratorPtr10.GetDecoratorIndex() == 0, "FAnimationAnimNextRuntimeTest_NodeLifetime -> Allocated decorator pointer should point to first decorator");
-		AddErrorIfFalse(!DecoratorPtr10.IsWeak(), "FAnimationAnimNextRuntimeTest_NodeLifetime -> Allocated decorator pointer should not be weak, we have the same parent but a different node handle");
-		AddErrorIfFalse(DecoratorPtr10.GetNodeInstance()->GetNodeHandle() == NodeHandles[1], "FAnimationAnimNextRuntimeTest_NodeLifetime -> Allocated node should point to the provided node handle");
-		AddErrorIfFalse(DecoratorPtr10.GetNodeInstance() != DecoratorPtr00.GetNodeInstance(), "FAnimationAnimNextRuntimeTest_NodeLifetime -> Handles should not point to the same node instance");
-		AddErrorIfFalse(DecoratorPtr10.GetNodeInstance()->GetReferenceCount() == 1, "FAnimationAnimNextRuntimeTest_NodeLifetime -> Allocated node should have one reference");
-	}
+			AddErrorIfFalse(DecoratorPtr00.GetNodeInstance()->GetReferenceCount() == 1, "FAnimationAnimNextRuntimeTest_NodeLifetime -> Strong references should decrease ref count");
 
-	// Validate constructors and destructors
-	{
-		TArray<FDecoratorUID> ConstructedDecorators;
-		TArray<FDecoratorUID> DestructedDecorators;
+			{
+				FDecoratorPtr DecoratorPtr00_1 = DecoratorPtr00;
+				AddErrorIfFalse(DecoratorPtr00_1.GetNodeInstance() == DecoratorPtr00.GetNodeInstance(), "FAnimationAnimNextRuntimeTest_NodeLifetime -> Strong reference should point to the same node instance");
+				AddErrorIfFalse(DecoratorPtr00_1.GetDecoratorIndex() == DecoratorPtr00.GetDecoratorIndex(), "FAnimationAnimNextRuntimeTest_NodeLifetime -> Strong reference should point to the same decorator index");
+				AddErrorIfFalse(DecoratorPtr00.GetNodeInstance()->GetReferenceCount() == 2, "FAnimationAnimNextRuntimeTest_NodeLifetime -> Strong references should increase ref count");
+			}
 
-		Private::ConstructedDecorators = &ConstructedDecorators;
-		Private::DestructedDecorators = &DestructedDecorators;
+			AddErrorIfFalse(DecoratorPtr00.GetNodeInstance()->GetReferenceCount() == 1, "FAnimationAnimNextRuntimeTest_NodeLifetime -> Strong references should decrease ref count");
+		}
 
+		// Validate parent support
 		{
 			FDecoratorBinding RootBinding;									// Empty, no parent
 			FAnimNextDecoratorHandle DecoratorHandle00(NodeHandles[0], 0);	// Point to first node, first base decorator
+			FAnimNextDecoratorHandle DecoratorHandle03(NodeHandles[0], 3);	// Point to first node, second base decorator
+			FAnimNextDecoratorHandle DecoratorHandle10(NodeHandles[1], 0);	// Point to second node, first base decorator
 
-			// Allocate our node instance
+			// Allocate our first node
 			FDecoratorPtr DecoratorPtr00 = Context.AllocateNodeInstance(RootBinding, DecoratorHandle00);
 			AddErrorIfFalse(DecoratorPtr00.IsValid(), "FAnimationAnimNextRuntimeTest_NodeLifetime -> Failed to allocate a node instance");
 
-			// Validate instance constructors
-			AddErrorIfFalse(ConstructedDecorators.Num() == 5, "FAnimationAnimNextRuntimeTest_NodeLifetime -> Expected all 5 decorators to have been constructed");
-			AddErrorIfFalse(DestructedDecorators.Num() == 0, "FAnimationAnimNextRuntimeTest_NodeLifetime -> Expected no decorators to have been destructed");
-			AddErrorIfFalse(ConstructedDecorators[0] == NodeTemplateDecoratorList[0], "FAnimationAnimNextRuntimeTest_NodeLifetime -> Unexpected constructor order");
-			AddErrorIfFalse(ConstructedDecorators[1] == NodeTemplateDecoratorList[1], "FAnimationAnimNextRuntimeTest_NodeLifetime -> Unexpected constructor order");
-			AddErrorIfFalse(ConstructedDecorators[2] == NodeTemplateDecoratorList[2], "FAnimationAnimNextRuntimeTest_NodeLifetime -> Unexpected constructor order");
-			AddErrorIfFalse(ConstructedDecorators[3] == NodeTemplateDecoratorList[3], "FAnimationAnimNextRuntimeTest_NodeLifetime -> Unexpected constructor order");
-			AddErrorIfFalse(ConstructedDecorators[4] == NodeTemplateDecoratorList[4], "FAnimationAnimNextRuntimeTest_NodeLifetime -> Unexpected constructor order");
-
-			// Destruct our node instance
+			// Allocate a new node, using the first as a parent
+			// Both decorators live on the same node, the returned handle should be weak on the parent
+			FDecoratorPtr DecoratorPtr03 = Context.AllocateNodeInstance(DecoratorPtr00, DecoratorHandle03);
+			AddErrorIfFalse(DecoratorPtr03.IsValid(), "FAnimationAnimNextRuntimeTest_NodeLifetime -> Failed to allocate a node instance");
+			AddErrorIfFalse(DecoratorPtr03.GetDecoratorIndex() == 3, "FAnimationAnimNextRuntimeTest_NodeLifetime -> Allocated decorator pointer should point to fourth decorator");
+			AddErrorIfFalse(DecoratorPtr03.IsWeak(), "FAnimationAnimNextRuntimeTest_NodeLifetime -> Allocated decorator pointer should be weak, we have the same parent");
+			AddErrorIfFalse(DecoratorPtr03.GetNodeInstance()->GetNodeHandle() == NodeHandles[0], "FAnimationAnimNextRuntimeTest_NodeLifetime -> Allocated node should point to the provided node handle");
+			AddErrorIfFalse(DecoratorPtr03.GetNodeInstance() == DecoratorPtr00.GetNodeInstance(), "FAnimationAnimNextRuntimeTest_NodeLifetime -> Handles should point to the same node instance");
+			AddErrorIfFalse(DecoratorPtr03.GetNodeInstance()->GetReferenceCount() == 1, "FAnimationAnimNextRuntimeTest_NodeLifetime -> Allocated node should have one reference");
+
+			// Allocate a new node, using the first as a parent
+			// The second decorator lives on a new node, a new node instance will be allocated
+			FDecoratorPtr DecoratorPtr10 = Context.AllocateNodeInstance(DecoratorPtr00, DecoratorHandle10);
+			AddErrorIfFalse(DecoratorPtr10.IsValid(), "FAnimationAnimNextRuntimeTest_NodeLifetime -> Failed to allocate a node instance");
+			AddErrorIfFalse(DecoratorPtr10.GetDecoratorIndex() == 0, "FAnimationAnimNextRuntimeTest_NodeLifetime -> Allocated decorator pointer should point to first decorator");
+			AddErrorIfFalse(!DecoratorPtr10.IsWeak(), "FAnimationAnimNextRuntimeTest_NodeLifetime -> Allocated decorator pointer should not be weak, we have the same parent but a different node handle");
+			AddErrorIfFalse(DecoratorPtr10.GetNodeInstance()->GetNodeHandle() == NodeHandles[1], "FAnimationAnimNextRuntimeTest_NodeLifetime -> Allocated node should point to the provided node handle");
+			AddErrorIfFalse(DecoratorPtr10.GetNodeInstance() != DecoratorPtr00.GetNodeInstance(), "FAnimationAnimNextRuntimeTest_NodeLifetime -> Handles should not point to the same node instance");
+			AddErrorIfFalse(DecoratorPtr10.GetNodeInstance()->GetReferenceCount() == 1, "FAnimationAnimNextRuntimeTest_NodeLifetime -> Allocated node should have one reference");
 		}
 
-		// Validate instance destructors
-		AddErrorIfFalse(ConstructedDecorators.Num() == 5, "FAnimationAnimNextRuntimeTest_NodeLifetime -> Expected no decorators to have been constructed");
-		AddErrorIfFalse(DestructedDecorators.Num() == 5, "FAnimationAnimNextRuntimeTest_NodeLifetime -> Expected all 5 decorators to have been destructed");
-		AddErrorIfFalse(DestructedDecorators[0] == NodeTemplateDecoratorList[4], "FAnimationAnimNextRuntimeTest_NodeLifetime -> Unexpected destructor order");
-		AddErrorIfFalse(DestructedDecorators[1] == NodeTemplateDecoratorList[3], "FAnimationAnimNextRuntimeTest_NodeLifetime -> Unexpected destructor order");
-		AddErrorIfFalse(DestructedDecorators[2] == NodeTemplateDecoratorList[2], "FAnimationAnimNextRuntimeTest_NodeLifetime -> Unexpected destructor order");
-		AddErrorIfFalse(DestructedDecorators[3] == NodeTemplateDecoratorList[1], "FAnimationAnimNextRuntimeTest_NodeLifetime -> Unexpected destructor order");
-		AddErrorIfFalse(DestructedDecorators[4] == NodeTemplateDecoratorList[0], "FAnimationAnimNextRuntimeTest_NodeLifetime -> Unexpected destructor order");
-
-		Private::ConstructedDecorators = nullptr;
-		Private::DestructedDecorators = nullptr;
-	}
+		// Validate constructors and destructors
+		{
+			TArray<FDecoratorUID> ConstructedDecorators;
+			TArray<FDecoratorUID> DestructedDecorators;
+
+			Private::ConstructedDecorators = &ConstructedDecorators;
+			Private::DestructedDecorators = &DestructedDecorators;
+
+			{
+				FDecoratorBinding RootBinding;									// Empty, no parent
+				FAnimNextDecoratorHandle DecoratorHandle00(NodeHandles[0], 0);	// Point to first node, first base decorator
+
+				// Allocate our node instance
+				FDecoratorPtr DecoratorPtr00 = Context.AllocateNodeInstance(RootBinding, DecoratorHandle00);
+				AddErrorIfFalse(DecoratorPtr00.IsValid(), "FAnimationAnimNextRuntimeTest_NodeLifetime -> Failed to allocate a node instance");
+
+				// Validate instance constructors
+				AddErrorIfFalse(ConstructedDecorators.Num() == 5, "FAnimationAnimNextRuntimeTest_NodeLifetime -> Expected all 5 decorators to have been constructed");
+				AddErrorIfFalse(DestructedDecorators.Num() == 0, "FAnimationAnimNextRuntimeTest_NodeLifetime -> Expected no decorators to have been destructed");
+				AddErrorIfFalse(ConstructedDecorators[0] == NodeTemplateDecoratorList[0], "FAnimationAnimNextRuntimeTest_NodeLifetime -> Unexpected constructor order");
+				AddErrorIfFalse(ConstructedDecorators[1] == NodeTemplateDecoratorList[1], "FAnimationAnimNextRuntimeTest_NodeLifetime -> Unexpected constructor order");
+				AddErrorIfFalse(ConstructedDecorators[2] == NodeTemplateDecoratorList[2], "FAnimationAnimNextRuntimeTest_NodeLifetime -> Unexpected constructor order");
+				AddErrorIfFalse(ConstructedDecorators[3] == NodeTemplateDecoratorList[3], "FAnimationAnimNextRuntimeTest_NodeLifetime -> Unexpected constructor order");
+				AddErrorIfFalse(ConstructedDecorators[4] == NodeTemplateDecoratorList[4], "FAnimationAnimNextRuntimeTest_NodeLifetime -> Unexpected constructor order");
+
+				// Destruct our node instance
+			}
 
-	// Unregister our templates
-	Registry.Unregister(NodeTemplate0);
+			// Validate instance destructors
+			AddErrorIfFalse(ConstructedDecorators.Num() == 5, "FAnimationAnimNextRuntimeTest_NodeLifetime -> Expected no decorators to have been constructed");
+			AddErrorIfFalse(DestructedDecorators.Num() == 5, "FAnimationAnimNextRuntimeTest_NodeLifetime -> Expected all 5 decorators to have been destructed");
+			AddErrorIfFalse(DestructedDecorators[0] == NodeTemplateDecoratorList[4], "FAnimationAnimNextRuntimeTest_NodeLifetime -> Unexpected destructor order");
+			AddErrorIfFalse(DestructedDecorators[1] == NodeTemplateDecoratorList[3], "FAnimationAnimNextRuntimeTest_NodeLifetime -> Unexpected destructor order");
+			AddErrorIfFalse(DestructedDecorators[2] == NodeTemplateDecoratorList[2], "FAnimationAnimNextRuntimeTest_NodeLifetime -> Unexpected destructor order");
+			AddErrorIfFalse(DestructedDecorators[3] == NodeTemplateDecoratorList[1], "FAnimationAnimNextRuntimeTest_NodeLifetime -> Unexpected destructor order");
+			AddErrorIfFalse(DestructedDecorators[4] == NodeTemplateDecoratorList[0], "FAnimationAnimNextRuntimeTest_NodeLifetime -> Unexpected destructor order");
+
+			Private::ConstructedDecorators = nullptr;
+			Private::DestructedDecorators = nullptr;
+		}
 
-	AddErrorIfFalse(Registry.GetNum() == 0, "FAnimationAnimNextRuntimeTest_NodeLifetime -> Registry should contain 0 templates");
+		// Unregister our templates
+		Registry.Unregister(NodeTemplate0);
 
+		AddErrorIfFalse(Registry.GetNum() == 0, "FAnimationAnimNextRuntimeTest_NodeLifetime -> Registry should contain 0 templates");
+	}
+	Tests::FUtils::CleanupAfterTests();
+	
 	return true;
 }
 
@@ -707,265 +710,269 @@ bool FAnimationAnimNextRuntimeTest_GetDecoratorInterface::RunTest(const FString&
 {
 	using namespace UE::AnimNext;
 
-	AUTO_REGISTER_ANIM_DECORATOR(FDecoratorA_Base)
-	AUTO_REGISTER_ANIM_DECORATOR(FDecoratorAB_Add)
-	AUTO_REGISTER_ANIM_DECORATOR(FDecoratorAC_Add)
+	{
+		AUTO_REGISTER_ANIM_DECORATOR(FDecoratorA_Base)
+		AUTO_REGISTER_ANIM_DECORATOR(FDecoratorAB_Add)
+		AUTO_REGISTER_ANIM_DECORATOR(FDecoratorAC_Add)
 
-	UFactory* GraphFactory = NewObject<UAnimNextGraphFactory>();
-	UAnimNextGraph* AnimNextGraph = CastChecked<UAnimNextGraph>(GraphFactory->FactoryCreateNew(UAnimNextGraph::StaticClass(), GetTransientPackage(), TEXT("TestAnimNextGraph"), RF_Transient, nullptr, nullptr, NAME_None));
-	UE_RETURN_ON_ERROR(AnimNextGraph != nullptr, "FAnimationAnimNextRuntimeTest_GetDecoratorInterface -> Failed to create animation graph");
+		UFactory* GraphFactory = NewObject<UAnimNextGraphFactory>();
+		UAnimNextGraph* AnimNextGraph = CastChecked<UAnimNextGraph>(GraphFactory->FactoryCreateNew(UAnimNextGraph::StaticClass(), GetTransientPackage(), TEXT("TestAnimNextGraph"), RF_Transient, nullptr, nullptr, NAME_None));
+		UE_RETURN_ON_ERROR(AnimNextGraph != nullptr, "FAnimationAnimNextRuntimeTest_GetDecoratorInterface -> Failed to create animation graph");
 
-	FScopedClearNodeTemplateRegistry ScopedClearNodeTemplateRegistry;
-	FNodeTemplateRegistry& Registry = FNodeTemplateRegistry::Get();
+		FScopedClearNodeTemplateRegistry ScopedClearNodeTemplateRegistry;
+		FNodeTemplateRegistry& Registry = FNodeTemplateRegistry::Get();
 
-	TArray<FDecoratorUID> NodeTemplateDecoratorList;
-	NodeTemplateDecoratorList.Add(FDecoratorA_Base::DecoratorUID);
-	NodeTemplateDecoratorList.Add(FDecoratorAB_Add::DecoratorUID);
-	NodeTemplateDecoratorList.Add(FDecoratorAC_Add::DecoratorUID);
-	NodeTemplateDecoratorList.Add(FDecoratorA_Base::DecoratorUID);
-	NodeTemplateDecoratorList.Add(FDecoratorAC_Add::DecoratorUID);
+		TArray<FDecoratorUID> NodeTemplateDecoratorList;
+		NodeTemplateDecoratorList.Add(FDecoratorA_Base::DecoratorUID);
+		NodeTemplateDecoratorList.Add(FDecoratorAB_Add::DecoratorUID);
+		NodeTemplateDecoratorList.Add(FDecoratorAC_Add::DecoratorUID);
+		NodeTemplateDecoratorList.Add(FDecoratorA_Base::DecoratorUID);
+		NodeTemplateDecoratorList.Add(FDecoratorAC_Add::DecoratorUID);
 
-	// Populate our node template registry
-	TArray<uint8> NodeTemplateBuffer0;
-	const FNodeTemplate* NodeTemplate0 = FNodeTemplateBuilder::BuildNodeTemplate(NodeTemplateDecoratorList, NodeTemplateBuffer0);
+		// Populate our node template registry
+		TArray<uint8> NodeTemplateBuffer0;
+		const FNodeTemplate* NodeTemplate0 = FNodeTemplateBuilder::BuildNodeTemplate(NodeTemplateDecoratorList, NodeTemplateBuffer0);
 
-	FNodeTemplateRegistryHandle TemplateHandle0 = Registry.FindOrAdd(NodeTemplate0);
-	AddErrorIfFalse(TemplateHandle0.IsValid(), "FAnimationAnimNextRuntimeTest_GetDecoratorInterface -> Registry should contain our template");
+		FNodeTemplateRegistryHandle TemplateHandle0 = Registry.FindOrAdd(NodeTemplate0);
+		AddErrorIfFalse(TemplateHandle0.IsValid(), "FAnimationAnimNextRuntimeTest_GetDecoratorInterface -> Registry should contain our template");
 
-	TArray<FNodeHandle> NodeHandles;
+		TArray<FNodeHandle> NodeHandles;
 
-	// Write our graph
-	TArray<uint8> GraphSharedDataArchiveBuffer;
-	TArray<TObjectPtr<UObject>> GraphReferencedObjects;
-	{
-		FDecoratorWriter DecoratorWriter;
+		// Write our graph
+		TArray<uint8> GraphSharedDataArchiveBuffer;
+		TArray<TObjectPtr<UObject>> GraphReferencedObjects;
+		{
+			FDecoratorWriter DecoratorWriter;
 
-		NodeHandles.Add(DecoratorWriter.RegisterNode(*NodeTemplate0));
+			NodeHandles.Add(DecoratorWriter.RegisterNode(*NodeTemplate0));
 
-		// We don't have decorator properties
+			// We don't have decorator properties
 
-		DecoratorWriter.BeginNodeWriting();
-		DecoratorWriter.WriteNode(NodeHandles[0],
-			[](uint32 DecoratorIndex, const FString& PropertyName)
-			{
-				return FString();
-			},
-			[](uint32 DecoratorIndex, const FString& PropertyName)
-			{
-				return false;
-			});
-		DecoratorWriter.EndNodeWriting();
+			DecoratorWriter.BeginNodeWriting();
+			DecoratorWriter.WriteNode(NodeHandles[0],
+				[](uint32 DecoratorIndex, const FString& PropertyName)
+				{
+					return FString();
+				},
+				[](uint32 DecoratorIndex, const FString& PropertyName)
+				{
+					return false;
+				});
+			DecoratorWriter.EndNodeWriting();
 
-		AddErrorIfFalse(DecoratorWriter.GetErrorState() == FDecoratorWriter::EErrorState::None, "FAnimationAnimNextRuntimeTest_GetDecoratorInterface -> Failed to write decorators");
-		GraphSharedDataArchiveBuffer = DecoratorWriter.GetGraphSharedData();
-		GraphReferencedObjects = DecoratorWriter.GetGraphReferencedObjects();
-	}
+			AddErrorIfFalse(DecoratorWriter.GetErrorState() == FDecoratorWriter::EErrorState::None, "FAnimationAnimNextRuntimeTest_GetDecoratorInterface -> Failed to write decorators");
+			GraphSharedDataArchiveBuffer = DecoratorWriter.GetGraphSharedData();
+			GraphReferencedObjects = DecoratorWriter.GetGraphReferencedObjects();
+		}
 
-	// Read our graph
-	FTestUtils::LoadFromArchiveBuffer(*AnimNextGraph, NodeHandles, GraphSharedDataArchiveBuffer);
+		// Read our graph
+		FTestUtils::LoadFromArchiveBuffer(*AnimNextGraph, NodeHandles, GraphSharedDataArchiveBuffer);
 
-	FAnimNextGraphInstancePtr GraphInstance;
-	AnimNextGraph->AllocateInstance(GraphInstance);
+		FAnimNextGraphInstancePtr GraphInstance;
+		AnimNextGraph->AllocateInstance(GraphInstance);
 
-	FExecutionContext Context(GraphInstance);
+		FExecutionContext Context(GraphInstance);
 
-	// Validate from the first base decorator
-	{
-		FDecoratorBinding ParentBinding;								// Empty, no parent
-		FAnimNextDecoratorHandle DecoratorHandle00(NodeHandles[0], 0);	// Point to first node, first base decorator
-
-		FDecoratorPtr DecoratorPtr00 = Context.AllocateNodeInstance(ParentBinding, DecoratorHandle00);
-		AddErrorIfFalse(DecoratorPtr00.IsValid(), "FAnimationAnimNextRuntimeTest_GetDecoratorInterface -> Failed to allocate a node instance");
-
-		// Validate GetInterface from a decorator handle
-		TDecoratorBinding<IInterfaceC> Binding00C;
-		AddErrorIfFalse(Context.GetInterface(DecoratorPtr00, Binding00C), "FAnimationAnimNextRuntimeTest_GetDecoratorInterface -> InterfaceC not found");
-		AddErrorIfFalse(Binding00C.IsValid(), "FAnimationAnimNextRuntimeTest_GetDecoratorInterface -> InterfaceC binding not valid");
-		AddErrorIfFalse(Binding00C.GetInterfaceUID() == IInterfaceC::InterfaceUID, "FAnimationAnimNextRuntimeTest_GetDecoratorInterface -> Unexpected interface UID found in decorator binding");
-		AddErrorIfFalse(Binding00C.GetDecoratorPtr().GetDecoratorIndex() == 2, "FAnimationAnimNextRuntimeTest_GetDecoratorInterface -> InterfaceC not found on expected decorator");
-		AddErrorIfFalse(Binding00C.GetDecoratorPtr().GetNodeInstance()->GetNodeHandle() == NodeHandles[0], "FAnimationAnimNextRuntimeTest_GetDecoratorInterface -> InterfaceC not found on expected node");
-		AddErrorIfFalse(Binding00C.GetSharedData<FDecoratorAC_Add::FSharedData>()->DecoratorUID == FDecoratorAC_Add::DecoratorUID, "FAnimationAnimNextRuntimeTest_GetDecoratorInterface -> Unexpected shared data in decorator binding");
-		AddErrorIfFalse(Binding00C.GetInstanceData<FDecoratorAC_Add::FInstanceData>()->DecoratorUID == FDecoratorAC_Add::DecoratorUID, "FAnimationAnimNextRuntimeTest_GetDecoratorInterface -> Unexpected instance data in decorator binding");
-
-		TDecoratorBinding<IInterfaceB> Binding00B;
-		AddErrorIfFalse(Context.GetInterface(DecoratorPtr00, Binding00B), "FAnimationAnimNextRuntimeTest_GetDecoratorInterface -> InterfaceB not found");
-		AddErrorIfFalse(Binding00B.IsValid(), "FAnimationAnimNextRuntimeTest_GetDecoratorInterface -> InterfaceB binding not valid");
-		AddErrorIfFalse(Binding00B.GetInterfaceUID() == IInterfaceB::InterfaceUID, "FAnimationAnimNextRuntimeTest_GetDecoratorInterface -> Unexpected interface UID found in decorator binding");
-		AddErrorIfFalse(Binding00B.GetDecoratorPtr().GetDecoratorIndex() == 1, "FAnimationAnimNextRuntimeTest_GetDecoratorInterface -> InterfaceB not found on expected decorator");
-		AddErrorIfFalse(Binding00B.GetDecoratorPtr().GetNodeInstance()->GetNodeHandle() == NodeHandles[0], "FAnimationAnimNextRuntimeTest_GetDecoratorInterface -> InterfaceB not found on expected node");
-		AddErrorIfFalse(Binding00B.GetSharedData<FDecoratorAB_Add::FSharedData>()->DecoratorUID == FDecoratorAB_Add::DecoratorUID, "FAnimationAnimNextRuntimeTest_GetDecoratorInterface -> Unexpected shared data in decorator binding");
-		AddErrorIfFalse(Binding00B.GetInstanceData<FDecoratorAB_Add::FInstanceData>()->DecoratorUID == FDecoratorAB_Add::DecoratorUID, "FAnimationAnimNextRuntimeTest_GetDecoratorInterface -> Unexpected instance data in decorator binding");
-
-		TDecoratorBinding<IInterfaceA> Binding00A;
-		AddErrorIfFalse(Context.GetInterface(DecoratorPtr00, Binding00A), "FAnimationAnimNextRuntimeTest_GetDecoratorInterface -> InterfaceA not found");
-		AddErrorIfFalse(Binding00A.IsValid(), "FAnimationAnimNextRuntimeTest_GetDecoratorInterface -> InterfaceA binding not valid");
-		AddErrorIfFalse(Binding00A.GetInterfaceUID() == IInterfaceA::InterfaceUID, "FAnimationAnimNextRuntimeTest_GetDecoratorInterface -> Unexpected interface UID found in decorator binding");
-		AddErrorIfFalse(Binding00A.GetDecoratorPtr().GetDecoratorIndex() == 2, "FAnimationAnimNextRuntimeTest_GetDecoratorInterface -> InterfaceA not found on expected decorator");
-		AddErrorIfFalse(Binding00A.GetDecoratorPtr().GetNodeInstance()->GetNodeHandle() == NodeHandles[0], "FAnimationAnimNextRuntimeTest_GetDecoratorInterface -> InterfaceA not found on expected node");
-		AddErrorIfFalse(Binding00A.GetSharedData<FDecoratorAC_Add::FSharedData>()->DecoratorUID == FDecoratorAC_Add::DecoratorUID, "FAnimationAnimNextRuntimeTest_GetDecoratorInterface -> Unexpected shared data in decorator binding");
-		AddErrorIfFalse(Binding00A.GetInstanceData<FDecoratorAC_Add::FInstanceData>()->DecoratorUID == FDecoratorAC_Add::DecoratorUID, "FAnimationAnimNextRuntimeTest_GetDecoratorInterface -> Unexpected instance data in decorator binding");
-
-		// Validate GetInterface from a decorator binding
+		// Validate from the first base decorator
 		{
+			FDecoratorBinding ParentBinding;								// Empty, no parent
+			FAnimNextDecoratorHandle DecoratorHandle00(NodeHandles[0], 0);	// Point to first node, first base decorator
+
+			FDecoratorPtr DecoratorPtr00 = Context.AllocateNodeInstance(ParentBinding, DecoratorHandle00);
+			AddErrorIfFalse(DecoratorPtr00.IsValid(), "FAnimationAnimNextRuntimeTest_GetDecoratorInterface -> Failed to allocate a node instance");
+
+			// Validate GetInterface from a decorator handle
+			TDecoratorBinding<IInterfaceC> Binding00C;
+			AddErrorIfFalse(Context.GetInterface(DecoratorPtr00, Binding00C), "FAnimationAnimNextRuntimeTest_GetDecoratorInterface -> InterfaceC not found");
+			AddErrorIfFalse(Binding00C.IsValid(), "FAnimationAnimNextRuntimeTest_GetDecoratorInterface -> InterfaceC binding not valid");
+			AddErrorIfFalse(Binding00C.GetInterfaceUID() == IInterfaceC::InterfaceUID, "FAnimationAnimNextRuntimeTest_GetDecoratorInterface -> Unexpected interface UID found in decorator binding");
+			AddErrorIfFalse(Binding00C.GetDecoratorPtr().GetDecoratorIndex() == 2, "FAnimationAnimNextRuntimeTest_GetDecoratorInterface -> InterfaceC not found on expected decorator");
+			AddErrorIfFalse(Binding00C.GetDecoratorPtr().GetNodeInstance()->GetNodeHandle() == NodeHandles[0], "FAnimationAnimNextRuntimeTest_GetDecoratorInterface -> InterfaceC not found on expected node");
+			AddErrorIfFalse(Binding00C.GetSharedData<FDecoratorAC_Add::FSharedData>()->DecoratorUID == FDecoratorAC_Add::DecoratorUID, "FAnimationAnimNextRuntimeTest_GetDecoratorInterface -> Unexpected shared data in decorator binding");
+			AddErrorIfFalse(Binding00C.GetInstanceData<FDecoratorAC_Add::FInstanceData>()->DecoratorUID == FDecoratorAC_Add::DecoratorUID, "FAnimationAnimNextRuntimeTest_GetDecoratorInterface -> Unexpected instance data in decorator binding");
+
+			TDecoratorBinding<IInterfaceB> Binding00B;
+			AddErrorIfFalse(Context.GetInterface(DecoratorPtr00, Binding00B), "FAnimationAnimNextRuntimeTest_GetDecoratorInterface -> InterfaceB not found");
+			AddErrorIfFalse(Binding00B.IsValid(), "FAnimationAnimNextRuntimeTest_GetDecoratorInterface -> InterfaceB binding not valid");
+			AddErrorIfFalse(Binding00B.GetInterfaceUID() == IInterfaceB::InterfaceUID, "FAnimationAnimNextRuntimeTest_GetDecoratorInterface -> Unexpected interface UID found in decorator binding");
+			AddErrorIfFalse(Binding00B.GetDecoratorPtr().GetDecoratorIndex() == 1, "FAnimationAnimNextRuntimeTest_GetDecoratorInterface -> InterfaceB not found on expected decorator");
+			AddErrorIfFalse(Binding00B.GetDecoratorPtr().GetNodeInstance()->GetNodeHandle() == NodeHandles[0], "FAnimationAnimNextRuntimeTest_GetDecoratorInterface -> InterfaceB not found on expected node");
+			AddErrorIfFalse(Binding00B.GetSharedData<FDecoratorAB_Add::FSharedData>()->DecoratorUID == FDecoratorAB_Add::DecoratorUID, "FAnimationAnimNextRuntimeTest_GetDecoratorInterface -> Unexpected shared data in decorator binding");
+			AddErrorIfFalse(Binding00B.GetInstanceData<FDecoratorAB_Add::FInstanceData>()->DecoratorUID == FDecoratorAB_Add::DecoratorUID, "FAnimationAnimNextRuntimeTest_GetDecoratorInterface -> Unexpected instance data in decorator binding");
+
+			TDecoratorBinding<IInterfaceA> Binding00A;
+			AddErrorIfFalse(Context.GetInterface(DecoratorPtr00, Binding00A), "FAnimationAnimNextRuntimeTest_GetDecoratorInterface -> InterfaceA not found");
+			AddErrorIfFalse(Binding00A.IsValid(), "FAnimationAnimNextRuntimeTest_GetDecoratorInterface -> InterfaceA binding not valid");
+			AddErrorIfFalse(Binding00A.GetInterfaceUID() == IInterfaceA::InterfaceUID, "FAnimationAnimNextRuntimeTest_GetDecoratorInterface -> Unexpected interface UID found in decorator binding");
+			AddErrorIfFalse(Binding00A.GetDecoratorPtr().GetDecoratorIndex() == 2, "FAnimationAnimNextRuntimeTest_GetDecoratorInterface -> InterfaceA not found on expected decorator");
+			AddErrorIfFalse(Binding00A.GetDecoratorPtr().GetNodeInstance()->GetNodeHandle() == NodeHandles[0], "FAnimationAnimNextRuntimeTest_GetDecoratorInterface -> InterfaceA not found on expected node");
+			AddErrorIfFalse(Binding00A.GetSharedData<FDecoratorAC_Add::FSharedData>()->DecoratorUID == FDecoratorAC_Add::DecoratorUID, "FAnimationAnimNextRuntimeTest_GetDecoratorInterface -> Unexpected shared data in decorator binding");
+			AddErrorIfFalse(Binding00A.GetInstanceData<FDecoratorAC_Add::FInstanceData>()->DecoratorUID == FDecoratorAC_Add::DecoratorUID, "FAnimationAnimNextRuntimeTest_GetDecoratorInterface -> Unexpected instance data in decorator binding");
+
+			// Validate GetInterface from a decorator binding
 			{
 				{
-					TDecoratorBinding<IInterfaceC> Binding00C_;
-					AddErrorIfFalse(Context.GetInterface(Binding00C, Binding00C_), "FAnimationAnimNextRuntimeTest_GetDecoratorInterface -> InterfaceC not found");
-					AddErrorIfFalse(Binding00C == Binding00C_, "FAnimationAnimNextRuntimeTest_GetDecoratorInterface -> GetInterface methods should return the same result");
-				}
+					{
+						TDecoratorBinding<IInterfaceC> Binding00C_;
+						AddErrorIfFalse(Context.GetInterface(Binding00C, Binding00C_), "FAnimationAnimNextRuntimeTest_GetDecoratorInterface -> InterfaceC not found");
+						AddErrorIfFalse(Binding00C == Binding00C_, "FAnimationAnimNextRuntimeTest_GetDecoratorInterface -> GetInterface methods should return the same result");
+					}
 
-				{
-					TDecoratorBinding<IInterfaceC> Binding00C_;
-					AddErrorIfFalse(Context.GetInterface(Binding00B, Binding00C_), "FAnimationAnimNextRuntimeTest_GetDecoratorInterface -> InterfaceC not found");
-					AddErrorIfFalse(Binding00C == Binding00C_, "FAnimationAnimNextRuntimeTest_GetDecoratorInterface -> GetInterface methods should return the same result");
-				}
+					{
+						TDecoratorBinding<IInterfaceC> Binding00C_;
+						AddErrorIfFalse(Context.GetInterface(Binding00B, Binding00C_), "FAnimationAnimNextRuntimeTest_GetDecoratorInterface -> InterfaceC not found");
+						AddErrorIfFalse(Binding00C == Binding00C_, "FAnimationAnimNextRuntimeTest_GetDecoratorInterface -> GetInterface methods should return the same result");
+					}
 
-				{
-					TDecoratorBinding<IInterfaceC> Binding00C_;
-					AddErrorIfFalse(Context.GetInterface(Binding00A, Binding00C_), "FAnimationAnimNextRuntimeTest_GetDecoratorInterface -> InterfaceC not found");
-					AddErrorIfFalse(Binding00C == Binding00C_, "FAnimationAnimNextRuntimeTest_GetDecoratorInterface -> GetInterface methods should return the same result");
+					{
+						TDecoratorBinding<IInterfaceC> Binding00C_;
+						AddErrorIfFalse(Context.GetInterface(Binding00A, Binding00C_), "FAnimationAnimNextRuntimeTest_GetDecoratorInterface -> InterfaceC not found");
+						AddErrorIfFalse(Binding00C == Binding00C_, "FAnimationAnimNextRuntimeTest_GetDecoratorInterface -> GetInterface methods should return the same result");
+					}
 				}
-			}
 
-			{
 				{
-					TDecoratorBinding<IInterfaceB> Binding00B_;
-					AddErrorIfFalse(Context.GetInterface(Binding00C, Binding00B_), "FAnimationAnimNextRuntimeTest_GetDecoratorInterface -> InterfaceB not found");
-					AddErrorIfFalse(Binding00B == Binding00B_, "FAnimationAnimNextRuntimeTest_GetDecoratorInterface -> GetInterface methods should return the same result");
-				}
+					{
+						TDecoratorBinding<IInterfaceB> Binding00B_;
+						AddErrorIfFalse(Context.GetInterface(Binding00C, Binding00B_), "FAnimationAnimNextRuntimeTest_GetDecoratorInterface -> InterfaceB not found");
+						AddErrorIfFalse(Binding00B == Binding00B_, "FAnimationAnimNextRuntimeTest_GetDecoratorInterface -> GetInterface methods should return the same result");
+					}
 
-				{
-					TDecoratorBinding<IInterfaceB> Binding00B_;
-					AddErrorIfFalse(Context.GetInterface(Binding00B, Binding00B_), "FAnimationAnimNextRuntimeTest_GetDecoratorInterface -> InterfaceB not found");
-					AddErrorIfFalse(Binding00B == Binding00B_, "FAnimationAnimNextRuntimeTest_GetDecoratorInterface -> GetInterface methods should return the same result");
-				}
+					{
+						TDecoratorBinding<IInterfaceB> Binding00B_;
+						AddErrorIfFalse(Context.GetInterface(Binding00B, Binding00B_), "FAnimationAnimNextRuntimeTest_GetDecoratorInterface -> InterfaceB not found");
+						AddErrorIfFalse(Binding00B == Binding00B_, "FAnimationAnimNextRuntimeTest_GetDecoratorInterface -> GetInterface methods should return the same result");
+					}
 
-				{
-					TDecoratorBinding<IInterfaceB> Binding00B_;
-					AddErrorIfFalse(Context.GetInterface(Binding00A, Binding00B_), "FAnimationAnimNextRuntimeTest_GetDecoratorInterface -> InterfaceB not found");
-					AddErrorIfFalse(Binding00B == Binding00B_, "FAnimationAnimNextRuntimeTest_GetDecoratorInterface -> GetInterface methods should return the same result");
+					{
+						TDecoratorBinding<IInterfaceB> Binding00B_;
+						AddErrorIfFalse(Context.GetInterface(Binding00A, Binding00B_), "FAnimationAnimNextRuntimeTest_GetDecoratorInterface -> InterfaceB not found");
+						AddErrorIfFalse(Binding00B == Binding00B_, "FAnimationAnimNextRuntimeTest_GetDecoratorInterface -> GetInterface methods should return the same result");
+					}
 				}
-			}
 
-			{
 				{
-					TDecoratorBinding<IInterfaceA> Binding00A_;
-					AddErrorIfFalse(Context.GetInterface(Binding00C, Binding00A_), "FAnimationAnimNextRuntimeTest_GetDecoratorInterface -> InterfaceA not found");
-					AddErrorIfFalse(Binding00A == Binding00A_, "FAnimationAnimNextRuntimeTest_GetDecoratorInterface -> GetInterface methods should return the same result");
-				}
+					{
+						TDecoratorBinding<IInterfaceA> Binding00A_;
+						AddErrorIfFalse(Context.GetInterface(Binding00C, Binding00A_), "FAnimationAnimNextRuntimeTest_GetDecoratorInterface -> InterfaceA not found");
+						AddErrorIfFalse(Binding00A == Binding00A_, "FAnimationAnimNextRuntimeTest_GetDecoratorInterface -> GetInterface methods should return the same result");
+					}
 
-				{
-					TDecoratorBinding<IInterfaceA> Binding00A_;
-					AddErrorIfFalse(Context.GetInterface(Binding00B, Binding00A_), "FAnimationAnimNextRuntimeTest_GetDecoratorInterface -> InterfaceA not found");
-					AddErrorIfFalse(Binding00A == Binding00A_, "FAnimationAnimNextRuntimeTest_GetDecoratorInterface -> GetInterface methods should return the same result");
-				}
+					{
+						TDecoratorBinding<IInterfaceA> Binding00A_;
+						AddErrorIfFalse(Context.GetInterface(Binding00B, Binding00A_), "FAnimationAnimNextRuntimeTest_GetDecoratorInterface -> InterfaceA not found");
+						AddErrorIfFalse(Binding00A == Binding00A_, "FAnimationAnimNextRuntimeTest_GetDecoratorInterface -> GetInterface methods should return the same result");
+					}
 
-				{
-					TDecoratorBinding<IInterfaceA> Binding00A_;
-					AddErrorIfFalse(Context.GetInterface(Binding00A, Binding00A_), "FAnimationAnimNextRuntimeTest_GetDecoratorInterface -> InterfaceA not found");
-					AddErrorIfFalse(Binding00A == Binding00A_, "FAnimationAnimNextRuntimeTest_GetDecoratorInterface -> GetInterface methods should return the same result");
+					{
+						TDecoratorBinding<IInterfaceA> Binding00A_;
+						AddErrorIfFalse(Context.GetInterface(Binding00A, Binding00A_), "FAnimationAnimNextRuntimeTest_GetDecoratorInterface -> InterfaceA not found");
+						AddErrorIfFalse(Binding00A == Binding00A_, "FAnimationAnimNextRuntimeTest_GetDecoratorInterface -> GetInterface methods should return the same result");
+					}
 				}
 			}
 		}
-	}
 
-	// Validate from the second base decorator
-	{
-		FDecoratorBinding ParentBinding;								// Empty, no parent
-		FAnimNextDecoratorHandle DecoratorHandle03(NodeHandles[0], 3);	// Point to first node, second base decorator
-
-		FDecoratorPtr DecoratorPtr03 = Context.AllocateNodeInstance(ParentBinding, DecoratorHandle03);
-		AddErrorIfFalse(DecoratorPtr03.IsValid(), "FAnimationAnimNextRuntimeTest_GetDecoratorInterface -> Failed to allocate a node instance");
-
-		// Validate GetInterface from a decorator handle
-		TDecoratorBinding<IInterfaceC> Binding03C;
-		AddErrorIfFalse(Context.GetInterface(DecoratorPtr03, Binding03C), "FAnimationAnimNextRuntimeTest_GetDecoratorInterface -> InterfaceC not found");
-		AddErrorIfFalse(Binding03C.IsValid(), "FAnimationAnimNextRuntimeTest_GetDecoratorInterface -> InterfaceC binding not valid");
-		AddErrorIfFalse(Binding03C.GetInterfaceUID() == IInterfaceC::InterfaceUID, "FAnimationAnimNextRuntimeTest_GetDecoratorInterface -> Unexpected interface UID found in decorator binding");
-		AddErrorIfFalse(Binding03C.GetDecoratorPtr().GetDecoratorIndex() == 4, "FAnimationAnimNextRuntimeTest_GetDecoratorInterface -> InterfaceC not found on expected decorator");
-		AddErrorIfFalse(Binding03C.GetDecoratorPtr().GetNodeInstance()->GetNodeHandle() == NodeHandles[0], "FAnimationAnimNextRuntimeTest_GetDecoratorInterface -> InterfaceC not found on expected node");
-		AddErrorIfFalse(Binding03C.GetSharedData<FDecoratorAC_Add::FSharedData>()->DecoratorUID == FDecoratorAC_Add::DecoratorUID, "FAnimationAnimNextRuntimeTest_GetDecoratorInterface -> Unexpected shared data in decorator binding");
-		AddErrorIfFalse(Binding03C.GetInstanceData<FDecoratorAC_Add::FInstanceData>()->DecoratorUID == FDecoratorAC_Add::DecoratorUID, "FAnimationAnimNextRuntimeTest_GetDecoratorInterface -> Unexpected instance data in decorator binding");
-
-		TDecoratorBinding<IInterfaceB> Binding03B;
-		AddErrorIfFalse(!Context.GetInterface(DecoratorPtr03, Binding03B), "FAnimationAnimNextRuntimeTest_GetDecoratorInterface -> InterfaceB found");
-		AddErrorIfFalse(!Binding03B.IsValid(), "FAnimationAnimNextRuntimeTest_GetDecoratorInterface -> InterfaceB binding should not be valid");
-
-		TDecoratorBinding<IInterfaceA> Binding03A;
-		AddErrorIfFalse(Context.GetInterface(DecoratorPtr03, Binding03A), "FAnimationAnimNextRuntimeTest_GetDecoratorInterface -> InterfaceA not found");
-		AddErrorIfFalse(Binding03A.IsValid(), "FAnimationAnimNextRuntimeTest_GetDecoratorInterface -> InterfaceA binding not valid");
-		AddErrorIfFalse(Binding03A.GetInterfaceUID() == IInterfaceA::InterfaceUID, "FAnimationAnimNextRuntimeTest_GetDecoratorInterface -> Unexpected interface UID found in decorator binding");
-		AddErrorIfFalse(Binding03A.GetDecoratorPtr().GetDecoratorIndex() == 4, "FAnimationAnimNextRuntimeTest_GetDecoratorInterface -> InterfaceA not found on expected decorator");
-		AddErrorIfFalse(Binding03A.GetDecoratorPtr().GetNodeInstance()->GetNodeHandle() == NodeHandles[0], "FAnimationAnimNextRuntimeTest_GetDecoratorInterface -> InterfaceA not found on expected node");
-		AddErrorIfFalse(Binding03A.GetSharedData<FDecoratorAC_Add::FSharedData>()->DecoratorUID == FDecoratorAC_Add::DecoratorUID, "FAnimationAnimNextRuntimeTest_GetDecoratorInterface -> Unexpected shared data in decorator binding");
-		AddErrorIfFalse(Binding03A.GetInstanceData<FDecoratorAC_Add::FInstanceData>()->DecoratorUID == FDecoratorAC_Add::DecoratorUID, "FAnimationAnimNextRuntimeTest_GetDecoratorInterface -> Unexpected instance data in decorator binding");
-
-		// Validate GetInterface from a decorator binding
+		// Validate from the second base decorator
 		{
+			FDecoratorBinding ParentBinding;								// Empty, no parent
+			FAnimNextDecoratorHandle DecoratorHandle03(NodeHandles[0], 3);	// Point to first node, second base decorator
+
+			FDecoratorPtr DecoratorPtr03 = Context.AllocateNodeInstance(ParentBinding, DecoratorHandle03);
+			AddErrorIfFalse(DecoratorPtr03.IsValid(), "FAnimationAnimNextRuntimeTest_GetDecoratorInterface -> Failed to allocate a node instance");
+
+			// Validate GetInterface from a decorator handle
+			TDecoratorBinding<IInterfaceC> Binding03C;
+			AddErrorIfFalse(Context.GetInterface(DecoratorPtr03, Binding03C), "FAnimationAnimNextRuntimeTest_GetDecoratorInterface -> InterfaceC not found");
+			AddErrorIfFalse(Binding03C.IsValid(), "FAnimationAnimNextRuntimeTest_GetDecoratorInterface -> InterfaceC binding not valid");
+			AddErrorIfFalse(Binding03C.GetInterfaceUID() == IInterfaceC::InterfaceUID, "FAnimationAnimNextRuntimeTest_GetDecoratorInterface -> Unexpected interface UID found in decorator binding");
+			AddErrorIfFalse(Binding03C.GetDecoratorPtr().GetDecoratorIndex() == 4, "FAnimationAnimNextRuntimeTest_GetDecoratorInterface -> InterfaceC not found on expected decorator");
+			AddErrorIfFalse(Binding03C.GetDecoratorPtr().GetNodeInstance()->GetNodeHandle() == NodeHandles[0], "FAnimationAnimNextRuntimeTest_GetDecoratorInterface -> InterfaceC not found on expected node");
+			AddErrorIfFalse(Binding03C.GetSharedData<FDecoratorAC_Add::FSharedData>()->DecoratorUID == FDecoratorAC_Add::DecoratorUID, "FAnimationAnimNextRuntimeTest_GetDecoratorInterface -> Unexpected shared data in decorator binding");
+			AddErrorIfFalse(Binding03C.GetInstanceData<FDecoratorAC_Add::FInstanceData>()->DecoratorUID == FDecoratorAC_Add::DecoratorUID, "FAnimationAnimNextRuntimeTest_GetDecoratorInterface -> Unexpected instance data in decorator binding");
+
+			TDecoratorBinding<IInterfaceB> Binding03B;
+			AddErrorIfFalse(!Context.GetInterface(DecoratorPtr03, Binding03B), "FAnimationAnimNextRuntimeTest_GetDecoratorInterface -> InterfaceB found");
+			AddErrorIfFalse(!Binding03B.IsValid(), "FAnimationAnimNextRuntimeTest_GetDecoratorInterface -> InterfaceB binding should not be valid");
+
+			TDecoratorBinding<IInterfaceA> Binding03A;
+			AddErrorIfFalse(Context.GetInterface(DecoratorPtr03, Binding03A), "FAnimationAnimNextRuntimeTest_GetDecoratorInterface -> InterfaceA not found");
+			AddErrorIfFalse(Binding03A.IsValid(), "FAnimationAnimNextRuntimeTest_GetDecoratorInterface -> InterfaceA binding not valid");
+			AddErrorIfFalse(Binding03A.GetInterfaceUID() == IInterfaceA::InterfaceUID, "FAnimationAnimNextRuntimeTest_GetDecoratorInterface -> Unexpected interface UID found in decorator binding");
+			AddErrorIfFalse(Binding03A.GetDecoratorPtr().GetDecoratorIndex() == 4, "FAnimationAnimNextRuntimeTest_GetDecoratorInterface -> InterfaceA not found on expected decorator");
+			AddErrorIfFalse(Binding03A.GetDecoratorPtr().GetNodeInstance()->GetNodeHandle() == NodeHandles[0], "FAnimationAnimNextRuntimeTest_GetDecoratorInterface -> InterfaceA not found on expected node");
+			AddErrorIfFalse(Binding03A.GetSharedData<FDecoratorAC_Add::FSharedData>()->DecoratorUID == FDecoratorAC_Add::DecoratorUID, "FAnimationAnimNextRuntimeTest_GetDecoratorInterface -> Unexpected shared data in decorator binding");
+			AddErrorIfFalse(Binding03A.GetInstanceData<FDecoratorAC_Add::FInstanceData>()->DecoratorUID == FDecoratorAC_Add::DecoratorUID, "FAnimationAnimNextRuntimeTest_GetDecoratorInterface -> Unexpected instance data in decorator binding");
+
+			// Validate GetInterface from a decorator binding
 			{
 				{
-					TDecoratorBinding<IInterfaceC> Binding03C_;
-					AddErrorIfFalse(Context.GetInterface(Binding03C, Binding03C_), "FAnimationAnimNextRuntimeTest_GetDecoratorInterface -> InterfaceC not found");
-					AddErrorIfFalse(Binding03C == Binding03C_, "FAnimationAnimNextRuntimeTest_GetDecoratorInterface -> GetInterface methods should return the same result");
-				}
+					{
+						TDecoratorBinding<IInterfaceC> Binding03C_;
+						AddErrorIfFalse(Context.GetInterface(Binding03C, Binding03C_), "FAnimationAnimNextRuntimeTest_GetDecoratorInterface -> InterfaceC not found");
+						AddErrorIfFalse(Binding03C == Binding03C_, "FAnimationAnimNextRuntimeTest_GetDecoratorInterface -> GetInterface methods should return the same result");
+					}
 
-				{
-					TDecoratorBinding<IInterfaceC> Binding03C_;
-					AddErrorIfFalse(!Context.GetInterface(Binding03B, Binding03C_), "FAnimationAnimNextRuntimeTest_GetDecoratorInterface -> InterfaceC found");
-				}
+					{
+						TDecoratorBinding<IInterfaceC> Binding03C_;
+						AddErrorIfFalse(!Context.GetInterface(Binding03B, Binding03C_), "FAnimationAnimNextRuntimeTest_GetDecoratorInterface -> InterfaceC found");
+					}
 
-				{
-					TDecoratorBinding<IInterfaceC> Binding03C_;
-					AddErrorIfFalse(Context.GetInterface(Binding03A, Binding03C_), "FAnimationAnimNextRuntimeTest_GetDecoratorInterface -> InterfaceC not found");
-					AddErrorIfFalse(Binding03C == Binding03C_, "FAnimationAnimNextRuntimeTest_GetDecoratorInterface -> GetInterface methods should return the same result");
+					{
+						TDecoratorBinding<IInterfaceC> Binding03C_;
+						AddErrorIfFalse(Context.GetInterface(Binding03A, Binding03C_), "FAnimationAnimNextRuntimeTest_GetDecoratorInterface -> InterfaceC not found");
+						AddErrorIfFalse(Binding03C == Binding03C_, "FAnimationAnimNextRuntimeTest_GetDecoratorInterface -> GetInterface methods should return the same result");
+					}
 				}
-			}
 
-			{
 				{
-					TDecoratorBinding<IInterfaceB> Binding03B_;
-					AddErrorIfFalse(!Context.GetInterface(Binding03C, Binding03B_), "FAnimationAnimNextRuntimeTest_GetDecoratorInterface -> InterfaceB found");
-					AddErrorIfFalse(Binding03B == Binding03B_, "FAnimationAnimNextRuntimeTest_GetDecoratorInterface -> GetInterface methods should return the same result");
-				}
+					{
+						TDecoratorBinding<IInterfaceB> Binding03B_;
+						AddErrorIfFalse(!Context.GetInterface(Binding03C, Binding03B_), "FAnimationAnimNextRuntimeTest_GetDecoratorInterface -> InterfaceB found");
+						AddErrorIfFalse(Binding03B == Binding03B_, "FAnimationAnimNextRuntimeTest_GetDecoratorInterface -> GetInterface methods should return the same result");
+					}
 
-				{
-					TDecoratorBinding<IInterfaceB> Binding03B_;
-					AddErrorIfFalse(!Context.GetInterface(Binding03B, Binding03B_), "FAnimationAnimNextRuntimeTest_GetDecoratorInterface -> InterfaceB found");
-					AddErrorIfFalse(Binding03B == Binding03B_, "FAnimationAnimNextRuntimeTest_GetDecoratorInterface -> GetInterface methods should return the same result");
-				}
+					{
+						TDecoratorBinding<IInterfaceB> Binding03B_;
+						AddErrorIfFalse(!Context.GetInterface(Binding03B, Binding03B_), "FAnimationAnimNextRuntimeTest_GetDecoratorInterface -> InterfaceB found");
+						AddErrorIfFalse(Binding03B == Binding03B_, "FAnimationAnimNextRuntimeTest_GetDecoratorInterface -> GetInterface methods should return the same result");
+					}
 
-				{
-					TDecoratorBinding<IInterfaceB> Binding03B_;
-					AddErrorIfFalse(!Context.GetInterface(Binding03A, Binding03B_), "FAnimationAnimNextRuntimeTest_GetDecoratorInterface -> InterfaceB found");
-					AddErrorIfFalse(Binding03B == Binding03B_, "FAnimationAnimNextRuntimeTest_GetDecoratorInterface -> GetInterface methods should return the same result");
+					{
+						TDecoratorBinding<IInterfaceB> Binding03B_;
+						AddErrorIfFalse(!Context.GetInterface(Binding03A, Binding03B_), "FAnimationAnimNextRuntimeTest_GetDecoratorInterface -> InterfaceB found");
+						AddErrorIfFalse(Binding03B == Binding03B_, "FAnimationAnimNextRuntimeTest_GetDecoratorInterface -> GetInterface methods should return the same result");
+					}
 				}
-			}
 
-			{
 				{
-					TDecoratorBinding<IInterfaceA> Binding03A_;
-					AddErrorIfFalse(Context.GetInterface(Binding03C, Binding03A_), "FAnimationAnimNextRuntimeTest_GetDecoratorInterface -> InterfaceA not found");
-					AddErrorIfFalse(Binding03A == Binding03A_, "FAnimationAnimNextRuntimeTest_GetDecoratorInterface -> GetInterface methods should return the same result");
-				}
+					{
+						TDecoratorBinding<IInterfaceA> Binding03A_;
+						AddErrorIfFalse(Context.GetInterface(Binding03C, Binding03A_), "FAnimationAnimNextRuntimeTest_GetDecoratorInterface -> InterfaceA not found");
+						AddErrorIfFalse(Binding03A == Binding03A_, "FAnimationAnimNextRuntimeTest_GetDecoratorInterface -> GetInterface methods should return the same result");
+					}
 
-				{
-					TDecoratorBinding<IInterfaceA> Binding03A_;
-					AddErrorIfFalse(!Context.GetInterface(Binding03B, Binding03A_), "FAnimationAnimNextRuntimeTest_GetDecoratorInterface -> InterfaceA found");
-				}
+					{
+						TDecoratorBinding<IInterfaceA> Binding03A_;
+						AddErrorIfFalse(!Context.GetInterface(Binding03B, Binding03A_), "FAnimationAnimNextRuntimeTest_GetDecoratorInterface -> InterfaceA found");
+					}
 
-				{
-					TDecoratorBinding<IInterfaceA> Binding03A_;
-					AddErrorIfFalse(Context.GetInterface(Binding03A, Binding03A_), "FAnimationAnimNextRuntimeTest_GetDecoratorInterface -> InterfaceA not found");
-					AddErrorIfFalse(Binding03A == Binding03A_, "FAnimationAnimNextRuntimeTest_GetDecoratorInterface -> GetInterface methods should return the same result");
+					{
+						TDecoratorBinding<IInterfaceA> Binding03A_;
+						AddErrorIfFalse(Context.GetInterface(Binding03A, Binding03A_), "FAnimationAnimNextRuntimeTest_GetDecoratorInterface -> InterfaceA not found");
+						AddErrorIfFalse(Binding03A == Binding03A_, "FAnimationAnimNextRuntimeTest_GetDecoratorInterface -> GetInterface methods should return the same result");
+					}
 				}
 			}
 		}
-	}
 
-	Registry.Unregister(NodeTemplate0);
+		Registry.Unregister(NodeTemplate0);
 
-	AddErrorIfFalse(Registry.GetNum() == 0, "FAnimationAnimNextRuntimeTest_GetDecoratorInterface -> Registry should contain 0 templates");
+		AddErrorIfFalse(Registry.GetNum() == 0, "FAnimationAnimNextRuntimeTest_GetDecoratorInterface -> Registry should contain 0 templates");
+	}
 
+	Tests::FUtils::CleanupAfterTests();
+	
 	return true;
 }
 
@@ -975,196 +982,200 @@ bool FAnimationAnimNextRuntimeTest_GetDecoratorInterfaceSuper::RunTest(const FSt
 {
 	using namespace UE::AnimNext;
 
-	AUTO_REGISTER_ANIM_DECORATOR(FDecoratorA_Base)
-	AUTO_REGISTER_ANIM_DECORATOR(FDecoratorAB_Add)
-	AUTO_REGISTER_ANIM_DECORATOR(FDecoratorAC_Add)
+	{
+		AUTO_REGISTER_ANIM_DECORATOR(FDecoratorA_Base)
+		AUTO_REGISTER_ANIM_DECORATOR(FDecoratorAB_Add)
+		AUTO_REGISTER_ANIM_DECORATOR(FDecoratorAC_Add)
 
-	UFactory* GraphFactory = NewObject<UAnimNextGraphFactory>();
-	UAnimNextGraph* AnimNextGraph = CastChecked<UAnimNextGraph>(GraphFactory->FactoryCreateNew(UAnimNextGraph::StaticClass(), GetTransientPackage(), TEXT("TestAnimNextGraph"), RF_Transient, nullptr, nullptr, NAME_None));
-	UE_RETURN_ON_ERROR(AnimNextGraph != nullptr, "FAnimationAnimNextRuntimeTest_GetDecoratorInterfaceSuper -> Failed to create animation graph");
+		UFactory* GraphFactory = NewObject<UAnimNextGraphFactory>();
+		UAnimNextGraph* AnimNextGraph = CastChecked<UAnimNextGraph>(GraphFactory->FactoryCreateNew(UAnimNextGraph::StaticClass(), GetTransientPackage(), TEXT("TestAnimNextGraph"), RF_Transient, nullptr, nullptr, NAME_None));
+		UE_RETURN_ON_ERROR(AnimNextGraph != nullptr, "FAnimationAnimNextRuntimeTest_GetDecoratorInterfaceSuper -> Failed to create animation graph");
 
-	FScopedClearNodeTemplateRegistry ScopedClearNodeTemplateRegistry;
-	FNodeTemplateRegistry& Registry = FNodeTemplateRegistry::Get();
+		FScopedClearNodeTemplateRegistry ScopedClearNodeTemplateRegistry;
+		FNodeTemplateRegistry& Registry = FNodeTemplateRegistry::Get();
 
-	TArray<FDecoratorUID> NodeTemplateDecoratorList;
-	NodeTemplateDecoratorList.Add(FDecoratorA_Base::DecoratorUID);
-	NodeTemplateDecoratorList.Add(FDecoratorAB_Add::DecoratorUID);
-	NodeTemplateDecoratorList.Add(FDecoratorAC_Add::DecoratorUID);
-	NodeTemplateDecoratorList.Add(FDecoratorA_Base::DecoratorUID);
-	NodeTemplateDecoratorList.Add(FDecoratorAC_Add::DecoratorUID);
+		TArray<FDecoratorUID> NodeTemplateDecoratorList;
+		NodeTemplateDecoratorList.Add(FDecoratorA_Base::DecoratorUID);
+		NodeTemplateDecoratorList.Add(FDecoratorAB_Add::DecoratorUID);
+		NodeTemplateDecoratorList.Add(FDecoratorAC_Add::DecoratorUID);
+		NodeTemplateDecoratorList.Add(FDecoratorA_Base::DecoratorUID);
+		NodeTemplateDecoratorList.Add(FDecoratorAC_Add::DecoratorUID);
 
-	// Populate our node template registry
-	TArray<uint8> NodeTemplateBuffer0;
-	const FNodeTemplate* NodeTemplate0 = FNodeTemplateBuilder::BuildNodeTemplate(NodeTemplateDecoratorList, NodeTemplateBuffer0);
+		// Populate our node template registry
+		TArray<uint8> NodeTemplateBuffer0;
+		const FNodeTemplate* NodeTemplate0 = FNodeTemplateBuilder::BuildNodeTemplate(NodeTemplateDecoratorList, NodeTemplateBuffer0);
 
-	FNodeTemplateRegistryHandle TemplateHandle0 = Registry.FindOrAdd(NodeTemplate0);
-	AddErrorIfFalse(Registry.GetNum() == 1, "FAnimationAnimNextRuntimeTest_GetDecoratorInterfaceSuper -> Registry should contain 1 template");
-	AddErrorIfFalse(TemplateHandle0.IsValid(), "FAnimationAnimNextRuntimeTest_GetDecoratorInterfaceSuper -> Registry should contain our template");
+		FNodeTemplateRegistryHandle TemplateHandle0 = Registry.FindOrAdd(NodeTemplate0);
+		AddErrorIfFalse(Registry.GetNum() == 1, "FAnimationAnimNextRuntimeTest_GetDecoratorInterfaceSuper -> Registry should contain 1 template");
+		AddErrorIfFalse(TemplateHandle0.IsValid(), "FAnimationAnimNextRuntimeTest_GetDecoratorInterfaceSuper -> Registry should contain our template");
 
-	TArray<FNodeHandle> NodeHandles;
+		TArray<FNodeHandle> NodeHandles;
 
-	// Write our graph
-	TArray<uint8> GraphSharedDataArchiveBuffer;
-	TArray<TObjectPtr<UObject>> GraphReferencedObjects;
-	{
-		FDecoratorWriter DecoratorWriter;
+		// Write our graph
+		TArray<uint8> GraphSharedDataArchiveBuffer;
+		TArray<TObjectPtr<UObject>> GraphReferencedObjects;
+		{
+			FDecoratorWriter DecoratorWriter;
 
-		NodeHandles.Add(DecoratorWriter.RegisterNode(*NodeTemplate0));
+			NodeHandles.Add(DecoratorWriter.RegisterNode(*NodeTemplate0));
 
-		// We don't have decorator properties
+			// We don't have decorator properties
 
-		DecoratorWriter.BeginNodeWriting();
-		DecoratorWriter.WriteNode(NodeHandles[0],
-			[](uint32 DecoratorIndex, const FString& PropertyName)
-			{
-				return FString();
-			},
-			[](uint32 DecoratorIndex, const FString& PropertyName)
-			{
-				return false;
-			});
-		DecoratorWriter.EndNodeWriting();
+			DecoratorWriter.BeginNodeWriting();
+			DecoratorWriter.WriteNode(NodeHandles[0],
+				[](uint32 DecoratorIndex, const FString& PropertyName)
+				{
+					return FString();
+				},
+				[](uint32 DecoratorIndex, const FString& PropertyName)
+				{
+					return false;
+				});
+			DecoratorWriter.EndNodeWriting();
 
-		AddErrorIfFalse(DecoratorWriter.GetErrorState() == FDecoratorWriter::EErrorState::None, "FAnimationAnimNextRuntimeTest_GetDecoratorInterfaceSuper -> Failed to write decorators");
-		GraphSharedDataArchiveBuffer = DecoratorWriter.GetGraphSharedData();
-		GraphReferencedObjects = DecoratorWriter.GetGraphReferencedObjects();
-	}
+			AddErrorIfFalse(DecoratorWriter.GetErrorState() == FDecoratorWriter::EErrorState::None, "FAnimationAnimNextRuntimeTest_GetDecoratorInterfaceSuper -> Failed to write decorators");
+			GraphSharedDataArchiveBuffer = DecoratorWriter.GetGraphSharedData();
+			GraphReferencedObjects = DecoratorWriter.GetGraphReferencedObjects();
+		}
 
-	// Read our graph
-	FTestUtils::LoadFromArchiveBuffer(*AnimNextGraph, NodeHandles, GraphSharedDataArchiveBuffer);
+		// Read our graph
+		FTestUtils::LoadFromArchiveBuffer(*AnimNextGraph, NodeHandles, GraphSharedDataArchiveBuffer);
 
-	FAnimNextGraphInstancePtr GraphInstance;
-	AnimNextGraph->AllocateInstance(GraphInstance);
+		FAnimNextGraphInstancePtr GraphInstance;
+		AnimNextGraph->AllocateInstance(GraphInstance);
 
-	FExecutionContext Context(GraphInstance);
+		FExecutionContext Context(GraphInstance);
 
-	// Validate from the first base decorator
-	{
-		FDecoratorBinding ParentBinding;								// Empty, no parent
-		FAnimNextDecoratorHandle DecoratorHandle00(NodeHandles[0], 0);	// Point to first node, first base decorator
+		// Validate from the first base decorator
+		{
+			FDecoratorBinding ParentBinding;								// Empty, no parent
+			FAnimNextDecoratorHandle DecoratorHandle00(NodeHandles[0], 0);	// Point to first node, first base decorator
 
-		FDecoratorPtr DecoratorPtr00 = Context.AllocateNodeInstance(ParentBinding, DecoratorHandle00);
-		AddErrorIfFalse(DecoratorPtr00.IsValid(), "FAnimationAnimNextRuntimeTest_GetDecoratorInterfaceSuper -> Failed to allocate a node instance");
+			FDecoratorPtr DecoratorPtr00 = Context.AllocateNodeInstance(ParentBinding, DecoratorHandle00);
+			AddErrorIfFalse(DecoratorPtr00.IsValid(), "FAnimationAnimNextRuntimeTest_GetDecoratorInterfaceSuper -> Failed to allocate a node instance");
 
-		{
-			// Get a valid decorator binding: FDecoratorAC_Add
-			TDecoratorBinding<IInterfaceC> Binding02C;
-			AddErrorIfFalse(Context.GetInterface(DecoratorPtr00, Binding02C), "FAnimationAnimNextRuntimeTest_GetDecoratorInterfaceSuper -> InterfaceC not found");
-
-			// Validate GetInterfaceSuper from a decorator handle
-			TDecoratorBinding<IInterfaceC> SuperBinding02C;
-			AddErrorIfFalse(!Context.GetInterfaceSuper(Binding02C.GetDecoratorPtr(), SuperBinding02C), "FAnimationAnimNextRuntimeTest_GetDecoratorInterfaceSuper -> InterfaceC found");
-
-			// Validate GetInterfaceSuper from a decorator binding
-			TDecoratorBinding<IInterfaceC> SuperBinding02C_;
-			AddErrorIfFalse(!Context.GetInterfaceSuper(Binding02C, SuperBinding02C_), "FAnimationAnimNextRuntimeTest_GetDecoratorInterfaceSuper -> InterfaceC found");
-			AddErrorIfFalse(SuperBinding02C == SuperBinding02C_, "FAnimationAnimNextRuntimeTest_GetDecoratorInterfaceSuper -> GetInterfaceSuper methods should return the same result");
-		}
+			{
+				// Get a valid decorator binding: FDecoratorAC_Add
+				TDecoratorBinding<IInterfaceC> Binding02C;
+				AddErrorIfFalse(Context.GetInterface(DecoratorPtr00, Binding02C), "FAnimationAnimNextRuntimeTest_GetDecoratorInterfaceSuper -> InterfaceC not found");
+
+				// Validate GetInterfaceSuper from a decorator handle
+				TDecoratorBinding<IInterfaceC> SuperBinding02C;
+				AddErrorIfFalse(!Context.GetInterfaceSuper(Binding02C.GetDecoratorPtr(), SuperBinding02C), "FAnimationAnimNextRuntimeTest_GetDecoratorInterfaceSuper -> InterfaceC found");
+
+				// Validate GetInterfaceSuper from a decorator binding
+				TDecoratorBinding<IInterfaceC> SuperBinding02C_;
+				AddErrorIfFalse(!Context.GetInterfaceSuper(Binding02C, SuperBinding02C_), "FAnimationAnimNextRuntimeTest_GetDecoratorInterfaceSuper -> InterfaceC found");
+				AddErrorIfFalse(SuperBinding02C == SuperBinding02C_, "FAnimationAnimNextRuntimeTest_GetDecoratorInterfaceSuper -> GetInterfaceSuper methods should return the same result");
+			}
 
-		{
-			// Get a valid decorator binding: FDecoratorAC_Add
-			TDecoratorBinding<IInterfaceA> Binding02A;
-			AddErrorIfFalse(Context.GetInterface(DecoratorPtr00, Binding02A), "FAnimationAnimNextRuntimeTest_GetDecoratorInterfaceSuper -> InterfaceA not found");
-
-			// Validate GetInterfaceSuper from a decorator handle, FDecoratorAB_Add
-			TDecoratorBinding<IInterfaceA> SuperBinding02A;
-			AddErrorIfFalse(Context.GetInterfaceSuper(Binding02A.GetDecoratorPtr(), SuperBinding02A), "FAnimationAnimNextRuntimeTest_GetDecoratorInterfaceSuper -> InterfaceA not found");
-			AddErrorIfFalse(SuperBinding02A.IsValid(), "FAnimationAnimNextRuntimeTest_GetDecoratorInterfaceSuper -> InterfaceA binding not valid");
-			AddErrorIfFalse(SuperBinding02A.GetInterfaceUID() == IInterfaceA::InterfaceUID, "FAnimationAnimNextRuntimeTest_GetDecoratorInterfaceSuper -> Unexpected interface UID found in decorator binding");
-			AddErrorIfFalse(SuperBinding02A.GetDecoratorPtr().GetDecoratorIndex() == 1, "FAnimationAnimNextRuntimeTest_GetDecoratorInterfaceSuper -> InterfaceA not found on expected decorator");
-			AddErrorIfFalse(SuperBinding02A.GetDecoratorPtr().GetNodeInstance()->GetNodeHandle() == NodeHandles[0], "FAnimationAnimNextRuntimeTest_GetDecoratorInterfaceSuper -> InterfaceA not found on expected node");
-			AddErrorIfFalse(SuperBinding02A.GetSharedData<FDecoratorAB_Add::FSharedData>()->DecoratorUID == FDecoratorAB_Add::DecoratorUID, "FAnimationAnimNextRuntimeTest_GetDecoratorInterfaceSuper -> Unexpected shared data in decorator binding");
-			AddErrorIfFalse(SuperBinding02A.GetInstanceData<FDecoratorAB_Add::FInstanceData>()->DecoratorUID == FDecoratorAB_Add::DecoratorUID, "FAnimationAnimNextRuntimeTest_GetDecoratorInterfaceSuper -> Unexpected instance data in decorator binding");
-
-			// Validate GetInterfaceSuper from a decorator binding, FDecoratorAB_Add
-			TDecoratorBinding<IInterfaceA> SuperBinding02A_;
-			AddErrorIfFalse(Context.GetInterfaceSuper(Binding02A, SuperBinding02A_), "FAnimationAnimNextRuntimeTest_GetDecoratorInterfaceSuper -> InterfaceA not found");
-			AddErrorIfFalse(SuperBinding02A == SuperBinding02A_, "FAnimationAnimNextRuntimeTest_GetDecoratorInterfaceSuper -> GetInterfaceSuper methods should return the same result");
-
-			// Validate GetInterfaceSuper from a decorator handle, FDecoratorA_Base
-			TDecoratorBinding<IInterfaceA> SuperBinding01A;
-			AddErrorIfFalse(Context.GetInterfaceSuper(SuperBinding02A.GetDecoratorPtr(), SuperBinding01A), "FAnimationAnimNextRuntimeTest_GetDecoratorInterfaceSuper -> InterfaceA not found");
-			AddErrorIfFalse(SuperBinding01A.IsValid(), "FAnimationAnimNextRuntimeTest_GetDecoratorInterfaceSuper -> InterfaceA binding not valid");
-			AddErrorIfFalse(SuperBinding01A.GetInterfaceUID() == IInterfaceA::InterfaceUID, "FAnimationAnimNextRuntimeTest_GetDecoratorInterfaceSuper -> Unexpected interface UID found in decorator binding");
-			AddErrorIfFalse(SuperBinding01A.GetDecoratorPtr().GetDecoratorIndex() == 0, "FAnimationAnimNextRuntimeTest_GetDecoratorInterfaceSuper -> InterfaceA not found on expected decorator");
-			AddErrorIfFalse(SuperBinding01A.GetDecoratorPtr().GetNodeInstance()->GetNodeHandle() == NodeHandles[0], "FAnimationAnimNextRuntimeTest_GetDecoratorInterfaceSuper -> InterfaceA not found on expected node");
-			AddErrorIfFalse(SuperBinding01A.GetSharedData<FDecoratorA_Base::FSharedData>()->DecoratorUID == FDecoratorA_Base::DecoratorUID, "FAnimationAnimNextRuntimeTest_GetDecoratorInterfaceSuper -> Unexpected shared data in decorator binding");
-			AddErrorIfFalse(SuperBinding01A.GetInstanceData<FDecoratorA_Base::FInstanceData>()->DecoratorUID == FDecoratorA_Base::DecoratorUID, "FAnimationAnimNextRuntimeTest_GetDecoratorInterfaceSuper -> Unexpected instance data in decorator binding");
-
-			// Validate GetInterfaceSuper from a decorator binding, FDecoratorA_Base
-			TDecoratorBinding<IInterfaceA> SuperBinding01A_;
-			AddErrorIfFalse(Context.GetInterfaceSuper(SuperBinding02A, SuperBinding01A_), "FAnimationAnimNextRuntimeTest_GetDecoratorInterfaceSuper -> InterfaceA not found");
-			AddErrorIfFalse(SuperBinding01A == SuperBinding01A_, "FAnimationAnimNextRuntimeTest_GetDecoratorInterfaceSuper -> GetInterfaceSuper methods should return the same result");
-
-			// Validate GetInterfaceSuper from a decorator handle
-			TDecoratorBinding<IInterfaceA> SuperBinding00A;
-			AddErrorIfFalse(!Context.GetInterfaceSuper(SuperBinding01A.GetDecoratorPtr(), SuperBinding00A), "FAnimationAnimNextRuntimeTest_GetDecoratorInterfaceSuper -> InterfaceA found");
-
-			// Validate GetInterfaceSuper from a decorator binding
-			TDecoratorBinding<IInterfaceA> SuperBinding00A_;
-			AddErrorIfFalse(!Context.GetInterfaceSuper(SuperBinding01A, SuperBinding00A_), "FAnimationAnimNextRuntimeTest_GetDecoratorInterfaceSuper -> InterfaceA found");
-			AddErrorIfFalse(SuperBinding00A == SuperBinding00A_, "FAnimationAnimNextRuntimeTest_GetDecoratorInterfaceSuper -> GetInterfaceSuper methods should return the same result");
+			{
+				// Get a valid decorator binding: FDecoratorAC_Add
+				TDecoratorBinding<IInterfaceA> Binding02A;
+				AddErrorIfFalse(Context.GetInterface(DecoratorPtr00, Binding02A), "FAnimationAnimNextRuntimeTest_GetDecoratorInterfaceSuper -> InterfaceA not found");
+
+				// Validate GetInterfaceSuper from a decorator handle, FDecoratorAB_Add
+				TDecoratorBinding<IInterfaceA> SuperBinding02A;
+				AddErrorIfFalse(Context.GetInterfaceSuper(Binding02A.GetDecoratorPtr(), SuperBinding02A), "FAnimationAnimNextRuntimeTest_GetDecoratorInterfaceSuper -> InterfaceA not found");
+				AddErrorIfFalse(SuperBinding02A.IsValid(), "FAnimationAnimNextRuntimeTest_GetDecoratorInterfaceSuper -> InterfaceA binding not valid");
+				AddErrorIfFalse(SuperBinding02A.GetInterfaceUID() == IInterfaceA::InterfaceUID, "FAnimationAnimNextRuntimeTest_GetDecoratorInterfaceSuper -> Unexpected interface UID found in decorator binding");
+				AddErrorIfFalse(SuperBinding02A.GetDecoratorPtr().GetDecoratorIndex() == 1, "FAnimationAnimNextRuntimeTest_GetDecoratorInterfaceSuper -> InterfaceA not found on expected decorator");
+				AddErrorIfFalse(SuperBinding02A.GetDecoratorPtr().GetNodeInstance()->GetNodeHandle() == NodeHandles[0], "FAnimationAnimNextRuntimeTest_GetDecoratorInterfaceSuper -> InterfaceA not found on expected node");
+				AddErrorIfFalse(SuperBinding02A.GetSharedData<FDecoratorAB_Add::FSharedData>()->DecoratorUID == FDecoratorAB_Add::DecoratorUID, "FAnimationAnimNextRuntimeTest_GetDecoratorInterfaceSuper -> Unexpected shared data in decorator binding");
+				AddErrorIfFalse(SuperBinding02A.GetInstanceData<FDecoratorAB_Add::FInstanceData>()->DecoratorUID == FDecoratorAB_Add::DecoratorUID, "FAnimationAnimNextRuntimeTest_GetDecoratorInterfaceSuper -> Unexpected instance data in decorator binding");
+
+				// Validate GetInterfaceSuper from a decorator binding, FDecoratorAB_Add
+				TDecoratorBinding<IInterfaceA> SuperBinding02A_;
+				AddErrorIfFalse(Context.GetInterfaceSuper(Binding02A, SuperBinding02A_), "FAnimationAnimNextRuntimeTest_GetDecoratorInterfaceSuper -> InterfaceA not found");
+				AddErrorIfFalse(SuperBinding02A == SuperBinding02A_, "FAnimationAnimNextRuntimeTest_GetDecoratorInterfaceSuper -> GetInterfaceSuper methods should return the same result");
+
+				// Validate GetInterfaceSuper from a decorator handle, FDecoratorA_Base
+				TDecoratorBinding<IInterfaceA> SuperBinding01A;
+				AddErrorIfFalse(Context.GetInterfaceSuper(SuperBinding02A.GetDecoratorPtr(), SuperBinding01A), "FAnimationAnimNextRuntimeTest_GetDecoratorInterfaceSuper -> InterfaceA not found");
+				AddErrorIfFalse(SuperBinding01A.IsValid(), "FAnimationAnimNextRuntimeTest_GetDecoratorInterfaceSuper -> InterfaceA binding not valid");
+				AddErrorIfFalse(SuperBinding01A.GetInterfaceUID() == IInterfaceA::InterfaceUID, "FAnimationAnimNextRuntimeTest_GetDecoratorInterfaceSuper -> Unexpected interface UID found in decorator binding");
+				AddErrorIfFalse(SuperBinding01A.GetDecoratorPtr().GetDecoratorIndex() == 0, "FAnimationAnimNextRuntimeTest_GetDecoratorInterfaceSuper -> InterfaceA not found on expected decorator");
+				AddErrorIfFalse(SuperBinding01A.GetDecoratorPtr().GetNodeInstance()->GetNodeHandle() == NodeHandles[0], "FAnimationAnimNextRuntimeTest_GetDecoratorInterfaceSuper -> InterfaceA not found on expected node");
+				AddErrorIfFalse(SuperBinding01A.GetSharedData<FDecoratorA_Base::FSharedData>()->DecoratorUID == FDecoratorA_Base::DecoratorUID, "FAnimationAnimNextRuntimeTest_GetDecoratorInterfaceSuper -> Unexpected shared data in decorator binding");
+				AddErrorIfFalse(SuperBinding01A.GetInstanceData<FDecoratorA_Base::FInstanceData>()->DecoratorUID == FDecoratorA_Base::DecoratorUID, "FAnimationAnimNextRuntimeTest_GetDecoratorInterfaceSuper -> Unexpected instance data in decorator binding");
+
+				// Validate GetInterfaceSuper from a decorator binding, FDecoratorA_Base
+				TDecoratorBinding<IInterfaceA> SuperBinding01A_;
+				AddErrorIfFalse(Context.GetInterfaceSuper(SuperBinding02A, SuperBinding01A_), "FAnimationAnimNextRuntimeTest_GetDecoratorInterfaceSuper -> InterfaceA not found");
+				AddErrorIfFalse(SuperBinding01A == SuperBinding01A_, "FAnimationAnimNextRuntimeTest_GetDecoratorInterfaceSuper -> GetInterfaceSuper methods should return the same result");
+
+				// Validate GetInterfaceSuper from a decorator handle
+				TDecoratorBinding<IInterfaceA> SuperBinding00A;
+				AddErrorIfFalse(!Context.GetInterfaceSuper(SuperBinding01A.GetDecoratorPtr(), SuperBinding00A), "FAnimationAnimNextRuntimeTest_GetDecoratorInterfaceSuper -> InterfaceA found");
+
+				// Validate GetInterfaceSuper from a decorator binding
+				TDecoratorBinding<IInterfaceA> SuperBinding00A_;
+				AddErrorIfFalse(!Context.GetInterfaceSuper(SuperBinding01A, SuperBinding00A_), "FAnimationAnimNextRuntimeTest_GetDecoratorInterfaceSuper -> InterfaceA found");
+				AddErrorIfFalse(SuperBinding00A == SuperBinding00A_, "FAnimationAnimNextRuntimeTest_GetDecoratorInterfaceSuper -> GetInterfaceSuper methods should return the same result");
+			}
 		}
-	}
 
-	// Validate from the second base decorator
-	{
-		FDecoratorBinding ParentBinding;								// Empty, no parent
-		FAnimNextDecoratorHandle DecoratorHandle03(NodeHandles[0], 3);	// Point to first node, second base decorator
+		// Validate from the second base decorator
+		{
+			FDecoratorBinding ParentBinding;								// Empty, no parent
+			FAnimNextDecoratorHandle DecoratorHandle03(NodeHandles[0], 3);	// Point to first node, second base decorator
 
-		FDecoratorPtr DecoratorPtr03 = Context.AllocateNodeInstance(ParentBinding, DecoratorHandle03);
-		AddErrorIfFalse(DecoratorPtr03.IsValid(), "FAnimationAnimNextRuntimeTest_GetDecoratorInterfaceSuper -> Failed to allocate a node instance");
+			FDecoratorPtr DecoratorPtr03 = Context.AllocateNodeInstance(ParentBinding, DecoratorHandle03);
+			AddErrorIfFalse(DecoratorPtr03.IsValid(), "FAnimationAnimNextRuntimeTest_GetDecoratorInterfaceSuper -> Failed to allocate a node instance");
 
-		{
-			// Get a valid decorator binding: FDecoratorAC_Add
-			TDecoratorBinding<IInterfaceC> Binding04C;
-			AddErrorIfFalse(Context.GetInterface(DecoratorPtr03, Binding04C), "FAnimationAnimNextRuntimeTest_GetDecoratorInterfaceSuper -> InterfaceC not found");
-
-			// Validate GetInterfaceSuper from a decorator handle
-			TDecoratorBinding<IInterfaceC> SuperBinding04C;
-			AddErrorIfFalse(!Context.GetInterfaceSuper(Binding04C.GetDecoratorPtr(), SuperBinding04C), "FAnimationAnimNextRuntimeTest_GetDecoratorInterfaceSuper -> InterfaceC found");
-
-			// Validate GetInterfaceSuper from a decorator binding
-			TDecoratorBinding<IInterfaceC> SuperBinding04C_;
-			AddErrorIfFalse(!Context.GetInterfaceSuper(Binding04C, SuperBinding04C_), "FAnimationAnimNextRuntimeTest_GetDecoratorInterfaceSuper -> InterfaceC found");
-			AddErrorIfFalse(SuperBinding04C == SuperBinding04C_, "FAnimationAnimNextRuntimeTest_GetDecoratorInterfaceSuper -> GetInterfaceSuper methods should return the same result");
-		}
+			{
+				// Get a valid decorator binding: FDecoratorAC_Add
+				TDecoratorBinding<IInterfaceC> Binding04C;
+				AddErrorIfFalse(Context.GetInterface(DecoratorPtr03, Binding04C), "FAnimationAnimNextRuntimeTest_GetDecoratorInterfaceSuper -> InterfaceC not found");
+
+				// Validate GetInterfaceSuper from a decorator handle
+				TDecoratorBinding<IInterfaceC> SuperBinding04C;
+				AddErrorIfFalse(!Context.GetInterfaceSuper(Binding04C.GetDecoratorPtr(), SuperBinding04C), "FAnimationAnimNextRuntimeTest_GetDecoratorInterfaceSuper -> InterfaceC found");
+
+				// Validate GetInterfaceSuper from a decorator binding
+				TDecoratorBinding<IInterfaceC> SuperBinding04C_;
+				AddErrorIfFalse(!Context.GetInterfaceSuper(Binding04C, SuperBinding04C_), "FAnimationAnimNextRuntimeTest_GetDecoratorInterfaceSuper -> InterfaceC found");
+				AddErrorIfFalse(SuperBinding04C == SuperBinding04C_, "FAnimationAnimNextRuntimeTest_GetDecoratorInterfaceSuper -> GetInterfaceSuper methods should return the same result");
+			}
 
-		{
-			// Get a valid decorator binding: FDecoratorAC_Add
-			TDecoratorBinding<IInterfaceA> Binding04A;
-			AddErrorIfFalse(Context.GetInterface(DecoratorPtr03, Binding04A), "FAnimationAnimNextRuntimeTest_GetDecoratorInterfaceSuper -> InterfaceA not found");
-
-			// Validate GetInterfaceSuper from a decorator handle, FDecoratorA_Base
-			TDecoratorBinding<IInterfaceA> SuperBinding04A;
-			AddErrorIfFalse(Context.GetInterfaceSuper(Binding04A.GetDecoratorPtr(), SuperBinding04A), "FAnimationAnimNextRuntimeTest_GetDecoratorInterfaceSuper -> InterfaceA not found");
-			AddErrorIfFalse(SuperBinding04A.IsValid(), "FAnimationAnimNextRuntimeTest_GetDecoratorInterfaceSuper -> InterfaceA binding not valid");
-			AddErrorIfFalse(SuperBinding04A.GetInterfaceUID() == IInterfaceA::InterfaceUID, "FAnimationAnimNextRuntimeTest_GetDecoratorInterfaceSuper -> Unexpected interface UID found in decorator binding");
-			AddErrorIfFalse(SuperBinding04A.GetDecoratorPtr().GetDecoratorIndex() == 3, "FAnimationAnimNextRuntimeTest_GetDecoratorInterfaceSuper -> InterfaceA not found on expected decorator");
-			AddErrorIfFalse(SuperBinding04A.GetDecoratorPtr().GetNodeInstance()->GetNodeHandle() == NodeHandles[0], "FAnimationAnimNextRuntimeTest_GetDecoratorInterfaceSuper -> InterfaceA not found on expected node");
-			AddErrorIfFalse(SuperBinding04A.GetSharedData<FDecoratorA_Base::FSharedData>()->DecoratorUID == FDecoratorA_Base::DecoratorUID, "FAnimationAnimNextRuntimeTest_GetDecoratorInterfaceSuper -> Unexpected shared data in decorator binding");
-			AddErrorIfFalse(SuperBinding04A.GetInstanceData<FDecoratorA_Base::FInstanceData>()->DecoratorUID == FDecoratorA_Base::DecoratorUID, "FAnimationAnimNextRuntimeTest_GetDecoratorInterfaceSuper -> Unexpected instance data in decorator binding");
-
-			// Validate GetInterfaceSuper from a decorator binding, FDecoratorA_Base
-			TDecoratorBinding<IInterfaceA> SuperBinding04A_;
-			AddErrorIfFalse(Context.GetInterfaceSuper(Binding04A, SuperBinding04A_), "FAnimationAnimNextRuntimeTest_GetDecoratorInterfaceSuper -> InterfaceA not found");
-			AddErrorIfFalse(SuperBinding04A == SuperBinding04A_, "FAnimationAnimNextRuntimeTest_GetDecoratorInterfaceSuper -> GetInterfaceSuper methods should return the same result");
-
-			// Validate GetInterfaceSuper from a decorator handle
-			TDecoratorBinding<IInterfaceA> SuperBinding03A;
-			AddErrorIfFalse(!Context.GetInterfaceSuper(SuperBinding04A.GetDecoratorPtr(), SuperBinding03A), "FAnimationAnimNextRuntimeTest_GetDecoratorInterfaceSuper -> InterfaceA found");
-
-			// Validate GetInterfaceSuper from a decorator binding
-			TDecoratorBinding<IInterfaceA> SuperBinding03A_;
-			AddErrorIfFalse(!Context.GetInterfaceSuper(SuperBinding04A, SuperBinding03A_), "FAnimationAnimNextRuntimeTest_GetDecoratorInterfaceSuper -> InterfaceA found");
-			AddErrorIfFalse(SuperBinding03A == SuperBinding03A_, "FAnimationAnimNextRuntimeTest_GetDecoratorInterfaceSuper -> GetInterfaceSuper methods should return the same result");
+			{
+				// Get a valid decorator binding: FDecoratorAC_Add
+				TDecoratorBinding<IInterfaceA> Binding04A;
+				AddErrorIfFalse(Context.GetInterface(DecoratorPtr03, Binding04A), "FAnimationAnimNextRuntimeTest_GetDecoratorInterfaceSuper -> InterfaceA not found");
+
+				// Validate GetInterfaceSuper from a decorator handle, FDecoratorA_Base
+				TDecoratorBinding<IInterfaceA> SuperBinding04A;
+				AddErrorIfFalse(Context.GetInterfaceSuper(Binding04A.GetDecoratorPtr(), SuperBinding04A), "FAnimationAnimNextRuntimeTest_GetDecoratorInterfaceSuper -> InterfaceA not found");
+				AddErrorIfFalse(SuperBinding04A.IsValid(), "FAnimationAnimNextRuntimeTest_GetDecoratorInterfaceSuper -> InterfaceA binding not valid");
+				AddErrorIfFalse(SuperBinding04A.GetInterfaceUID() == IInterfaceA::InterfaceUID, "FAnimationAnimNextRuntimeTest_GetDecoratorInterfaceSuper -> Unexpected interface UID found in decorator binding");
+				AddErrorIfFalse(SuperBinding04A.GetDecoratorPtr().GetDecoratorIndex() == 3, "FAnimationAnimNextRuntimeTest_GetDecoratorInterfaceSuper -> InterfaceA not found on expected decorator");
+				AddErrorIfFalse(SuperBinding04A.GetDecoratorPtr().GetNodeInstance()->GetNodeHandle() == NodeHandles[0], "FAnimationAnimNextRuntimeTest_GetDecoratorInterfaceSuper -> InterfaceA not found on expected node");
+				AddErrorIfFalse(SuperBinding04A.GetSharedData<FDecoratorA_Base::FSharedData>()->DecoratorUID == FDecoratorA_Base::DecoratorUID, "FAnimationAnimNextRuntimeTest_GetDecoratorInterfaceSuper -> Unexpected shared data in decorator binding");
+				AddErrorIfFalse(SuperBinding04A.GetInstanceData<FDecoratorA_Base::FInstanceData>()->DecoratorUID == FDecoratorA_Base::DecoratorUID, "FAnimationAnimNextRuntimeTest_GetDecoratorInterfaceSuper -> Unexpected instance data in decorator binding");
+
+				// Validate GetInterfaceSuper from a decorator binding, FDecoratorA_Base
+				TDecoratorBinding<IInterfaceA> SuperBinding04A_;
+				AddErrorIfFalse(Context.GetInterfaceSuper(Binding04A, SuperBinding04A_), "FAnimationAnimNextRuntimeTest_GetDecoratorInterfaceSuper -> InterfaceA not found");
+				AddErrorIfFalse(SuperBinding04A == SuperBinding04A_, "FAnimationAnimNextRuntimeTest_GetDecoratorInterfaceSuper -> GetInterfaceSuper methods should return the same result");
+
+				// Validate GetInterfaceSuper from a decorator handle
+				TDecoratorBinding<IInterfaceA> SuperBinding03A;
+				AddErrorIfFalse(!Context.GetInterfaceSuper(SuperBinding04A.GetDecoratorPtr(), SuperBinding03A), "FAnimationAnimNextRuntimeTest_GetDecoratorInterfaceSuper -> InterfaceA found");
+
+				// Validate GetInterfaceSuper from a decorator binding
+				TDecoratorBinding<IInterfaceA> SuperBinding03A_;
+				AddErrorIfFalse(!Context.GetInterfaceSuper(SuperBinding04A, SuperBinding03A_), "FAnimationAnimNextRuntimeTest_GetDecoratorInterfaceSuper -> InterfaceA found");
+				AddErrorIfFalse(SuperBinding03A == SuperBinding03A_, "FAnimationAnimNextRuntimeTest_GetDecoratorInterfaceSuper -> GetInterfaceSuper methods should return the same result");
+			}
 		}
-	}
 
-	Registry.Unregister(NodeTemplate0);
+		Registry.Unregister(NodeTemplate0);
 
-	AddErrorIfFalse(Registry.GetNum() == 0, "FAnimationAnimNextRuntimeTest_GetDecoratorInterface -> Registry should contain 0 templates");
+		AddErrorIfFalse(Registry.GetNum() == 0, "FAnimationAnimNextRuntimeTest_GetDecoratorInterface -> Registry should contain 0 templates");
+	}
 
+	Tests::FUtils::CleanupAfterTests();
+	
 	return true;
 }
 
@@ -1511,6 +1522,8 @@ bool FAnimationAnimNextRuntimeTest_DecoratorSerialization::RunTest(const FString
 		}
 	}
 
+	Tests::FUtils::CleanupAfterTests();
+	
 	return true;
 }
 
diff --git a/Engine/Plugins/Experimental/AnimNext/Source/AnimNextTestSuite/Private/AnimNextDecoratorGraphTest.cpp b/Engine/Plugins/Experimental/AnimNext/Source/AnimNextTestSuite/Private/AnimNextDecoratorGraphTest.cpp
--- a/Engine/Plugins/Experimental/AnimNext/Source/AnimNextTestSuite/Private/AnimNextDecoratorGraphTest.cpp
+++ b/Engine/Plugins/Experimental/AnimNext/Source/AnimNextTestSuite/Private/AnimNextDecoratorGraphTest.cpp
@@ -3,12 +3,14 @@
 #include "AnimNextDecoratorGraphTest.h"
 
 #include "AnimNextRuntimeTest.h"
+#include "AnimNextTest.h"
 #include "AssetToolsModule.h"
 #include "Context.h"
 #include "UncookedOnlyUtils.h"
 #include "DecoratorBase/DecoratorRegistry.h"
 #include "DecoratorInterfaces/IEvaluate.h"
 #include "DecoratorInterfaces/IUpdate.h"
+#include "Editor/Transactor.h"
 #include "Graph/AnimNextExecuteContext.h"
 #include "Graph/AnimNextGraph.h"
 #include "Graph/AnimNextGraph_EditorData.h"
@@ -168,6 +170,8 @@ bool FAnimationAnimNextRuntimeTest_GraphAddDecorator::RunTest(const FString& InP
 	UE_RETURN_ON_ERROR(DecoratorPin->GetSubPins()[5]->GetDefaultValue() == TEXT("34.000000"), TEXT("FAnimationAnimNextRuntimeTest_GraphAddDecorator -> Unexpected decorator pin value"));
 	UE_RETURN_ON_ERROR(DecoratorPin->GetSubPins()[5]->IsLazy(), TEXT("FAnimationAnimNextRuntimeTest_GraphAddDecorator -> Expected lazy decorator pin"));
 
+	Tests::FUtils::CleanupAfterTests();
+
 	return true;
 }
 
@@ -280,6 +284,8 @@ bool FAnimationAnimNextRuntimeTest_GraphExecute::RunTest(const FString& InParame
 
 	FParamStack::DetachFromCurrentThread();
 
+	Tests::FUtils::CleanupAfterTests();
+
 	return true;
 }
 
@@ -423,6 +429,8 @@ bool FAnimationAnimNextRuntimeTest_GraphExecuteLatent::RunTest(const FString& In
 
 	FParamStack::DetachFromCurrentThread();
 
+	Tests::FUtils::CleanupAfterTests();
+
 	return true;
 }
 
