AnimNext: Begin merging parameter blocks and graphs

diff --git a/Engine/Plugins/Experimental/AnimNext/Source/AnimNextTestSuite/Private/AnimNextDecoratorInterfacesTest.cpp b/Engine/Plugins/Experimental/AnimNext/Source/AnimNextTestSuite/Private/AnimNextDecoratorInterfacesTest.cpp
--- a/Engine/Plugins/Experimental/AnimNext/Source/AnimNextTestSuite/Private/AnimNextDecoratorInterfacesTest.cpp
+++ b/Engine/Plugins/Experimental/AnimNext/Source/AnimNextTestSuite/Private/AnimNextDecoratorInterfacesTest.cpp
@@ -2,6 +2,7 @@
 
 #include "AnimNextDecoratorInterfacesTest.h"
 #include "AnimNextRuntimeTest.h"
+#include "AnimNextTest.h"
 
 #include "Misc/AutomationTest.h"
 
@@ -222,180 +223,184 @@ bool FAnimationAnimNextRuntimeTest_IHierarchy::RunTest(const FString& InParamete
 {
 	using namespace UE::AnimNext;
 
-	AUTO_REGISTER_ANIM_DECORATOR(FDecoratorWithNoChildren)
-	AUTO_REGISTER_ANIM_DECORATOR(FDecoratorWithOneChild)
-	AUTO_REGISTER_ANIM_DECORATOR(FDecoratorWithChildren)
-
-	UFactory* GraphFactory = NewObject<UAnimNextGraphFactory>();
-	UAnimNextGraph* AnimNextGraph = CastChecked<UAnimNextGraph>(GraphFactory->FactoryCreateNew(UAnimNextGraph::StaticClass(), GetTransientPackage(), TEXT("TestAnimNextGraph"), RF_Transient, nullptr, nullptr, NAME_None));
-	UE_RETURN_ON_ERROR(AnimNextGraph != nullptr, "FAnimationAnimNextRuntimeTest_IHierarchy -> Failed to create animation graph");
-
-	FScopedClearNodeTemplateRegistry ScopedClearNodeTemplateRegistry;
-	FNodeTemplateRegistry& Registry = FNodeTemplateRegistry::Get();
-
-	// We create a few node templates
-	// Template A has a single decorator with no children
-	TArray<FDecoratorUID> NodeTemplateDecoratorListA;
-	NodeTemplateDecoratorListA.Add(FDecoratorWithNoChildren::DecoratorUID);
-
-	// Template B has a single decorator with one child
-	TArray<FDecoratorUID> NodeTemplateDecoratorListB;
-	NodeTemplateDecoratorListB.Add(FDecoratorWithOneChild::DecoratorUID);
-
-	// Template C has two decorators, each with one child
-	TArray<FDecoratorUID> NodeTemplateDecoratorListC;
-	NodeTemplateDecoratorListC.Add(FDecoratorWithOneChild::DecoratorUID);
-	NodeTemplateDecoratorListC.Add(FDecoratorWithOneChild::DecoratorUID);
-
-	// Template D has a single decorator with children
-	TArray<FDecoratorUID> NodeTemplateDecoratorListD;
-	NodeTemplateDecoratorListD.Add(FDecoratorWithChildren::DecoratorUID);
-
-	// Populate our node template registry
-	TArray<uint8> NodeTemplateBufferA, NodeTemplateBufferB, NodeTemplateBufferC, NodeTemplateBufferD;
-	const FNodeTemplate* NodeTemplateA = FNodeTemplateBuilder::BuildNodeTemplate(NodeTemplateDecoratorListA, NodeTemplateBufferA);
-	const FNodeTemplate* NodeTemplateB = FNodeTemplateBuilder::BuildNodeTemplate(NodeTemplateDecoratorListB, NodeTemplateBufferB);
-	const FNodeTemplate* NodeTemplateC = FNodeTemplateBuilder::BuildNodeTemplate(NodeTemplateDecoratorListC, NodeTemplateBufferC);
-	const FNodeTemplate* NodeTemplateD = FNodeTemplateBuilder::BuildNodeTemplate(NodeTemplateDecoratorListD, NodeTemplateBufferD);
-
-	// Build our graph, it as follow (each node template has a single node instance):
-	// NodeA has no children
-	// NodeB has one child: NodeA
-	// NodeC has two children: NodeA and NodeB (but both decorators are base, only NodeB will be referenced)
-	// NodeD has two children: NodeA and NodeC
-
-	TArray<FNodeHandle> NodeHandles;
-
-	// Write our graph
-	TArray<uint8> GraphSharedDataArchiveBuffer;
-	TArray<TObjectPtr<UObject>> GraphReferencedObjects;
 	{
-		FDecoratorWriter DecoratorWriter;
-
-		NodeHandles.Add(DecoratorWriter.RegisterNode(*NodeTemplateA));
-		NodeHandles.Add(DecoratorWriter.RegisterNode(*NodeTemplateB));
-		NodeHandles.Add(DecoratorWriter.RegisterNode(*NodeTemplateC));
-		NodeHandles.Add(DecoratorWriter.RegisterNode(*NodeTemplateD));
-
-		// We don't have decorator properties
-		TArray<TMap<FString, FString>> DecoratorPropertiesA;
-		DecoratorPropertiesA.AddDefaulted(NodeTemplateDecoratorListA.Num());
-
-		TArray<TMap<FString, FString>> DecoratorPropertiesB;
-		DecoratorPropertiesB.AddDefaulted(NodeTemplateDecoratorListB.Num());
-		DecoratorPropertiesB[0].Add(TEXT("Child"), ToString<FDecoratorWithOneChild::FSharedData>(TEXT("Child"), FAnimNextDecoratorHandle(NodeHandles[0])));
-
-		TArray<TMap<FString, FString>> DecoratorPropertiesC;
-		DecoratorPropertiesC.AddDefaulted(NodeTemplateDecoratorListC.Num());
-		DecoratorPropertiesC[0].Add(TEXT("Child"), ToString<FDecoratorWithOneChild::FSharedData>(TEXT("Child"), FAnimNextDecoratorHandle(NodeHandles[0])));
-		DecoratorPropertiesC[1].Add(TEXT("Child"), ToString<FDecoratorWithOneChild::FSharedData>(TEXT("Child"), FAnimNextDecoratorHandle(NodeHandles[1])));
-
-		TArray<TMap<FString, FString>> DecoratorPropertiesD;
-		DecoratorPropertiesD.AddDefaulted(NodeTemplateDecoratorListD.Num());
-		FAnimNextDecoratorHandle ChildrenHandlesD[2] = { FAnimNextDecoratorHandle(NodeHandles[0]), FAnimNextDecoratorHandle(NodeHandles[2], 1)};
-		DecoratorPropertiesD[0].Add(TEXT("Children"), ToString<FDecoratorWithChildren::FSharedData>(TEXT("Children"), ChildrenHandlesD));
-
-		DecoratorWriter.BeginNodeWriting();
-		DecoratorWriter.WriteNode(NodeHandles[0],
-			[&DecoratorPropertiesA](uint32 DecoratorIndex, const FString& PropertyName)
-			{
-				return DecoratorPropertiesA[DecoratorIndex][PropertyName];
-			},
-			[](uint32 DecoratorIndex, const FString& PropertyName)
-			{
-				return false;
-			});
-		DecoratorWriter.WriteNode(NodeHandles[1],
-			[&DecoratorPropertiesB](uint32 DecoratorIndex, const FString& PropertyName)
-			{
-				return DecoratorPropertiesB[DecoratorIndex][PropertyName];
-			},
-			[](uint32 DecoratorIndex, const FString& PropertyName)
-			{
-				return false;
-			});
-		DecoratorWriter.WriteNode(NodeHandles[2],
-			[&DecoratorPropertiesC](uint32 DecoratorIndex, const FString& PropertyName)
-			{
-				return DecoratorPropertiesC[DecoratorIndex][PropertyName];
-			},
-			[](uint32 DecoratorIndex, const FString& PropertyName)
-			{
-				return false;
-			});
-		DecoratorWriter.WriteNode(NodeHandles[3],
-			[&DecoratorPropertiesD](uint32 DecoratorIndex, const FString& PropertyName)
-			{
-				return DecoratorPropertiesD[DecoratorIndex][PropertyName];
-			},
-			[](uint32 DecoratorIndex, const FString& PropertyName)
-			{
-				return false;
-			});
-		DecoratorWriter.EndNodeWriting();
-
-		AddErrorIfFalse(DecoratorWriter.GetErrorState() == FDecoratorWriter::EErrorState::None, "FAnimationAnimNextRuntimeTest_IHierarchy -> Failed to write decorators");
-		GraphSharedDataArchiveBuffer = DecoratorWriter.GetGraphSharedData();
-		GraphReferencedObjects = DecoratorWriter.GetGraphReferencedObjects();
-	}
+		AUTO_REGISTER_ANIM_DECORATOR(FDecoratorWithNoChildren)
+		AUTO_REGISTER_ANIM_DECORATOR(FDecoratorWithOneChild)
+		AUTO_REGISTER_ANIM_DECORATOR(FDecoratorWithChildren)
+
+		UFactory* GraphFactory = NewObject<UAnimNextGraphFactory>();
+		UAnimNextGraph* AnimNextGraph = CastChecked<UAnimNextGraph>(GraphFactory->FactoryCreateNew(UAnimNextGraph::StaticClass(), GetTransientPackage(), TEXT("TestAnimNextGraph"), RF_Transient, nullptr, nullptr, NAME_None));
+		UE_RETURN_ON_ERROR(AnimNextGraph != nullptr, "FAnimationAnimNextRuntimeTest_IHierarchy -> Failed to create animation graph");
+
+		FScopedClearNodeTemplateRegistry ScopedClearNodeTemplateRegistry;
+		FNodeTemplateRegistry& Registry = FNodeTemplateRegistry::Get();
+
+		// We create a few node templates
+		// Template A has a single decorator with no children
+		TArray<FDecoratorUID> NodeTemplateDecoratorListA;
+		NodeTemplateDecoratorListA.Add(FDecoratorWithNoChildren::DecoratorUID);
+
+		// Template B has a single decorator with one child
+		TArray<FDecoratorUID> NodeTemplateDecoratorListB;
+		NodeTemplateDecoratorListB.Add(FDecoratorWithOneChild::DecoratorUID);
+
+		// Template C has two decorators, each with one child
+		TArray<FDecoratorUID> NodeTemplateDecoratorListC;
+		NodeTemplateDecoratorListC.Add(FDecoratorWithOneChild::DecoratorUID);
+		NodeTemplateDecoratorListC.Add(FDecoratorWithOneChild::DecoratorUID);
+
+		// Template D has a single decorator with children
+		TArray<FDecoratorUID> NodeTemplateDecoratorListD;
+		NodeTemplateDecoratorListD.Add(FDecoratorWithChildren::DecoratorUID);
+
+		// Populate our node template registry
+		TArray<uint8> NodeTemplateBufferA, NodeTemplateBufferB, NodeTemplateBufferC, NodeTemplateBufferD;
+		const FNodeTemplate* NodeTemplateA = FNodeTemplateBuilder::BuildNodeTemplate(NodeTemplateDecoratorListA, NodeTemplateBufferA);
+		const FNodeTemplate* NodeTemplateB = FNodeTemplateBuilder::BuildNodeTemplate(NodeTemplateDecoratorListB, NodeTemplateBufferB);
+		const FNodeTemplate* NodeTemplateC = FNodeTemplateBuilder::BuildNodeTemplate(NodeTemplateDecoratorListC, NodeTemplateBufferC);
+		const FNodeTemplate* NodeTemplateD = FNodeTemplateBuilder::BuildNodeTemplate(NodeTemplateDecoratorListD, NodeTemplateBufferD);
+
+		// Build our graph, it as follow (each node template has a single node instance):
+		// NodeA has no children
+		// NodeB has one child: NodeA
+		// NodeC has two children: NodeA and NodeB (but both decorators are base, only NodeB will be referenced)
+		// NodeD has two children: NodeA and NodeC
+
+		TArray<FNodeHandle> NodeHandles;
+
+		// Write our graph
+		TArray<uint8> GraphSharedDataArchiveBuffer;
+		TArray<TObjectPtr<UObject>> GraphReferencedObjects;
+		{
+			FDecoratorWriter DecoratorWriter;
+
+			NodeHandles.Add(DecoratorWriter.RegisterNode(*NodeTemplateA));
+			NodeHandles.Add(DecoratorWriter.RegisterNode(*NodeTemplateB));
+			NodeHandles.Add(DecoratorWriter.RegisterNode(*NodeTemplateC));
+			NodeHandles.Add(DecoratorWriter.RegisterNode(*NodeTemplateD));
+
+			// We don't have decorator properties
+			TArray<TMap<FString, FString>> DecoratorPropertiesA;
+			DecoratorPropertiesA.AddDefaulted(NodeTemplateDecoratorListA.Num());
+
+			TArray<TMap<FString, FString>> DecoratorPropertiesB;
+			DecoratorPropertiesB.AddDefaulted(NodeTemplateDecoratorListB.Num());
+			DecoratorPropertiesB[0].Add(TEXT("Child"), ToString<FDecoratorWithOneChild::FSharedData>(TEXT("Child"), FAnimNextDecoratorHandle(NodeHandles[0])));
+
+			TArray<TMap<FString, FString>> DecoratorPropertiesC;
+			DecoratorPropertiesC.AddDefaulted(NodeTemplateDecoratorListC.Num());
+			DecoratorPropertiesC[0].Add(TEXT("Child"), ToString<FDecoratorWithOneChild::FSharedData>(TEXT("Child"), FAnimNextDecoratorHandle(NodeHandles[0])));
+			DecoratorPropertiesC[1].Add(TEXT("Child"), ToString<FDecoratorWithOneChild::FSharedData>(TEXT("Child"), FAnimNextDecoratorHandle(NodeHandles[1])));
+
+			TArray<TMap<FString, FString>> DecoratorPropertiesD;
+			DecoratorPropertiesD.AddDefaulted(NodeTemplateDecoratorListD.Num());
+			FAnimNextDecoratorHandle ChildrenHandlesD[2] = { FAnimNextDecoratorHandle(NodeHandles[0]), FAnimNextDecoratorHandle(NodeHandles[2], 1)};
+			DecoratorPropertiesD[0].Add(TEXT("Children"), ToString<FDecoratorWithChildren::FSharedData>(TEXT("Children"), ChildrenHandlesD));
+
+			DecoratorWriter.BeginNodeWriting();
+			DecoratorWriter.WriteNode(NodeHandles[0],
+				[&DecoratorPropertiesA](uint32 DecoratorIndex, const FString& PropertyName)
+				{
+					return DecoratorPropertiesA[DecoratorIndex][PropertyName];
+				},
+				[](uint32 DecoratorIndex, const FString& PropertyName)
+				{
+					return false;
+				});
+			DecoratorWriter.WriteNode(NodeHandles[1],
+				[&DecoratorPropertiesB](uint32 DecoratorIndex, const FString& PropertyName)
+				{
+					return DecoratorPropertiesB[DecoratorIndex][PropertyName];
+				},
+				[](uint32 DecoratorIndex, const FString& PropertyName)
+				{
+					return false;
+				});
+			DecoratorWriter.WriteNode(NodeHandles[2],
+				[&DecoratorPropertiesC](uint32 DecoratorIndex, const FString& PropertyName)
+				{
+					return DecoratorPropertiesC[DecoratorIndex][PropertyName];
+				},
+				[](uint32 DecoratorIndex, const FString& PropertyName)
+				{
+					return false;
+				});
+			DecoratorWriter.WriteNode(NodeHandles[3],
+				[&DecoratorPropertiesD](uint32 DecoratorIndex, const FString& PropertyName)
+				{
+					return DecoratorPropertiesD[DecoratorIndex][PropertyName];
+				},
+				[](uint32 DecoratorIndex, const FString& PropertyName)
+				{
+					return false;
+				});
+			DecoratorWriter.EndNodeWriting();
+
+			AddErrorIfFalse(DecoratorWriter.GetErrorState() == FDecoratorWriter::EErrorState::None, "FAnimationAnimNextRuntimeTest_IHierarchy -> Failed to write decorators");
+			GraphSharedDataArchiveBuffer = DecoratorWriter.GetGraphSharedData();
+			GraphReferencedObjects = DecoratorWriter.GetGraphReferencedObjects();
+		}
 
-	// Read our graph
-	FTestUtils::LoadFromArchiveBuffer(*AnimNextGraph, NodeHandles, GraphSharedDataArchiveBuffer);
+		// Read our graph
+		FTestUtils::LoadFromArchiveBuffer(*AnimNextGraph, NodeHandles, GraphSharedDataArchiveBuffer);
 
-	FAnimNextGraphInstancePtr GraphInstance;
-	AnimNextGraph->AllocateInstance(GraphInstance);
+		FAnimNextGraphInstancePtr GraphInstance;
+		AnimNextGraph->AllocateInstance(GraphInstance);
 
-	FExecutionContext Context(GraphInstance);
+		FExecutionContext Context(GraphInstance);
 
-	{
-		FMemMark Mark(FMemStack::Get());
+		{
+			FMemMark Mark(FMemStack::Get());
 
-		FWeakDecoratorPtr NullPtr;								// Empty, no parent
-		FAnimNextDecoratorHandle RootHandle(NodeHandles[3]);	// Point to NodeD, first base decorator
+			FWeakDecoratorPtr NullPtr;								// Empty, no parent
+			FAnimNextDecoratorHandle RootHandle(NodeHandles[3]);	// Point to NodeD, first base decorator
 
-		FDecoratorPtr NodeDPtr = Context.AllocateNodeInstance(NullPtr, RootHandle);
-		AddErrorIfFalse(NodeDPtr.IsValid(), "FAnimationAnimNextRuntimeTest_IHierarchy -> Failed to allocate root node instance");
+			FDecoratorPtr NodeDPtr = Context.AllocateNodeInstance(NullPtr, RootHandle);
+			AddErrorIfFalse(NodeDPtr.IsValid(), "FAnimationAnimNextRuntimeTest_IHierarchy -> Failed to allocate root node instance");
 
-		TDecoratorBinding<IHierarchy> HierarchyBindingNodeD;
-		AddErrorIfFalse(Context.GetInterface(NodeDPtr, HierarchyBindingNodeD), "FAnimationAnimNextRuntimeTest_IHierarchy -> IHierarchy not found");
+			TDecoratorBinding<IHierarchy> HierarchyBindingNodeD;
+			AddErrorIfFalse(Context.GetInterface(NodeDPtr, HierarchyBindingNodeD), "FAnimationAnimNextRuntimeTest_IHierarchy -> IHierarchy not found");
 
-		FChildrenArray ChildrenNodeD;
-		HierarchyBindingNodeD.GetChildren(Context, ChildrenNodeD);
+			FChildrenArray ChildrenNodeD;
+			HierarchyBindingNodeD.GetChildren(Context, ChildrenNodeD);
 
-		AddErrorIfFalse(ChildrenNodeD.Num() == 2, "FAnimationAnimNextRuntimeTest_IHierarchy -> Expected 2 children");
-		AddErrorIfFalse(ChildrenNodeD[0].IsValid() && ChildrenNodeD[0].GetNodeInstance()->GetNodeHandle() == NodeHandles[0], "FAnimationAnimNextRuntimeTest_IHierarchy -> Expected child: NodeA");
-		AddErrorIfFalse(ChildrenNodeD[1].IsValid() && ChildrenNodeD[1].GetNodeInstance()->GetNodeHandle() == NodeHandles[2], "FAnimationAnimNextRuntimeTest_IHierarchy -> Expected child: NodeC");
+			AddErrorIfFalse(ChildrenNodeD.Num() == 2, "FAnimationAnimNextRuntimeTest_IHierarchy -> Expected 2 children");
+			AddErrorIfFalse(ChildrenNodeD[0].IsValid() && ChildrenNodeD[0].GetNodeInstance()->GetNodeHandle() == NodeHandles[0], "FAnimationAnimNextRuntimeTest_IHierarchy -> Expected child: NodeA");
+			AddErrorIfFalse(ChildrenNodeD[1].IsValid() && ChildrenNodeD[1].GetNodeInstance()->GetNodeHandle() == NodeHandles[2], "FAnimationAnimNextRuntimeTest_IHierarchy -> Expected child: NodeC");
 
-		{
-			TDecoratorBinding<IHierarchy> HierarchyBindingNodeC;
-			AddErrorIfFalse(Context.GetInterface(ChildrenNodeD[1], HierarchyBindingNodeC), "FAnimationAnimNextRuntimeTest_IHierarchy -> IHierarchy not found");
+			{
+				TDecoratorBinding<IHierarchy> HierarchyBindingNodeC;
+				AddErrorIfFalse(Context.GetInterface(ChildrenNodeD[1], HierarchyBindingNodeC), "FAnimationAnimNextRuntimeTest_IHierarchy -> IHierarchy not found");
 
-			FChildrenArray ChildrenNodeC;
-			HierarchyBindingNodeC.GetChildren(Context, ChildrenNodeC);
+				FChildrenArray ChildrenNodeC;
+				HierarchyBindingNodeC.GetChildren(Context, ChildrenNodeC);
 
-			AddErrorIfFalse(ChildrenNodeC.Num() == 1, "FAnimationAnimNextRuntimeTest_IHierarchy -> Expected 2 children");
-			AddErrorIfFalse(ChildrenNodeC[0].IsValid() && ChildrenNodeC[0].GetNodeInstance()->GetNodeHandle() == NodeHandles[1], "FAnimationAnimNextRuntimeTest_IHierarchy -> Expected child: NodeB");
+				AddErrorIfFalse(ChildrenNodeC.Num() == 1, "FAnimationAnimNextRuntimeTest_IHierarchy -> Expected 2 children");
+				AddErrorIfFalse(ChildrenNodeC[0].IsValid() && ChildrenNodeC[0].GetNodeInstance()->GetNodeHandle() == NodeHandles[1], "FAnimationAnimNextRuntimeTest_IHierarchy -> Expected child: NodeB");
 
-			{
-				TDecoratorBinding<IHierarchy> HierarchyBindingNodeB;
-				AddErrorIfFalse(Context.GetInterface(ChildrenNodeC[0], HierarchyBindingNodeB), "FAnimationAnimNextRuntimeTest_IHierarchy -> IHierarchy not found");
+				{
+					TDecoratorBinding<IHierarchy> HierarchyBindingNodeB;
+					AddErrorIfFalse(Context.GetInterface(ChildrenNodeC[0], HierarchyBindingNodeB), "FAnimationAnimNextRuntimeTest_IHierarchy -> IHierarchy not found");
 
-				FChildrenArray ChildrenNodeB;
-				HierarchyBindingNodeB.GetChildren(Context, ChildrenNodeB);
+					FChildrenArray ChildrenNodeB;
+					HierarchyBindingNodeB.GetChildren(Context, ChildrenNodeB);
 
-				AddErrorIfFalse(ChildrenNodeB.Num() == 1, "FAnimationAnimNextRuntimeTest_IHierarchy -> Expected 1 child");
-				AddErrorIfFalse(ChildrenNodeB[0].IsValid() && ChildrenNodeB[0].GetNodeInstance()->GetNodeHandle() == NodeHandles[0], "FAnimationAnimNextRuntimeTest_IHierarchy -> Expected child: NodeA");
+					AddErrorIfFalse(ChildrenNodeB.Num() == 1, "FAnimationAnimNextRuntimeTest_IHierarchy -> Expected 1 child");
+					AddErrorIfFalse(ChildrenNodeB[0].IsValid() && ChildrenNodeB[0].GetNodeInstance()->GetNodeHandle() == NodeHandles[0], "FAnimationAnimNextRuntimeTest_IHierarchy -> Expected child: NodeA");
+				}
 			}
 		}
-	}
 
-	Registry.Unregister(NodeTemplateA);
-	Registry.Unregister(NodeTemplateB);
-	Registry.Unregister(NodeTemplateC);
-	Registry.Unregister(NodeTemplateD);
+		Registry.Unregister(NodeTemplateA);
+		Registry.Unregister(NodeTemplateB);
+		Registry.Unregister(NodeTemplateC);
+		Registry.Unregister(NodeTemplateD);
 
-	AddErrorIfFalse(Registry.GetNum() == 0, "FAnimationAnimNextRuntimeTest_IHierarchy -> Registry should contain 0 templates");
+		AddErrorIfFalse(Registry.GetNum() == 0, "FAnimationAnimNextRuntimeTest_IHierarchy -> Registry should contain 0 templates");
+	}
 
+	Tests::FUtils::CleanupAfterTests();
+	
 	return true;
 }
 
@@ -404,133 +409,135 @@ IMPLEMENT_SIMPLE_AUTOMATION_TEST(FAnimationAnimNextRuntimeTest_IUpdate, "Animati
 bool FAnimationAnimNextRuntimeTest_IUpdate::RunTest(const FString& InParameters)
 {
 	using namespace UE::AnimNext;
-
-	AUTO_REGISTER_ANIM_DECORATOR(FDecoratorWithNoChildren)
-	AUTO_REGISTER_ANIM_DECORATOR(FDecoratorWithOneChild)
-	AUTO_REGISTER_ANIM_DECORATOR(FDecoratorWithChildren)
-
-	UFactory* GraphFactory = NewObject<UAnimNextGraphFactory>();
-	UAnimNextGraph* AnimNextGraph = CastChecked<UAnimNextGraph>(GraphFactory->FactoryCreateNew(UAnimNextGraph::StaticClass(), GetTransientPackage(), TEXT("TestAnimNextGraph"), RF_Transient, nullptr, nullptr, NAME_None));
-	UE_RETURN_ON_ERROR(AnimNextGraph != nullptr, "FAnimationAnimNextRuntimeTest_IUpdate -> Failed to create animation graph");
-
-	FScopedClearNodeTemplateRegistry ScopedClearNodeTemplateRegistry;
-	FNodeTemplateRegistry& Registry = FNodeTemplateRegistry::Get();
-
-	// We create a few node templates
-	// Template A has a single decorator with no children
-	TArray<FDecoratorUID> NodeTemplateDecoratorListA;
-	NodeTemplateDecoratorListA.Add(FDecoratorWithNoChildren::DecoratorUID);
-
-	// Template B has a single decorator with one child, it doesn't update
-	TArray<FDecoratorUID> NodeTemplateDecoratorListB;
-	NodeTemplateDecoratorListB.Add(FDecoratorWithOneChild::DecoratorUID);
-
-	// Template C has a single decorator with children
-	TArray<FDecoratorUID> NodeTemplateDecoratorListC;
-	NodeTemplateDecoratorListC.Add(FDecoratorWithChildren::DecoratorUID);
-
-	// Populate our node template registry
-	TArray<uint8> NodeTemplateBufferA, NodeTemplateBufferB, NodeTemplateBufferC;
-	const FNodeTemplate* NodeTemplateA = FNodeTemplateBuilder::BuildNodeTemplate(NodeTemplateDecoratorListA, NodeTemplateBufferA);
-	const FNodeTemplate* NodeTemplateB = FNodeTemplateBuilder::BuildNodeTemplate(NodeTemplateDecoratorListB, NodeTemplateBufferB);
-	const FNodeTemplate* NodeTemplateC = FNodeTemplateBuilder::BuildNodeTemplate(NodeTemplateDecoratorListC, NodeTemplateBufferC);
-
-	// Build our graph, it as follow (each node template has a single node instance):
-	// NodeA has no children
-	// NodeB has one child: NodeA (it doesn't update)
-	// NodeC (root) has two children: NodeA and NodeB
-
-	TArray<FNodeHandle> NodeHandles;
-
-	// Write our graph
-	TArray<uint8> GraphSharedDataArchiveBuffer;
-	TArray<TObjectPtr<UObject>> GraphReferencedObjects;
 	{
-		FDecoratorWriter DecoratorWriter;
-
-		NodeHandles.Add(DecoratorWriter.RegisterNode(*NodeTemplateC));	// Root node
-		NodeHandles.Add(DecoratorWriter.RegisterNode(*NodeTemplateA));
-		NodeHandles.Add(DecoratorWriter.RegisterNode(*NodeTemplateB));
-
-		// We don't have decorator properties
-		TArray<TMap<FString, FString>> DecoratorPropertiesA;
-		DecoratorPropertiesA.AddDefaulted(NodeTemplateDecoratorListA.Num());
-
-		TArray<TMap<FString, FString>> DecoratorPropertiesB;
-		DecoratorPropertiesB.AddDefaulted(NodeTemplateDecoratorListB.Num());
-		DecoratorPropertiesB[0].Add(TEXT("Child"), ToString<FDecoratorWithOneChild::FSharedData>(TEXT("Child"), FAnimNextDecoratorHandle(NodeHandles[1])));
-
-		TArray<TMap<FString, FString>> DecoratorPropertiesC;
-		DecoratorPropertiesC.AddDefaulted(NodeTemplateDecoratorListC.Num());
-		FAnimNextDecoratorHandle ChildrenHandlesC[2] = { FAnimNextDecoratorHandle(NodeHandles[1]), FAnimNextDecoratorHandle(NodeHandles[2])};
-		DecoratorPropertiesC[0].Add(TEXT("Children"), ToString<FDecoratorWithChildren::FSharedData>(TEXT("Children"), ChildrenHandlesC));
+		AUTO_REGISTER_ANIM_DECORATOR(FDecoratorWithNoChildren)
+		AUTO_REGISTER_ANIM_DECORATOR(FDecoratorWithOneChild)
+		AUTO_REGISTER_ANIM_DECORATOR(FDecoratorWithChildren)
+
+		UFactory* GraphFactory = NewObject<UAnimNextGraphFactory>();
+		UAnimNextGraph* AnimNextGraph = CastChecked<UAnimNextGraph>(GraphFactory->FactoryCreateNew(UAnimNextGraph::StaticClass(), GetTransientPackage(), TEXT("TestAnimNextGraph"), RF_Transient, nullptr, nullptr, NAME_None));
+		UE_RETURN_ON_ERROR(AnimNextGraph != nullptr, "FAnimationAnimNextRuntimeTest_IUpdate -> Failed to create animation graph");
+
+		FScopedClearNodeTemplateRegistry ScopedClearNodeTemplateRegistry;
+		FNodeTemplateRegistry& Registry = FNodeTemplateRegistry::Get();
+
+		// We create a few node templates
+		// Template A has a single decorator with no children
+		TArray<FDecoratorUID> NodeTemplateDecoratorListA;
+		NodeTemplateDecoratorListA.Add(FDecoratorWithNoChildren::DecoratorUID);
+
+		// Template B has a single decorator with one child, it doesn't update
+		TArray<FDecoratorUID> NodeTemplateDecoratorListB;
+		NodeTemplateDecoratorListB.Add(FDecoratorWithOneChild::DecoratorUID);
+
+		// Template C has a single decorator with children
+		TArray<FDecoratorUID> NodeTemplateDecoratorListC;
+		NodeTemplateDecoratorListC.Add(FDecoratorWithChildren::DecoratorUID);
+
+		// Populate our node template registry
+		TArray<uint8> NodeTemplateBufferA, NodeTemplateBufferB, NodeTemplateBufferC;
+		const FNodeTemplate* NodeTemplateA = FNodeTemplateBuilder::BuildNodeTemplate(NodeTemplateDecoratorListA, NodeTemplateBufferA);
+		const FNodeTemplate* NodeTemplateB = FNodeTemplateBuilder::BuildNodeTemplate(NodeTemplateDecoratorListB, NodeTemplateBufferB);
+		const FNodeTemplate* NodeTemplateC = FNodeTemplateBuilder::BuildNodeTemplate(NodeTemplateDecoratorListC, NodeTemplateBufferC);
+
+		// Build our graph, it as follow (each node template has a single node instance):
+		// NodeA has no children
+		// NodeB has one child: NodeA (it doesn't update)
+		// NodeC (root) has two children: NodeA and NodeB
+
+		TArray<FNodeHandle> NodeHandles;
+
+		// Write our graph
+		TArray<uint8> GraphSharedDataArchiveBuffer;
+		TArray<TObjectPtr<UObject>> GraphReferencedObjects;
+		{
+			FDecoratorWriter DecoratorWriter;
+
+			NodeHandles.Add(DecoratorWriter.RegisterNode(*NodeTemplateC));	// Root node
+			NodeHandles.Add(DecoratorWriter.RegisterNode(*NodeTemplateA));
+			NodeHandles.Add(DecoratorWriter.RegisterNode(*NodeTemplateB));
+
+			// We don't have decorator properties
+			TArray<TMap<FString, FString>> DecoratorPropertiesA;
+			DecoratorPropertiesA.AddDefaulted(NodeTemplateDecoratorListA.Num());
+
+			TArray<TMap<FString, FString>> DecoratorPropertiesB;
+			DecoratorPropertiesB.AddDefaulted(NodeTemplateDecoratorListB.Num());
+			DecoratorPropertiesB[0].Add(TEXT("Child"), ToString<FDecoratorWithOneChild::FSharedData>(TEXT("Child"), FAnimNextDecoratorHandle(NodeHandles[1])));
+
+			TArray<TMap<FString, FString>> DecoratorPropertiesC;
+			DecoratorPropertiesC.AddDefaulted(NodeTemplateDecoratorListC.Num());
+			FAnimNextDecoratorHandle ChildrenHandlesC[2] = { FAnimNextDecoratorHandle(NodeHandles[1]), FAnimNextDecoratorHandle(NodeHandles[2])};
+			DecoratorPropertiesC[0].Add(TEXT("Children"), ToString<FDecoratorWithChildren::FSharedData>(TEXT("Children"), ChildrenHandlesC));
+
+			DecoratorWriter.BeginNodeWriting();
+			DecoratorWriter.WriteNode(NodeHandles[0],
+				[&DecoratorPropertiesC](uint32 DecoratorIndex, const FString& PropertyName)
+				{
+					return DecoratorPropertiesC[DecoratorIndex][PropertyName];
+				},
+				[](uint32 DecoratorIndex, const FString& PropertyName)
+				{
+					return false;
+				});
+			DecoratorWriter.WriteNode(NodeHandles[1],
+				[&DecoratorPropertiesA](uint32 DecoratorIndex, const FString& PropertyName)
+				{
+					return DecoratorPropertiesA[DecoratorIndex][PropertyName];
+				},
+				[](uint32 DecoratorIndex, const FString& PropertyName)
+				{
+					return false;
+				});
+			DecoratorWriter.WriteNode(NodeHandles[2],
+				[&DecoratorPropertiesB](uint32 DecoratorIndex, const FString& PropertyName)
+				{
+					return DecoratorPropertiesB[DecoratorIndex][PropertyName];
+				},
+				[](uint32 DecoratorIndex, const FString& PropertyName)
+				{
+					return false;
+				});
+			DecoratorWriter.EndNodeWriting();
+
+			AddErrorIfFalse(DecoratorWriter.GetErrorState() == FDecoratorWriter::EErrorState::None, "FAnimationAnimNextRuntimeTest_IUpdate -> Failed to write decorators");
+			GraphSharedDataArchiveBuffer = DecoratorWriter.GetGraphSharedData();
+			GraphReferencedObjects = DecoratorWriter.GetGraphReferencedObjects();
+		}
 
-		DecoratorWriter.BeginNodeWriting();
-		DecoratorWriter.WriteNode(NodeHandles[0],
-			[&DecoratorPropertiesC](uint32 DecoratorIndex, const FString& PropertyName)
-			{
-				return DecoratorPropertiesC[DecoratorIndex][PropertyName];
-			},
-			[](uint32 DecoratorIndex, const FString& PropertyName)
-			{
-				return false;
-			});
-		DecoratorWriter.WriteNode(NodeHandles[1],
-			[&DecoratorPropertiesA](uint32 DecoratorIndex, const FString& PropertyName)
-			{
-				return DecoratorPropertiesA[DecoratorIndex][PropertyName];
-			},
-			[](uint32 DecoratorIndex, const FString& PropertyName)
-			{
-				return false;
-			});
-		DecoratorWriter.WriteNode(NodeHandles[2],
-			[&DecoratorPropertiesB](uint32 DecoratorIndex, const FString& PropertyName)
-			{
-				return DecoratorPropertiesB[DecoratorIndex][PropertyName];
-			},
-			[](uint32 DecoratorIndex, const FString& PropertyName)
-			{
-				return false;
-			});
-		DecoratorWriter.EndNodeWriting();
+		// Read our graph
+		FTestUtils::LoadFromArchiveBuffer(*AnimNextGraph, NodeHandles, GraphSharedDataArchiveBuffer);
 
-		AddErrorIfFalse(DecoratorWriter.GetErrorState() == FDecoratorWriter::EErrorState::None, "FAnimationAnimNextRuntimeTest_IUpdate -> Failed to write decorators");
-		GraphSharedDataArchiveBuffer = DecoratorWriter.GetGraphSharedData();
-		GraphReferencedObjects = DecoratorWriter.GetGraphReferencedObjects();
-	}
+		FAnimNextGraphInstancePtr GraphInstance;
+		AnimNextGraph->AllocateInstance(GraphInstance);
 
-	// Read our graph
-	FTestUtils::LoadFromArchiveBuffer(*AnimNextGraph, NodeHandles, GraphSharedDataArchiveBuffer);
+		FExecutionContext Context(GraphInstance);
 
-	FAnimNextGraphInstancePtr GraphInstance;
-	AnimNextGraph->AllocateInstance(GraphInstance);
+		{
+			TArray<FDecoratorUID> UpdatedDecorators;
+			Private::UpdatedDecorators = &UpdatedDecorators;
 
-	FExecutionContext Context(GraphInstance);
+			// Call pre/post update on our graph
+			UpdateGraph(GraphInstance, 0.0333f);
 
-	{
-		TArray<FDecoratorUID> UpdatedDecorators;
-		Private::UpdatedDecorators = &UpdatedDecorators;
+			AddErrorIfFalse(UpdatedDecorators.Num() == 6, "FAnimationAnimNextRuntimeTest_IUpdate -> Expected 6 nodes to have been visited during the update traversal");
+			AddErrorIfFalse(UpdatedDecorators[0] == FDecoratorWithChildren::DecoratorUID, "FAnimationAnimNextRuntimeTest_IUpdate -> Unexpected update order");		// NodeC
+			AddErrorIfFalse(UpdatedDecorators[1] == FDecoratorWithNoChildren::DecoratorUID, "FAnimationAnimNextRuntimeTest_IUpdate -> Unexpected update order");		// NodeA
+			AddErrorIfFalse(UpdatedDecorators[2] == FDecoratorWithNoChildren::DecoratorUID, "FAnimationAnimNextRuntimeTest_IUpdate -> Unexpected update order");		// NodeA
+			AddErrorIfFalse(UpdatedDecorators[3] == FDecoratorWithNoChildren::DecoratorUID, "FAnimationAnimNextRuntimeTest_IUpdate -> Unexpected update order");		// NodeB -> NodeA
+			AddErrorIfFalse(UpdatedDecorators[4] == FDecoratorWithNoChildren::DecoratorUID, "FAnimationAnimNextRuntimeTest_IUpdate -> Unexpected update order");		// NodeB -> NodeA
+			AddErrorIfFalse(UpdatedDecorators[5] == FDecoratorWithChildren::DecoratorUID, "FAnimationAnimNextRuntimeTest_IUpdate -> Unexpected update order");		// NodeC
 
-		// Call pre/post update on our graph
-		UpdateGraph(GraphInstance, 0.0333f);
+			Private::UpdatedDecorators = nullptr;
+		}
 
-		AddErrorIfFalse(UpdatedDecorators.Num() == 6, "FAnimationAnimNextRuntimeTest_IUpdate -> Expected 6 nodes to have been visited during the update traversal");
-		AddErrorIfFalse(UpdatedDecorators[0] == FDecoratorWithChildren::DecoratorUID, "FAnimationAnimNextRuntimeTest_IUpdate -> Unexpected update order");		// NodeC
-		AddErrorIfFalse(UpdatedDecorators[1] == FDecoratorWithNoChildren::DecoratorUID, "FAnimationAnimNextRuntimeTest_IUpdate -> Unexpected update order");		// NodeA
-		AddErrorIfFalse(UpdatedDecorators[2] == FDecoratorWithNoChildren::DecoratorUID, "FAnimationAnimNextRuntimeTest_IUpdate -> Unexpected update order");		// NodeA
-		AddErrorIfFalse(UpdatedDecorators[3] == FDecoratorWithNoChildren::DecoratorUID, "FAnimationAnimNextRuntimeTest_IUpdate -> Unexpected update order");		// NodeB -> NodeA
-		AddErrorIfFalse(UpdatedDecorators[4] == FDecoratorWithNoChildren::DecoratorUID, "FAnimationAnimNextRuntimeTest_IUpdate -> Unexpected update order");		// NodeB -> NodeA
-		AddErrorIfFalse(UpdatedDecorators[5] == FDecoratorWithChildren::DecoratorUID, "FAnimationAnimNextRuntimeTest_IUpdate -> Unexpected update order");		// NodeC
+		Registry.Unregister(NodeTemplateA);
+		Registry.Unregister(NodeTemplateB);
+		Registry.Unregister(NodeTemplateC);
 
-		Private::UpdatedDecorators = nullptr;
+		AddErrorIfFalse(Registry.GetNum() == 0, "FAnimationAnimNextRuntimeTest_IUpdate -> Registry should contain 0 templates");
 	}
-
-	Registry.Unregister(NodeTemplateA);
-	Registry.Unregister(NodeTemplateB);
-	Registry.Unregister(NodeTemplateC);
-
-	AddErrorIfFalse(Registry.GetNum() == 0, "FAnimationAnimNextRuntimeTest_IUpdate -> Registry should contain 0 templates");
+	Tests::FUtils::CleanupAfterTests();
 
 	return true;
 }
@@ -540,132 +547,134 @@ IMPLEMENT_SIMPLE_AUTOMATION_TEST(FAnimationAnimNextRuntimeTest_IEvaluate, "Anima
 bool FAnimationAnimNextRuntimeTest_IEvaluate::RunTest(const FString& InParameters)
 {
 	using namespace UE::AnimNext;
-
-	AUTO_REGISTER_ANIM_DECORATOR(FDecoratorWithNoChildren)
-	AUTO_REGISTER_ANIM_DECORATOR(FDecoratorWithOneChild)
-	AUTO_REGISTER_ANIM_DECORATOR(FDecoratorWithChildren)
-
-	UFactory* GraphFactory = NewObject<UAnimNextGraphFactory>();
-	UAnimNextGraph* AnimNextGraph = CastChecked<UAnimNextGraph>(GraphFactory->FactoryCreateNew(UAnimNextGraph::StaticClass(), GetTransientPackage(), TEXT("TestAnimNextGraph"), RF_Transient, nullptr, nullptr, NAME_None));
-	UE_RETURN_ON_ERROR(AnimNextGraph != nullptr, "FAnimationAnimNextRuntimeTest_IEvaluate -> Failed to create animation graph");
-
-	FScopedClearNodeTemplateRegistry ScopedClearNodeTemplateRegistry;
-	FNodeTemplateRegistry& Registry = FNodeTemplateRegistry::Get();
-
-	// We create a few node templates
-	// Template A has a single decorator with no children
-	TArray<FDecoratorUID> NodeTemplateDecoratorListA;
-	NodeTemplateDecoratorListA.Add(FDecoratorWithNoChildren::DecoratorUID);
-
-	// Template B has a single decorator with one child, it doesn't evaluate
-	TArray<FDecoratorUID> NodeTemplateDecoratorListB;
-	NodeTemplateDecoratorListB.Add(FDecoratorWithOneChild::DecoratorUID);
-
-	// Template C has a single decorator with children
-	TArray<FDecoratorUID> NodeTemplateDecoratorListC;
-	NodeTemplateDecoratorListC.Add(FDecoratorWithChildren::DecoratorUID);
-
-	// Populate our node template registry
-	TArray<uint8> NodeTemplateBufferA, NodeTemplateBufferB, NodeTemplateBufferC;
-	const FNodeTemplate* NodeTemplateA = FNodeTemplateBuilder::BuildNodeTemplate(NodeTemplateDecoratorListA, NodeTemplateBufferA);
-	const FNodeTemplate* NodeTemplateB = FNodeTemplateBuilder::BuildNodeTemplate(NodeTemplateDecoratorListB, NodeTemplateBufferB);
-	const FNodeTemplate* NodeTemplateC = FNodeTemplateBuilder::BuildNodeTemplate(NodeTemplateDecoratorListC, NodeTemplateBufferC);
-
-	// Build our graph, it as follow (each node template has a single node instance):
-	// NodeA has no children
-	// NodeB has one child: NodeA (it doesn't evaluate)
-	// NodeC (root) has two children: NodeA and NodeB
-
-	TArray<FNodeHandle> NodeHandles;
-
-	// Write our graph
-	TArray<uint8> GraphSharedDataArchiveBuffer;
-	TArray<TObjectPtr<UObject>> GraphReferencedObjects;
 	{
-		FDecoratorWriter DecoratorWriter;
-
-		NodeHandles.Add(DecoratorWriter.RegisterNode(*NodeTemplateC));	// Root node
-		NodeHandles.Add(DecoratorWriter.RegisterNode(*NodeTemplateA));
-		NodeHandles.Add(DecoratorWriter.RegisterNode(*NodeTemplateB));
-
-		// We don't have decorator properties
-		TArray<TMap<FString, FString>> DecoratorPropertiesA;
-		DecoratorPropertiesA.AddDefaulted(NodeTemplateDecoratorListA.Num());
-
-		TArray<TMap<FString, FString>> DecoratorPropertiesB;
-		DecoratorPropertiesB.AddDefaulted(NodeTemplateDecoratorListB.Num());
-		DecoratorPropertiesB[0].Add(TEXT("Child"), ToString<FDecoratorWithOneChild::FSharedData>(TEXT("Child"), FAnimNextDecoratorHandle(NodeHandles[1])));
+		AUTO_REGISTER_ANIM_DECORATOR(FDecoratorWithNoChildren)
+		AUTO_REGISTER_ANIM_DECORATOR(FDecoratorWithOneChild)
+		AUTO_REGISTER_ANIM_DECORATOR(FDecoratorWithChildren)
+
+		UFactory* GraphFactory = NewObject<UAnimNextGraphFactory>();
+		UAnimNextGraph* AnimNextGraph = CastChecked<UAnimNextGraph>(GraphFactory->FactoryCreateNew(UAnimNextGraph::StaticClass(), GetTransientPackage(), TEXT("TestAnimNextGraph"), RF_Transient, nullptr, nullptr, NAME_None));
+		UE_RETURN_ON_ERROR(AnimNextGraph != nullptr, "FAnimationAnimNextRuntimeTest_IEvaluate -> Failed to create animation graph");
+
+		FScopedClearNodeTemplateRegistry ScopedClearNodeTemplateRegistry;
+		FNodeTemplateRegistry& Registry = FNodeTemplateRegistry::Get();
+
+		// We create a few node templates
+		// Template A has a single decorator with no children
+		TArray<FDecoratorUID> NodeTemplateDecoratorListA;
+		NodeTemplateDecoratorListA.Add(FDecoratorWithNoChildren::DecoratorUID);
+
+		// Template B has a single decorator with one child, it doesn't evaluate
+		TArray<FDecoratorUID> NodeTemplateDecoratorListB;
+		NodeTemplateDecoratorListB.Add(FDecoratorWithOneChild::DecoratorUID);
+
+		// Template C has a single decorator with children
+		TArray<FDecoratorUID> NodeTemplateDecoratorListC;
+		NodeTemplateDecoratorListC.Add(FDecoratorWithChildren::DecoratorUID);
+
+		// Populate our node template registry
+		TArray<uint8> NodeTemplateBufferA, NodeTemplateBufferB, NodeTemplateBufferC;
+		const FNodeTemplate* NodeTemplateA = FNodeTemplateBuilder::BuildNodeTemplate(NodeTemplateDecoratorListA, NodeTemplateBufferA);
+		const FNodeTemplate* NodeTemplateB = FNodeTemplateBuilder::BuildNodeTemplate(NodeTemplateDecoratorListB, NodeTemplateBufferB);
+		const FNodeTemplate* NodeTemplateC = FNodeTemplateBuilder::BuildNodeTemplate(NodeTemplateDecoratorListC, NodeTemplateBufferC);
+
+		// Build our graph, it as follow (each node template has a single node instance):
+		// NodeA has no children
+		// NodeB has one child: NodeA (it doesn't evaluate)
+		// NodeC (root) has two children: NodeA and NodeB
+
+		TArray<FNodeHandle> NodeHandles;
+
+		// Write our graph
+		TArray<uint8> GraphSharedDataArchiveBuffer;
+		TArray<TObjectPtr<UObject>> GraphReferencedObjects;
+		{
+			FDecoratorWriter DecoratorWriter;
+
+			NodeHandles.Add(DecoratorWriter.RegisterNode(*NodeTemplateC));	// Root node
+			NodeHandles.Add(DecoratorWriter.RegisterNode(*NodeTemplateA));
+			NodeHandles.Add(DecoratorWriter.RegisterNode(*NodeTemplateB));
+
+			// We don't have decorator properties
+			TArray<TMap<FString, FString>> DecoratorPropertiesA;
+			DecoratorPropertiesA.AddDefaulted(NodeTemplateDecoratorListA.Num());
+
+			TArray<TMap<FString, FString>> DecoratorPropertiesB;
+			DecoratorPropertiesB.AddDefaulted(NodeTemplateDecoratorListB.Num());
+			DecoratorPropertiesB[0].Add(TEXT("Child"), ToString<FDecoratorWithOneChild::FSharedData>(TEXT("Child"), FAnimNextDecoratorHandle(NodeHandles[1])));
+
+			TArray<TMap<FString, FString>> DecoratorPropertiesC;
+			DecoratorPropertiesC.AddDefaulted(NodeTemplateDecoratorListC.Num());
+
+			FAnimNextDecoratorHandle ChildrenHandlesC[2] = { FAnimNextDecoratorHandle(NodeHandles[1]), FAnimNextDecoratorHandle(NodeHandles[2])};
+			DecoratorPropertiesC[0].Add(TEXT("Children"), ToString<FDecoratorWithChildren::FSharedData>(TEXT("Children"), ChildrenHandlesC));
+
+			DecoratorWriter.BeginNodeWriting();
+			DecoratorWriter.WriteNode(NodeHandles[0],
+				[&DecoratorPropertiesC](uint32 DecoratorIndex, const FString& PropertyName)
+				{
+					return DecoratorPropertiesC[DecoratorIndex][PropertyName];
+				},
+				[](uint32 DecoratorIndex, const FString& PropertyName)
+				{
+					return false;
+				});
+			DecoratorWriter.WriteNode(NodeHandles[1],
+				[&DecoratorPropertiesA](uint32 DecoratorIndex, const FString& PropertyName)
+				{
+					return DecoratorPropertiesA[DecoratorIndex][PropertyName];
+				},
+				[](uint32 DecoratorIndex, const FString& PropertyName)
+				{
+					return false;
+				});
+			DecoratorWriter.WriteNode(NodeHandles[2],
+				[&DecoratorPropertiesB](uint32 DecoratorIndex, const FString& PropertyName)
+				{
+					return DecoratorPropertiesB[DecoratorIndex][PropertyName];
+				},
+				[](uint32 DecoratorIndex, const FString& PropertyName)
+				{
+					return false;
+				});
+			DecoratorWriter.EndNodeWriting();
+
+			AddErrorIfFalse(DecoratorWriter.GetErrorState() == FDecoratorWriter::EErrorState::None, "FAnimationAnimNextRuntimeTest_IEvaluate -> Failed to write decorators");
+			GraphSharedDataArchiveBuffer = DecoratorWriter.GetGraphSharedData();
+			GraphReferencedObjects = DecoratorWriter.GetGraphReferencedObjects();
+		}
 
-		TArray<TMap<FString, FString>> DecoratorPropertiesC;
-		DecoratorPropertiesC.AddDefaulted(NodeTemplateDecoratorListC.Num());
+		// Read our graph
+		FTestUtils::LoadFromArchiveBuffer(*AnimNextGraph, NodeHandles, GraphSharedDataArchiveBuffer);
 
-		FAnimNextDecoratorHandle ChildrenHandlesC[2] = { FAnimNextDecoratorHandle(NodeHandles[1]), FAnimNextDecoratorHandle(NodeHandles[2])};
-		DecoratorPropertiesC[0].Add(TEXT("Children"), ToString<FDecoratorWithChildren::FSharedData>(TEXT("Children"), ChildrenHandlesC));
+		FAnimNextGraphInstancePtr GraphInstance;
+		AnimNextGraph->AllocateInstance(GraphInstance);
 
-		DecoratorWriter.BeginNodeWriting();
-		DecoratorWriter.WriteNode(NodeHandles[0],
-			[&DecoratorPropertiesC](uint32 DecoratorIndex, const FString& PropertyName)
-			{
-				return DecoratorPropertiesC[DecoratorIndex][PropertyName];
-			},
-			[](uint32 DecoratorIndex, const FString& PropertyName)
-			{
-				return false;
-			});
-		DecoratorWriter.WriteNode(NodeHandles[1],
-			[&DecoratorPropertiesA](uint32 DecoratorIndex, const FString& PropertyName)
-			{
-				return DecoratorPropertiesA[DecoratorIndex][PropertyName];
-			},
-			[](uint32 DecoratorIndex, const FString& PropertyName)
-			{
-				return false;
-			});
-		DecoratorWriter.WriteNode(NodeHandles[2],
-			[&DecoratorPropertiesB](uint32 DecoratorIndex, const FString& PropertyName)
-			{
-				return DecoratorPropertiesB[DecoratorIndex][PropertyName];
-			},
-			[](uint32 DecoratorIndex, const FString& PropertyName)
-			{
-				return false;
-			});
-		DecoratorWriter.EndNodeWriting();
-
-		AddErrorIfFalse(DecoratorWriter.GetErrorState() == FDecoratorWriter::EErrorState::None, "FAnimationAnimNextRuntimeTest_IEvaluate -> Failed to write decorators");
-		GraphSharedDataArchiveBuffer = DecoratorWriter.GetGraphSharedData();
-		GraphReferencedObjects = DecoratorWriter.GetGraphReferencedObjects();
-	}
+		{
+			TArray<FDecoratorUID> EvaluatedDecorators;
+			Private::EvaluatedDecorators = &EvaluatedDecorators;
 
-	// Read our graph
-	FTestUtils::LoadFromArchiveBuffer(*AnimNextGraph, NodeHandles, GraphSharedDataArchiveBuffer);
+			// Call pre/post evaluate on our graph
+			(void)EvaluateGraph(GraphInstance);
 
-	FAnimNextGraphInstancePtr GraphInstance;
-	AnimNextGraph->AllocateInstance(GraphInstance);
+			AddErrorIfFalse(EvaluatedDecorators.Num() == 6, "FAnimationAnimNextRuntimeTest_IEvaluate -> Expected 6 nodes to have been visited during the evaluate traversal");
+			AddErrorIfFalse(EvaluatedDecorators[0] == FDecoratorWithChildren::DecoratorUID, "FAnimationAnimNextRuntimeTest_IEvaluate -> Unexpected evaluate order");		// NodeC
+			AddErrorIfFalse(EvaluatedDecorators[1] == FDecoratorWithNoChildren::DecoratorUID, "FAnimationAnimNextRuntimeTest_IEvaluate -> Unexpected evaluate order");			// NodeA
+			AddErrorIfFalse(EvaluatedDecorators[2] == FDecoratorWithNoChildren::DecoratorUID, "FAnimationAnimNextRuntimeTest_IEvaluate -> Unexpected evaluate order");			// NodeA
+			AddErrorIfFalse(EvaluatedDecorators[3] == FDecoratorWithNoChildren::DecoratorUID, "FAnimationAnimNextRuntimeTest_IEvaluate -> Unexpected evaluate order");			// NodeB -> NodeA
+			AddErrorIfFalse(EvaluatedDecorators[4] == FDecoratorWithNoChildren::DecoratorUID, "FAnimationAnimNextRuntimeTest_IEvaluate -> Unexpected evaluate order");			// NodeB -> NodeA
+			AddErrorIfFalse(EvaluatedDecorators[5] == FDecoratorWithChildren::DecoratorUID, "FAnimationAnimNextRuntimeTest_IEvaluate -> Unexpected evaluate order");		// NodeC
 
-	{
-		TArray<FDecoratorUID> EvaluatedDecorators;
-		Private::EvaluatedDecorators = &EvaluatedDecorators;
-
-		// Call pre/post evaluate on our graph
-		(void)EvaluateGraph(GraphInstance);
+			Private::EvaluatedDecorators = nullptr;
+		}
 
-		AddErrorIfFalse(EvaluatedDecorators.Num() == 6, "FAnimationAnimNextRuntimeTest_IEvaluate -> Expected 6 nodes to have been visited during the evaluate traversal");
-		AddErrorIfFalse(EvaluatedDecorators[0] == FDecoratorWithChildren::DecoratorUID, "FAnimationAnimNextRuntimeTest_IEvaluate -> Unexpected evaluate order");		// NodeC
-		AddErrorIfFalse(EvaluatedDecorators[1] == FDecoratorWithNoChildren::DecoratorUID, "FAnimationAnimNextRuntimeTest_IEvaluate -> Unexpected evaluate order");			// NodeA
-		AddErrorIfFalse(EvaluatedDecorators[2] == FDecoratorWithNoChildren::DecoratorUID, "FAnimationAnimNextRuntimeTest_IEvaluate -> Unexpected evaluate order");			// NodeA
-		AddErrorIfFalse(EvaluatedDecorators[3] == FDecoratorWithNoChildren::DecoratorUID, "FAnimationAnimNextRuntimeTest_IEvaluate -> Unexpected evaluate order");			// NodeB -> NodeA
-		AddErrorIfFalse(EvaluatedDecorators[4] == FDecoratorWithNoChildren::DecoratorUID, "FAnimationAnimNextRuntimeTest_IEvaluate -> Unexpected evaluate order");			// NodeB -> NodeA
-		AddErrorIfFalse(EvaluatedDecorators[5] == FDecoratorWithChildren::DecoratorUID, "FAnimationAnimNextRuntimeTest_IEvaluate -> Unexpected evaluate order");		// NodeC
+		Registry.Unregister(NodeTemplateA);
+		Registry.Unregister(NodeTemplateB);
+		Registry.Unregister(NodeTemplateC);
 
-		Private::EvaluatedDecorators = nullptr;
+		AddErrorIfFalse(Registry.GetNum() == 0, "FAnimationAnimNextRuntimeTest_IEvaluate -> Registry should contain 0 templates");
 	}
-
-	Registry.Unregister(NodeTemplateA);
-	Registry.Unregister(NodeTemplateB);
-	Registry.Unregister(NodeTemplateC);
-
-	AddErrorIfFalse(Registry.GetNum() == 0, "FAnimationAnimNextRuntimeTest_IEvaluate -> Registry should contain 0 templates");
+	Tests::FUtils::CleanupAfterTests();
 
 	return true;
 }
diff --git a/Engine/Plugins/Experimental/AnimNext/Source/AnimNextTestSuite/Private/AnimNextEditorTest.cpp b/Engine/Plugins/Experimental/AnimNext/Source/AnimNextTestSuite/Private/AnimNextEditorTest.cpp
--- a/Engine/Plugins/Experimental/AnimNext/Source/AnimNextTestSuite/Private/AnimNextEditorTest.cpp
+++ b/Engine/Plugins/Experimental/AnimNext/Source/AnimNextTestSuite/Private/AnimNextEditorTest.cpp
@@ -1,13 +1,13 @@
 // Copyright Epic Games, Inc. All Rights Reserved.
 
 #include "CoreMinimal.h"
+#include "AnimNextTest.h"
 #include "UncookedOnlyUtils.h"
 #include "Param/AnimNextParameterBlock_EditorData.h"
 #include "Misc/AutomationTest.h"
 #include "Param/ParameterBlockFactory.h"
 #include "Param/AnimNextParameterBlock.h"
 #include "Param/AnimNextParameterBlockParameter.h"
-#include "Param/AnimNextParameterBlockBinding.h"
 #include "Animation/AnimSequence.h"
 #if WITH_EDITOR
 #include "ScopedTransaction.h"
@@ -43,16 +43,17 @@ bool FAnimationAnimNextParametersEditorTest_Block::RunTest(const FString& InPara
 	static FName TestParameterName = TEXT("TestParam");
 	
 	// AddParameter
+	UAnimNextParameterBlockParameter* Parameter = nullptr;
 	{
 		FScopedTransaction Transaction(FText::GetEmpty());		
-		UAnimNextParameterBlockParameter* Parameter = EditorData->AddParameter(TestParameterName, FAnimNextParamType::GetType<bool>());
+		Parameter = EditorData->AddParameter(TestParameterName, FAnimNextParamType::GetType<bool>());
 
 		if (AddErrorIfFalse(Parameter != nullptr, TEXT("Could not create new parameter in block.")))
 		{
 			AddErrorIfFalse(Parameter->GetParamType() == FAnimNextParamType::GetType<bool>(), TEXT("Incorrect parameter type found"));			
 		}
 	}
-		
+
 	AddExpectedError(TEXT("UAnimNextParameterBlock_EditorData::AddParameter: A parameter already exists for the supplied parameter name."));
 	AddErrorIfFalse(EditorData->AddParameter(TestParameterName, FAnimNextParamType::GetType<bool>()) == nullptr, TEXT("Expected duplicate parameter name argument to fail"));
 
@@ -61,9 +62,7 @@ bool FAnimationAnimNextParametersEditorTest_Block::RunTest(const FString& InPara
 
 	GEditor->RedoTransaction();
 	AddErrorIfFalse(EditorData->Entries.Num() == 1, FString::Printf(TEXT("Unexpected entry count found in parameter block (Have %d, expected 1)."), EditorData->Entries.Num()));
-	
-	GEditor->UndoTransaction();
-	
+
 	// Failure cases
 	AddExpectedError(TEXT("UAnimNextParameterBlock_EditorData::AddParameter: Invalid parameter name supplied."));
 	AddErrorIfFalse(EditorData->AddParameter(NAME_None, FAnimNextParamType::GetType<bool>()) == nullptr, TEXT("Expected invalid argument to fail"));
@@ -76,9 +75,9 @@ bool FAnimationAnimNextParametersEditorTest_Block::RunTest(const FString& InPara
 		{
 			AddErrorIfFalse(TypedParameter->GetParamType() == InType, TEXT("Incorrect parameter type found"));
 			EditorData->RemoveEntry(TypedParameter);
-		}		
+		}
 	};
-	
+
 	// Various types
 	TestParameterType(FAnimNextParamType::GetType<bool>());
 	TestParameterType(FAnimNextParamType::GetType<uint8>());
@@ -98,39 +97,10 @@ bool FAnimationAnimNextParametersEditorTest_Block::RunTest(const FString& InPara
 	TestParameterType(FAnimNextParamType::GetType<TArray<float>>());
 	TestParameterType(FAnimNextParamType::GetType<TArray<TObjectPtr<UAnimSequence>>>());
 
-	// AddBinding
-	static FName TestBindingName = TEXT("TestParamBinding");
-	UAnimNextParameterBlockBinding* Binding = nullptr;
-	{
-		FScopedTransaction Transaction(FText::GetEmpty());
-		Binding = EditorData->AddBinding(TestBindingName, FAnimNextParamType::GetType<int32>());
-		AddErrorIfFalse(Binding != nullptr, TEXT("Could not create new binding in block."));
-	}
-
-	GEditor->UndoTransaction();
-	AddErrorIfFalse(EditorData->Entries.Num() == 0, FString::Printf(TEXT("Unexpected entry count found in parameter block (Have %d, expected 0)."), EditorData->Entries.Num()));
-
-	GEditor->RedoTransaction();
-
-	// RemoveAllBindings
-	{
-		FScopedTransaction Transaction(FText::GetEmpty());
-		AddErrorIfFalse(EditorData->RemoveAllBindings(TestBindingName), TEXT("Failed to remove binding."));
-	}
-
-	GEditor->UndoTransaction();
-	AddErrorIfFalse(EditorData->Entries.Num() == 1, FString::Printf(TEXT("Unexpected entry count found in parameter block (Have %d, expected 1)."), EditorData->Entries.Num()));
-
-	GEditor->RedoTransaction();
-	AddErrorIfFalse(EditorData->Entries.Num() == 0, FString::Printf(TEXT("Unexpected entry count found in parameter block (Have %d, expected 0)."), EditorData->Entries.Num()));
-
-	GEditor->UndoTransaction();
-	AddErrorIfFalse(EditorData->Entries.Num() == 1, FString::Printf(TEXT("Unexpected entry count found in parameter block (Have %d, expected 1)."), EditorData->Entries.Num()));
-
 	// RemoveEntry
 	{
 		FScopedTransaction Transaction(FText::GetEmpty());
-		AddErrorIfFalse(EditorData->RemoveEntry(Binding), TEXT("Failed to remove entry."));
+		AddErrorIfFalse(EditorData->RemoveEntry(Parameter), TEXT("Failed to remove entry."));
 	}
 
 	GEditor->UndoTransaction();
@@ -142,47 +112,11 @@ bool FAnimationAnimNextParametersEditorTest_Block::RunTest(const FString& InPara
 	GEditor->UndoTransaction();
 	AddErrorIfFalse(EditorData->Entries.Num() == 1, FString::Printf(TEXT("Unexpected entry count found in parameter block (Have %d, expected 1)."), EditorData->Entries.Num()));
 
-	// FindBinding
-	AddErrorIfFalse(EditorData->FindBinding(TestBindingName) != nullptr, TEXT("Could not find binding in block."));
-
-	const UAnimNextParameterBlock* OtherBlock = Cast<UAnimNextParameterBlock>(BlockFactory->FactoryCreateNew(UAnimNextParameterBlock::StaticClass(), GetTransientPackage(), TEXT("TestAnimNextParameterBlock2"), RF_Transient, nullptr, nullptr, NAME_None));
-	if(OtherBlock == nullptr)
-	{
-		AddError(TEXT("Could not create additional parameter block."));
-		return false;
-	}
-	
-	UAnimNextParameterBlock_EditorData* OtherEditorData = UncookedOnly::FUtils::GetEditorData(OtherBlock);
-	if(OtherEditorData == nullptr)
-	{
-		AddError(TEXT("Additional parameter block has no editor data."));
-		return false;
-	}
-
-	// AddBindingReference
-	{
-		FScopedTransaction Transaction(FText::GetEmpty());
-		UAnimNextParameterBlockBindingReference* Ref = OtherEditorData->AddBindingReference(TestBindingName, Block);
-		AddErrorIfFalse(Ref != nullptr, TEXT("Could not create new binding reference in block."));
-	}
-
+	// FindEntry
+	AddErrorIfFalse(EditorData->FindEntry(TestParameterName) != nullptr, TEXT("Could not find entry in block."));
 	GEditor->UndoTransaction();
-	AddErrorIfFalse(OtherEditorData->Entries.Num() == 0, FString::Printf(TEXT("Unexpected entry count found in parameter block (Have %d, expected 0)."), OtherEditorData->Entries.Num()));
-
-	GEditor->RedoTransaction();
-	AddErrorIfFalse(OtherEditorData->Entries.Num() == 1, FString::Printf(TEXT("Unexpected entry count found in parameter block (Have %d, expected 1)."), OtherEditorData->Entries.Num()));
 
-	// FindBinding
-	AddErrorIfFalse(OtherEditorData->FindBinding(TestBindingName) != nullptr, TEXT("Could not find binding refernce in block."));
-
-	GEditor->UndoTransaction();
-	GEditor->UndoTransaction();
-	AddErrorIfFalse(EditorData->Entries.Num() == 0, FString::Printf(TEXT("Unexpected entry count found in parameter block (Have %d, expected 0)."), EditorData->Entries.Num()));
-
-	GEditor->RedoTransaction();
-	AddErrorIfFalse(EditorData->Entries.Num() == 1, FString::Printf(TEXT("Unexpected entry count found in parameter block (Have %d, expected 1)."), EditorData->Entries.Num()));
-	GEditor->UndoTransaction();
-	
+	// Add graph
 	UAnimNextParameterBlockGraph* Graph = nullptr;
 	{
 		FScopedTransaction Transaction(FText::GetEmpty());
@@ -197,7 +131,7 @@ bool FAnimationAnimNextParametersEditorTest_Block::RunTest(const FString& InPara
 	AddErrorIfFalse(EditorData->Entries.Num() == 1, FString::Printf(TEXT("Unexpected entry count found in parameter block (Have %d, expected 1)."), EditorData->Entries.Num()));
 	GEditor->UndoTransaction();
 
-	CollectGarbage(GARBAGE_COLLECTION_KEEPFLAGS);
+	Tests::FUtils::CleanupAfterTests();
 
 	return true;
 }
@@ -211,18 +145,14 @@ bool FAnimationAnimNextParametersEditorTest_Python::RunTest(const FString& InPar
 	const TCHAR* Script = TEXT(
 		"asset_tools = unreal.AssetToolsHelpers.get_asset_tools()\n"
 		"block = unreal.AssetTools.create_asset(asset_tools, asset_name = \"TestBlock\", package_path = \"/Game/\", asset_class = unreal.AnimNextParameterBlock, factory = unreal.AnimNextParameterBlockFactory())\n"
-		"other_block = unreal.AssetTools.create_asset(asset_tools, asset_name = \"TestBlock1\", package_path = \"/Game/\", asset_class = unreal.AnimNextParameterBlock, factory = unreal.AnimNextParameterBlockFactory())\n"
 		"block.add_parameter(name = \"TestParam\", value_type = unreal.PropertyBagPropertyType.BOOL, container_type = unreal.PropertyBagContainerType.NONE)\n"
-		"block.add_binding(name = \"TestParamBinding\", value_type = unreal.PropertyBagPropertyType.BOOL, container_type = unreal.PropertyBagContainerType.NONE)\n"
-		"other_block.add_binding_reference(name = \"TestParamBinding\", referenced_block = block)\n"
 		"block.add_graph(name = \"TestGraph\")\n"
 		"unreal.EditorAssetLibrary.delete_loaded_asset(block)\n"
-		"unreal.EditorAssetLibrary.delete_loaded_asset(other_block)\n"
 	);
 
 	IPythonScriptPlugin::Get()->ExecPythonCommand(Script);
 
-	CollectGarbage(GARBAGE_COLLECTION_KEEPFLAGS);
+	Tests::FUtils::CleanupAfterTests();
 
 	return true;
 }
diff --git a/Engine/Plugins/Experimental/AnimNext/Source/AnimNextTestSuite/Private/AnimNextTest.cpp b/Engine/Plugins/Experimental/AnimNext/Source/AnimNextTestSuite/Private/AnimNextTest.cpp
--- a/Engine/Plugins/Experimental/AnimNext/Source/AnimNextTestSuite/Private/AnimNextTest.cpp
+++ b/Engine/Plugins/Experimental/AnimNext/Source/AnimNextTestSuite/Private/AnimNextTest.cpp
@@ -8,6 +8,10 @@
 #include "DataRegistry.h"
 #include "ReferencePose.h"
 #include <chrono>
+#if WITH_EDITOR
+#include "Editor.h"
+#include "Editor/Transactor.h"
+#endif
 
 #if WITH_DEV_AUTOMATION_TESTS
 
@@ -39,6 +43,17 @@ public:
 
 }
 
+void FUtils::CleanupAfterTests()
+{
+#if WITH_EDITOR
+	if(GEditor && GEditor->Trans)
+	{
+		GEditor->Trans->Reset(NSLOCTEXT("AnimNextTests", "CleanupAfterTest", "Cleaning up after AnimNext test"));
+	}
+#endif
+	CollectGarbage(GARBAGE_COLLECTION_KEEPFLAGS);
+}
+
 IMPLEMENT_SIMPLE_AUTOMATION_TEST(FDataRegistryTransformsTest, "Animation.AnimNext.AnimationDataRegistry.TransformsTest", EAutomationTestFlags::EditorContext | EAutomationTestFlags::EngineFilter)
 
 bool FDataRegistryTransformsTest::RunTest(const FString& InParameters)
diff --git a/Engine/Plugins/Experimental/AnimNext/Source/AnimNextUncookedOnly/Internal/AnimNextRigVMAssetEditorData.h b/Engine/Plugins/Experimental/AnimNext/Source/AnimNextUncookedOnly/Internal/AnimNextRigVMAssetEditorData.h
new file mode 100644
--- /dev/null
+++ b/Engine/Plugins/Experimental/AnimNext/Source/AnimNextUncookedOnly/Internal/AnimNextRigVMAssetEditorData.h
@@ -0,0 +1,266 @@
+﻿// Copyright Epic Games, Inc. All Rights Reserved.
+
+#pragma once
+
+#include "CoreMinimal.h"
+#include "RigVMModel/RigVMGraph.h"
+#include "RigVMCore/RigVMGraphFunctionHost.h"
+#include "RigVMBlueprint.h"
+#include "Kismet/BlueprintFunctionLibrary.h"
+#include "AnimNextRigVMAssetEditorData.generated.h"
+
+enum class ERigVMGraphNotifType : uint8;
+class UAnimNextRigVMAssetEntry;
+class UAnimNextRigVMAssetEditorData;
+
+namespace UE::AnimNext::UncookedOnly
+{
+	struct FUtils;
+	struct FUtilsPrivate;
+}
+
+namespace UE::AnimNext::Editor
+{
+	struct FUtils;
+	class FParameterBlockEditor;
+	class SRigVMAssetView;
+	class SParameterPicker;
+	class FParameterBlockTabSummoner;
+	class SRigVMAssetViewRow;
+	class FParameterBlockParameterCustomization;
+	class FModule;
+	class FWorkspaceEditor;
+}
+
+namespace UE::AnimNext::UncookedOnly
+{
+	// A delegate for subscribing / reacting to editor data modifications.
+	DECLARE_MULTICAST_DELEGATE_OneParam(FOnEditorDataModified, UAnimNextRigVMAssetEditorData* /* InEditorData */);
+}
+
+// Script-callable editor API hoisted onto UAnimNextRigVMAsset
+UCLASS()
+class UAnimNextRigVMAssetLibrary : public UBlueprintFunctionLibrary
+{
+	GENERATED_BODY()
+
+	UFUNCTION(BlueprintCallable, Category = "AnimNext|RigVM Asset", meta=(ScriptMethod))
+	static UAnimNextRigVMAssetEntry* FindEntry(UAnimNextRigVMAsset* InAsset, FName InName);
+
+	UFUNCTION(BlueprintCallable, Category = "AnimNext|RigVM Asset", meta=(ScriptMethod))
+	static bool RemoveEntry(UAnimNextRigVMAsset* InAsset, UAnimNextRigVMAssetEntry* InEntry, bool bSetupUndoRedo = true, bool bPrintPythonCommand = true);
+
+	UFUNCTION(BlueprintCallable, Category = "AnimNext|RigVM Asset", meta=(ScriptMethod))
+	static bool RemoveEntries(UAnimNextRigVMAsset* InAsset, const TArray<UAnimNextRigVMAssetEntry*>& InEntries, bool bSetupUndoRedo = true, bool bPrintPythonCommand = true);
+};
+
+/* Base class for all AnimNext editor data objects that use RigVM */
+UCLASS(MinimalAPI, Abstract)
+class UAnimNextRigVMAssetEditorData : public UObject, public IRigVMClientHost, public IRigVMGraphFunctionHost
+{
+	GENERATED_BODY()
+
+protected:
+	friend class UE::AnimNext::Editor::SRigVMAssetView;
+	friend class UE::AnimNext::Editor::SRigVMAssetViewRow;
+	friend struct UE::AnimNext::Editor::FUtils;
+	friend struct UE::AnimNext::UncookedOnly::FUtils;
+	friend class UE::AnimNext::Editor::FModule;
+	friend class UE::AnimNext::Editor::FWorkspaceEditor;
+	friend class UAnimNextRigVMAssetEntry;
+	friend class UAnimNextRigVMAssetLibrary;
+
+	// UObject interface
+	virtual void Serialize(FArchive& Ar) override;
+	virtual void PostLoad() override;
+	virtual void PostTransacted(const FTransactionObjectEvent& TransactionEvent) override;
+	virtual void PostDuplicate(EDuplicateMode::Type DuplicateMode) override;
+	virtual void GetAssetRegistryTags(FAssetRegistryTagsContext Context) const override;
+	virtual bool IsEditorOnly() const override { return true; }
+
+	void HandlePackageDone(const FEndLoadPackageContext& Context);
+	void HandlePackageDone();
+
+	// IRigVMClientHost interface
+	virtual FRigVMClient* GetRigVMClient() override;
+	virtual const FRigVMClient* GetRigVMClient() const override;
+	virtual IRigVMGraphFunctionHost* GetRigVMGraphFunctionHost() override;
+	virtual const IRigVMGraphFunctionHost* GetRigVMGraphFunctionHost() const override;
+	virtual void HandleRigVMGraphAdded(const FRigVMClient* InClient, const FString& InNodePath) override;
+	virtual void HandleRigVMGraphRemoved(const FRigVMClient* InClient, const FString& InNodePath) override;
+	virtual void HandleRigVMGraphRenamed(const FRigVMClient* InClient, const FString& InOldNodePath, const FString& InNewNodePath) override {}
+	virtual void HandleConfigureRigVMController(const FRigVMClient* InClient, URigVMController* InControllerToConfigure) override;
+	virtual UObject* GetEditorObjectForRigVMGraph(URigVMGraph* InVMGraph) const override;
+	virtual URigVMGraph* GetRigVMGraphForEditorObject(UObject* InObject) const override;
+	virtual void RecompileVM() override PURE_VIRTUAL(UAnimNextRigVMAssetEditorData::RecompileVM, )
+	virtual void RecompileVMIfRequired() override;
+	virtual void RequestAutoVMRecompilation() override;
+	virtual void SetAutoVMRecompile(bool bAutoRecompile) override;
+	virtual bool GetAutoVMRecompile() const override;
+	virtual void IncrementVMRecompileBracket() override;
+	virtual void DecrementVMRecompileBracket() override;
+	virtual void RefreshAllModels(ERigVMLoadType InLoadType) override;
+	virtual void OnRigVMRegistryChanged() override;
+	virtual void RequestRigVMInit() override;
+	virtual URigVMGraph* GetModel(const UEdGraph* InEdGraph = nullptr) const override;
+	virtual URigVMGraph* GetModel(const FString& InNodePath) const override;
+	virtual URigVMGraph* GetDefaultModel() const override;
+	virtual TArray<URigVMGraph*> GetAllModels() const override;
+	virtual URigVMFunctionLibrary* GetLocalFunctionLibrary() const override;
+	virtual URigVMGraph* AddModel(FString InName = TEXT("Rig Graph"), bool bSetupUndoRedo = true, bool bPrintPythonCommand = true) override;
+	virtual bool RemoveModel(FString InName = TEXT("Rig Graph"), bool bSetupUndoRedo = true, bool bPrintPythonCommand = true) override;
+	virtual FRigVMGetFocusedGraph& OnGetFocusedGraph() override;
+	virtual const FRigVMGetFocusedGraph& OnGetFocusedGraph() const override;
+	virtual URigVMGraph* GetFocusedModel() const override;
+	virtual URigVMController* GetController(const URigVMGraph* InGraph = nullptr) const override;
+	virtual URigVMController* GetControllerByName(const FString InGraphName = TEXT("")) const override;
+	virtual URigVMController* GetOrCreateController(URigVMGraph* InGraph = nullptr) override;
+	virtual URigVMController* GetController(const UEdGraph* InEdGraph) const override;
+	virtual URigVMController* GetOrCreateController(const UEdGraph* InGraph) override;
+	virtual TArray<FString> GeneratePythonCommands(const FString InNewBlueprintName) override;
+
+	// IRigVMGraphFunctionHost interface
+	virtual FRigVMGraphFunctionStore* GetRigVMGraphFunctionStore() override;
+	virtual const FRigVMGraphFunctionStore* GetRigVMGraphFunctionStore() const override;
+
+	// Override called during initialization to determine what RigVM controller class is used
+	virtual TSubclassOf<URigVMController> GetControllerClass() const { return URigVMController::StaticClass(); }
+
+	// Override called during initialization to determine what RigVM schema class is used
+	virtual TSubclassOf<URigVMSchema> GetRigVMSchemaClass() const PURE_VIRTUAL(UAnimNextRigVMAssetEditorData::GetRigVMSchemaClass, return nullptr;)
+
+	// Override called during initialization to determine what RigVM execute struct is used
+	virtual UScriptStruct* GetExecuteContextStruct() const PURE_VIRTUAL(UAnimNextRigVMAssetEditorData::GetExecuteContextStruct, return nullptr;)
+
+	// Create and store a UEdGraph that corresponds to a URigVMGraph
+	virtual UEdGraph* CreateEdGraph(URigVMGraph* InRigVMGraph, bool bForce) PURE_VIRTUAL(UAnimNextRigVMAssetEditorData::CreateEdGraph, return nullptr;)
+
+	// Create and store a UEdGraph that corresponds to a URigVMCollapseNode
+	virtual void CreateEdGraphForCollapseNode(URigVMCollapseNode* InNode) PURE_VIRTUAL(UAnimNextRigVMAssetEditorData::CreateEdGraphForCollapseNode, )
+
+	// Remove the UEdGraph that corresponds to a URigVMGraph
+	virtual bool RemoveEdGraph(URigVMGraph* InModel) PURE_VIRTUAL(UAnimNextRigVMAssetEditorData::RemoveEdGraph, return false;)
+
+	// Initialize the asset for use
+	virtual void Initialize(bool bRecompileVM);
+
+	// Handle RigVM modification events
+	virtual void HandleModifiedEvent(ERigVMGraphNotifType InNotifType, URigVMGraph* InGraph, UObject* InSubject);
+
+	// Get all the kinds of entry for this asset
+	virtual TConstArrayView<TSubclassOf<UAnimNextRigVMAssetEntry>> GetEntryClasses() const PURE_VIRTUAL(UAnimNextRigVMAssetEditorData::GetEntryClasses, return {};)
+
+	// Helper for creating new sub-entries. Sets package flags and outers appropriately 
+	static UObject* CreateNewSubEntry(UAnimNextRigVMAssetEditorData* InEditorData, TSubclassOf<UObject> InClass);
+
+	// Helper for creating new sub-entries. Sets package flags and outers appropriately
+	template<typename EntryClassType>
+	static EntryClassType* CreateNewSubEntry(UAnimNextRigVMAssetEditorData* InEditorData)
+	{
+		return CastChecked<EntryClassType>(CreateNewSubEntry(InEditorData, EntryClassType::StaticClass()));
+	}
+
+	// Helper that generates a fixed UEdGraph name counterpart to a a URigVMGraph 
+	static FName MakeEdGraphName(FName InRigVMGraphName);
+
+	// Get all the entries for this asset
+	TConstArrayView<TObjectPtr<UAnimNextRigVMAssetEntry>> GetAllEntries() const { return Entries; } 
+
+	// Access all the UEdGraphs in this asset
+	TArray<UEdGraph*> GetAllEdGraphs() const;
+
+	// Iterate over all entries of the specified type
+	// If predicate returns false, iteration is stopped
+	template<typename EntryType, typename PredicateType>
+	void ForEachEntryOfType(PredicateType InPredicate) const
+	{
+		for(UAnimNextRigVMAssetEntry* Entry : Entries)
+		{
+			if(EntryType* TypedEntry = Cast<EntryType>(Entry))
+			{
+				if(!InPredicate(TypedEntry))
+				{
+					return;
+				}
+			}
+		}
+	}
+
+	// Find an entry by name
+	ANIMNEXTUNCOOKEDONLY_API UAnimNextRigVMAssetEntry* FindEntry(FName InName) const;
+
+	// Remove an entry from the asset
+	// @return true if the item was removed
+	ANIMNEXTUNCOOKEDONLY_API bool RemoveEntry(UAnimNextRigVMAssetEntry* InEntry, bool bSetupUndoRedo = true, bool bPrintPythonCommand = true);
+
+	// Remove a number of entries from the asset
+	// @return true if any items were removed
+	ANIMNEXTUNCOOKEDONLY_API bool RemoveEntries(const TArray<UAnimNextRigVMAssetEntry*>& InEntries, bool bSetupUndoRedo = true, bool bPrintPythonCommand = true);
+
+	ANIMNEXTUNCOOKEDONLY_API void BroadcastModified();
+
+	void ReportError(const TCHAR* InMessage) const;
+
+	void ReconstructAllNodes();
+
+	// Called from PostLoad to load external packages
+	void PostLoadExternalPackages();
+
+	// Find an entry that corresponds to the specified RigVMGraph. This uses the name of the graph to match the entry 
+	UAnimNextRigVMAssetEntry* FindEntryForRigVMGraph(URigVMGraph* InRigVMGraph) const;
+
+	// Find an entry that corresponds to the specified EdGraph
+	UAnimNextRigVMAssetEntry* FindEntryForEdGraph(URigVMEdGraph* InEdGraph) const;
+
+	/** All entries in this asset - not saved, either serialized or discovered at load time */
+	UPROPERTY(transient)
+	TArray<TObjectPtr<UAnimNextRigVMAssetEntry>> Entries;
+
+	UPROPERTY()
+	FRigVMClient RigVMClient;
+
+	UPROPERTY()
+	FRigVMGraphFunctionStore GraphFunctionStore;
+
+	UPROPERTY(EditAnywhere, Category = "User Interface")
+	FRigVMEdGraphDisplaySettings RigGraphDisplaySettings;
+
+	UPROPERTY(EditAnywhere, Category = "VM")
+	FRigVMRuntimeSettings VMRuntimeSettings;
+	
+	UPROPERTY(EditAnywhere, BlueprintReadWrite, Category = "VM", meta = (AllowPrivateAccess = "true"))
+	FRigVMCompileSettings VMCompileSettings;
+
+	UPROPERTY(transient, DuplicateTransient)
+	TMap<FString, FRigVMOperand> PinToOperandMap;
+
+	UPROPERTY()
+	TArray<FEditedDocumentInfo> LastEditedDocuments;
+
+	UPROPERTY(transient, DuplicateTransient)
+	int32 VMRecompilationBracket = 0;
+
+	UPROPERTY(transient, DuplicateTransient)
+	bool bVMRecompilationRequired = false;
+
+	UPROPERTY(transient, DuplicateTransient)
+	bool bIsCompiling = false;
+
+	FCompilerResultsLog CompileLog;
+
+	FOnRigVMCompiledEvent RigVMCompiledEvent;
+
+	FRigVMGraphModifiedEvent RigVMGraphModifiedEvent;
+
+	// Delegate to subscribe to modifications to this editor data
+	UE::AnimNext::UncookedOnly::FOnEditorDataModified ModifiedDelegate;
+
+	bool bAutoRecompileVM = true;
+	bool bErrorsDuringCompilation = false;
+	bool bSuspendModelNotificationsForSelf = false;
+	bool bSuspendModelNotificationsForOthers = false;
+	bool bSuspendAllNotifications = false;
+	bool bCompileInDebugMode = false;
+	bool bSuspendPythonMessagesForRigVMClient = true;
+	bool bSuspendEditorDataNotifications = false;
+};
diff --git a/Engine/Plugins/Experimental/AnimNext/Source/AnimNextUncookedOnly/Internal/AnimNextRigVMAssetEntry.h b/Engine/Plugins/Experimental/AnimNext/Source/AnimNextUncookedOnly/Internal/AnimNextRigVMAssetEntry.h
new file mode 100644
--- /dev/null
+++ b/Engine/Plugins/Experimental/AnimNext/Source/AnimNextUncookedOnly/Internal/AnimNextRigVMAssetEntry.h
@@ -0,0 +1,47 @@
+﻿// Copyright Epic Games, Inc. All Rights Reserved.
+
+#pragma once
+
+#include "CoreMinimal.h"
+#include "AnimNextRigVMAssetEntry.generated.h"
+
+class UAnimNextRigVMAssetEditorData;
+enum class ERigVMGraphNotifType : uint8;
+class URigVMGraph;
+
+namespace UE::AnimNext::Editor
+{
+	class SRigVMAssetViewRow;
+	struct FRigVMAssetViewEntry;
+}
+
+/** Base class that defines an entry in a parameter block, e.g. a parameter binding */
+UCLASS(MinimalAPI, BlueprintType, Abstract)
+class UAnimNextRigVMAssetEntry : public UObject
+{
+	GENERATED_BODY()
+public:
+	// Binds delegates to owning editor data
+	virtual void Initialize(UAnimNextRigVMAssetEditorData* InEditorData);
+
+	// Allow entries the opportunity to handle RigVM modification events
+	virtual void HandleRigVMGraphModifiedEvent(ERigVMGraphNotifType InNotifType, URigVMGraph* InGraph, UObject* InSubject) {}
+
+	// Get this entry's name
+	virtual FName GetEntryName() const PURE_VIRTUAL(UAnimNextRigVMAssetEntry::GetEntryName, return NAME_None;)
+
+	// Set this entry's name
+	virtual void SetEntryName(FName InName, bool bSetupUndoRedo = true) {}
+
+	// Get the name to be displayed in the UI for this entry
+	virtual FText GetDisplayName() const { return FText::FromName(GetEntryName()); }
+
+	// Get the tooltip to be displayed for the name in the UI for this entry
+	virtual FText GetDisplayNameTooltip() const { return FText::FromName(GetEntryName()); }
+
+	// UObject interface
+	virtual bool IsAsset() const override;
+
+protected:
+	void BroadcastModified();
+};
\ No newline at end of file
diff --git a/Engine/Plugins/Experimental/AnimNext/Source/AnimNextUncookedOnly/Internal/Graph/AnimNextGraphEntry.h b/Engine/Plugins/Experimental/AnimNext/Source/AnimNextUncookedOnly/Internal/Graph/AnimNextGraphEntry.h
--- a/Engine/Plugins/Experimental/AnimNext/Source/AnimNextUncookedOnly/Internal/Graph/AnimNextGraphEntry.h
+++ b/Engine/Plugins/Experimental/AnimNext/Source/AnimNextUncookedOnly/Internal/Graph/AnimNextGraphEntry.h
@@ -3,11 +3,13 @@
 #pragma once
 
 #include "CoreMinimal.h"
+#include "AnimNextRigVMAssetEntry.h"
+#include "IAnimNextRigVMGraphInterface.h"
 #include "AnimNextGraphEntry.generated.h"
 
 class UAnimNextGraph_EditorData;
+class UAnimNextGraph_EdGraph;
 enum class ERigVMGraphNotifType : uint8;
-class URigVMGraph;
 
 namespace UE::AnimNext::Editor
 {
@@ -15,35 +17,31 @@ namespace UE::AnimNext::Editor
 }
 
 /** A single entry in an AnimNext graph asset */
-UCLASS(MinimalAPI, BlueprintType)
-class UAnimNextGraphEntry : public UObject
+UCLASS(MinimalAPI, Category = "Animation Graphs")
+class UAnimNextGraphEntry : public UAnimNextRigVMAssetEntry, public IAnimNextRigVMGraphInterface
 {
 	GENERATED_BODY()
 
 	friend class UAnimNextGraph_EditorData;
-	friend struct UE::AnimNext::Editor::FUtils;
-	
-	void Initialize(UAnimNextGraph_EditorData* InEditorData);
+	friend struct UE::AnimNext::Editor::FUtils;	
 
-	void HandleRigVMGraphModifiedEvent(ERigVMGraphNotifType InNotifType, URigVMGraph* InGraph, UObject* InSubject);
-	
-	// Get the name to be displayed in the UI for this entry
-	virtual FText GetDisplayName() const;
+	// UAnimNextRigVMAssetEntry interface
+	virtual FName GetEntryName() const override;
 
-	// Get the tooltip to be displayed for the name in the UI for this entry
-	virtual FText GetDisplayNameTooltip() const;
-
-	// UObject interface
-	virtual bool IsAsset() const override;
+	// IAnimNextRigVMGraphInterface interface
+	virtual URigVMGraph* GetRigVMGraph() const override;
+	virtual URigVMEdGraph* GetEdGraph() const override;
 
 protected:
-	void BroadcastModified();
-
 	/** The name of the graph */
 	UPROPERTY(VisibleAnywhere, Category = Parameter)
 	FName GraphName;
 
-	/** Graph */
+	/** RigVM graph */
 	UPROPERTY()
 	TObjectPtr<URigVMGraph> Graph;
+
+	/** Editor graph */
+	UPROPERTY()
+	TObjectPtr<UAnimNextGraph_EdGraph> EdGraph;
 };
\ No newline at end of file
diff --git a/Engine/Plugins/Experimental/AnimNext/Source/AnimNextUncookedOnly/Internal/Graph/AnimNextGraph_EditorData.h b/Engine/Plugins/Experimental/AnimNext/Source/AnimNextUncookedOnly/Internal/Graph/AnimNextGraph_EditorData.h
--- a/Engine/Plugins/Experimental/AnimNext/Source/AnimNextUncookedOnly/Internal/Graph/AnimNextGraph_EditorData.h
+++ b/Engine/Plugins/Experimental/AnimNext/Source/AnimNextUncookedOnly/Internal/Graph/AnimNextGraph_EditorData.h
@@ -3,10 +3,11 @@
 #pragma once
 
 #include "CoreMinimal.h"
+#include "AnimNextGraph_Controller.h"
 #include "RigVMModel/RigVMGraph.h"
 #include "AnimNextGraph_EdGraph.h"
-#include "RigVMCore/RigVMGraphFunctionHost.h"
-#include "RigVMBlueprint.h"
+#include "AnimNextRigVMAssetEditorData.h"
+#include "Graph/AnimNextExecuteContext.h"
 #include "Kismet/BlueprintFunctionLibrary.h"
 #include "AnimNextGraph_EditorData.generated.h"
 
@@ -28,18 +29,6 @@ namespace UE::AnimNext::Editor
 	struct FUtils;
 }
 
-enum class UE_DEPRECATED(5.4, "Please, use ERigVMLoadType") EAnimNextGraphLoadType : uint8
-{
-	PostLoad,
-	CheckUserDefinedStructs
-};
-
-namespace UE::AnimNext::UncookedOnly
-{
-// A delegate for subscribing / reacting to graph modifications.
-DECLARE_MULTICAST_DELEGATE_OneParam(FOnGraphModified, UAnimNextGraph_EditorData* /* InEditorData */);
-}
-
 /**
  * The Schema is used to determine which actions are allowed
  * on a graph. This includes any topological change.
@@ -61,23 +50,12 @@ class ANIMNEXTUNCOOKEDONLY_API UAnimNextGraphLibrary : public UBlueprintFunction
 	static UAnimNextGraphEntry* AddGraph(UAnimNextGraph* InGraph, FName InName, bool bSetupUndoRedo = true, bool bPrintPythonCommand = true);
 };
 
-/**
-  * Implements a RigVM client host and RigVM graph function store host
-  * 
-  * A RigVM client holds a graph schema, multiple graph models and controllers, as well as an action stack, undo/redo
-  * information, and a function library.
-  * 
-  * A RigVM function store holds various RigVM functions that can be called and linked against. Functions are created
-  * in the editor UI under the 'My Blueprint' tab.
-  */
+/** Editor data for AnimNext graphs */
 UCLASS(MinimalAPI)
-class UAnimNextGraph_EditorData : public UObject, public IRigVMClientHost, public IRigVMGraphFunctionHost
+class UAnimNextGraph_EditorData : public UAnimNextRigVMAssetEditorData
 {
 	GENERATED_BODY()
 
-public:
-	UAnimNextGraph_EditorData(const FObjectInitializer& ObjectInitializer);
-	
 	friend class UAnimNextGraphFactory;
 	friend class UAnimNextGraph_EdGraph;
 	friend class UAnimNextGraphEntry;
@@ -90,145 +68,29 @@ public:
 	friend class FAnimationAnimNextRuntimeTest_GraphAddDecorator;
 	friend class FAnimationAnimNextRuntimeTest_GraphExecute;
 	friend class FAnimationAnimNextRuntimeTest_GraphExecuteLatent;
-
+	
+public:
+	/** Adds a graph to this asset */
 	ANIMNEXTUNCOOKEDONLY_API UAnimNextGraphEntry* AddGraph(FName InName, bool bSetupUndoRedo = true, bool bPrintPythonCommand = true);
 
+private:
 	// UObject interface
-	virtual void Serialize(FArchive& Ar) override;
 	virtual void PostLoad() override;
-	virtual void PostDuplicate(EDuplicateMode::Type DuplicateMode) override;
-	virtual bool IsEditorOnly() const override { return true; }
-	virtual void GetAssetRegistryTags(FAssetRegistryTagsContext Context) const override;
-	UE_DEPRECATED(5.4, "Implement the version that takes FAssetRegistryTagsContext instead.")
-	virtual void GetAssetRegistryTags(TArray<FAssetRegistryTag>& OutTags) const override;
-#if WITH_EDITOR
-	void HandlePackageDone(const FEndLoadPackageContext& Context);
-	void HandlePackageDone();
-#endif // WITH_EDITOR
 
 	// IRigVMClientHost interface
-	virtual FRigVMClient* GetRigVMClient() override;
-	virtual const FRigVMClient* GetRigVMClient() const override;
-	virtual IRigVMGraphFunctionHost* GetRigVMGraphFunctionHost() override;
-	virtual const IRigVMGraphFunctionHost* GetRigVMGraphFunctionHost() const override;
-	virtual void HandleRigVMGraphAdded(const FRigVMClient* InClient, const FString& InNodePath) override;
-	virtual void HandleRigVMGraphRemoved(const FRigVMClient* InClient, const FString& InNodePath) override;
-	virtual void HandleRigVMGraphRenamed(const FRigVMClient* InClient, const FString& InOldNodePath, const FString& InNewNodePath) override {}
-	virtual void HandleConfigureRigVMController(const FRigVMClient* InClient, URigVMController* InControllerToConfigure) override;
-	virtual UObject* GetEditorObjectForRigVMGraph(URigVMGraph* InVMGraph) const override;
-	virtual URigVMGraph* GetRigVMGraphForEditorObject(UObject* InObject) const override;
 	virtual void RecompileVM() override;
-	virtual void RecompileVMIfRequired() override;
-	virtual void RequestAutoVMRecompilation() override;
-	virtual void SetAutoVMRecompile(bool bAutoRecompile) override;
-	virtual bool GetAutoVMRecompile() const override;
-	virtual void IncrementVMRecompileBracket() override;
-	virtual void DecrementVMRecompileBracket() override;
-	virtual void RefreshAllModels(ERigVMLoadType InLoadType) override;
-	virtual void OnRigVMRegistryChanged() override;
-	virtual void RequestRigVMInit() override;
-	virtual URigVMGraph* GetModel(const UEdGraph* InEdGraph = nullptr) const override;
-	virtual URigVMGraph* GetModel(const FString& InNodePath) const override;
-	virtual URigVMGraph* GetDefaultModel() const override;
-	virtual TArray<URigVMGraph*> GetAllModels() const override;
-	virtual URigVMFunctionLibrary* GetLocalFunctionLibrary() const override;
-	virtual URigVMGraph* AddModel(FString InName = TEXT("Rig Graph"), bool bSetupUndoRedo = true, bool bPrintPythonCommand = true) override;
-	virtual bool RemoveModel(FString InName = TEXT("Rig Graph"), bool bSetupUndoRedo = true, bool bPrintPythonCommand = true) override;
-	virtual FRigVMGetFocusedGraph& OnGetFocusedGraph() override;
-	virtual const FRigVMGetFocusedGraph& OnGetFocusedGraph() const override;
-	virtual URigVMGraph* GetFocusedModel() const override;
-	virtual URigVMController* GetController(const URigVMGraph* InGraph = nullptr) const override;
-	virtual URigVMController* GetControllerByName(const FString InGraphName = TEXT("")) const override;
-	virtual URigVMController* GetOrCreateController(URigVMGraph* InGraph = nullptr) override;
-	virtual URigVMController* GetController(const UEdGraph* InEdGraph) const override;
-	virtual URigVMController* GetOrCreateController(const UEdGraph* InGraph) override;
-	virtual TArray<FString> GeneratePythonCommands(const FString InNewBlueprintName) override;
-
-
-	// IRigVMGraphFunctionHost interface
-	virtual FRigVMGraphFunctionStore* GetRigVMGraphFunctionStore() override;
-	virtual const FRigVMGraphFunctionStore* GetRigVMGraphFunctionStore() const override;
-
-	UEdGraph* CreateEdGraph(URigVMGraph* InRigVMGraph, bool bForce);
-	bool RemoveEdGraph(URigVMGraph* InModel);
-
-	FRigVMGraphModifiedEvent& GetRigVMGraphModifiedEvent()
-	{
-		return RigVMGraphModifiedEvent;
-	}
-
-	ANIMNEXTUNCOOKEDONLY_API void Initialize(bool bRecompileVM);
-
-protected:
-#if WITH_EDITOR
-	void GetAllGraphs(TArray<UEdGraph*>& Graphs) const;
-#endif
-
-	
-	void HandleModifiedEvent(ERigVMGraphNotifType InNotifType, URigVMGraph* InGraph, UObject* InSubject);
-
-	void CreateEdGraphForCollapseNode(URigVMCollapseNode* InNode);
-
-	void BroadcastModified();
-
-	void ReportError(const TCHAR* InMessage) const;
-
-	void ReconstructAllNodes();
 
+	// UAnimNextRigVMAssetEditorData interface
+	virtual TSubclassOf<URigVMController> GetControllerClass() const override { return UAnimNextGraph_Controller::StaticClass(); }
+	virtual TSubclassOf<URigVMSchema> GetRigVMSchemaClass() const override { return UAnimNextGraph_Schema::StaticClass(); }
+	virtual UScriptStruct* GetExecuteContextStruct() const override { return FAnimNextExecuteContext::StaticStruct(); }
+	virtual UEdGraph* CreateEdGraph(URigVMGraph* InRigVMGraph, bool bForce) override;
+	virtual bool RemoveEdGraph(URigVMGraph* InModel) override;
+	virtual void CreateEdGraphForCollapseNode(URigVMCollapseNode* InNode) override;
+	virtual void HandleModifiedEvent(ERigVMGraphNotifType InNotifType, URigVMGraph* InGraph, UObject* InSubject) override;
+	virtual TConstArrayView<TSubclassOf<UAnimNextRigVMAssetEntry>> GetEntryClasses() const override;
+
+private:
 	UPROPERTY()
-	TArray<TObjectPtr<UAnimNextGraph_EdGraph>> Graphs;
-
-	UPROPERTY()
-	FRigVMClient RigVMClient;
-
-	UPROPERTY()
-	FRigVMGraphFunctionStore GraphFunctionStore;
-
-	UPROPERTY()
-	TObjectPtr<URigVMLibraryNode> EntryPoint;
-
-	UPROPERTY(transient)
-	TMap<TObjectPtr<URigVMGraph>, TObjectPtr<URigVMController>> Controllers;
-
-	UPROPERTY(EditAnywhere, Category = "User Interface")
-	FRigVMEdGraphDisplaySettings RigGraphDisplaySettings;
-
-	UPROPERTY(EditAnywhere, Category = "VM")
-	FRigVMRuntimeSettings VMRuntimeSettings;
-	
-	UPROPERTY(EditAnywhere, BlueprintReadWrite, Category = "VM", meta = (AllowPrivateAccess = "true"))
-	FRigVMCompileSettings VMCompileSettings;
-
-	UPROPERTY(transient, DuplicateTransient)
-	TMap<FString, FRigVMOperand> PinToOperandMap;
-
-	/** All entries in this graph asset - not saved, either serialized or discovered at load time */
-	UPROPERTY(transient)
-	TArray<TObjectPtr<UAnimNextGraphEntry>> Entries;
-	
-	UPROPERTY(transient, DuplicateTransient)
-	int32 VMRecompilationBracket = 0;
-
-	UPROPERTY(transient, DuplicateTransient)
-	bool bVMRecompilationRequired = false;
-
-	UPROPERTY(transient, DuplicateTransient)
-	bool bIsCompiling = false;
-	
-	FCompilerResultsLog CompileLog;
-
-	FOnRigVMCompiledEvent RigVMCompiledEvent;
-	FRigVMGraphModifiedEvent RigVMGraphModifiedEvent;
-
-	// Delegate to subscribe to modifications to this graph
-	UE::AnimNext::UncookedOnly::FOnGraphModified ModifiedDelegate;
-
-	bool bAutoRecompileVM = true;
-	bool bErrorsDuringCompilation = false;
-	bool bSuspendModelNotificationsForSelf = false;
-	bool bSuspendModelNotificationsForOthers = false;
-	bool bSuspendPythonMessagesForRigVMClient = false;
-	bool bSuspendAllNotifications = false;
-	bool bSuspendGraphNotifications = false;
-	bool bCompileInDebugMode = false;
+	TArray<TObjectPtr<UAnimNextGraph_EdGraph>> Graphs_DEPRECATED;
 };
diff --git a/Engine/Plugins/Experimental/AnimNext/Source/AnimNextUncookedOnly/Internal/IAnimNextRigVMGraphInterface.h b/Engine/Plugins/Experimental/AnimNext/Source/AnimNextUncookedOnly/Internal/IAnimNextRigVMGraphInterface.h
new file mode 100644
--- /dev/null
+++ b/Engine/Plugins/Experimental/AnimNext/Source/AnimNextUncookedOnly/Internal/IAnimNextRigVMGraphInterface.h
@@ -0,0 +1,28 @@
+// Copyright Epic Games, Inc. All Rights Reserved.
+
+#pragma once
+
+#include "UObject/Interface.h"
+#include "IAnimNextRigVMGraphInterface.generated.h"
+
+class URigVMGraph;
+class URigVMEdGraph;
+
+UINTERFACE(meta=(CannotImplementInterfaceInBlueprint))
+class ANIMNEXTUNCOOKEDONLY_API UAnimNextRigVMGraphInterface : public UInterface
+{
+	GENERATED_BODY()
+};
+
+// Interface for entries that contain a RigVM graph
+class ANIMNEXTUNCOOKEDONLY_API IAnimNextRigVMGraphInterface
+{
+	GENERATED_BODY()
+
+public:
+	// Get the RigVM graph
+	virtual URigVMGraph* GetRigVMGraph() const = 0;
+
+	// Get the Editor graph
+	virtual URigVMEdGraph* GetEdGraph() const = 0;
+};
\ No newline at end of file
diff --git a/Engine/Plugins/Experimental/AnimNext/Source/AnimNextUncookedOnly/Internal/IAnimNextRigVMParameterInterface.h b/Engine/Plugins/Experimental/AnimNext/Source/AnimNextUncookedOnly/Internal/IAnimNextRigVMParameterInterface.h
new file mode 100644
--- /dev/null
+++ b/Engine/Plugins/Experimental/AnimNext/Source/AnimNextUncookedOnly/Internal/IAnimNextRigVMParameterInterface.h
@@ -0,0 +1,30 @@
+﻿// Copyright Epic Games, Inc. All Rights Reserved.
+
+#pragma once
+
+#include "UObject/Interface.h"
+#include "IAnimNextRigVMParameterInterface.generated.h"
+
+struct FAnimNextParamType;
+struct FInstancedPropertyBag;
+
+UINTERFACE(meta=(CannotImplementInterfaceInBlueprint))
+class ANIMNEXTUNCOOKEDONLY_API UAnimNextRigVMParameterInterface : public UInterface
+{
+	GENERATED_BODY()
+};
+
+class ANIMNEXTUNCOOKEDONLY_API IAnimNextRigVMParameterInterface
+{
+	GENERATED_BODY()
+
+public:
+	// Get the parameter type
+	virtual FAnimNextParamType GetParamType() const = 0;
+
+	// Set the parameter type
+	virtual bool SetParamType(const FAnimNextParamType& InType, bool bSetupUndoRedo = true) = 0;
+
+	// Access the backing storage property bag for the parameter
+	virtual FInstancedPropertyBag& GetPropertyBag() const = 0;
+};
\ No newline at end of file
diff --git a/Engine/Plugins/Experimental/AnimNext/Source/AnimNextUncookedOnly/Internal/Param/AnimNextParameterBlockGraph.h b/Engine/Plugins/Experimental/AnimNext/Source/AnimNextUncookedOnly/Internal/Param/AnimNextParameterBlockGraph.h
--- a/Engine/Plugins/Experimental/AnimNext/Source/AnimNextUncookedOnly/Internal/Param/AnimNextParameterBlockGraph.h
+++ b/Engine/Plugins/Experimental/AnimNext/Source/AnimNextUncookedOnly/Internal/Param/AnimNextParameterBlockGraph.h
@@ -3,29 +3,30 @@
 #pragma once
 
 #include "CoreMinimal.h"
-#include "AnimNextParameterBlockEntry.h"
-#include "IAnimNextParameterBlockGraphInterface.h"
+#include "AnimNextRigVMAssetEntry.h"
+#include "IAnimNextRigVMGraphInterface.h"
 #include "AnimNextParameterBlockGraph.generated.h"
 
 class UAnimNextParameterLibrary;
 class UAnimNextParameterBlock_EditorData;
-class URigVMGraph;
+class UAnimNextParameterBlock_EdGraph;
 
-UCLASS()
-class UAnimNextParameterBlockGraph : public UAnimNextParameterBlockEntry, public IAnimNextParameterBlockGraphInterface
+UCLASS(Category = "Parameter Graphs")
+class UAnimNextParameterBlockGraph : public UAnimNextRigVMAssetEntry, public IAnimNextRigVMGraphInterface
 {
 	GENERATED_BODY()
 
 	friend class UAnimNextParameterBlock_EditorData;
 
-	// UAnimNextParameterBlockEntry interface
+	// UAnimNextRigVMAssetEntry interface
+	virtual FName GetEntryName() const override { return GraphName; }
+	virtual void SetEntryName(FName InName, bool bSetupUndoRedo = true) override;
 	virtual FText GetDisplayName() const override;
 	virtual FText GetDisplayNameTooltip() const override;
 
-	// IAnimNextParameterBlockGraphInterface interface
-	virtual URigVMGraph* GetGraph() const override { return Graph; }
-	virtual FName GetGraphName() const override { return GraphName; }
-	virtual void SetGraphName(FName InName, bool bSetupUndoRedo = true) override;
+	// IAnimNextRigVMGraphInterface interface
+	virtual URigVMGraph* GetRigVMGraph() const override;
+	virtual URigVMEdGraph* GetEdGraph() const override;
 
 	/** The name of the graph */
 	UPROPERTY(VisibleAnywhere, Category = Parameter)
@@ -34,4 +35,8 @@ class UAnimNextParameterBlockGraph : public UAnimNextParameterBlockEntry, public
 	/** Graph */
 	UPROPERTY()
 	TObjectPtr<URigVMGraph> Graph;
+
+	/** Graph */
+	UPROPERTY()
+	TObjectPtr<UAnimNextParameterBlock_EdGraph> EdGraph;
 };
\ No newline at end of file
diff --git a/Engine/Plugins/Experimental/AnimNext/Source/AnimNextUncookedOnly/Internal/Param/AnimNextParameterBlockParameter.h b/Engine/Plugins/Experimental/AnimNext/Source/AnimNextUncookedOnly/Internal/Param/AnimNextParameterBlockParameter.h
--- a/Engine/Plugins/Experimental/AnimNext/Source/AnimNextUncookedOnly/Internal/Param/AnimNextParameterBlockParameter.h
+++ b/Engine/Plugins/Experimental/AnimNext/Source/AnimNextUncookedOnly/Internal/Param/AnimNextParameterBlockParameter.h
@@ -3,8 +3,9 @@
 #pragma once
 
 #include "CoreMinimal.h"
-#include "AnimNextParameterBlockEntry.h"
-#include "IAnimNextParameterBlockParameterInterface.h"
+#include "AnimNextRigVMAssetEntry.h"
+#include "IAnimNextRigVMParameterInterface.h"
+#include "Param/ParamType.h"
 #include "AnimNextParameterBlockParameter.generated.h"
 
 class UAnimNextParameterLibrary;
@@ -15,24 +16,25 @@ namespace UE::AnimNext::Editor
 	class FParameterBlockParameterCustomization;
 }
 
-UCLASS(MinimalAPI)
-class UAnimNextParameterBlockParameter : public UAnimNextParameterBlockEntry, public IAnimNextParameterBlockParameterInterface
+UCLASS(MinimalAPI, Category = "Parameters")
+class UAnimNextParameterBlockParameter : public UAnimNextRigVMAssetEntry, public IAnimNextRigVMParameterInterface
 {
 	GENERATED_BODY()
 
 	friend class UAnimNextParameterBlock_EditorData;
 	friend class FAnimationAnimNextParametersEditorTest_Block;
 	friend class UE::AnimNext::Editor::FParameterBlockParameterCustomization;
-	
-	// UAnimNextParameterBlockEntry interface
+
+	// UAnimNextRigVMAssetEntry interface
+	virtual FName GetEntryName() const override;
+	virtual void SetEntryName(FName InName, bool bSetupUndoRedo = true) override;
 	virtual FText GetDisplayName() const override;
 	virtual FText GetDisplayNameTooltip() const override;
 
-	// IAnimNextParameterBlockParameterInterface interface
+	// IAnimNextRigVMParameterInterface interface
 	virtual FAnimNextParamType GetParamType() const override;
-	virtual void SetParameterName(FName InName, bool bSetupUndoRedo = true) override;
-	virtual FName GetParameterName() const override;
 	virtual bool SetParamType(const FAnimNextParamType& InType, bool bSetupUndoRedo = true) override;
+	virtual FInstancedPropertyBag& GetPropertyBag() const override;
 	
 	/** Parameter name we reference */
 	UPROPERTY(VisibleAnywhere, Category = Parameter)
