AnimNext: Begin merging parameter blocks and graphs

diff --git a/Engine/Plugins/Experimental/AnimNext/Source/AnimNextUncookedOnly/Internal/Param/AnimNextParameterBlock_EditorData.h b/Engine/Plugins/Experimental/AnimNext/Source/AnimNextUncookedOnly/Internal/Param/AnimNextParameterBlock_EditorData.h
--- a/Engine/Plugins/Experimental/AnimNext/Source/AnimNextUncookedOnly/Internal/Param/AnimNextParameterBlock_EditorData.h
+++ b/Engine/Plugins/Experimental/AnimNext/Source/AnimNextUncookedOnly/Internal/Param/AnimNextParameterBlock_EditorData.h
@@ -3,16 +3,14 @@
 #pragma once
 
 #include "CoreMinimal.h"
-#include "RigVMModel/RigVMGraph.h"
 #include "AnimNextParameterBlock_EdGraph.h"
-#include "RigVMCore/RigVMGraphFunctionHost.h"
-#include "RigVMBlueprint.h"
+#include "AnimNextRigVMAssetEditorData.h"
 #include "Kismet/BlueprintFunctionLibrary.h"
+#include "Param/AnimNextParameterExecuteContext.h"
 #include "Param/ParamType.h"
 #include "AnimNextParameterBlock_EditorData.generated.h"
 
 class UAnimNextParameterBlock;
-class UAnimNextParameterBlockEntry;
 class UAnimNextParameterBlockParameter;
 class UAnimNextParameterBlock_Controller;
 enum class ERigVMGraphNotifType : uint8;
@@ -28,25 +26,14 @@ namespace UE::AnimNext::Editor
 {
 	struct FUtils;
 	class FParameterBlockEditor;
-	class SParameterBlockView;
+	class SRigVMAssetView;
 	class SParameterPicker;
 	class FParameterBlockTabSummoner;
-	class SParameterBlockViewRow;
+	class SRigVMAssetViewRow;
 	class FParameterBlockParameterCustomization;
+	class FModule;
 }
 
-namespace UE::AnimNext::UncookedOnly
-{
-	// A delegate for subscribing / reacting to parameter block modifications.
-	DECLARE_MULTICAST_DELEGATE_OneParam(FOnParameterBlockModified, UAnimNextParameterBlock_EditorData* /* InEditorData */);
-}
-
-enum class UE_DEPRECATED(5.4, "Please use ERigVMLoadType") EAnimNextParameterLoadType : uint8
-{
-	PostLoad,
-	CheckUserDefinedStructs
-};
-
 UCLASS()
 class UAnimNextParameterBlockLibrary_Schema : public URigVMSchema
 {
@@ -64,205 +51,51 @@ class ANIMNEXTUNCOOKEDONLY_API UAnimNextParameterBlockLibrary : public UBlueprin
 
 	UFUNCTION(BlueprintCallable, Category = "AnimNext|Parameter Block", meta=(ScriptMethod))
 	static UAnimNextParameterBlockGraph* AddGraph(UAnimNextParameterBlock* InBlock, FName InName, bool bSetupUndoRedo = true, bool bPrintPythonCommand = true);
-	
-	UFUNCTION(BlueprintCallable, Category = "AnimNext|Parameter Block", meta=(ScriptMethod))
-	static UAnimNextParameterBlockBinding* AddBinding(UAnimNextParameterBlock* InBlock, FName InName, EPropertyBagPropertyType InValueType, EPropertyBagContainerType InContainerType = EPropertyBagContainerType::None, const UObject* InValueTypeObject = nullptr, bool bSetupUndoRedo = true, bool bPrintPythonCommand = true);
-
-	UFUNCTION(BlueprintCallable, Category = "AnimNext|Parameter Block", meta=(ScriptMethod))
-	static UAnimNextParameterBlockBindingReference* AddBindingReference(UAnimNextParameterBlock* InBlock, FName InName, UAnimNextParameterBlock* InReferencedBlock, bool bSetupUndoRedo = true, bool bPrintPythonCommand = true);
-
-	UFUNCTION(BlueprintCallable, Category = "AnimNext|Parameter Block", meta=(ScriptMethod))
-	static bool RemoveAllBindings(UAnimNextParameterBlock* InBlock, FName InName, bool bSetupUndoRedo = true, bool bPrintPythonCommand = true);
-
-	UFUNCTION(BlueprintCallable, Category = "AnimNext|Parameter Block", meta=(ScriptMethod))
-	static bool RemoveEntry(UAnimNextParameterBlock* InBlock, UAnimNextParameterBlockEntry* InEntry, bool bSetupUndoRedo = true, bool bPrintPythonCommand = true);
-
-	UFUNCTION(BlueprintCallable, Category = "AnimNext|Parameter Block", meta=(ScriptMethod))
-	static bool RemoveEntries(UAnimNextParameterBlock* InBlock, const TArray<UAnimNextParameterBlockEntry*>& InEntries, bool bSetupUndoRedo = true, bool bPrintPythonCommand = true);
-
-	UFUNCTION(BlueprintCallable, Category = "AnimNext|Parameter Block", meta=(ScriptMethod))
-	static UAnimNextParameterBlockEntry* FindBinding(UAnimNextParameterBlock* InBlock, FName InName);
 };
 
 UCLASS(MinimalAPI)
-class UAnimNextParameterBlock_EditorData : public UObject, public IRigVMClientHost, public IRigVMGraphFunctionHost
+class UAnimNextParameterBlock_EditorData : public UAnimNextRigVMAssetEditorData
 {
 	GENERATED_BODY()
-
-	UAnimNextParameterBlock_EditorData(const FObjectInitializer& ObjectInitializer);
-
-#if WITH_EDITORONLY_DATA
+	
 	ANIMNEXTUNCOOKEDONLY_API static const FName ExportsAssetRegistryTag; 
 
 	friend class UAnimNextParameterBlockLibrary;
 	friend class UAnimNextParameterBlockFactory;
-	friend class UAnimNextParameterBlockEntry;
+	friend class UAnimNextRigVMAssetEntry;
 	friend class UAnimNextParameterBlock_EdGraph;
 	friend class UAnimNextParameterBlockParameter;
 	friend class UAnimNextParameterBlockBindingReference;
 	friend struct UE::AnimNext::Editor::FUtils;
 	friend struct UE::AnimNext::UncookedOnly::FUtils;
 	friend struct UE::AnimNext::UncookedOnly::FUtilsPrivate;
+	friend class UE::AnimNext::Editor::FModule;
 	friend class UE::AnimNext::Editor::FParameterBlockEditor;
-	friend struct FAnimNextParameterSchemaAction_RigUnit;
-	friend struct FAnimNextParameterSchemaAction_DispatchFactory;
-	friend class UE::AnimNext::Editor::SParameterBlockView;
+	friend class UE::AnimNext::Editor::SRigVMAssetView;
 	friend class UE::AnimNext::Editor::SParameterPicker;
 	friend class UE::AnimNext::Editor::FParameterBlockTabSummoner;
 	friend class FAnimationAnimNextParametersEditorTest_Block;
-	friend class UE::AnimNext::Editor::SParameterBlockViewRow;
+	friend class UE::AnimNext::Editor::SRigVMAssetViewRow;
 	friend class UE::AnimNext::Editor::FParameterBlockParameterCustomization;
 
 	ANIMNEXTUNCOOKEDONLY_API UAnimNextParameterBlockParameter* AddParameter(FName InName, FAnimNextParamType InType, bool bSetupUndoRedo = true, bool bPrintPythonCommand = true);
 
 	ANIMNEXTUNCOOKEDONLY_API UAnimNextParameterBlockGraph* AddGraph(FName InName, bool bSetupUndoRedo = true, bool bPrintPythonCommand = true);
-	
-	ANIMNEXTUNCOOKEDONLY_API UAnimNextParameterBlockBinding* AddBinding(FName InName, FAnimNextParamType InType, bool bSetupUndoRedo = true, bool bPrintPythonCommand = true);
-
-	ANIMNEXTUNCOOKEDONLY_API UAnimNextParameterBlockBindingReference* AddBindingReference(FName InName, UAnimNextParameterBlock* InBlock, bool bSetupUndoRedo = true, bool bPrintPythonCommand = true);
-
-	ANIMNEXTUNCOOKEDONLY_API bool RemoveAllBindings(FName InName, bool bSetupUndoRedo = true, bool bPrintPythonCommand = true);
-
-	ANIMNEXTUNCOOKEDONLY_API bool RemoveEntry(UAnimNextParameterBlockEntry* InEntry, bool bSetupUndoRedo = true, bool bPrintPythonCommand = true);
-
-	ANIMNEXTUNCOOKEDONLY_API bool RemoveEntries(const TArray<UAnimNextParameterBlockEntry*>& InEntries, bool bSetupUndoRedo = true, bool bPrintPythonCommand = true);
-
-	ANIMNEXTUNCOOKEDONLY_API UAnimNextParameterBlockEntry* FindBinding(FName InName) const;
 
 	// UObject interface
-	virtual void Serialize(FArchive& Ar) override;
 	virtual void PostLoad() override;
-	virtual void PostDuplicate(EDuplicateMode::Type DuplicateMode) override;
-	virtual void PostTransacted(const FTransactionObjectEvent& TransactionEvent) override;
-	virtual bool IsEditorOnly() const override { return true; }
-	virtual void GetAssetRegistryTags(FAssetRegistryTagsContext Context) const override;
-	UE_DEPRECATED(5.4, "Implement the version that takes FAssetRegistryTagsContext instead.")
-	virtual void GetAssetRegistryTags(TArray<FAssetRegistryTag>& OutTags) const override;
-
-#if WITH_EDITOR
-	void HandlePackageDone(const FEndLoadPackageContext& Context);
-	void HandlePackageDone();
-#endif
 
 	// IRigVMClientHost interface
-	virtual FRigVMClient* GetRigVMClient() override;
-	virtual const FRigVMClient* GetRigVMClient() const override;
-	virtual IRigVMGraphFunctionHost* GetRigVMGraphFunctionHost() override;
-	virtual const IRigVMGraphFunctionHost* GetRigVMGraphFunctionHost() const override;
-	virtual void HandleRigVMGraphAdded(const FRigVMClient* InClient, const FString& InNodePath) override;
-	virtual void HandleRigVMGraphRemoved(const FRigVMClient* InClient, const FString& InNodePath) override;
-	virtual void HandleRigVMGraphRenamed(const FRigVMClient* InClient, const FString& InOldNodePath, const FString& InNewNodePath) override {}
-	virtual void HandleConfigureRigVMController(const FRigVMClient* InClient, URigVMController* InControllerToConfigure) override;
-	virtual UObject* GetEditorObjectForRigVMGraph(URigVMGraph* InVMGraph) const override;
-	virtual URigVMGraph* GetRigVMGraphForEditorObject(UObject* InObject) const override;
 	virtual void RecompileVM() override;
-	virtual void RecompileVMIfRequired() override;
-	virtual void RequestAutoVMRecompilation() override;
-	virtual void SetAutoVMRecompile(bool bAutoRecompile) override;
-	virtual bool GetAutoVMRecompile() const override;
-	virtual void IncrementVMRecompileBracket() override;
-	virtual void DecrementVMRecompileBracket() override;
-	virtual void RefreshAllModels(ERigVMLoadType InLoadType) override;
-	virtual void OnRigVMRegistryChanged() override;
-	virtual void RequestRigVMInit() override;
-	virtual URigVMGraph* GetModel(const UEdGraph* InEdGraph = nullptr) const override;
-	virtual URigVMGraph* GetModel(const FString& InNodePath) const override;
-	virtual URigVMGraph* GetDefaultModel() const override;
-	virtual TArray<URigVMGraph*> GetAllModels() const override;
-	virtual URigVMFunctionLibrary* GetLocalFunctionLibrary() const override;
-	virtual URigVMGraph* AddModel(FString InName = TEXT("Rig Graph"), bool bSetupUndoRedo = true, bool bPrintPythonCommand = true) override;
-	virtual bool RemoveModel(FString InName = TEXT("Rig Graph"), bool bSetupUndoRedo = true, bool bPrintPythonCommand = true) override;
-	virtual FRigVMGetFocusedGraph& OnGetFocusedGraph() override;
-	virtual const FRigVMGetFocusedGraph& OnGetFocusedGraph() const override;
-	virtual URigVMGraph* GetFocusedModel() const override;
-	virtual URigVMController* GetController(const URigVMGraph* InGraph = nullptr) const override;
-	virtual URigVMController* GetControllerByName(const FString InGraphName = TEXT("")) const override;
-	virtual URigVMController* GetOrCreateController(URigVMGraph* InGraph = nullptr) override;
-	virtual URigVMController* GetController(const UEdGraph* InEdGraph) const override;
-	virtual URigVMController* GetOrCreateController(const UEdGraph* InGraph) override;
-	virtual TArray<FString> GeneratePythonCommands(const FString InNewBlueprintName) override;
-
 
-	// IRigVMGraphFunctionHost interface
-	virtual FRigVMGraphFunctionStore* GetRigVMGraphFunctionStore() override;
-	virtual const FRigVMGraphFunctionStore* GetRigVMGraphFunctionStore() const override;
-
-	UEdGraph* CreateEdGraph(URigVMGraph* InRigVMGraph, bool bForce);
-	bool RemoveEdGraph(URigVMGraph* InModel);
-	
-	ANIMNEXTUNCOOKEDONLY_API void Initialize(bool bRecompileVM);
-
-	void Recompile();
-
-	void RecompileIfRequired();
-
-	void RequestAutoRecompilation();
-	
-	void HandleModifiedEvent(ERigVMGraphNotifType InNotifType, URigVMGraph* InGraph, UObject* InSubject);
-
-	void CreateEdGraphForCollapseNode(URigVMCollapseNode* InNode);
-
-	ANIMNEXTUNCOOKEDONLY_API void BroadcastModified();
-
-	void ReportError(const TCHAR* InMessage) const;
-
-	void ReconstructAllNodes();
-	
-	UPROPERTY()
-	TArray<TObjectPtr<UAnimNextParameterBlock_EdGraph>> Graphs;
+	// UAnimNextRigVMAssetEditorData interface
+	virtual TSubclassOf<URigVMSchema> GetRigVMSchemaClass() const override { return UAnimNextParameterBlockLibrary_Schema::StaticClass(); }
+	virtual UScriptStruct* GetExecuteContextStruct() const override { return FAnimNextParameterExecuteContext::StaticStruct(); }
+	virtual UEdGraph* CreateEdGraph(URigVMGraph* InRigVMGraph, bool bForce) override;
+	virtual bool RemoveEdGraph(URigVMGraph* InModel) override;
+	virtual void CreateEdGraphForCollapseNode(URigVMCollapseNode* InNode) override;
+	virtual TConstArrayView<TSubclassOf<UAnimNextRigVMAssetEntry>> GetEntryClasses() const override;
 
 	UPROPERTY()
-	FRigVMClient RigVMClient;
-
-	UPROPERTY()
-	FRigVMGraphFunctionStore GraphFunctionStore;
-
-	UPROPERTY(EditAnywhere, Category = "User Interface")
-	FRigVMEdGraphDisplaySettings RigGraphDisplaySettings;
-
-	UPROPERTY(EditAnywhere, Category = "VM")
-	FRigVMRuntimeSettings VMRuntimeSettings;
-	
-	UPROPERTY(EditAnywhere, BlueprintReadWrite, Category = "VM", meta = (AllowPrivateAccess = "true"))
-	FRigVMCompileSettings VMCompileSettings;
-
-	UPROPERTY(transient, DuplicateTransient)
-	TMap<FString, FRigVMOperand> PinToOperandMap;
-
-	/** All entries in this parameter block - not saved, either serialized or discovered at load time */
-	UPROPERTY(transient)
-	TArray<TObjectPtr<UAnimNextParameterBlockEntry>> Entries;
-
-	UPROPERTY()
-	TArray<FEditedDocumentInfo> LastEditedDocuments;
-
-	UPROPERTY(transient, DuplicateTransient)
-	int32 VMRecompilationBracket = 0;
-
-	UPROPERTY(transient, DuplicateTransient)
-	bool bVMRecompilationRequired = false;
-
-	UPROPERTY(transient, DuplicateTransient)
-	bool bStructRecompilationRequired = false;
-
-	UPROPERTY(transient, DuplicateTransient)
-	bool bIsCompiling = false;
-	
-	FCompilerResultsLog CompileLog;
-
-	FOnRigVMCompiledEvent RigVMCompiledEvent;
-	FRigVMGraphModifiedEvent RigVMGraphModifiedEvent;
-
-	// Delegate to subscribe to modifications to this block
-	UE::AnimNext::UncookedOnly::FOnParameterBlockModified ModifiedDelegate;
-
-	bool bAutoRecompileVM = true;
-	bool bErrorsDuringCompilation = false;
-	bool bSuspendModelNotificationsForSelf = false;
-	bool bSuspendModelNotificationsForOthers = false;
-	bool bSuspendAllNotifications = false;
-	bool bCompileInDebugMode = false;
-	bool bSuspendPythonMessagesForRigVMClient = true;
-	bool bSuspendBlockNotifications = false;
-#endif // #if WITH_EDITORONLY_DATA
+	TArray<TObjectPtr<UAnimNextParameterBlock_EdGraph>> Graphs_DEPRECATED;
 };
diff --git a/Engine/Plugins/Experimental/AnimNext/Source/AnimNextUncookedOnly/Internal/Param/IAnimNextParameterBlockReferenceInterface.h b/Engine/Plugins/Experimental/AnimNext/Source/AnimNextUncookedOnly/Internal/Param/IAnimNextParameterBlockReferenceInterface.h
--- a/Engine/Plugins/Experimental/AnimNext/Source/AnimNextUncookedOnly/Internal/Param/IAnimNextParameterBlockReferenceInterface.h
+++ b/Engine/Plugins/Experimental/AnimNext/Source/AnimNextUncookedOnly/Internal/Param/IAnimNextParameterBlockReferenceInterface.h
@@ -11,8 +11,8 @@ class UAnimNextParameterBlock;
 
 namespace UE::AnimNext::Editor
 {
-class SParameterBlockView;
-class SParameterBlockViewRow;
+class SRigVMAssetView;
+class SRigVMAssetViewRow;
 }
 
 namespace UE::AnimNext::UncookedOnly
@@ -30,8 +30,8 @@ class ANIMNEXTUNCOOKEDONLY_API IAnimNextParameterBlockReferenceInterface
 {
 	GENERATED_BODY()
 
-	friend class UE::AnimNext::Editor::SParameterBlockView;
-	friend class UE::AnimNext::Editor::SParameterBlockViewRow;
+	friend class UE::AnimNext::Editor::SRigVMAssetView;
+	friend class UE::AnimNext::Editor::SRigVMAssetViewRow;
 	friend struct UE::AnimNext::UncookedOnly::FUtilsPrivate;
 
 	// Get the block we reference
diff --git a/Engine/Plugins/Experimental/AnimNext/Source/AnimNextUncookedOnly/Internal/UncookedOnlyUtils.h b/Engine/Plugins/Experimental/AnimNext/Source/AnimNextUncookedOnly/Internal/UncookedOnlyUtils.h
--- a/Engine/Plugins/Experimental/AnimNext/Source/AnimNextUncookedOnly/Internal/UncookedOnlyUtils.h
+++ b/Engine/Plugins/Experimental/AnimNext/Source/AnimNextUncookedOnly/Internal/UncookedOnlyUtils.h
@@ -5,6 +5,8 @@
 #include "CoreMinimal.h"
 #include "Param/AnimNextParameterBlock_EditorData.h"
 #include "Graph/AnimNextGraph_EditorData.h"
+#include "AnimNextRigVMAssetEditorData.h"
+#include "IAnimNextRigVMGraphInterface.h"
 #include "Param/ParamTypeHandle.h"
 #include "RigVMCore/RigVMTemplate.h"
 
@@ -21,6 +23,8 @@ class UAnimNextParameterBlock;
 class UAnimNextParameterBlock_EditorData;
 class UAnimNextGraph_EdGraph;
 struct FEdGraphPinType;
+class UAnimNextRigVMAsset;
+class UAnimNextRigVMAssetEditorData;
 
 namespace UE
 {
@@ -111,6 +115,10 @@ struct ANIMNEXTUNCOOKEDONLY_API FUtils
 
 	static void RecreateVM(UAnimNextParameterBlock* InParameterBlock);
 
+	static UAnimNextRigVMAsset* GetAsset(UAnimNextRigVMAssetEditorData* InEditorData);
+
+	static UAnimNextRigVMAssetEditorData* GetEditorData(UAnimNextRigVMAsset* InAsset);
+
 	/**
 	 * Get an AnimNext parameter type from an FEdGraphPinType.
 	 * Note that the returned handle may not be valid, so should be checked using IsValid() before use.
@@ -149,30 +157,17 @@ struct ANIMNEXTUNCOOKEDONLY_API FUtils
 
 	/** Returns all nodes in all graphs of the specified class */
 	template<class T>
-	static void GetAllNodesOfClass(const UAnimNextParameterBlock_EditorData* InEditorData, TArray<T*>& OutNodes)
-	{
-		for(const UAnimNextParameterBlock_EdGraph* Graph : InEditorData->Graphs)
-		{
-			check(Graph);
-			TArray<T*> GraphNodes;
-			Graph->GetNodesOfClass<T>(GraphNodes);
-			OutNodes.Append(GraphNodes);
-		}
-	}
-
-	/** Returns all nodes in all graphs of the specified class */
-	template<class T>
-	static void GetAllNodesOfClass(const UAnimNextGraph_EditorData* InEditorData, TArray<T*>& OutNodes)
+	static void GetAllNodesOfClass(const UAnimNextRigVMAssetEditorData* InEditorData, TArray<T*>& OutNodes)
 	{
-		for(const UAnimNextGraph_EdGraph* Graph : InEditorData->Graphs)
+		InEditorData->ForEachEntryOfType<IAnimNextRigVMGraphInterface>([&OutNodes](IAnimNextRigVMGraphInterface* InGraphInterface)
 		{
-			check(Graph);
 			TArray<T*> GraphNodes;
-			Graph->GetNodesOfClass<T>(GraphNodes);
+			InGraphInterface->GetEdGraph()->GetNodesOfClass<T>(GraphNodes);
 			OutNodes.Append(GraphNodes);
-		}
+			return true;
+		});
 	}
-
+	
 	// Gets the parameters that are exported to the asset registry for an asset
 	static bool GetExportedParametersForAsset(const FAssetData& InAsset, FAnimNextParameterProviderAssetRegistryExports& OutExports);
 
diff --git a/Engine/Plugins/Experimental/AnimNext/Source/AnimNextUncookedOnly/Private/AnimNextRigVMAssetEditorData.cpp b/Engine/Plugins/Experimental/AnimNext/Source/AnimNextUncookedOnly/Private/AnimNextRigVMAssetEditorData.cpp
new file mode 100644
--- /dev/null
+++ b/Engine/Plugins/Experimental/AnimNext/Source/AnimNextUncookedOnly/Private/AnimNextRigVMAssetEditorData.cpp
@@ -0,0 +1,810 @@
+﻿// Copyright Epic Games, Inc. All Rights Reserved.
+
+#include "AnimNextRigVMAssetEditorData.h"
+
+#include "AnimNextRigVMAsset.h"
+#include "AnimNextRigVMAssetEntry.h"
+#include "ControlRigDefines.h"
+#include "ExternalPackageHelper.h"
+#include "IAnimNextRigVMGraphInterface.h"
+#include "IAnimNextRigVMParameterInterface.h"
+#include "UncookedOnlyUtils.h"
+#include "Misc/TransactionObjectEvent.h"
+#include "RigVMModel/RigVMFunctionLibrary.h"
+#include "RigVMModel/RigVMNotifications.h"
+#include "RigVMModel/Nodes/RigVMCollapseNode.h"
+#include "UObject/AssetRegistryTagsContext.h"
+
+void UAnimNextRigVMAssetEditorData::BroadcastModified()
+{
+	RecompileVM();
+
+	if(!bSuspendEditorDataNotifications)
+	{
+		ModifiedDelegate.Broadcast(this);
+	}
+}
+
+void UAnimNextRigVMAssetEditorData::ReportError(const TCHAR* InMessage) const
+{
+	FScriptExceptionHandler::Get().HandleException(ELogVerbosity::Error, InMessage, TEXT(""));
+}
+
+void UAnimNextRigVMAssetEditorData::ReconstructAllNodes()
+{
+	// Avoid refreshing EdGraph nodes during cook
+	if (GIsCookerLoadingPackage)
+	{
+		return;
+	}
+	
+	if (GetRigVMClient()->GetDefaultModel() == nullptr)
+	{
+		return;
+	}
+
+	TArray<URigVMEdGraphNode*> AllNodes;
+	UE::AnimNext::UncookedOnly::FUtils::GetAllNodesOfClass(this, AllNodes);
+
+	for (URigVMEdGraphNode* Node : AllNodes)
+	{
+		Node->SetFlags(RF_Transient);
+	}
+
+	for(URigVMEdGraphNode* Node : AllNodes)
+	{
+		Node->ReconstructNode();
+	}
+
+	for (URigVMEdGraphNode* Node : AllNodes)
+	{
+		Node->ClearFlags(RF_Transient);
+	}
+}
+
+void UAnimNextRigVMAssetEditorData::Serialize(FArchive& Ar)
+{
+	RigVMClient.SetOuterClientHost(this, GET_MEMBER_NAME_CHECKED(UAnimNextRigVMAssetEditorData, RigVMClient));
+
+	Super::Serialize(Ar);
+}
+
+void UAnimNextRigVMAssetEditorData::Initialize(bool bRecompileVM)
+{
+	RigVMClient.bDefaultModelCanBeRemoved = true;
+	RigVMClient.SetControllerClass(GetControllerClass());
+	RigVMClient.SetSchemaClass(GetRigVMSchemaClass());
+	RigVMClient.SetExecuteContextStruct(GetExecuteContextStruct());
+	RigVMClient.SetOuterClientHost(this, GET_MEMBER_NAME_CHECKED(UAnimNextRigVMAssetEditorData, RigVMClient));
+	{
+		TGuardValue<bool> DisableClientNotifs(RigVMClient.bSuspendNotifications, true);
+		RigVMClient.GetOrCreateFunctionLibrary(false);
+	}
+
+	if (RigVMClient.GetController(0) == nullptr)
+	{
+		if(RigVMClient.GetDefaultModel())
+		{
+			RigVMClient.GetOrCreateController(RigVMClient.GetDefaultModel());
+		}
+
+		check(RigVMClient.GetFunctionLibrary());
+		RigVMClient.GetOrCreateController(RigVMClient.GetFunctionLibrary());
+
+		// Init function library controllers
+		for(URigVMLibraryNode* LibraryNode : RigVMClient.GetFunctionLibrary()->GetFunctions())
+		{
+			RigVMClient.GetOrCreateController(LibraryNode->GetContainedGraph());
+		}
+		
+		if(bRecompileVM)
+		{
+			RecompileVM();
+		}
+	}
+
+	for(UAnimNextRigVMAssetEntry* Entry : Entries)
+	{
+		Entry->Initialize(this);
+	}
+}
+
+void UAnimNextRigVMAssetEditorData::PostLoad()
+{
+	Super::PostLoad();
+
+	Initialize(/*bRecompileVM*/false);
+
+	RefreshAllModels(ERigVMLoadType::PostLoad);
+
+	PostLoadExternalPackages();
+
+	// delay compilation until the package has been loaded
+	FCoreUObjectDelegates::OnEndLoadPackage.AddUObject(this, &UAnimNextRigVMAssetEditorData::HandlePackageDone);
+}
+
+void UAnimNextRigVMAssetEditorData::PostLoadExternalPackages()
+{
+	FExternalPackageHelper::LoadObjectsFromExternalPackages<UAnimNextRigVMAssetEntry>(this, [this](UAnimNextRigVMAssetEntry* InLoadedEntry)
+	{
+		check(IsValid(InLoadedEntry));
+		InLoadedEntry->Initialize(this);
+		Entries.Add(InLoadedEntry);
+	});
+}
+
+void UAnimNextRigVMAssetEditorData::PostTransacted(const FTransactionObjectEvent& TransactionEvent)
+{
+	Super::PostTransacted(TransactionEvent);
+
+	if (TransactionEvent.GetEventType() == ETransactionObjectEventType::UndoRedo)
+	{
+		BroadcastModified();
+	}
+}
+
+void UAnimNextRigVMAssetEditorData::PostDuplicate(EDuplicateMode::Type DuplicateMode)
+{
+	Initialize(/*bRecompileVM*/true);
+}
+
+void UAnimNextRigVMAssetEditorData::GetAssetRegistryTags(FAssetRegistryTagsContext Context) const
+{
+	Super::GetAssetRegistryTags(Context);
+
+	FAnimNextParameterProviderAssetRegistryExports ExportParameters;
+	ExportParameters.Parameters.Reserve(Entries.Num());
+
+	for(const UAnimNextRigVMAssetEntry* Entry : Entries)
+	{
+		if(const IAnimNextRigVMParameterInterface* ParameterInterface = Cast<IAnimNextRigVMParameterInterface>(Entry))
+		{
+			ExportParameters.Parameters.Emplace(Entry->GetEntryName(), ParameterInterface->GetParamType(), EAnimNextParameterFlags::Bound);
+		}
+		else if(const IAnimNextRigVMGraphInterface* GraphInterface = Cast<IAnimNextRigVMGraphInterface>(Entry))
+		{
+			UE::AnimNext::UncookedOnly::FUtils::GetGraphParameters(GraphInterface->GetRigVMGraph(), ExportParameters);
+		}
+	}
+
+	FString TagValue;
+	FAnimNextParameterProviderAssetRegistryExports::StaticStruct()->ExportText(TagValue, &ExportParameters, nullptr, nullptr, PPF_None, nullptr);
+	Context.AddTag(FAssetRegistryTag(UE::AnimNext::ExportsAnimNextAssetRegistryTag, TagValue, FAssetRegistryTag::TT_Hidden));
+}
+
+void UAnimNextRigVMAssetEditorData::HandlePackageDone(const FEndLoadPackageContext& Context)
+{
+	if (!Context.LoadedPackages.Contains(GetPackage()))
+	{
+		return;
+	}
+	HandlePackageDone();
+}
+
+void UAnimNextRigVMAssetEditorData::HandlePackageDone()
+{
+	FCoreUObjectDelegates::OnEndLoadPackage.RemoveAll(this);
+
+	RecompileVM();
+
+	ReconstructAllNodes();
+}
+
+void UAnimNextRigVMAssetEditorData::RefreshAllModels(ERigVMLoadType InLoadType)
+{
+	const bool bIsPostLoad = InLoadType == ERigVMLoadType::PostLoad;
+
+	TGuardValue<bool> IsCompilingGuard(bIsCompiling, true);
+	TGuardValue<bool> ClientIgnoreModificationsGuard(RigVMClient.bIgnoreModelNotifications, true);
+
+	TArray<URigVMGraph*> GraphsToDetach = RigVMClient.GetAllModels(true, false);
+	TMap<const URigVMGraph*, TArray<URigVMController::FLinkedPath>> LinkedPaths;
+	
+	if (ensure(IsInGameThread()))
+	{
+		for (const URigVMGraph* GraphToDetach : GraphsToDetach)
+		{
+			URigVMController* Controller = RigVMClient.GetOrCreateController(GraphToDetach);
+			// temporarily disable default value validation during load time, serialized values should always be accepted
+			TGuardValue<bool> PerGraphDisablePinDefaultValueValidation(Controller->bValidatePinDefaults, false);
+			FRigVMControllerNotifGuard NotifGuard(Controller, true);
+			LinkedPaths.Add(GraphToDetach, Controller->GetLinkedPaths());
+			Controller->FastBreakLinkedPaths(LinkedPaths.FindChecked(GraphToDetach));
+			TArray<URigVMNode*> Nodes = GraphToDetach->GetNodes();
+			for (URigVMNode* Node : Nodes)
+			{
+				Controller->RepopulatePinsOnNode(Node, true, true);
+			}
+		}
+		//SetupPinRedirectorsForBackwardsCompatibility();
+	}
+
+	for (const URigVMGraph* GraphToDetach : GraphsToDetach)
+	{
+		URigVMController* Controller = RigVMClient.GetOrCreateController(GraphToDetach);
+		// at this stage, allow all links to be reattached,
+		// RecomputeAllTemplateFilteredPermutations() later should break any invalid links
+		FRigVMControllerNotifGuard NotifGuard(Controller, true);
+		Controller->RestoreLinkedPaths(LinkedPaths.FindChecked(GraphToDetach));
+	}
+
+	if (bIsPostLoad)
+	{
+		//PatchTemplateNodesWithPreferredPermutation();
+	}
+
+	TArray<URigVMGraph*> GraphsToClean = RigVMClient.GetAllModels(true, true);
+
+	// Sort from leaf graphs to root
+	TArray<URigVMGraph*> SortedGraphsToClean;
+	SortedGraphsToClean.Reserve(GraphsToClean.Num());
+	while (SortedGraphsToClean.Num() < GraphsToClean.Num())
+	{
+		bool bGraphAdded = false;
+		for (URigVMGraph* Graph : GraphsToClean)
+		{
+			if (SortedGraphsToClean.Contains(Graph))
+			{
+				continue;
+			}
+
+			TArray<URigVMGraph*> ContainedGraphs;
+			for (URigVMNode* Node : Graph->GetNodes())
+			{
+				if (URigVMLibraryNode* LibraryNode = Cast<URigVMLibraryNode>(Node))
+				{
+					if (URigVMFunctionReferenceNode* FunctionReferenceNode = Cast<URigVMFunctionReferenceNode>(LibraryNode))
+					{
+						if (FunctionReferenceNode->GetReferencedFunctionHeader().LibraryPointer.LibraryNode.GetLongPackageName() != GetPackage()->GetPathName())
+						{
+							continue;
+						}
+						if (URigVMLibraryNode* ReferencedNode = FunctionReferenceNode->LoadReferencedNode())
+						{
+							ContainedGraphs.Add(ReferencedNode->GetContainedGraph());
+							continue;
+						}
+					}
+
+					if (URigVMGraph* ContainedGraph = LibraryNode->GetContainedGraph())
+					{
+						ContainedGraphs.Add(ContainedGraph);
+					}
+				}
+			}
+
+			bool bAllContained = true;
+			for (URigVMGraph* Contained : ContainedGraphs)
+			{
+				if (!SortedGraphsToClean.Contains(Contained))
+				{
+					bAllContained = false;
+					break;
+				}
+			}
+			if (bAllContained)
+			{
+				SortedGraphsToClean.Add(Graph);
+				bGraphAdded = true;
+			}
+		}
+		ensure(bGraphAdded);
+	}
+
+	for (int32 GraphIndex = 0; GraphIndex < SortedGraphsToClean.Num(); GraphIndex++)
+	{
+		URigVMGraph* GraphToClean = SortedGraphsToClean[GraphIndex];
+		URigVMController* Controller = RigVMClient.GetOrCreateController(GraphToClean);
+		//TGuardValue<bool> GuardEditGraph(GraphToClean->bEditable, true);
+		FRigVMControllerNotifGuard NotifGuard(Controller, true);
+
+		for (URigVMNode* ModelNode : GraphToClean->GetNodes())
+		{
+			Controller->RemoveUnusedOrphanedPins(ModelNode);
+		}
+	}
+}
+
+void UAnimNextRigVMAssetEditorData::OnRigVMRegistryChanged()
+{
+	RefreshAllModels(ERigVMLoadType::PostLoad);
+	//RebuildGraphFromModel(); // TODO zzz : How we do this on AnimNext ?
+}
+
+void UAnimNextRigVMAssetEditorData::RequestRigVMInit()
+{
+	// TODO zzz : How we do this on AnimNext ?
+}
+
+URigVMGraph* UAnimNextRigVMAssetEditorData::GetModel(const UEdGraph* InEdGraph) const
+{
+	return RigVMClient.GetModel(InEdGraph);
+}
+
+URigVMGraph* UAnimNextRigVMAssetEditorData::GetModel(const FString& InNodePath) const
+{
+	return RigVMClient.GetModel(InNodePath);
+}
+
+URigVMGraph* UAnimNextRigVMAssetEditorData::GetDefaultModel() const 
+{
+	return RigVMClient.GetDefaultModel();
+}
+
+TArray<URigVMGraph*> UAnimNextRigVMAssetEditorData::GetAllModels() const
+{
+	return RigVMClient.GetAllModels(true, true);
+}
+
+URigVMFunctionLibrary* UAnimNextRigVMAssetEditorData::GetLocalFunctionLibrary() const
+{
+	return RigVMClient.GetFunctionLibrary();
+}
+
+URigVMGraph* UAnimNextRigVMAssetEditorData::AddModel(FString InName, bool bSetupUndoRedo, bool bPrintPythonCommand)
+{
+	TGuardValue<bool> EnablePythonPrint(bSuspendPythonMessagesForRigVMClient, !bPrintPythonCommand);
+	return RigVMClient.AddModel(InName, bSetupUndoRedo, bPrintPythonCommand);
+}
+
+bool UAnimNextRigVMAssetEditorData::RemoveModel(FString InName, bool bSetupUndoRedo, bool bPrintPythonCommand)
+{
+	TGuardValue<bool> EnablePythonPrint(bSuspendPythonMessagesForRigVMClient, !bPrintPythonCommand);
+	return RigVMClient.RemoveModel(InName, bSetupUndoRedo, bPrintPythonCommand);
+}
+
+FRigVMGetFocusedGraph& UAnimNextRigVMAssetEditorData::OnGetFocusedGraph()
+{
+	return RigVMClient.OnGetFocusedGraph();
+}
+
+const FRigVMGetFocusedGraph& UAnimNextRigVMAssetEditorData::OnGetFocusedGraph() const
+{
+	return RigVMClient.OnGetFocusedGraph();
+}
+
+URigVMGraph* UAnimNextRigVMAssetEditorData::GetFocusedModel() const
+{
+	return RigVMClient.GetFocusedModel();
+}
+
+URigVMController* UAnimNextRigVMAssetEditorData::GetController(const URigVMGraph* InGraph) const
+{
+	return RigVMClient.GetController(InGraph);
+};
+
+URigVMController* UAnimNextRigVMAssetEditorData::GetControllerByName(const FString InGraphName) const
+{
+	return RigVMClient.GetControllerByName(InGraphName);
+};
+
+URigVMController* UAnimNextRigVMAssetEditorData::GetOrCreateController(URigVMGraph* InGraph)
+{
+	return RigVMClient.GetOrCreateController(InGraph);
+};
+
+URigVMController* UAnimNextRigVMAssetEditorData::GetController(const UEdGraph* InEdGraph) const
+{
+	return RigVMClient.GetController(InEdGraph);
+};
+
+URigVMController* UAnimNextRigVMAssetEditorData::GetOrCreateController(const UEdGraph* InEdGraph)
+{
+	return RigVMClient.GetOrCreateController(InEdGraph);
+};
+
+TArray<FString> UAnimNextRigVMAssetEditorData::GeneratePythonCommands(const FString InNewBlueprintName)
+{
+	return TArray<FString>();
+}
+
+FRigVMClient* UAnimNextRigVMAssetEditorData::GetRigVMClient()
+{
+	return &RigVMClient;
+}
+
+const FRigVMClient* UAnimNextRigVMAssetEditorData::GetRigVMClient() const
+{
+	return &RigVMClient;
+}
+
+IRigVMGraphFunctionHost* UAnimNextRigVMAssetEditorData::GetRigVMGraphFunctionHost()
+{
+	return this;
+}
+
+const IRigVMGraphFunctionHost* UAnimNextRigVMAssetEditorData::GetRigVMGraphFunctionHost() const
+{
+	return this;
+}
+
+void UAnimNextRigVMAssetEditorData::HandleRigVMGraphAdded(const FRigVMClient* InClient, const FString& InNodePath)
+{
+	if(URigVMGraph* RigVMGraph = InClient->GetModel(InNodePath))
+	{
+		RigVMGraph->SetExecuteContextStruct(GetExecuteContextStruct());
+
+		if(!HasAnyFlags(RF_ClassDefaultObject | RF_NeedInitialization | RF_NeedLoad | RF_NeedPostLoad) &&
+			GetOuter() != GetTransientPackage())
+		{
+			CreateEdGraph(RigVMGraph, true);
+			RecompileVM();
+		}
+		
+#if WITH_EDITOR
+		if(!bSuspendPythonMessagesForRigVMClient)
+		{
+			const FString AssetName = InClient->GetSchema()->GetSanitizedName(GetName(), true, false);
+			RigVMPythonUtils::Print(AssetName, FString::Printf(TEXT("asset.add_graph('%s')"), *RigVMGraph->GetName()));
+		}
+#endif
+	}
+}
+
+void UAnimNextRigVMAssetEditorData::HandleRigVMGraphRemoved(const FRigVMClient* InClient, const FString& InNodePath)
+{
+	if(URigVMGraph* RigVMGraph = InClient->GetModel(InNodePath))
+	{
+		RemoveEdGraph(RigVMGraph);
+		RecompileVM();
+
+#if WITH_EDITOR
+		if(!bSuspendPythonMessagesForRigVMClient)
+		{
+			const FString AssetName = InClient->GetSchema()->GetSanitizedName(GetName(), true, false);
+			RigVMPythonUtils::Print(AssetName, FString::Printf(TEXT("asset.add_graph('%s')"), *RigVMGraph->GetName()));
+		}
+#endif
+	}
+}
+
+void UAnimNextRigVMAssetEditorData::HandleConfigureRigVMController(const FRigVMClient* InClient, URigVMController* InControllerToConfigure)
+{
+	InControllerToConfigure->OnModified().AddUObject(this, &UAnimNextRigVMAssetEditorData::HandleModifiedEvent);
+
+	TWeakObjectPtr<UAnimNextRigVMAssetEditorData> WeakThis(this);
+
+	// this delegate is used by the controller to retrieve the current bytecode of the VM
+	InControllerToConfigure->GetCurrentByteCodeDelegate.BindLambda([WeakThis]() -> const FRigVMByteCode*
+	{
+		if (WeakThis.IsValid())
+		{
+			if(UAnimNextRigVMAsset* Asset = WeakThis->GetTypedOuter<UAnimNextRigVMAsset>())
+			{
+				if (Asset->VM)
+				{
+					return &Asset->VM->GetByteCode();
+				}
+			}
+		}
+		return nullptr;
+
+	});
+
+#if WITH_EDITOR
+	InControllerToConfigure->SetupDefaultUnitNodeDelegates(TDelegate<FName(FRigVMExternalVariable, FString)>::CreateLambda(
+		[](FRigVMExternalVariable InVariableToCreate, FString InDefaultValue) -> FName
+		{
+			return NAME_None;
+		}
+	));
+#endif
+}
+
+UObject* UAnimNextRigVMAssetEditorData::GetEditorObjectForRigVMGraph(URigVMGraph* InVMGraph) const
+{
+	if(InVMGraph)
+	{
+		for(UAnimNextRigVMAssetEntry* Entry : Entries)
+		{
+			if(IAnimNextRigVMGraphInterface* GraphInterface = Cast<IAnimNextRigVMGraphInterface>(Entry))
+			{
+				URigVMEdGraph* EdGraph = GraphInterface->GetEdGraph();
+				if(EdGraph->ModelNodePath == InVMGraph->GetNodePath())
+				{
+					return EdGraph;
+				}
+			}
+		}
+	}
+	return nullptr;
+}
+
+URigVMGraph* UAnimNextRigVMAssetEditorData::GetRigVMGraphForEditorObject(UObject* InObject) const
+{
+	if(const URigVMEdGraph* Graph = Cast<URigVMEdGraph>(InObject))
+	{
+		if (Graph->bIsFunctionDefinition)
+		{
+			if (URigVMLibraryNode* LibraryNode = RigVMClient.GetFunctionLibrary()->FindFunction(*Graph->ModelNodePath))
+			{
+				return LibraryNode->GetContainedGraph();
+			}
+		}
+		else
+		{
+			return RigVMClient.GetModel(Graph->ModelNodePath);
+		}
+	}
+
+	return nullptr;
+}
+
+FRigVMGraphFunctionStore* UAnimNextRigVMAssetEditorData::GetRigVMGraphFunctionStore()
+{
+	return &GraphFunctionStore;
+}
+
+const FRigVMGraphFunctionStore* UAnimNextRigVMAssetEditorData::GetRigVMGraphFunctionStore() const
+{
+	return &GraphFunctionStore;
+}
+
+void UAnimNextRigVMAssetEditorData::RecompileVMIfRequired()
+{
+	if (bVMRecompilationRequired)
+	{
+		RecompileVM();
+	}
+}
+
+void UAnimNextRigVMAssetEditorData::RequestAutoVMRecompilation()
+{
+	bVMRecompilationRequired = true;
+	if (bAutoRecompileVM && VMRecompilationBracket == 0)
+	{
+		RecompileVMIfRequired();
+	}
+}
+
+void UAnimNextRigVMAssetEditorData::SetAutoVMRecompile(bool bAutoRecompile)
+{
+	bAutoRecompileVM = bAutoRecompile;
+}
+
+bool UAnimNextRigVMAssetEditorData::GetAutoVMRecompile() const
+{
+	return bAutoRecompileVM;
+}
+
+void UAnimNextRigVMAssetEditorData::IncrementVMRecompileBracket()
+{
+	VMRecompilationBracket++;
+}
+
+void UAnimNextRigVMAssetEditorData::DecrementVMRecompileBracket()
+{
+	if (VMRecompilationBracket == 1)
+	{
+		if (bAutoRecompileVM)
+		{
+			RecompileVMIfRequired();
+		}
+		VMRecompilationBracket = 0;
+	}
+	else if (VMRecompilationBracket > 0)
+	{
+		VMRecompilationBracket--;
+	}
+}
+
+void UAnimNextRigVMAssetEditorData::HandleModifiedEvent(ERigVMGraphNotifType InNotifType, URigVMGraph* InGraph, UObject* InSubject)
+{
+	bool bNotifForOthersPending = true;
+
+	switch(InNotifType)
+	{
+	case ERigVMGraphNotifType::InteractionBracketOpened:
+		{
+			IncrementVMRecompileBracket();
+			break;
+		}
+	case ERigVMGraphNotifType::InteractionBracketClosed:
+	case ERigVMGraphNotifType::InteractionBracketCanceled:
+		{
+			DecrementVMRecompileBracket();
+			break;
+		}
+	case ERigVMGraphNotifType::NodeAdded:
+		{
+			if (URigVMCollapseNode* CollapseNode = Cast<URigVMCollapseNode>(InSubject))
+			{
+				CreateEdGraphForCollapseNode(CollapseNode);
+				break;
+			}
+		}
+		// Fall through to the next case
+	case ERigVMGraphNotifType::LinkAdded:
+	case ERigVMGraphNotifType::LinkRemoved:
+	case ERigVMGraphNotifType::PinArraySizeChanged:
+	case ERigVMGraphNotifType::PinDirectionChanged:
+		{
+			RequestAutoVMRecompilation();
+			break;
+		}
+
+	case ERigVMGraphNotifType::PinDefaultValueChanged:
+		{
+			if (InGraph->GetRuntimeAST().IsValid())
+			{
+				URigVMPin* RootPin = CastChecked<URigVMPin>(InSubject)->GetRootPin();
+				FRigVMASTProxy RootPinProxy = FRigVMASTProxy::MakeFromUObject(RootPin);
+				const FRigVMExprAST* Expression = InGraph->GetRuntimeAST()->GetExprForSubject(RootPinProxy);
+				if (Expression == nullptr)
+				{
+					InGraph->ClearAST();
+				}
+				else if (Expression->NumParents() > 1)
+				{
+					InGraph->ClearAST();
+				}
+			}
+
+			RequestAutoVMRecompilation();	// We need to rebuild our metadata when a default value changes
+			break;
+		}
+	}
+	
+	// if the notification still has to be sent...
+	if (bNotifForOthersPending && !bSuspendModelNotificationsForOthers)
+	{
+		if (RigVMGraphModifiedEvent.IsBound())
+		{
+			RigVMGraphModifiedEvent.Broadcast(InNotifType, InGraph, InSubject);
+		}
+	}
+}
+
+TArray<UEdGraph*> UAnimNextRigVMAssetEditorData::GetAllEdGraphs() const
+{
+	TArray<UEdGraph*> Graphs;
+	for(UAnimNextRigVMAssetEntry* Entry : Entries)
+	{
+		if(IAnimNextRigVMGraphInterface* GraphInterface = Cast<IAnimNextRigVMGraphInterface>(Entry))
+		{
+			Graphs.Add(GraphInterface->GetEdGraph());
+		}
+	}
+	return Graphs;
+}
+
+UAnimNextRigVMAssetEntry* UAnimNextRigVMAssetLibrary::FindEntry(UAnimNextRigVMAsset* InAsset, FName InName)
+{
+	return UE::AnimNext::UncookedOnly::FUtils::GetEditorData(InAsset)->FindEntry(InName);
+}
+
+UAnimNextRigVMAssetEntry* UAnimNextRigVMAssetEditorData::FindEntry(FName InName) const
+{
+	if(InName == NAME_None)
+	{
+		ReportError(TEXT("UAnimNextRigVMAssetEditorData::FindEntry: Invalid name supplied."));
+		return nullptr;
+	}
+
+	const TObjectPtr<UAnimNextRigVMAssetEntry>* FoundEntry = Entries.FindByPredicate([InName](const UAnimNextRigVMAssetEntry* InEntry)
+	{
+		return InEntry->GetEntryName() == InName;
+	});
+
+	return FoundEntry != nullptr ? *FoundEntry : nullptr;
+}
+
+bool UAnimNextRigVMAssetLibrary::RemoveEntry(UAnimNextRigVMAsset* InAsset, UAnimNextRigVMAssetEntry* InEntry,  bool bSetupUndoRedo, bool bPrintPythonCommand)
+{
+	return UE::AnimNext::UncookedOnly::FUtils::GetEditorData(InAsset)->RemoveEntry(InEntry, bSetupUndoRedo, bPrintPythonCommand);
+}
+
+bool UAnimNextRigVMAssetEditorData::RemoveEntry(UAnimNextRigVMAssetEntry* InEntry, bool bSetupUndoRedo, bool bPrintPythonCommand)
+{
+	if(InEntry == nullptr)
+	{
+		ReportError(TEXT("UAnimNextRigVMAssetEditorData::RemoveEntry: Invalid entry supplied."));
+		return false;
+	}
+	
+	TObjectPtr<UAnimNextRigVMAssetEntry>* EntryToRemovePtr = Entries.FindByKey(InEntry);
+	if(EntryToRemovePtr == nullptr)
+	{
+		ReportError(TEXT("UAnimNextRigVMAssetEditorData::RemoveEntry: Asset does not contain the supplied entry."));
+		return false;
+	}
+
+	if(bSetupUndoRedo)
+	{
+		Modify();
+	}
+
+	// Remove from internal array
+	UAnimNextRigVMAssetEntry* EntryToRemove = *EntryToRemovePtr;
+	Entries.Remove(EntryToRemove);
+
+	if (bSetupUndoRedo)
+	{
+		EntryToRemove->Modify();
+	}
+
+	bool bResult = true;
+	if(const IAnimNextRigVMGraphInterface* GraphInterface = Cast<IAnimNextRigVMGraphInterface>(EntryToRemove))
+	{
+		// Remove any graphs
+		if(URigVMGraph* RigVMGraph = GraphInterface->GetRigVMGraph())
+		{
+			TGuardValue<bool> EnablePythonPrint(bSuspendPythonMessagesForRigVMClient, !bPrintPythonCommand);
+			TGuardValue<bool> DisableAutoCompile(bAutoRecompileVM, false);
+			bResult = RigVMClient.RemoveModel(RigVMGraph->GetNodePath(), bSetupUndoRedo);
+		}
+	}
+
+	// This will cause any external package to be removed when saved
+	EntryToRemove->MarkAsGarbage();
+
+	BroadcastModified();
+
+	return bResult;
+}
+
+bool UAnimNextRigVMAssetLibrary::RemoveEntries(UAnimNextRigVMAsset* InAsset, const TArray<UAnimNextRigVMAssetEntry*>& InEntries,  bool bSetupUndoRedo, bool bPrintPythonCommand)
+{
+	return UE::AnimNext::UncookedOnly::FUtils::GetEditorData(InAsset)->RemoveEntries(InEntries, bSetupUndoRedo, bPrintPythonCommand);
+}
+
+bool UAnimNextRigVMAssetEditorData::RemoveEntries(const TArray<UAnimNextRigVMAssetEntry*>& InEntries, bool bSetupUndoRedo, bool bPrintPythonCommand)
+{
+	bool bResult = false;
+	{
+		TGuardValue<bool> DisableEditorDataNotifications(bSuspendEditorDataNotifications, true);
+		TGuardValue<bool> DisableAutoCompile(bAutoRecompileVM, false);
+		for(UAnimNextRigVMAssetEntry* Entry : InEntries)
+		{
+			bResult |= RemoveEntry(Entry, bSetupUndoRedo, bPrintPythonCommand);
+		}
+	}
+
+	BroadcastModified();
+
+	return bResult;
+}
+
+UObject* UAnimNextRigVMAssetEditorData::CreateNewSubEntry(UAnimNextRigVMAssetEditorData* InEditorData, TSubclassOf<UObject> InClass)
+{
+	UObject* NewEntry = NewObject<UObject>(InEditorData, InClass.Get(), NAME_None, RF_Transactional);
+	// If we are a transient asset, dont use external packages
+	UAnimNextRigVMAsset* Asset = UE::AnimNext::UncookedOnly::FUtils::GetAsset(InEditorData);
+	check(Asset);
+	if(!Asset->HasAnyFlags(RF_Transient))
+	{
+		FExternalPackageHelper::SetPackagingMode(NewEntry, InEditorData, true, false, PKG_None);
+	}
+	return NewEntry;
+}
+
+UAnimNextRigVMAssetEntry* UAnimNextRigVMAssetEditorData::FindEntryForRigVMGraph(URigVMGraph* InRigVMGraph) const
+{
+	for(UAnimNextRigVMAssetEntry* Entry : Entries)
+	{
+		if(Entry->GetEntryName() == InRigVMGraph->GetFName())
+		{
+			return Entry;
+		}
+	}
+
+	return nullptr;
+}
+
+UAnimNextRigVMAssetEntry* UAnimNextRigVMAssetEditorData::FindEntryForEdGraph(URigVMEdGraph* InEdGraph) const
+{
+	for(UAnimNextRigVMAssetEntry* Entry : Entries)
+	{
+		if(MakeEdGraphName(Entry->GetEntryName()) == InEdGraph->GetFName())
+		{
+			return Entry;
+		}
+	}
+
+	return nullptr;
+}
+
+FName UAnimNextRigVMAssetEditorData::MakeEdGraphName(FName InRigVMGraphName)
+{
+	return *FString::Printf(TEXT("%s_EdGraph"), *InRigVMGraphName.ToString());
+}
\ No newline at end of file
diff --git a/Engine/Plugins/Experimental/AnimNext/Source/AnimNextUncookedOnly/Private/AnimNextRigVMAssetEntry.cpp b/Engine/Plugins/Experimental/AnimNext/Source/AnimNextUncookedOnly/Private/AnimNextRigVMAssetEntry.cpp
new file mode 100644
--- /dev/null
+++ b/Engine/Plugins/Experimental/AnimNext/Source/AnimNextUncookedOnly/Private/AnimNextRigVMAssetEntry.cpp
@@ -0,0 +1,25 @@
+﻿// Copyright Epic Games, Inc. All Rights Reserved.
+
+#include "AnimNextRigVMAssetEntry.h"
+#include "AnimNextRigVMAssetEditorData.h"
+
+void UAnimNextRigVMAssetEntry::Initialize(UAnimNextRigVMAssetEditorData* InEditorData)
+{
+	InEditorData->RigVMGraphModifiedEvent.RemoveAll(this);
+	InEditorData->RigVMGraphModifiedEvent.AddUObject(this, &UAnimNextRigVMAssetEntry::HandleRigVMGraphModifiedEvent);
+}
+
+bool UAnimNextRigVMAssetEntry::IsAsset() const
+{
+	// Entries are considered assets to allow using the asset logic for save dialogs, etc.
+	// Also, they return true even if pending kill, in order to show up as deleted in these dialogs.
+	return IsPackageExternal() && !GetPackage()->HasAnyFlags(RF_Transient) && !HasAnyFlags(RF_Transient | RF_ClassDefaultObject);
+}
+
+void UAnimNextRigVMAssetEntry::BroadcastModified()
+{
+	if(UAnimNextRigVMAssetEditorData* EditorData = Cast<UAnimNextRigVMAssetEditorData>(GetOuter()))
+	{
+		EditorData->BroadcastModified();
+	}
+}
\ No newline at end of file
diff --git a/Engine/Plugins/Experimental/AnimNext/Source/AnimNextUncookedOnly/Private/Graph/AnimNextGraphEntry.cpp b/Engine/Plugins/Experimental/AnimNext/Source/AnimNextUncookedOnly/Private/Graph/AnimNextGraphEntry.cpp
--- a/Engine/Plugins/Experimental/AnimNext/Source/AnimNextUncookedOnly/Private/Graph/AnimNextGraphEntry.cpp
+++ b/Engine/Plugins/Experimental/AnimNext/Source/AnimNextUncookedOnly/Private/Graph/AnimNextGraphEntry.cpp
@@ -1,40 +1,20 @@
 ﻿// Copyright Epic Games, Inc. All Rights Reserved.
 
 #include "Graph/AnimNextGraphEntry.h"
-#include "Graph/AnimNextGraph_EditorData.h"
+#include "RigVMModel/RigVMGraph.h"
+#include "Graph/AnimNextGraph_EdGraph.h"
 
-void UAnimNextGraphEntry::Initialize(UAnimNextGraph_EditorData* InEditorData)
+FName UAnimNextGraphEntry::GetEntryName() const
 {
-	InEditorData->RigVMGraphModifiedEvent.RemoveAll(this);
-	InEditorData->RigVMGraphModifiedEvent.AddUObject(this, &UAnimNextGraphEntry::HandleRigVMGraphModifiedEvent);
+	return GraphName;
 }
 
-void UAnimNextGraphEntry::HandleRigVMGraphModifiedEvent(ERigVMGraphNotifType InNotifType, URigVMGraph* InGraph, UObject* InSubject)
+URigVMGraph* UAnimNextGraphEntry::GetRigVMGraph() const
 {
-	
+	return Graph;
 }
 
-bool UAnimNextGraphEntry::IsAsset() const
+URigVMEdGraph* UAnimNextGraphEntry::GetEdGraph() const
 {
-	// Entries are considered assets to allow using the asset logic for save dialogs, etc.
-	// Also, they return true even if pending kill, in order to show up as deleted in these dialogs.
-	return IsPackageExternal() && !GetPackage()->HasAnyFlags(RF_Transient) && !HasAnyFlags(RF_Transient | RF_ClassDefaultObject);
+	return EdGraph;
 }
-
-FText UAnimNextGraphEntry::GetDisplayName() const
-{
-	return FText::FromName(GraphName);
-}
-
-FText UAnimNextGraphEntry::GetDisplayNameTooltip() const
-{
-	return FText::FromName(GraphName);
-}
-
-void UAnimNextGraphEntry::BroadcastModified()
-{
-	if(UAnimNextGraph_EditorData* EditorData = Cast<UAnimNextGraph_EditorData>(GetOuter()))
-	{
-		EditorData->BroadcastModified();
-	}
-}
\ No newline at end of file
diff --git a/Engine/Plugins/Experimental/AnimNext/Source/AnimNextUncookedOnly/Private/Graph/AnimNextGraph_EditorData.cpp b/Engine/Plugins/Experimental/AnimNext/Source/AnimNextUncookedOnly/Private/Graph/AnimNextGraph_EditorData.cpp
--- a/Engine/Plugins/Experimental/AnimNext/Source/AnimNextUncookedOnly/Private/Graph/AnimNextGraph_EditorData.cpp
+++ b/Engine/Plugins/Experimental/AnimNext/Source/AnimNextUncookedOnly/Private/Graph/AnimNextGraph_EditorData.cpp
@@ -2,585 +2,58 @@
 
 #include "Graph/AnimNextGraph_EditorData.h"
 
-#include "ControlRigDefines.h"
-#include "ExternalPackageHelper.h"
+#include "UncookedOnlyUtils.h"
+#include "Curves/CurveFloat.h"
 #include "Graph/AnimNextGraph.h"
-#include "Graph/AnimNextGraph_Controller.h"
+#include "Graph/AnimNextGraphEntry.h"
 #include "Graph/AnimNextGraph_EdGraphSchema.h"
-#include "UncookedOnlyUtils.h"
-#include "Graph/AnimNextExecuteContext.h"
-#include "Rigs/RigHierarchyPose.h"
 #include "RigVMModel/RigVMFunctionLibrary.h"
 #include "RigVMModel/RigVMNotifications.h"
 #include "RigVMModel/Nodes/RigVMCollapseNode.h"
-#include "Curves/CurveFloat.h"
-#include "UObject/AssetRegistryTagsContext.h"
-#include "UObject/ObjectSaveContext.h"
-#include "Kismet2/BlueprintEditorUtils.h"
-#include "Graph/AnimNextGraphEntry.h"
-#include "Param/RigVMDispatch_GetLayerParameter.h"
-#include "Param/RigVMDispatch_GetParameter.h"
-
-namespace UE::AnimNext::Graph::Private
-{
-
-static UAnimNextGraphEntry* CreateNewGraphEntry(UAnimNextGraph_EditorData* InEditorData, TSubclassOf<UAnimNextGraphEntry> InClass)
-{
-	UAnimNextGraphEntry* NewEntry = NewObject<UAnimNextGraphEntry>(InEditorData, InClass.Get(), NAME_None, RF_Transactional);
-	// If we are a transient asset, dont use external packages
-	UAnimNextGraph* Graph = UncookedOnly::FUtils::GetGraph(InEditorData);
-	check(Graph);
-	if(!Graph->HasAnyFlags(RF_Transient))
-	{
-		FExternalPackageHelper::SetPackagingMode(NewEntry, InEditorData, true, false, PKG_None);
-	}
-	return NewEntry;
-}
-
-template<typename EntryClassType>
-static EntryClassType* CreateNewGraphEntry(UAnimNextGraph_EditorData* InEditorData)
-{
-	return CastChecked<EntryClassType>(CreateNewGraphEntry(InEditorData, EntryClassType::StaticClass()));
-}
-
-}
-
-UAnimNextGraph_EditorData::UAnimNextGraph_EditorData(const FObjectInitializer& ObjectInitializer)
-{
-	RigVMClient.Reset();
-	RigVMClient.SetSchemaClass(UAnimNextGraph_Schema::StaticClass());
-	RigVMClient.SetControllerClass(UAnimNextGraph_Controller::StaticClass());
-	RigVMClient.SetOuterClientHost(this, GET_MEMBER_NAME_CHECKED(UAnimNextGraph_EditorData, RigVMClient));
-	{
-		TGuardValue<bool> DisableClientNotifs(RigVMClient.bSuspendNotifications, true);
-		RigVMClient.GetOrCreateFunctionLibrary(false, &ObjectInitializer);
-	}
-	RigVMClient.SetExecuteContextStruct(FAnimNextExecuteContext::StaticStruct());
-}
-
-
-void UAnimNextGraph_EditorData::BroadcastModified()
-{
-	RecompileVM();
-
-	if(!bSuspendGraphNotifications)
-	{
-		ModifiedDelegate.Broadcast(this);
-	}
-}
-
-void UAnimNextGraph_EditorData::ReportError(const TCHAR* InMessage) const
-{
-#if WITH_EDITOR
-	FScriptExceptionHandler::Get().HandleException(ELogVerbosity::Error, InMessage, TEXT(""));
-#endif
-}
-
-void UAnimNextGraph_EditorData::ReconstructAllNodes()
-{
-#ifdef WITH_EDITORONLY_DATA
-	// Avoid refreshing EdGraph nodes during cook
-	if (GIsCookerLoadingPackage)
-	{
-		return;
-	}
-	
-	if (GetRigVMClient()->GetDefaultModel() == nullptr)
-	{
-		return;
-	}
-
-	TArray<URigVMEdGraphNode*> AllNodes;
-	UE::AnimNext::UncookedOnly::FUtils::GetAllNodesOfClass(this, AllNodes);
-
-	for (URigVMEdGraphNode* Node : AllNodes)
-	{
-		Node->SetFlags(RF_Transient);
-	}
-
-	for(URigVMEdGraphNode* Node : AllNodes)
-	{
-		Node->ReconstructNode();
-	}
-
-	for (URigVMEdGraphNode* Node : AllNodes)
-	{
-		Node->ClearFlags(RF_Transient);
-	}
-#endif	
-}
-
-void UAnimNextGraph_EditorData::Serialize(FArchive& Ar)
-{
-	RigVMClient.SetOuterClientHost(this, GET_MEMBER_NAME_CHECKED(UAnimNextGraph_EditorData, RigVMClient));
-
-	UObject::Serialize(Ar);
-}
-
-void UAnimNextGraph_EditorData::Initialize(bool bRecompileVM)
-{
-	UAnimNextGraph* AnimNextGraph = GetTypedOuter<UAnimNextGraph>();
-
-	if (RigVMClient.GetController(0) == nullptr)
-	{
-		check(RigVMClient.Num() == 1);
-		check(RigVMClient.GetFunctionLibrary());
-		
-		RigVMClient.GetOrCreateController(RigVMClient.GetDefaultModel());
-		RigVMClient.GetOrCreateController(RigVMClient.GetFunctionLibrary());
-
-		// Init function library controllers
-		for(URigVMLibraryNode* LibraryNode : RigVMClient.GetFunctionLibrary()->GetFunctions())
-		{
-			RigVMClient.GetOrCreateController(LibraryNode->GetContainedGraph());
-		}
-		
-		if(bRecompileVM)
-		{
-			UE::AnimNext::UncookedOnly::FUtils::Compile(AnimNextGraph);
-		}
-	}
-
-	for(UAnimNextGraph_EdGraph* Graph : Graphs)
-	{
-		Graph->Initialize(this);
-	}
-}
+#include "UObject/FortniteMainBranchObjectVersion.h"
+#include "UObject/LinkerLoad.h"
 
 void UAnimNextGraph_EditorData::PostLoad()
 {
 	Super::PostLoad();
-	
-	Initialize(/*bRecompileVM*/false);
-	RefreshAllModels(ERigVMLoadType::PostLoad);
-
-#if WITH_EDITOR
-	FExternalPackageHelper::LoadObjectsFromExternalPackages<UAnimNextGraphEntry>(this, [this](UAnimNextGraphEntry* InLoadedEntry)
-	{
-		check(IsValid(InLoadedEntry));
-		Entries.Add(InLoadedEntry);
-	});
-	
-	// delay compilation until the package has been loaded
-	FCoreUObjectDelegates::OnEndLoadPackage.AddUObject(this, &UAnimNextGraph_EditorData::HandlePackageDone);
-#else // !WITH_EDITOR
-	RecompileVMIfRequired();
-#endif // WITH_EDITOR
-}
-
-void UAnimNextGraph_EditorData::PostDuplicate(EDuplicateMode::Type DuplicateMode)
-{
-	Initialize(/*bRecompileVM*/true);
-}
-
-#if WITH_EDITOR
-void UAnimNextGraph_EditorData::GetAssetRegistryTags(TArray<FAssetRegistryTag>& OutTags) const
-{
-	PRAGMA_DISABLE_DEPRECATION_WARNINGS;
-	Super::GetAssetRegistryTags(OutTags);
-	PRAGMA_ENABLE_DEPRECATION_WARNINGS;
-}
-
-void UAnimNextGraph_EditorData::GetAssetRegistryTags(FAssetRegistryTagsContext Context) const
-{
-	Super::GetAssetRegistryTags(Context);
-
-	FAnimNextParameterProviderAssetRegistryExports ExportParameters;
-	
-	for (const TObjectPtr<UAnimNextGraph_EdGraph>& Graph : Graphs)
-	{
-		if (const URigVMGraph* RigGraph = Graph->GetModel())
-		{
-			UE::AnimNext::UncookedOnly::FUtils::GetGraphParameters(RigGraph, ExportParameters);
-		}
-	}
-
-	if (ExportParameters.Parameters.Num())
-	{
-		FString TagValue;	
-		FAnimNextParameterProviderAssetRegistryExports::StaticStruct()->ExportText(TagValue, &ExportParameters, nullptr, nullptr, PPF_None, nullptr);
-		Context.AddTag(FAssetRegistryTag(UE::AnimNext::ExportsAnimNextAssetRegistryTag, TagValue, FAssetRegistryTag::TT_Hidden));	
-	}
-}
-
-void UAnimNextGraph_EditorData::HandlePackageDone(const FEndLoadPackageContext& Context)
-{
-	if (!Context.LoadedPackages.Contains(GetPackage()))
-	{
-		return;
-	}
-	HandlePackageDone();
-}
-
-void UAnimNextGraph_EditorData::HandlePackageDone()
-{
-	FCoreUObjectDelegates::OnEndLoadPackage.RemoveAll(this);
-
-	RecompileVM();
-
-	ReconstructAllNodes();
-}
-
-void UAnimNextGraph_EditorData::RefreshAllModels(ERigVMLoadType InLoadType)
-{
-	const bool bIsPostLoad = InLoadType == ERigVMLoadType::PostLoad;
-
-	TGuardValue<bool> IsCompilingGuard(bIsCompiling, true);
-	TGuardValue<bool> ClientIgnoreModificationsGuard(RigVMClient.bIgnoreModelNotifications, true);
-
-	TArray<URigVMGraph*> GraphsToDetach = RigVMClient.GetAllModels(true, false);
-	TMap<const URigVMGraph*, TArray<URigVMController::FLinkedPath>> LinkedPaths;
-	
-	if (ensure(IsInGameThread()))
-	{
-		for (const URigVMGraph* GraphToDetach : GraphsToDetach)
-		{
-			URigVMController* Controller = RigVMClient.GetOrCreateController(GraphToDetach);
-			// temporarily disable default value validation during load time, serialized values should always be accepted
-			TGuardValue<bool> PerGraphDisablePinDefaultValueValidation(Controller->bValidatePinDefaults, false);
-			FRigVMControllerNotifGuard NotifGuard(Controller, true);
-			LinkedPaths.Add(GraphToDetach, Controller->GetLinkedPaths());
-			Controller->FastBreakLinkedPaths(LinkedPaths.FindChecked(GraphToDetach));
-			TArray<URigVMNode*> Nodes = GraphToDetach->GetNodes();
-			for (URigVMNode* Node : Nodes)
-			{
-				Controller->RepopulatePinsOnNode(Node, true, true);
-			}
-		}
-		//SetupPinRedirectorsForBackwardsCompatibility();
-	}
-
-	for (const URigVMGraph* GraphToDetach : GraphsToDetach)
-	{
-		URigVMController* Controller = RigVMClient.GetOrCreateController(GraphToDetach);
-		// at this stage, allow all links to be reattached,
-		// RecomputeAllTemplateFilteredPermutations() later should break any invalid links
-		FRigVMControllerNotifGuard NotifGuard(Controller, true);
-		Controller->RestoreLinkedPaths(LinkedPaths.FindChecked(GraphToDetach));
-	}
-
-	if (bIsPostLoad)
-	{
-		//PatchTemplateNodesWithPreferredPermutation();
-	}
-
-	TArray<URigVMGraph*> GraphsToClean = RigVMClient.GetAllModels(true, true);
 
-	// Sort from leaf graphs to root
-	TArray<URigVMGraph*> SortedGraphsToClean;
-	SortedGraphsToClean.Reserve(GraphsToClean.Num());
-	while (SortedGraphsToClean.Num() < GraphsToClean.Num())
+	if(GetLinkerCustomVersion(FFortniteMainBranchObjectVersion::GUID) < FFortniteMainBranchObjectVersion::AnimNextMoveGraphsToEntries)
 	{
-		bool bGraphAdded = false;
-		for (URigVMGraph* Graph : GraphsToClean)
+		for(TObjectPtr<UAnimNextGraph_EdGraph> Graph : Graphs_DEPRECATED)
 		{
-			if (SortedGraphsToClean.Contains(Graph))
+			// Must preload entries so their data is populated or we cannot find the appropriate entries for graphs
+			for(UAnimNextRigVMAssetEntry* Entry : Entries) 
 			{
-				continue;
+				Entry->GetLinker()->Preload(Entry);
 			}
-
-			TArray<URigVMGraph*> ContainedGraphs;
-			for (URigVMNode* Node : Graph->GetNodes())
+			
+			URigVMGraph* FoundRigVMGraph = GetRigVMGraphForEditorObject(Graph);
+			if(FoundRigVMGraph)
 			{
-				if (URigVMLibraryNode* LibraryNode = Cast<URigVMLibraryNode>(Node))
+				UAnimNextRigVMAssetEntry* FoundEntry = nullptr;
+				for(UAnimNextRigVMAssetEntry* Entry : Entries)
 				{
-					if (URigVMFunctionReferenceNode* FunctionReferenceNode = Cast<URigVMFunctionReferenceNode>(LibraryNode))
+					if(UAnimNextGraphEntry* GraphEntry = Cast<UAnimNextGraphEntry>(Entry))
 					{
-						if (FunctionReferenceNode->GetReferencedFunctionHeader().LibraryPointer.LibraryNode.GetLongPackageName() != GetPackage()->GetPathName())
-						{
-							continue;
-						}
-						if (URigVMLibraryNode* ReferencedNode = FunctionReferenceNode->LoadReferencedNode())
+						if(FoundRigVMGraph == static_cast<IAnimNextRigVMGraphInterface*>(GraphEntry)->GetRigVMGraph())
 						{
-							ContainedGraphs.Add(ReferencedNode->GetContainedGraph());
-							continue;
+							FoundEntry = Entry;
+							break;
 						}
 					}
-
-					if (URigVMGraph* ContainedGraph = LibraryNode->GetContainedGraph())
-					{
-						ContainedGraphs.Add(ContainedGraph);
-					}
-				}
-			}
-
-			bool bAllContained = true;
-			for (URigVMGraph* Contained : ContainedGraphs)
-			{
-				if (!SortedGraphsToClean.Contains(Contained))
-				{
-					bAllContained = false;
-					break;
-				}
-			}
-			if (bAllContained)
-			{
-				SortedGraphsToClean.Add(Graph);
-				bGraphAdded = true;
-			}
-		}
-		ensure(bGraphAdded);
-	}
-
-	for (int32 GraphIndex = 0; GraphIndex < SortedGraphsToClean.Num(); GraphIndex++)
-	{
-		URigVMGraph* GraphToClean = SortedGraphsToClean[GraphIndex];
-		URigVMController* Controller = RigVMClient.GetOrCreateController(GraphToClean);
-		//TGuardValue<bool> GuardEditGraph(GraphToClean->bEditable, true);
-		FRigVMControllerNotifGuard NotifGuard(Controller, true);
-
-		for (URigVMNode* ModelNode : GraphToClean->GetNodes())
-		{
-			Controller->RemoveUnusedOrphanedPins(ModelNode);
-		}
-	}
-}
-
-void UAnimNextGraph_EditorData::OnRigVMRegistryChanged()
-{
-	RefreshAllModels(ERigVMLoadType::PostLoad);
-	//RebuildGraphFromModel(); // TODO zzz : How we do this on AnimNext ?
-}
-
-void UAnimNextGraph_EditorData::RequestRigVMInit()
-{
-	// TODO zzz : How we do this on AnimNext ?
-}
-
-URigVMGraph* UAnimNextGraph_EditorData::GetModel(const UEdGraph* InEdGraph) const
-{
-	return RigVMClient.GetModel(InEdGraph);
-}
-
-URigVMGraph* UAnimNextGraph_EditorData::GetModel(const FString& InNodePath) const
-{
-	return RigVMClient.GetModel(InNodePath);
-}
-
-URigVMGraph* UAnimNextGraph_EditorData::GetDefaultModel() const 
-{
-	return RigVMClient.GetDefaultModel();
-}
-
-TArray<URigVMGraph*> UAnimNextGraph_EditorData::GetAllModels() const
-{
-	return RigVMClient.GetAllModels(true, true);
-}
-
-URigVMFunctionLibrary* UAnimNextGraph_EditorData::GetLocalFunctionLibrary() const
-{
-	return RigVMClient.GetFunctionLibrary();
-}
-
-URigVMGraph* UAnimNextGraph_EditorData::AddModel(FString InName, bool bSetupUndoRedo, bool bPrintPythonCommand)
-{
-	TGuardValue<bool> EnablePythonPrint(bSuspendPythonMessagesForRigVMClient, !bPrintPythonCommand);
-	return RigVMClient.AddModel(InName, bSetupUndoRedo, bPrintPythonCommand);
-}
-
-bool UAnimNextGraph_EditorData::RemoveModel(FString InName, bool bSetupUndoRedo, bool bPrintPythonCommand)
-{
-	TGuardValue<bool> EnablePythonPrint(bSuspendPythonMessagesForRigVMClient, !bPrintPythonCommand);
-	return RigVMClient.RemoveModel(InName, bSetupUndoRedo, bPrintPythonCommand);
-}
-
-FRigVMGetFocusedGraph& UAnimNextGraph_EditorData::OnGetFocusedGraph()
-{
-	return RigVMClient.OnGetFocusedGraph();
-}
-
-const FRigVMGetFocusedGraph& UAnimNextGraph_EditorData::OnGetFocusedGraph() const
-{
-	return RigVMClient.OnGetFocusedGraph();
-}
-
-URigVMGraph* UAnimNextGraph_EditorData::GetFocusedModel() const
-{
-	return RigVMClient.GetFocusedModel();
-}
-
-URigVMController* UAnimNextGraph_EditorData::GetController(const URigVMGraph* InGraph) const
-{
-	return RigVMClient.GetController(InGraph);
-};
-
-URigVMController* UAnimNextGraph_EditorData::GetControllerByName(const FString InGraphName) const
-{
-	return RigVMClient.GetControllerByName(InGraphName);
-};
-
-URigVMController* UAnimNextGraph_EditorData::GetOrCreateController(URigVMGraph* InGraph)
-{
-	return RigVMClient.GetOrCreateController(InGraph);
-};
-
-URigVMController* UAnimNextGraph_EditorData::GetController(const UEdGraph* InEdGraph) const
-{
-	return RigVMClient.GetController(InEdGraph);
-};
-
-URigVMController* UAnimNextGraph_EditorData::GetOrCreateController(const UEdGraph* InEdGraph)
-{
-	return RigVMClient.GetOrCreateController(InEdGraph);
-};
-
-TArray<FString> UAnimNextGraph_EditorData::GeneratePythonCommands(const FString InNewBlueprintName)
-{
-	return TArray<FString>();
-}
-
-void UAnimNextGraph_EditorData::GetAllGraphs(TArray<UEdGraph*>& OutGraphs) const
-{
-	OutGraphs.Reset();
-	OutGraphs.Append(Graphs);
-}
-
-#endif // WITH_EDITOR
-
-FRigVMClient* UAnimNextGraph_EditorData::GetRigVMClient()
-{
-	return &RigVMClient;
-}
-
-const FRigVMClient* UAnimNextGraph_EditorData::GetRigVMClient() const
-{
-	return &RigVMClient;
-}
-
-IRigVMGraphFunctionHost* UAnimNextGraph_EditorData::GetRigVMGraphFunctionHost()
-{
-	return this;
-}
-
-const IRigVMGraphFunctionHost* UAnimNextGraph_EditorData::GetRigVMGraphFunctionHost() const
-{
-	return this;
-}
-
-
-void UAnimNextGraph_EditorData::HandleRigVMGraphAdded(const FRigVMClient* InClient, const FString& InNodePath)
-{
-	if(URigVMGraph* RigVMGraph = InClient->GetModel(InNodePath))
-	{
-		RigVMGraph->SetExecuteContextStruct(FAnimNextExecuteContext::StaticStruct());
-
-		if(!HasAnyFlags(RF_ClassDefaultObject | RF_NeedInitialization | RF_NeedLoad | RF_NeedPostLoad) &&
-			GetOuter() != GetTransientPackage())
-		{
-			CreateEdGraph(RigVMGraph, true);
-			RecompileVM();
-		}
-	}
-}
-
-void UAnimNextGraph_EditorData::HandleRigVMGraphRemoved(const FRigVMClient* InClient, const FString& InNodePath)
-{
-	if(URigVMGraph* RigVMGraph = InClient->GetModel(InNodePath))
-	{
-		RemoveEdGraph(RigVMGraph);
-		RecompileVM();
-	}
-}
-
-void UAnimNextGraph_EditorData::HandleConfigureRigVMController(const FRigVMClient* InClient,
-                                                                    URigVMController* InControllerToConfigure)
-{
-	InControllerToConfigure->OnModified().AddUObject(this, &UAnimNextGraph_EditorData::HandleModifiedEvent);
-
-	TWeakObjectPtr<UAnimNextGraph_EditorData> WeakThis(this);
-
-	// this delegate is used by the controller to determine variable validity
-	// during a bind process. the controller itself doesn't own the variables,
-	// so we need a delegate to request them from the owning blueprint
-	InControllerToConfigure->GetExternalVariablesDelegate.BindLambda([](URigVMGraph* InGraph) -> TArray<FRigVMExternalVariable>
-	{
-		if (InGraph)
-		{
-			if(UAnimNextGraph_EditorData* EditorData = InGraph->GetTypedOuter<UAnimNextGraph_EditorData>())
-			{
-				if (UAnimNextGraph* Graph = EditorData->GetTypedOuter<UAnimNextGraph>())
-				{
-					return Graph->GetRigVMExternalVariables();
 				}
-			}
-		}
-		return TArray<FRigVMExternalVariable>();
-	});
 
-	// this delegate is used by the controller to retrieve the current bytecode of the VM
-	InControllerToConfigure->GetCurrentByteCodeDelegate.BindLambda([WeakThis]() -> const FRigVMByteCode*
-	{
-		if (WeakThis.IsValid())
-		{
-			if(UAnimNextGraph* Graph = WeakThis->GetTypedOuter<UAnimNextGraph>())
-			{
-				if (Graph->VM)
+				if(FoundEntry)
 				{
-					return &Graph->VM->GetByteCode();
-				}
-			}
-		}
-		return nullptr;
+					UAnimNextGraphEntry* GraphEntry = CastChecked<UAnimNextGraphEntry>(FoundEntry);
+					GraphEntry->EdGraph = Graph;
 
-	});
-
-#if WITH_EDITOR
-	InControllerToConfigure->SetupDefaultUnitNodeDelegates(TDelegate<FName(FRigVMExternalVariable, FString)>::CreateLambda(
-		[](FRigVMExternalVariable InVariableToCreate, FString InDefaultValue) -> FName
-		{
-			return NAME_None;
-		}
-	));
-#endif
-
-}
-
-UObject* UAnimNextGraph_EditorData::GetEditorObjectForRigVMGraph(URigVMGraph* InVMGraph) const
-{
-	if(InVMGraph)
-	{
-		for(UAnimNextGraph_EdGraph* EdGraph : Graphs)
-		{
-			if(EdGraph)
-			{
-				if(EdGraph->ModelNodePath == InVMGraph->GetNodePath())
-				{
-					return EdGraph;
+					Graph->Rename(*MakeEdGraphName(FoundRigVMGraph->GetFName()).ToString(), FoundEntry, REN_ForceNoResetLoaders | REN_DoNotDirty | REN_DontCreateRedirectors | REN_NonTransactional);
+					Graph->Initialize(this);
 				}
 			}
 		}
 	}
-	return nullptr;
-}
-
-URigVMGraph* UAnimNextGraph_EditorData::GetRigVMGraphForEditorObject(UObject* InObject) const
-{
-	if(const UAnimNextGraph_EdGraph* Graph = Cast<UAnimNextGraph_EdGraph>(InObject))
-	{
-		if (Graph->bIsFunctionDefinition)
-		{
-			if (URigVMLibraryNode* LibraryNode = RigVMClient.GetFunctionLibrary()->FindFunction(*Graph->ModelNodePath))
-			{
-				return LibraryNode->GetContainedGraph();
-			}
-		}
-		else
-		{
-			return RigVMClient.GetModel(Graph->ModelNodePath);
-		}
-	}
-
-	return nullptr;
-}
-
-FRigVMGraphFunctionStore* UAnimNextGraph_EditorData::GetRigVMGraphFunctionStore()
-{
-	return &GraphFunctionStore;
-}
-
-const FRigVMGraphFunctionStore* UAnimNextGraph_EditorData::GetRigVMGraphFunctionStore() const
-{
-	return &GraphFunctionStore;
 }
 
 void UAnimNextGraph_EditorData::RecompileVM()
@@ -588,88 +61,10 @@ void UAnimNextGraph_EditorData::RecompileVM()
 	UE::AnimNext::UncookedOnly::FUtils::Compile(GetTypedOuter<UAnimNextGraph>());
 }
 
-void UAnimNextGraph_EditorData::RecompileVMIfRequired()
-{
-	if (bVMRecompilationRequired)
-	{
-		RecompileVM();
-	}
-}
-
-void UAnimNextGraph_EditorData::RequestAutoVMRecompilation()
-{
-	bVMRecompilationRequired = true;
-	if (bAutoRecompileVM && VMRecompilationBracket == 0)
-	{
-		RecompileVMIfRequired();
-	}
-}
-
-void UAnimNextGraph_EditorData::SetAutoVMRecompile(bool bAutoRecompile)
-{
-	bAutoRecompileVM = bAutoRecompile;
-}
-
-bool UAnimNextGraph_EditorData::GetAutoVMRecompile() const
-{
-	return bAutoRecompileVM;
-}
-
-void UAnimNextGraph_EditorData::IncrementVMRecompileBracket()
-{
-	VMRecompilationBracket++;
-}
-
-void UAnimNextGraph_EditorData::DecrementVMRecompileBracket()
-{
-	if (VMRecompilationBracket == 1)
-	{
-		if (bAutoRecompileVM)
-		{
-			RecompileVMIfRequired();
-		}
-		VMRecompilationBracket = 0;
-	}
-	else if (VMRecompilationBracket > 0)
-	{
-		VMRecompilationBracket--;
-	}
-}
-
 void UAnimNextGraph_EditorData::HandleModifiedEvent(ERigVMGraphNotifType InNotifType, URigVMGraph* InGraph, UObject* InSubject)
 {
-	bool bNotifForOthersPending = true;
-
 	switch(InNotifType)
 	{
-	case ERigVMGraphNotifType::InteractionBracketOpened:
-		{
-			IncrementVMRecompileBracket();
-			break;
-		}
-	case ERigVMGraphNotifType::InteractionBracketClosed:
-	case ERigVMGraphNotifType::InteractionBracketCanceled:
-		{
-			DecrementVMRecompileBracket();
-			break;
-		}
-	case ERigVMGraphNotifType::NodeAdded:
-		{
-			if (URigVMCollapseNode* CollapseNode = Cast<URigVMCollapseNode>(InSubject))
-			{
-				CreateEdGraphForCollapseNode(CollapseNode);
-				break;
-			}
-		}
-		// Fall through to the next case
-	case ERigVMGraphNotifType::LinkAdded:
-	case ERigVMGraphNotifType::LinkRemoved:
-	case ERigVMGraphNotifType::PinArraySizeChanged:
-	case ERigVMGraphNotifType::PinDirectionChanged:
-		{
-			RequestAutoVMRecompilation();
-			break;
-		}
 	case ERigVMGraphNotifType::PinAdded:
 		{
 			if (URigVMPin* Pin = Cast<URigVMPin>(InSubject))
@@ -681,36 +76,19 @@ void UAnimNextGraph_EditorData::HandleModifiedEvent(ERigVMGraphNotifType InNotif
 			}
 			break;
 		}
-	case ERigVMGraphNotifType::PinDefaultValueChanged:
-		{
-			if (InGraph->GetRuntimeAST().IsValid())
-			{
-				URigVMPin* RootPin = CastChecked<URigVMPin>(InSubject)->GetRootPin();
-				FRigVMASTProxy RootPinProxy = FRigVMASTProxy::MakeFromUObject(RootPin);
-				const FRigVMExprAST* Expression = InGraph->GetRuntimeAST()->GetExprForSubject(RootPinProxy);
-				if (Expression == nullptr)
-				{
-					InGraph->ClearAST();
-				}
-				else if (Expression->NumParents() > 1)
-				{
-					InGraph->ClearAST();
-				}
-			}
-
-			RequestAutoVMRecompilation();	// We need to rebuild our metadata when a default value changes
-			break;
-		}
 	}
-	
-	// if the notification still has to be sent...
-	if (bNotifForOthersPending && !bSuspendModelNotificationsForOthers)
+
+	Super::HandleModifiedEvent(InNotifType, InGraph, InSubject);
+}
+
+TConstArrayView<TSubclassOf<UAnimNextRigVMAssetEntry>> UAnimNextGraph_EditorData::GetEntryClasses() const
+{
+	static const TSubclassOf<UAnimNextRigVMAssetEntry> Classes[] =
 	{
-		if (RigVMGraphModifiedEvent.IsBound())
-		{
-			RigVMGraphModifiedEvent.Broadcast(InNotifType, InGraph, InSubject);
-		}
-	}
+		UAnimNextGraphEntry::StaticClass(),
+	};
+	
+	return Classes;
 }
 
 void UAnimNextGraph_EditorData::CreateEdGraphForCollapseNode(URigVMCollapseNode* InNode)
@@ -727,7 +105,7 @@ void UAnimNextGraph_EditorData::CreateEdGraphForCollapseNode(URigVMCollapseNode*
 			RigFunctionGraph->bAllowDeletion = true;
 			RigFunctionGraph->ModelNodePath = ContainedGraph->GetNodePath();
 			RigFunctionGraph->bIsFunctionDefinition = true;
-		
+
 			RigFunctionGraph->Initialize(this);
 
 			RigVMClient.GetOrCreateController(ContainedGraph)->ResendAllNotifications();
@@ -739,42 +117,51 @@ UEdGraph* UAnimNextGraph_EditorData::CreateEdGraph(URigVMGraph* InRigVMGraph, bo
 {
 	check(InRigVMGraph);
 
-	if(bForce)
+	if(InRigVMGraph->IsA<URigVMFunctionLibrary>())
 	{
-		RemoveEdGraph(InRigVMGraph);
+		return nullptr;
 	}
 
-	FString GraphName = InRigVMGraph->GetName();
+	UAnimNextGraphEntry* Entry = Cast<UAnimNextGraphEntry>(FindEntryForRigVMGraph(InRigVMGraph));
+	if(Entry == nullptr)
+	{
+		return nullptr;
+	}
 
-	if(GraphName.IsEmpty())
+	if(bForce)
 	{
-		GraphName = URigVMEdGraphSchema::GraphName_RigVM.ToString();
+		RemoveEdGraph(InRigVMGraph);
 	}
 
-	GraphName = RigVMClient.GetUniqueName(*GraphName).ToString();
+	FString GraphName = InRigVMGraph->GetName();
+	check(!GraphName.IsEmpty());
 
-	UAnimNextGraph_EdGraph* RigFunctionGraph = NewObject<UAnimNextGraph_EdGraph>(this, *GraphName, RF_Transactional);
+	UAnimNextGraph_EdGraph* RigFunctionGraph = NewObject<UAnimNextGraph_EdGraph>(this, MakeEdGraphName(*GraphName), RF_Transactional);
 	RigFunctionGraph->Schema = UAnimNextGraph_EdGraphSchema::StaticClass();
 	RigFunctionGraph->bAllowDeletion = true;
 	RigFunctionGraph->bIsFunctionDefinition = false;
 	RigFunctionGraph->ModelNodePath = InRigVMGraph->GetNodePath();
 	RigFunctionGraph->Initialize(this);
 
-	Graphs.Add(RigFunctionGraph);
+	Entry->EdGraph = RigFunctionGraph;
+	if(Entry->Graph == nullptr)
+	{
+		Entry->Graph = InRigVMGraph;
+	}
+	else
+	{
+		check(Entry->Graph == InRigVMGraph);
+	}
 
 	return RigFunctionGraph;
 }
 
 bool UAnimNextGraph_EditorData::RemoveEdGraph(URigVMGraph* InModel)
 {
-	if(UAnimNextGraph_EdGraph* RigFunctionGraph = Cast<UAnimNextGraph_EdGraph>(GetEditorObjectForRigVMGraph(InModel)))
+	if(UAnimNextGraphEntry* Entry = Cast<UAnimNextGraphEntry>(FindEntryForRigVMGraph(InModel)))
 	{
-		if(Graphs.Contains(RigFunctionGraph))
-		{
-			Modify();
-			Graphs.Remove(RigFunctionGraph);
-		}
-		RigVMClient.DestroyObject(RigFunctionGraph);
+		RigVMClient.DestroyObject(Entry->EdGraph);
+		Entry->EdGraph = nullptr;
 		return true;
 	}
 	return false;
@@ -795,9 +182,9 @@ UAnimNextGraphEntry* UAnimNextGraph_EditorData::AddGraph(FName InName, bool bSet
 
 	// Check for duplicate name
 	FName NewGraphName = InName;
-	auto DuplicateNamePredicate = [&NewGraphName](const UAnimNextGraphEntry* InEntry)
+	auto DuplicateNamePredicate = [&NewGraphName](const UAnimNextRigVMAssetEntry* InEntry)
 	{
-		return InEntry->GraphName == NewGraphName;
+		return InEntry->GetEntryName() == NewGraphName;
 	};
 
 	bool bAlreadyExists = Entries.ContainsByPredicate(DuplicateNamePredicate);
@@ -808,7 +195,7 @@ UAnimNextGraphEntry* UAnimNextGraph_EditorData::AddGraph(FName InName, bool bSet
 		bAlreadyExists =  Entries.ContainsByPredicate(DuplicateNamePredicate);
 	}
 
-	UAnimNextGraphEntry* NewEntry = UE::AnimNext::Graph::Private::CreateNewGraphEntry<UAnimNextGraphEntry>(this);
+	UAnimNextGraphEntry* NewEntry = CreateNewSubEntry<UAnimNextGraphEntry>(this);
 	NewEntry->GraphName = NewGraphName;
 
 	if(bSetupUndoRedo)
diff --git a/Engine/Plugins/Experimental/AnimNext/Source/AnimNextUncookedOnly/Private/Param/AnimNextParameterBlockGraph.cpp b/Engine/Plugins/Experimental/AnimNext/Source/AnimNextUncookedOnly/Private/Param/AnimNextParameterBlockGraph.cpp
--- a/Engine/Plugins/Experimental/AnimNext/Source/AnimNextUncookedOnly/Private/Param/AnimNextParameterBlockGraph.cpp
+++ b/Engine/Plugins/Experimental/AnimNext/Source/AnimNextUncookedOnly/Private/Param/AnimNextParameterBlockGraph.cpp
@@ -1,6 +1,7 @@
 ﻿// Copyright Epic Games, Inc. All Rights Reserved.
 
 #include "Param/AnimNextParameterBlockGraph.h"
+#include "Param/AnimNextParameterBlock_EdGraph.h"
 
 FText UAnimNextParameterBlockGraph::GetDisplayName() const
 {
@@ -12,7 +13,7 @@ FText UAnimNextParameterBlockGraph::GetDisplayNameTooltip() const
 	return FText::FromName(GraphName);
 }
 
-void UAnimNextParameterBlockGraph::SetGraphName(FName InName, bool bSetupUndoRedo)
+void UAnimNextParameterBlockGraph::SetEntryName(FName InName, bool bSetupUndoRedo)
 {
 	if(bSetupUndoRedo)
 	{
@@ -22,4 +23,14 @@ void UAnimNextParameterBlockGraph::SetGraphName(FName InName, bool bSetupUndoRed
 	GraphName = InName;
 
 	BroadcastModified();
+}
+
+URigVMGraph* UAnimNextParameterBlockGraph::GetRigVMGraph() const
+{
+	return Graph;
+}
+
+URigVMEdGraph* UAnimNextParameterBlockGraph::GetEdGraph() const
+{
+	return EdGraph;
 }
\ No newline at end of file
diff --git a/Engine/Plugins/Experimental/AnimNext/Source/AnimNextUncookedOnly/Private/Param/AnimNextParameterBlockParameter.cpp b/Engine/Plugins/Experimental/AnimNext/Source/AnimNextUncookedOnly/Private/Param/AnimNextParameterBlockParameter.cpp
--- a/Engine/Plugins/Experimental/AnimNext/Source/AnimNextUncookedOnly/Private/Param/AnimNextParameterBlockParameter.cpp
+++ b/Engine/Plugins/Experimental/AnimNext/Source/AnimNextUncookedOnly/Private/Param/AnimNextParameterBlockParameter.cpp
@@ -1,6 +1,8 @@
 // Copyright Epic Games, Inc. All Rights Reserved.
 
 #include "Param/AnimNextParameterBlockParameter.h"
+
+#include "Param/AnimNextParameterBlock.h"
 #include "Param/AnimNextParameterBlock_EditorData.h"
 #include "Param/ExternalParameterRegistry.h"
 
@@ -18,7 +20,7 @@ FAnimNextParamType UAnimNextParameterBlockParameter::GetParamType() const
 	return Type;
 }
 
-FName UAnimNextParameterBlockParameter::GetParameterName() const
+FName UAnimNextParameterBlockParameter::GetEntryName() const
 {
 	return ParameterName;
 }
@@ -37,7 +39,15 @@ bool UAnimNextParameterBlockParameter::SetParamType(const FAnimNextParamType& In
 	return true;
 }
 
-void UAnimNextParameterBlockParameter::SetParameterName(FName InName, bool bSetupUndoRedo)
+FInstancedPropertyBag& UAnimNextParameterBlockParameter::GetPropertyBag() const
+{
+	// TODO: move property bag for defaults onto this entry!
+	UAnimNextParameterBlock* Asset = GetTypedOuter<UAnimNextParameterBlock>();
+	check(Asset);
+	return Asset->PropertyBag;
+}
+
+void UAnimNextParameterBlockParameter::SetEntryName(FName InName, bool bSetupUndoRedo)
 {
 	if(bSetupUndoRedo)
 	{
