AnimNext: Begin merging parameter blocks and graphs

diff --git a/Engine/Plugins/Experimental/AnimNext/Source/AnimNextUncookedOnly/Private/Param/AnimNextParameterBlock_EditorData.cpp b/Engine/Plugins/Experimental/AnimNext/Source/AnimNextUncookedOnly/Private/Param/AnimNextParameterBlock_EditorData.cpp
--- a/Engine/Plugins/Experimental/AnimNext/Source/AnimNextUncookedOnly/Private/Param/AnimNextParameterBlock_EditorData.cpp
+++ b/Engine/Plugins/Experimental/AnimNext/Source/AnimNextUncookedOnly/Private/Param/AnimNextParameterBlock_EditorData.cpp
@@ -2,76 +2,19 @@
 
 #include "Param/AnimNextParameterBlock_EditorData.h"
 
-#include "ControlRigDefines.h"
-#include "Param/AnimNextParameterBlock_EdGraph.h"
-#include "Param/AnimNextParameterBlock_EdGraphSchema.h"
 #include "UncookedOnlyUtils.h"
-#include "Param/AnimNextParameterExecuteContext.h"
-#include "Rigs/RigHierarchyPose.h"
-#include "RigVMModel/RigVMFunctionLibrary.h"
-#include "RigVMModel/RigVMNotifications.h"
-#include "RigVMModel/Nodes/RigVMCollapseNode.h"
 #include "Curves/CurveFloat.h"
 #include "Param/AnimNextParameterBlock.h"
-#include "Param/AnimNextParameterBlockEntry.h"
-#include "ExternalPackageHelper.h"
-#include "Param/AnimNextParameterBlockBinding.h"
-#include "Param/AnimNextParameterBlockBindingReference.h"
-#include "Misc/TransactionObjectEvent.h"
-#include "ObjectTools.h"
 #include "Param/AnimNextParameterBlockGraph.h"
 #include "Param/AnimNextParameterBlockParameter.h"
-#include "Param/RigVMDispatch_GetLayerParameter.h"
-#include "Param/RigVMDispatch_GetParameter.h"
-#include "UObject/AssetRegistryTagsContext.h"
-
-#if WITH_EDITORONLY_DATA
+#include "Param/AnimNextParameterBlock_EdGraph.h"
+#include "Param/AnimNextParameterBlock_EdGraphSchema.h"
+#include "RigVMModel/RigVMFunctionLibrary.h"
+#include "RigVMModel/Nodes/RigVMCollapseNode.h"
+#include "UObject/LinkerLoad.h"
 
 const FName UAnimNextParameterBlock_EditorData::ExportsAssetRegistryTag = TEXT("Exports");
 
-namespace UE::AnimNext::Parameters::Private
-{
-
-static constexpr TCHAR DefaultBindingName[] = TEXT("BindingGraph");
-
-static UAnimNextParameterBlockEntry* CreateNewParameterBlockEntry(UAnimNextParameterBlock_EditorData* InEditorData, TSubclassOf<UAnimNextParameterBlockEntry> InClass)
-{
-	UAnimNextParameterBlockEntry* NewEntry = NewObject<UAnimNextParameterBlockEntry>(InEditorData, InClass.Get(), NAME_None, RF_Transactional);
-	// If we are a transient asset, dont use external packages
-	UAnimNextParameterBlock* Block = UncookedOnly::FUtils::GetBlock(InEditorData);
-	check(Block);
-	if(!Block->HasAnyFlags(RF_Transient))
-	{
-		FExternalPackageHelper::SetPackagingMode(NewEntry, InEditorData, true, false, PKG_None);
-	}
-	return NewEntry;
-}
-
-template<typename EntryClassType>
-static EntryClassType* CreateNewParameterBlockEntry(UAnimNextParameterBlock_EditorData* InEditorData)
-{
-	return CastChecked<EntryClassType>(CreateNewParameterBlockEntry(InEditorData, EntryClassType::StaticClass()));
-}
-
-}
-
-#endif // #if WITH_EDITORONLY_DATA
-
-UAnimNextParameterBlock_EditorData::UAnimNextParameterBlock_EditorData(const FObjectInitializer& ObjectInitializer)
-{
-	RigVMClient.Reset();
-	RigVMClient.SetSchemaClass(UAnimNextParameterBlockLibrary_Schema::StaticClass());
-	RigVMClient.SetExecuteContextStruct(FAnimNextParameterExecuteContext::StaticStruct());
-	RigVMClient.SetOuterClientHost(this, GET_MEMBER_NAME_CHECKED(UAnimNextParameterBlock_EditorData, RigVMClient));
-	{
-		TGuardValue<bool> DisableClientNotifs(RigVMClient.bSuspendNotifications, true);
-		RigVMClient.AddModel(TEXT("RigVMGraph"), false, &ObjectInitializer);
-		RigVMClient.GetOrCreateFunctionLibrary(false, &ObjectInitializer);
-	}
-}
-
-#if WITH_EDITORONLY_DATA
-
 UAnimNextParameterBlockParameter* UAnimNextParameterBlock_EditorData::AddParameter(FName InName, FAnimNextParamType InType, bool bSetupUndoRedo, bool bPrintPythonCommand)
 {
 	if(InName == NAME_None)
@@ -81,7 +24,7 @@ UAnimNextParameterBlockParameter* UAnimNextParameterBlock_EditorData::AddParamet
 	}
 
 	// Check for duplicate parameter
-	const bool bAlreadyExists = Entries.ContainsByPredicate([InName](const UAnimNextParameterBlockEntry* InEntry)
+	const bool bAlreadyExists = Entries.ContainsByPredicate([InName](const UAnimNextRigVMAssetEntry* InEntry)
 	{
 		if(const UAnimNextParameterBlockParameter* Parameter = Cast<UAnimNextParameterBlockParameter>(InEntry))
 		{
@@ -96,7 +39,7 @@ UAnimNextParameterBlockParameter* UAnimNextParameterBlock_EditorData::AddParamet
 		return nullptr;
 	}
 
-	UAnimNextParameterBlockParameter* NewEntry = UE::AnimNext::Parameters::Private::CreateNewParameterBlockEntry<UAnimNextParameterBlockParameter>(this);
+	UAnimNextParameterBlockParameter* NewEntry = CreateNewSubEntry<UAnimNextParameterBlockParameter>(this);
 	NewEntry->ParameterName = InName;
 	NewEntry->Type = InType;
 
@@ -135,13 +78,9 @@ UAnimNextParameterBlockGraph* UAnimNextParameterBlock_EditorData::AddGraph(FName
 
 	// Check for duplicate name
 	FName NewGraphName = InName;
-	auto DuplicateNamePredicate = [&NewGraphName](const UAnimNextParameterBlockEntry* InEntry)
+	auto DuplicateNamePredicate = [&NewGraphName](const UAnimNextRigVMAssetEntry* InEntry)
 	{
-		if(const IAnimNextParameterBlockGraphInterface* GraphInterface = Cast<IAnimNextParameterBlockGraphInterface>(InEntry))
-		{
-			return GraphInterface->GetGraphName() == NewGraphName;
-		}
-		return false;
+		return InEntry->GetEntryName() == NewGraphName;
 	};
 
 	bool bAlreadyExists = Entries.ContainsByPredicate(DuplicateNamePredicate);
@@ -152,7 +91,7 @@ UAnimNextParameterBlockGraph* UAnimNextParameterBlock_EditorData::AddGraph(FName
 		bAlreadyExists =  Entries.ContainsByPredicate(DuplicateNamePredicate);
 	}
 
-	UAnimNextParameterBlockGraph* NewEntry = UE::AnimNext::Parameters::Private::CreateNewParameterBlockEntry<UAnimNextParameterBlockGraph>(this);
+	UAnimNextParameterBlockGraph* NewEntry = CreateNewSubEntry<UAnimNextParameterBlockGraph>(this);
 	NewEntry->GraphName = NewGraphName;
 
 	if(bSetupUndoRedo)
@@ -160,7 +99,7 @@ UAnimNextParameterBlockGraph* UAnimNextParameterBlock_EditorData::AddGraph(FName
 		NewEntry->Modify();
 		Modify();
 	}
-	
+
 	Entries.Add(NewEntry);
 
 	// Add new graph
@@ -180,1050 +119,164 @@ UAnimNextParameterBlockGraph* UAnimNextParameterBlock_EditorData::AddGraph(FName
 	return NewEntry;
 }
 
-
-UAnimNextParameterBlockBinding* UAnimNextParameterBlockLibrary::AddBinding(UAnimNextParameterBlock* InBlock, FName InName, EPropertyBagPropertyType InValueType, EPropertyBagContainerType InContainerType, const UObject* InValueTypeObject, bool bSetupUndoRedo, bool bPrintPythonCommand)
-{
-	return UE::AnimNext::UncookedOnly::FUtils::GetEditorData(InBlock)->AddBinding(InName, FAnimNextParamType(InValueType, InContainerType, InValueTypeObject), bSetupUndoRedo, bPrintPythonCommand);
-}
-
-UAnimNextParameterBlockBinding* UAnimNextParameterBlock_EditorData::AddBinding(FName InName, FAnimNextParamType InType, bool bSetupUndoRedo, bool bPrintPythonCommand)
+TConstArrayView<TSubclassOf<UAnimNextRigVMAssetEntry>> UAnimNextParameterBlock_EditorData::GetEntryClasses() const
 {
-	if(InName == NAME_None)
-	{
-		ReportError(TEXT("UAnimNextParameterBlock_EditorData::AddBinding: Invalid parameter name supplied."));
-		return nullptr;
-	}
-
-	// Check for duplicate bindings
-	const bool bAlreadyExists = Entries.ContainsByPredicate([InName](const UAnimNextParameterBlockEntry* InEntry)
-	{
-		if(const IAnimNextParameterBlockParameterInterface* Binding = Cast<IAnimNextParameterBlockParameterInterface>(InEntry))
-		{
-			return Binding->GetParameterName() == InName;
-		}
-		return false;
-	});
-
-	if(bAlreadyExists)
-	{
-		ReportError(TEXT("UAnimNextParameterBlock_EditorData::AddBinding: A binding already exists for the supplied parameter name and library."));
-		return nullptr;
-	}
-
-	UAnimNextParameterBlockBinding* NewEntry = UE::AnimNext::Parameters::Private::CreateNewParameterBlockEntry<UAnimNextParameterBlockBinding>(this);
-	NewEntry->ParameterName = InName;
-	NewEntry->Type = InType;
-
-	if(bSetupUndoRedo)
+	static const TSubclassOf<UAnimNextRigVMAssetEntry> Classes[] =
 	{
-		NewEntry->Modify();
-		Modify();
-	}
+		UAnimNextParameterBlockGraph::StaticClass(),
+		UAnimNextParameterBlockParameter::StaticClass(),
+	};
 	
-	Entries.Add(NewEntry);
-
-	// Add new binding graph
-	{
-		TGuardValue<bool> EnablePythonPrint(bSuspendPythonMessagesForRigVMClient, !bPrintPythonCommand);
-		TGuardValue<bool> DisableAutoCompile(bAutoRecompileVM, false);
-		URigVMGraph* NewGraph = RigVMClient.AddModel(UE::AnimNext::Parameters::Private::DefaultBindingName, bSetupUndoRedo);
-		ensure(NewGraph);
-		NewEntry->BindingGraph = NewGraph;
-
-		URigVMController* Controller = RigVMClient.GetController(NewGraph);
-		UE::AnimNext::UncookedOnly::FUtils::SetupBindingGraph(Controller, InName, NewEntry->GetParamType());
-	}
-
-	BroadcastModified();
-
-	return NewEntry;
+	return Classes;
 }
 
-UAnimNextParameterBlockBindingReference* UAnimNextParameterBlockLibrary::AddBindingReference(UAnimNextParameterBlock* InBlock, FName InName, UAnimNextParameterBlock* InReferencedBlock, bool bSetupUndoRedo, bool bPrintPythonCommand)
-{
-	return UE::AnimNext::UncookedOnly::FUtils::GetEditorData(InBlock)->AddBindingReference(InName, InReferencedBlock, bSetupUndoRedo, bPrintPythonCommand);
-}
-
-UAnimNextParameterBlockBindingReference* UAnimNextParameterBlock_EditorData::AddBindingReference(FName InName, UAnimNextParameterBlock* InReferencedBlock, bool bSetupUndoRedo, bool bPrintPythonCommand)
+void UAnimNextParameterBlock_EditorData::PostLoad()
 {
-	if(InName == NAME_None)
-	{
-		ReportError(TEXT("UAnimNextParameterBlock_EditorData::AddBindingReference: Invalid parameter name supplied."));
-		return nullptr;
-	}
-
-	if(InReferencedBlock == nullptr)
-	{
-		ReportError(TEXT("UAnimNextParameterBlock_EditorData::AddBindingReference: Invalid block supplied."));
-		return nullptr;
-	}
-
-	const UAnimNextParameterBlockEntry* ReferencedEntry = UE::AnimNext::UncookedOnly::FUtils::GetEditorData(InReferencedBlock)->FindBinding(InName);
-	if(ReferencedEntry == nullptr)
-	{
-		ReportError(TEXT("UAnimNextParameterBlock_EditorData::AddBindingReference: Named parameter to reference not found in block supplied."));
-		return nullptr;
-	}
+	Super::PostLoad();
 
-	// Check for duplicates
-	const bool bAlreadyExists = Entries.ContainsByPredicate([InName](const UAnimNextParameterBlockEntry* InEntry)
+	if(GetLinkerCustomVersion(FFortniteMainBranchObjectVersion::GUID) < FFortniteMainBranchObjectVersion::AnimNextMoveGraphsToEntries)
 	{
-		if(const IAnimNextParameterBlockParameterInterface* Binding = Cast<IAnimNextParameterBlockParameterInterface>(InEntry))
+		// Must preload entries so their data is populated or we cannot find the appropriate entries for graphs
+		for(UAnimNextRigVMAssetEntry* Entry : Entries) 
 		{
-			return Binding->GetParameterName() == InName;
+			Entry->GetLinker()->Preload(Entry);
 		}
-
-		return false;
-	});
-
-	if(bAlreadyExists)
-	{
-		ReportError(TEXT("UAnimNextParameterBlock_EditorData::AddBindingReference: A binding already exists for the supplied parameter and library."));
-		return nullptr;
-	}
-
-	UAnimNextParameterBlockBindingReference* NewEntry = UE::AnimNext::Parameters::Private::CreateNewParameterBlockEntry<UAnimNextParameterBlockBindingReference>(this);
-	NewEntry->ParameterName = InName;
-	NewEntry->Block = InReferencedBlock;
-
-	if(bSetupUndoRedo)
-	{
-		NewEntry->Modify();
-		Modify();
-	}
-
-	Entries.Add(NewEntry);
-
-	BroadcastModified();
-
-	return NewEntry;
-}
-
-bool UAnimNextParameterBlockLibrary::RemoveAllBindings(UAnimNextParameterBlock* InBlock, FName InName,  bool bSetupUndoRedo, bool bPrintPythonCommand)
-{
-	return UE::AnimNext::UncookedOnly::FUtils::GetEditorData(InBlock)->RemoveAllBindings(InName, bSetupUndoRedo, bPrintPythonCommand);
-}
-
-bool UAnimNextParameterBlock_EditorData::RemoveAllBindings(FName InName, bool bSetupUndoRedo, bool bPrintPythonCommand)
-{
-	if(InName == NAME_None)
-	{
-		ReportError(TEXT("UAnimNextParameterBlock_EditorData::RemoveBinding: Invalid parameter name supplied."));
-		return false;
-	}
-
-	// Check parameter binding exists
-	TObjectPtr<UAnimNextParameterBlockEntry>* EntryToRemove = Entries.FindByPredicate([InName](const UAnimNextParameterBlockEntry* InEntry)
-	{
-		if(const IAnimNextParameterBlockParameterInterface* Binding = Cast<IAnimNextParameterBlockParameterInterface>(InEntry))
+		
+		for(TObjectPtr<UAnimNextParameterBlock_EdGraph> Graph : Graphs_DEPRECATED)
 		{
-			return Binding->GetParameterName() == InName;
-		}
-		return false;
-	});
-
-	if(EntryToRemove == nullptr)
-	{
-		ReportError(TEXT("UAnimNextParameterBlock_EditorData::RemoveBinding: Block does not contain a binding for parameter."));
-		return false;
-	}
-
-	UAnimNextParameterBlockBinding* BindingToRemove = CastChecked<UAnimNextParameterBlockBinding>(*EntryToRemove);
+			URigVMGraph* FoundRigVMGraph = GetRigVMGraphForEditorObject(Graph);
+			if(FoundRigVMGraph)
+			{
+				UAnimNextRigVMAssetEntry* FoundEntry = nullptr;
+				for(UAnimNextRigVMAssetEntry* Entry : Entries)
+				{
+					if(UAnimNextParameterBlockGraph* GraphEntry = Cast<UAnimNextParameterBlockGraph>(Entry))
+					{
+						if(FoundRigVMGraph == static_cast<IAnimNextRigVMGraphInterface*>(GraphEntry)->GetRigVMGraph())
+						{
+							FoundEntry = Entry;
+							break;
+						}
+					}
+				}
 
-	if(bSetupUndoRedo)
-	{
-		Modify();
-	}
+				if(FoundEntry)
+				{
+					UAnimNextParameterBlockGraph* GraphEntry = CastChecked<UAnimNextParameterBlockGraph>(FoundEntry);
+					GraphEntry->EdGraph = Graph;
 
-	Entries.RemoveAll([InName](const UAnimNextParameterBlockEntry* InEntry)
-	{
-		if(const IAnimNextParameterBlockParameterInterface* Binding = Cast<IAnimNextParameterBlockParameterInterface>(InEntry))
-		{
-			return Binding->GetParameterName() == InName;
+					Graph->Rename(*MakeEdGraphName(FoundRigVMGraph->GetFName()).ToString(), FoundEntry, REN_ForceNoResetLoaders | REN_DoNotDirty | REN_DontCreateRedirectors | REN_NonTransactional);
+					Graph->Initialize(this);
+				}
+			}
 		}
-		return false;
-	});
 
-	// Remove any binding graphs
-	TGuardValue<bool> EnablePythonPrint(bSuspendPythonMessagesForRigVMClient, !bPrintPythonCommand);
-	TGuardValue<bool> DisableAutoCompile(bAutoRecompileVM, false);
-	const bool bResult = RigVMClient.RemoveModel(BindingToRemove->BindingGraph->GetNodePath(), bSetupUndoRedo);
+		// We used to add a default model that is no longer needed
+		URigVMGraph* DefaultModel = RigVMClient.GetDefaultModel();
+		if(DefaultModel && DefaultModel->GetName() == TEXT("RigVMGraph"))
+		{
+			bool bFound = false;
+			for(UAnimNextRigVMAssetEntry* Entry : Entries)
+			{
+				if(UAnimNextParameterBlockGraph* GraphEntry = Cast<UAnimNextParameterBlockGraph>(Entry))
+				{
+					if(DefaultModel == static_cast<IAnimNextRigVMGraphInterface*>(GraphEntry)->GetRigVMGraph())
+					{
+						bFound = true;
+						break;
+					}
+				}
+			}
 
-	BroadcastModified();
+			if(!bFound)
+			{
+				TGuardValue<bool> DisablePythonPrint(bSuspendPythonMessagesForRigVMClient, false);
+				TGuardValue<bool> DisableAutoCompile(bAutoRecompileVM, false);
+				RigVMClient.RemoveModel(DefaultModel->GetNodePath(), false);
+			}
+		}
 
-	return bResult;
+		RecompileVM();
+	}
 }
 
-bool UAnimNextParameterBlockLibrary::RemoveEntry(UAnimNextParameterBlock* InBlock, UAnimNextParameterBlockEntry* InEntry,  bool bSetupUndoRedo, bool bPrintPythonCommand)
+void UAnimNextParameterBlock_EditorData::RecompileVM()
 {
-	return UE::AnimNext::UncookedOnly::FUtils::GetEditorData(InBlock)->RemoveEntry(InEntry, bSetupUndoRedo, bPrintPythonCommand);
+	UE::AnimNext::UncookedOnly::FUtils::CompileVM(GetTypedOuter<UAnimNextParameterBlock>());
 }
 
-bool UAnimNextParameterBlock_EditorData::RemoveEntry(UAnimNextParameterBlockEntry* InEntry, bool bSetupUndoRedo, bool bPrintPythonCommand)
+void UAnimNextParameterBlock_EditorData::CreateEdGraphForCollapseNode(URigVMCollapseNode* InNode)
 {
-	if(InEntry == nullptr)
-	{
-		ReportError(TEXT("UAnimNextParameterBlock_EditorData::RemoveEntry: Invalid entry supplied."));
-		return false;
-	}
-
-	// Remove from internal array
-	TObjectPtr<UAnimNextParameterBlockEntry>* EntryToRemovePtr = Entries.FindByKey(InEntry);
-	if(EntryToRemovePtr == nullptr)
-	{
-		ReportError(TEXT("UAnimNextParameterBlock_EditorData::RemoveEntry: Block does not contain the supplied entry."));
-		return false;
-	}
-
-	if(bSetupUndoRedo)
-	{
-		Modify();
-	}
-
-	UAnimNextParameterBlockEntry* EntryToRemove = *EntryToRemovePtr;
-	Entries.Remove(EntryToRemove);
-
-	if (bSetupUndoRedo)
-	{
-		EntryToRemove->Modify();
-	}
-
-	bool bResult = true;
-	if(const IAnimNextParameterBlockGraphInterface* GraphInterface = Cast<IAnimNextParameterBlockGraphInterface>(EntryToRemove))
+	if (InNode->GetGraph()->IsA<URigVMFunctionLibrary>())
 	{
-		// Remove any binding graphs
-		if(GraphInterface->GetGraph())
+		if (URigVMGraph* ContainedGraph = InNode->GetContainedGraph())
 		{
-			TGuardValue<bool> EnablePythonPrint(bSuspendPythonMessagesForRigVMClient, !bPrintPythonCommand);
-			TGuardValue<bool> DisableAutoCompile(bAutoRecompileVM, false);
-			bResult = RigVMClient.RemoveModel(GraphInterface->GetGraph()->GetNodePath(), bSetupUndoRedo);
-		}
-	}
-
-	// This will cause any external package to be removed when saved
-	EntryToRemove->MarkAsGarbage();
-	
-	BroadcastModified();
-
-	return bResult;
-}
+			// create a sub graph
+			UAnimNextParameterBlock_EdGraph* RigFunctionGraph = NewObject<UAnimNextParameterBlock_EdGraph>(this, *InNode->GetName(), RF_Transactional);
+			RigFunctionGraph->Schema = UAnimNextParameterBlock_EdGraphSchema::StaticClass();
+			RigFunctionGraph->bAllowRenaming = true;
+			RigFunctionGraph->bEditable = true;
+			RigFunctionGraph->bAllowDeletion = true;
+			RigFunctionGraph->ModelNodePath = ContainedGraph->GetNodePath();
+			RigFunctionGraph->bIsFunctionDefinition = true;
 
-bool UAnimNextParameterBlockLibrary::RemoveEntries(UAnimNextParameterBlock* InBlock, const TArray<UAnimNextParameterBlockEntry*>& InEntries,  bool bSetupUndoRedo, bool bPrintPythonCommand)
-{
-	return UE::AnimNext::UncookedOnly::FUtils::GetEditorData(InBlock)->RemoveEntries(InEntries, bSetupUndoRedo, bPrintPythonCommand);
-}
+			RigFunctionGraph->Initialize(this);
 
-bool UAnimNextParameterBlock_EditorData::RemoveEntries(const TArray<UAnimNextParameterBlockEntry*>& InEntries, bool bSetupUndoRedo, bool bPrintPythonCommand)
-{
-	bool bResult = true;
-	{
-		TGuardValue<bool> DisableBlockNotifications(bSuspendBlockNotifications, true);
-		TGuardValue<bool> DisableAutoCompile(bAutoRecompileVM, false);
-		for(UAnimNextParameterBlockEntry* Entry : InEntries)
-		{
-			bResult &= RemoveEntry(Entry, bSetupUndoRedo, bPrintPythonCommand);
+			RigVMClient.GetOrCreateController(ContainedGraph)->ResendAllNotifications();
 		}
 	}
-
-	BroadcastModified();
-
-	return bResult;
 }
 
-UAnimNextParameterBlockEntry* UAnimNextParameterBlockLibrary::FindBinding(UAnimNextParameterBlock* InBlock, FName InName)
+UEdGraph* UAnimNextParameterBlock_EditorData::CreateEdGraph(URigVMGraph* InRigVMGraph, bool bForce)
 {
-	return UE::AnimNext::UncookedOnly::FUtils::GetEditorData(InBlock)->FindBinding(InName);
-}
+	check(InRigVMGraph);
 
-UAnimNextParameterBlockEntry* UAnimNextParameterBlock_EditorData::FindBinding(FName InName) const
-{
-	if(InName == NAME_None)
+	if(InRigVMGraph->IsA<URigVMFunctionLibrary>())
 	{
-		ReportError(TEXT("UAnimNextParameterBlock_EditorData::FindBinding: Invalid parameter name supplied."));
 		return nullptr;
 	}
 
-	const TObjectPtr<UAnimNextParameterBlockEntry>* FoundEntry = Entries.FindByPredicate([InName](const UAnimNextParameterBlockEntry* InEntry)
-	{
-		if(const IAnimNextParameterBlockParameterInterface* Binding = Cast<IAnimNextParameterBlockParameterInterface>(InEntry))
-		{
-			return Binding->GetParameterName() == InName;
-		}
-		return false;
-	});
-
-	return FoundEntry != nullptr ? *FoundEntry : nullptr;
-}
-
-void UAnimNextParameterBlock_EditorData::BroadcastModified()
-{
-	Recompile();
-
-	if(!bSuspendBlockNotifications)
+	UAnimNextParameterBlockGraph* Entry = Cast<UAnimNextParameterBlockGraph>(FindEntryForRigVMGraph(InRigVMGraph));
+	if(Entry == nullptr)
 	{
-		ModifiedDelegate.Broadcast(this);
+		return nullptr;
 	}
-}
 
-void UAnimNextParameterBlock_EditorData::ReportError(const TCHAR* InMessage) const
-{
-#if WITH_EDITOR
-	FScriptExceptionHandler::Get().HandleException(ELogVerbosity::Error, InMessage, TEXT(""));
-#endif
-}
-
-void UAnimNextParameterBlock_EditorData::ReconstructAllNodes()
-{
-#ifdef WITH_EDITORONLY_DATA
-	// Avoid refreshing EdGraph nodes during cook
-	if (GIsCookerLoadingPackage)
-	{
-		return;
-	}
-	
-	if (GetRigVMClient()->GetDefaultModel() == nullptr)
+	if(bForce)
 	{
-		return;
+		RemoveEdGraph(InRigVMGraph);
 	}
 
-	TArray<URigVMEdGraphNode*> AllNodes;
-	UE::AnimNext::UncookedOnly::FUtils::GetAllNodesOfClass(this, AllNodes);
+	FString GraphName = InRigVMGraph->GetName();
+	check(!GraphName.IsEmpty());
 
-	for (URigVMEdGraphNode* Node : AllNodes)
-	{
-		Node->SetFlags(RF_Transient);
-	}
+	UAnimNextParameterBlock_EdGraph* RigFunctionGraph = NewObject<UAnimNextParameterBlock_EdGraph>(this, MakeEdGraphName(*GraphName), RF_Transactional);
+	RigFunctionGraph->Schema = UAnimNextParameterBlock_EdGraphSchema::StaticClass();
+	RigFunctionGraph->bAllowDeletion = true;
+	RigFunctionGraph->bIsFunctionDefinition = false;
+	RigFunctionGraph->ModelNodePath = InRigVMGraph->GetNodePath();
+	RigFunctionGraph->Initialize(this);
 
-	for(URigVMEdGraphNode* Node : AllNodes)
+	Entry->EdGraph = RigFunctionGraph;
+	if(Entry->Graph == nullptr)
 	{
-		Node->ReconstructNode();
+		Entry->Graph = InRigVMGraph;
 	}
-
-	for (URigVMEdGraphNode* Node : AllNodes)
+	else
 	{
-		Node->ClearFlags(RF_Transient);
+		check(Entry->Graph == InRigVMGraph);
 	}
-#endif	
-}
-
-void UAnimNextParameterBlock_EditorData::Serialize(FArchive& Ar)
-{
-	RigVMClient.SetOuterClientHost(this, GET_MEMBER_NAME_CHECKED(UAnimNextParameterBlock_EditorData, RigVMClient));
 
-	UObject::Serialize(Ar);
+	return RigFunctionGraph;
 }
 
-void UAnimNextParameterBlock_EditorData::Initialize(bool bRecompileVM)
+bool UAnimNextParameterBlock_EditorData::RemoveEdGraph(URigVMGraph* InModel)
 {
-	UAnimNextParameterBlock* ParameterBlock = GetTypedOuter<UAnimNextParameterBlock>();
-
-	if (RigVMClient.GetController(0) == nullptr)
+	if(UAnimNextParameterBlockGraph* Entry = Cast<UAnimNextParameterBlockGraph>(FindEntryForRigVMGraph(InModel)))
 	{
-		check(RigVMClient.Num() == 1);
-		check(RigVMClient.GetFunctionLibrary());
-		
-		RigVMClient.GetOrCreateController(RigVMClient.GetDefaultModel());
-		RigVMClient.GetOrCreateController(RigVMClient.GetFunctionLibrary());
-
-		// Init function library controllers
-		for(URigVMLibraryNode* LibraryNode : RigVMClient.GetFunctionLibrary()->GetFunctions())
-		{
-			RigVMClient.GetOrCreateController(LibraryNode->GetContainedGraph());
-		}
-		
-		if(bRecompileVM)
-		{
-			UE::AnimNext::UncookedOnly::FUtils::Compile(ParameterBlock);
-		}
-	}
-
-	for(UAnimNextParameterBlock_EdGraph* Graph : Graphs)
-	{
-		Graph->Initialize(this);
-	}
-}
-
-void UAnimNextParameterBlock_EditorData::PostLoad()
-{
-	Super::PostLoad();
-
-	Initialize(/*bRecompileVM*/false);
-	RefreshAllModels(ERigVMLoadType::PostLoad);
-
-#if WITH_EDITOR
-	FExternalPackageHelper::LoadObjectsFromExternalPackages<UAnimNextParameterBlockEntry>(this, [this](UAnimNextParameterBlockEntry* InLoadedEntry)
-	{
-		check(IsValid(InLoadedEntry));
-		Entries.Add(InLoadedEntry);
-	});
-	
-	// delay compilation until the package has been loaded
-	FCoreUObjectDelegates::OnEndLoadPackage.AddUObject(this, &UAnimNextParameterBlock_EditorData::HandlePackageDone);
-#else // !WITH_EDITOR
-	RecompileVMIfRequired();
-#endif // WITH_EDITOR
-}
-
-void UAnimNextParameterBlock_EditorData::PostDuplicate(EDuplicateMode::Type DuplicateMode)
-{
-	Initialize(/*bRecompileVM*/true);
-}
-
-void UAnimNextParameterBlock_EditorData::PostTransacted(const FTransactionObjectEvent& TransactionEvent)
-{
-	Super::PostTransacted(TransactionEvent);
-
-	if (TransactionEvent.GetEventType() == ETransactionObjectEventType::UndoRedo)
-	{
-		BroadcastModified();
-	}
-}
-
-void UAnimNextParameterBlock_EditorData::GetAssetRegistryTags(TArray<FAssetRegistryTag>& OutTags) const
-{
-	PRAGMA_DISABLE_DEPRECATION_WARNINGS;
-	Super::GetAssetRegistryTags(OutTags);
-	PRAGMA_ENABLE_DEPRECATION_WARNINGS;
-}
-
-void UAnimNextParameterBlock_EditorData::GetAssetRegistryTags(FAssetRegistryTagsContext Context) const
-{
-	Super::GetAssetRegistryTags(Context);
-
-	FAnimNextParameterProviderAssetRegistryExports ExportParameters;
-	ExportParameters.Parameters.Reserve(Entries.Num());
-
-	for(const UAnimNextParameterBlockEntry* Entry : Entries)
-	{
-		if(const IAnimNextParameterBlockParameterInterface* Binding = Cast<IAnimNextParameterBlockParameterInterface>(Entry))
-		{
-			ExportParameters.Parameters.Emplace(Binding->GetParameterName(), Binding->GetParamType(), EAnimNextParameterFlags::Bound);
-		}
-	}
-
-	for (const TObjectPtr<UAnimNextParameterBlock_EdGraph>& Graph : Graphs)
-	{
-		if (const URigVMGraph* RigGraph = Graph->GetModel())
-		{
-			UE::AnimNext::UncookedOnly::FUtils::GetGraphParameters(RigGraph, ExportParameters);
-		}
-	}
-
-	FString TagValue;
-	FAnimNextParameterProviderAssetRegistryExports::StaticStruct()->ExportText(TagValue, &ExportParameters, nullptr, nullptr, PPF_None, nullptr);
-	Context.AddTag(FAssetRegistryTag(UE::AnimNext::ExportsAnimNextAssetRegistryTag, TagValue, FAssetRegistryTag::TT_Hidden));
-}
-
-#if WITH_EDITOR
-void UAnimNextParameterBlock_EditorData::HandlePackageDone(const FEndLoadPackageContext& Context)
-{
-	if (!Context.LoadedPackages.Contains(GetPackage()))
-	{
-		return;
-	}
-	HandlePackageDone();
-}
-
-void UAnimNextParameterBlock_EditorData::HandlePackageDone()
-{
-	FCoreUObjectDelegates::OnEndLoadPackage.RemoveAll(this);
-
-	RecompileVM();
-
-	ReconstructAllNodes();
-}
-
-void UAnimNextParameterBlock_EditorData::RefreshAllModels(ERigVMLoadType InLoadType)
-{
-	const bool bIsPostLoad = InLoadType == ERigVMLoadType::PostLoad;
-
-	TGuardValue<bool> IsCompilingGuard(bIsCompiling, true);
-	TGuardValue<bool> ClientIgnoreModificationsGuard(RigVMClient.bIgnoreModelNotifications, true);
-
-	TArray<URigVMGraph*> GraphsToDetach = RigVMClient.GetAllModels(true, false);
-	TMap<const URigVMGraph*, TArray<URigVMController::FLinkedPath>> LinkedPaths;
-	
-	if (ensure(IsInGameThread()))
-	{
-		for (URigVMGraph* GraphToDetach : GraphsToDetach)
-		{
-			URigVMController* Controller = RigVMClient.GetOrCreateController(GraphToDetach);
-			// temporarily disable default value validation during load time, serialized values should always be accepted
-			TGuardValue<bool> PerGraphDisablePinDefaultValueValidation(Controller->bValidatePinDefaults, false);
-			FRigVMControllerNotifGuard NotifGuard(Controller, true);
-			LinkedPaths.Add(GraphToDetach, Controller->GetLinkedPaths());
-			Controller->FastBreakLinkedPaths(LinkedPaths.FindChecked(GraphToDetach));
-			TArray<URigVMNode*> Nodes = GraphToDetach->GetNodes();
-			for (URigVMNode* Node : Nodes)
-			{
-				Controller->RepopulatePinsOnNode(Node, true, true);
-			}
-		}
-		//SetupPinRedirectorsForBackwardsCompatibility();
-	}
-
-	for (URigVMGraph* GraphToDetach : GraphsToDetach)
-	{
-		URigVMController* Controller = RigVMClient.GetOrCreateController(GraphToDetach);
-		// at this stage, allow all links to be reattached,
-		// RecomputeAllTemplateFilteredPermutations() later should break any invalid links
-		FRigVMControllerNotifGuard NotifGuard(Controller, true);
-		Controller->RestoreLinkedPaths(LinkedPaths.FindChecked(GraphToDetach));
-	}
-
-	if (bIsPostLoad)
-	{
-		//PatchTemplateNodesWithPreferredPermutation();
-	}
-
-	TArray<URigVMGraph*> GraphsToClean = RigVMClient.GetAllModels(true, true);
-
-	// Sort from leaf graphs to root
-	TArray<URigVMGraph*> SortedGraphsToClean;
-	SortedGraphsToClean.Reserve(GraphsToClean.Num());
-	while (SortedGraphsToClean.Num() < GraphsToClean.Num())
-	{
-		bool bGraphAdded = false;
-		for (URigVMGraph* Graph : GraphsToClean)
-		{
-			if (SortedGraphsToClean.Contains(Graph))
-			{
-				continue;
-			}
-
-			TArray<URigVMGraph*> ContainedGraphs;
-			for (URigVMNode* Node : Graph->GetNodes())
-			{
-				if (URigVMLibraryNode* LibraryNode = Cast<URigVMLibraryNode>(Node))
-				{
-					if (URigVMFunctionReferenceNode* FunctionReferenceNode = Cast<URigVMFunctionReferenceNode>(LibraryNode))
-					{
-						if (FunctionReferenceNode->GetReferencedFunctionHeader().LibraryPointer.LibraryNode.GetLongPackageName() != GetPackage()->GetPathName())
-						{
-							continue;
-						}
-						if (URigVMLibraryNode* ReferencedNode = FunctionReferenceNode->LoadReferencedNode())
-						{
-							ContainedGraphs.Add(ReferencedNode->GetContainedGraph());
-							continue;
-						}
-					}
-
-					if (URigVMGraph* ContainedGraph = LibraryNode->GetContainedGraph())
-					{
-						ContainedGraphs.Add(ContainedGraph);
-					}
-				}
-			}
-
-			bool bAllContained = true;
-			for (URigVMGraph* Contained : ContainedGraphs)
-			{
-				if (!SortedGraphsToClean.Contains(Contained))
-				{
-					bAllContained = false;
-					break;
-				}
-			}
-			if (bAllContained)
-			{
-				SortedGraphsToClean.Add(Graph);
-				bGraphAdded = true;
-			}
-		}
-		ensure(bGraphAdded);
-	}
-
-	for (int32 GraphIndex = 0; GraphIndex < SortedGraphsToClean.Num(); GraphIndex++)
-	{
-		URigVMGraph* GraphToClean = SortedGraphsToClean[GraphIndex];
-		URigVMController* Controller = RigVMClient.GetOrCreateController(GraphToClean);
-		//TGuardValue<bool> GuardEditGraph(GraphToClean->bEditable, true);
-		FRigVMControllerNotifGuard NotifGuard(Controller, true);
-
-		for (URigVMNode* ModelNode : GraphToClean->GetNodes())
-		{
-			Controller->RemoveUnusedOrphanedPins(ModelNode);
-		}
-	}
-}
-
-void UAnimNextParameterBlock_EditorData::OnRigVMRegistryChanged()
-{
-	RefreshAllModels(ERigVMLoadType::PostLoad);
-	//RebuildGraphFromModel(); // TODO zzz : How we do this on AnimNext ?
-}
-
-void UAnimNextParameterBlock_EditorData::RequestRigVMInit()
-{
-	// TODO zzz : How we do this on AnimNext ?
-}
-
-URigVMGraph* UAnimNextParameterBlock_EditorData::GetModel(const UEdGraph* InEdGraph) const
-{
-	return RigVMClient.GetModel(InEdGraph);
-}
-
-URigVMGraph* UAnimNextParameterBlock_EditorData::GetModel(const FString& InNodePath) const
-{
-	return RigVMClient.GetModel(InNodePath);
-}
-
-URigVMGraph* UAnimNextParameterBlock_EditorData::GetDefaultModel() const
-{
-	return RigVMClient.GetDefaultModel();
-}
-
-TArray<URigVMGraph*> UAnimNextParameterBlock_EditorData::GetAllModels() const
-{
-	return RigVMClient.GetAllModels(true, true);
-}
-
-URigVMFunctionLibrary* UAnimNextParameterBlock_EditorData::GetLocalFunctionLibrary() const
-{
-	return RigVMClient.GetFunctionLibrary();
-}
-
-URigVMGraph* UAnimNextParameterBlock_EditorData::AddModel(FString InName, bool bSetupUndoRedo, bool bPrintPythonCommand)
-{
-	TGuardValue<bool> EnablePythonPrint(bSuspendPythonMessagesForRigVMClient, !bPrintPythonCommand);
-	return RigVMClient.AddModel(InName, bSetupUndoRedo, bPrintPythonCommand);
-}
-
-bool UAnimNextParameterBlock_EditorData::RemoveModel(FString InName, bool bSetupUndoRedo, bool bPrintPythonCommand)
-{
-	TGuardValue<bool> EnablePythonPrint(bSuspendPythonMessagesForRigVMClient, !bPrintPythonCommand);
-	return RigVMClient.RemoveModel(InName, bSetupUndoRedo, bPrintPythonCommand);
-}
-
-FRigVMGetFocusedGraph& UAnimNextParameterBlock_EditorData::OnGetFocusedGraph()
-{
-	return RigVMClient.OnGetFocusedGraph();
-}
-
-const FRigVMGetFocusedGraph& UAnimNextParameterBlock_EditorData::OnGetFocusedGraph() const
-{
-	return RigVMClient.OnGetFocusedGraph();
-}
-
-URigVMGraph* UAnimNextParameterBlock_EditorData::GetFocusedModel() const
-{
-	return RigVMClient.GetFocusedModel();
-}
-
-URigVMController* UAnimNextParameterBlock_EditorData::GetController(const URigVMGraph* InGraph) const
-{
-	return RigVMClient.GetController(InGraph);
-};
-
-URigVMController* UAnimNextParameterBlock_EditorData::GetControllerByName(const FString InGraphName) const
-{
-	return RigVMClient.GetControllerByName(InGraphName);
-};
-
-URigVMController* UAnimNextParameterBlock_EditorData::GetOrCreateController(URigVMGraph* InGraph)
-{
-	return RigVMClient.GetOrCreateController(InGraph);
-};
-
-URigVMController* UAnimNextParameterBlock_EditorData::GetController(const UEdGraph* InEdGraph) const
-{
-	return RigVMClient.GetController(InEdGraph);
-};
-
-URigVMController* UAnimNextParameterBlock_EditorData::GetOrCreateController(const UEdGraph* InEdGraph)
-{
-	return RigVMClient.GetOrCreateController(InEdGraph);
-};
-
-TArray<FString> UAnimNextParameterBlock_EditorData::GeneratePythonCommands(const FString InNewBlueprintName)
-{
-	return TArray<FString>();
-}
-
-#endif // WITH_EDITOR
-
-FRigVMClient* UAnimNextParameterBlock_EditorData::GetRigVMClient()
-{
-	return &RigVMClient;
-}
-
-const FRigVMClient* UAnimNextParameterBlock_EditorData::GetRigVMClient() const
-{
-	return &RigVMClient;
-}
-
-IRigVMGraphFunctionHost* UAnimNextParameterBlock_EditorData::GetRigVMGraphFunctionHost()
-{
-	return this;
-}
-
-const IRigVMGraphFunctionHost* UAnimNextParameterBlock_EditorData::GetRigVMGraphFunctionHost() const
-{
-	return this;
-}
-
-void UAnimNextParameterBlock_EditorData::HandleRigVMGraphAdded(const FRigVMClient* InClient, const FString& InNodePath)
-{
-	if(URigVMGraph* RigVMGraph = InClient->GetModel(InNodePath))
-	{
-		RigVMGraph->SetExecuteContextStruct(FAnimNextParameterExecuteContext::StaticStruct());
-
-		if(!HasAnyFlags(RF_ClassDefaultObject | RF_NeedInitialization | RF_NeedLoad | RF_NeedPostLoad) &&
-			GetOuter() != GetTransientPackage())
-		{
-			CreateEdGraph(RigVMGraph, true);
-			RecompileVM();
-		}
-		
-#if WITH_EDITOR
-		if(!bSuspendPythonMessagesForRigVMClient)
-		{
-			const FString BlueprintName = InClient->GetSchema()->GetSanitizedName(GetName(), true, false);
-			RigVMPythonUtils::Print(BlueprintName, FString::Printf(TEXT("block.add_binding('%s')"), *RigVMGraph->GetName()));
-		}
-#endif
-	}
-}
-
-void UAnimNextParameterBlock_EditorData::HandleRigVMGraphRemoved(const FRigVMClient* InClient, const FString& InNodePath)
-{
-	if(URigVMGraph* RigVMGraph = InClient->GetModel(InNodePath))
-	{
-		RemoveEdGraph(RigVMGraph);
-		RecompileVM();
-
-#if WITH_EDITOR
-		if(!bSuspendPythonMessagesForRigVMClient)
-		{
-			const FString BlueprintName = InClient->GetSchema()->GetSanitizedName(GetName(), true, false);
-			RigVMPythonUtils::Print(BlueprintName, FString::Printf(TEXT("block.remove_binding('%s')"), *RigVMGraph->GetName()));
-		}
-#endif
-	}
-}
-
-void UAnimNextParameterBlock_EditorData::HandleConfigureRigVMController(const FRigVMClient* InClient, URigVMController* InControllerToConfigure)
-{
-	InControllerToConfigure->OnModified().AddUObject(this, &UAnimNextParameterBlock_EditorData::HandleModifiedEvent);
-
-	TWeakObjectPtr<UAnimNextParameterBlock_EditorData> WeakThis(this);
-
-	// this delegate is used by the controller to determine variable validity
-	// during a bind process. the controller itself doesn't own the variables,
-	// so we need a delegate to request them from the owning blueprint
-	InControllerToConfigure->GetExternalVariablesDelegate.BindLambda([](URigVMGraph* InGraph) -> TArray<FRigVMExternalVariable>
-	{
-		if (InGraph)
-		{
-			if(UAnimNextParameterBlock_EditorData* EditorData = InGraph->GetTypedOuter<UAnimNextParameterBlock_EditorData>())
-			{
-				if (UAnimNextParameterBlock* ParameterBlock = EditorData->GetTypedOuter<UAnimNextParameterBlock>())
-				{
-					return ParameterBlock->GetExternalVariables();
-				}
-			}
-		}
-		return TArray<FRigVMExternalVariable>();
-	});
-
-	// this delegate is used by the controller to retrieve the current bytecode of the VM
-	InControllerToConfigure->GetCurrentByteCodeDelegate.BindLambda([WeakThis]() -> const FRigVMByteCode*
-	{
-		if (WeakThis.IsValid())
-		{
-			if(UAnimNextParameterBlock* ParameterBlock = WeakThis->GetTypedOuter<UAnimNextParameterBlock>())
-			{
-				if (ParameterBlock->VM)
-				{
-					return &ParameterBlock->VM->GetByteCode();
-				}
-			}
-		}
-		return nullptr;
-
-	});
-
-#if WITH_EDITOR
-	InControllerToConfigure->SetupDefaultUnitNodeDelegates(TDelegate<FName(FRigVMExternalVariable, FString)>::CreateLambda(
-		[](FRigVMExternalVariable InVariableToCreate, FString InDefaultValue) -> FName
-		{
-			return NAME_None;
-		}
-	));
-#endif
-
-}
-
-UObject* UAnimNextParameterBlock_EditorData::GetEditorObjectForRigVMGraph(URigVMGraph* InVMGraph) const
-{
-	if(InVMGraph)
-	{
-		for(UAnimNextParameterBlock_EdGraph* EdGraph : Graphs)
-		{
-			if(EdGraph)
-			{
-				if(EdGraph->ModelNodePath == InVMGraph->GetNodePath())
-				{
-					return EdGraph;
-				}
-			}
-		}
-	}
-	return nullptr;
-}
-
-URigVMGraph* UAnimNextParameterBlock_EditorData::GetRigVMGraphForEditorObject(UObject* InObject) const
-{
-	if(const UAnimNextParameterBlock_EdGraph* Graph = Cast<UAnimNextParameterBlock_EdGraph>(InObject))
-	{
-		if (Graph->bIsFunctionDefinition)
-		{
-			if (URigVMLibraryNode* LibraryNode = RigVMClient.GetFunctionLibrary()->FindFunction(*Graph->ModelNodePath))
-			{
-				return LibraryNode->GetContainedGraph();
-			}
-		}
-		else
-		{
-			return RigVMClient.GetModel(Graph->ModelNodePath);
-		}
-	}
-
-	return nullptr;
-}
-
-FRigVMGraphFunctionStore* UAnimNextParameterBlock_EditorData::GetRigVMGraphFunctionStore()
-{
-	return &GraphFunctionStore;
-}
-
-const FRigVMGraphFunctionStore* UAnimNextParameterBlock_EditorData::GetRigVMGraphFunctionStore() const
-{
-	return &GraphFunctionStore;
-}
-
-void UAnimNextParameterBlock_EditorData::RecompileVM()
-{
-	UE::AnimNext::UncookedOnly::FUtils::CompileVM(GetTypedOuter<UAnimNextParameterBlock>());
-}
-
-void UAnimNextParameterBlock_EditorData::RecompileVMIfRequired()
-{
-	if (bVMRecompilationRequired)
-	{
-		RecompileVM();
-	}
-}
-
-void UAnimNextParameterBlock_EditorData::Recompile()
-{
-	// NOTE: Recompiling the struct entails recompiling the VM to ensure that offsets are up to date in bytecode
-	UE::AnimNext::UncookedOnly::FUtils::Compile(GetTypedOuter<UAnimNextParameterBlock>());
-}
-
-void UAnimNextParameterBlock_EditorData::RecompileIfRequired()
-{
-	if (bStructRecompilationRequired || bVMRecompilationRequired)
-	{
-		Recompile();
-	}
-}
-
-void UAnimNextParameterBlock_EditorData::RequestAutoRecompilation()
-{
-	bVMRecompilationRequired = true;
-	bStructRecompilationRequired = true;
-	if (bAutoRecompileVM && VMRecompilationBracket == 0)
-	{
-		RecompileIfRequired();
-	}
-}
-
-void UAnimNextParameterBlock_EditorData::RequestAutoVMRecompilation()
-{
-	bVMRecompilationRequired = true;
-	if (bAutoRecompileVM && VMRecompilationBracket == 0)
-	{
-		RecompileVMIfRequired();
-	}
-}
-
-void UAnimNextParameterBlock_EditorData::SetAutoVMRecompile(bool bAutoRecompile)
-{
-	bAutoRecompileVM = bAutoRecompile;
-}
-
-bool UAnimNextParameterBlock_EditorData::GetAutoVMRecompile() const
-{
-	return bAutoRecompileVM;
-}
-
-
-void UAnimNextParameterBlock_EditorData::IncrementVMRecompileBracket()
-{
-	VMRecompilationBracket++;
-}
-
-void UAnimNextParameterBlock_EditorData::DecrementVMRecompileBracket()
-{
-	if (VMRecompilationBracket == 1)
-	{
-		if (bAutoRecompileVM)
-		{
-			RecompileVMIfRequired();
-		}
-		VMRecompilationBracket = 0;
-	}
-	else if (VMRecompilationBracket > 0)
-	{
-		VMRecompilationBracket--;
-	}
-}
-
-void UAnimNextParameterBlock_EditorData::HandleModifiedEvent(ERigVMGraphNotifType InNotifType, URigVMGraph* InGraph, UObject* InSubject)
-{
-	bool bNotifForOthersPending = true;
-
-	switch(InNotifType)
-	{
-	case ERigVMGraphNotifType::InteractionBracketOpened:
-		{
-			IncrementVMRecompileBracket();
-			break;
-		}
-	case ERigVMGraphNotifType::InteractionBracketClosed:
-	case ERigVMGraphNotifType::InteractionBracketCanceled:
-		{
-			DecrementVMRecompileBracket();
-			break;
-		}
-	case ERigVMGraphNotifType::NodeAdded:
-		{
-			if (URigVMCollapseNode* CollapseNode = Cast<URigVMCollapseNode>(InSubject))
-			{
-				CreateEdGraphForCollapseNode(CollapseNode);
-				break;
-			}
-		}
-		// Fall through to the next case
-	case ERigVMGraphNotifType::LinkAdded:
-	case ERigVMGraphNotifType::LinkRemoved:
-	case ERigVMGraphNotifType::PinArraySizeChanged:
-	case ERigVMGraphNotifType::PinDirectionChanged:
-		{
-			RequestAutoVMRecompilation();
-			break;
-		}
-
-	case ERigVMGraphNotifType::PinDefaultValueChanged:
-		{
-			if (InGraph->GetRuntimeAST().IsValid())
-			{
-				URigVMPin* RootPin = CastChecked<URigVMPin>(InSubject)->GetRootPin();
-				FRigVMASTProxy RootPinProxy = FRigVMASTProxy::MakeFromUObject(RootPin);
-				const FRigVMExprAST* Expression = InGraph->GetRuntimeAST()->GetExprForSubject(RootPinProxy);
-				if (Expression == nullptr)
-				{
-					InGraph->ClearAST();
-					break;
-				}
-				else if (Expression->NumParents() > 1)
-				{
-					InGraph->ClearAST();
-					break;
-				}
-			}
-
-			RequestAutoVMRecompilation();
-			break;
-		}
-	}
-	
-	// if the notification still has to be sent...
-	if (bNotifForOthersPending && !bSuspendModelNotificationsForOthers)
-	{
-		if (RigVMGraphModifiedEvent.IsBound())
-		{
-			RigVMGraphModifiedEvent.Broadcast(InNotifType, InGraph, InSubject);
-		}
-	}
-}
-
-void UAnimNextParameterBlock_EditorData::CreateEdGraphForCollapseNode(URigVMCollapseNode* InNode)
-{
-	if (InNode->GetGraph()->IsA<URigVMFunctionLibrary>())
-	{
-		if (URigVMGraph* ContainedGraph = InNode->GetContainedGraph())
-		{
-			// create a sub graph
-			UAnimNextParameterBlock_EdGraph* RigFunctionGraph = NewObject<UAnimNextParameterBlock_EdGraph>(this, *InNode->GetName(), RF_Transactional);
-			RigFunctionGraph->Schema = UAnimNextParameterBlock_EdGraphSchema::StaticClass();
-			RigFunctionGraph->bAllowRenaming = true;
-			RigFunctionGraph->bEditable = true;
-			RigFunctionGraph->bAllowDeletion = true;
-			RigFunctionGraph->ModelNodePath = ContainedGraph->GetNodePath();
-			RigFunctionGraph->bIsFunctionDefinition = true;
-
-			RigFunctionGraph->Initialize(this);
-
-			RigVMClient.GetOrCreateController(ContainedGraph)->ResendAllNotifications();
-		}
-	}
-}
-
-UEdGraph* UAnimNextParameterBlock_EditorData::CreateEdGraph(URigVMGraph* InRigVMGraph, bool bForce)
-{
-	check(InRigVMGraph);
-
-// #if WITH_EDITORONLY_DATA
-// 	if(InRigVMGraph->IsA<URigVMFunctionLibrary>())
-// 	{
-// 		return FunctionLibraryEdGraph;
-// 	}
-// #endif
-	
-	if(bForce)
-	{
-		RemoveEdGraph(InRigVMGraph);
-	}
-
-	FString GraphName = InRigVMGraph->GetName();
-//	GraphName.RemoveFromStart(RigVMModelPrefix);
-//	GraphName.TrimStartAndEndInline();
-
-	if(GraphName.IsEmpty())
-	{
-		GraphName = URigVMEdGraphSchema::GraphName_RigVM.ToString();
-	}
-
-	GraphName = RigVMClient.GetUniqueName(*GraphName).ToString();
-
-	UAnimNextParameterBlock_EdGraph* RigFunctionGraph = NewObject<UAnimNextParameterBlock_EdGraph>(this, *GraphName, RF_Transactional);
-	RigFunctionGraph->Schema = UAnimNextParameterBlock_EdGraphSchema::StaticClass();
-	RigFunctionGraph->bAllowDeletion = true;
-	RigFunctionGraph->bIsFunctionDefinition = false;
-	RigFunctionGraph->ModelNodePath = InRigVMGraph->GetNodePath();
-	RigFunctionGraph->Initialize(this);
-
-	Graphs.Add(RigFunctionGraph);
-
-	return RigFunctionGraph;
-}
-
-bool UAnimNextParameterBlock_EditorData::RemoveEdGraph(URigVMGraph* InModel)
-{
-	if(UAnimNextParameterBlock_EdGraph* RigFunctionGraph = Cast<UAnimNextParameterBlock_EdGraph>(GetEditorObjectForRigVMGraph(InModel)))
-	{
-		if(Graphs.Contains(RigFunctionGraph))
-		{
-			Modify();
-			Graphs.Remove(RigFunctionGraph);
-		}
-		RigVMClient.DestroyObject(RigFunctionGraph);
+		RigVMClient.DestroyObject(Entry->EdGraph);
+		Entry->EdGraph = nullptr;
 		return true;
 	}
 	return false;
 }
-
-#endif // #if WITH_EDI
\ No newline at end of file
diff --git a/Engine/Plugins/Experimental/AnimNext/Source/AnimNextUncookedOnly/Private/UncookedOnlyUtils.cpp b/Engine/Plugins/Experimental/AnimNext/Source/AnimNextUncookedOnly/Private/UncookedOnlyUtils.cpp
--- a/Engine/Plugins/Experimental/AnimNext/Source/AnimNextUncookedOnly/Private/UncookedOnlyUtils.cpp
+++ b/Engine/Plugins/Experimental/AnimNext/Source/AnimNextUncookedOnly/Private/UncookedOnlyUtils.cpp
@@ -14,23 +14,16 @@
 #include "Graph/RigUnit_AnimNextGraphRoot.h"
 #include "Graph/RigUnit_AnimNextGraphEvaluator.h"
 #include "Graph/RigUnit_AnimNextShimRoot.h"
-#include "Graph/AnimNextExecuteContext.h"
 #include "Param/AnimNextParameterBlock.h"
 #include "Param/AnimNextParameterBlock_EditorData.h"
-#include "Param/AnimNextParameterBlock_EdGraph.h"
 #include "Graph/RigDecorator_AnimNextCppDecorator.h"
 #include "Param/RigUnit_AnimNextParameterBeginExecution.h"
 #include "Param/RigVMDispatch_GetParameter.h"
 #include "Param/RigVMDispatch_SetLayerParameter.h"
-#include "Param/AnimNextParameterBlockEntry.h"
-#include "Param/IAnimNextParameterBlockParameterInterface.h"
+#include "IAnimNextRigVMParameterInterface.h"
 #include "DecoratorBase/DecoratorReader.h"
 #include "DecoratorBase/DecoratorWriter.h"
-#include "DecoratorBase/NodeTemplate.h"
 #include "DecoratorBase/NodeTemplateBuilder.h"
-#include "DecoratorBase/NodeTemplateRegistry.h"
-#include "DecoratorBase/NodeDescription.h"
-#include "DecoratorBase/NodeInstance.h"
 #include "DecoratorBase/DecoratorRegistry.h"
 #include "DecoratorBase/Decorator.h"
 #include "Graph/RigUnit_AnimNextBeginExecution.h"
@@ -45,10 +38,12 @@
 #include "RigVMCompiler/RigVMCompiler.h"
 #include "RigVMCore/RigVM.h"
 #include "Param/ExternalParameterRegistry.h"
-#include "Param/ParamUtils.h"
 #include "Scheduler/AnimNextExternalTaskBinding.h"
 #include "Scheduler/AnimNextSchedule.h"
 #include "Scheduler/AnimNextSchedulePort.h"
+#include "AnimNextRigVMAsset.h"
+#include "AnimNextRigVMAssetEditorData.h"
+#include "AnimNextRigVMAssetEntry.h"
 
 namespace UE::AnimNext::UncookedOnly
 {
@@ -404,6 +399,8 @@ void FUtils::Compile(UAnimNextGraph* InGraph)
 
 	TGuardValue<bool> CompilingGuard(EditorData->bIsCompiling, true);
 
+
+	
 	// Before we re-compile a graph, we need to release and live instances since we need the metadata we are about to replace
 	// to call decorator destructors etc
 	InGraph->FreezeGraphInstances();
@@ -423,9 +420,15 @@ void FUtils::Compile(UAnimNextGraph* InGraph)
 	EditorData->CompileLog.Messages.Reset();
 	EditorData->CompileLog.NumErrors = EditorData->CompileLog.NumWarnings = 0;
 
-	// We use a temporary graph model to build our final graph that we'll compile
 	FRigVMClient* VMClient = EditorData->GetRigVMClient();
 	URigVMGraph* VMRootGraph = VMClient->GetDefaultModel();
+
+	if(VMRootGraph == nullptr)
+	{
+		return;
+	}
+	
+	// We use a temporary graph model to build our final graph that we'll compile
 	URigVMGraph* VMTempGraph = CastChecked<URigVMGraph>(StaticDuplicateObject(VMRootGraph, GetTransientPackage(), VMClient->GetUniqueName(TEXT("TempRigVMGraph"))));
 
 	UAnimNextGraph_Controller* TempController = CastChecked<UAnimNextGraph_Controller>(VMClient->GetOrCreateController(VMTempGraph));
@@ -513,7 +516,7 @@ void FUtils::Compile(UAnimNextGraph* InGraph)
 	URigVMCompiler* Compiler = URigVMCompiler::StaticClass()->GetDefaultObject<URigVMCompiler>();
 	EditorData->VMCompileSettings.SetExecuteContextStruct(EditorData->RigVMClient.GetExecuteContextStruct());
 	const FRigVMCompileSettings Settings = (EditorData->bCompileInDebugMode) ? FRigVMCompileSettings::Fast(EditorData->VMCompileSettings.GetExecuteContextStruct()) : EditorData->VMCompileSettings;
-	Compiler->Compile(Settings, { VMTempGraph }, TempController, InGraph->VM, InGraph->ExtendedExecuteContext, InGraph->GetRigVMExternalVariables(), & EditorData->PinToOperandMap);
+	Compiler->Compile(Settings, { VMTempGraph }, TempController, InGraph->VM, InGraph->ExtendedExecuteContext, TArray<FRigVMExternalVariable>(), & EditorData->PinToOperandMap);
 
 	// Initialize right away, in packaged builds we initialize during PostLoad
 	InGraph->VM->Initialize(InGraph->ExtendedExecuteContext);
@@ -715,12 +718,20 @@ void FUtils::CompileVM(UAnimNextParameterBlock* InParameterBlock)
 	EditorData->CompileLog.Messages.Reset();
 	EditorData->CompileLog.NumErrors = EditorData->CompileLog.NumWarnings = 0;
 
+	FRigVMClient* VMClient = EditorData->GetRigVMClient();
+	URigVMGraph* RootGraph = VMClient->GetDefaultModel();
+
+	if(RootGraph == nullptr)
+	{
+		return;
+	}
+
 	URigVMCompiler* Compiler = URigVMCompiler::StaticClass()->GetDefaultObject<URigVMCompiler>();
 	EditorData->VMCompileSettings.SetExecuteContextStruct(EditorData->RigVMClient.GetExecuteContextStruct());
 	FRigVMExtendedExecuteContext& CDOContext = InParameterBlock->GetRigVMExtendedExecuteContext();
 	const FRigVMCompileSettings Settings = (EditorData->bCompileInDebugMode) ? FRigVMCompileSettings::Fast(EditorData->VMCompileSettings.GetExecuteContextStruct()) : EditorData->VMCompileSettings;
-	URigVMController* RootController = EditorData->GetRigVMClient()->GetOrCreateController(EditorData->GetRigVMClient()->GetDefaultModel());
-	Compiler->Compile(Settings, EditorData->GetRigVMClient()->GetAllModels(false, false), RootController, InParameterBlock->VM, CDOContext, InParameterBlock->GetExternalVariables(), &EditorData->PinToOperandMap);
+	URigVMController* RootController = VMClient->GetOrCreateController(RootGraph);
+	Compiler->Compile(Settings, { RootGraph }, RootController, InParameterBlock->VM, CDOContext, InParameterBlock->GetExternalVariables(), &EditorData->PinToOperandMap);
 
 	InParameterBlock->VM->Initialize(CDOContext);
 	InParameterBlock->GenerateUserDefinedDependenciesData(CDOContext);
@@ -766,13 +777,13 @@ void FUtils::CompileStruct(UAnimNextParameterBlock* InParameterBlock)
 	PropertyDescs.Reserve(EditorData->Entries.Num());
 	
 	// Gather all parameters in this block
-	for(const UAnimNextParameterBlockEntry* Entry : EditorData->Entries)
+	for(const UAnimNextRigVMAssetEntry* Entry : EditorData->Entries)
 	{
-		if(const IAnimNextParameterBlockParameterInterface* Binding = Cast<IAnimNextParameterBlockParameterInterface>(Entry))
+		if(const IAnimNextRigVMParameterInterface* Binding = Cast<IAnimNextRigVMParameterInterface>(Entry))
 		{
 			const FAnimNextParamType& Type = Binding->GetParamType();
 			ensure(Type.IsValid());
-			PropertyDescs.Emplace(Binding->GetParameterName(), Type.GetContainerType(), Type.GetValueType(), Type.GetValueTypeObject());
+			PropertyDescs.Emplace(Entry->GetEntryName(), Type.GetContainerType(), Type.GetValueType(), Type.GetValueTypeObject());
 		}
 	}
 
@@ -802,8 +813,6 @@ void FUtils::CompileStruct(UAnimNextParameterBlock* InParameterBlock)
 	{
 		InParameterBlock->PropertyBag.Reset();
 	}
-
-	EditorData->bStructRecompilationRequired = false;
 }
 
 void FUtils::Compile(UAnimNextParameterBlock* InParameterBlock)
@@ -824,6 +833,18 @@ void FUtils::RecreateVM(UAnimNextParameterBlock* InParameterBlock)
 	InParameterBlock->RigVM = InParameterBlock->VM; // Local serialization
 }
 
+UAnimNextRigVMAsset* FUtils::GetAsset(UAnimNextRigVMAssetEditorData* InEditorData)
+{
+	check(InEditorData);
+	return CastChecked<UAnimNextRigVMAsset>(InEditorData->GetOuter());
+}
+
+UAnimNextRigVMAssetEditorData* FUtils::GetEditorData(UAnimNextRigVMAsset* InAsset)
+{
+	check(InAsset);
+	return CastChecked<UAnimNextRigVMAssetEditorData>(InAsset->EditorData);
+}
+
 UAnimNextParameterBlock_EditorData* FUtils::GetEditorData(const UAnimNextParameterBlock* InParameterBlock)
 {
 	check(InParameterBlock);
diff --git a/Engine/Plugins/Runtime/RigVM/Source/RigVM/Private/RigVMCore/RigVM.cpp b/Engine/Plugins/Runtime/RigVM/Source/RigVM/Private/RigVMCore/RigVM.cpp
--- a/Engine/Plugins/Runtime/RigVM/Source/RigVM/Private/RigVMCore/RigVM.cpp
+++ b/Engine/Plugins/Runtime/RigVM/Source/RigVM/Private/RigVMCore/RigVM.cpp
@@ -568,8 +568,6 @@ bool URigVM::ValidateBytecode()
 
 const URigVMHost* URigVM::GetHostCDO() const
 {
-	ensure(HasAnyFlags(RF_ClassDefaultObject | RF_DefaultSubObject));
-
 	return GetTypedOuter<const URigVMHost>();
 }
 
diff --git a/Engine/Plugins/Runtime/RigVM/Source/RigVMDeveloper/Private/RigVMModel/RigVMClient.cpp b/Engine/Plugins/Runtime/RigVM/Source/RigVMDeveloper/Private/RigVMModel/RigVMClient.cpp
--- a/Engine/Plugins/Runtime/RigVM/Source/RigVMDeveloper/Private/RigVMModel/RigVMClient.cpp
+++ b/Engine/Plugins/Runtime/RigVM/Source/RigVMDeveloper/Private/RigVMModel/RigVMClient.cpp
@@ -171,11 +171,6 @@ URigVMGraph* FRigVMClient::GetModel(const UEdGraph* InEdGraph) const
 		return GetDefaultModel();
 	}
 
-	if (InEdGraph->GetOutermost() != GetOuter()->GetOutermost())
-	{
-		return nullptr;
-	}
-
 //#if WITH_EDITORONLY_DATA
 //	if (InEdGraph == FunctionLibraryEdGraph)
 //	{
@@ -556,7 +551,7 @@ bool FRigVMClient::RemoveModel(const FString& InNodePathOrName, bool bSetupUndoR
 {
 	if(URigVMGraph* Model = GetModel(InNodePathOrName))
 	{
-		if(Model == GetDefaultModel())
+		if(Model == GetDefaultModel() && !bDefaultModelCanBeRemoved)
 		{
 #if WITH_EDITOR
 			static constexpr TCHAR Message[] = TEXT("Cannot remove the default model.");
diff --git a/Engine/Plugins/Runtime/RigVM/Source/RigVMDeveloper/Public/RigVMModel/RigVMClient.h b/Engine/Plugins/Runtime/RigVM/Source/RigVMDeveloper/Public/RigVMModel/RigVMClient.h
--- a/Engine/Plugins/Runtime/RigVM/Source/RigVMDeveloper/Public/RigVMModel/RigVMClient.h
+++ b/Engine/Plugins/Runtime/RigVM/Source/RigVMDeveloper/Public/RigVMModel/RigVMClient.h
@@ -163,6 +163,7 @@ public:
 		, ActionStack(nullptr)
 		, bSuspendNotifications(false)
 		, bIgnoreModelNotifications(false)
+		, bDefaultModelCanBeRemoved(false)
 		, OuterClientHost(nullptr)
 		, OuterClientPropertyName(NAME_None)
 	{
@@ -178,7 +179,7 @@ public:
 	const URigVMSchema* GetSchema() const { return SchemaPtr; }
 	URigVMSchema* GetOrCreateSchema();
 	URigVMGraph* GetDefaultModel() const;
-	URigVMGraph* GetModel(int32 InIndex) const { return Models[InIndex]; }
+	URigVMGraph* GetModel(int32 InIndex) const { return Models.IsValidIndex(InIndex) ? Models[InIndex] : nullptr; }
 	URigVMGraph* GetModel(const UEdGraph* InEdGraph = nullptr) const;
 	URigVMGraph* GetModel(const FString& InNodePathOrName) const;
 	URigVMGraph* GetModel(const UObject* InEditorSideObject) const;
@@ -311,6 +312,7 @@ private:
 public:
 	bool bSuspendNotifications;
 	bool bIgnoreModelNotifications;
+	bool bDefaultModelCanBeRemoved;
 private:
 	TWeakObjectPtr<UObject> OuterClientHost;
 	FName OuterClientPropertyName;
diff --git a/Engine/Source/Runtime/Core/Public/UObject/FortniteMainBranchObjectVersion.h b/Engine/Source/Runtime/Core/Public/UObject/FortniteMainBranchObjectVersion.h
--- a/Engine/Source/Runtime/Core/Public/UObject/FortniteMainBranchObjectVersion.h
+++ b/Engine/Source/Runtime/Core/Public/UObject/FortniteMainBranchObjectVersion.h
@@ -399,6 +399,9 @@ struct FFortniteMainBranchObjectVersion
 		// AnimNode Layers will now start in a Shared Group, instead of being each one on a different group at runtime
 		AnimNodeRootDefaultGroupChange,
 
+		// Move AnimNext graphs to sub-entries of assets
+		AnimNextMoveGraphsToEntries,
+
 		// -----<new versions can be added above this line>-------------------------------------------------
 		VersionPlusOne,
 		LatestVersion = VersionPlusOne - 1
