Add a default reference pose node that we hook to missing input pins in animnext graphs. This ensures we don't have to check if we have a valid child everywhere: we always do. This fixup is done during compilation.

diff --git a/Engine/Plugins/Experimental/AnimNext/Source/AnimNext/Private/Decorators/SubGraphHost.cpp b/Engine/Plugins/Experimental/AnimNext/Source/AnimNext/Private/Decorators/SubGraphHost.cpp
--- a/Engine/Plugins/Experimental/AnimNext/Source/AnimNext/Private/Decorators/SubGraphHost.cpp
+++ b/Engine/Plugins/Experimental/AnimNext/Source/AnimNext/Private/Decorators/SubGraphHost.cpp
@@ -3,7 +3,6 @@
 #include "Decorators/SubGraphHost.h"
 
 #include "DecoratorBase/ExecutionContext.h"
-#include "EvaluationVM/Tasks/PushReferenceKeyframe.h"
 
 namespace UE::AnimNext
 {
@@ -11,7 +10,6 @@ namespace UE::AnimNext
 
 	DEFINE_ANIM_DECORATOR_BEGIN(FSubGraphHostDecorator)
 		DEFINE_ANIM_DECORATOR_IMPLEMENTS_INTERFACE(IDiscreteBlend)
-		DEFINE_ANIM_DECORATOR_IMPLEMENTS_INTERFACE(IEvaluate)
 		DEFINE_ANIM_DECORATOR_IMPLEMENTS_INTERFACE(IGarbageCollection)
 		DEFINE_ANIM_DECORATOR_IMPLEMENTS_INTERFACE(IHierarchy)
 		DEFINE_ANIM_DECORATOR_IMPLEMENTS_INTERFACE(IUpdate)
@@ -22,6 +20,9 @@ namespace UE::AnimNext
 		FDecorator::FInstanceData::Construct(Context, Binding);
 
 		IGarbageCollection::RegisterWithGC(Context, Binding);
+
+		const FSharedData* SharedData = Binding.GetSharedData<FSharedData>();
+		ReferencePoseChildPtr = Context.AllocateNodeInstance(Binding, SharedData->ReferencePoseChild);
 	}
 
 	void FSubGraphHostDecorator::FInstanceData::Destruct(const FExecutionContext& Context, const FDecoratorBinding& Binding)
@@ -34,7 +35,6 @@ namespace UE::AnimNext
 	uint32 FSubGraphHostDecorator::GetNumChildren(const FExecutionContext& Context, const TDecoratorBinding<IHierarchy>& Binding) const
 	{
 		const FInstanceData* InstanceData = Binding.GetInstanceData<FInstanceData>();
-
 		return InstanceData->SubGraphSlots.Num();
 	}
 
@@ -42,24 +42,17 @@ namespace UE::AnimNext
 	{
 		const FInstanceData* InstanceData = Binding.GetInstanceData<FInstanceData>();
 
-		for (const FInstanceData::FSubGraphSlot& SubGraphEntry : InstanceData->SubGraphSlots)
-		{
-			Children.Add(SubGraphEntry.GraphInstance.GetGraphRootPtr());
-		}
-	}
-
-	void FSubGraphHostDecorator::PostEvaluate(FEvaluateTraversalContext& Context, const TDecoratorBinding<IEvaluate>& Binding) const
-	{
-		const FInstanceData* InstanceData = Binding.GetInstanceData<FInstanceData>();
-
-		if (InstanceData->SubGraphSlots.IsEmpty())
+		for (const FSubGraphSlot& SubGraphEntry : InstanceData->SubGraphSlots)
 		{
-			// We have no children, output a non-additive reference pose
-			Context.AppendTask(FAnimNextPushReferenceKeyframeTask::MakeFromSkeleton());
-		}
-		else
-		{
-			// We only have one child or another decorator handles this, do nothing
+			if (SubGraphEntry.State == ESlotState::ActiveWithReferencePose)
+			{
+				Children.Add(InstanceData->ReferencePoseChildPtr);
+			}
+			else
+			{
+				// Even if the slot is inactive, we queue an empty handle
+				Children.Add(SubGraphEntry.GraphInstance.GetGraphRootPtr());
+			}
 		}
 	}
 
@@ -68,27 +61,27 @@ namespace UE::AnimNext
 		const FSharedData* SharedData = Binding.GetSharedData<FSharedData>();
 		FInstanceData* InstanceData = Binding.GetInstanceData<FInstanceData>();
 
+		const bool bHasActiveSubGraph = InstanceData->CurrentlyActiveSubGraphIndex != INDEX_NONE;
+
 		TObjectPtr<const UAnimNextGraph> CurrentActiveSubGraph;
-		if (InstanceData->CurrentlyActiveSubGraphIndex != INDEX_NONE)
+		if (bHasActiveSubGraph)
 		{
 			CurrentActiveSubGraph = InstanceData->SubGraphSlots[InstanceData->CurrentlyActiveSubGraphIndex].SubGraph;
 		}
 
-		const TObjectPtr<const UAnimNextGraph> SubGraph = SharedData->GetSubGraph(Context, Binding);
-		if (CurrentActiveSubGraph != SubGraph)
-		{
-			// We can't blend from a sub-graph to an empty pose right now
-			check(SubGraph);
+		const TObjectPtr<const UAnimNextGraph> DesiredSubGraph = SharedData->GetSubGraph(Context, Binding);
 
+		if (CurrentActiveSubGraph != DesiredSubGraph || !bHasActiveSubGraph)
+		{
 			// Find an empty slot we can use
 			int32 FreeSlotIndex = INDEX_NONE;
 
 			const int32 NumSubGraphSlots = InstanceData->SubGraphSlots.Num();
 			for (int32 SlotIndex = 0; SlotIndex < NumSubGraphSlots; ++SlotIndex)
 			{
-				if (!InstanceData->SubGraphSlots[SlotIndex].GraphInstance.IsValid())
+				if (InstanceData->SubGraphSlots[SlotIndex].State == ESlotState::Inactive)
 				{
-					// This graph instance is invalid, we can re-use it
+					// This slot is inactive, we can re-use it
 					FreeSlotIndex = SlotIndex;
 					break;
 				}
@@ -100,8 +93,9 @@ namespace UE::AnimNext
 				FreeSlotIndex = InstanceData->SubGraphSlots.AddDefaulted();
 			}
 
-			FInstanceData::FSubGraphSlot& SubGraphSlot = InstanceData->SubGraphSlots[FreeSlotIndex];
-			SubGraphSlot.SubGraph = SubGraph;
+			FSubGraphSlot& SubGraphSlot = InstanceData->SubGraphSlots[FreeSlotIndex];
+			SubGraphSlot.SubGraph = DesiredSubGraph;
+			SubGraphSlot.State = DesiredSubGraph ? ESlotState::ActiveWithGraph : ESlotState::ActiveWithReferencePose;
 
 			const int32 OldChildIndex = InstanceData->CurrentlyActiveSubGraphIndex;
 			const int32 NewChildIndex = FreeSlotIndex;
@@ -187,8 +181,12 @@ namespace UE::AnimNext
 		if (InstanceData->SubGraphSlots.IsValidIndex(ChildIndex))
 		{
 			// Allocate our new sub-graph instance
-			FInstanceData::FSubGraphSlot& SubGraphEntry = InstanceData->SubGraphSlots[ChildIndex];
-			SubGraphEntry.SubGraph->AllocateInstance(Context.GetGraphInstance(), SubGraphEntry.GraphInstance);
+			FSubGraphSlot& SubGraphEntry = InstanceData->SubGraphSlots[ChildIndex];
+
+			if (SubGraphEntry.State == ESlotState::ActiveWithGraph)
+			{
+				SubGraphEntry.SubGraph->AllocateInstance(Context.GetGraphInstance(), SubGraphEntry.GraphInstance);
+			}
 		}
 	}
 
@@ -199,7 +197,14 @@ namespace UE::AnimNext
 		if (InstanceData->SubGraphSlots.IsValidIndex(ChildIndex))
 		{
 			// Deallocate our sub-graph instance
-			InstanceData->SubGraphSlots[ChildIndex].GraphInstance.Release();
+			FSubGraphSlot& SubGraphEntry = InstanceData->SubGraphSlots[ChildIndex];
+
+			if (SubGraphEntry.State == ESlotState::ActiveWithGraph)
+			{
+				InstanceData->SubGraphSlots[ChildIndex].GraphInstance.Release();
+			}
+
+			SubGraphEntry.State = ESlotState::Inactive;
 		}
 	}
 
@@ -207,7 +212,7 @@ namespace UE::AnimNext
 	{
 		FInstanceData* InstanceData = Binding.GetInstanceData<FInstanceData>();
 
-		for (FInstanceData::FSubGraphSlot& SubGraphEntry : InstanceData->SubGraphSlots)
+		for (FSubGraphSlot& SubGraphEntry : InstanceData->SubGraphSlots)
 		{
 			Collector.AddPropertyReferencesWithStructARO(FAnimNextGraphInstancePtr::StaticStruct(), &SubGraphEntry.GraphInstance);
 		}
diff --git a/Engine/Plugins/Experimental/AnimNext/Source/AnimNext/Private/Decorators/SubGraphHost.h b/Engine/Plugins/Experimental/AnimNext/Source/AnimNext/Private/Decorators/SubGraphHost.h
--- a/Engine/Plugins/Experimental/AnimNext/Source/AnimNext/Private/Decorators/SubGraphHost.h
+++ b/Engine/Plugins/Experimental/AnimNext/Source/AnimNext/Private/Decorators/SubGraphHost.h
@@ -6,7 +6,6 @@
 
 #include "DecoratorBase/Decorator.h"
 #include "DecoratorInterfaces/IDiscreteBlend.h"
-#include "DecoratorInterfaces/IEvaluate.h"
 #include "DecoratorInterfaces/IGarbageCollection.h"
 #include "DecoratorInterfaces/IHierarchy.h"
 #include "DecoratorInterfaces/IUpdate.h"
@@ -24,6 +23,10 @@ struct FAnimNextSubGraphHostDecoratorSharedData : public FAnimNextDecoratorShare
 	UPROPERTY(EditAnywhere, Category = "Default")
 	TObjectPtr<const UAnimNextGraph> SubGraph;
 
+	/** A dummy child that we can use to output the bind pose. This property is hidden and automatically populated during compilation. */
+	UPROPERTY(meta = (Hidden))
+	FAnimNextDecoratorHandle ReferencePoseChild;
+
 	// Latent pin support boilerplate
 	#define DECORATOR_LATENT_PROPERTIES_ENUMERATOR(GeneratorMacro) \
 		GeneratorMacro(SubGraph) \
@@ -39,10 +42,29 @@ namespace UE::AnimNext
 	 * 
 	 * A decorator that hosts and manages a sub-graph instance.
 	 */
-	struct FSubGraphHostDecorator : FBaseDecorator, IEvaluate, IUpdate, IHierarchy, IDiscreteBlend, IGarbageCollection
+	struct FSubGraphHostDecorator : FBaseDecorator, IUpdate, IHierarchy, IDiscreteBlend, IGarbageCollection
 	{
 		DECLARE_ANIM_DECORATOR(FSubGraphHostDecorator, 0xb0cfe72e, FBaseDecorator)
 
+		enum class ESlotState
+		{
+			ActiveWithGraph,
+			ActiveWithReferencePose,
+			Inactive,
+		};
+
+		struct FSubGraphSlot
+		{
+			// The sub-graph to use
+			TObjectPtr<const UAnimNextGraph> SubGraph;
+
+			// The graph instance
+			FAnimNextGraphInstancePtr GraphInstance;
+
+			// The current slot state
+			ESlotState State = ESlotState::Inactive;
+		};
+
 		using FSharedData = FAnimNextSubGraphHostDecoratorSharedData;
 
 		struct FInstanceData : FDecorator::FInstanceData
@@ -50,25 +72,16 @@ namespace UE::AnimNext
 			void Construct(const FExecutionContext& Context, const FDecoratorBinding& Binding);
 			void Destruct(const FExecutionContext& Context, const FDecoratorBinding& Binding);
 
-			struct FSubGraphSlot
-			{
-				// The sub-graph to use
-				TObjectPtr<const UAnimNextGraph> SubGraph;
-
-				// The graph instance
-				FAnimNextGraphInstancePtr GraphInstance;
-			};
-
 			// List of sub-graph slots
 			TArray<FSubGraphSlot> SubGraphSlots;
 
 			// The index of the currently active sub-graph slot
 			// All other sub-graphs are blending out
 			int32 CurrentlyActiveSubGraphIndex = INDEX_NONE;
-		};
 
-		// IEvaluate impl
-		virtual void PostEvaluate(FEvaluateTraversalContext& Context, const TDecoratorBinding<IEvaluate>& Binding) const override;
+			// Our child node pointer. This child is shared between all slots that have no graph provided.
+			FDecoratorPtr ReferencePoseChildPtr;
+		};
 
 		// IUpdate impl
 		virtual void PreUpdate(FUpdateTraversalContext& Context, const TDecoratorBinding<IUpdate>& Binding, const FDecoratorUpdateState& DecoratorState) const override;
diff --git a/Engine/Plugins/Experimental/AnimNext/Source/AnimNext/Private/Graph/RigDecorator_AnimNextCppDecorator.cpp b/Engine/Plugins/Experimental/AnimNext/Source/AnimNext/Private/Graph/RigDecorator_AnimNextCppDecorator.cpp
--- a/Engine/Plugins/Experimental/AnimNext/Source/AnimNext/Private/Graph/RigDecorator_AnimNextCppDecorator.cpp
+++ b/Engine/Plugins/Experimental/AnimNext/Source/AnimNext/Private/Graph/RigDecorator_AnimNextCppDecorator.cpp
@@ -38,10 +38,13 @@ void FRigDecorator_AnimNextCppDecorator::GetProgrammaticPins(URigVMController* I
 			continue;
 		}
 
-		const bool bIsHidden = PinInfo.Property->HasMetaData(FRigVMStruct::HiddenMetaName);
 		const bool bIsInline = PinInfo.Property->HasMetaData("Inline");
 		const bool bIsDecoratorHandle = PinInfo.Property->GetCPPType() == TEXT("FAnimNextDecoratorHandle");
 
+		// Decorator handle pins are never hidden because we need to still be able to link things to it
+		// UI display will use the hidden property if specified
+		const bool bIsHidden = bIsDecoratorHandle ? false : PinInfo.Property->HasMetaData(FRigVMStruct::HiddenMetaName);
+
 		// Check if the metadata stipulates that we should explicitly hide this property, if not we mark it as an input
 		PinInfo.Direction = bIsHidden ? ERigVMPinDirection::Hidden : ERigVMPinDirection::Input;
 
diff --git a/Engine/Plugins/Experimental/AnimNext/Source/AnimNextEditor/Private/Graph/SAnimNextGraphNode.cpp b/Engine/Plugins/Experimental/AnimNext/Source/AnimNextEditor/Private/Graph/SAnimNextGraphNode.cpp
--- a/Engine/Plugins/Experimental/AnimNext/Source/AnimNextEditor/Private/Graph/SAnimNextGraphNode.cpp
+++ b/Engine/Plugins/Experimental/AnimNext/Source/AnimNextEditor/Private/Graph/SAnimNextGraphNode.cpp
@@ -236,27 +236,31 @@ void SAnimNextGraphNode::UpdatePinTreeView()
 		
 		if(UEdGraphPin* InputEdGraphPin = RigGraphNode->FindPin(ModelPin->GetPinPath(), EEdGraphPinDirection::EGPD_Input))
 		{
-			if(const int32* PinIndexPtr = EdGraphPinToInputPin.Find(InputEdGraphPin))
+			// Make sure to handle input pins that we explicitly hide, like decorator handles
+			if (!InputEdGraphPin->bHidden)
 			{
-				PinInfo.InputPinWidget = InputPins[*PinIndexPtr];
-				PinInfo.InputPinWidget->SetVisibility(PinVisibilityAttribute);
-				PinWidgetForExpander = PinInfo.InputPinWidget;
-				bPinWidgetForExpanderLeft = true;
-				bPinInfoIsValid = true;
-				if (Template && !ModelPin->IsExecuteContext())
+				if (const int32* PinIndexPtr = EdGraphPinToInputPin.Find(InputEdGraphPin))
 				{
-					if (URigVMPin* RootPin = ModelPin->GetRootPin())
+					PinInfo.InputPinWidget = InputPins[*PinIndexPtr];
+					PinInfo.InputPinWidget->SetVisibility(PinVisibilityAttribute);
+					PinWidgetForExpander = PinInfo.InputPinWidget;
+					bPinWidgetForExpanderLeft = true;
+					bPinInfoIsValid = true;
+					if (Template && !ModelPin->IsExecuteContext())
 					{
-						FLinearColor PinColorAndOpacity = PinInfo.InputPinWidget->GetColorAndOpacity();
-						if (Template->FindArgument(RootPin->GetFName()) == nullptr && Template->FindExecuteArgument(RootPin->GetFName(), DispatchContext) == nullptr)
+						if (URigVMPin* RootPin = ModelPin->GetRootPin())
 						{
-							PinColorAndOpacity.A = 0.2f;
-						}
-						else
-						{
-							PinColorAndOpacity.A = 1.0f;
+							FLinearColor PinColorAndOpacity = PinInfo.InputPinWidget->GetColorAndOpacity();
+							if (Template->FindArgument(RootPin->GetFName()) == nullptr && Template->FindExecuteArgument(RootPin->GetFName(), DispatchContext) == nullptr)
+							{
+								PinColorAndOpacity.A = 0.2f;
+							}
+							else
+							{
+								PinColorAndOpacity.A = 1.0f;
+							}
+							PinInfo.InputPinWidget->SetColorAndOpacity(PinColorAndOpacity);
 						}
-						PinInfo.InputPinWidget->SetColorAndOpacity(PinColorAndOpacity);
 					}
 				}
 			}
diff --git a/Engine/Plugins/Experimental/AnimNext/Source/AnimNextUncookedOnly/Internal/Graph/AnimNextGraph_EdGraphNode.h b/Engine/Plugins/Experimental/AnimNext/Source/AnimNextUncookedOnly/Internal/Graph/AnimNextGraph_EdGraphNode.h
--- a/Engine/Plugins/Experimental/AnimNext/Source/AnimNextUncookedOnly/Internal/Graph/AnimNextGraph_EdGraphNode.h
+++ b/Engine/Plugins/Experimental/AnimNext/Source/AnimNextUncookedOnly/Internal/Graph/AnimNextGraph_EdGraphNode.h
@@ -19,6 +19,10 @@ public:
 	// UEdGraphNode implementation
 	virtual void GetNodeContextMenuActions(class UToolMenu* Menu, class UGraphNodeContextMenuContext* Context) const override;
 
+	//////////////////////////////////////////////////////////////////////////
+	// URigVMEdGraphNode implementation
+	virtual void ConfigurePin(UEdGraphPin* EdGraphPin, const URigVMPin* ModelPin) const override;
+
 	//////////////////////////////////////////////////////////////////////////
 	// Our implementation
 
diff --git a/Engine/Plugins/Experimental/AnimNext/Source/AnimNextUncookedOnly/Private/Graph/AnimNextGraph_EdGraphNode.cpp b/Engine/Plugins/Experimental/AnimNext/Source/AnimNextUncookedOnly/Private/Graph/AnimNextGraph_EdGraphNode.cpp
--- a/Engine/Plugins/Experimental/AnimNext/Source/AnimNextUncookedOnly/Private/Graph/AnimNextGraph_EdGraphNode.cpp
+++ b/Engine/Plugins/Experimental/AnimNext/Source/AnimNextUncookedOnly/Private/Graph/AnimNextGraph_EdGraphNode.cpp
@@ -2,6 +2,7 @@
 
 #include "Graph/AnimNextGraph_EdGraphNode.h"
 
+#include "DecoratorBase/DecoratorHandle.h"
 #include "Graph/RigUnit_AnimNextDecoratorStack.h"
 #include "Graph/RigDecorator_AnimNextCppDecorator.h"
 #include "RigVMModel/RigVMController.h"
@@ -28,6 +29,35 @@ void UAnimNextGraph_EdGraphNode::GetNodeContextMenuActions(UToolMenu* Menu, clas
 	}
 }
 
+void UAnimNextGraph_EdGraphNode::ConfigurePin(UEdGraphPin* EdGraphPin, const URigVMPin* ModelPin) const
+{
+	Super::ConfigurePin(EdGraphPin, ModelPin);
+
+	// Decorator handles always remain as a RigVM input pins so that we can still link things to them even if they are hidden
+	// We handle visibility for those explicitly here
+	const bool bIsInputPin = ModelPin->GetDirection() == ERigVMPinDirection::Input;
+	const bool bIsDecoratorHandle = ModelPin->GetCPPTypeObject() == FAnimNextDecoratorHandle::StaticStruct();
+	if (bIsInputPin && bIsDecoratorHandle)
+	{
+		if (const URigVMPin* DecoratorPin = ModelPin->GetParentPin())
+		{
+			if (DecoratorPin->IsDecoratorPin())
+			{
+				check(DecoratorPin->GetScriptStruct() == FRigDecorator_AnimNextCppDecorator::StaticStruct());
+
+				TSharedPtr<FStructOnScope> DecoratorScope = DecoratorPin->GetDecoratorInstance();
+				const FRigDecorator_AnimNextCppDecorator* VMDecorator = (const FRigDecorator_AnimNextCppDecorator*)DecoratorScope->GetStructMemory();
+
+				const UScriptStruct* DecoratorStruct = VMDecorator->GetDecoratorSharedDataStruct();
+				check(DecoratorStruct != nullptr);
+
+				const FProperty* PinProperty = DecoratorStruct->FindPropertyByName(ModelPin->GetFName());
+				EdGraphPin->bHidden = PinProperty->HasMetaData(FRigVMStruct::HiddenMetaName);
+			}
+		}
+	}
+}
+
 bool UAnimNextGraph_EdGraphNode::IsDecoratorStack() const
 {
 	if (const URigVMUnitNode* VMNode = Cast<URigVMUnitNode>(GetModelNode()))
diff --git a/Engine/Plugins/Experimental/AnimNext/Source/AnimNextUncookedOnly/Private/UncookedOnlyUtils.cpp b/Engine/Plugins/Experimental/AnimNext/Source/AnimNextUncookedOnly/Private/UncookedOnlyUtils.cpp
--- a/Engine/Plugins/Experimental/AnimNext/Source/AnimNextUncookedOnly/Private/UncookedOnlyUtils.cpp
+++ b/Engine/Plugins/Experimental/AnimNext/Source/AnimNextUncookedOnly/Private/UncookedOnlyUtils.cpp
@@ -49,11 +49,17 @@ namespace UE::AnimNext::UncookedOnly
 {
 namespace Private
 {
+	// Represents a decorator entry on a node
 	struct FDecoratorEntryMapping
 	{
-		const URigVMNode* DecoratorStackNode;
-		const URigVMPin* DecoratorEntryPin;
-		const FDecorator* Decorator;
+		// The RigVM node that hosts this RigVM decorator
+		const URigVMNode* DecoratorStackNode = nullptr;
+
+		// The RigVM decorator pin on our host node
+		const URigVMPin* DecoratorEntryPin = nullptr;
+
+		// The AnimNext decorator
+		const FDecorator* Decorator = nullptr;
 
 		FDecoratorEntryMapping(const URigVMNode* InDecoratorStackNode, const URigVMPin* InDecoratorEntryPin, const FDecorator* InDecorator)
 			: DecoratorStackNode(InDecoratorStackNode)
@@ -62,11 +68,16 @@ namespace Private
 		{}
 	};
 
+	// Represents a node that contains a decorator list
 	struct FDecoratorStackMapping
 	{
-		const URigVMNode* DecoratorStackNode;
+		// The RigVM node that hosts the RigVM decorators
+		const URigVMNode* DecoratorStackNode = nullptr;
+
+		// The decorator list on this node
 		TArray<FDecoratorEntryMapping> DecoratorEntries;
 
+		// The node handle assigned to this RigVM node
 		FNodeHandle DecoratorStackNodeHandle;
 
 		explicit FDecoratorStackMapping(const URigVMNode* InDecoratorStackNode)
@@ -137,7 +148,7 @@ namespace Private
 				{
 					// Decorator handle pins don't have a value, just an optional link
 					const TArray<URigVMLink*>& PinLinks = Pin->GetLinks();
-					if (PinLinks.Num() != 0)
+					if (!PinLinks.IsEmpty())
 					{
 						// Something is connected to us, find the corresponding node handle so that we can encode it as our property value
 						check(PinLinks.Num() == 1);
@@ -151,7 +162,12 @@ namespace Private
 						if (SourceDecoratorStack != nullptr)
 						{
 							SourceNodeHandle = SourceDecoratorStack->DecoratorStackNodeHandle;
-							SourceDecoratorIndex = 0;	// We always bind to the first decorator index since we only allow one base decorator per stack for now
+
+							// If the source pin is null, we are a node where the result pin lives on the stack node instead of a decorator sub-pin
+							// If this is the case, we bind to the first decorator index since we only allowed a single base decorator per stack
+							// Otherwise we lookup the decorator index we are linked to
+							const URigVMPin* SourceDecoratorPin = PinLinks[0]->GetSourcePin()->GetParentPin();
+							SourceDecoratorIndex = SourceDecoratorPin != nullptr ? SourceDecoratorStack->DecoratorStackNode->GetDecoratorPins().IndexOfByKey(SourceDecoratorPin) : 0;
 						}
 
 						if (SourceNodeHandle.IsValid())
@@ -235,6 +251,86 @@ namespace Private
 		return nullptr;
 	}
 
+	void AddMissingInputLinks(const URigVMPin* DecoratorPin, UAnimNextGraph_Controller* VMController)
+	{
+		const TArray<URigVMPin*>& Pins = DecoratorPin->GetSubPins();
+		for (URigVMPin* Pin : Pins)
+		{
+			const ERigVMPinDirection PinDirection = Pin->GetDirection();
+			if (PinDirection != ERigVMPinDirection::Input && PinDirection != ERigVMPinDirection::Hidden)
+			{
+				continue;	// We only look for hidden or input pins
+			}
+
+			if (Pin->GetCPPTypeObject() != FAnimNextDecoratorHandle::StaticStruct())
+			{
+				continue;	// We only look for decorator handle pins
+			}
+
+			const TArray<URigVMLink*>& PinLinks = Pin->GetLinks();
+			if (!PinLinks.IsEmpty())
+			{
+				continue;	// This pin already has a link, all good
+			}
+
+			// Add a dummy node that will output a reference pose to ensure every link is valid.
+			// RigVM doesn't let us link two decorators on a same node together or linking a child back to a parent
+			// as this would create a cycle in the RigVM graph. The AnimNext graph decorators do support it
+			// and so perhaps we could have a merging pass later on to remove useless dummy nodes like this.
+
+			URigVMUnitNode* VMReferencePoseNode = VMController->AddUnitNode(FRigUnit_AnimNextDecoratorStack::StaticStruct(), FRigVMStruct::ExecuteName, FVector2D(0.0f, 0.0f), FString(), false);
+			check(VMReferencePoseNode != nullptr);
+
+			const UScriptStruct* CppDecoratorStruct = FRigDecorator_AnimNextCppDecorator::StaticStruct();
+
+			FString DefaultValue;
+			{
+				const UE::AnimNext::FDecoratorUID ReferencePoseDecoratorUID(0xc03d6afc);	// Decorator header is private, reference by UID directly
+				const FDecorator* Decorator = FDecoratorRegistry::Get().Find(ReferencePoseDecoratorUID);
+				check(Decorator != nullptr);
+
+				const FRigDecorator_AnimNextCppDecorator DefaultCppDecoratorStructInstance;
+				FRigDecorator_AnimNextCppDecorator CppDecoratorStructInstance;
+				CppDecoratorStructInstance.DecoratorSharedDataStruct = Decorator->GetDecoratorSharedDataStruct();
+
+				const FProperty* Prop = FAnimNextCppDecoratorWrapper::StaticStruct()->FindPropertyByName(GET_MEMBER_NAME_STRING_CHECKED(FAnimNextCppDecoratorWrapper, CppDecorator));
+				check(Prop != nullptr);
+
+				Prop->ExportText_Direct(DefaultValue, &CppDecoratorStructInstance, &DefaultCppDecoratorStructInstance, nullptr, PPF_None);
+			}
+
+			const FName ReferencePoseDecoratorName = VMController->AddDecorator(VMReferencePoseNode->GetFName(), *CppDecoratorStruct->GetPathName(), TEXT("ReferencePose"), DefaultValue, INDEX_NONE, false, false);
+			check(!ReferencePoseDecoratorName.IsNone());
+
+			URigVMPin* OutputPin = VMReferencePoseNode->FindPin(GET_MEMBER_NAME_STRING_CHECKED(FRigUnit_AnimNextDecoratorStack, Result));
+			check(OutputPin != nullptr);
+
+			ensure(VMController->AddLink(OutputPin, Pin, false));
+		}
+	}
+
+	void AddMissingInputLinks(const URigVMGraph* VMGraph, UAnimNextGraph_Controller* VMController)
+	{
+		const TArray<URigVMNode*> VMNodes = VMGraph->GetNodes();	// Copy since we might add new nodes
+		for (URigVMNode* VMNode : VMNodes)
+		{
+			if (const URigVMUnitNode* VMUnitNode = Cast<URigVMUnitNode>(VMNode))
+			{
+				const UScriptStruct* ScriptStruct = VMUnitNode->GetScriptStruct();
+				if (ScriptStruct != FRigUnit_AnimNextDecoratorStack::StaticStruct())
+				{
+					continue;	// Skip non-decorator nodes
+				}
+
+				ForEachDecoratorInStack(VMNode,
+					[VMController](const URigVMNode* DecoratorStackNode, const URigVMPin* DecoratorPin, const FDecorator* Decorator)
+					{
+						AddMissingInputLinks(DecoratorPin, VMController);
+					});
+			}
+		}
+	}
+
 	TArray<FDecoratorStackMapping> CollectDecoratorStacks(const URigVMGraph* VMGraph, UAnimNextGraph_Controller* VMController)
 	{
 		const TArray<URigVMNode*>& VMNodes = VMGraph->GetNodes();
@@ -248,6 +344,9 @@ namespace Private
 			VMRootNode = VMController->AddUnitNode(FRigUnit_AnimNextGraphRoot::StaticStruct(), FRigUnit_AnimNextGraphRoot::EventName, FVector2D(0.0f, 0.0f), FString(), false);
 		}
 
+		// Make sure we don't have empty input pins
+		AddMissingInputLinks(VMGraph, VMController);
+
 		TArray<const URigVMNode*> NodesToVisit;
 		NodesToVisit.Add(VMRootNode);
 
diff --git a/Engine/Plugins/Runtime/RigVM/Source/RigVMDeveloper/Private/EdGraph/RigVMEdGraphNode.cpp b/Engine/Plugins/Runtime/RigVM/Source/RigVMDeveloper/Private/EdGraph/RigVMEdGraphNode.cpp
--- a/Engine/Plugins/Runtime/RigVM/Source/RigVMDeveloper/Private/EdGraph/RigVMEdGraphNode.cpp
+++ b/Engine/Plugins/Runtime/RigVM/Source/RigVMDeveloper/Private/EdGraph/RigVMEdGraphNode.cpp
@@ -1764,7 +1764,7 @@ FEdGraphPinType URigVMEdGraphNode::GetPinTypeForModelPin(const URigVMPin* InMode
 	return PinType;
 }
 
-void URigVMEdGraphNode::ConfigurePin(UEdGraphPin* EdGraphPin, const URigVMPin* ModelPin)
+void URigVMEdGraphNode::ConfigurePin(UEdGraphPin* EdGraphPin, const URigVMPin* ModelPin) const
 {
 	const bool bConnectable =
 		ModelPin->GetDirection() == ERigVMPinDirection::Input ||
diff --git a/Engine/Plugins/Runtime/RigVM/Source/RigVMDeveloper/Public/EdGraph/RigVMEdGraphNode.h b/Engine/Plugins/Runtime/RigVM/Source/RigVMDeveloper/Public/EdGraph/RigVMEdGraphNode.h
--- a/Engine/Plugins/Runtime/RigVM/Source/RigVMDeveloper/Public/EdGraph/RigVMEdGraphNode.h
+++ b/Engine/Plugins/Runtime/RigVM/Source/RigVMDeveloper/Public/EdGraph/RigVMEdGraphNode.h
@@ -246,12 +246,12 @@ protected:
 
 	static FEdGraphPinType GetPinTypeForModelPin(const URigVMPin* InModelPin);
 
+	virtual void ConfigurePin(UEdGraphPin* EdGraphPin, const URigVMPin* ModelPin) const;
 private:
 
 	int32 GetNodeTopologyVersion() const { return NodeTopologyVersion; }
 	int32 NodeTopologyVersion;
 
-	static void ConfigurePin(UEdGraphPin* EdGraphPin, const URigVMPin* ModelPin);
 	TArray<URigVMPin*>& PinListForPin(const URigVMPin* InModelPin);
 
 #if WITH_EDITOR
