AnimNext: Create ed graphs with their entry as their outer

diff --git a/Engine/Plugins/Experimental/AnimNext/Source/AnimNextUncookedOnly/Internal/AnimNextRigVMAssetEditorData.h b/Engine/Plugins/Experimental/AnimNext/Source/AnimNextUncookedOnly/Internal/AnimNextRigVMAssetEditorData.h
--- a/Engine/Plugins/Experimental/AnimNext/Source/AnimNextUncookedOnly/Internal/AnimNextRigVMAssetEditorData.h
+++ b/Engine/Plugins/Experimental/AnimNext/Source/AnimNextUncookedOnly/Internal/AnimNextRigVMAssetEditorData.h
@@ -12,6 +12,8 @@
 enum class ERigVMGraphNotifType : uint8;
 class UAnimNextRigVMAssetEntry;
 class UAnimNextRigVMAssetEditorData;
+class UAnimNextParameterBlock_EdGraph;
+class UAnimNextGraph_EdGraph;
 
 namespace UE::AnimNext::UncookedOnly
 {
@@ -69,6 +71,8 @@ protected:
 	friend class UE::AnimNext::Editor::FWorkspaceEditor;
 	friend class UAnimNextRigVMAssetEntry;
 	friend class UAnimNextRigVMAssetLibrary;
+	friend class UAnimNextGraph_EdGraph;
+	friend class UAnimNextParameterBlock_EdGraph;
 
 	// UObject interface
 	virtual void Serialize(FArchive& Ar) override;
diff --git a/Engine/Plugins/Experimental/AnimNext/Source/AnimNextUncookedOnly/Internal/Graph/AnimNextGraph_EdGraph.h b/Engine/Plugins/Experimental/AnimNext/Source/AnimNextUncookedOnly/Internal/Graph/AnimNextGraph_EdGraph.h
--- a/Engine/Plugins/Experimental/AnimNext/Source/AnimNextUncookedOnly/Internal/Graph/AnimNextGraph_EdGraph.h
+++ b/Engine/Plugins/Experimental/AnimNext/Source/AnimNextUncookedOnly/Internal/Graph/AnimNextGraph_EdGraph.h
@@ -6,7 +6,8 @@
 #include "EdGraph/RigVMEdGraph.h"
 #include "AnimNextGraph_EdGraph.generated.h"
 
-class UAnimNextGraph_EditorData;
+class UAnimNextGraphEntry;
+class UAnimNextRigVMAssetEditorData;
 
 namespace UE::AnimNext::UncookedOnly
 {
@@ -21,10 +22,14 @@ class UAnimNextGraph_EdGraph : public URigVMEdGraph
 {
 	GENERATED_BODY()
 
+	friend class UAnimNextGraphEntry;
 	friend class UAnimNextGraph_EditorData;
 
+	// UObject interface
+	virtual void PostLoad() override;
+
 	// URigVMEdGraph interface
 	virtual FRigVMClient* GetRigVMClient() const override;
 	
-	void Initialize(UAnimNextGraph_EditorData* InEditorData);
+	void Initialize(UAnimNextRigVMAssetEditorData* InEditorData);
 };
\ No newline at end of file
diff --git a/Engine/Plugins/Experimental/AnimNext/Source/AnimNextUncookedOnly/Internal/Param/AnimNextParameterBlock_EdGraph.h b/Engine/Plugins/Experimental/AnimNext/Source/AnimNextUncookedOnly/Internal/Param/AnimNextParameterBlock_EdGraph.h
--- a/Engine/Plugins/Experimental/AnimNext/Source/AnimNextUncookedOnly/Internal/Param/AnimNextParameterBlock_EdGraph.h
+++ b/Engine/Plugins/Experimental/AnimNext/Source/AnimNextUncookedOnly/Internal/Param/AnimNextParameterBlock_EdGraph.h
@@ -6,7 +6,8 @@
 #include "EdGraph/RigVMEdGraph.h"
 #include "AnimNextParameterBlock_EdGraph.generated.h"
 
-class UAnimNextParameterBlock_EditorData;
+class UAnimNextParameterBlockGraph;
+class UAnimNextRigVMAssetEditorData;
 
 namespace UE::AnimNext::UncookedOnly
 {
@@ -18,10 +19,14 @@ class UAnimNextParameterBlock_EdGraph : public URigVMEdGraph
 {
 	GENERATED_BODY()
 
+	friend class UAnimNextParameterBlockGraph;
 	friend class UAnimNextParameterBlock_EditorData;
 
+	// UObject interface
+	virtual void PostLoad() override;
+
 	// UControlRigGraph interface
 	virtual FRigVMClient* GetRigVMClient() const override;
-	
-	void Initialize(UAnimNextParameterBlock_EditorData* InEditorData);
+
+	void Initialize(UAnimNextRigVMAssetEditorData* InEditorData);
 };
\ No newline at end of file
diff --git a/Engine/Plugins/Experimental/AnimNext/Source/AnimNextUncookedOnly/Private/Graph/AnimNextGraphEntry.cpp b/Engine/Plugins/Experimental/AnimNext/Source/AnimNextUncookedOnly/Private/Graph/AnimNextGraphEntry.cpp
--- a/Engine/Plugins/Experimental/AnimNext/Source/AnimNextUncookedOnly/Private/Graph/AnimNextGraphEntry.cpp
+++ b/Engine/Plugins/Experimental/AnimNext/Source/AnimNextUncookedOnly/Private/Graph/AnimNextGraphEntry.cpp
@@ -1,7 +1,6 @@
 ﻿// Copyright Epic Games, Inc. All Rights Reserved.
 
 #include "Graph/AnimNextGraphEntry.h"
-#include "RigVMModel/RigVMGraph.h"
 #include "Graph/AnimNextGraph_EdGraph.h"
 
 FName UAnimNextGraphEntry::GetEntryName() const
diff --git a/Engine/Plugins/Experimental/AnimNext/Source/AnimNextUncookedOnly/Private/Graph/AnimNextGraph_EdGraph.cpp b/Engine/Plugins/Experimental/AnimNext/Source/AnimNextUncookedOnly/Private/Graph/AnimNextGraph_EdGraph.cpp
--- a/Engine/Plugins/Experimental/AnimNext/Source/AnimNextUncookedOnly/Private/Graph/AnimNextGraph_EdGraph.cpp
+++ b/Engine/Plugins/Experimental/AnimNext/Source/AnimNextUncookedOnly/Private/Graph/AnimNextGraph_EdGraph.cpp
@@ -3,7 +3,16 @@
 #include "Graph/AnimNextGraph_EdGraph.h"
 #include "Graph/AnimNextGraph_EditorData.h"
 
-void UAnimNextGraph_EdGraph::Initialize(UAnimNextGraph_EditorData* InEditorData)
+void UAnimNextGraph_EdGraph::PostLoad()
+{
+	Super::PostLoad();
+
+	UAnimNextRigVMAssetEditorData* EditorData = GetTypedOuter<UAnimNextRigVMAssetEditorData>();
+	check(EditorData);
+	Initialize(EditorData);
+}
+
+void UAnimNextGraph_EdGraph::Initialize(UAnimNextRigVMAssetEditorData* InEditorData)
 {
 	InEditorData->RigVMGraphModifiedEvent.RemoveAll(this);
 	InEditorData->RigVMGraphModifiedEvent.AddUObject(this, &UAnimNextGraph_EdGraph::HandleModifiedEvent);
diff --git a/Engine/Plugins/Experimental/AnimNext/Source/AnimNextUncookedOnly/Private/Graph/AnimNextGraph_EditorData.cpp b/Engine/Plugins/Experimental/AnimNext/Source/AnimNextUncookedOnly/Private/Graph/AnimNextGraph_EditorData.cpp
--- a/Engine/Plugins/Experimental/AnimNext/Source/AnimNextUncookedOnly/Private/Graph/AnimNextGraph_EditorData.cpp
+++ b/Engine/Plugins/Experimental/AnimNext/Source/AnimNextUncookedOnly/Private/Graph/AnimNextGraph_EditorData.cpp
@@ -136,7 +136,7 @@ UEdGraph* UAnimNextGraph_EditorData::CreateEdGraph(URigVMGraph* InRigVMGraph, bo
 	FString GraphName = InRigVMGraph->GetName();
 	check(!GraphName.IsEmpty());
 
-	UAnimNextGraph_EdGraph* RigFunctionGraph = NewObject<UAnimNextGraph_EdGraph>(this, MakeEdGraphName(*GraphName), RF_Transactional);
+	UAnimNextGraph_EdGraph* RigFunctionGraph = NewObject<UAnimNextGraph_EdGraph>(Entry, MakeEdGraphName(*GraphName), RF_Transactional);
 	RigFunctionGraph->Schema = UAnimNextGraph_EdGraphSchema::StaticClass();
 	RigFunctionGraph->bAllowDeletion = true;
 	RigFunctionGraph->bIsFunctionDefinition = false;
diff --git a/Engine/Plugins/Experimental/AnimNext/Source/AnimNextUncookedOnly/Private/Param/AnimNextParameterBlock_EdGraph.cpp b/Engine/Plugins/Experimental/AnimNext/Source/AnimNextUncookedOnly/Private/Param/AnimNextParameterBlock_EdGraph.cpp
--- a/Engine/Plugins/Experimental/AnimNext/Source/AnimNextUncookedOnly/Private/Param/AnimNextParameterBlock_EdGraph.cpp
+++ b/Engine/Plugins/Experimental/AnimNext/Source/AnimNextUncookedOnly/Private/Param/AnimNextParameterBlock_EdGraph.cpp
@@ -3,7 +3,16 @@
 #include "Param/AnimNextParameterBlock_EdGraph.h"
 #include "Param/AnimNextParameterBlock_EditorData.h"
 
-void UAnimNextParameterBlock_EdGraph::Initialize(UAnimNextParameterBlock_EditorData* InEditorData)
+void UAnimNextParameterBlock_EdGraph::PostLoad()
+{
+	Super::PostLoad();
+
+	UAnimNextRigVMAssetEditorData* EditorData = GetTypedOuter<UAnimNextRigVMAssetEditorData>();
+	check(EditorData);
+	Initialize(EditorData);
+}
+
+void UAnimNextParameterBlock_EdGraph::Initialize(UAnimNextRigVMAssetEditorData* InEditorData)
 {
 	InEditorData->RigVMGraphModifiedEvent.RemoveAll(this);
 	InEditorData->RigVMGraphModifiedEvent.AddUObject(this, &UAnimNextParameterBlock_EdGraph::HandleModifiedEvent);
diff --git a/Engine/Plugins/Experimental/AnimNext/Source/AnimNextUncookedOnly/Private/Param/AnimNextParameterBlock_EditorData.cpp b/Engine/Plugins/Experimental/AnimNext/Source/AnimNextUncookedOnly/Private/Param/AnimNextParameterBlock_EditorData.cpp
--- a/Engine/Plugins/Experimental/AnimNext/Source/AnimNextUncookedOnly/Private/Param/AnimNextParameterBlock_EditorData.cpp
+++ b/Engine/Plugins/Experimental/AnimNext/Source/AnimNextUncookedOnly/Private/Param/AnimNextParameterBlock_EditorData.cpp
@@ -250,7 +250,7 @@ UEdGraph* UAnimNextParameterBlock_EditorData::CreateEdGraph(URigVMGraph* InRigVM
 	FString GraphName = InRigVMGraph->GetName();
 	check(!GraphName.IsEmpty());
 
-	UAnimNextParameterBlock_EdGraph* RigFunctionGraph = NewObject<UAnimNextParameterBlock_EdGraph>(this, MakeEdGraphName(*GraphName), RF_Transactional);
+	UAnimNextParameterBlock_EdGraph* RigFunctionGraph = NewObject<UAnimNextParameterBlock_EdGraph>(Entry, MakeEdGraphName(*GraphName), RF_Transactional);
 	RigFunctionGraph->Schema = UAnimNextParameterBlock_EdGraphSchema::StaticClass();
 	RigFunctionGraph->bAllowDeletion = true;
 	RigFunctionGraph->bIsFunctionDefinition = false;
