Implement latent property snapshot and caching to ensure that we have a consistent view of latent properties no matter when in the frame we query them. We take a snapshot of latent properties before a node updates and values remain cached until the next update.

Rework decorator implementation macro boilerplate to be more flexible. Implemented GetDecoratorInterfaces to return a sorted list.

diff --git a/Engine/Plugins/Experimental/AnimNext/Source/AnimNext/Internal/DecoratorBase/DecoratorTemplate.h b/Engine/Plugins/Experimental/AnimNext/Source/AnimNext/Internal/DecoratorBase/DecoratorTemplate.h
--- a/Engine/Plugins/Experimental/AnimNext/Source/AnimNext/Internal/DecoratorBase/DecoratorTemplate.h
+++ b/Engine/Plugins/Experimental/AnimNext/Source/AnimNext/Internal/DecoratorBase/DecoratorTemplate.h
@@ -50,6 +50,17 @@ namespace UE::AnimNext
 			return AdditiveIndexOrNumAdditive;
 		}
 
+		// Returns the number of latent properties on this decorator
+		uint32 GetNumLatentPropreties() const noexcept { return NumLatentProperties; }
+
+		// Returns the number of latent properties on the decorator sub-stack
+		// Only available on the base decorator of the sub-stack
+		uint32 GetNumSubStackLatentPropreties() const
+		{
+			check(GetMode() == EDecoratorMode::Base);
+			return NumSubStackLatentProperties;
+		}
+
 		// Returns the offset into the shared data where the descriptor begins, relative to the root of the node's shared data
 		uint32 GetNodeSharedOffset() const noexcept { return NodeSharedOffset; }
 
@@ -74,16 +85,28 @@ namespace UE::AnimNext
 			return reinterpret_cast<const FAnimNextDecoratorSharedData*>(reinterpret_cast<const uint8*>(&NodeDescription) + GetNodeSharedOffset());
 		}
 
+		// Returns a reference to the specified decorator latent properties header on the current node
+		FLatentPropertiesHeader& GetDecoratorLatentPropertiesHeader(FNodeDescription& NodeDescription) const noexcept
+		{
+			return *reinterpret_cast<FLatentPropertiesHeader*>(reinterpret_cast<uint8*>(&NodeDescription) + GetNodeSharedLatentPropertyHandlesOffset());
+		}
+
+		// Returns a reference to the specified decorator latent properties header on the current node
+		const FLatentPropertiesHeader& GetDecoratorLatentPropertiesHeader(const FNodeDescription& NodeDescription) const noexcept
+		{
+			return *reinterpret_cast<const FLatentPropertiesHeader*>(reinterpret_cast<const uint8*>(&NodeDescription) + GetNodeSharedLatentPropertyHandlesOffset());
+		}
+
 		// Returns a pointer to the specified decorator latent property handles on the current node
 		FLatentPropertyHandle* GetDecoratorLatentPropertyHandles(FNodeDescription& NodeDescription) const noexcept
 		{
-			return reinterpret_cast<FLatentPropertyHandle*>(reinterpret_cast<uint8*>(&NodeDescription) + GetNodeSharedLatentPropertyHandlesOffset());
+			return reinterpret_cast<FLatentPropertyHandle*>(reinterpret_cast<uint8*>(&NodeDescription) + GetNodeSharedLatentPropertyHandlesOffset() + sizeof(FLatentPropertiesHeader));
 		}
 
 		// Returns a pointer to the specified decorator latent property handles on the current node
 		const FLatentPropertyHandle* GetDecoratorLatentPropertyHandles(const FNodeDescription& NodeDescription) const noexcept
 		{
-			return reinterpret_cast<const FLatentPropertyHandle*>(reinterpret_cast<const uint8*>(&NodeDescription) + GetNodeSharedLatentPropertyHandlesOffset());
+			return reinterpret_cast<const FLatentPropertyHandle*>(reinterpret_cast<const uint8*>(&NodeDescription) + GetNodeSharedLatentPropertyHandlesOffset() + sizeof(FLatentPropertiesHeader));
 		}
 
 		// Returns a pointer to the specified decorator instance on the current node
@@ -110,25 +133,49 @@ namespace UE::AnimNext
 			, RegistryHandle(InRegistryHandle)
 			, Mode(static_cast<uint8>(InMode))
 			, AdditiveIndexOrNumAdditive(InAdditiveIndexOrNumAdditive)
+			, NumLatentProperties(0)
+			, NumSubStackLatentProperties(0)
+			// For shared and instance data, 0 is an invalid offset since the data follows their respective header (FNodeDescription or FNodeInstance)
 			, NodeSharedOffset(0)
 			, NodeSharedLatentPropertyHandlesOffset(0)
-			, NodeInstanceOffset(0)				// For instance data, 0 is an invalid offset since the data follows an instance of FNodeInstance
+			, NodeInstanceOffset(0)
 			, Padding0(0)
 		{}
 
-		FDecoratorUIDRaw			UID;			// decorator globally unique identifier
+		// Decorator globally unique identifier (32 bits)
+		FDecoratorUIDRaw			UID;
+
+		// Cached decorator registry handle (16 bits)
+		FDecoratorRegistryHandle	RegistryHandle;
 
-		FDecoratorRegistryHandle	RegistryHandle;	// decorator registry handle
-		uint8						Mode;			// the decorator mode (we only need 1 bit)
+		// Decorator mode (we only need 1 bit, we could store other flags here)
+		uint8	Mode;
 
 		// For base decorators, this contains the number of additive decorators on top
 		// For additive decorators, this contains its index relative to the base decorator
 		// The first additive decorator has index 1, there is no index 0 (we re-purpose it for the base decorator)
 		uint8	AdditiveIndexOrNumAdditive;
 
+		// For each latent property defined on a decorator, we store a handle in the per node shared data
+		// These handles specify various metadata of the latent property: RigVM memory handle index, whether the property can freeze, instance data offset
+		// The handles of each decorator that lives on a sub-stack (base + additives) are stored contiguously
+		// The base decorator has the root handle offset and the total number as well as its local number of latent properties
+		// Each additive decorator has a handle offset that points into that contiguous list and its local number of latent properties
+
+		// How many latent properties are defined on this decorator (cached value of FDecorator::GetNumLatentPropreties to avoid repeated lookups)
+		// Not serialized, @see FNodeTemplate::Finalize
+		uint16	NumLatentProperties;
+
+		// How many latent properties are defined on this decorator sub-stack (stored on base decorator only)
+		// Not serialized, @see FNodeTemplate::Finalize
+		uint16	NumSubStackLatentProperties;
+
 		// Offsets into the shared read-only and instance data portions of a node
 		// These are not serialized, @see FNodeTemplate::Finalize
 		// These are relative to the root of the node description and instance data, respectively (max 64 KB per node each)
+		// These offsets are fixed in the template meaning each node will have the same memory layout up to the last decorator
+		// Optional data like cached latent properties are stored after this fixed layout ends
+
 		uint16	NodeSharedOffset;						// Start of shared data for this decorator
 		uint16	NodeSharedLatentPropertyHandlesOffset;	// Start of shared latent property handles for this decorator
 		uint16	NodeInstanceOffset;						// Start of instance data for this decorator
diff --git a/Engine/Plugins/Experimental/AnimNext/Source/AnimNext/Internal/DecoratorBase/DecoratorWriter.h b/Engine/Plugins/Experimental/AnimNext/Source/AnimNext/Internal/DecoratorBase/DecoratorWriter.h
--- a/Engine/Plugins/Experimental/AnimNext/Source/AnimNext/Internal/DecoratorBase/DecoratorWriter.h
+++ b/Engine/Plugins/Experimental/AnimNext/Source/AnimNext/Internal/DecoratorBase/DecoratorWriter.h
@@ -50,8 +50,8 @@ namespace UE::AnimNext
 		// Nodes must be written in the same order they were registered in
 		void WriteNode(
 			const FNodeHandle NodeHandle,
-			const TFunction<FString(uint32 DecoratorIndex, const FString& PropertyName)>& GetDecoratorProperty,
-			const TFunction<bool(uint32 DecoratorIndex, const FString& PropertyName)>& IsDecoratorPropertyLatent
+			const TFunction<FString(uint32 DecoratorIndex, FName PropertyName)>& GetDecoratorProperty,
+			const TFunction<uint16(uint32 DecoratorIndex, FName PropertyName)>& GetDecoratorLatentPropertyIndex
 			);
 
 		// Returns the error state
@@ -89,7 +89,6 @@ namespace UE::AnimNext
 		// To track node writing
 		TArray<UObject*> GraphReferencedObjects;
 		uint32 NumNodesWritten;
-		FLatentPropertyHandle CurrentLatentPropertyHandle;
 		bool bIsNodeWriting;
 
 		EErrorState ErrorState;
diff --git a/Engine/Plugins/Experimental/AnimNext/Source/AnimNext/Internal/DecoratorBase/NodeTemplate.h b/Engine/Plugins/Experimental/AnimNext/Source/AnimNext/Internal/DecoratorBase/NodeTemplate.h
--- a/Engine/Plugins/Experimental/AnimNext/Source/AnimNext/Internal/DecoratorBase/NodeTemplate.h
+++ b/Engine/Plugins/Experimental/AnimNext/Source/AnimNext/Internal/DecoratorBase/NodeTemplate.h
@@ -76,7 +76,7 @@ namespace UE::AnimNext
 		uint32	UID;					// globally unique template identifier or hash
 
 		uint16	NodeSharedDataSize;		// size in bytes of a node shared data (not serialized, @see FNodeTemplate::Finalize)
-		uint16	NodeInstanceDataSize;	// size in bytes of a node instance (not serialized, @see FNodeTemplate::Finalize)
+		uint16	NodeInstanceDataSize;	// size in bytes of a node instance, excludes optional latent properties (not serialized, @see FNodeTemplate::Finalize)
 
 		uint8	NumDecorators;
 		uint8	Padding[3];
diff --git a/Engine/Plugins/Experimental/AnimNext/Source/AnimNext/Internal/Graph/AnimNextExecuteContext.h b/Engine/Plugins/Experimental/AnimNext/Source/AnimNext/Internal/Graph/AnimNextExecuteContext.h
--- a/Engine/Plugins/Experimental/AnimNext/Source/AnimNext/Internal/Graph/AnimNextExecuteContext.h
+++ b/Engine/Plugins/Experimental/AnimNext/Source/AnimNext/Internal/Graph/AnimNextExecuteContext.h
@@ -7,6 +7,11 @@
 #include "Units/RigUnit.h"
 #include "AnimNextExecuteContext.generated.h"
 
+namespace UE::AnimNext
+{
+	struct FLatentPropertyHandle;
+}
+
 USTRUCT(BlueprintType)
 struct FAnimNextExecuteContext : public FRigVMExecuteContext
 {
@@ -14,34 +19,39 @@ struct FAnimNextExecuteContext : public FRigVMExecuteContext
 
 	FAnimNextExecuteContext() = default;
 
-	int32 GetLatentPinIndex() const { return LatentPinIndex; }
-	void* GetDestinationPtr() const { return DestinationPtr; }
+	const TConstArrayView<UE::AnimNext::FLatentPropertyHandle>& GetLatentHandles() const { return LatentHandles; }
+	void* GetDestinationBasePtr() const { return DestinationBasePtr; }
+	bool IsFrozen() const { return bIsFrozen; }
 
 	virtual void Copy(const FRigVMExecuteContext* InOtherContext) override
 	{
 		Super::Copy(InOtherContext);
 
 		const FAnimNextExecuteContext* OtherContext = (const FAnimNextExecuteContext*)InOtherContext;
-		LatentPinIndex = OtherContext->LatentPinIndex;
-		DestinationPtr = OtherContext->DestinationPtr;
+		LatentHandles = OtherContext->LatentHandles;
+		DestinationBasePtr = OtherContext->DestinationBasePtr;
+		bIsFrozen = OtherContext->bIsFrozen;
 	}
 
 private:
-	void SetupForExecution(int32 InLatentPinIndex, void* InDestinationPtr)
+	void SetupForExecution(const TConstArrayView<UE::AnimNext::FLatentPropertyHandle>& InLatentHandles, void* InDestinationBasePtr, bool bInIsFrozen)
 	{
-		LatentPinIndex = InLatentPinIndex;
-		DestinationPtr = InDestinationPtr;
+		LatentHandles = InLatentHandles;
+		DestinationBasePtr = InDestinationBasePtr;
+		bIsFrozen = bInIsFrozen;
 	}
 
 	// Call this to reset the context to its original state to detect stale usage, can't call it Reset due to virtual in base with that name
 	void DebugReset()
 	{
-		LatentPinIndex = INDEX_NONE;
-		DestinationPtr = nullptr;
+		LatentHandles = TConstArrayView<UE::AnimNext::FLatentPropertyHandle>();
+		DestinationBasePtr = nullptr;
+		bIsFrozen = false;
 	}
 
-	int32 LatentPinIndex = INDEX_NONE;
-	void* DestinationPtr = nullptr;
+	TConstArrayView<UE::AnimNext::FLatentPropertyHandle> LatentHandles;
+	void* DestinationBasePtr = nullptr;
+	bool bIsFrozen = false;
 
 	friend struct FAnimNextGraphInstance;
 };
diff --git a/Engine/Plugins/Experimental/AnimNext/Source/AnimNext/Internal/Graph/AnimNextGraph.h b/Engine/Plugins/Experimental/AnimNext/Source/AnimNext/Internal/Graph/AnimNextGraph.h
--- a/Engine/Plugins/Experimental/AnimNext/Source/AnimNext/Internal/Graph/AnimNextGraph.h
+++ b/Engine/Plugins/Experimental/AnimNext/Source/AnimNext/Internal/Graph/AnimNextGraph.h
@@ -31,6 +31,7 @@ namespace UE::AnimNext
 {
 	struct FContext;
 	struct FExecutionContext;
+	class FModule;
 	struct FTestUtils;
 }
 
@@ -109,6 +110,7 @@ protected:
 	friend UE::AnimNext::FExecutionContext;
 	friend class UAnimNextSchedule;
 	friend struct FAnimNextScheduleGraphTask;
+	friend UE::AnimNext::FModule;
 	
 #if WITH_EDITORONLY_DATA
 	mutable FRWLock GraphInstancesLock;
diff --git a/Engine/Plugins/Experimental/AnimNext/Source/AnimNext/Private/DecoratorBase/Decorator.cpp b/Engine/Plugins/Experimental/AnimNext/Source/AnimNext/Private/DecoratorBase/Decorator.cpp
--- a/Engine/Plugins/Experimental/AnimNext/Source/AnimNext/Private/DecoratorBase/Decorator.cpp
+++ b/Engine/Plugins/Experimental/AnimNext/Source/AnimNext/Private/DecoratorBase/Decorator.cpp
@@ -12,8 +12,52 @@ namespace UE::AnimNext
 		SharedDataStruct->SerializeItem(Ar, &SharedData, nullptr);
 	}
 
+	FDecoratorLatentPropertyMemoryLayout FDecorator::GetLatentPropertyMemoryLayoutImpl(
+		FName PropertyName,
+		uint32 PropertyIndex,
+		TArray<FDecoratorLatentPropertyMemoryLayout>& LatentPropertyMemoryLayouts) const
+	{
+		check(LatentPropertyMemoryLayouts.IsValidIndex(PropertyIndex));
+		if (LatentPropertyMemoryLayouts[PropertyIndex].Size == 0)
+		{
+			// This is a new entry, initialize it
+			// No need for locking, this is a deterministic write
+			const UScriptStruct* SharedDataStruct = GetDecoratorSharedDataStruct();
+			const FProperty* Property = SharedDataStruct->FindPropertyByName(PropertyName);
+			check(Property != nullptr);
+
+			LatentPropertyMemoryLayouts[PropertyIndex].Alignment = Property->GetMinAlignment();
+
+			// Ensure alignment is visible before we write the size to avoid torn reads
+			FPlatformMisc::MemoryBarrier();
+
+			LatentPropertyMemoryLayouts[PropertyIndex].Size = Property->GetSize();
+		}
+		
+		return LatentPropertyMemoryLayouts[PropertyIndex];
+	}
+
+	TArray<FDecoratorInterfaceUID> FDecorator::BuildDecoratorInterfaceList(
+		const TConstArrayView<FDecoratorInterfaceUID>& SuperInterfaces,
+		std::initializer_list<FDecoratorInterfaceUID> InterfaceList)
+	{
+		TArray<FDecoratorInterfaceUID> Result;
+		Result.Reserve(SuperInterfaces.Num() + InterfaceList.size());
+
+		Result.Append(SuperInterfaces);
+
+		for (FDecoratorInterfaceUID InterfaceID : InterfaceList)
+		{
+			Result.AddUnique(InterfaceID);
+		}
+
+		Result.Shrink();
+		Result.Sort();
+		return Result;
+	}
+
 #if WITH_EDITOR
-	void FDecorator::SaveDecoratorSharedData(const TFunction<FString(const FString& PropertyName)>& GetDecoratorProperty, FAnimNextDecoratorSharedData& OutSharedData) const
+	void FDecorator::SaveDecoratorSharedData(const TFunction<FString(FName PropertyName)>& GetDecoratorProperty, FAnimNextDecoratorSharedData& OutSharedData) const
 	{
 		const UScriptStruct* SharedDataStruct = GetDecoratorSharedDataStruct();
 
@@ -28,7 +72,7 @@ namespace UE::AnimNext
 		{
 			// No need to skip editor only properties since serialization will take care of that afterwards
 
-			const FString PropertyValue = GetDecoratorProperty(Property->GetName());
+			const FString PropertyValue = GetDecoratorProperty(Property->GetFName());
 			if (PropertyValue.Len() != 0)
 			{
 				const TCHAR* PropertyValuePtr = *PropertyValue;
@@ -62,7 +106,9 @@ namespace UE::AnimNext
 		}
 	}
 
-	TArray<FLatentPropertyHandle> FDecorator::GetLatentPropertyHandles(bool bFilterEditorOnly, const TFunction<bool(const FString& PropertyName)>& IsDecoratorPropertyLatent, FLatentPropertyHandle& CurrentLatentPropertyHandle) const
+	TArray<FLatentPropertyMetadata> FDecorator::GetLatentPropertyHandles(
+		bool bFilterEditorOnly,
+		const TFunction<uint16(FName PropertyName)>& GetDecoratorLatentPropertyIndex) const
 	{
 		const UStruct* BaseStruct = GetDecoratorSharedDataStruct();
 
@@ -78,12 +124,13 @@ namespace UE::AnimNext
 		}
 		while (BaseStruct != nullptr);
 
-		TArray<FLatentPropertyHandle> LatentPropertyHandles;
+		TArray<FLatentPropertyMetadata> LatentPropertyHandles;
 
 		// Gather our latent properties from base to most derived
 		for (auto It = StructHierarchy.rbegin(); It != StructHierarchy.rend(); ++It)
 		{
-			for (const FField* Field = (*It)->ChildProperties; Field; Field = Field->Next)
+			const UStruct* SharedDataStruct = *It;
+			for (const FField* Field = SharedDataStruct->ChildProperties; Field != nullptr; Field = Field->Next)
 			{
 				const FProperty* Property = CastField<FProperty>(Field);
 
@@ -97,7 +144,8 @@ namespace UE::AnimNext
 				//     - Properties marked as hidden are not visible in the editor and cannot be hooked up manually
 				//     - Properties marked as inline are only visible in the details panel and cannot be hooked up to another node
 				//     - Properties of decorator handle type are never lazy since they just encode graph connectivity
-				const bool bIsPotentiallyLatent = !Property->HasMetaData(TEXT("Hidden")) &&
+				const bool bIsPotentiallyLatent =
+					!Property->HasMetaData(TEXT("Hidden")) &&
 					!Property->HasMetaData(TEXT("Inline")) &&
 					Property->GetCPPType() != TEXT("FAnimNextDecoratorHandle");
 
@@ -106,16 +154,14 @@ namespace UE::AnimNext
 					continue;	// Skip non-latent properties
 				}
 
-				FLatentPropertyHandle LatentHandle;
+				FLatentPropertyMetadata Metadata;
+				Metadata.Name = Property->GetFName();
+				Metadata.RigVMIndex = GetDecoratorLatentPropertyIndex(Property->GetFName());
 
-				if (IsDecoratorPropertyLatent(Property->GetName()))
-				{
-					// This property is marked latent and it has a non-inline value, grab our next latent handle and increment it
-					LatentHandle = CurrentLatentPropertyHandle;
-					CurrentLatentPropertyHandle = CurrentLatentPropertyHandle.GetNextHandle();
-				}
+				// Always false for now, we don't support freezing yet
+				Metadata.bCanFreeze = false;
 
-				LatentPropertyHandles.Add(LatentHandle);
+				LatentPropertyHandles.Add(Metadata);
 			}
 		}
 
diff --git a/Engine/Plugins/Experimental/AnimNext/Source/AnimNext/Private/DecoratorBase/DecoratorReader.cpp b/Engine/Plugins/Experimental/AnimNext/Source/AnimNext/Private/DecoratorBase/DecoratorReader.cpp
--- a/Engine/Plugins/Experimental/AnimNext/Source/AnimNext/Private/DecoratorBase/DecoratorReader.cpp
+++ b/Engine/Plugins/Experimental/AnimNext/Source/AnimNext/Private/DecoratorBase/DecoratorReader.cpp
@@ -2,6 +2,7 @@
 
 #include "DecoratorBase/DecoratorReader.h"
 
+#include "DecoratorBase/DecoratorRegistry.h"
 #include "DecoratorBase/NodeDescription.h"
 #include "DecoratorBase/NodeInstance.h"
 #include "DecoratorBase/NodeTemplate.h"
@@ -39,6 +40,7 @@ namespace UE::AnimNext
 		TArray<FNodeTemplateRegistryHandle> NodeTemplateHandles;
 		NodeTemplateHandles.Reserve(NumNodeTemplates);
 
+		const FDecoratorRegistry& DecoratorRegistry = FDecoratorRegistry::Get();
 		FNodeTemplateRegistry& NodeTemplateRegistry = FNodeTemplateRegistry::Get();
 
 		{
@@ -81,8 +83,6 @@ namespace UE::AnimNext
 		uint32 SharedDataSize = 0;
 		for (uint32 NodeIndex = 0; NodeIndex < NumNodes; ++NodeIndex)
 		{
-			SharedDataSize = Align(SharedDataSize, alignof(FNodeDescription));	// Make sure we respect our alignment constraints
-
 			if (SharedDataSize > MAXIMUM_GRAPH_SHARED_DATA_SIZE)
 			{
 				// This graph shared data is too large, we won't be able to create handles to this node
@@ -116,23 +116,77 @@ namespace UE::AnimNext
 			FNodeDescription* NodeDesc = reinterpret_cast<FNodeDescription*>(&GraphSharedData[SharedDataOffset]);
 			NodeDesc->Serialize(*this);
 
-			// Serialize our latent property handles
 			const FNodeTemplate* NodeTemplate = NodeTemplateRegistry.Find(NodeDesc->GetTemplateHandle());
+			check(NodeTemplate != nullptr);
+
+			uint32 NodeInstanceDataSize = NodeTemplate->GetNodeInstanceDataSize();
+			uint32 LatentPropertyOffset = NodeInstanceDataSize;
 
+			// Read the latent properties and add them to our instance data (if any)
 			const uint32 NumDecorators = NodeTemplate->GetNumDecorators();
 			const FDecoratorTemplate* DecoratorTemplates = NodeTemplate->GetDecorators();
-
 			for (uint32 DecoratorIndex = 0; DecoratorIndex < NumDecorators; ++DecoratorIndex)
 			{
-				int32 NumLatentHandles = 0;
-				*this << NumLatentHandles;
+				const FDecoratorTemplate& DecoratorTemplate = DecoratorTemplates[DecoratorIndex];
+
+				int32 NumLatentProperties = 0;
+				*this << NumLatentProperties;
+				check(NumLatentProperties == DecoratorTemplate.GetNumLatentPropreties());
+
+				if (NumLatentProperties == 0)
+				{
+					continue;	// Nothing to do
+				}
+
+				FLatentPropertiesHeader& LatentHeader = DecoratorTemplate.GetDecoratorLatentPropertiesHeader(*NodeDesc);
+				FLatentPropertyHandle* LatentHandles = DecoratorTemplate.GetDecoratorLatentPropertyHandles(*NodeDesc);
+
+				bool bHasValidLatentProperties = false;
+				bool bCanAllPropertiesFreeze = true;
+
+				const FDecoratorRegistryHandle DecoratorHandle = DecoratorTemplate.GetRegistryHandle();
+				const FDecorator* Decorator = DecoratorRegistry.Find(DecoratorHandle);
 
-				FLatentPropertyHandle* LatentHandles = DecoratorTemplates[DecoratorIndex].GetDecoratorLatentPropertyHandles(*NodeDesc);
-				for (int32 LatentHandleIndex = 0; LatentHandleIndex < NumLatentHandles; ++LatentHandleIndex)
+				for (int32 LatentPropertyIndex = 0; LatentPropertyIndex < NumLatentProperties; ++LatentPropertyIndex)
 				{
-					*this << LatentHandles[LatentHandleIndex];
+					FLatentPropertyMetadata Metadata;
+					*this << Metadata;
+
+					uint16 RigVMIndex = MAX_uint16;
+					uint32 CurrentLatentPropertyOffset = 0;
+					bool bCanFreeze = true;
+
+					if (Decorator != nullptr)
+					{
+						// If this property is valid, setup out binding for it
+						if (Metadata.RigVMIndex != MAX_uint16)
+						{
+							const FDecoratorLatentPropertyMemoryLayout PropertyMemoryLayout = Decorator->GetLatentPropertyMemoryLayout(Metadata.Name, LatentPropertyIndex);
+
+							// Align our property
+							LatentPropertyOffset = Align(LatentPropertyOffset, PropertyMemoryLayout.Alignment);
+
+							RigVMIndex = Metadata.RigVMIndex;
+							CurrentLatentPropertyOffset = LatentPropertyOffset;
+							bCanFreeze = Metadata.bCanFreeze;
+
+							bHasValidLatentProperties = true;
+							bCanAllPropertiesFreeze &= bCanFreeze;
+
+							// Consume the property size
+							LatentPropertyOffset += PropertyMemoryLayout.Size;
+						}
+					}
+
+					LatentHandles[LatentPropertyIndex] = FLatentPropertyHandle(RigVMIndex, CurrentLatentPropertyOffset, bCanFreeze);
 				}
+
+				LatentHeader.bHasValidLatentProperties = bHasValidLatentProperties;
+				LatentHeader.bCanAllPropertiesFreeze = bCanAllPropertiesFreeze;
 			}
+
+			// Set our final node instance data size that factors in our latent properties
+			NodeDesc->NodeInstanceDataSize = LatentPropertyOffset;
 		}
 
 		return EErrorState::None;
diff --git a/Engine/Plugins/Experimental/AnimNext/Source/AnimNext/Private/DecoratorBase/DecoratorWriter.cpp b/Engine/Plugins/Experimental/AnimNext/Source/AnimNext/Private/DecoratorBase/DecoratorWriter.cpp
--- a/Engine/Plugins/Experimental/AnimNext/Source/AnimNext/Private/DecoratorBase/DecoratorWriter.cpp
+++ b/Engine/Plugins/Experimental/AnimNext/Source/AnimNext/Private/DecoratorBase/DecoratorWriter.cpp
@@ -75,7 +75,6 @@ namespace UE::AnimNext
 
 		bIsNodeWriting = true;
 		GraphReferencedObjects.Reset();
-		CurrentLatentPropertyHandle = FLatentPropertyHandle::GetFirstHandle();
 
 		// Serialize the node templates
 		TArray<FNodeTemplateRegistryHandle> NodeTemplateHandles;
@@ -124,8 +123,8 @@ namespace UE::AnimNext
 
 	void FDecoratorWriter::WriteNode(
 		const FNodeHandle NodeHandle,
-		const TFunction<FString (uint32 DecoratorIndex, const FString& PropertyName)>& GetDecoratorProperty,
-		const TFunction<bool(uint32 DecoratorIndex, const FString& PropertyName)>& IsDecoratorPropertyLatent)
+		const TFunction<FString (uint32 DecoratorIndex, FName PropertyName)>& GetDecoratorProperty,
+		const TFunction<uint16(uint32 DecoratorIndex, FName PropertyName)>& GetDecoratorLatentPropertyIndex)
 	{
 		ensure(bIsNodeWriting);
 
@@ -168,7 +167,7 @@ namespace UE::AnimNext
 			FAnimNextDecoratorSharedData* SharedData = DecoratorTemplates[DecoratorIndex].GetDecoratorDescription(*NodeDesc);
 
 			// Curry our lambda with the decorator index
-			const auto GetDecoratorPropertyAt = [&GetDecoratorProperty, DecoratorIndex](const FString& PropertyName)
+			const auto GetDecoratorPropertyAt = [&GetDecoratorProperty, DecoratorIndex](FName PropertyName)
 			{
 				return GetDecoratorProperty(DecoratorIndex, PropertyName);
 			};
@@ -180,25 +179,35 @@ namespace UE::AnimNext
 		NodeDesc->Serialize(*this);
 
 		// Append our decorator latent property handles to our archive
+		// We only write out the properties that will be present at runtime
+		// This takes into account editor only latent properties which can be stripped in cooked builds
+		// Other forms of property stripping are not currently supported
+		// The latent property offsets will be computed at runtime on load to support property sizes/alignment
+		// changing between the editor and the runtime platform (e.g. 32 vs 64 bit pointers)
+		// To that end, we serialize the following property metadata:
+		//     * RigVM memory handle index
+		//     * Whether the property supports freezing or not
+		//     * The property name and index for us to look it up at runtime
+
 		for (uint32 DecoratorIndex = 0; DecoratorIndex < NumDecorators; ++DecoratorIndex)
 		{
 			const FDecoratorRegistryHandle DecoratorHandle = DecoratorTemplates[DecoratorIndex].GetRegistryHandle();
 			const FDecorator* Decorator = DecoratorRegistry.Find(DecoratorHandle);
 
 			// Curry our lambda with the decorator index
-			const auto IsDecoratorPropertyLatentAt = [&IsDecoratorPropertyLatent, DecoratorIndex](const FString& PropertyName)
+			const auto GetDecoratorLatentPropertyIndexAt = [&GetDecoratorLatentPropertyIndex, DecoratorIndex](FName PropertyName)
 			{
-				return IsDecoratorPropertyLatent(DecoratorIndex, PropertyName);
+				return GetDecoratorLatentPropertyIndex(DecoratorIndex, PropertyName);
 			};
 
-			const TArray<FLatentPropertyHandle> LatentHandles = Decorator->GetLatentPropertyHandles(IsFilterEditorOnly(), IsDecoratorPropertyLatentAt, CurrentLatentPropertyHandle);
+			TArray<FLatentPropertyMetadata> LatentProperties = Decorator->GetLatentPropertyHandles(IsFilterEditorOnly(), GetDecoratorLatentPropertyIndexAt);
 
-			int32 NumLatentHandles = LatentHandles.Num();
-			*this << NumLatentHandles;
+			int32 NumLatentProperties = LatentProperties.Num();
+			*this << NumLatentProperties;
 
-			for (FLatentPropertyHandle Handle : LatentHandles)
+			for (FLatentPropertyMetadata& Metadata : LatentProperties)
 			{
-				*this << Handle;
+				*this << Metadata;
 			}
 		}
 
diff --git a/Engine/Plugins/Experimental/AnimNext/Source/AnimNext/Private/DecoratorBase/ExecutionContext.cpp b/Engine/Plugins/Experimental/AnimNext/Source/AnimNext/Private/DecoratorBase/ExecutionContext.cpp
--- a/Engine/Plugins/Experimental/AnimNext/Source/AnimNext/Private/DecoratorBase/ExecutionContext.cpp
+++ b/Engine/Plugins/Experimental/AnimNext/Source/AnimNext/Private/DecoratorBase/ExecutionContext.cpp
@@ -11,7 +11,6 @@
 #include "DecoratorBase/NodeTemplateRegistry.h"
 #include "Graph/AnimNextGraph.h"
 #include "Graph/AnimNextGraphInstance.h"
-#include "RigVMCore/RigVMExecuteContext.h"
 
 namespace UE::AnimNext
 {
@@ -122,7 +121,7 @@ namespace UE::AnimNext
 		// We need to allocate a new node instance
 		const FDecoratorTemplate* DecoratorDescs = NodeTemplate->GetDecorators();
 
-		const uint32 InstanceSize = NodeTemplate->GetNodeInstanceDataSize();
+		const uint32 InstanceSize = NodeDesc.GetNodeInstanceDataSize();
 		uint8* NodeInstanceBuffer = reinterpret_cast<uint8*>(FMemory::Malloc(InstanceSize, 16));
 		FNodeInstance* NodeInstance = new(NodeInstanceBuffer) FNodeInstance(*GraphInstance, ChildNodeHandle);
 
@@ -383,10 +382,48 @@ namespace UE::AnimNext
 		return false;
 	}
 
-	void FExecutionContext::EvaluateLatentPinImpl(FLatentPropertyHandle LatentPropertyHandle, void* DestinationPtr) const
+	void FExecutionContext::SnapshotLatentProperties(const FWeakDecoratorPtr& DecoratorPtr, bool bIsFrozen) const
 	{
-		check(LatentPropertyHandle.IsValid());
-		GraphInstance->ExecuteLatentPin(LatentPropertyHandle.GetLatentPropertyIndex(), DestinationPtr);
+		if (!DecoratorPtr.IsValid())
+		{
+			return;	// Nothing to do
+		}
+
+		FNodeInstance* NodeInstance = DecoratorPtr.GetNodeInstance();
+
+		if (!ensure(IsBoundTo(NodeInstance->GetOwner())))
+		{
+			return;	// The execution context isn't bound to the right graph instance
+		}
+
+		const FNodeDescription& NodeDesc = GetNodeDescription(NodeInstance->GetNodeHandle());
+
+		const FNodeTemplate* NodeTemplate = GetNodeTemplate(NodeDesc);
+		if (!ensure(NodeTemplate != nullptr))
+		{
+			return;	// Node template wasn't found, node descriptor is perhaps corrupted
+		}
+
+		const FDecoratorTemplate* DecoratorDescs = NodeTemplate->GetDecorators();
+
+		// We only snapshot the partial stack the specified decorator lives in
+		const FDecoratorTemplate* CurrentDecoratorDesc = DecoratorDescs + DecoratorPtr.GetDecoratorIndex();
+		const FDecoratorTemplate* BaseDecoratorDesc = CurrentDecoratorDesc->GetMode() == EDecoratorMode::Base ? CurrentDecoratorDesc : (CurrentDecoratorDesc - CurrentDecoratorDesc->GetAdditiveDecoratorIndex());
+
+		const FLatentPropertiesHeader& LatentHeader = BaseDecoratorDesc->GetDecoratorLatentPropertiesHeader(NodeDesc);
+		if (!LatentHeader.bHasValidLatentProperties)
+		{
+			return;	// All latent properties are inline, nothing to snapshot
+		}
+		else if (bIsFrozen && LatentHeader.bCanAllPropertiesFreeze)
+		{
+			return;	// We are frozen and all latent properties support freezing, nothing to snapshot
+		}
+
+		const FLatentPropertyHandle* LatentHandles = BaseDecoratorDesc->GetDecoratorLatentPropertyHandles(NodeDesc);
+		const uint32 NumLatentHandles = BaseDecoratorDesc->GetNumSubStackLatentPropreties();
+
+		GraphInstance->ExecuteLatentPins(TConstArrayView<FLatentPropertyHandle>(LatentHandles, NumLatentHandles), NodeInstance, bIsFrozen);
 	}
 
 	FGraphInstanceComponent* FExecutionContext::TryGetComponent(int32 ComponentNameHash, FName ComponentName) const
diff --git a/Engine/Plugins/Experimental/AnimNext/Source/AnimNext/Private/DecoratorBase/NodeTemplate.cpp b/Engine/Plugins/Experimental/AnimNext/Source/AnimNext/Private/DecoratorBase/NodeTemplate.cpp
--- a/Engine/Plugins/Experimental/AnimNext/Source/AnimNext/Private/DecoratorBase/NodeTemplate.cpp
+++ b/Engine/Plugins/Experimental/AnimNext/Source/AnimNext/Private/DecoratorBase/NodeTemplate.cpp
@@ -10,6 +10,34 @@
 
 namespace UE::AnimNext
 {
+	namespace Private
+	{
+		static uint32 GetNumSubStackLatentProperties(const FDecoratorRegistry& DecoratorRegistry, const FDecoratorTemplate* DecoratorTemplates, uint32 NumDecorators, uint32 BaseDecoratorIndex)
+		{
+			const FDecoratorTemplate& BaseDecoratorTemplate = DecoratorTemplates[BaseDecoratorIndex];
+			check(BaseDecoratorTemplate.GetMode() == EDecoratorMode::Base);
+
+			const uint32 NumSubStackDecorators = BaseDecoratorTemplate.GetNumAdditiveDecorators() + 1;
+			uint32 NumSubStackLatentProperties = 0;
+
+			for (uint32 SubStackDecoratorIndex = 0; SubStackDecoratorIndex < NumSubStackDecorators; ++SubStackDecoratorIndex)
+			{
+				const uint32 DecoratorIndex = BaseDecoratorIndex + SubStackDecoratorIndex;
+				check(DecoratorIndex < NumDecorators);
+
+				const FDecoratorTemplate& DecoratorTemplate = DecoratorTemplates[DecoratorIndex];
+				const FDecoratorUID DecoratorUID = DecoratorTemplate.GetUID();
+
+				if (const FDecorator* Decorator = DecoratorRegistry.Find(DecoratorUID))
+				{
+					NumSubStackLatentProperties += Decorator->GetNumLatentDecoratorProperties();
+				}
+			}
+
+			return NumSubStackLatentProperties;
+		}
+	}
+
 	void FNodeTemplate::Serialize(FArchive& Ar)
 	{
 		Ar << UID;
@@ -39,16 +67,23 @@ namespace UE::AnimNext
 
 		uint32 SharedDataOffset = sizeof(FNodeDescription);
 		uint32 InstanceDataOffset = sizeof(FNodeInstance);
+		uint32 SharedLatentPropertyHandlesOffset = 0;
 
 		for (uint32 DecoratorIndex = 0; DecoratorIndex < NumDecorators_; ++DecoratorIndex)
 		{
 			FDecoratorTemplate& DecoratorTemplate = DecoratorTemplates[DecoratorIndex];
 			const FDecoratorUID DecoratorUID = DecoratorTemplate.GetUID();
 
+			uint32 NumLatentProperties = 0;
 			uint32 DecoratorSharedDataOffset = 0;
 			uint32 DecoratorSharedLatentPropertyHandlesOffset = 0;
 			uint32 DecoratorInstanceDataOffset = 0;	// For instance data, 0 is an invalid offset since the data follows an instance of FNodeInstance
 
+			uint32 NumSubStackLatentProperties = 0;
+			if (DecoratorTemplate.GetMode() == EDecoratorMode::Base)
+			{
+				NumSubStackLatentProperties = Private::GetNumSubStackLatentProperties(DecoratorRegistry, DecoratorTemplates, NumDecorators_, DecoratorIndex);
+			}
 			// Skip decorators that we can't find
 			// If a decorator isn't loaded and we attempt to run the graph, it will be a no-op entry
 			if (const FDecorator* Decorator = DecoratorRegistry.Find(DecoratorUID))
@@ -67,14 +102,38 @@ namespace UE::AnimNext
 				SharedDataOffset += MemoryLayout.SharedDataSize;
 				InstanceDataOffset += MemoryLayout.InstanceDataSize;
 
-				// Save our latent pins offset (if we have any)
-				const uint32 NumLatentProperties = Decorator->GetNumLatentDecoratorProperties();
-				DecoratorSharedLatentPropertyHandlesOffset = NumLatentProperties != 0 ? Align(SharedDataOffset, alignof(FLatentPropertyHandle)) : 0;
+				// Base decorators include the list of all latent property handles in its shared data
+				// Latent property offsets will point into that list
+				if (DecoratorTemplate.GetMode() == EDecoratorMode::Base && NumSubStackLatentProperties != 0)
+				{
+					// Align our handles
+					SharedDataOffset = Align(SharedDataOffset, alignof(FLatentPropertiesHeader));
+
+					// Save the offset where we start, we'll increment it as we consume it
+					SharedLatentPropertyHandlesOffset = SharedDataOffset;
+
+					// Include the handles in the shared data and their header
+					SharedDataOffset += sizeof(FLatentPropertiesHeader) + (NumSubStackLatentProperties * sizeof(FLatentPropertyHandle));
 
-				// Include our latent pins
-				SharedDataOffset += NumLatentProperties * sizeof(FLatentPropertyHandle);
+					// Skip the header
+					SharedLatentPropertyHandlesOffset += sizeof(FLatentPropertiesHeader);
+				}
+
+				// Save our latent pins offset (if we have any)
+				NumLatentProperties = Decorator->GetNumLatentDecoratorProperties();
+				if (NumLatentProperties != 0)
+				{
+					DecoratorSharedLatentPropertyHandlesOffset = SharedLatentPropertyHandlesOffset;
+					SharedLatentPropertyHandlesOffset += NumLatentProperties * sizeof(FLatentPropertyHandle);
+				}
 			}
 
+			check(NumSubStackLatentProperties <= MAX_uint16);
+			check(NumLatentProperties <= MAX_uint16);
+
+			DecoratorTemplate.NumLatentProperties = static_cast<uint16>(NumLatentProperties);
+			DecoratorTemplate.NumSubStackLatentProperties = static_cast<uint16>(NumSubStackLatentProperties);
+
 			check(DecoratorSharedDataOffset <= MAX_uint16);
 			check(DecoratorSharedLatentPropertyHandlesOffset <= MAX_uint16);
 			check(DecoratorInstanceDataOffset <= MAX_uint16);
@@ -85,6 +144,9 @@ namespace UE::AnimNext
 			DecoratorTemplate.NodeInstanceOffset = static_cast<uint16>(DecoratorInstanceDataOffset);
 		}
 
+		// Make sure we respect our alignment constraints
+		SharedDataOffset = Align(SharedDataOffset, alignof(FNodeDescription));
+
 		// Our size is the offset of the decorator that would follow afterwards
 		// If the size is too large, we'll end up truncating the offsets/size
 		// Set a value of 0 to be able to detect it later
diff --git a/Engine/Plugins/Experimental/AnimNext/Source/AnimNext/Private/DecoratorInterfaces/IUpdate.cpp b/Engine/Plugins/Experimental/AnimNext/Source/AnimNext/Private/DecoratorInterfaces/IUpdate.cpp
--- a/Engine/Plugins/Experimental/AnimNext/Source/AnimNext/Private/DecoratorInterfaces/IUpdate.cpp
+++ b/Engine/Plugins/Experimental/AnimNext/Source/AnimNext/Private/DecoratorInterfaces/IUpdate.cpp
@@ -252,9 +252,15 @@ namespace UE::AnimNext
 
 			if (!Entry->bHasPreUpdated)
 			{
+				// This is the first time we visit this node, time to pre-update
+				// But first, if it has latent pins, we must execute and cache their results
+				// This will ensure that other calls into this node will have a consistent view of
+				// what the node saw when it started to update. We thus take a snapshot.
+				const bool bIsFrozen = false;	// Not yet supported
+				ExecutionContext.SnapshotLatentProperties(EntryDecoratorPtr, bIsFrozen);
+
 				if (ExecutionContext.GetInterface(EntryDecoratorPtr, Entry->UpdateDecorator))
 				{
-					// This is the first time we visit this node, time to pre-update
 					Entry->UpdateDecorator.PreUpdate(TraversalContext, Entry->DecoratorState);
 
 					// Make sure that next time we visit this entry, we'll post-update
diff --git a/Engine/Plugins/Experimental/AnimNext/Source/AnimNext/Private/Decorators/BlendByBool.cpp b/Engine/Plugins/Experimental/AnimNext/Source/AnimNext/Private/Decorators/BlendByBool.cpp
--- a/Engine/Plugins/Experimental/AnimNext/Source/AnimNext/Private/Decorators/BlendByBool.cpp
+++ b/Engine/Plugins/Experimental/AnimNext/Source/AnimNext/Private/Decorators/BlendByBool.cpp
@@ -10,11 +10,14 @@ namespace UE::AnimNext
 {
 	AUTO_REGISTER_ANIM_DECORATOR(FBlendByBoolDecorator)
 
-	DEFINE_ANIM_DECORATOR_BEGIN(FBlendByBoolDecorator)
-		DEFINE_ANIM_DECORATOR_IMPLEMENTS_INTERFACE(IDiscreteBlend)
-		DEFINE_ANIM_DECORATOR_IMPLEMENTS_INTERFACE(IHierarchy)
-		DEFINE_ANIM_DECORATOR_IMPLEMENTS_INTERFACE(IUpdate)
-	DEFINE_ANIM_DECORATOR_END(FBlendByBoolDecorator)
+	// Decorator implementation boilerplate
+	#define DECORATOR_INTERFACE_ENUMERATOR(GeneratorMacro) \
+		GeneratorMacro(IDiscreteBlend) \
+		GeneratorMacro(IHierarchy) \
+		GeneratorMacro(IUpdate) \
+
+	GENERATE_ANIM_DECORATOR_IMPLEMENTATION(FBlendByBoolDecorator, DECORATOR_INTERFACE_ENUMERATOR)
+	#undef DECORATOR_INTERFACE_ENUMERATOR
 
 	static constexpr int32 TRUE_CHILD_INDEX = 0;
 	static constexpr int32 FALSE_CHILD_INDEX = 1;
diff --git a/Engine/Plugins/Experimental/AnimNext/Source/AnimNext/Private/Decorators/BlendInertializer.cpp b/Engine/Plugins/Experimental/AnimNext/Source/AnimNext/Private/Decorators/BlendInertializer.cpp
--- a/Engine/Plugins/Experimental/AnimNext/Source/AnimNext/Private/Decorators/BlendInertializer.cpp
+++ b/Engine/Plugins/Experimental/AnimNext/Source/AnimNext/Private/Decorators/BlendInertializer.cpp
@@ -8,10 +8,13 @@ namespace UE::AnimNext
 {
 	AUTO_REGISTER_ANIM_DECORATOR(FBlendInertializerDecorator)
 
-	DEFINE_ANIM_DECORATOR_BEGIN(FBlendInertializerDecorator)
-		DEFINE_ANIM_DECORATOR_IMPLEMENTS_INTERFACE(IDiscreteBlend)
-		DEFINE_ANIM_DECORATOR_IMPLEMENTS_INTERFACE(ISmoothBlend)
-	DEFINE_ANIM_DECORATOR_END(FBlendInertializerDecorator)
+	// Decorator implementation boilerplate
+	#define DECORATOR_INTERFACE_ENUMERATOR(GeneratorMacro) \
+		GeneratorMacro(IDiscreteBlend) \
+		GeneratorMacro(ISmoothBlend) \
+
+	GENERATE_ANIM_DECORATOR_IMPLEMENTATION(FBlendInertializerDecorator, DECORATOR_INTERFACE_ENUMERATOR)
+	#undef DECORATOR_INTERFACE_ENUMERATOR
 
 	void FBlendInertializerDecorator::OnBlendTransition(const FExecutionContext& Context, const TDecoratorBinding<IDiscreteBlend>& Binding, int32 OldChildIndex, int32 NewChildIndex) const
 	{
diff --git a/Engine/Plugins/Experimental/AnimNext/Source/AnimNext/Private/Decorators/BlendSmoother.cpp b/Engine/Plugins/Experimental/AnimNext/Source/AnimNext/Private/Decorators/BlendSmoother.cpp
--- a/Engine/Plugins/Experimental/AnimNext/Source/AnimNext/Private/Decorators/BlendSmoother.cpp
+++ b/Engine/Plugins/Experimental/AnimNext/Source/AnimNext/Private/Decorators/BlendSmoother.cpp
@@ -12,12 +12,15 @@ namespace UE::AnimNext
 {
 	AUTO_REGISTER_ANIM_DECORATOR(FBlendSmootherDecorator)
 
-	DEFINE_ANIM_DECORATOR_BEGIN(FBlendSmootherDecorator)
-		DEFINE_ANIM_DECORATOR_IMPLEMENTS_INTERFACE(IDiscreteBlend)
-		DEFINE_ANIM_DECORATOR_IMPLEMENTS_INTERFACE(IEvaluate)
-		DEFINE_ANIM_DECORATOR_IMPLEMENTS_INTERFACE(ISmoothBlend)
-		DEFINE_ANIM_DECORATOR_IMPLEMENTS_INTERFACE(IUpdate)
-	DEFINE_ANIM_DECORATOR_END(FBlendSmootherDecorator)
+	// Decorator implementation boilerplate
+	#define DECORATOR_INTERFACE_ENUMERATOR(GeneratorMacro) \
+		GeneratorMacro(IDiscreteBlend) \
+		GeneratorMacro(IEvaluate) \
+		GeneratorMacro(ISmoothBlend) \
+		GeneratorMacro(IUpdate) \
+
+	GENERATE_ANIM_DECORATOR_IMPLEMENTATION(FBlendSmootherDecorator, DECORATOR_INTERFACE_ENUMERATOR)
+	#undef DECORATOR_INTERFACE_ENUMERATOR
 
 	void FBlendSmootherDecorator::PostEvaluate(FEvaluateTraversalContext& Context, const TDecoratorBinding<IEvaluate>& Binding) const
 	{
diff --git a/Engine/Plugins/Experimental/AnimNext/Source/AnimNext/Private/Decorators/BlendSmootherPerBone.cpp b/Engine/Plugins/Experimental/AnimNext/Source/AnimNext/Private/Decorators/BlendSmootherPerBone.cpp
--- a/Engine/Plugins/Experimental/AnimNext/Source/AnimNext/Private/Decorators/BlendSmootherPerBone.cpp
+++ b/Engine/Plugins/Experimental/AnimNext/Source/AnimNext/Private/Decorators/BlendSmootherPerBone.cpp
@@ -13,11 +13,14 @@ namespace UE::AnimNext
 {
 	AUTO_REGISTER_ANIM_DECORATOR(FBlendSmootherPerBoneDecorator)
 
-	DEFINE_ANIM_DECORATOR_BEGIN(FBlendSmootherPerBoneDecorator)
-		DEFINE_ANIM_DECORATOR_IMPLEMENTS_INTERFACE(IEvaluate)
-		DEFINE_ANIM_DECORATOR_IMPLEMENTS_INTERFACE(IUpdate)
-		DEFINE_ANIM_DECORATOR_IMPLEMENTS_INTERFACE(IDiscreteBlend)
-	DEFINE_ANIM_DECORATOR_END(FBlendSmootherPerBoneDecorator)
+	// Decorator implementation boilerplate
+	#define DECORATOR_INTERFACE_ENUMERATOR(GeneratorMacro) \
+		GeneratorMacro(IDiscreteBlend) \
+		GeneratorMacro(IEvaluate) \
+		GeneratorMacro(IUpdate) \
+
+	GENERATE_ANIM_DECORATOR_IMPLEMENTATION(FBlendSmootherPerBoneDecorator, DECORATOR_INTERFACE_ENUMERATOR)
+	#undef DECORATOR_INTERFACE_ENUMERATOR
 
 	void FBlendSmootherPerBoneDecorator::PostEvaluate(FEvaluateTraversalContext& Context, const TDecoratorBinding<IEvaluate>& Binding) const
 	{
diff --git a/Engine/Plugins/Experimental/AnimNext/Source/AnimNext/Private/Decorators/BlendTwoWay.cpp b/Engine/Plugins/Experimental/AnimNext/Source/AnimNext/Private/Decorators/BlendTwoWay.cpp
--- a/Engine/Plugins/Experimental/AnimNext/Source/AnimNext/Private/Decorators/BlendTwoWay.cpp
+++ b/Engine/Plugins/Experimental/AnimNext/Source/AnimNext/Private/Decorators/BlendTwoWay.cpp
@@ -10,12 +10,15 @@ namespace UE::AnimNext
 {
 	AUTO_REGISTER_ANIM_DECORATOR(FBlendTwoWayDecorator)
 
-	DEFINE_ANIM_DECORATOR_BEGIN(FBlendTwoWayDecorator)
-		DEFINE_ANIM_DECORATOR_IMPLEMENTS_INTERFACE(IContinuousBlend)
-		DEFINE_ANIM_DECORATOR_IMPLEMENTS_INTERFACE(IEvaluate)
-		DEFINE_ANIM_DECORATOR_IMPLEMENTS_INTERFACE(IHierarchy)
-		DEFINE_ANIM_DECORATOR_IMPLEMENTS_INTERFACE(IUpdate)
-	DEFINE_ANIM_DECORATOR_END(FBlendTwoWayDecorator)
+	// Decorator implementation boilerplate
+	#define DECORATOR_INTERFACE_ENUMERATOR(GeneratorMacro) \
+		GeneratorMacro(IContinuousBlend) \
+		GeneratorMacro(IEvaluate) \
+		GeneratorMacro(IHierarchy) \
+		GeneratorMacro(IUpdate) \
+
+	GENERATE_ANIM_DECORATOR_IMPLEMENTATION(FBlendTwoWayDecorator, DECORATOR_INTERFACE_ENUMERATOR)
+	#undef DECORATOR_INTERFACE_ENUMERATOR
 
 	void FBlendTwoWayDecorator::PostEvaluate(FEvaluateTraversalContext& Context, const TDecoratorBinding<IEvaluate>& Binding) const
 	{
diff --git a/Engine/Plugins/Experimental/AnimNext/Source/AnimNext/Private/Decorators/ReferencePoseDecorator.cpp b/Engine/Plugins/Experimental/AnimNext/Source/AnimNext/Private/Decorators/ReferencePoseDecorator.cpp
--- a/Engine/Plugins/Experimental/AnimNext/Source/AnimNext/Private/Decorators/ReferencePoseDecorator.cpp
+++ b/Engine/Plugins/Experimental/AnimNext/Source/AnimNext/Private/Decorators/ReferencePoseDecorator.cpp
@@ -9,9 +9,12 @@ namespace UE::AnimNext
 {
 	AUTO_REGISTER_ANIM_DECORATOR(FReferencePoseDecorator)
 
-	DEFINE_ANIM_DECORATOR_BEGIN(FReferencePoseDecorator)
-		DEFINE_ANIM_DECORATOR_IMPLEMENTS_INTERFACE(IEvaluate)
-	DEFINE_ANIM_DECORATOR_END(FReferencePoseDecorator)
+	// Decorator implementation boilerplate
+	#define DECORATOR_INTERFACE_ENUMERATOR(GeneratorMacro) \
+		GeneratorMacro(IEvaluate) \
+
+	GENERATE_ANIM_DECORATOR_IMPLEMENTATION(FReferencePoseDecorator, DECORATOR_INTERFACE_ENUMERATOR)
+	#undef DECORATOR_INTERFACE_ENUMERATOR
 
 	void FReferencePoseDecorator::PreEvaluate(FEvaluateTraversalContext& Context, const TDecoratorBinding<IEvaluate>& Binding) const
 	{
diff --git a/Engine/Plugins/Experimental/AnimNext/Source/AnimNext/Private/Decorators/SequencePlayer.cpp b/Engine/Plugins/Experimental/AnimNext/Source/AnimNext/Private/Decorators/SequencePlayer.cpp
--- a/Engine/Plugins/Experimental/AnimNext/Source/AnimNext/Private/Decorators/SequencePlayer.cpp
+++ b/Engine/Plugins/Experimental/AnimNext/Source/AnimNext/Private/Decorators/SequencePlayer.cpp
@@ -10,11 +10,14 @@ namespace UE::AnimNext
 {
 	AUTO_REGISTER_ANIM_DECORATOR(FSequencePlayerDecorator)
 
-	DEFINE_ANIM_DECORATOR_BEGIN(FSequencePlayerDecorator)
-		DEFINE_ANIM_DECORATOR_IMPLEMENTS_INTERFACE(IEvaluate)
-		DEFINE_ANIM_DECORATOR_IMPLEMENTS_INTERFACE(IUpdate)
-		DEFINE_ANIM_DECORATOR_IMPLEMENTS_INTERFACE(ITimeline)
-	DEFINE_ANIM_DECORATOR_END(FSequencePlayerDecorator)
+	// Decorator implementation boilerplate
+	#define DECORATOR_INTERFACE_ENUMERATOR(GeneratorMacro) \
+		GeneratorMacro(IEvaluate) \
+		GeneratorMacro(ITimeline) \
+		GeneratorMacro(IUpdate) \
+
+	GENERATE_ANIM_DECORATOR_IMPLEMENTATION(FSequencePlayerDecorator, DECORATOR_INTERFACE_ENUMERATOR)
+	#undef DECORATOR_INTERFACE_ENUMERATOR
 
 	void FSequencePlayerDecorator::FInstanceData::Construct(const FExecutionContext& Context, const FDecoratorBinding& Binding)
 	{
diff --git a/Engine/Plugins/Experimental/AnimNext/Source/AnimNext/Private/Decorators/SubGraphHost.cpp b/Engine/Plugins/Experimental/AnimNext/Source/AnimNext/Private/Decorators/SubGraphHost.cpp
--- a/Engine/Plugins/Experimental/AnimNext/Source/AnimNext/Private/Decorators/SubGraphHost.cpp
+++ b/Engine/Plugins/Experimental/AnimNext/Source/AnimNext/Private/Decorators/SubGraphHost.cpp
@@ -8,12 +8,15 @@ namespace UE::AnimNext
 {
 	AUTO_REGISTER_ANIM_DECORATOR(FSubGraphHostDecorator)
 
-	DEFINE_ANIM_DECORATOR_BEGIN(FSubGraphHostDecorator)
-		DEFINE_ANIM_DECORATOR_IMPLEMENTS_INTERFACE(IDiscreteBlend)
-		DEFINE_ANIM_DECORATOR_IMPLEMENTS_INTERFACE(IGarbageCollection)
-		DEFINE_ANIM_DECORATOR_IMPLEMENTS_INTERFACE(IHierarchy)
-		DEFINE_ANIM_DECORATOR_IMPLEMENTS_INTERFACE(IUpdate)
-	DEFINE_ANIM_DECORATOR_END(FSubGraphHostDecorator)
+	// Decorator implementation boilerplate
+	#define DECORATOR_INTERFACE_ENUMERATOR(GeneratorMacro) \
+		GeneratorMacro(IDiscreteBlend) \
+		GeneratorMacro(IGarbageCollection) \
+		GeneratorMacro(IHierarchy) \
+		GeneratorMacro(IUpdate) \
+
+	GENERATE_ANIM_DECORATOR_IMPLEMENTATION(FSubGraphHostDecorator, DECORATOR_INTERFACE_ENUMERATOR)
+	#undef DECORATOR_INTERFACE_ENUMERATOR
 
 	void FSubGraphHostDecorator::FInstanceData::Construct(const FExecutionContext& Context, const FDecoratorBinding& Binding)
 	{
diff --git a/Engine/Plugins/Experimental/AnimNext/Source/AnimNext/Private/Decorators/SynchronizeUsingGroups.cpp b/Engine/Plugins/Experimental/AnimNext/Source/AnimNext/Private/Decorators/SynchronizeUsingGroups.cpp
--- a/Engine/Plugins/Experimental/AnimNext/Source/AnimNext/Private/Decorators/SynchronizeUsingGroups.cpp
+++ b/Engine/Plugins/Experimental/AnimNext/Source/AnimNext/Private/Decorators/SynchronizeUsingGroups.cpp
@@ -7,11 +7,14 @@ namespace UE::AnimNext
 {
 	AUTO_REGISTER_ANIM_DECORATOR(FSynchronizeUsingGroupsDecorator)
 
-	DEFINE_ANIM_DECORATOR_BEGIN(FSynchronizeUsingGroupsDecorator)
-		DEFINE_ANIM_DECORATOR_IMPLEMENTS_INTERFACE(IUpdate)
-		DEFINE_ANIM_DECORATOR_IMPLEMENTS_INTERFACE(IGroupSynchronization)
-		DEFINE_ANIM_DECORATOR_IMPLEMENTS_INTERFACE(ITimeline)
-	DEFINE_ANIM_DECORATOR_END(FSynchronizeUsingGroupsDecorator)
+	// Decorator implementation boilerplate
+	#define DECORATOR_INTERFACE_ENUMERATOR(GeneratorMacro) \
+		GeneratorMacro(IGroupSynchronization) \
+		GeneratorMacro(ITimeline) \
+		GeneratorMacro(IUpdate) \
+
+	GENERATE_ANIM_DECORATOR_IMPLEMENTATION(FSynchronizeUsingGroupsDecorator, DECORATOR_INTERFACE_ENUMERATOR)
+	#undef DECORATOR_INTERFACE_ENUMERATOR
 
 	void FSynchronizeUsingGroupsDecorator::PreUpdate(FUpdateTraversalContext& Context, const TDecoratorBinding<IUpdate>& Binding, const FDecoratorUpdateState& DecoratorState) const
 	{
diff --git a/Engine/Plugins/Experimental/AnimNext/Source/AnimNext/Private/Graph/AnimNextGraphInstance.cpp b/Engine/Plugins/Experimental/AnimNext/Source/AnimNext/Private/Graph/AnimNextGraphInstance.cpp
--- a/Engine/Plugins/Experimental/AnimNext/Source/AnimNext/Private/Graph/AnimNextGraphInstance.cpp
+++ b/Engine/Plugins/Experimental/AnimNext/Source/AnimNext/Private/Graph/AnimNextGraphInstance.cpp
@@ -98,7 +98,7 @@ GraphInstanceComponentMapType::TConstIterator FAnimNextGraphInstance::GetCompone
 	return RootGraphInstance->Components.CreateConstIterator();
 }
 
-void FAnimNextGraphInstance::ExecuteLatentPin(int32 LatentPinIndex, void* DestinationPtr)
+void FAnimNextGraphInstance::ExecuteLatentPins(const TConstArrayView<UE::AnimNext::FLatentPropertyHandle>& LatentHandles, void* DestinationBasePtr, bool bIsFrozen)
 {
 	SCOPE_CYCLE_COUNTER(STAT_AnimNext_Graph_RigVM);
 
@@ -110,7 +110,7 @@ void FAnimNextGraphInstance::ExecuteLatentPin(int32 LatentPinIndex, void* Destin
 	if (URigVM* VM = Graph->VM)
 	{
 		FAnimNextExecuteContext& AnimNextContext = ExtendedExecuteContext.GetPublicDataSafe<FAnimNextExecuteContext>();
-		AnimNextContext.SetupForExecution(LatentPinIndex, DestinationPtr);
+		AnimNextContext.SetupForExecution(LatentHandles, DestinationBasePtr, bIsFrozen);
 
 		VM->ExecuteVM(ExtendedExecuteContext, FRigUnit_AnimNextShimRoot::EventName);
 
diff --git a/Engine/Plugins/Experimental/AnimNext/Source/AnimNext/Private/Graph/RigUnit_AnimNextGraphEvaluator.cpp b/Engine/Plugins/Experimental/AnimNext/Source/AnimNext/Private/Graph/RigUnit_AnimNextGraphEvaluator.cpp
--- a/Engine/Plugins/Experimental/AnimNext/Source/AnimNext/Private/Graph/RigUnit_AnimNextGraphEvaluator.cpp
+++ b/Engine/Plugins/Experimental/AnimNext/Source/AnimNext/Private/Graph/RigUnit_AnimNextGraphEvaluator.cpp
@@ -2,6 +2,7 @@
 
 #include "Graph/RigUnit_AnimNextGraphEvaluator.h"
 #include "Context.h"
+#include "DecoratorBase/LatentPropertyHandle.h"
 
 #include UE_INLINE_GENERATED_CPP_BY_NAME(RigUnit_AnimNextGraphEvaluator)
 
@@ -27,31 +28,44 @@ void FRigUnit_AnimNextGraphEvaluator::StaticExecute(FRigVMExtendedExecuteContext
 {
 	const FAnimNextExecuteContext& VMExecuteContext = RigVMExecuteContext.GetPublicData<FAnimNextExecuteContext>();
 
-	const int32 LatentPinIndex = VMExecuteContext.GetLatentPinIndex();
-	if (LatentPinIndex == INDEX_NONE)
+	const TConstArrayView<UE::AnimNext::FLatentPropertyHandle>& LatentHandles = VMExecuteContext.GetLatentHandles();
+	uint8* DestinationBasePtr = (uint8*)VMExecuteContext.GetDestinationBasePtr();
+	const bool bIsFrozen = VMExecuteContext.IsFrozen();
+
+	for (UE::AnimNext::FLatentPropertyHandle Handle : LatentHandles)
 	{
-		return;
-	}
+		if (!Handle.IsIndexValid())
+		{
+			// This handle isn't valid
+			continue;
+		}
+
+		if (bIsFrozen && Handle.CanFreeze())
+		{
+			// This handle can freeze and we are frozen, no need to update it
+			continue;
+		}
 
-	FRigVMMemoryHandle& MemoryHandle = RigVMMemoryHandles[LatentPinIndex];
+		FRigVMMemoryHandle& MemoryHandle = RigVMMemoryHandles[Handle.GetLatentPropertyIndex()];
 
-	// This should be an assert. If this triggers, it means that we have a bug in how lazy memory handles
-	// are assigned during compilation. We keep it as an ensure because in this case, we can recover
-	// as even if the memory handle isn't lazy, it remains valid and we can use it. It won't have the
-	// value we expect but it'll work. The ensure will signal that we need to fix the bug.
-	if (ensure(MemoryHandle.IsLazy()))
-	{
-		MemoryHandle.ComputeLazyValueIfNecessary(RigVMExecuteContext, RigVMExecuteContext.GetSlice().GetIndex());
-	}
+		// This should be an assert. If this triggers, it means that we have a bug in how lazy memory handles
+		// are assigned during compilation. We keep it as an ensure because in this case, we can recover
+		// as even if the memory handle isn't lazy, it remains valid and we can use it. It won't have the
+		// value we expect but it'll work. The ensure will signal that we need to fix the bug.
+		if (ensure(MemoryHandle.IsLazy()))
+		{
+			MemoryHandle.ComputeLazyValueIfNecessary(RigVMExecuteContext, RigVMExecuteContext.GetSlice().GetIndex());
+		}
 
-	const uint8* SourcePtr = MemoryHandle.GetData();
-	uint8* DestinationPtr = (uint8*)VMExecuteContext.GetDestinationPtr();
+		const uint8* SourcePtr = MemoryHandle.GetData();
+		uint8* DestinationPtr = DestinationBasePtr + Handle.GetLatentPropertyOffset();
 
-	// Copy from our source into our destination
-	// We assume the source and destination properties are identical
-	URigVMMemoryStorage::CopyProperty(
-		MemoryHandle.GetProperty(), DestinationPtr,
-		MemoryHandle.GetProperty(), SourcePtr);
+		// Copy from our source into our destination
+		// We assume the source and destination properties are identical
+		URigVMMemoryStorage::CopyProperty(
+			MemoryHandle.GetProperty(), DestinationPtr,
+			MemoryHandle.GetProperty(), SourcePtr);
+	}
 }
 
 void FRigUnit_AnimNextGraphEvaluator::RegisterExecuteMethod(const FAnimNextGraphEvaluatorExecuteDefinition& ExecuteDefinition)
diff --git a/Engine/Plugins/Experimental/AnimNext/Source/AnimNext/Private/Module.cpp b/Engine/Plugins/Experimental/AnimNext/Source/AnimNext/Private/Module.cpp
--- a/Engine/Plugins/Experimental/AnimNext/Source/AnimNext/Private/Module.cpp
+++ b/Engine/Plugins/Experimental/AnimNext/Source/AnimNext/Private/Module.cpp
@@ -18,6 +18,18 @@
 #include "Param/ExternalParameterRegistry.h"
 #include "Param/ObjectProxyFactory.h"
 
+// Enable console commands only in development builds when logging is enabled
+#define WITH_ANIMNEXT_CONSOLE_COMMANDS (!UE_BUILD_SHIPPING && !NO_LOGGING)
+
+#if WITH_ANIMNEXT_CONSOLE_COMMANDS
+#include "HAL/IConsoleManager.h"
+#include "UObject/UObjectIterator.h"
+
+#include "DecoratorBase/DecoratorTemplate.h"
+#include "DecoratorBase/NodeDescription.h"
+#include "DecoratorBase/NodeTemplate.h"
+#endif
+
 namespace UE::AnimNext
 {
 
@@ -46,6 +58,24 @@ public:
 		FNodeTemplateRegistry::Init();
 		FScheduler::Init();
 		FRigVMRuntimeDataRegistry::Init();
+
+#if WITH_ANIMNEXT_CONSOLE_COMMANDS
+		if (!IsRunningCommandlet())
+		{
+			ConsoleCommands.Add(IConsoleManager::Get().RegisterConsoleCommand(
+				TEXT("AnimNext.ListNodeTemplates"),
+				TEXT("Dumps statistics about node templates to the log."),
+				FConsoleCommandWithArgsDelegate::CreateRaw(this, &FModule::ListNodeTemplates),
+				ECVF_Default
+			));
+			ConsoleCommands.Add(IConsoleManager::Get().RegisterConsoleCommand(
+				TEXT("AnimNext.ListAnimGraphs"),
+				TEXT("Dumps statistics about animation graphs to the log."),
+				FConsoleCommandWithArgsDelegate::CreateRaw(this, &FModule::ListAnimGraphs),
+				ECVF_Default
+			));
+		}
+#endif
 	}
 
 	virtual void ShutdownModule() override
@@ -57,7 +87,179 @@ public:
 		FDataRegistry::Destroy();
 		FObjectProxyFactory::Destroy();
 		FExternalParameterRegistry::Destroy();
+
+#if WITH_ANIMNEXT_CONSOLE_COMMANDS
+		for (IConsoleObject* Cmd : ConsoleCommands)
+		{
+			IConsoleManager::Get().UnregisterConsoleObject(Cmd);
+		}
+		ConsoleCommands.Empty();
+#endif
+	}
+
+#if WITH_ANIMNEXT_CONSOLE_COMMANDS
+	TArray<IConsoleObject*> ConsoleCommands;
+
+	void ListNodeTemplates(const TArray<FString>& Args)
+	{
+		// Turn off log times to make diff-ing easier
+		TGuardValue<ELogTimes::Type> DisableLogTimes(GPrintLogTimes, ELogTimes::None);
+
+		// Make sure to log everything
+		const ELogVerbosity::Type OldVerbosity = LogAnimation.GetVerbosity();
+		LogAnimation.SetVerbosity(ELogVerbosity::All);
+
+		const FNodeTemplateRegistry& NodeTemplateRegistry = FNodeTemplateRegistry::Get();
+		const FDecoratorRegistry& DecoratorRegistry = FDecoratorRegistry::Get();
+
+		UE_LOG(LogAnimation, Log, TEXT("===== AnimNext Node Templates ====="));
+		UE_LOG(LogAnimation, Log, TEXT("Template Buffer Size: %u bytes"), NodeTemplateRegistry.TemplateBuffer.GetAllocatedSize());
+
+		for (auto It = NodeTemplateRegistry.TemplateUIDToHandleMap.CreateConstIterator(); It; ++It)
+		{
+			const FNodeTemplateRegistryHandle Handle = It.Value();
+			const FNodeTemplate* NodeTemplate = NodeTemplateRegistry.Find(Handle);
+
+			const uint32 NumDecorators = NodeTemplate->GetNumDecorators();
+
+			UE_LOG(LogAnimation, Log, TEXT("[%x] has %u decorators ..."), NodeTemplate->GetUID(), NumDecorators);
+			UE_LOG(LogAnimation, Log, TEXT("    Template Size: %u bytes"), NodeTemplate->GetNodeTemplateSize());
+			UE_LOG(LogAnimation, Log, TEXT("    Shared Data Size: %u bytes"), NodeTemplate->GetNodeSharedDataSize());
+			UE_LOG(LogAnimation, Log, TEXT("    Instance Data Size: %u bytes"), NodeTemplate->GetNodeInstanceDataSize());
+			UE_LOG(LogAnimation, Log, TEXT("    Decorators ..."));
+
+			const FDecoratorTemplate* DecoratorTemplates = NodeTemplate->GetDecorators();
+			for (uint32 DecoratorIndex = 0; DecoratorIndex < NumDecorators; ++DecoratorIndex)
+			{
+				const FDecoratorTemplate* DecoratorTemplate = DecoratorTemplates + DecoratorIndex;
+				const FDecorator* Decorator = DecoratorRegistry.Find(DecoratorTemplate->GetRegistryHandle());
+				const FString DecoratorName = Decorator != nullptr ? Decorator->GetDecoratorName() : TEXT("<Unknown>");
+
+				const uint32 NextDecoratorIndex = DecoratorIndex + 1;
+				const uint32 EndOfNextDecoratorSharedData = NextDecoratorIndex < NumDecorators ? DecoratorTemplates[NextDecoratorIndex].GetNodeSharedOffset() : NodeTemplate->GetNodeSharedDataSize();
+				const uint32 DecoratorSharedDataSize = EndOfNextDecoratorSharedData - DecoratorTemplate->GetNodeSharedOffset();
+
+				const uint32 EndOfNextDecoratorInstanceData = NextDecoratorIndex < NumDecorators ? DecoratorTemplates[NextDecoratorIndex].GetNodeInstanceOffset() : NodeTemplate->GetNodeInstanceDataSize();
+				const uint32 DecoratorInstanceDataSize = EndOfNextDecoratorInstanceData - DecoratorTemplate->GetNodeInstanceOffset();
+
+				UE_LOG(LogAnimation, Log, TEXT("            %u: [%x] %s (%s)"), DecoratorIndex, DecoratorTemplate->GetUID().GetUID(), *DecoratorName, DecoratorTemplate->GetMode() == EDecoratorMode::Base ? TEXT("Base") : TEXT("Additive"));
+				UE_LOG(LogAnimation, Log, TEXT("                Shared Data: [Offset: %u bytes, Size: %u bytes]"), DecoratorTemplate->GetNodeSharedOffset(), DecoratorSharedDataSize);
+				if (DecoratorTemplate->HasLatentProperties() && Decorator != nullptr)
+				{
+					UE_LOG(LogAnimation, Log, TEXT("                Shared Data Latent Property Handles: [Offset: %u bytes, Count: %u]"), DecoratorTemplate->GetNodeSharedLatentPropertyHandlesOffset(), Decorator->GetNumLatentDecoratorProperties());
+				}
+				UE_LOG(LogAnimation, Log, TEXT("                Instance Data: [Offset: %u bytes, Size: %u bytes]"), DecoratorTemplate->GetNodeInstanceOffset(), DecoratorInstanceDataSize);
+			}
+		}
+
+		LogAnimation.SetVerbosity(OldVerbosity);
+	}
+
+	void ListAnimGraphs(const TArray<FString>& Args)
+	{
+		// Turn off log times to make diff-ing easier
+		TGuardValue<ELogTimes::Type> DisableLogTimes(GPrintLogTimes, ELogTimes::None);
+
+		// Make sure to log everything
+		const ELogVerbosity::Type OldVerbosity = LogAnimation.GetVerbosity();
+		LogAnimation.SetVerbosity(ELogVerbosity::All);
+
+		TArray<const UAnimNextGraph*> AnimGraphs;
+
+		for (TObjectIterator<UAnimNextGraph> It; It; ++It)
+		{
+			AnimGraphs.Add(*It);
+		}
+
+		struct FCompareObjectNames
+		{
+			FORCEINLINE bool operator()(const UAnimNextGraph& Lhs, const UAnimNextGraph& Rhs) const
+			{
+				return Lhs.GetPathName().Compare(Rhs.GetPathName()) < 0;
+			}
+		};
+		AnimGraphs.Sort(FCompareObjectNames());
+
+		const FNodeTemplateRegistry& NodeTemplateRegistry = FNodeTemplateRegistry::Get();
+		const FDecoratorRegistry& DecoratorRegistry = FDecoratorRegistry::Get();
+		const bool bDetailedOutput = true;
+
+		UE_LOG(LogAnimation, Log, TEXT("===== AnimNext Animation Graphs ====="));
+		UE_LOG(LogAnimation, Log, TEXT("Num Graphs: %u"), AnimGraphs.Num());
+
+		for (const UAnimNextGraph* AnimGraph : AnimGraphs)
+		{
+			uint32 TotalInstanceSize = 0;
+			uint32 NumNodes = 0;
+			{
+				// We always have a node at offset 0
+				int32 NodeOffset = 0;
+
+				while (NodeOffset < AnimGraph->SharedDataBuffer.Num())
+				{
+					const FNodeDescription* NodeDesc = reinterpret_cast<const FNodeDescription*>(&AnimGraph->SharedDataBuffer[NodeOffset]);
+
+					TotalInstanceSize += NodeDesc->GetNodeInstanceDataSize();
+					NumNodes++;
+
+					const FNodeTemplate* NodeTemplate = NodeTemplateRegistry.Find(NodeDesc->GetTemplateHandle());
+					NodeOffset += NodeTemplate->GetNodeSharedDataSize();
+				}
+			}
+
+			UE_LOG(LogAnimation, Log, TEXT("    %s ..."), *AnimGraph->GetPathName());
+			UE_LOG(LogAnimation, Log, TEXT("        Shared Data Size: %.2f KB"), double(AnimGraph->SharedDataBuffer.Num()) / 1024.0);
+			UE_LOG(LogAnimation, Log, TEXT("        Max Instance Data Size: %.2f KB"), double(TotalInstanceSize) / 1024.0);
+			UE_LOG(LogAnimation, Log, TEXT("        Num Nodes: %u"), NumNodes);
+
+			if (bDetailedOutput)
+			{
+				// We always have a node at offset 0
+				int32 NodeOffset = 0;
+
+				while (NodeOffset < AnimGraph->SharedDataBuffer.Num())
+				{
+					const FNodeDescription* NodeDesc = reinterpret_cast<const FNodeDescription*>(&AnimGraph->SharedDataBuffer[NodeOffset]);
+					const FNodeTemplate* NodeTemplate = NodeTemplateRegistry.Find(NodeDesc->GetTemplateHandle());
+
+					const uint32 NumDecorators = NodeTemplate->GetNumDecorators();
+
+					UE_LOG(LogAnimation, Log, TEXT("        Node %u: [Template %x with %u decorators]"), NodeDesc->GetUID().GetNodeIndex(), NodeTemplate->GetUID(), NumDecorators);
+					UE_LOG(LogAnimation, Log, TEXT("            Shared Data: [Offset: %u bytes, Size: %u bytes]"), NodeOffset, NodeTemplate->GetNodeSharedDataSize());
+					UE_LOG(LogAnimation, Log, TEXT("            Instance Data Size: %u bytes"), NodeDesc->GetNodeInstanceDataSize());
+					UE_LOG(LogAnimation, Log, TEXT("            Decorators ..."));
+
+					const FDecoratorTemplate* DecoratorTemplates = NodeTemplate->GetDecorators();
+					for (uint32 DecoratorIndex = 0; DecoratorIndex < NumDecorators; ++DecoratorIndex)
+					{
+						const FDecoratorTemplate* DecoratorTemplate = DecoratorTemplates + DecoratorIndex;
+						const FDecorator* Decorator = DecoratorRegistry.Find(DecoratorTemplate->GetRegistryHandle());
+						const FString DecoratorName = Decorator != nullptr ? Decorator->GetDecoratorName() : TEXT("<Unknown>");
+
+						const uint32 NextDecoratorIndex = DecoratorIndex + 1;
+						const uint32 EndOfNextDecoratorSharedData = NextDecoratorIndex < NumDecorators ? DecoratorTemplates[NextDecoratorIndex].GetNodeSharedOffset() : NodeTemplate->GetNodeSharedDataSize();
+						const uint32 DecoratorSharedDataSize = EndOfNextDecoratorSharedData - DecoratorTemplate->GetNodeSharedOffset();
+
+						const uint32 EndOfNextDecoratorInstanceData = NextDecoratorIndex < NumDecorators ? DecoratorTemplates[NextDecoratorIndex].GetNodeInstanceOffset() : NodeTemplate->GetNodeInstanceDataSize();
+						const uint32 DecoratorInstanceDataSize = EndOfNextDecoratorInstanceData - DecoratorTemplate->GetNodeInstanceOffset();
+
+						UE_LOG(LogAnimation, Log, TEXT("                    %u: [%x] %s (%s)"), DecoratorIndex, DecoratorTemplate->GetUID().GetUID(), *DecoratorName, DecoratorTemplate->GetMode() == EDecoratorMode::Base ? TEXT("Base") : TEXT("Additive"));
+						UE_LOG(LogAnimation, Log, TEXT("                        Shared Data: [Offset: %u bytes, Size: %u bytes]"), DecoratorTemplate->GetNodeSharedOffset(), DecoratorSharedDataSize);
+						if (DecoratorTemplate->HasLatentProperties() && Decorator != nullptr)
+						{
+							UE_LOG(LogAnimation, Log, TEXT("                        Shared Data Latent Property Handles: [Offset: %u bytes, Count: %u]"), DecoratorTemplate->GetNodeSharedLatentPropertyHandlesOffset(), Decorator->GetNumLatentDecoratorProperties());
+						}
+						UE_LOG(LogAnimation, Log, TEXT("                        Instance Data: [Offset: %u bytes, Size: %u bytes]"), DecoratorTemplate->GetNodeInstanceOffset(), DecoratorInstanceDataSize);
+					}
+
+					NodeOffset += NodeTemplate->GetNodeSharedDataSize();
+				}
+			}
+		}
+
+		LogAnimation.SetVerbosity(OldVerbosity);
 	}
+#endif
 };
 
 }
diff --git a/Engine/Plugins/Experimental/AnimNext/Source/AnimNext/Public/DecoratorBase/Decorator.h b/Engine/Plugins/Experimental/AnimNext/Source/AnimNext/Public/DecoratorBase/Decorator.h
--- a/Engine/Plugins/Experimental/AnimNext/Source/AnimNext/Public/DecoratorBase/Decorator.h
+++ b/Engine/Plugins/Experimental/AnimNext/Source/AnimNext/Public/DecoratorBase/Decorator.h
@@ -25,14 +25,16 @@ class FArchive;
 	/* FDecorator impl */ \
 	static constexpr UE::AnimNext::FDecoratorUID DecoratorUID = UE::AnimNext::FDecoratorUID(DecoratorNameHash, TEXT(#DecoratorName)); \
 	virtual UE::AnimNext::FDecoratorUID GetDecoratorUID() const override { return DecoratorUID; } \
-	virtual FString GetDecoratorName() const { return TEXT(#DecoratorName); } \
+	virtual FString GetDecoratorName() const override { return TEXT(#DecoratorName); } \
 	static const UE::AnimNext::FDecoratorMemoryLayout DecoratorMemoryDescription; \
 	virtual UE::AnimNext::FDecoratorMemoryLayout GetDecoratorMemoryDescription() const override { return DecoratorMemoryDescription; } \
 	virtual UScriptStruct* GetDecoratorSharedDataStruct() const override { return FSharedData::StaticStruct(); } \
 	virtual void ConstructDecoratorInstance(const UE::AnimNext::FExecutionContext& Context, const UE::AnimNext::FDecoratorBinding& Binding) const override; \
 	virtual void DestructDecoratorInstance(const UE::AnimNext::FExecutionContext& Context, const UE::AnimNext::FDecoratorBinding& Binding) const override; \
 	virtual const UE::AnimNext::IDecoratorInterface* GetDecoratorInterface(UE::AnimNext::FDecoratorInterfaceUID InterfaceUID) const override; \
+	virtual TConstArrayView<FDecoratorInterfaceUID> GetDecoratorInterfaces() const override; \
 	virtual uint32 GetNumLatentDecoratorProperties() const override { return -FSharedData::GetLatentPropertyIndex(~(size_t)0); } \
+	virtual FDecoratorLatentPropertyMemoryLayout GetLatentPropertyMemoryLayout(FName PropertyName, uint32 PropertyIndex) const override; \
 	static_assert(std::is_base_of<FAnimNextDecoratorSharedData, FSharedData>::value, "Decorator shared data must derive from FAnimNextDecoratorSharedData"); \
 	static_assert(std::is_base_of<FDecoratorInstanceData, FInstanceData>::value, "Decorator instance data must derive from FDecoratorInstanceData");
 
@@ -41,20 +43,28 @@ class FArchive;
 	/* FDecorator impl */ \
 	static constexpr UE::AnimNext::FDecoratorUID DecoratorUID = UE::AnimNext::FDecoratorUID(DecoratorNameHash, TEXT(#DecoratorName)); \
 	virtual UE::AnimNext::FDecoratorUID GetDecoratorUID() const override { return DecoratorUID; } \
-	virtual FString GetDecoratorName() const { return TEXT(#DecoratorName); }
+	virtual FString GetDecoratorName() const override { return TEXT(#DecoratorName); }
 
 // In the decorator cpp, these three macros implement the base functionality
 // 
 // Usage is as follow:
-// DEFINE_ANIM_DECORATOR_BEGIN(FSequencePlayerDecorator)
-//     DEFINE_ANIM_DECORATOR_IMPLEMENTS_INTERFACE(IEvaluate)
-//     DEFINE_ANIM_DECORATOR_IMPLEMENTS_INTERFACE(IUpdate)
-//     DEFINE_ANIM_DECORATOR_IMPLEMENTS_INTERFACE(ITimeline)
-// DEFINE_ANIM_DECORATOR_END(FSequencePlayerDecorator)
-
-#define DEFINE_ANIM_DECORATOR_BEGIN(DecoratorName) \
+// #define DECORATOR_INTERFACE_ENUMERATOR(GeneratorMacro) \
+//		GeneratorMacro(IHierarchy) \
+//		GeneratorMacro(IUpdate) \
+//
+// GENERATE_ANIM_DECORATOR_IMPLEMENTATION(FMyDecorator, DECORATOR_INTERFACE_ENUMERATOR)
+// #undef DECORATOR_INTERFACE_ENUMERATOR
+
+// Implements various parts of FDecorator
+#define ANIM_NEXT_IMPL_DEFINE_ANIM_DECORATOR(DecoratorName) \
 	const UE::AnimNext::FDecoratorMemoryLayout DecoratorName::DecoratorMemoryDescription = \
 		UE::AnimNext::FDecoratorMemoryLayout{ sizeof(DecoratorName), alignof(DecoratorName), sizeof(DecoratorName::FSharedData), alignof(DecoratorName::FSharedData), sizeof(DecoratorName::FInstanceData), alignof(DecoratorName::FInstanceData) }; \
+	FDecoratorLatentPropertyMemoryLayout DecoratorName::GetLatentPropertyMemoryLayout(FName PropertyName, uint32 PropertyIndex) const \
+	{ \
+		/* Thread safe cache initialization */ \
+		static TArray<FDecoratorLatentPropertyMemoryLayout> CachedLatentPropertyMemoryLayouts = [this](){ TArray<FDecoratorLatentPropertyMemoryLayout> Result; Result.SetNum(GetNumLatentDecoratorProperties()); return Result; }(); \
+		return GetLatentPropertyMemoryLayoutImpl(PropertyName, PropertyIndex, CachedLatentPropertyMemoryLayouts); \
+	} \
 	void DecoratorName::ConstructDecoratorInstance(const UE::AnimNext::FExecutionContext& Context, const UE::AnimNext::FDecoratorBinding& Binding) const \
 	{ \
 		FInstanceData* Data = new(Binding.GetInstanceData<FInstanceData>()) FInstanceData(); \
@@ -65,21 +75,48 @@ class FArchive;
 		FInstanceData* Data = Binding.GetInstanceData<FInstanceData>(); \
 		Data->Destruct(Context, Binding); \
 		Data->~FInstanceData(); \
-	} \
-	const UE::AnimNext::IDecoratorInterface* DecoratorName::GetDecoratorInterface(UE::AnimNext::FDecoratorInterfaceUID InterfaceUID_) const \
-	{
+	}
 
-#define DEFINE_ANIM_DECORATOR_END(DecoratorName) \
+// Helper that handles the GetDecoratorInterface() details for each interface specified by the generator macro
+#define ANIM_NEXT_IMPL_GET_INTERFACE_IMPL_FOR_INTERFACE(InterfaceName) \
+	if (InInterfaceUID == InterfaceName::InterfaceUID) \
+	{ \
+		return static_cast<const InterfaceName*>(this); \
+	}
+
+// Implements GetDecoratorInterface()
+#define ANIM_NEXT_IMPL_DEFINE_ANIM_DECORATOR_GET_INTERFACE(DecoratorName, InterfaceEnumeratorMacro) \
+	const UE::AnimNext::IDecoratorInterface* DecoratorName::GetDecoratorInterface(UE::AnimNext::FDecoratorInterfaceUID InInterfaceUID) const \
+	{ \
+		InterfaceEnumeratorMacro(ANIM_NEXT_IMPL_GET_INTERFACE_IMPL_FOR_INTERFACE) \
 		/* Forward to base implementation */ \
-		return DecoratorSuper::GetDecoratorInterface(InterfaceUID_); \
+		return DecoratorSuper::GetDecoratorInterface(InInterfaceUID); \
 	}
 
-#define DEFINE_ANIM_DECORATOR_IMPLEMENTS_INTERFACE(InterfaceName) \
-	if (InterfaceUID_ == InterfaceName::InterfaceUID) \
+// Helper that handles the GetDecoratorInterfaces() details for each interface specified by the generator macro
+#define ANIM_NEXT_IMPL_GET_INTERFACES_IMPL_FOR_INTERFACE(InterfaceName) InterfaceName::InterfaceUID,
+
+// Implements GetDecoratorInterfaces()
+#define ANIM_NEXT_IMPL_DEFINE_ANIM_DECORATOR_GET_INTERFACES(DecoratorName, InterfaceEnumeratorMacro) \
+	TConstArrayView<UE::AnimNext::FDecoratorInterfaceUID> DecoratorName::GetDecoratorInterfaces() const \
 	{ \
-		return static_cast<const InterfaceName*>(this); \
+		/* Thread safe cache initialization */ \
+		static TArray<UE::AnimNext::FDecoratorInterfaceUID> CachedInterfaceList = FDecorator::BuildDecoratorInterfaceList( \
+		DecoratorSuper::GetDecoratorInterfaces(), \
+		{ \
+			InterfaceEnumeratorMacro(ANIM_NEXT_IMPL_GET_INTERFACES_IMPL_FOR_INTERFACE) \
+		}); \
+		return CachedInterfaceList; \
 	}
 
+/**
+  * This macro defines the necessary boilerplate for implementing FDecorator. See above for usage example.
+  */
+#define GENERATE_ANIM_DECORATOR_IMPLEMENTATION(DecoratorName, InterfaceEnumeratorMacro) \
+	ANIM_NEXT_IMPL_DEFINE_ANIM_DECORATOR(DecoratorName) \
+	ANIM_NEXT_IMPL_DEFINE_ANIM_DECORATOR_GET_INTERFACE(DecoratorName, InterfaceEnumeratorMacro) \
+	ANIM_NEXT_IMPL_DEFINE_ANIM_DECORATOR_GET_INTERFACES(DecoratorName, InterfaceEnumeratorMacro)
+
 // Allows a decorator to auto-register and unregister within the current execution scope
 // The decorator must be found in the current scope without a namespace qualification
 #define AUTO_REGISTER_ANIM_DECORATOR(DecoratorName) \
@@ -115,19 +152,33 @@ namespace UE::AnimNext
 		uint32 DecoratorSize = 0;
 
 		// The alignment in bytes of an instance of the decorator class which derives from FDecorator
-		uint32 DecoratorAlignment = 0;
+		uint32 DecoratorAlignment = 1;
 
 		// The size in bytes of the shared data for the decorator which derives from FAnimNextDecoratorSharedData
 		uint32 SharedDataSize = 0;
 
 		// The alignment in bytes of the shared data for the decorator which derives from FAnimNextDecoratorSharedData
-		uint32 SharedDataAlignment = 0;
+		uint32 SharedDataAlignment = 1;
 
 		// The size in bytes of the instance data for the decorator which derives from FDecoratorInstanceData
 		uint32 InstanceDataSize = 0;
 
 		// The alignment in bytes of the instance data for the decorator which derives from FDecoratorInstanceData
-		uint32 InstanceDataAlignment = 0;
+		uint32 InstanceDataAlignment = 1;
+	};
+
+	/**
+	 * FDecoratorLatentPropertyMemoryLayout
+	 *
+	 * Encapsulates size/alignment details for a latent property.
+	 */
+	struct FDecoratorLatentPropertyMemoryLayout
+	{
+		// The size in bytes of the latent property
+		uint32 Size = 0;
+
+		// The alignment in bytes of the latent property
+		uint32 Alignment = 1;
 	};
 
 	/**
@@ -208,9 +259,15 @@ namespace UE::AnimNext
 			return nullptr;
 		}
 
+		// Returns a list of interfaces that this decorator supports
+		virtual TConstArrayView<FDecoratorInterfaceUID> GetDecoratorInterfaces() const { return TConstArrayView<FDecoratorInterfaceUID>(); }
+
 		// The number of latent property properties in the shared data of this decorator
 		virtual uint32 GetNumLatentDecoratorProperties() const { return 0; }
 
+		// Returns the memory layout of the specified latent property
+		virtual FDecoratorLatentPropertyMemoryLayout GetLatentPropertyMemoryLayout(FName PropertyName, uint32 PropertyIndex) const { return FDecoratorLatentPropertyMemoryLayout(); }
+
 		// Called to serialize decorator shared data
 		virtual void SerializeDecoratorSharedData(FArchive& Ar, FAnimNextDecoratorSharedData& SharedData) const;
 
@@ -219,11 +276,25 @@ namespace UE::AnimNext
 		// derived type using UE reflection.
 		// Decorators can override this function to control how editor only properties are coerced into the runtime shared data
 		// instance.
-		virtual void SaveDecoratorSharedData(const TFunction<FString(const FString& PropertyName)>& GetDecoratorProperty, FAnimNextDecoratorSharedData& OutSharedData) const;
+		virtual void SaveDecoratorSharedData(const TFunction<FString(FName PropertyName)>& GetDecoratorProperty, FAnimNextDecoratorSharedData& OutSharedData) const;
 
-		// Takes the editor properties as authored in the graph and returns the latent property handles using UE reflection.
-		TArray<FLatentPropertyHandle> GetLatentPropertyHandles(bool bFilterEditorOnly, const TFunction<bool(const FString& PropertyName)>& IsDecoratorPropertyLatent, FLatentPropertyHandle& CurrentLatentPropertyHandle) const;
+		// Takes the editor properties as authored in the graph and returns the latent property metadata using UE reflection.
+		TArray<FLatentPropertyMetadata> GetLatentPropertyHandles(
+			bool bFilterEditorOnly,
+			const TFunction<uint16(FName PropertyName)>& GetDecoratorLatentPropertyIndex) const;
 #endif
+
+	protected:
+		// Implements GetLatentPropertyMemoryLayout() by allowing a map to be provided for caching purposes to speed up repeated queries
+		FDecoratorLatentPropertyMemoryLayout GetLatentPropertyMemoryLayoutImpl(
+			FName PropertyName,
+			uint32 PropertyIndex,
+			TArray<FDecoratorLatentPropertyMemoryLayout>& LatentPropertyMemoryLayouts) const;
+
+		// Builds a list of interfaces with the provided super interfaces and current interfaces as an initializer list
+		static TArray<FDecoratorInterfaceUID> BuildDecoratorInterfaceList(
+			const TConstArrayView<FDecoratorInterfaceUID>& SuperInterfaces,
+			std::initializer_list<FDecoratorInterfaceUID> InterfaceList);
 	};
 
 	// Base class for base decorators that are standalone
diff --git a/Engine/Plugins/Experimental/AnimNext/Source/AnimNext/Public/DecoratorBase/DecoratorBinding.h b/Engine/Plugins/Experimental/AnimNext/Source/AnimNext/Public/DecoratorBase/DecoratorBinding.h
--- a/Engine/Plugins/Experimental/AnimNext/Source/AnimNext/Public/DecoratorBase/DecoratorBinding.h
+++ b/Engine/Plugins/Experimental/AnimNext/Source/AnimNext/Public/DecoratorBase/DecoratorBinding.h
@@ -73,6 +73,24 @@ namespace UE::AnimNext
 			return DecoratorTemplate->GetDecoratorLatentPropertyHandles(*NodeDescription);
 		}
 
+		// Returns a pointer to the latent property specified by the provided handle or nullptr if the binding/handle are invalid
+		template<typename PropertyType>
+		const PropertyType* GetLatentProperty(FLatentPropertyHandle Handle) const
+		{
+			if (!IsValid())
+			{
+				return nullptr;
+			}
+
+			if (!Handle.IsOffsetValid())
+			{
+				return nullptr;
+			}
+
+			const uint8* NodeInstance = (const uint8*)DecoratorPtr.GetNodeInstance();
+			return (const PropertyType*)(NodeInstance + Handle.GetLatentPropertyOffset());
+		}
+
 		// Returns the decorator pointer we are bound to.
 		FWeakDecoratorPtr GetDecoratorPtr() const { return DecoratorPtr; }
 
diff --git a/Engine/Plugins/Experimental/AnimNext/Source/AnimNext/Public/DecoratorBase/DecoratorInterfaceUID.h b/Engine/Plugins/Experimental/AnimNext/Source/AnimNext/Public/DecoratorBase/DecoratorInterfaceUID.h
--- a/Engine/Plugins/Experimental/AnimNext/Source/AnimNext/Public/DecoratorBase/DecoratorInterfaceUID.h
+++ b/Engine/Plugins/Experimental/AnimNext/Source/AnimNext/Public/DecoratorBase/DecoratorInterfaceUID.h
@@ -61,4 +61,7 @@ namespace UE::AnimNext
 	// Compares for equality and inequality
 	constexpr bool operator==(FDecoratorInterfaceUID LHS, FDecoratorInterfaceUID RHS) { return LHS.GetUID() == RHS.GetUID(); }
 	constexpr bool operator!=(FDecoratorInterfaceUID LHS, FDecoratorInterfaceUID RHS) { return LHS.GetUID() != RHS.GetUID(); }
+
+	// For sorting
+	constexpr bool operator<(FDecoratorInterfaceUID LHS, FDecoratorInterfaceUID RHS) { return LHS.GetUID() < RHS.GetUID(); }
 }
diff --git a/Engine/Plugins/Experimental/AnimNext/Source/AnimNext/Public/DecoratorBase/DecoratorSharedData.h b/Engine/Plugins/Experimental/AnimNext/Source/AnimNext/Public/DecoratorBase/DecoratorSharedData.h
--- a/Engine/Plugins/Experimental/AnimNext/Source/AnimNext/Public/DecoratorBase/DecoratorSharedData.h
+++ b/Engine/Plugins/Experimental/AnimNext/Source/AnimNext/Public/DecoratorBase/DecoratorSharedData.h
@@ -47,14 +47,14 @@
 		static_assert(LatentPropertyIndex > 0, "Property " #PropertyName " isn't latent"); \
 		const UE::AnimNext::FLatentPropertyHandle* LatentPropertyHandles = Binding.GetLatentPropertyHandles(); \
 		const UE::AnimNext::FLatentPropertyHandle LatentPropertyHandle = LatentPropertyHandles[LatentPropertyIndex - 1]; \
-		/* An invalid handle means we are inline, otherwise we are latent */ \
-		if (!LatentPropertyHandle.IsValid()) \
+		/* An invalid offset means we are inline, otherwise we are cached */ \
+		if (!LatentPropertyHandle.IsOffsetValid()) \
 		{ \
 			return PropertyName; \
 		} \
 		else \
 		{ \
-			return Context.EvaluateLatentPin<decltype(PropertyName)>(LatentPropertyHandle); \
+			return *Binding.GetLatentProperty<decltype(PropertyName)>(LatentPropertyHandle); \
 		} \
 	} \
 
diff --git a/Engine/Plugins/Experimental/AnimNext/Source/AnimNext/Public/DecoratorBase/ExecutionContext.h b/Engine/Plugins/Experimental/AnimNext/Source/AnimNext/Public/DecoratorBase/ExecutionContext.h
--- a/Engine/Plugins/Experimental/AnimNext/Source/AnimNext/Public/DecoratorBase/ExecutionContext.h
+++ b/Engine/Plugins/Experimental/AnimNext/Source/AnimNext/Public/DecoratorBase/ExecutionContext.h
@@ -90,9 +90,10 @@ namespace UE::AnimNext
 		// there are no more references remaining, reseting the pointer in the process
 		void ReleaseNodeInstance(FDecoratorPtr& NodePtr) const;
 
-		// Evaluates the latent pin with the specified handle
-		template<typename LatentPinType>
-		LatentPinType EvaluateLatentPin(FLatentPropertyHandle LatentPropertyHandle) const;
+		// Takes a snapshot of all latent properties on the provided node sub-stack (all decorators on the sub-stack of the provided one)
+		// Properties can be marked as always updating or as supporting freezing (e.g. when a branch of the graph blends out)
+		// A freezable property does not update when a snapshot is taken of a frozen node
+		void SnapshotLatentProperties(const FWeakDecoratorPtr& DecoratorPtr, bool bIsFrozen) const;
 
 		// Returns a typed graph instance component, creating it lazily the first time it is queried
 		template<class ComponentType>
@@ -115,7 +116,6 @@ namespace UE::AnimNext
 
 		bool GetInterfaceImpl(FDecoratorInterfaceUID InterfaceUID, const FWeakDecoratorPtr& DecoratorPtr, FDecoratorBinding& InterfaceBinding) const;
 		bool GetInterfaceSuperImpl(FDecoratorInterfaceUID InterfaceUID, const FWeakDecoratorPtr& DecoratorPtr, FDecoratorBinding& SuperBinding) const;
-		void EvaluateLatentPinImpl(FLatentPropertyHandle LatentPropertyHandle, void* DestinationPtr) const;
 		FGraphInstanceComponent* TryGetComponent(int32 ComponentNameHash, FName ComponentName) const;
 		FGraphInstanceComponent& AddComponent(int32 ComponentNameHash, FName ComponentName, TSharedPtr<FGraphInstanceComponent>&& Component) const;
 
@@ -166,20 +166,6 @@ namespace UE::AnimNext
 		return AllocateNodeInstance(ParentBinding.GetDecoratorPtr(), ChildDecoratorHandle);
 	}
 
-	template<typename LatentPinType>
-	inline LatentPinType FExecutionContext::EvaluateLatentPin(FLatentPropertyHandle LatentPropertyHandle) const
-	{
-		LatentPinType Result{};
-
-		// Latent pin handle needs to be valid
-		if (ensure(LatentPropertyHandle.IsValid()))
-		{
-			EvaluateLatentPinImpl(LatentPropertyHandle, &Result);
-		}
-
-		return Result;
-	}
-
 	template<class ComponentType>
 	ComponentType& FExecutionContext::GetComponent() const
 	{
diff --git a/Engine/Plugins/Experimental/AnimNext/Source/AnimNext/Public/DecoratorBase/ExecutionContextProxy.h b/Engine/Plugins/Experimental/AnimNext/Source/AnimNext/Public/DecoratorBase/ExecutionContextProxy.h
--- a/Engine/Plugins/Experimental/AnimNext/Source/AnimNext/Public/DecoratorBase/ExecutionContextProxy.h
+++ b/Engine/Plugins/Experimental/AnimNext/Source/AnimNext/Public/DecoratorBase/ExecutionContextProxy.h
@@ -96,11 +96,12 @@ namespace UE::AnimNext
 			return ExecutionContext.ReleaseNodeInstance(NodePtr);
 		}
 
-		// Evaluates the latent pin with the specified handle
-		template<typename LatentPinType>
-		LatentPinType EvaluateLatentPin(FLatentPropertyHandle LatentPropertyHandle) const
+		// Takes a snapshot of all latent properties on the provided node sub-stack (all decorators on the sub-stack of the provided one)
+		// Properties can be marked as always updating or as supporting freezing (e.g. when a branch of the graph blends out)
+		// A freezable property does not update when a snapshot is taken of a frozen node
+		void SnapshotLatentProperties(const FWeakDecoratorPtr& DecoratorPtr, bool bIsFrozen) const
 		{
-			return ExecutionContext.EvaluateLatentPin<LatentPinType>(LatentPropertyHandle);
+			ExecutionContext.SnapshotLatentProperties(DecoratorPtr, bIsFrozen);
 		}
 
 		// Returns a typed graph instance component, creating it lazily the first time it is queried
diff --git a/Engine/Plugins/Experimental/AnimNext/Source/AnimNext/Public/DecoratorBase/LatentPropertyHandle.h b/Engine/Plugins/Experimental/AnimNext/Source/AnimNext/Public/DecoratorBase/LatentPropertyHandle.h
--- a/Engine/Plugins/Experimental/AnimNext/Source/AnimNext/Public/DecoratorBase/LatentPropertyHandle.h
+++ b/Engine/Plugins/Experimental/AnimNext/Source/AnimNext/Public/DecoratorBase/LatentPropertyHandle.h
@@ -7,51 +7,133 @@
 
 namespace UE::AnimNext
 {
+	class FDecoratorReader;
+
+	/**
+	 * Latent Property Metadata
+	 * 
+	 * The metadata we serialize in the archive for each latent property.
+	 */
+	struct FLatentPropertyMetadata final
+	{
+		// The name of the latent property
+		FName Name;
+
+		// The index of the RigVM memory handle to evaluate to compute our latent value
+		// An invalid index of ~0 means that the latent property has a hardcoded inline value
+		uint16 RigVMIndex = MAX_uint16;
+
+		// Whether or not the property supports freezing
+		bool bCanFreeze = false;
+	};
+
+	// Serializes latent property metadata
+	inline FArchive& operator<<(FArchive& Ar, FLatentPropertyMetadata& Metadata)
+	{
+		Ar << Metadata.Name;
+		Ar << Metadata.RigVMIndex;
+		Ar << Metadata.bCanFreeze;
+		return Ar;
+	}
+
+	/**
+	 * Latent Properties Header
+	 * Holds metadata about the stored latent properties
+	 *
+	 * @see FExecutionContext
+	 */
+	struct alignas(uint16) FLatentPropertiesHeader final
+	{
+		// Whether or not this base decorator has any bound latent properties that can evaluate
+		uint8 bHasValidLatentProperties = false;
+
+		// Whether or not this base decorator has all its latent properties support freezing
+		uint8 bCanAllPropertiesFreeze = false;
+	};
+
 	/**
 	 * Latent Property Handle
-	 * Latent property handles represent a 1-based index into the RigVM memory handles array during execution.
+	 * A latent property handle represents:
+	 *    - A 1-based index into the RigVM memory handles array during execution
+	 *    - A flag that specifies if the latent property should always update or if it supports freezing
+	 *    - A node instance local offset to the cached latent property value
+	 * 
+	 * A latent property can be in one of 3 states:
+	 *    - Invalid (property is inline and not latent)
+	 *    - Latent and Cached (property is latent and is cached at some offset)
+	 *    - Cached Only (property is not latent but is cached at some offset)
+	 * 
+	 * The Cached Only state can occur if a base decorator has duplicate entries within its property handles.
+	 * Duplicate entries would have the same RigVM index but different offsets and thus evaluate it repeatedly.
+	 * To avoid the waste of cycles and space, subsequent properties are marked as Cached Only with the same
+	 * offset as the first property that is Latent and Cached (it will update the cached location).
 	 *
 	 * @see FExecutionContext
 	 */
 	struct FLatentPropertyHandle final
 	{
 		// Creates an invalid latent property handle
-		constexpr FLatentPropertyHandle() noexcept
-			: HandleValue(INVALID_HANDLE_VALUE)
-		{}
+		constexpr FLatentPropertyHandle() noexcept = default;
+
+		// Returns true if this latent property handle has a valid RigVM index, false otherwise
+		constexpr bool IsIndexValid() const noexcept { return (PackedIndexCanFreezeValue & HANDLE_INDEX_MASK) != INVALID_HANDLE_VALUE; }
 
-		// Returns true if this latent property handle is valid, false otherwise
-		constexpr bool IsValid() const noexcept { return HandleValue != INVALID_HANDLE_VALUE; }
+		// Returns true if this latent property handle has a valid property offset, false otherwise
+		constexpr bool IsOffsetValid() const noexcept { return OffsetValue != INVALID_OFFSET_VALUE; }
 
-		// Returns the latent property index represented by this handle if valid, -1 otherwise
-		constexpr int32 GetLatentPropertyIndex() const noexcept { return IsValid() ? (HandleValue - 1) : INDEX_NONE; }
+		// Returns the latent property RigVM index represented by this handle if valid, -1 otherwise
+		constexpr int32 GetLatentPropertyIndex() const noexcept { return IsIndexValid() ? (PackedIndexCanFreezeValue & HANDLE_INDEX_MASK) : INDEX_NONE; }
 
-		// The first valid latent handle
-		static constexpr FLatentPropertyHandle GetFirstHandle() noexcept { return FLatentPropertyHandle(1); }
+		// Returns the latent property offset represented by this handle if valid, 0 otherwise
+		// This offset is relative to the start of the node instance data
+		constexpr uint32 GetLatentPropertyOffset() const noexcept { return OffsetValue; }
 
-		// Returns the next latent handle
-		// If the current handle is invalid, the next is invalid as well
-		// Note that this function does not handle wrapping if more than 65536 handles are used
-		// If wrapping occurs, the next handle returned will be invalid and all subsequent handles will be invalid
-		constexpr FLatentPropertyHandle GetNextHandle() const noexcept { return IsValid() ? FLatentPropertyHandle(HandleValue + 1) : FLatentPropertyHandle(); }
+		// Returns whether or not the latent property supports freezing during snapshots
+		// Latent properties that do not support freezing always update in a snapshot
+		constexpr bool CanFreeze() const noexcept { return IsIndexValid() ? !!(PackedIndexCanFreezeValue & HANDLE_CAN_FREEZE_MASK) : false; }
 
 	private:
-		// Constructs a FLatentPropertyHandle instance from a raw value
-		explicit constexpr FLatentPropertyHandle(uint16 InHandleValue) noexcept : HandleValue(InHandleValue) {}
+		// Constructs a FLatentPropertyHandle instance
+		FLatentPropertyHandle(uint16 InRigRMIndex, uint32 InPropertyOffset, bool bInCanFreeze)
+			: PackedIndexCanFreezeValue((InRigRMIndex & HANDLE_INDEX_MASK) | (uint16(bInCanFreeze) << HANDLE_CAN_FREEZE_SHIFT_OFFSET))
+			, OffsetValue(uint16(InPropertyOffset))
+		{
+			check(InRigRMIndex <= INVALID_HANDLE_VALUE || InRigRMIndex == MAX_uint16);	// Must fit on 15 bits when valid
+			check(InPropertyOffset <= MAX_uint16);	// Must fit on 16 bits
+		}
 
-		// Latent property handles are a 1-based index, 0 is invalid
-		static constexpr uint16 INVALID_HANDLE_VALUE = 0;
+		// A mask for the RigVM index portion of the packed value
+		static constexpr uint16 HANDLE_INDEX_MASK = 0x7FFF;
 
-		// The raw 1-based latent property index or 0 if invalid
-		uint16		HandleValue;
+		// A mask for the CanFreeze metadata flag in the packed value
+		static constexpr uint16 HANDLE_CAN_FREEZE_MASK = 0x8000;
+
+		// A shift offset for the CanFreeze metadata flag in the packed value
+		static constexpr uint16 HANDLE_CAN_FREEZE_SHIFT_OFFSET = 15;
+
+		// Latent property handles are a 15 bit index and a 1 bit flag to denote if the property can freeze, ~0 is invalid
+		static constexpr uint16 INVALID_HANDLE_VALUE = 0x7FFF;
+
+		// Latent property offsets are relative to the start of the node instance data, 0 is invalid
+		static constexpr uint16 INVALID_OFFSET_VALUE = 0;
+
+		// A packed RigVM index and metadata flag for whether or not this property can be frozen
+		uint16		PackedIndexCanFreezeValue = INVALID_HANDLE_VALUE;
+
+		// An offset relative to the start of the node instance where the latent property is cached
+		uint16		OffsetValue = INVALID_OFFSET_VALUE;
 
 		friend FArchive& operator<<(FArchive& Ar, FLatentPropertyHandle& Handle);
+		friend FDecoratorReader;
 	};
 
+	static_assert(alignof(FLatentPropertiesHeader) == alignof(FLatentPropertyHandle), "Header and handle must have matching alignment");
+
 	// Serializes a latent handle value
 	inline FArchive& operator<<(FArchive& Ar, FLatentPropertyHandle& Handle)
 	{
-		Ar << Handle.HandleValue;
+		Ar << Handle.PackedIndexCanFreezeValue;
+		Ar << Handle.OffsetValue;
 		return Ar;
 	}
 }
diff --git a/Engine/Plugins/Experimental/AnimNext/Source/AnimNext/Public/DecoratorBase/NodeDescription.h b/Engine/Plugins/Experimental/AnimNext/Source/AnimNext/Public/DecoratorBase/NodeDescription.h
--- a/Engine/Plugins/Experimental/AnimNext/Source/AnimNext/Public/DecoratorBase/NodeDescription.h
+++ b/Engine/Plugins/Experimental/AnimNext/Source/AnimNext/Public/DecoratorBase/NodeDescription.h
@@ -10,6 +10,8 @@ class FArchive;
 
 namespace UE::AnimNext
 {
+	class FDecoratorReader;
+
 	/**
 	 * Node Description
 	 * A node description represents a unique instance in the authored static graph.
@@ -32,6 +34,8 @@ namespace UE::AnimNext
 		FNodeDescription(FNodeID InNodeID, FNodeTemplateRegistryHandle InTemplateHandle)
 			: NodeID(InNodeID)
 			, TemplateHandle(InTemplateHandle)
+			, NodeInstanceDataSize(0)
+			, Padding0(0)
 		{}
 
 		// Returns the node UID, unique to the owning sub-graph
@@ -40,13 +44,45 @@ namespace UE::AnimNext
 		// Returns the handle of the node's template in the node template registry
 		FNodeTemplateRegistryHandle GetTemplateHandle() const { return TemplateHandle; }
 
+		// Returns the node instance data size factoring in any cached latent properties
+		uint32 GetNodeInstanceDataSize() const { return NodeInstanceDataSize; }
+
 		// Serializes this node description instance and the shared data of each decorator that follows
 		ANIMNEXT_API void Serialize(FArchive& Ar);
 
 	private:
-		FNodeID							NodeID;				// assigned during export/cook, unique to current sub-graph
-		FNodeTemplateRegistryHandle		TemplateHandle;		// offset of the node template within the global list
+		// Assigned during export/cook, unique to current sub-graph (16 bits)
+		FNodeID							NodeID;
+
+		// Offset of the node template within the global list (16 bits)
+		FNodeTemplateRegistryHandle		TemplateHandle;
+
+		// The node instance data size, includes latent properties (not serialized, @see FNodeTemplate::Finalize)
+		uint16	NodeInstanceDataSize;
+
+		uint16	Padding0;	// Unused
+
+		friend FDecoratorReader;
 
-		// Followed by a list of [FAnimNextDecoratorSharedData] instances and optional padding
+		// This structures is the header for a node's shared data. The memory layout is as follows:
+		// 
+		// [FNodeDescription] for the header
+		// [FAnimNextDecoratorSharedData] for decorator 1 (base)
+		// [FLatentPropertiesHeader][FLatentPropertyHandle][...] for the latent properties of decorator 1
+		// [FAnimNextDecoratorSharedData] for decorator 2 (additive)
+		// [FAnimNextDecoratorSharedData] for decorator 3 (base)
+		// [FLatentPropertiesHeader][FLatentPropertyHandle][...] for the latent properties of decorator 2
+		// [...]
+		// 
+		// Each node is thus followed by the decorator shared data contiguously.
+		// After each base decorator's shared data, an optional list of latent property handles follows
+		// before the next decorator's shared data begins. This list contains all latent property handles
+		// of the whole sub-stack (including the additive decorators on top of the base).
+		// 
+		// Each decorator contains a shared data structure that derives from FAnimNextDecoratorSharedData.
+		// That derived structure is what is contained in the actual buffer. As such, sizes and offsets
+		// vary as required. The [FDecoratorTemplate] contains the offsets that map here.
+		// 
+		// Optional padding is inserted as required by alignment constraints.
 	};
 }
diff --git a/Engine/Plugins/Experimental/AnimNext/Source/AnimNext/Public/DecoratorBase/NodeID.h b/Engine/Plugins/Experimental/AnimNext/Source/AnimNext/Public/DecoratorBase/NodeID.h
--- a/Engine/Plugins/Experimental/AnimNext/Source/AnimNext/Public/DecoratorBase/NodeID.h
+++ b/Engine/Plugins/Experimental/AnimNext/Source/AnimNext/Public/DecoratorBase/NodeID.h
@@ -10,7 +10,7 @@ namespace UE::AnimNext
 	/**
 	  * Node ID
 	  * 
-	  * A node ID fits on 16 bits and is used to encore a node index.
+	  * A node ID fits on 16 bits and is used to encode a node index.
 	  * The zero ID is reserved as an invalid value, this means that the
 	  * node index is 1-based.
 	  * 
diff --git a/Engine/Plugins/Experimental/AnimNext/Source/AnimNext/Public/DecoratorBase/NodeInstance.h b/Engine/Plugins/Experimental/AnimNext/Source/AnimNext/Public/DecoratorBase/NodeInstance.h
--- a/Engine/Plugins/Experimental/AnimNext/Source/AnimNext/Public/DecoratorBase/NodeInstance.h
+++ b/Engine/Plugins/Experimental/AnimNext/Source/AnimNext/Public/DecoratorBase/NodeInstance.h
@@ -53,15 +53,33 @@ namespace UE::AnimNext
 		// Decrements the reference count and returns true if any references remain
 		bool RemoveReference() { check(ReferenceCount > 0); return ReferenceCount-- != 1; }
 
-		FAnimNextGraphInstance& Owner;		// The graph instance that owns this node instance
+		FAnimNextGraphInstance&		Owner;					// The graph instance that owns this node instance
 
-		uint32		ReferenceCount;			// how many non-weak FDecoratorPtr handles point to us, not thread safe
-		FNodeHandle	NodeHandle;				// relative to root of sub-graph, should this be a pointer?
-
-		// Followed by a list of [FDecoratorInstanceData] instances and optional padding
+		uint32						ReferenceCount;			// how many non-weak FDecoratorPtr handles point to us, not thread safe
+		FNodeHandle					NodeHandle;				// relative to root of sub-graph, should this be a pointer?
 
 		friend struct FDecoratorPtr;
 		friend struct FExecutionContext;
 		friend struct FNodePtr;
+
+		// This structures is the header for a node's instance data. The memory layout is as follows:
+		// 
+		// [FNodeInstance] for the header
+		// [FDecoratorInstanceData] for decorator 1
+		// [FDecoratorInstanceData] for decorator 2
+		// [FDecoratorInstanceData] for decorator 3
+		// [...]
+		// [LatentProperty] for some decorator
+		// [LatentProperty] for some decorator
+		// [...]
+		// 
+		// Each node is thus followed by the decorator instance data contiguously.
+		// Following the decorator instance data, cached latent properties are present
+		// 
+		// Each decorator contains an instance data structure that derives from FDecoratorInstanceData.
+		// That derived structure is what is contained in the actual buffer. As such, sizes and offsets
+		// vary as required. The [FDecoratorTemplate] contains the offsets that map here.
+		// 
+		// Optional padding is inserted as required by alignment constraints.
 	};
 }
diff --git a/Engine/Plugins/Experimental/AnimNext/Source/AnimNext/Public/Graph/AnimNextGraphInstance.h b/Engine/Plugins/Experimental/AnimNext/Source/AnimNext/Public/Graph/AnimNextGraphInstance.h
--- a/Engine/Plugins/Experimental/AnimNext/Source/AnimNext/Public/Graph/AnimNextGraphInstance.h
+++ b/Engine/Plugins/Experimental/AnimNext/Source/AnimNext/Public/Graph/AnimNextGraphInstance.h
@@ -18,6 +18,7 @@ namespace UE::AnimNext
 {
 	struct FExecutionContext;
 	struct FGraphInstanceComponent;
+	struct FLatentPropertyHandle;
 }
 
 using GraphInstanceComponentMapType = TMap<FName, TSharedPtr<UE::AnimNext::FGraphInstanceComponent>>;
@@ -89,8 +90,9 @@ private:
 	// Adds the specified component and returns a reference to it
 	UE::AnimNext::FGraphInstanceComponent& AddComponent(int32 ComponentNameHash, FName ComponentName, TSharedPtr<UE::AnimNext::FGraphInstanceComponent>&& Component);
 
-	// Executes a latent RigVM pin and writes the result into the destination pointer
-	void ExecuteLatentPin(int32 LatentPinIndex, void* DestinationPtr);
+	// Executes a list of latent RigVM pins and writes the result into the destination pointer (latent handle offsets are using the destination as base)
+	// When frozen, latent handles that can freeze are skipped, all others will execute
+	void ExecuteLatentPins(const TConstArrayView<UE::AnimNext::FLatentPropertyHandle>& LatentHandles, void* DestinationBasePtr, bool bIsFrozen);
 
 	// During graph compilation, if we have existing graph instances, we freeze them by releasing their memory before thawing them
 	// Freezing is a partial release of resources that retains the necessary information to re-create things safely
diff --git a/Engine/Plugins/Experimental/AnimNext/Source/AnimNextTestSuite/Private/AnimNextDecoratorBaseTest.cpp b/Engine/Plugins/Experimental/AnimNext/Source/AnimNextTestSuite/Private/AnimNextDecoratorBaseTest.cpp
--- a/Engine/Plugins/Experimental/AnimNext/Source/AnimNextTestSuite/Private/AnimNextDecoratorBaseTest.cpp
+++ b/Engine/Plugins/Experimental/AnimNext/Source/AnimNextTestSuite/Private/AnimNextDecoratorBaseTest.cpp
@@ -157,9 +157,12 @@ namespace UE::AnimNext
 		}
 	};
 
-	DEFINE_ANIM_DECORATOR_BEGIN(FDecoratorA_Base)
-		DEFINE_ANIM_DECORATOR_IMPLEMENTS_INTERFACE(IInterfaceA)
-	DEFINE_ANIM_DECORATOR_END(FDecoratorA_Base)
+	// Decorator implementation boilerplate
+	#define DECORATOR_INTERFACE_ENUMERATOR(GeneratorMacro) \
+		GeneratorMacro(IInterfaceA) \
+
+	GENERATE_ANIM_DECORATOR_IMPLEMENTATION(FDecoratorA_Base, DECORATOR_INTERFACE_ENUMERATOR)
+	#undef DECORATOR_INTERFACE_ENUMERATOR
 
 	//////////////////////////////////////////////////////////////////////////
 
@@ -201,10 +204,13 @@ namespace UE::AnimNext
 		}
 	};
 
-	DEFINE_ANIM_DECORATOR_BEGIN(FDecoratorAB_Add)
-		DEFINE_ANIM_DECORATOR_IMPLEMENTS_INTERFACE(IInterfaceA)
-		DEFINE_ANIM_DECORATOR_IMPLEMENTS_INTERFACE(IInterfaceB)
-	DEFINE_ANIM_DECORATOR_END(FDecoratorAB_Add)
+	// Decorator implementation boilerplate
+	#define DECORATOR_INTERFACE_ENUMERATOR(GeneratorMacro) \
+		GeneratorMacro(IInterfaceA) \
+		GeneratorMacro(IInterfaceB) \
+
+	GENERATE_ANIM_DECORATOR_IMPLEMENTATION(FDecoratorAB_Add, DECORATOR_INTERFACE_ENUMERATOR)
+	#undef DECORATOR_INTERFACE_ENUMERATOR
 
 	//////////////////////////////////////////////////////////////////////////
 
@@ -246,10 +252,13 @@ namespace UE::AnimNext
 		}
 	};
 
-	DEFINE_ANIM_DECORATOR_BEGIN(FDecoratorAC_Add)
-		DEFINE_ANIM_DECORATOR_IMPLEMENTS_INTERFACE(IInterfaceA)
-		DEFINE_ANIM_DECORATOR_IMPLEMENTS_INTERFACE(IInterfaceC)
-	DEFINE_ANIM_DECORATOR_END(FDecoratorAC_Add)
+	// Decorator implementation boilerplate
+	#define DECORATOR_INTERFACE_ENUMERATOR(GeneratorMacro) \
+		GeneratorMacro(IInterfaceA) \
+		GeneratorMacro(IInterfaceC) \
+
+	GENERATE_ANIM_DECORATOR_IMPLEMENTATION(FDecoratorAC_Add, DECORATOR_INTERFACE_ENUMERATOR)
+	#undef DECORATOR_INTERFACE_ENUMERATOR
 
 
 	//////////////////////////////////////////////////////////////////////////
@@ -261,9 +270,12 @@ namespace UE::AnimNext
 		using FSharedData = FDecoratorSerialization_BaseSharedData;
 	};
 
-	DEFINE_ANIM_DECORATOR_BEGIN(FDecoratorSerialization_Base)
-		DEFINE_ANIM_DECORATOR_IMPLEMENTS_INTERFACE(IInterfaceA)
-	DEFINE_ANIM_DECORATOR_END(FDecoratorSerialization_Base)
+	// Decorator implementation boilerplate
+	#define DECORATOR_INTERFACE_ENUMERATOR(GeneratorMacro) \
+		GeneratorMacro(IInterfaceA) \
+
+	GENERATE_ANIM_DECORATOR_IMPLEMENTATION(FDecoratorSerialization_Base, DECORATOR_INTERFACE_ENUMERATOR)
+	#undef DECORATOR_INTERFACE_ENUMERATOR
 
 	//////////////////////////////////////////////////////////////////////////
 
@@ -274,9 +286,12 @@ namespace UE::AnimNext
 		using FSharedData = FDecoratorSerialization_AddSharedData;
 	};
 
-	DEFINE_ANIM_DECORATOR_BEGIN(FDecoratorSerialization_Add)
-		DEFINE_ANIM_DECORATOR_IMPLEMENTS_INTERFACE(IInterfaceB)
-	DEFINE_ANIM_DECORATOR_END(FDecoratorSerialization_Add)
+	// Decorator implementation boilerplate
+	#define DECORATOR_INTERFACE_ENUMERATOR(GeneratorMacro) \
+		GeneratorMacro(IInterfaceB) \
+
+	GENERATE_ANIM_DECORATOR_IMPLEMENTATION(FDecoratorSerialization_Add, DECORATOR_INTERFACE_ENUMERATOR)
+	#undef DECORATOR_INTERFACE_ENUMERATOR
 
 	//////////////////////////////////////////////////////////////////////////
 
@@ -287,9 +302,12 @@ namespace UE::AnimNext
 		using FSharedData = FDecoratorNativeSerialization_AddSharedData;
 	};
 
-	DEFINE_ANIM_DECORATOR_BEGIN(FDecoratorNativeSerialization_Add)
-		DEFINE_ANIM_DECORATOR_IMPLEMENTS_INTERFACE(IInterfaceC)
-	DEFINE_ANIM_DECORATOR_END(FDecoratorNativeSerialization_Add)
+	// Decorator implementation boilerplate
+	#define DECORATOR_INTERFACE_ENUMERATOR(GeneratorMacro) \
+		GeneratorMacro(IInterfaceC) \
+
+	GENERATE_ANIM_DECORATOR_IMPLEMENTATION(FDecoratorNativeSerialization_Add, DECORATOR_INTERFACE_ENUMERATOR)
+	#undef DECORATOR_INTERFACE_ENUMERATOR
 }
 
 IMPLEMENT_SIMPLE_AUTOMATION_TEST(FAnimationAnimNextRuntimeTest_DecoratorRegistry, "Animation.AnimNext.Runtime.DecoratorRegistry", EAutomationTestFlags::EditorContext | EAutomationTestFlags::EngineFilter)
@@ -539,22 +557,22 @@ bool FAnimationAnimNextRuntimeTest_NodeLifetime::RunTest(const FString& InParame
 
 			DecoratorWriter.BeginNodeWriting();
 			DecoratorWriter.WriteNode(NodeHandles[0],
-				[](uint32 DecoratorIndex, const FString& PropertyName)
+				[](uint32 DecoratorIndex, FName PropertyName)
 				{
 					return FString();
 				},
-				[](uint32 DecoratorIndex, const FString& PropertyName)
+				[](uint32 DecoratorIndex, FName PropertyName)
 				{
-					return false;
+					return MAX_uint16;
 				});
 			DecoratorWriter.WriteNode(NodeHandles[1],
-				[](uint32 DecoratorIndex, const FString& PropertyName)
+				[](uint32 DecoratorIndex, FName PropertyName)
 				{
 					return FString();
 				},
-				[](uint32 DecoratorIndex, const FString& PropertyName)
+				[](uint32 DecoratorIndex, FName PropertyName)
 				{
-					return false;
+					return MAX_uint16;
 				});
 			DecoratorWriter.EndNodeWriting();
 
@@ -750,13 +768,13 @@ bool FAnimationAnimNextRuntimeTest_GetDecoratorInterface::RunTest(const FString&
 
 			DecoratorWriter.BeginNodeWriting();
 			DecoratorWriter.WriteNode(NodeHandles[0],
-				[](uint32 DecoratorIndex, const FString& PropertyName)
+				[](uint32 DecoratorIndex, FName PropertyName)
 				{
 					return FString();
 				},
-				[](uint32 DecoratorIndex, const FString& PropertyName)
+				[](uint32 DecoratorIndex, FName PropertyName)
 				{
-					return false;
+					return MAX_uint16;
 				});
 			DecoratorWriter.EndNodeWriting();
 
@@ -1023,13 +1041,13 @@ bool FAnimationAnimNextRuntimeTest_GetDecoratorInterfaceSuper::RunTest(const FSt
 
 			DecoratorWriter.BeginNodeWriting();
 			DecoratorWriter.WriteNode(NodeHandles[0],
-				[](uint32 DecoratorIndex, const FString& PropertyName)
+				[](uint32 DecoratorIndex, FName PropertyName)
 				{
 					return FString();
 				},
-				[](uint32 DecoratorIndex, const FString& PropertyName)
+				[](uint32 DecoratorIndex, FName PropertyName)
 				{
-					return false;
+					return MAX_uint16;
 				});
 			DecoratorWriter.EndNodeWriting();
 
@@ -1185,339 +1203,341 @@ bool FAnimationAnimNextRuntimeTest_DecoratorSerialization::RunTest(const FString
 {
 	using namespace UE::AnimNext;
 
-	AUTO_REGISTER_ANIM_DECORATOR(FDecoratorSerialization_Base)
-	AUTO_REGISTER_ANIM_DECORATOR(FDecoratorSerialization_Add)
-	AUTO_REGISTER_ANIM_DECORATOR(FDecoratorNativeSerialization_Add)
+	{
+		FScopedClearNodeTemplateRegistry ScopedClearNodeTemplateRegistry;
+		FNodeTemplateRegistry& Registry = FNodeTemplateRegistry::Get();
 
-	UFactory* GraphFactory = NewObject<UAnimNextGraphFactory>();
-	UAnimNextGraph* AnimNextGraph = CastChecked<UAnimNextGraph>(GraphFactory->FactoryCreateNew(UAnimNextGraph::StaticClass(), GetTransientPackage(), TEXT("TestAnimNextGraph"), RF_Transient, nullptr, nullptr, NAME_None));
-	UE_RETURN_ON_ERROR(AnimNextGraph != nullptr, "FAnimationAnimNextRuntimeTest_DecoratorSerialization -> Failed to create animation graph");
+		AUTO_REGISTER_ANIM_DECORATOR(FDecoratorSerialization_Base)
+		AUTO_REGISTER_ANIM_DECORATOR(FDecoratorSerialization_Add)
+		AUTO_REGISTER_ANIM_DECORATOR(FDecoratorNativeSerialization_Add)
 
-	FScopedClearNodeTemplateRegistry ScopedClearNodeTemplateRegistry;
-	FNodeTemplateRegistry& Registry = FNodeTemplateRegistry::Get();
+		UFactory* GraphFactory = NewObject<UAnimNextGraphFactory>();
+		UAnimNextGraph* AnimNextGraph = CastChecked<UAnimNextGraph>(GraphFactory->FactoryCreateNew(UAnimNextGraph::StaticClass(), GetTransientPackage(), TEXT("TestAnimNextGraph"), RF_Transient, nullptr, nullptr, NAME_None));
+		UE_RETURN_ON_ERROR(AnimNextGraph != nullptr, "FAnimationAnimNextRuntimeTest_DecoratorSerialization -> Failed to create animation graph");
 
-	TArray<FDecoratorUID> NodeTemplateDecoratorList;
-	NodeTemplateDecoratorList.Add(FDecoratorSerialization_Base::DecoratorUID);
-	NodeTemplateDecoratorList.Add(FDecoratorSerialization_Add::DecoratorUID);
-	NodeTemplateDecoratorList.Add(FDecoratorNativeSerialization_Add::DecoratorUID);
+		TArray<FDecoratorUID> NodeTemplateDecoratorList;
+		NodeTemplateDecoratorList.Add(FDecoratorSerialization_Base::DecoratorUID);
+		NodeTemplateDecoratorList.Add(FDecoratorSerialization_Add::DecoratorUID);
+		NodeTemplateDecoratorList.Add(FDecoratorNativeSerialization_Add::DecoratorUID);
 
-	// Populate our node template registry
-	TArray<uint8> NodeTemplateBuffer0;
-	const FNodeTemplate* NodeTemplate0 = FNodeTemplateBuilder::BuildNodeTemplate(NodeTemplateDecoratorList, NodeTemplateBuffer0);
+		// Populate our node template registry
+		TArray<uint8> NodeTemplateBuffer0;
+		const FNodeTemplate* NodeTemplate0 = FNodeTemplateBuilder::BuildNodeTemplate(NodeTemplateDecoratorList, NodeTemplateBuffer0);
 
-	FNodeTemplateRegistryHandle TemplateHandle0 = Registry.FindOrAdd(NodeTemplate0);
-	AddErrorIfFalse(TemplateHandle0.IsValid(), "FAnimationAnimNextRuntimeTest_DecoratorSerialization -> Registry should contain our template");
-
-	FDecoratorSerialization_Base::FSharedData DecoratorBaseRef0;
-	DecoratorBaseRef0.Integer = 1651;
-	DecoratorBaseRef0.IntegerArray[0] = 1071;
-	DecoratorBaseRef0.IntegerArray[1] = -158;
-	DecoratorBaseRef0.IntegerArray[2] = 88116;
-	DecoratorBaseRef0.IntegerArray[3] = 0x417;
-	DecoratorBaseRef0.IntegerTArray = { -8162, 88152, 0x8152f };
-	DecoratorBaseRef0.Vector = FVector(0.1917, 12435.1, -18200.1726);
-	DecoratorBaseRef0.VectorArray[0] = FVector(192.1716, -1927.115, 99176.12);
-	DecoratorBaseRef0.VectorArray[1] = FVector(961.811, -18956.117, 81673.44);
-	DecoratorBaseRef0.VectorTArray = { FVector(-1927.8771, 1826.9917, -123.1555), FVector(9177.011, -71.44, -917.88), FVector(123.91, 852.11, -81652.1) };
-	DecoratorBaseRef0.String = TEXT("sample string 123");
-	DecoratorBaseRef0.Name = FName(TEXT("sample name 999178"));
-
-	FDecoratorSerialization_Add::FSharedData DecoratorAddRef0;
-	DecoratorAddRef0.Integer = 16511;
-	DecoratorAddRef0.IntegerArray[0] = 10711;
-	DecoratorAddRef0.IntegerArray[1] = -1581;
-	DecoratorAddRef0.IntegerArray[2] = 881161;
-	DecoratorAddRef0.IntegerArray[3] = 0x4171;
-	DecoratorAddRef0.IntegerTArray = { -81621, 881521, 0x8152f1 };
-	DecoratorAddRef0.Vector = FVector(0.19171, 12435.11, -18200.17261);
-	DecoratorAddRef0.VectorArray[0] = FVector(192.17161, -1927.1151, 99176.121);
-	DecoratorAddRef0.VectorArray[1] = FVector(961.8111, -18956.1171, 81673.441);
-	DecoratorAddRef0.VectorTArray = { FVector(-1927.87711, 1826.99171, -123.15551), FVector(9177.0111, -71.441, -917.881), FVector(123.911, 852.111, -81652.11) };
-	DecoratorAddRef0.String = TEXT("sample string 1231");
-	DecoratorAddRef0.Name = FName(TEXT("sample name 9991781"));
-
-	FDecoratorNativeSerialization_Add::FSharedData DecoratorNativeRef0;
-	DecoratorNativeRef0.Integer = 16514;
-	DecoratorNativeRef0.IntegerArray[0] = 10714;
-	DecoratorNativeRef0.IntegerArray[1] = -1584;
-	DecoratorNativeRef0.IntegerArray[2] = 881164;
-	DecoratorNativeRef0.IntegerArray[3] = 0x4174;
-	DecoratorNativeRef0.IntegerTArray = { -81624, 881524, 0x8152f4 };
-	DecoratorNativeRef0.Vector = FVector(0.19174, 12435.14, -18200.17264);
-	DecoratorNativeRef0.VectorArray[0] = FVector(192.17164, -1927.1154, 99176.124);
-	DecoratorNativeRef0.VectorArray[1] = FVector(961.8114, -18956.1174, 81673.444);
-	DecoratorNativeRef0.VectorTArray = { FVector(-1927.87714, 1826.99174, -123.15554), FVector(9177.0114, -71.444, -917.884), FVector(123.914, 852.114, -81652.14) };
-	DecoratorNativeRef0.String = TEXT("sample string 1234");
-	DecoratorNativeRef0.Name = FName(TEXT("sample name 9991784"));
-
-	FDecoratorSerialization_Base::FSharedData DecoratorBaseRef1;
-	DecoratorBaseRef1.Integer = 16512;
-	DecoratorBaseRef1.IntegerArray[0] = 10712;
-	DecoratorBaseRef1.IntegerArray[1] = -1582;
-	DecoratorBaseRef1.IntegerArray[2] = 881162;
-	DecoratorBaseRef1.IntegerArray[3] = 0x4172;
-	DecoratorBaseRef1.IntegerTArray = { -81622, 881522, 0x8152f2 };
-	DecoratorBaseRef1.Vector = FVector(0.19172, 12435.12, -18200.17262);
-	DecoratorBaseRef1.VectorArray[0] = FVector(192.17162, -1927.1152, 99176.122);
-	DecoratorBaseRef1.VectorArray[1] = FVector(961.8112, -18956.1172, 81673.442);
-	DecoratorBaseRef1.VectorTArray = { FVector(-1927.87712, 1826.99172, -123.15552), FVector(9177.0112, -71.442, -917.882), FVector(123.912, 852.112, -81652.12) };
-	DecoratorBaseRef1.String = TEXT("sample string 1232");
-	DecoratorBaseRef1.Name = FName(TEXT("sample name 9991782"));
-
-	FDecoratorSerialization_Add::FSharedData DecoratorAddRef1;
-	DecoratorAddRef1.Integer = 16513;
-	DecoratorAddRef1.IntegerArray[0] = 10713;
-	DecoratorAddRef1.IntegerArray[1] = -1583;
-	DecoratorAddRef1.IntegerArray[2] = 881163;
-	DecoratorAddRef1.IntegerArray[3] = 0x4173;
-	DecoratorAddRef1.IntegerTArray = { -81623, 881523, 0x8152f3 };
-	DecoratorAddRef1.Vector = FVector(0.19173, 12435.13, -18200.17263);
-	DecoratorAddRef1.VectorArray[0] = FVector(192.17163, -1927.1153, 99176.123);
-	DecoratorAddRef1.VectorArray[1] = FVector(961.8113, -18956.1173, 81673.443);
-	DecoratorAddRef1.VectorTArray = { FVector(-1927.87713, 1826.99173, -123.15553), FVector(9177.0113, -71.443, -917.883), FVector(123.913, 852.113, -81652.13) };
-	DecoratorAddRef1.String = TEXT("sample string 1233");
-	DecoratorAddRef1.Name = FName(TEXT("sample name 9991783"));
-
-	FDecoratorNativeSerialization_Add::FSharedData DecoratorNativeRef1;
-	DecoratorNativeRef1.Integer = 16515;
-	DecoratorNativeRef1.IntegerArray[0] = 10715;
-	DecoratorNativeRef1.IntegerArray[1] = -1585;
-	DecoratorNativeRef1.IntegerArray[2] = 881165;
-	DecoratorNativeRef1.IntegerArray[3] = 0x4175;
-	DecoratorNativeRef1.IntegerTArray = { -81625, 881525, 0x8152f5 };
-	DecoratorNativeRef1.Vector = FVector(0.19175, 12435.15, -18200.17265);
-	DecoratorNativeRef1.VectorArray[0] = FVector(192.17165, -1927.1155, 99176.125);
-	DecoratorNativeRef1.VectorArray[1] = FVector(961.8115, -18956.1175, 81673.445);
-	DecoratorNativeRef1.VectorTArray = { FVector(-1927.87715, 1826.99175, -123.15555), FVector(9177.0115, -71.445, -917.885), FVector(123.915, 852.115, -81652.15) };
-	DecoratorNativeRef1.String = TEXT("sample string 1235");
-	DecoratorNativeRef1.Name = FName(TEXT("sample name 9991785"));
-
-	TArray<FNodeHandle> NodeHandles;
-
-	// Write our graph
-	TArray<uint8> GraphSharedDataArchiveBuffer;
-	TArray<TObjectPtr<UObject>> GraphReferencedObjects;
-	{
-		FDecoratorWriter DecoratorWriter;
-
-		NodeHandles.Add(DecoratorWriter.RegisterNode(*NodeTemplate0));
-		NodeHandles.Add(DecoratorWriter.RegisterNode(*NodeTemplate0));
-
-		// We don't have decorator properties
-		TArray<TMap<FString, FString>> DecoratorProperties0;
-		DecoratorProperties0.AddDefaulted(NodeTemplateDecoratorList.Num());
-
-		DecoratorProperties0[0].Add(TEXT("Integer"), ToString<FDecoratorSerialization_Base::FSharedData>(TEXT("Integer"), DecoratorBaseRef0.Integer));
-		DecoratorProperties0[0].Add(TEXT("IntegerArray"), ToString<FDecoratorSerialization_Base::FSharedData>(TEXT("IntegerArray"), DecoratorBaseRef0.IntegerArray));
-		DecoratorProperties0[0].Add(TEXT("IntegerTArray"), ToString<FDecoratorSerialization_Base::FSharedData>(TEXT("IntegerTArray"), DecoratorBaseRef0.IntegerTArray));
-		DecoratorProperties0[0].Add(TEXT("Vector"), ToString<FDecoratorSerialization_Base::FSharedData>(TEXT("Vector"), DecoratorBaseRef0.Vector));
-		DecoratorProperties0[0].Add(TEXT("VectorArray"), ToString<FDecoratorSerialization_Base::FSharedData>(TEXT("VectorArray"), DecoratorBaseRef0.VectorArray));
-		DecoratorProperties0[0].Add(TEXT("VectorTArray"), ToString<FDecoratorSerialization_Base::FSharedData>(TEXT("VectorTArray"), DecoratorBaseRef0.VectorTArray));
-		DecoratorProperties0[0].Add(TEXT("String"), ToString<FDecoratorSerialization_Base::FSharedData>(TEXT("String"), DecoratorBaseRef0.String));
-		DecoratorProperties0[0].Add(TEXT("Name"), ToString<FDecoratorSerialization_Base::FSharedData>(TEXT("Name"), DecoratorBaseRef0.Name));
-
-		DecoratorProperties0[1].Add(TEXT("Integer"), ToString<FDecoratorSerialization_Add::FSharedData>(TEXT("Integer"), DecoratorAddRef0.Integer));
-		DecoratorProperties0[1].Add(TEXT("IntegerArray"), ToString<FDecoratorSerialization_Add::FSharedData>(TEXT("IntegerArray"), DecoratorAddRef0.IntegerArray));
-		DecoratorProperties0[1].Add(TEXT("IntegerTArray"), ToString<FDecoratorSerialization_Add::FSharedData>(TEXT("IntegerTArray"), DecoratorAddRef0.IntegerTArray));
-		DecoratorProperties0[1].Add(TEXT("Vector"), ToString<FDecoratorSerialization_Add::FSharedData>(TEXT("Vector"), DecoratorAddRef0.Vector));
-		DecoratorProperties0[1].Add(TEXT("VectorArray"), ToString<FDecoratorSerialization_Add::FSharedData>(TEXT("VectorArray"), DecoratorAddRef0.VectorArray));
-		DecoratorProperties0[1].Add(TEXT("VectorTArray"), ToString<FDecoratorSerialization_Add::FSharedData>(TEXT("VectorTArray"), DecoratorAddRef0.VectorTArray));
-		DecoratorProperties0[1].Add(TEXT("String"), ToString<FDecoratorSerialization_Add::FSharedData>(TEXT("String"), DecoratorAddRef0.String));
-		DecoratorProperties0[1].Add(TEXT("Name"), ToString<FDecoratorSerialization_Add::FSharedData>(TEXT("Name"), DecoratorAddRef0.Name));
-
-		DecoratorProperties0[2].Add(TEXT("Integer"), ToString<FDecoratorNativeSerialization_Add::FSharedData>(TEXT("Integer"), DecoratorNativeRef0.Integer));
-		DecoratorProperties0[2].Add(TEXT("IntegerArray"), ToString<FDecoratorNativeSerialization_Add::FSharedData>(TEXT("IntegerArray"), DecoratorNativeRef0.IntegerArray));
-		DecoratorProperties0[2].Add(TEXT("IntegerTArray"), ToString<FDecoratorNativeSerialization_Add::FSharedData>(TEXT("IntegerTArray"), DecoratorNativeRef0.IntegerTArray));
-		DecoratorProperties0[2].Add(TEXT("Vector"), ToString<FDecoratorNativeSerialization_Add::FSharedData>(TEXT("Vector"), DecoratorNativeRef0.Vector));
-		DecoratorProperties0[2].Add(TEXT("VectorArray"), ToString<FDecoratorNativeSerialization_Add::FSharedData>(TEXT("VectorArray"), DecoratorNativeRef0.VectorArray));
-		DecoratorProperties0[2].Add(TEXT("VectorTArray"), ToString<FDecoratorNativeSerialization_Add::FSharedData>(TEXT("VectorTArray"), DecoratorNativeRef0.VectorTArray));
-		DecoratorProperties0[2].Add(TEXT("String"), ToString<FDecoratorNativeSerialization_Add::FSharedData>(TEXT("String"), DecoratorNativeRef0.String));
-		DecoratorProperties0[2].Add(TEXT("Name"), ToString<FDecoratorNativeSerialization_Add::FSharedData>(TEXT("Name"), DecoratorNativeRef0.Name));
-
-		TArray<TMap<FString, FString>> DecoratorProperties1;
-		DecoratorProperties1.AddDefaulted(NodeTemplateDecoratorList.Num());
-
-		DecoratorProperties1[0].Add(TEXT("Integer"), ToString<FDecoratorSerialization_Base::FSharedData>(TEXT("Integer"), DecoratorBaseRef1.Integer));
-		DecoratorProperties1[0].Add(TEXT("IntegerArray"), ToString<FDecoratorSerialization_Base::FSharedData>(TEXT("IntegerArray"), DecoratorBaseRef1.IntegerArray));
-		DecoratorProperties1[0].Add(TEXT("IntegerTArray"), ToString<FDecoratorSerialization_Base::FSharedData>(TEXT("IntegerTArray"), DecoratorBaseRef1.IntegerTArray));
-		DecoratorProperties1[0].Add(TEXT("Vector"), ToString<FDecoratorSerialization_Base::FSharedData>(TEXT("Vector"), DecoratorBaseRef1.Vector));
-		DecoratorProperties1[0].Add(TEXT("VectorArray"), ToString<FDecoratorSerialization_Base::FSharedData>(TEXT("VectorArray"), DecoratorBaseRef1.VectorArray));
-		DecoratorProperties1[0].Add(TEXT("VectorTArray"), ToString<FDecoratorSerialization_Base::FSharedData>(TEXT("VectorTArray"), DecoratorBaseRef1.VectorTArray));
-		DecoratorProperties1[0].Add(TEXT("String"), ToString<FDecoratorSerialization_Base::FSharedData>(TEXT("String"), DecoratorBaseRef1.String));
-		DecoratorProperties1[0].Add(TEXT("Name"), ToString<FDecoratorSerialization_Base::FSharedData>(TEXT("Name"), DecoratorBaseRef1.Name));
-
-		DecoratorProperties1[1].Add(TEXT("Integer"), ToString<FDecoratorSerialization_Add::FSharedData>(TEXT("Integer"), DecoratorAddRef1.Integer));
-		DecoratorProperties1[1].Add(TEXT("IntegerArray"), ToString<FDecoratorSerialization_Add::FSharedData>(TEXT("IntegerArray"), DecoratorAddRef1.IntegerArray));
-		DecoratorProperties1[1].Add(TEXT("IntegerTArray"), ToString<FDecoratorSerialization_Add::FSharedData>(TEXT("IntegerTArray"), DecoratorAddRef1.IntegerTArray));
-		DecoratorProperties1[1].Add(TEXT("Vector"), ToString<FDecoratorSerialization_Add::FSharedData>(TEXT("Vector"), DecoratorAddRef1.Vector));
-		DecoratorProperties1[1].Add(TEXT("VectorArray"), ToString<FDecoratorSerialization_Add::FSharedData>(TEXT("VectorArray"), DecoratorAddRef1.VectorArray));
-		DecoratorProperties1[1].Add(TEXT("VectorTArray"), ToString<FDecoratorSerialization_Add::FSharedData>(TEXT("VectorTArray"), DecoratorAddRef1.VectorTArray));
-		DecoratorProperties1[1].Add(TEXT("String"), ToString<FDecoratorSerialization_Add::FSharedData>(TEXT("String"), DecoratorAddRef1.String));
-		DecoratorProperties1[1].Add(TEXT("Name"), ToString<FDecoratorSerialization_Add::FSharedData>(TEXT("Name"), DecoratorAddRef1.Name));
-
-		DecoratorProperties1[2].Add(TEXT("Integer"), ToString<FDecoratorNativeSerialization_Add::FSharedData>(TEXT("Integer"), DecoratorNativeRef1.Integer));
-		DecoratorProperties1[2].Add(TEXT("IntegerArray"), ToString<FDecoratorNativeSerialization_Add::FSharedData>(TEXT("IntegerArray"), DecoratorNativeRef1.IntegerArray));
-		DecoratorProperties1[2].Add(TEXT("IntegerTArray"), ToString<FDecoratorNativeSerialization_Add::FSharedData>(TEXT("IntegerTArray"), DecoratorNativeRef1.IntegerTArray));
-		DecoratorProperties1[2].Add(TEXT("Vector"), ToString<FDecoratorNativeSerialization_Add::FSharedData>(TEXT("Vector"), DecoratorNativeRef1.Vector));
-		DecoratorProperties1[2].Add(TEXT("VectorArray"), ToString<FDecoratorNativeSerialization_Add::FSharedData>(TEXT("VectorArray"), DecoratorNativeRef1.VectorArray));
-		DecoratorProperties1[2].Add(TEXT("VectorTArray"), ToString<FDecoratorNativeSerialization_Add::FSharedData>(TEXT("VectorTArray"), DecoratorNativeRef1.VectorTArray));
-		DecoratorProperties1[2].Add(TEXT("String"), ToString<FDecoratorNativeSerialization_Add::FSharedData>(TEXT("String"), DecoratorNativeRef1.String));
-		DecoratorProperties1[2].Add(TEXT("Name"), ToString<FDecoratorNativeSerialization_Add::FSharedData>(TEXT("Name"), DecoratorNativeRef1.Name));
-
-		DecoratorWriter.BeginNodeWriting();
-		DecoratorWriter.WriteNode(NodeHandles[0],
-			[&DecoratorProperties0](uint32 DecoratorIndex, const FString& PropertyName)
-			{
-				return DecoratorProperties0[DecoratorIndex][PropertyName];
-			},
-			[](uint32 DecoratorIndex, const FString& PropertyName)
-			{
-				return false;
-			});
-		DecoratorWriter.WriteNode(NodeHandles[1],
-			[&DecoratorProperties1](uint32 DecoratorIndex, const FString& PropertyName)
-			{
-				return DecoratorProperties1[DecoratorIndex][PropertyName];
-			},
-			[](uint32 DecoratorIndex, const FString& PropertyName)
-			{
-				return false;
-			});
-		DecoratorWriter.EndNodeWriting();
+		FNodeTemplateRegistryHandle TemplateHandle0 = Registry.FindOrAdd(NodeTemplate0);
+		AddErrorIfFalse(TemplateHandle0.IsValid(), "FAnimationAnimNextRuntimeTest_DecoratorSerialization -> Registry should contain our template");
+
+		FDecoratorSerialization_Base::FSharedData DecoratorBaseRef0;
+		DecoratorBaseRef0.Integer = 1651;
+		DecoratorBaseRef0.IntegerArray[0] = 1071;
+		DecoratorBaseRef0.IntegerArray[1] = -158;
+		DecoratorBaseRef0.IntegerArray[2] = 88116;
+		DecoratorBaseRef0.IntegerArray[3] = 0x417;
+		DecoratorBaseRef0.IntegerTArray = { -8162, 88152, 0x8152f };
+		DecoratorBaseRef0.Vector = FVector(0.1917, 12435.1, -18200.1726);
+		DecoratorBaseRef0.VectorArray[0] = FVector(192.1716, -1927.115, 99176.12);
+		DecoratorBaseRef0.VectorArray[1] = FVector(961.811, -18956.117, 81673.44);
+		DecoratorBaseRef0.VectorTArray = { FVector(-1927.8771, 1826.9917, -123.1555), FVector(9177.011, -71.44, -917.88), FVector(123.91, 852.11, -81652.1) };
+		DecoratorBaseRef0.String = TEXT("sample string 123");
+		DecoratorBaseRef0.Name = FName(TEXT("sample name 999178"));
+
+		FDecoratorSerialization_Add::FSharedData DecoratorAddRef0;
+		DecoratorAddRef0.Integer = 16511;
+		DecoratorAddRef0.IntegerArray[0] = 10711;
+		DecoratorAddRef0.IntegerArray[1] = -1581;
+		DecoratorAddRef0.IntegerArray[2] = 881161;
+		DecoratorAddRef0.IntegerArray[3] = 0x4171;
+		DecoratorAddRef0.IntegerTArray = { -81621, 881521, 0x8152f1 };
+		DecoratorAddRef0.Vector = FVector(0.19171, 12435.11, -18200.17261);
+		DecoratorAddRef0.VectorArray[0] = FVector(192.17161, -1927.1151, 99176.121);
+		DecoratorAddRef0.VectorArray[1] = FVector(961.8111, -18956.1171, 81673.441);
+		DecoratorAddRef0.VectorTArray = { FVector(-1927.87711, 1826.99171, -123.15551), FVector(9177.0111, -71.441, -917.881), FVector(123.911, 852.111, -81652.11) };
+		DecoratorAddRef0.String = TEXT("sample string 1231");
+		DecoratorAddRef0.Name = FName(TEXT("sample name 9991781"));
+
+		FDecoratorNativeSerialization_Add::FSharedData DecoratorNativeRef0;
+		DecoratorNativeRef0.Integer = 16514;
+		DecoratorNativeRef0.IntegerArray[0] = 10714;
+		DecoratorNativeRef0.IntegerArray[1] = -1584;
+		DecoratorNativeRef0.IntegerArray[2] = 881164;
+		DecoratorNativeRef0.IntegerArray[3] = 0x4174;
+		DecoratorNativeRef0.IntegerTArray = { -81624, 881524, 0x8152f4 };
+		DecoratorNativeRef0.Vector = FVector(0.19174, 12435.14, -18200.17264);
+		DecoratorNativeRef0.VectorArray[0] = FVector(192.17164, -1927.1154, 99176.124);
+		DecoratorNativeRef0.VectorArray[1] = FVector(961.8114, -18956.1174, 81673.444);
+		DecoratorNativeRef0.VectorTArray = { FVector(-1927.87714, 1826.99174, -123.15554), FVector(9177.0114, -71.444, -917.884), FVector(123.914, 852.114, -81652.14) };
+		DecoratorNativeRef0.String = TEXT("sample string 1234");
+		DecoratorNativeRef0.Name = FName(TEXT("sample name 9991784"));
+
+		FDecoratorSerialization_Base::FSharedData DecoratorBaseRef1;
+		DecoratorBaseRef1.Integer = 16512;
+		DecoratorBaseRef1.IntegerArray[0] = 10712;
+		DecoratorBaseRef1.IntegerArray[1] = -1582;
+		DecoratorBaseRef1.IntegerArray[2] = 881162;
+		DecoratorBaseRef1.IntegerArray[3] = 0x4172;
+		DecoratorBaseRef1.IntegerTArray = { -81622, 881522, 0x8152f2 };
+		DecoratorBaseRef1.Vector = FVector(0.19172, 12435.12, -18200.17262);
+		DecoratorBaseRef1.VectorArray[0] = FVector(192.17162, -1927.1152, 99176.122);
+		DecoratorBaseRef1.VectorArray[1] = FVector(961.8112, -18956.1172, 81673.442);
+		DecoratorBaseRef1.VectorTArray = { FVector(-1927.87712, 1826.99172, -123.15552), FVector(9177.0112, -71.442, -917.882), FVector(123.912, 852.112, -81652.12) };
+		DecoratorBaseRef1.String = TEXT("sample string 1232");
+		DecoratorBaseRef1.Name = FName(TEXT("sample name 9991782"));
+
+		FDecoratorSerialization_Add::FSharedData DecoratorAddRef1;
+		DecoratorAddRef1.Integer = 16513;
+		DecoratorAddRef1.IntegerArray[0] = 10713;
+		DecoratorAddRef1.IntegerArray[1] = -1583;
+		DecoratorAddRef1.IntegerArray[2] = 881163;
+		DecoratorAddRef1.IntegerArray[3] = 0x4173;
+		DecoratorAddRef1.IntegerTArray = { -81623, 881523, 0x8152f3 };
+		DecoratorAddRef1.Vector = FVector(0.19173, 12435.13, -18200.17263);
+		DecoratorAddRef1.VectorArray[0] = FVector(192.17163, -1927.1153, 99176.123);
+		DecoratorAddRef1.VectorArray[1] = FVector(961.8113, -18956.1173, 81673.443);
+		DecoratorAddRef1.VectorTArray = { FVector(-1927.87713, 1826.99173, -123.15553), FVector(9177.0113, -71.443, -917.883), FVector(123.913, 852.113, -81652.13) };
+		DecoratorAddRef1.String = TEXT("sample string 1233");
+		DecoratorAddRef1.Name = FName(TEXT("sample name 9991783"));
+
+		FDecoratorNativeSerialization_Add::FSharedData DecoratorNativeRef1;
+		DecoratorNativeRef1.Integer = 16515;
+		DecoratorNativeRef1.IntegerArray[0] = 10715;
+		DecoratorNativeRef1.IntegerArray[1] = -1585;
+		DecoratorNativeRef1.IntegerArray[2] = 881165;
+		DecoratorNativeRef1.IntegerArray[3] = 0x4175;
+		DecoratorNativeRef1.IntegerTArray = { -81625, 881525, 0x8152f5 };
+		DecoratorNativeRef1.Vector = FVector(0.19175, 12435.15, -18200.17265);
+		DecoratorNativeRef1.VectorArray[0] = FVector(192.17165, -1927.1155, 99176.125);
+		DecoratorNativeRef1.VectorArray[1] = FVector(961.8115, -18956.1175, 81673.445);
+		DecoratorNativeRef1.VectorTArray = { FVector(-1927.87715, 1826.99175, -123.15555), FVector(9177.0115, -71.445, -917.885), FVector(123.915, 852.115, -81652.15) };
+		DecoratorNativeRef1.String = TEXT("sample string 1235");
+		DecoratorNativeRef1.Name = FName(TEXT("sample name 9991785"));
 
-		AddErrorIfFalse(DecoratorWriter.GetErrorState() == FDecoratorWriter::EErrorState::None, "FAnimationAnimNextRuntimeTest_DecoratorSerialization -> Failed to write decorators");
-		GraphSharedDataArchiveBuffer = DecoratorWriter.GetGraphSharedData();
-		GraphReferencedObjects = DecoratorWriter.GetGraphReferencedObjects();
-	}
+		TArray<FNodeHandle> NodeHandles;
 
-	// Clear out the node template registry to test registration on load
-	{
-		FScopedClearNodeTemplateRegistry ScopedClearNodeTemplateRegistryForLoad;
+		// Write our graph
+		TArray<uint8> GraphSharedDataArchiveBuffer;
+		TArray<TObjectPtr<UObject>> GraphReferencedObjects;
+		{
+			FDecoratorWriter DecoratorWriter;
 
-		AddErrorIfFalse(Registry.GetNum() == 0, "FAnimationAnimNextRuntimeTest_DecoratorSerialization -> Registry should contain 0 templates");
+			NodeHandles.Add(DecoratorWriter.RegisterNode(*NodeTemplate0));
+			NodeHandles.Add(DecoratorWriter.RegisterNode(*NodeTemplate0));
 
-		// Read our graph
-		FTestUtils::LoadFromArchiveBuffer(*AnimNextGraph, NodeHandles, GraphSharedDataArchiveBuffer);
+			// We don't have decorator properties
+			TArray<TMap<FName, FString>> DecoratorProperties0;
+			DecoratorProperties0.AddDefaulted(NodeTemplateDecoratorList.Num());
+
+			DecoratorProperties0[0].Add(TEXT("Integer"), ToString<FDecoratorSerialization_Base::FSharedData>(TEXT("Integer"), DecoratorBaseRef0.Integer));
+			DecoratorProperties0[0].Add(TEXT("IntegerArray"), ToString<FDecoratorSerialization_Base::FSharedData>(TEXT("IntegerArray"), DecoratorBaseRef0.IntegerArray));
+			DecoratorProperties0[0].Add(TEXT("IntegerTArray"), ToString<FDecoratorSerialization_Base::FSharedData>(TEXT("IntegerTArray"), DecoratorBaseRef0.IntegerTArray));
+			DecoratorProperties0[0].Add(TEXT("Vector"), ToString<FDecoratorSerialization_Base::FSharedData>(TEXT("Vector"), DecoratorBaseRef0.Vector));
+			DecoratorProperties0[0].Add(TEXT("VectorArray"), ToString<FDecoratorSerialization_Base::FSharedData>(TEXT("VectorArray"), DecoratorBaseRef0.VectorArray));
+			DecoratorProperties0[0].Add(TEXT("VectorTArray"), ToString<FDecoratorSerialization_Base::FSharedData>(TEXT("VectorTArray"), DecoratorBaseRef0.VectorTArray));
+			DecoratorProperties0[0].Add(TEXT("String"), ToString<FDecoratorSerialization_Base::FSharedData>(TEXT("String"), DecoratorBaseRef0.String));
+			DecoratorProperties0[0].Add(TEXT("Name"), ToString<FDecoratorSerialization_Base::FSharedData>(TEXT("Name"), DecoratorBaseRef0.Name));
+
+			DecoratorProperties0[1].Add(TEXT("Integer"), ToString<FDecoratorSerialization_Add::FSharedData>(TEXT("Integer"), DecoratorAddRef0.Integer));
+			DecoratorProperties0[1].Add(TEXT("IntegerArray"), ToString<FDecoratorSerialization_Add::FSharedData>(TEXT("IntegerArray"), DecoratorAddRef0.IntegerArray));
+			DecoratorProperties0[1].Add(TEXT("IntegerTArray"), ToString<FDecoratorSerialization_Add::FSharedData>(TEXT("IntegerTArray"), DecoratorAddRef0.IntegerTArray));
+			DecoratorProperties0[1].Add(TEXT("Vector"), ToString<FDecoratorSerialization_Add::FSharedData>(TEXT("Vector"), DecoratorAddRef0.Vector));
+			DecoratorProperties0[1].Add(TEXT("VectorArray"), ToString<FDecoratorSerialization_Add::FSharedData>(TEXT("VectorArray"), DecoratorAddRef0.VectorArray));
+			DecoratorProperties0[1].Add(TEXT("VectorTArray"), ToString<FDecoratorSerialization_Add::FSharedData>(TEXT("VectorTArray"), DecoratorAddRef0.VectorTArray));
+			DecoratorProperties0[1].Add(TEXT("String"), ToString<FDecoratorSerialization_Add::FSharedData>(TEXT("String"), DecoratorAddRef0.String));
+			DecoratorProperties0[1].Add(TEXT("Name"), ToString<FDecoratorSerialization_Add::FSharedData>(TEXT("Name"), DecoratorAddRef0.Name));
+
+			DecoratorProperties0[2].Add(TEXT("Integer"), ToString<FDecoratorNativeSerialization_Add::FSharedData>(TEXT("Integer"), DecoratorNativeRef0.Integer));
+			DecoratorProperties0[2].Add(TEXT("IntegerArray"), ToString<FDecoratorNativeSerialization_Add::FSharedData>(TEXT("IntegerArray"), DecoratorNativeRef0.IntegerArray));
+			DecoratorProperties0[2].Add(TEXT("IntegerTArray"), ToString<FDecoratorNativeSerialization_Add::FSharedData>(TEXT("IntegerTArray"), DecoratorNativeRef0.IntegerTArray));
+			DecoratorProperties0[2].Add(TEXT("Vector"), ToString<FDecoratorNativeSerialization_Add::FSharedData>(TEXT("Vector"), DecoratorNativeRef0.Vector));
+			DecoratorProperties0[2].Add(TEXT("VectorArray"), ToString<FDecoratorNativeSerialization_Add::FSharedData>(TEXT("VectorArray"), DecoratorNativeRef0.VectorArray));
+			DecoratorProperties0[2].Add(TEXT("VectorTArray"), ToString<FDecoratorNativeSerialization_Add::FSharedData>(TEXT("VectorTArray"), DecoratorNativeRef0.VectorTArray));
+			DecoratorProperties0[2].Add(TEXT("String"), ToString<FDecoratorNativeSerialization_Add::FSharedData>(TEXT("String"), DecoratorNativeRef0.String));
+			DecoratorProperties0[2].Add(TEXT("Name"), ToString<FDecoratorNativeSerialization_Add::FSharedData>(TEXT("Name"), DecoratorNativeRef0.Name));
+
+			TArray<TMap<FName, FString>> DecoratorProperties1;
+			DecoratorProperties1.AddDefaulted(NodeTemplateDecoratorList.Num());
+
+			DecoratorProperties1[0].Add(TEXT("Integer"), ToString<FDecoratorSerialization_Base::FSharedData>(TEXT("Integer"), DecoratorBaseRef1.Integer));
+			DecoratorProperties1[0].Add(TEXT("IntegerArray"), ToString<FDecoratorSerialization_Base::FSharedData>(TEXT("IntegerArray"), DecoratorBaseRef1.IntegerArray));
+			DecoratorProperties1[0].Add(TEXT("IntegerTArray"), ToString<FDecoratorSerialization_Base::FSharedData>(TEXT("IntegerTArray"), DecoratorBaseRef1.IntegerTArray));
+			DecoratorProperties1[0].Add(TEXT("Vector"), ToString<FDecoratorSerialization_Base::FSharedData>(TEXT("Vector"), DecoratorBaseRef1.Vector));
+			DecoratorProperties1[0].Add(TEXT("VectorArray"), ToString<FDecoratorSerialization_Base::FSharedData>(TEXT("VectorArray"), DecoratorBaseRef1.VectorArray));
+			DecoratorProperties1[0].Add(TEXT("VectorTArray"), ToString<FDecoratorSerialization_Base::FSharedData>(TEXT("VectorTArray"), DecoratorBaseRef1.VectorTArray));
+			DecoratorProperties1[0].Add(TEXT("String"), ToString<FDecoratorSerialization_Base::FSharedData>(TEXT("String"), DecoratorBaseRef1.String));
+			DecoratorProperties1[0].Add(TEXT("Name"), ToString<FDecoratorSerialization_Base::FSharedData>(TEXT("Name"), DecoratorBaseRef1.Name));
+
+			DecoratorProperties1[1].Add(TEXT("Integer"), ToString<FDecoratorSerialization_Add::FSharedData>(TEXT("Integer"), DecoratorAddRef1.Integer));
+			DecoratorProperties1[1].Add(TEXT("IntegerArray"), ToString<FDecoratorSerialization_Add::FSharedData>(TEXT("IntegerArray"), DecoratorAddRef1.IntegerArray));
+			DecoratorProperties1[1].Add(TEXT("IntegerTArray"), ToString<FDecoratorSerialization_Add::FSharedData>(TEXT("IntegerTArray"), DecoratorAddRef1.IntegerTArray));
+			DecoratorProperties1[1].Add(TEXT("Vector"), ToString<FDecoratorSerialization_Add::FSharedData>(TEXT("Vector"), DecoratorAddRef1.Vector));
+			DecoratorProperties1[1].Add(TEXT("VectorArray"), ToString<FDecoratorSerialization_Add::FSharedData>(TEXT("VectorArray"), DecoratorAddRef1.VectorArray));
+			DecoratorProperties1[1].Add(TEXT("VectorTArray"), ToString<FDecoratorSerialization_Add::FSharedData>(TEXT("VectorTArray"), DecoratorAddRef1.VectorTArray));
+			DecoratorProperties1[1].Add(TEXT("String"), ToString<FDecoratorSerialization_Add::FSharedData>(TEXT("String"), DecoratorAddRef1.String));
+			DecoratorProperties1[1].Add(TEXT("Name"), ToString<FDecoratorSerialization_Add::FSharedData>(TEXT("Name"), DecoratorAddRef1.Name));
+
+			DecoratorProperties1[2].Add(TEXT("Integer"), ToString<FDecoratorNativeSerialization_Add::FSharedData>(TEXT("Integer"), DecoratorNativeRef1.Integer));
+			DecoratorProperties1[2].Add(TEXT("IntegerArray"), ToString<FDecoratorNativeSerialization_Add::FSharedData>(TEXT("IntegerArray"), DecoratorNativeRef1.IntegerArray));
+			DecoratorProperties1[2].Add(TEXT("IntegerTArray"), ToString<FDecoratorNativeSerialization_Add::FSharedData>(TEXT("IntegerTArray"), DecoratorNativeRef1.IntegerTArray));
+			DecoratorProperties1[2].Add(TEXT("Vector"), ToString<FDecoratorNativeSerialization_Add::FSharedData>(TEXT("Vector"), DecoratorNativeRef1.Vector));
+			DecoratorProperties1[2].Add(TEXT("VectorArray"), ToString<FDecoratorNativeSerialization_Add::FSharedData>(TEXT("VectorArray"), DecoratorNativeRef1.VectorArray));
+			DecoratorProperties1[2].Add(TEXT("VectorTArray"), ToString<FDecoratorNativeSerialization_Add::FSharedData>(TEXT("VectorTArray"), DecoratorNativeRef1.VectorTArray));
+			DecoratorProperties1[2].Add(TEXT("String"), ToString<FDecoratorNativeSerialization_Add::FSharedData>(TEXT("String"), DecoratorNativeRef1.String));
+			DecoratorProperties1[2].Add(TEXT("Name"), ToString<FDecoratorNativeSerialization_Add::FSharedData>(TEXT("Name"), DecoratorNativeRef1.Name));
 
-		FAnimNextGraphInstancePtr GraphInstance;
-		AnimNextGraph->AllocateInstance(GraphInstance);
+			DecoratorWriter.BeginNodeWriting();
+			DecoratorWriter.WriteNode(NodeHandles[0],
+				[&DecoratorProperties0](uint32 DecoratorIndex, FName PropertyName)
+				{
+					return DecoratorProperties0[DecoratorIndex][PropertyName];
+				},
+				[](uint32 DecoratorIndex, FName PropertyName)
+				{
+					return MAX_uint16;
+				});
+			DecoratorWriter.WriteNode(NodeHandles[1],
+				[&DecoratorProperties1](uint32 DecoratorIndex, FName PropertyName)
+				{
+					return DecoratorProperties1[DecoratorIndex][PropertyName];
+				},
+				[](uint32 DecoratorIndex, FName PropertyName)
+				{
+					return MAX_uint16;
+				});
+			DecoratorWriter.EndNodeWriting();
 
-		FExecutionContext Context(GraphInstance);
+			AddErrorIfFalse(DecoratorWriter.GetErrorState() == FDecoratorWriter::EErrorState::None, "FAnimationAnimNextRuntimeTest_DecoratorSerialization -> Failed to write decorators");
+			GraphSharedDataArchiveBuffer = DecoratorWriter.GetGraphSharedData();
+			GraphReferencedObjects = DecoratorWriter.GetGraphReferencedObjects();
+		}
 
-		// Validate decorator serialization
+		// Clear out the node template registry to test registration on load
 		{
-			FDecoratorBinding ParentBinding;								// Empty, no parent
-			FAnimNextDecoratorHandle DecoratorHandle0(NodeHandles[0], 0);	// Point to first node, first base decorator
-			FAnimNextDecoratorHandle DecoratorHandle1(NodeHandles[1], 0);	// Point to second node, first base decorator
+			FScopedClearNodeTemplateRegistry ScopedClearNodeTemplateRegistryForLoad;
 
-			FDecoratorPtr DecoratorPtr0 = Context.AllocateNodeInstance(ParentBinding, DecoratorHandle0);
-			AddErrorIfFalse(DecoratorPtr0.IsValid(), "FAnimationAnimNextRuntimeTest_DecoratorSerialization -> Failed to allocate a node instance");
+			AddErrorIfFalse(Registry.GetNum() == 0, "FAnimationAnimNextRuntimeTest_DecoratorSerialization -> Registry should contain 0 templates");
 
-			FDecoratorPtr DecoratorPtr1 = Context.AllocateNodeInstance(ParentBinding, DecoratorHandle1);
-			AddErrorIfFalse(DecoratorPtr1.IsValid(), "FAnimationAnimNextRuntimeTest_DecoratorSerialization -> Failed to allocate a node instance");
+			// Read our graph
+			FTestUtils::LoadFromArchiveBuffer(*AnimNextGraph, NodeHandles, GraphSharedDataArchiveBuffer);
 
-			// Validate shared data for base decorator on node 0
-			{
-				TDecoratorBinding<IInterfaceA> BindingA0;
-				AddErrorIfFalse(Context.GetInterface(DecoratorPtr0, BindingA0), "FAnimationAnimNextRuntimeTest_DecoratorSerialization -> InterfaceA not found");
-
-				const auto* SharedDataA0 = BindingA0.GetSharedData<FDecoratorSerialization_Base::FSharedData>();
-
-				AddErrorIfFalse(SharedDataA0->Integer == DecoratorBaseRef0.Integer, "FAnimationAnimNextRuntimeTest_DecoratorSerialization -> Unexpected serialized value");
-				AddErrorIfFalse(FMemory::Memcmp(SharedDataA0->IntegerArray, DecoratorBaseRef0.IntegerArray, sizeof(DecoratorBaseRef0.IntegerArray)) == 0, "FAnimationAnimNextRuntimeTest_DecoratorSerialization -> Unexpected serialized value");
-				AddErrorIfFalse(SharedDataA0->IntegerTArray == DecoratorBaseRef0.IntegerTArray, "FAnimationAnimNextRuntimeTest_DecoratorSerialization -> Unexpected serialized value");
-				AddErrorIfFalse(SharedDataA0->Vector == DecoratorBaseRef0.Vector, "FAnimationAnimNextRuntimeTest_DecoratorSerialization -> Unexpected serialized value");
-				AddErrorIfFalse(FMemory::Memcmp(SharedDataA0->VectorArray, DecoratorBaseRef0.VectorArray, sizeof(DecoratorBaseRef0.VectorArray)) == 0, "FAnimationAnimNextRuntimeTest_DecoratorSerialization -> Unexpected serialized value");
-				AddErrorIfFalse(SharedDataA0->VectorTArray == DecoratorBaseRef0.VectorTArray, "FAnimationAnimNextRuntimeTest_DecoratorSerialization -> Unexpected serialized value");
-				AddErrorIfFalse(SharedDataA0->String == DecoratorBaseRef0.String, "FAnimationAnimNextRuntimeTest_DecoratorSerialization -> Unexpected serialized value");
-				AddErrorIfFalse(SharedDataA0->Name == DecoratorBaseRef0.Name, "FAnimationAnimNextRuntimeTest_DecoratorSerialization -> Unexpected serialized value");
-			}
+			FAnimNextGraphInstancePtr GraphInstance;
+			AnimNextGraph->AllocateInstance(GraphInstance);
 
-			// Validate shared data for additive decorator on node 0
-			{
-				TDecoratorBinding<IInterfaceB> BindingB0;
-				AddErrorIfFalse(Context.GetInterface(DecoratorPtr0, BindingB0), "FAnimationAnimNextRuntimeTest_DecoratorSerialization -> InterfaceB not found");
-
-				const auto* SharedDataB0 = BindingB0.GetSharedData<FDecoratorSerialization_Add::FSharedData>();
-
-				AddErrorIfFalse(SharedDataB0->Integer == DecoratorAddRef0.Integer, "FAnimationAnimNextRuntimeTest_DecoratorSerialization -> Unexpected serialized value");
-				AddErrorIfFalse(FMemory::Memcmp(SharedDataB0->IntegerArray, DecoratorAddRef0.IntegerArray, sizeof(DecoratorAddRef0.IntegerArray)) == 0, "FAnimationAnimNextRuntimeTest_DecoratorSerialization -> Unexpected serialized value");
-				AddErrorIfFalse(SharedDataB0->IntegerTArray == DecoratorAddRef0.IntegerTArray, "FAnimationAnimNextRuntimeTest_DecoratorSerialization -> Unexpected serialized value");
-				AddErrorIfFalse(SharedDataB0->Vector == DecoratorAddRef0.Vector, "FAnimationAnimNextRuntimeTest_DecoratorSerialization -> Unexpected serialized value");
-				AddErrorIfFalse(FMemory::Memcmp(SharedDataB0->VectorArray, DecoratorAddRef0.VectorArray, sizeof(DecoratorAddRef0.VectorArray)) == 0, "FAnimationAnimNextRuntimeTest_DecoratorSerialization -> Unexpected serialized value");
-				AddErrorIfFalse(SharedDataB0->VectorTArray == DecoratorAddRef0.VectorTArray, "FAnimationAnimNextRuntimeTest_DecoratorSerialization -> Unexpected serialized value");
-				AddErrorIfFalse(SharedDataB0->String == DecoratorAddRef0.String, "FAnimationAnimNextRuntimeTest_DecoratorSerialization -> Unexpected serialized value");
-				AddErrorIfFalse(SharedDataB0->Name == DecoratorAddRef0.Name, "FAnimationAnimNextRuntimeTest_DecoratorSerialization -> Unexpected serialized value");
-			}
+			FExecutionContext Context(GraphInstance);
 
-			// Validate shared data for native decorator on node 0
+			// Validate decorator serialization
 			{
-				TDecoratorBinding<IInterfaceC> BindingC0;
-				AddErrorIfFalse(Context.GetInterface(DecoratorPtr0, BindingC0), "FAnimationAnimNextRuntimeTest_DecoratorSerialization -> InterfaceC not found");
-
-				const auto* SharedDataC0 = BindingC0.GetSharedData<FDecoratorNativeSerialization_Add::FSharedData>();
-
-				AddErrorIfFalse(SharedDataC0->Integer == DecoratorNativeRef0.Integer, "FAnimationAnimNextRuntimeTest_DecoratorSerialization -> Unexpected serialized value");
-				AddErrorIfFalse(FMemory::Memcmp(SharedDataC0->IntegerArray, DecoratorNativeRef0.IntegerArray, sizeof(DecoratorNativeRef0.IntegerArray)) == 0, "FAnimationAnimNextRuntimeTest_DecoratorSerialization -> Unexpected serialized value");
-				AddErrorIfFalse(SharedDataC0->IntegerTArray == DecoratorNativeRef0.IntegerTArray, "FAnimationAnimNextRuntimeTest_DecoratorSerialization -> Unexpected serialized value");
-				AddErrorIfFalse(SharedDataC0->Vector == DecoratorNativeRef0.Vector, "FAnimationAnimNextRuntimeTest_DecoratorSerialization -> Unexpected serialized value");
-				AddErrorIfFalse(FMemory::Memcmp(SharedDataC0->VectorArray, DecoratorNativeRef0.VectorArray, sizeof(DecoratorNativeRef0.VectorArray)) == 0, "FAnimationAnimNextRuntimeTest_DecoratorSerialization -> Unexpected serialized value");
-				AddErrorIfFalse(SharedDataC0->VectorTArray == DecoratorNativeRef0.VectorTArray, "FAnimationAnimNextRuntimeTest_DecoratorSerialization -> Unexpected serialized value");
-				AddErrorIfFalse(SharedDataC0->String == DecoratorNativeRef0.String, "FAnimationAnimNextRuntimeTest_DecoratorSerialization -> Unexpected serialized value");
-				AddErrorIfFalse(SharedDataC0->Name == DecoratorNativeRef0.Name, "FAnimationAnimNextRuntimeTest_DecoratorSerialization -> Unexpected serialized value");
-				AddErrorIfFalse(SharedDataC0->bSerializeCalled, "FAnimationAnimNextRuntimeTest_DecoratorSerialization -> Unexpected serialized value");
-			}
+				FDecoratorBinding ParentBinding;								// Empty, no parent
+				FAnimNextDecoratorHandle DecoratorHandle0(NodeHandles[0], 0);	// Point to first node, first base decorator
+				FAnimNextDecoratorHandle DecoratorHandle1(NodeHandles[1], 0);	// Point to second node, first base decorator
 
-			// Validate shared data for base decorator on node 1
-			{
-				TDecoratorBinding<IInterfaceA> BindingA1;
-				AddErrorIfFalse(Context.GetInterface(DecoratorPtr1, BindingA1), "FAnimationAnimNextRuntimeTest_DecoratorSerialization -> InterfaceA not found");
-
-				const auto* SharedDataA1 = BindingA1.GetSharedData<FDecoratorSerialization_Base::FSharedData>();
-
-				AddErrorIfFalse(SharedDataA1->Integer == DecoratorBaseRef1.Integer, "FAnimationAnimNextRuntimeTest_DecoratorSerialization -> Unexpected serialized value");
-				AddErrorIfFalse(FMemory::Memcmp(SharedDataA1->IntegerArray, DecoratorBaseRef1.IntegerArray, sizeof(DecoratorBaseRef1.IntegerArray)) == 0, "FAnimationAnimNextRuntimeTest_DecoratorSerialization -> Unexpected serialized value");
-				AddErrorIfFalse(SharedDataA1->IntegerTArray == DecoratorBaseRef1.IntegerTArray, "FAnimationAnimNextRuntimeTest_DecoratorSerialization -> Unexpected serialized value");
-				AddErrorIfFalse(SharedDataA1->Vector == DecoratorBaseRef1.Vector, "FAnimationAnimNextRuntimeTest_DecoratorSerialization -> Unexpected serialized value");
-				AddErrorIfFalse(FMemory::Memcmp(SharedDataA1->VectorArray, DecoratorBaseRef1.VectorArray, sizeof(DecoratorBaseRef1.VectorArray)) == 0, "FAnimationAnimNextRuntimeTest_DecoratorSerialization -> Unexpected serialized value");
-				AddErrorIfFalse(SharedDataA1->VectorTArray == DecoratorBaseRef1.VectorTArray, "FAnimationAnimNextRuntimeTest_DecoratorSerialization -> Unexpected serialized value");
-				AddErrorIfFalse(SharedDataA1->String == DecoratorBaseRef1.String, "FAnimationAnimNextRuntimeTest_DecoratorSerialization -> Unexpected serialized value");
-				AddErrorIfFalse(SharedDataA1->Name == DecoratorBaseRef1.Name, "FAnimationAnimNextRuntimeTest_DecoratorSerialization -> Unexpected serialized value");
-			}
+				FDecoratorPtr DecoratorPtr0 = Context.AllocateNodeInstance(ParentBinding, DecoratorHandle0);
+				AddErrorIfFalse(DecoratorPtr0.IsValid(), "FAnimationAnimNextRuntimeTest_DecoratorSerialization -> Failed to allocate a node instance");
 
-			// Validate shared data for additive decorator on node 1
-			{
-				TDecoratorBinding<IInterfaceB> BindingB1;
-				AddErrorIfFalse(Context.GetInterface(DecoratorPtr1, BindingB1), "FAnimationAnimNextRuntimeTest_DecoratorSerialization -> InterfaceB not found");
-
-				const auto* SharedDataB1 = BindingB1.GetSharedData<FDecoratorSerialization_Add::FSharedData>();
-
-				AddErrorIfFalse(SharedDataB1->Integer == DecoratorAddRef1.Integer, "FAnimationAnimNextRuntimeTest_DecoratorSerialization -> Unexpected serialized value");
-				AddErrorIfFalse(FMemory::Memcmp(SharedDataB1->IntegerArray, DecoratorAddRef1.IntegerArray, sizeof(DecoratorAddRef1.IntegerArray)) == 0, "FAnimationAnimNextRuntimeTest_DecoratorSerialization -> Unexpected serialized value");
-				AddErrorIfFalse(SharedDataB1->IntegerTArray == DecoratorAddRef1.IntegerTArray, "FAnimationAnimNextRuntimeTest_DecoratorSerialization -> Unexpected serialized value");
-				AddErrorIfFalse(SharedDataB1->Vector == DecoratorAddRef1.Vector, "FAnimationAnimNextRuntimeTest_DecoratorSerialization -> Unexpected serialized value");
-				AddErrorIfFalse(FMemory::Memcmp(SharedDataB1->VectorArray, DecoratorAddRef1.VectorArray, sizeof(DecoratorAddRef1.VectorArray)) == 0, "FAnimationAnimNextRuntimeTest_DecoratorSerialization -> Unexpected serialized value");
-				AddErrorIfFalse(SharedDataB1->VectorTArray == DecoratorAddRef1.VectorTArray, "FAnimationAnimNextRuntimeTest_DecoratorSerialization -> Unexpected serialized value");
-				AddErrorIfFalse(SharedDataB1->String == DecoratorAddRef1.String, "FAnimationAnimNextRuntimeTest_DecoratorSerialization -> Unexpected serialized value");
-				AddErrorIfFalse(SharedDataB1->Name == DecoratorAddRef1.Name, "FAnimationAnimNextRuntimeTest_DecoratorSerialization -> Unexpected serialized value");
-			}
+				FDecoratorPtr DecoratorPtr1 = Context.AllocateNodeInstance(ParentBinding, DecoratorHandle1);
+				AddErrorIfFalse(DecoratorPtr1.IsValid(), "FAnimationAnimNextRuntimeTest_DecoratorSerialization -> Failed to allocate a node instance");
 
-			// Validate shared data for native decorator on node 1
-			{
-				TDecoratorBinding<IInterfaceC> BindingC1;
-				AddErrorIfFalse(Context.GetInterface(DecoratorPtr1, BindingC1), "FAnimationAnimNextRuntimeTest_DecoratorSerialization -> InterfaceC not found");
-
-				const auto* SharedDataC1 = BindingC1.GetSharedData<FDecoratorNativeSerialization_Add::FSharedData>();
-
-				AddErrorIfFalse(SharedDataC1->Integer == DecoratorNativeRef1.Integer, "FAnimationAnimNextRuntimeTest_DecoratorSerialization -> Unexpected serialized value");
-				AddErrorIfFalse(FMemory::Memcmp(SharedDataC1->IntegerArray, DecoratorNativeRef1.IntegerArray, sizeof(DecoratorNativeRef1.IntegerArray)) == 0, "FAnimationAnimNextRuntimeTest_DecoratorSerialization -> Unexpected serialized value");
-				AddErrorIfFalse(SharedDataC1->IntegerTArray == DecoratorNativeRef1.IntegerTArray, "FAnimationAnimNextRuntimeTest_DecoratorSerialization -> Unexpected serialized value");
-				AddErrorIfFalse(SharedDataC1->Vector == DecoratorNativeRef1.Vector, "FAnimationAnimNextRuntimeTest_DecoratorSerialization -> Unexpected serialized value");
-				AddErrorIfFalse(FMemory::Memcmp(SharedDataC1->VectorArray, DecoratorNativeRef1.VectorArray, sizeof(DecoratorNativeRef1.VectorArray)) == 0, "FAnimationAnimNextRuntimeTest_DecoratorSerialization -> Unexpected serialized value");
-				AddErrorIfFalse(SharedDataC1->VectorTArray == DecoratorNativeRef1.VectorTArray, "FAnimationAnimNextRuntimeTest_DecoratorSerialization -> Unexpected serialized value");
-				AddErrorIfFalse(SharedDataC1->String == DecoratorNativeRef1.String, "FAnimationAnimNextRuntimeTest_DecoratorSerialization -> Unexpected serialized value");
-				AddErrorIfFalse(SharedDataC1->Name == DecoratorNativeRef1.Name, "FAnimationAnimNextRuntimeTest_DecoratorSerialization -> Unexpected serialized value");
-				AddErrorIfFalse(SharedDataC1->bSerializeCalled, "FAnimationAnimNextRuntimeTest_DecoratorSerialization -> Unexpected serialized value");
+				// Validate shared data for base decorator on node 0
+				{
+					TDecoratorBinding<IInterfaceA> BindingA0;
+					AddErrorIfFalse(Context.GetInterface(DecoratorPtr0, BindingA0), "FAnimationAnimNextRuntimeTest_DecoratorSerialization -> InterfaceA not found");
+
+					const auto* SharedDataA0 = BindingA0.GetSharedData<FDecoratorSerialization_Base::FSharedData>();
+
+					AddErrorIfFalse(SharedDataA0->Integer == DecoratorBaseRef0.Integer, "FAnimationAnimNextRuntimeTest_DecoratorSerialization -> Unexpected serialized value");
+					AddErrorIfFalse(FMemory::Memcmp(SharedDataA0->IntegerArray, DecoratorBaseRef0.IntegerArray, sizeof(DecoratorBaseRef0.IntegerArray)) == 0, "FAnimationAnimNextRuntimeTest_DecoratorSerialization -> Unexpected serialized value");
+					AddErrorIfFalse(SharedDataA0->IntegerTArray == DecoratorBaseRef0.IntegerTArray, "FAnimationAnimNextRuntimeTest_DecoratorSerialization -> Unexpected serialized value");
+					AddErrorIfFalse(SharedDataA0->Vector == DecoratorBaseRef0.Vector, "FAnimationAnimNextRuntimeTest_DecoratorSerialization -> Unexpected serialized value");
+					AddErrorIfFalse(FMemory::Memcmp(SharedDataA0->VectorArray, DecoratorBaseRef0.VectorArray, sizeof(DecoratorBaseRef0.VectorArray)) == 0, "FAnimationAnimNextRuntimeTest_DecoratorSerialization -> Unexpected serialized value");
+					AddErrorIfFalse(SharedDataA0->VectorTArray == DecoratorBaseRef0.VectorTArray, "FAnimationAnimNextRuntimeTest_DecoratorSerialization -> Unexpected serialized value");
+					AddErrorIfFalse(SharedDataA0->String == DecoratorBaseRef0.String, "FAnimationAnimNextRuntimeTest_DecoratorSerialization -> Unexpected serialized value");
+					AddErrorIfFalse(SharedDataA0->Name == DecoratorBaseRef0.Name, "FAnimationAnimNextRuntimeTest_DecoratorSerialization -> Unexpected serialized value");
+				}
+
+				// Validate shared data for additive decorator on node 0
+				{
+					TDecoratorBinding<IInterfaceB> BindingB0;
+					AddErrorIfFalse(Context.GetInterface(DecoratorPtr0, BindingB0), "FAnimationAnimNextRuntimeTest_DecoratorSerialization -> InterfaceB not found");
+
+					const auto* SharedDataB0 = BindingB0.GetSharedData<FDecoratorSerialization_Add::FSharedData>();
+
+					AddErrorIfFalse(SharedDataB0->Integer == DecoratorAddRef0.Integer, "FAnimationAnimNextRuntimeTest_DecoratorSerialization -> Unexpected serialized value");
+					AddErrorIfFalse(FMemory::Memcmp(SharedDataB0->IntegerArray, DecoratorAddRef0.IntegerArray, sizeof(DecoratorAddRef0.IntegerArray)) == 0, "FAnimationAnimNextRuntimeTest_DecoratorSerialization -> Unexpected serialized value");
+					AddErrorIfFalse(SharedDataB0->IntegerTArray == DecoratorAddRef0.IntegerTArray, "FAnimationAnimNextRuntimeTest_DecoratorSerialization -> Unexpected serialized value");
+					AddErrorIfFalse(SharedDataB0->Vector == DecoratorAddRef0.Vector, "FAnimationAnimNextRuntimeTest_DecoratorSerialization -> Unexpected serialized value");
+					AddErrorIfFalse(FMemory::Memcmp(SharedDataB0->VectorArray, DecoratorAddRef0.VectorArray, sizeof(DecoratorAddRef0.VectorArray)) == 0, "FAnimationAnimNextRuntimeTest_DecoratorSerialization -> Unexpected serialized value");
+					AddErrorIfFalse(SharedDataB0->VectorTArray == DecoratorAddRef0.VectorTArray, "FAnimationAnimNextRuntimeTest_DecoratorSerialization -> Unexpected serialized value");
+					AddErrorIfFalse(SharedDataB0->String == DecoratorAddRef0.String, "FAnimationAnimNextRuntimeTest_DecoratorSerialization -> Unexpected serialized value");
+					AddErrorIfFalse(SharedDataB0->Name == DecoratorAddRef0.Name, "FAnimationAnimNextRuntimeTest_DecoratorSerialization -> Unexpected serialized value");
+				}
+
+				// Validate shared data for native decorator on node 0
+				{
+					TDecoratorBinding<IInterfaceC> BindingC0;
+					AddErrorIfFalse(Context.GetInterface(DecoratorPtr0, BindingC0), "FAnimationAnimNextRuntimeTest_DecoratorSerialization -> InterfaceC not found");
+
+					const auto* SharedDataC0 = BindingC0.GetSharedData<FDecoratorNativeSerialization_Add::FSharedData>();
+
+					AddErrorIfFalse(SharedDataC0->Integer == DecoratorNativeRef0.Integer, "FAnimationAnimNextRuntimeTest_DecoratorSerialization -> Unexpected serialized value");
+					AddErrorIfFalse(FMemory::Memcmp(SharedDataC0->IntegerArray, DecoratorNativeRef0.IntegerArray, sizeof(DecoratorNativeRef0.IntegerArray)) == 0, "FAnimationAnimNextRuntimeTest_DecoratorSerialization -> Unexpected serialized value");
+					AddErrorIfFalse(SharedDataC0->IntegerTArray == DecoratorNativeRef0.IntegerTArray, "FAnimationAnimNextRuntimeTest_DecoratorSerialization -> Unexpected serialized value");
+					AddErrorIfFalse(SharedDataC0->Vector == DecoratorNativeRef0.Vector, "FAnimationAnimNextRuntimeTest_DecoratorSerialization -> Unexpected serialized value");
+					AddErrorIfFalse(FMemory::Memcmp(SharedDataC0->VectorArray, DecoratorNativeRef0.VectorArray, sizeof(DecoratorNativeRef0.VectorArray)) == 0, "FAnimationAnimNextRuntimeTest_DecoratorSerialization -> Unexpected serialized value");
+					AddErrorIfFalse(SharedDataC0->VectorTArray == DecoratorNativeRef0.VectorTArray, "FAnimationAnimNextRuntimeTest_DecoratorSerialization -> Unexpected serialized value");
+					AddErrorIfFalse(SharedDataC0->String == DecoratorNativeRef0.String, "FAnimationAnimNextRuntimeTest_DecoratorSerialization -> Unexpected serialized value");
+					AddErrorIfFalse(SharedDataC0->Name == DecoratorNativeRef0.Name, "FAnimationAnimNextRuntimeTest_DecoratorSerialization -> Unexpected serialized value");
+					AddErrorIfFalse(SharedDataC0->bSerializeCalled, "FAnimationAnimNextRuntimeTest_DecoratorSerialization -> Unexpected serialized value");
+				}
+
+				// Validate shared data for base decorator on node 1
+				{
+					TDecoratorBinding<IInterfaceA> BindingA1;
+					AddErrorIfFalse(Context.GetInterface(DecoratorPtr1, BindingA1), "FAnimationAnimNextRuntimeTest_DecoratorSerialization -> InterfaceA not found");
+
+					const auto* SharedDataA1 = BindingA1.GetSharedData<FDecoratorSerialization_Base::FSharedData>();
+
+					AddErrorIfFalse(SharedDataA1->Integer == DecoratorBaseRef1.Integer, "FAnimationAnimNextRuntimeTest_DecoratorSerialization -> Unexpected serialized value");
+					AddErrorIfFalse(FMemory::Memcmp(SharedDataA1->IntegerArray, DecoratorBaseRef1.IntegerArray, sizeof(DecoratorBaseRef1.IntegerArray)) == 0, "FAnimationAnimNextRuntimeTest_DecoratorSerialization -> Unexpected serialized value");
+					AddErrorIfFalse(SharedDataA1->IntegerTArray == DecoratorBaseRef1.IntegerTArray, "FAnimationAnimNextRuntimeTest_DecoratorSerialization -> Unexpected serialized value");
+					AddErrorIfFalse(SharedDataA1->Vector == DecoratorBaseRef1.Vector, "FAnimationAnimNextRuntimeTest_DecoratorSerialization -> Unexpected serialized value");
+					AddErrorIfFalse(FMemory::Memcmp(SharedDataA1->VectorArray, DecoratorBaseRef1.VectorArray, sizeof(DecoratorBaseRef1.VectorArray)) == 0, "FAnimationAnimNextRuntimeTest_DecoratorSerialization -> Unexpected serialized value");
+					AddErrorIfFalse(SharedDataA1->VectorTArray == DecoratorBaseRef1.VectorTArray, "FAnimationAnimNextRuntimeTest_DecoratorSerialization -> Unexpected serialized value");
+					AddErrorIfFalse(SharedDataA1->String == DecoratorBaseRef1.String, "FAnimationAnimNextRuntimeTest_DecoratorSerialization -> Unexpected serialized value");
+					AddErrorIfFalse(SharedDataA1->Name == DecoratorBaseRef1.Name, "FAnimationAnimNextRuntimeTest_DecoratorSerialization -> Unexpected serialized value");
+				}
+
+				// Validate shared data for additive decorator on node 1
+				{
+					TDecoratorBinding<IInterfaceB> BindingB1;
+					AddErrorIfFalse(Context.GetInterface(DecoratorPtr1, BindingB1), "FAnimationAnimNextRuntimeTest_DecoratorSerialization -> InterfaceB not found");
+
+					const auto* SharedDataB1 = BindingB1.GetSharedData<FDecoratorSerialization_Add::FSharedData>();
+
+					AddErrorIfFalse(SharedDataB1->Integer == DecoratorAddRef1.Integer, "FAnimationAnimNextRuntimeTest_DecoratorSerialization -> Unexpected serialized value");
+					AddErrorIfFalse(FMemory::Memcmp(SharedDataB1->IntegerArray, DecoratorAddRef1.IntegerArray, sizeof(DecoratorAddRef1.IntegerArray)) == 0, "FAnimationAnimNextRuntimeTest_DecoratorSerialization -> Unexpected serialized value");
+					AddErrorIfFalse(SharedDataB1->IntegerTArray == DecoratorAddRef1.IntegerTArray, "FAnimationAnimNextRuntimeTest_DecoratorSerialization -> Unexpected serialized value");
+					AddErrorIfFalse(SharedDataB1->Vector == DecoratorAddRef1.Vector, "FAnimationAnimNextRuntimeTest_DecoratorSerialization -> Unexpected serialized value");
+					AddErrorIfFalse(FMemory::Memcmp(SharedDataB1->VectorArray, DecoratorAddRef1.VectorArray, sizeof(DecoratorAddRef1.VectorArray)) == 0, "FAnimationAnimNextRuntimeTest_DecoratorSerialization -> Unexpected serialized value");
+					AddErrorIfFalse(SharedDataB1->VectorTArray == DecoratorAddRef1.VectorTArray, "FAnimationAnimNextRuntimeTest_DecoratorSerialization -> Unexpected serialized value");
+					AddErrorIfFalse(SharedDataB1->String == DecoratorAddRef1.String, "FAnimationAnimNextRuntimeTest_DecoratorSerialization -> Unexpected serialized value");
+					AddErrorIfFalse(SharedDataB1->Name == DecoratorAddRef1.Name, "FAnimationAnimNextRuntimeTest_DecoratorSerialization -> Unexpected serialized value");
+				}
+
+				// Validate shared data for native decorator on node 1
+				{
+					TDecoratorBinding<IInterfaceC> BindingC1;
+					AddErrorIfFalse(Context.GetInterface(DecoratorPtr1, BindingC1), "FAnimationAnimNextRuntimeTest_DecoratorSerialization -> InterfaceC not found");
+
+					const auto* SharedDataC1 = BindingC1.GetSharedData<FDecoratorNativeSerialization_Add::FSharedData>();
+
+					AddErrorIfFalse(SharedDataC1->Integer == DecoratorNativeRef1.Integer, "FAnimationAnimNextRuntimeTest_DecoratorSerialization -> Unexpected serialized value");
+					AddErrorIfFalse(FMemory::Memcmp(SharedDataC1->IntegerArray, DecoratorNativeRef1.IntegerArray, sizeof(DecoratorNativeRef1.IntegerArray)) == 0, "FAnimationAnimNextRuntimeTest_DecoratorSerialization -> Unexpected serialized value");
+					AddErrorIfFalse(SharedDataC1->IntegerTArray == DecoratorNativeRef1.IntegerTArray, "FAnimationAnimNextRuntimeTest_DecoratorSerialization -> Unexpected serialized value");
+					AddErrorIfFalse(SharedDataC1->Vector == DecoratorNativeRef1.Vector, "FAnimationAnimNextRuntimeTest_DecoratorSerialization -> Unexpected serialized value");
+					AddErrorIfFalse(FMemory::Memcmp(SharedDataC1->VectorArray, DecoratorNativeRef1.VectorArray, sizeof(DecoratorNativeRef1.VectorArray)) == 0, "FAnimationAnimNextRuntimeTest_DecoratorSerialization -> Unexpected serialized value");
+					AddErrorIfFalse(SharedDataC1->VectorTArray == DecoratorNativeRef1.VectorTArray, "FAnimationAnimNextRuntimeTest_DecoratorSerialization -> Unexpected serialized value");
+					AddErrorIfFalse(SharedDataC1->String == DecoratorNativeRef1.String, "FAnimationAnimNextRuntimeTest_DecoratorSerialization -> Unexpected serialized value");
+					AddErrorIfFalse(SharedDataC1->Name == DecoratorNativeRef1.Name, "FAnimationAnimNextRuntimeTest_DecoratorSerialization -> Unexpected serialized value");
+					AddErrorIfFalse(SharedDataC1->bSerializeCalled, "FAnimationAnimNextRuntimeTest_DecoratorSerialization -> Unexpected serialized value");
+				}
 			}
 		}
 	}
diff --git a/Engine/Plugins/Experimental/AnimNext/Source/AnimNextTestSuite/Private/AnimNextDecoratorGraphTest.cpp b/Engine/Plugins/Experimental/AnimNext/Source/AnimNextTestSuite/Private/AnimNextDecoratorGraphTest.cpp
--- a/Engine/Plugins/Experimental/AnimNext/Source/AnimNextTestSuite/Private/AnimNextDecoratorGraphTest.cpp
+++ b/Engine/Plugins/Experimental/AnimNext/Source/AnimNextTestSuite/Private/AnimNextDecoratorGraphTest.cpp
@@ -69,10 +69,13 @@ namespace UE::AnimNext
 		}
 	};
 
-	DEFINE_ANIM_DECORATOR_BEGIN(FTestDecorator)
-		DEFINE_ANIM_DECORATOR_IMPLEMENTS_INTERFACE(IEvaluate)
-		DEFINE_ANIM_DECORATOR_IMPLEMENTS_INTERFACE(IUpdate)
-	DEFINE_ANIM_DECORATOR_END(FTestDecorator)
+	// Decorator implementation boilerplate
+	#define DECORATOR_INTERFACE_ENUMERATOR(GeneratorMacro) \
+		GeneratorMacro(IEvaluate) \
+		GeneratorMacro(IUpdate) \
+
+	GENERATE_ANIM_DECORATOR_IMPLEMENTATION(FTestDecorator, DECORATOR_INTERFACE_ENUMERATOR)
+	#undef DECORATOR_INTERFACE_ENUMERATOR
 }
 
 IMPLEMENT_SIMPLE_AUTOMATION_TEST(FAnimationAnimNextRuntimeTest_GraphAddDecorator, "Animation.AnimNext.Runtime.Graph.AddDecorator", EAutomationTestFlags::EditorContext | EAutomationTestFlags::EngineFilter)
@@ -80,95 +83,98 @@ IMPLEMENT_SIMPLE_AUTOMATION_TEST(FAnimationAnimNextRuntimeTest_GraphAddDecorator
 bool FAnimationAnimNextRuntimeTest_GraphAddDecorator::RunTest(const FString& InParameters)
 {
 	using namespace UE::AnimNext;
-	AUTO_REGISTER_ANIM_DECORATOR(FTestDecorator)
 
-	FScopedClearNodeTemplateRegistry ScopedClearNodeTemplateRegistry;
+	{
+		AUTO_REGISTER_ANIM_DECORATOR(FTestDecorator)
 
-	UFactory* GraphFactory = NewObject<UAnimNextGraphFactory>();
-	UAnimNextGraph* AnimNextGraph = CastChecked<UAnimNextGraph>(GraphFactory->FactoryCreateNew(UAnimNextGraph::StaticClass(), GetTransientPackage(), TEXT("TestAnimNextGraph"), RF_Transient, nullptr, nullptr, NAME_None));
-	UE_RETURN_ON_ERROR(AnimNextGraph != nullptr, "FAnimationAnimNextRuntimeTest_GraphAddDecorator -> Failed to create animation graph");
+		FScopedClearNodeTemplateRegistry ScopedClearNodeTemplateRegistry;
 
-	UAnimNextGraph_EditorData* EditorData = UncookedOnly::FUtils::GetEditorData(AnimNextGraph);
-	UE_RETURN_ON_ERROR(EditorData != nullptr, "FAnimationAnimNextRuntimeTest_GraphAddDecorator -> Failed to find animation graph editor data");
+		UFactory* GraphFactory = NewObject<UAnimNextGraphFactory>();
+		UAnimNextGraph* AnimNextGraph = CastChecked<UAnimNextGraph>(GraphFactory->FactoryCreateNew(UAnimNextGraph::StaticClass(), GetTransientPackage(), TEXT("TestAnimNextGraph"), RF_Transient, nullptr, nullptr, NAME_None));
+		UE_RETURN_ON_ERROR(AnimNextGraph != nullptr, "FAnimationAnimNextRuntimeTest_GraphAddDecorator -> Failed to create animation graph");
 
-	URigVMController* Controller = EditorData->GetRigVMClient()->GetController(EditorData->GetRigVMClient()->GetDefaultModel());
-	UE_RETURN_ON_ERROR(Controller != nullptr, "FAnimationAnimNextRuntimeTest_GraphAddDecorator -> Failed to get RigVM controller");
+		UAnimNextGraph_EditorData* EditorData = UncookedOnly::FUtils::GetEditorData(AnimNextGraph);
+		UE_RETURN_ON_ERROR(EditorData != nullptr, "FAnimationAnimNextRuntimeTest_GraphAddDecorator -> Failed to find animation graph editor data");
 
-	// Create an empty decorator stack node
-	URigVMUnitNode* DecoratorStackNode = Controller->AddUnitNode(FRigUnit_AnimNextDecoratorStack::StaticStruct(), FRigVMStruct::ExecuteName, FVector2D(0.0f, 0.0f), FString(), false);
-	UE_RETURN_ON_ERROR(DecoratorStackNode != nullptr, "FAnimationAnimNextRuntimeTest_GraphAddDecorator -> Failed to create decorator stack node");
+		URigVMController* Controller = EditorData->GetRigVMClient()->GetController(EditorData->GetRigVMClient()->GetDefaultModel());
+		UE_RETURN_ON_ERROR(Controller != nullptr, "FAnimationAnimNextRuntimeTest_GraphAddDecorator -> Failed to get RigVM controller");
 
-	// Add a decorator
-	const UScriptStruct* CppDecoratorStruct = FRigDecorator_AnimNextCppDecorator::StaticStruct();
-	UE_RETURN_ON_ERROR(CppDecoratorStruct != nullptr, "FAnimationAnimNextRuntimeTest_GraphAddDecorator -> Failed to get find Cpp decorator static struct");
+		// Create an empty decorator stack node
+		URigVMUnitNode* DecoratorStackNode = Controller->AddUnitNode(FRigUnit_AnimNextDecoratorStack::StaticStruct(), FRigVMStruct::ExecuteName, FVector2D(0.0f, 0.0f), FString(), false);
+		UE_RETURN_ON_ERROR(DecoratorStackNode != nullptr, "FAnimationAnimNextRuntimeTest_GraphAddDecorator -> Failed to create decorator stack node");
 
-	const FDecorator* Decorator = FDecoratorRegistry::Get().Find(FTestDecorator::DecoratorUID);
-	UE_RETURN_ON_ERROR(Decorator != nullptr, "FAnimationAnimNextRuntimeTest_GraphAddDecorator -> Failed to find test decorator");
+		// Add a decorator
+		const UScriptStruct* CppDecoratorStruct = FRigDecorator_AnimNextCppDecorator::StaticStruct();
+		UE_RETURN_ON_ERROR(CppDecoratorStruct != nullptr, "FAnimationAnimNextRuntimeTest_GraphAddDecorator -> Failed to get find Cpp decorator static struct");
 
-	UScriptStruct* ScriptStruct = Decorator->GetDecoratorSharedDataStruct();
-	UE_RETURN_ON_ERROR(ScriptStruct != nullptr, "FAnimationAnimNextRuntimeTest_GraphAddDecorator -> Failed to find decorator shared data struct");
+		const FDecorator* Decorator = FDecoratorRegistry::Get().Find(FTestDecorator::DecoratorUID);
+		UE_RETURN_ON_ERROR(Decorator != nullptr, "FAnimationAnimNextRuntimeTest_GraphAddDecorator -> Failed to find test decorator");
 
-	FString DefaultValue;
-	{
-		const FRigDecorator_AnimNextCppDecorator DefaultCppDecoratorStructInstance;
-		FRigDecorator_AnimNextCppDecorator CppDecoratorStructInstance;
-		CppDecoratorStructInstance.DecoratorSharedDataStruct = ScriptStruct;
+		UScriptStruct* ScriptStruct = Decorator->GetDecoratorSharedDataStruct();
+		UE_RETURN_ON_ERROR(ScriptStruct != nullptr, "FAnimationAnimNextRuntimeTest_GraphAddDecorator -> Failed to find decorator shared data struct");
 
-		UE_RETURN_ON_ERROR(CppDecoratorStructInstance.CanBeAddedToNode(DecoratorStackNode, nullptr), "FAnimationAnimNextRuntimeTest_GraphAddDecorator -> Decorator cannot be added to decorator stack node");
+		FString DefaultValue;
+		{
+			const FRigDecorator_AnimNextCppDecorator DefaultCppDecoratorStructInstance;
+			FRigDecorator_AnimNextCppDecorator CppDecoratorStructInstance;
+			CppDecoratorStructInstance.DecoratorSharedDataStruct = ScriptStruct;
 
-		const FProperty* Prop = FAnimNextCppDecoratorWrapper::StaticStruct()->FindPropertyByName(GET_MEMBER_NAME_STRING_CHECKED(FAnimNextCppDecoratorWrapper, CppDecorator));
-		UE_RETURN_ON_ERROR(Prop != nullptr, "FAnimationAnimNextRuntimeTest_GraphAddDecorator -> Failed to find wrapper property");
+			UE_RETURN_ON_ERROR(CppDecoratorStructInstance.CanBeAddedToNode(DecoratorStackNode, nullptr), "FAnimationAnimNextRuntimeTest_GraphAddDecorator -> Decorator cannot be added to decorator stack node");
 
-		Prop->ExportText_Direct(DefaultValue, &CppDecoratorStructInstance, &DefaultCppDecoratorStructInstance, nullptr, PPF_None);
-	}
+			const FProperty* Prop = FAnimNextCppDecoratorWrapper::StaticStruct()->FindPropertyByName(GET_MEMBER_NAME_STRING_CHECKED(FAnimNextCppDecoratorWrapper, CppDecorator));
+			UE_RETURN_ON_ERROR(Prop != nullptr, "FAnimationAnimNextRuntimeTest_GraphAddDecorator -> Failed to find wrapper property");
 
-	FString DisplayNameMetadata;
-	ScriptStruct->GetStringMetaDataHierarchical(FRigVMStruct::DisplayNameMetaName, &DisplayNameMetadata);
-	const FString DisplayName = DisplayNameMetadata.IsEmpty() ? Decorator->GetDecoratorName() : DisplayNameMetadata;
-
-	const FName DecoratorName = Controller->AddDecorator(
-		DecoratorStackNode->GetFName(),
-		*CppDecoratorStruct->GetPathName(),
-		*DisplayName,
-		DefaultValue, INDEX_NONE, true, true);
-	UE_RETURN_ON_ERROR(DecoratorName == DisplayName, TEXT("FAnimationAnimNextRuntimeTest_GraphAddDecorator -> Unexpected decorator name"));
-
-	URigVMPin* DecoratorPin = DecoratorStackNode->FindPin(*DisplayName);
-	UE_RETURN_ON_ERROR(DecoratorPin != nullptr, TEXT("FAnimationAnimNextRuntimeTest_GraphAddDecorator -> Failed to find decorator pin"));
-
-	// Our first pin is the hard coded output result, decorator pins follow
-	UE_RETURN_ON_ERROR(DecoratorStackNode->GetPins().Num() == 2, TEXT("FAnimationAnimNextRuntimeTest_GraphAddDecorator -> Unexpected number of pins"));
-	UE_RETURN_ON_ERROR(DecoratorPin->IsDecoratorPin() == true, TEXT("FAnimationAnimNextRuntimeTest_GraphAddDecorator -> Unexpected pin type"));
-	UE_RETURN_ON_ERROR(DecoratorPin->GetFName() == DecoratorName, TEXT("FAnimationAnimNextRuntimeTest_GraphAddDecorator -> Unexpected pin name"));
-	UE_RETURN_ON_ERROR(DecoratorPin->GetCPPTypeObject() == FRigDecorator_AnimNextCppDecorator::StaticStruct(), TEXT("FAnimationAnimNextRuntimeTest_GraphAddDecorator -> Unexpected pin type"));
-
-	// Our first sub-pin is the hard coded script struct member that parametrizes the decorator, dynamic decorator sub-pins follow
-	UE_RETURN_ON_ERROR(DecoratorPin->GetSubPins().Num() == 6, TEXT("FAnimationAnimNextRuntimeTest_GraphAddDecorator -> Unexpected decorator sub pins"));
-
-	// SomeInt32
-	UE_RETURN_ON_ERROR(DecoratorPin->GetSubPins()[1]->GetCPPType() == TEXT("int32"), TEXT("FAnimationAnimNextRuntimeTest_GraphAddDecorator -> Unexpected decorator pin type"));
-	UE_RETURN_ON_ERROR(DecoratorPin->GetSubPins()[1]->GetDefaultValue() == TEXT("3"), TEXT("FAnimationAnimNextRuntimeTest_GraphAddDecorator -> Unexpected decorator pin value"));
-	UE_RETURN_ON_ERROR(!DecoratorPin->GetSubPins()[1]->IsLazy(), TEXT("FAnimationAnimNextRuntimeTest_GraphAddDecorator -> Expected non-lazy decorator pin"));
-
-	// SomeFloat
-	UE_RETURN_ON_ERROR(DecoratorPin->GetSubPins()[2]->GetCPPType() == TEXT("float"), TEXT("FAnimationAnimNextRuntimeTest_GraphAddDecorator -> Unexpected decorator pin type"));
-	UE_RETURN_ON_ERROR(DecoratorPin->GetSubPins()[2]->GetDefaultValue() == TEXT("34.000000"), TEXT("FAnimationAnimNextRuntimeTest_GraphAddDecorator -> Unexpected decorator pin value"));
-	UE_RETURN_ON_ERROR(!DecoratorPin->GetSubPins()[2]->IsLazy(), TEXT("FAnimationAnimNextRuntimeTest_GraphAddDecorator -> Expected non-lazy decorator pin"));
-
-	// SomeLatentInt32
-	UE_RETURN_ON_ERROR(DecoratorPin->GetSubPins()[3]->GetCPPType() == TEXT("int32"), TEXT("FAnimationAnimNextRuntimeTest_GraphAddDecorator -> Unexpected decorator pin type"));
-	UE_RETURN_ON_ERROR(DecoratorPin->GetSubPins()[3]->GetDefaultValue() == TEXT("3"), TEXT("FAnimationAnimNextRuntimeTest_GraphAddDecorator -> Unexpected decorator pin value"));
-	UE_RETURN_ON_ERROR(DecoratorPin->GetSubPins()[3]->IsLazy(), TEXT("FAnimationAnimNextRuntimeTest_GraphAddDecorator -> Expected lazy decorator pin"));
-
-	// SomeOtherLatentInt32
-	UE_RETURN_ON_ERROR(DecoratorPin->GetSubPins()[4]->GetCPPType() == TEXT("int32"), TEXT("FAnimationAnimNextRuntimeTest_GraphAddDecorator -> Unexpected decorator pin type"));
-	UE_RETURN_ON_ERROR(DecoratorPin->GetSubPins()[4]->GetDefaultValue() == TEXT("3"), TEXT("FAnimationAnimNextRuntimeTest_GraphAddDecorator -> Unexpected decorator pin value"));
-	UE_RETURN_ON_ERROR(DecoratorPin->GetSubPins()[4]->IsLazy(), TEXT("FAnimationAnimNextRuntimeTest_GraphAddDecorator -> Expected lazy decorator pin"));
-
-	// SomeLatentFloat
-	UE_RETURN_ON_ERROR(DecoratorPin->GetSubPins()[5]->GetCPPType() == TEXT("float"), TEXT("FAnimationAnimNextRuntimeTest_GraphAddDecorator -> Unexpected decorator pin type"));
-	UE_RETURN_ON_ERROR(DecoratorPin->GetSubPins()[5]->GetDefaultValue() == TEXT("34.000000"), TEXT("FAnimationAnimNextRuntimeTest_GraphAddDecorator -> Unexpected decorator pin value"));
-	UE_RETURN_ON_ERROR(DecoratorPin->GetSubPins()[5]->IsLazy(), TEXT("FAnimationAnimNextRuntimeTest_GraphAddDecorator -> Expected lazy decorator pin"));
+			Prop->ExportText_Direct(DefaultValue, &CppDecoratorStructInstance, &DefaultCppDecoratorStructInstance, nullptr, PPF_None);
+		}
+
+		FString DisplayNameMetadata;
+		ScriptStruct->GetStringMetaDataHierarchical(FRigVMStruct::DisplayNameMetaName, &DisplayNameMetadata);
+		const FString DisplayName = DisplayNameMetadata.IsEmpty() ? Decorator->GetDecoratorName() : DisplayNameMetadata;
+
+		const FName DecoratorName = Controller->AddDecorator(
+			DecoratorStackNode->GetFName(),
+			*CppDecoratorStruct->GetPathName(),
+			*DisplayName,
+			DefaultValue, INDEX_NONE, true, true);
+		UE_RETURN_ON_ERROR(DecoratorName == DisplayName, TEXT("FAnimationAnimNextRuntimeTest_GraphAddDecorator -> Unexpected decorator name"));
+
+		URigVMPin* DecoratorPin = DecoratorStackNode->FindPin(*DisplayName);
+		UE_RETURN_ON_ERROR(DecoratorPin != nullptr, TEXT("FAnimationAnimNextRuntimeTest_GraphAddDecorator -> Failed to find decorator pin"));
+
+		// Our first pin is the hard coded output result, decorator pins follow
+		UE_RETURN_ON_ERROR(DecoratorStackNode->GetPins().Num() == 2, TEXT("FAnimationAnimNextRuntimeTest_GraphAddDecorator -> Unexpected number of pins"));
+		UE_RETURN_ON_ERROR(DecoratorPin->IsDecoratorPin() == true, TEXT("FAnimationAnimNextRuntimeTest_GraphAddDecorator -> Unexpected pin type"));
+		UE_RETURN_ON_ERROR(DecoratorPin->GetFName() == DecoratorName, TEXT("FAnimationAnimNextRuntimeTest_GraphAddDecorator -> Unexpected pin name"));
+		UE_RETURN_ON_ERROR(DecoratorPin->GetCPPTypeObject() == FRigDecorator_AnimNextCppDecorator::StaticStruct(), TEXT("FAnimationAnimNextRuntimeTest_GraphAddDecorator -> Unexpected pin type"));
+
+		// Our first sub-pin is the hard coded script struct member that parametrizes the decorator, dynamic decorator sub-pins follow
+		UE_RETURN_ON_ERROR(DecoratorPin->GetSubPins().Num() == 6, TEXT("FAnimationAnimNextRuntimeTest_GraphAddDecorator -> Unexpected decorator sub pins"));
+
+		// SomeInt32
+		UE_RETURN_ON_ERROR(DecoratorPin->GetSubPins()[1]->GetCPPType() == TEXT("int32"), TEXT("FAnimationAnimNextRuntimeTest_GraphAddDecorator -> Unexpected decorator pin type"));
+		UE_RETURN_ON_ERROR(DecoratorPin->GetSubPins()[1]->GetDefaultValue() == TEXT("3"), TEXT("FAnimationAnimNextRuntimeTest_GraphAddDecorator -> Unexpected decorator pin value"));
+		UE_RETURN_ON_ERROR(!DecoratorPin->GetSubPins()[1]->IsLazy(), TEXT("FAnimationAnimNextRuntimeTest_GraphAddDecorator -> Expected non-lazy decorator pin"));
+
+		// SomeFloat
+		UE_RETURN_ON_ERROR(DecoratorPin->GetSubPins()[2]->GetCPPType() == TEXT("float"), TEXT("FAnimationAnimNextRuntimeTest_GraphAddDecorator -> Unexpected decorator pin type"));
+		UE_RETURN_ON_ERROR(DecoratorPin->GetSubPins()[2]->GetDefaultValue() == TEXT("34.000000"), TEXT("FAnimationAnimNextRuntimeTest_GraphAddDecorator -> Unexpected decorator pin value"));
+		UE_RETURN_ON_ERROR(!DecoratorPin->GetSubPins()[2]->IsLazy(), TEXT("FAnimationAnimNextRuntimeTest_GraphAddDecorator -> Expected non-lazy decorator pin"));
+
+		// SomeLatentInt32
+		UE_RETURN_ON_ERROR(DecoratorPin->GetSubPins()[3]->GetCPPType() == TEXT("int32"), TEXT("FAnimationAnimNextRuntimeTest_GraphAddDecorator -> Unexpected decorator pin type"));
+		UE_RETURN_ON_ERROR(DecoratorPin->GetSubPins()[3]->GetDefaultValue() == TEXT("3"), TEXT("FAnimationAnimNextRuntimeTest_GraphAddDecorator -> Unexpected decorator pin value"));
+		UE_RETURN_ON_ERROR(DecoratorPin->GetSubPins()[3]->IsLazy(), TEXT("FAnimationAnimNextRuntimeTest_GraphAddDecorator -> Expected lazy decorator pin"));
+
+		// SomeOtherLatentInt32
+		UE_RETURN_ON_ERROR(DecoratorPin->GetSubPins()[4]->GetCPPType() == TEXT("int32"), TEXT("FAnimationAnimNextRuntimeTest_GraphAddDecorator -> Unexpected decorator pin type"));
+		UE_RETURN_ON_ERROR(DecoratorPin->GetSubPins()[4]->GetDefaultValue() == TEXT("3"), TEXT("FAnimationAnimNextRuntimeTest_GraphAddDecorator -> Unexpected decorator pin value"));
+		UE_RETURN_ON_ERROR(DecoratorPin->GetSubPins()[4]->IsLazy(), TEXT("FAnimationAnimNextRuntimeTest_GraphAddDecorator -> Expected lazy decorator pin"));
+
+		// SomeLatentFloat
+		UE_RETURN_ON_ERROR(DecoratorPin->GetSubPins()[5]->GetCPPType() == TEXT("float"), TEXT("FAnimationAnimNextRuntimeTest_GraphAddDecorator -> Unexpected decorator pin type"));
+		UE_RETURN_ON_ERROR(DecoratorPin->GetSubPins()[5]->GetDefaultValue() == TEXT("34.000000"), TEXT("FAnimationAnimNextRuntimeTest_GraphAddDecorator -> Unexpected decorator pin value"));
+		UE_RETURN_ON_ERROR(DecoratorPin->GetSubPins()[5]->IsLazy(), TEXT("FAnimationAnimNextRuntimeTest_GraphAddDecorator -> Expected lazy decorator pin"));
+	}
 
 	Tests::FUtils::CleanupAfterTests();
 
@@ -180,109 +186,112 @@ IMPLEMENT_SIMPLE_AUTOMATION_TEST(FAnimationAnimNextRuntimeTest_GraphExecute, "An
 bool FAnimationAnimNextRuntimeTest_GraphExecute::RunTest(const FString& InParameters)
 {
 	using namespace UE::AnimNext;
-	AUTO_REGISTER_ANIM_DECORATOR(FTestDecorator)
 
-	FScopedClearNodeTemplateRegistry ScopedClearNodeTemplateRegistry;
+	{
+		AUTO_REGISTER_ANIM_DECORATOR(FTestDecorator)
 
-	UFactory* GraphFactory = NewObject<UAnimNextGraphFactory>();
-	UAnimNextGraph* AnimNextGraph = CastChecked<UAnimNextGraph>(GraphFactory->FactoryCreateNew(UAnimNextGraph::StaticClass(), GetTransientPackage(), TEXT("TestAnimNextGraph"), RF_Transient, nullptr, nullptr, NAME_None));
-	UE_RETURN_ON_ERROR(AnimNextGraph != nullptr, "FAnimationAnimNextRuntimeTest_GraphExecute -> Failed to create animation graph");
+		FScopedClearNodeTemplateRegistry ScopedClearNodeTemplateRegistry;
 
-	UAnimNextGraph_EditorData* EditorData = UncookedOnly::FUtils::GetEditorData(AnimNextGraph);
-	UE_RETURN_ON_ERROR(EditorData != nullptr, "FAnimationAnimNextRuntimeTest_GraphExecute -> Failed to find animation graph editor data");
+		UFactory* GraphFactory = NewObject<UAnimNextGraphFactory>();
+		UAnimNextGraph* AnimNextGraph = CastChecked<UAnimNextGraph>(GraphFactory->FactoryCreateNew(UAnimNextGraph::StaticClass(), GetTransientPackage(), TEXT("TestAnimNextGraph"), RF_Transient, nullptr, nullptr, NAME_None));
+		UE_RETURN_ON_ERROR(AnimNextGraph != nullptr, "FAnimationAnimNextRuntimeTest_GraphExecute -> Failed to create animation graph");
 
-	URigVMController* Controller = EditorData->GetRigVMClient()->GetController(EditorData->GetRigVMClient()->GetDefaultModel());
-	UE_RETURN_ON_ERROR(Controller != nullptr, "FAnimationAnimNextRuntimeTest_GraphExecute -> Failed to get RigVM controller");
+		UAnimNextGraph_EditorData* EditorData = UncookedOnly::FUtils::GetEditorData(AnimNextGraph);
+		UE_RETURN_ON_ERROR(EditorData != nullptr, "FAnimationAnimNextRuntimeTest_GraphExecute -> Failed to find animation graph editor data");
 
-	// Find graph entry point
-	URigVMNode* MainEntryPointNode = Controller->GetGraph()->FindNodeByName(FRigUnit_AnimNextGraphRoot::StaticStruct()->GetFName());
-	UE_RETURN_ON_ERROR(MainEntryPointNode != nullptr, "FAnimationAnimNextRuntimeTest_GraphExecute -> Failed to find main entry point node");
+		URigVMController* Controller = EditorData->GetRigVMClient()->GetController(EditorData->GetRigVMClient()->GetDefaultModel());
+		UE_RETURN_ON_ERROR(Controller != nullptr, "FAnimationAnimNextRuntimeTest_GraphExecute -> Failed to get RigVM controller");
 
-	URigVMPin* BeginExecutePin = MainEntryPointNode->FindPin(GET_MEMBER_NAME_STRING_CHECKED(FRigUnit_AnimNextGraphRoot, Result));
-	UE_RETURN_ON_ERROR(BeginExecutePin != nullptr && BeginExecutePin->GetDirection() == ERigVMPinDirection::Input, "FAnimationAnimNextRuntimeTest_GraphExecute -> Failed to create entry point");
+		// Find graph entry point
+		URigVMNode* MainEntryPointNode = Controller->GetGraph()->FindNodeByName(FRigUnit_AnimNextGraphRoot::StaticStruct()->GetFName());
+		UE_RETURN_ON_ERROR(MainEntryPointNode != nullptr, "FAnimationAnimNextRuntimeTest_GraphExecute -> Failed to find main entry point node");
 
-	// Create an empty decorator stack node
-	URigVMUnitNode* DecoratorStackNode = Controller->AddUnitNode(FRigUnit_AnimNextDecoratorStack::StaticStruct(), FRigVMStruct::ExecuteName, FVector2D(0.0f, 0.0f), FString(), false);
-	UE_RETURN_ON_ERROR(DecoratorStackNode != nullptr, "FAnimationAnimNextRuntimeTest_GraphExecute -> Failed to create decorator stack node");
+		URigVMPin* BeginExecutePin = MainEntryPointNode->FindPin(GET_MEMBER_NAME_STRING_CHECKED(FRigUnit_AnimNextGraphRoot, Result));
+		UE_RETURN_ON_ERROR(BeginExecutePin != nullptr && BeginExecutePin->GetDirection() == ERigVMPinDirection::Input, "FAnimationAnimNextRuntimeTest_GraphExecute -> Failed to create entry point");
 
-	// Link our stack result to our entry point
-	Controller->AddLink(DecoratorStackNode->GetPins()[0], MainEntryPointNode->GetPins()[0]);
+		// Create an empty decorator stack node
+		URigVMUnitNode* DecoratorStackNode = Controller->AddUnitNode(FRigUnit_AnimNextDecoratorStack::StaticStruct(), FRigVMStruct::ExecuteName, FVector2D(0.0f, 0.0f), FString(), false);
+		UE_RETURN_ON_ERROR(DecoratorStackNode != nullptr, "FAnimationAnimNextRuntimeTest_GraphExecute -> Failed to create decorator stack node");
 
-	// Add a decorator
-	const UScriptStruct* CppDecoratorStruct = FRigDecorator_AnimNextCppDecorator::StaticStruct();
-	UE_RETURN_ON_ERROR(CppDecoratorStruct != nullptr, "FAnimationAnimNextRuntimeTest_GraphExecute -> Failed to get find Cpp decorator static struct");
+		// Link our stack result to our entry point
+		Controller->AddLink(DecoratorStackNode->GetPins()[0], MainEntryPointNode->GetPins()[0]);
 
-	const FDecorator* Decorator = FDecoratorRegistry::Get().Find(FTestDecorator::DecoratorUID);
-	UE_RETURN_ON_ERROR(Decorator != nullptr, "FAnimationAnimNextRuntimeTest_GraphExecute -> Failed to find test decorator");
+		// Add a decorator
+		const UScriptStruct* CppDecoratorStruct = FRigDecorator_AnimNextCppDecorator::StaticStruct();
+		UE_RETURN_ON_ERROR(CppDecoratorStruct != nullptr, "FAnimationAnimNextRuntimeTest_GraphExecute -> Failed to get find Cpp decorator static struct");
 
-	UScriptStruct* ScriptStruct = Decorator->GetDecoratorSharedDataStruct();
-	UE_RETURN_ON_ERROR(ScriptStruct != nullptr, "FAnimationAnimNextRuntimeTest_GraphExecute -> Failed to find decorator shared data struct");
+		const FDecorator* Decorator = FDecoratorRegistry::Get().Find(FTestDecorator::DecoratorUID);
+		UE_RETURN_ON_ERROR(Decorator != nullptr, "FAnimationAnimNextRuntimeTest_GraphExecute -> Failed to find test decorator");
 
-	FString DefaultValue;
-	{
-		const FRigDecorator_AnimNextCppDecorator DefaultCppDecoratorStructInstance;
-		FRigDecorator_AnimNextCppDecorator CppDecoratorStructInstance;
-		CppDecoratorStructInstance.DecoratorSharedDataStruct = ScriptStruct;
+		UScriptStruct* ScriptStruct = Decorator->GetDecoratorSharedDataStruct();
+		UE_RETURN_ON_ERROR(ScriptStruct != nullptr, "FAnimationAnimNextRuntimeTest_GraphExecute -> Failed to find decorator shared data struct");
 
-		UE_RETURN_ON_ERROR(CppDecoratorStructInstance.CanBeAddedToNode(DecoratorStackNode, nullptr), "FAnimationAnimNextRuntimeTest_GraphExecute -> Decorator cannot be added to decorator stack node");
+		FString DefaultValue;
+		{
+			const FRigDecorator_AnimNextCppDecorator DefaultCppDecoratorStructInstance;
+			FRigDecorator_AnimNextCppDecorator CppDecoratorStructInstance;
+			CppDecoratorStructInstance.DecoratorSharedDataStruct = ScriptStruct;
 
-		const FProperty* Prop = FAnimNextCppDecoratorWrapper::StaticStruct()->FindPropertyByName(GET_MEMBER_NAME_STRING_CHECKED(FAnimNextCppDecoratorWrapper, CppDecorator));
-		UE_RETURN_ON_ERROR(Prop != nullptr, "FAnimationAnimNextRuntimeTest_GraphExecute -> Failed to find wrapper property");
+			UE_RETURN_ON_ERROR(CppDecoratorStructInstance.CanBeAddedToNode(DecoratorStackNode, nullptr), "FAnimationAnimNextRuntimeTest_GraphExecute -> Decorator cannot be added to decorator stack node");
 
-		Prop->ExportText_Direct(DefaultValue, &CppDecoratorStructInstance, &DefaultCppDecoratorStructInstance, nullptr, PPF_None);
-	}
+			const FProperty* Prop = FAnimNextCppDecoratorWrapper::StaticStruct()->FindPropertyByName(GET_MEMBER_NAME_STRING_CHECKED(FAnimNextCppDecoratorWrapper, CppDecorator));
+			UE_RETURN_ON_ERROR(Prop != nullptr, "FAnimationAnimNextRuntimeTest_GraphExecute -> Failed to find wrapper property");
 
-	FString DisplayNameMetadata;
-	ScriptStruct->GetStringMetaDataHierarchical(FRigVMStruct::DisplayNameMetaName, &DisplayNameMetadata);
-	const FString DisplayName = DisplayNameMetadata.IsEmpty() ? Decorator->GetDecoratorName() : DisplayNameMetadata;
+			Prop->ExportText_Direct(DefaultValue, &CppDecoratorStructInstance, &DefaultCppDecoratorStructInstance, nullptr, PPF_None);
+		}
 
-	const FName DecoratorName = Controller->AddDecorator(
-		DecoratorStackNode->GetFName(),
-		*CppDecoratorStruct->GetPathName(),
-		*DisplayName,
-		DefaultValue, INDEX_NONE, true, true);
-	UE_RETURN_ON_ERROR(DecoratorName == DisplayName, TEXT("FAnimationAnimNextRuntimeTest_GraphExecute -> Unexpected decorator name"));
+		FString DisplayNameMetadata;
+		ScriptStruct->GetStringMetaDataHierarchical(FRigVMStruct::DisplayNameMetaName, &DisplayNameMetadata);
+		const FString DisplayName = DisplayNameMetadata.IsEmpty() ? Decorator->GetDecoratorName() : DisplayNameMetadata;
 
-	URigVMPin* DecoratorPin = DecoratorStackNode->FindPin(*DisplayName);
-	UE_RETURN_ON_ERROR(DecoratorPin != nullptr, TEXT("FAnimationAnimNextRuntimeTest_GraphExecute -> Failed to find decorator pin"));
+		const FName DecoratorName = Controller->AddDecorator(
+			DecoratorStackNode->GetFName(),
+			*CppDecoratorStruct->GetPathName(),
+			*DisplayName,
+			DefaultValue, INDEX_NONE, true, true);
+		UE_RETURN_ON_ERROR(DecoratorName == DisplayName, TEXT("FAnimationAnimNextRuntimeTest_GraphExecute -> Unexpected decorator name"));
 
-	// Set some values on our decorator
-	Controller->SetPinDefaultValue(DecoratorPin->GetSubPins()[1]->GetPinPath(), TEXT("78"));
-	Controller->SetPinDefaultValue(DecoratorPin->GetSubPins()[2]->GetPinPath(), TEXT("142.33"));
+		URigVMPin* DecoratorPin = DecoratorStackNode->FindPin(*DisplayName);
+		UE_RETURN_ON_ERROR(DecoratorPin != nullptr, TEXT("FAnimationAnimNextRuntimeTest_GraphExecute -> Failed to find decorator pin"));
 
-	TSharedRef<FParamStack> ParamStack = MakeShared<FParamStack>();
-	FParamStack::AttachToCurrentThread(ParamStack);
+		// Set some values on our decorator
+		Controller->SetPinDefaultValue(DecoratorPin->GetSubPins()[1]->GetPinPath(), TEXT("78"));
+		Controller->SetPinDefaultValue(DecoratorPin->GetSubPins()[2]->GetPinPath(), TEXT("142.33"));
 
-	FAnimNextGraphInstancePtr GraphInstance;
-	AnimNextGraph->AllocateInstance(GraphInstance);
+		TSharedRef<FParamStack> ParamStack = MakeShared<FParamStack>();
+		FParamStack::AttachToCurrentThread(ParamStack);
 
-	FParamStack::FPushedLayerHandle LayerHandle = ParamStack->PushValues(
-		"UpdateCount", (int32)0,
-		"EvaluateCount", (int32)0,
-		"SomeInt32", (int32)0,
-		"SomeFloat", 0.0f,
-		"SomeLatentInt32", (int32)0,
-		"SomeOtherLatentInt32", (int32)0,
-		"SomeLatentFloat", 0.0f
-	);
+		FAnimNextGraphInstancePtr GraphInstance;
+		AnimNextGraph->AllocateInstance(GraphInstance);
 
-	{
-		UE::AnimNext::UpdateGraph(GraphInstance, 1.0f / 30.0f);
-		(void)UE::AnimNext::EvaluateGraph(GraphInstance);
-	}
+		FParamStack::FPushedLayerHandle LayerHandle = ParamStack->PushValues(
+			"UpdateCount", (int32)0,
+			"EvaluateCount", (int32)0,
+			"SomeInt32", (int32)0,
+			"SomeFloat", 0.0f,
+			"SomeLatentInt32", (int32)0,
+			"SomeOtherLatentInt32", (int32)0,
+			"SomeLatentFloat", 0.0f
+		);
 
-	AddErrorIfFalse(ParamStack->GetParam<int32>("UpdateCount") == 1, "FAnimationAnimNextRuntimeTest_GraphExecute -> Unexpected update count");
-	AddErrorIfFalse(ParamStack->GetParam<int32>("EvaluateCount") == 1, "FAnimationAnimNextRuntimeTest_GraphExecute -> Unexpected evaluate count");
-	AddErrorIfFalse(ParamStack->GetParam<int32>("SomeInt32") == 78, "FAnimationAnimNextRuntimeTest_GraphExecute -> Unexpected SomeInt32 value");
-	AddErrorIfFalse(ParamStack->GetParam<float>("SomeFloat") == 142.33f, "FAnimationAnimNextRuntimeTest_GraphExecute -> Unexpected SomeFloat value");
-	AddErrorIfFalse(ParamStack->GetParam<int32>("SomeLatentInt32") == 3, "FAnimationAnimNextRuntimeTest_GraphExecute -> Unexpected SomeLatentInt32 value");
-	AddErrorIfFalse(ParamStack->GetParam<int32>("SomeOtherLatentInt32") == 3, "FAnimationAnimNextRuntimeTest_GraphExecute -> Unexpected SomeOtherLatentInt32 value");
-	AddErrorIfFalse(ParamStack->GetParam<float>("SomeLatentFloat") == 34.0f, "FAnimationAnimNextRuntimeTest_GraphExecute -> Unexpected SomeLatentFloat value");
+		{
+			UE::AnimNext::UpdateGraph(GraphInstance, 1.0f / 30.0f);
+			(void)UE::AnimNext::EvaluateGraph(GraphInstance);
+		}
 
-	ParamStack->PopLayer(LayerHandle);
-	GraphInstance.Release();
+		AddErrorIfFalse(ParamStack->GetParam<int32>("UpdateCount") == 1, "FAnimationAnimNextRuntimeTest_GraphExecute -> Unexpected update count");
+		AddErrorIfFalse(ParamStack->GetParam<int32>("EvaluateCount") == 1, "FAnimationAnimNextRuntimeTest_GraphExecute -> Unexpected evaluate count");
+		AddErrorIfFalse(ParamStack->GetParam<int32>("SomeInt32") == 78, "FAnimationAnimNextRuntimeTest_GraphExecute -> Unexpected SomeInt32 value");
+		AddErrorIfFalse(ParamStack->GetParam<float>("SomeFloat") == 142.33f, "FAnimationAnimNextRuntimeTest_GraphExecute -> Unexpected SomeFloat value");
+		AddErrorIfFalse(ParamStack->GetParam<int32>("SomeLatentInt32") == 3, "FAnimationAnimNextRuntimeTest_GraphExecute -> Unexpected SomeLatentInt32 value");
+		AddErrorIfFalse(ParamStack->GetParam<int32>("SomeOtherLatentInt32") == 3, "FAnimationAnimNextRuntimeTest_GraphExecute -> Unexpected SomeOtherLatentInt32 value");
+		AddErrorIfFalse(ParamStack->GetParam<float>("SomeLatentFloat") == 34.0f, "FAnimationAnimNextRuntimeTest_GraphExecute -> Unexpected SomeLatentFloat value");
 
-	FParamStack::DetachFromCurrentThread();
+		ParamStack->PopLayer(LayerHandle);
+		GraphInstance.Release();
+
+		FParamStack::DetachFromCurrentThread();
+	}
 
 	Tests::FUtils::CleanupAfterTests();
 
@@ -294,140 +303,142 @@ IMPLEMENT_SIMPLE_AUTOMATION_TEST(FAnimationAnimNextRuntimeTest_GraphExecuteLaten
 bool FAnimationAnimNextRuntimeTest_GraphExecuteLatent::RunTest(const FString& InParameters)
 {
 	using namespace UE::AnimNext;
-	AUTO_REGISTER_ANIM_DECORATOR(FTestDecorator)
 
-	FScopedClearNodeTemplateRegistry ScopedClearNodeTemplateRegistry;
+	{
+		AUTO_REGISTER_ANIM_DECORATOR(FTestDecorator)
+		FScopedClearNodeTemplateRegistry ScopedClearNodeTemplateRegistry;
 
-	UFactory* GraphFactory = NewObject<UAnimNextGraphFactory>();
-	UAnimNextGraph* AnimNextGraph = CastChecked<UAnimNextGraph>(GraphFactory->FactoryCreateNew(UAnimNextGraph::StaticClass(), GetTransientPackage(), TEXT("TestAnimNextGraph"), RF_Transient, nullptr, nullptr, NAME_None));
-	UE_RETURN_ON_ERROR(AnimNextGraph != nullptr, "FAnimationAnimNextRuntimeTest_GraphExecuteLatent -> Failed to create animation graph");
+		UFactory* GraphFactory = NewObject<UAnimNextGraphFactory>();
+		UAnimNextGraph* AnimNextGraph = CastChecked<UAnimNextGraph>(GraphFactory->FactoryCreateNew(UAnimNextGraph::StaticClass(), GetTransientPackage(), TEXT("TestAnimNextGraph"), RF_Transient, nullptr, nullptr, NAME_None));
+		UE_RETURN_ON_ERROR(AnimNextGraph != nullptr, "FAnimationAnimNextRuntimeTest_GraphExecuteLatent -> Failed to create animation graph");
 
-	UAnimNextGraph_EditorData* EditorData = UncookedOnly::FUtils::GetEditorData(AnimNextGraph);
-	UE_RETURN_ON_ERROR(EditorData != nullptr, "FAnimationAnimNextRuntimeTest_GraphExecuteLatent -> Failed to find animation graph editor data");
+		UAnimNextGraph_EditorData* EditorData = UncookedOnly::FUtils::GetEditorData(AnimNextGraph);
+		UE_RETURN_ON_ERROR(EditorData != nullptr, "FAnimationAnimNextRuntimeTest_GraphExecuteLatent -> Failed to find animation graph editor data");
 
-	URigVMController* Controller = EditorData->GetRigVMClient()->GetController(EditorData->GetRigVMClient()->GetDefaultModel());
-	UE_RETURN_ON_ERROR(Controller != nullptr, "FAnimationAnimNextRuntimeTest_GraphExecuteLatent -> Failed to get RigVM controller");
+		URigVMController* Controller = EditorData->GetRigVMClient()->GetController(EditorData->GetRigVMClient()->GetDefaultModel());
+		UE_RETURN_ON_ERROR(Controller != nullptr, "FAnimationAnimNextRuntimeTest_GraphExecuteLatent -> Failed to get RigVM controller");
 
-	// Find graph entry point
-	URigVMNode* MainEntryPointNode = Controller->GetGraph()->FindNodeByName(FRigUnit_AnimNextGraphRoot::StaticStruct()->GetFName());
-	UE_RETURN_ON_ERROR(MainEntryPointNode != nullptr, "FAnimationAnimNextRuntimeTest_GraphExecuteLatent -> Failed to find main entry point node");
+		// Find graph entry point
+		URigVMNode* MainEntryPointNode = Controller->GetGraph()->FindNodeByName(FRigUnit_AnimNextGraphRoot::StaticStruct()->GetFName());
+		UE_RETURN_ON_ERROR(MainEntryPointNode != nullptr, "FAnimationAnimNextRuntimeTest_GraphExecuteLatent -> Failed to find main entry point node");
 
-	URigVMPin* BeginExecutePin = MainEntryPointNode->FindPin(GET_MEMBER_NAME_STRING_CHECKED(FRigUnit_AnimNextGraphRoot, Result));
-	UE_RETURN_ON_ERROR(BeginExecutePin != nullptr && BeginExecutePin->GetDirection() == ERigVMPinDirection::Input, "FAnimationAnimNextRuntimeTest_GraphExecuteLatent -> Failed to create entry point");
+		URigVMPin* BeginExecutePin = MainEntryPointNode->FindPin(GET_MEMBER_NAME_STRING_CHECKED(FRigUnit_AnimNextGraphRoot, Result));
+		UE_RETURN_ON_ERROR(BeginExecutePin != nullptr && BeginExecutePin->GetDirection() == ERigVMPinDirection::Input, "FAnimationAnimNextRuntimeTest_GraphExecuteLatent -> Failed to create entry point");
 
-	// Create an empty decorator stack node
-	URigVMUnitNode* DecoratorStackNode = Controller->AddUnitNode(FRigUnit_AnimNextDecoratorStack::StaticStruct(), FRigVMStruct::ExecuteName, FVector2D(0.0f, 0.0f), FString(), false);
-	UE_RETURN_ON_ERROR(DecoratorStackNode != nullptr, "FAnimationAnimNextRuntimeTest_GraphExecuteLatent -> Failed to create decorator stack node");
+		// Create an empty decorator stack node
+		URigVMUnitNode* DecoratorStackNode = Controller->AddUnitNode(FRigUnit_AnimNextDecoratorStack::StaticStruct(), FRigVMStruct::ExecuteName, FVector2D(0.0f, 0.0f), FString(), false);
+		UE_RETURN_ON_ERROR(DecoratorStackNode != nullptr, "FAnimationAnimNextRuntimeTest_GraphExecuteLatent -> Failed to create decorator stack node");
 
-	// Link our stack result to our entry point
-	Controller->AddLink(DecoratorStackNode->GetPins()[0], MainEntryPointNode->GetPins()[0]);
+		// Link our stack result to our entry point
+		Controller->AddLink(DecoratorStackNode->GetPins()[0], MainEntryPointNode->GetPins()[0]);
 
-	// Add a decorator
-	const UScriptStruct* CppDecoratorStruct = FRigDecorator_AnimNextCppDecorator::StaticStruct();
-	UE_RETURN_ON_ERROR(CppDecoratorStruct != nullptr, "FAnimationAnimNextRuntimeTest_GraphExecuteLatent -> Failed to get find Cpp decorator static struct");
+		// Add a decorator
+		const UScriptStruct* CppDecoratorStruct = FRigDecorator_AnimNextCppDecorator::StaticStruct();
+		UE_RETURN_ON_ERROR(CppDecoratorStruct != nullptr, "FAnimationAnimNextRuntimeTest_GraphExecuteLatent -> Failed to get find Cpp decorator static struct");
 
-	const FDecorator* Decorator = FDecoratorRegistry::Get().Find(FTestDecorator::DecoratorUID);
-	UE_RETURN_ON_ERROR(Decorator != nullptr, "FAnimationAnimNextRuntimeTest_GraphExecuteLatent -> Failed to find test decorator");
+		const FDecorator* Decorator = FDecoratorRegistry::Get().Find(FTestDecorator::DecoratorUID);
+		UE_RETURN_ON_ERROR(Decorator != nullptr, "FAnimationAnimNextRuntimeTest_GraphExecuteLatent -> Failed to find test decorator");
 
-	UScriptStruct* ScriptStruct = Decorator->GetDecoratorSharedDataStruct();
-	UE_RETURN_ON_ERROR(ScriptStruct != nullptr, "FAnimationAnimNextRuntimeTest_GraphExecuteLatent -> Failed to find decorator shared data struct");
+		UScriptStruct* ScriptStruct = Decorator->GetDecoratorSharedDataStruct();
+		UE_RETURN_ON_ERROR(ScriptStruct != nullptr, "FAnimationAnimNextRuntimeTest_GraphExecuteLatent -> Failed to find decorator shared data struct");
 
-	FString DefaultValue;
-	{
-		const FRigDecorator_AnimNextCppDecorator DefaultCppDecoratorStructInstance;
-		FRigDecorator_AnimNextCppDecorator CppDecoratorStructInstance;
-		CppDecoratorStructInstance.DecoratorSharedDataStruct = ScriptStruct;
+		FString DefaultValue;
+		{
+			const FRigDecorator_AnimNextCppDecorator DefaultCppDecoratorStructInstance;
+			FRigDecorator_AnimNextCppDecorator CppDecoratorStructInstance;
+			CppDecoratorStructInstance.DecoratorSharedDataStruct = ScriptStruct;
 
-		UE_RETURN_ON_ERROR(CppDecoratorStructInstance.CanBeAddedToNode(DecoratorStackNode, nullptr), "FAnimationAnimNextRuntimeTest_GraphExecuteLatent -> Decorator cannot be added to decorator stack node");
+			UE_RETURN_ON_ERROR(CppDecoratorStructInstance.CanBeAddedToNode(DecoratorStackNode, nullptr), "FAnimationAnimNextRuntimeTest_GraphExecuteLatent -> Decorator cannot be added to decorator stack node");
 
-		const FProperty* Prop = FAnimNextCppDecoratorWrapper::StaticStruct()->FindPropertyByName(GET_MEMBER_NAME_STRING_CHECKED(FAnimNextCppDecoratorWrapper, CppDecorator));
-		UE_RETURN_ON_ERROR(Prop != nullptr, "FAnimationAnimNextRuntimeTest_GraphExecuteLatent -> Failed to find wrapper property");
+			const FProperty* Prop = FAnimNextCppDecoratorWrapper::StaticStruct()->FindPropertyByName(GET_MEMBER_NAME_STRING_CHECKED(FAnimNextCppDecoratorWrapper, CppDecorator));
+			UE_RETURN_ON_ERROR(Prop != nullptr, "FAnimationAnimNextRuntimeTest_GraphExecuteLatent -> Failed to find wrapper property");
 
-		Prop->ExportText_Direct(DefaultValue, &CppDecoratorStructInstance, &DefaultCppDecoratorStructInstance, nullptr, PPF_None);
-	}
+			Prop->ExportText_Direct(DefaultValue, &CppDecoratorStructInstance, &DefaultCppDecoratorStructInstance, nullptr, PPF_None);
+		}
 
-	FString DisplayNameMetadata;
-	ScriptStruct->GetStringMetaDataHierarchical(FRigVMStruct::DisplayNameMetaName, &DisplayNameMetadata);
-	const FString DisplayName = DisplayNameMetadata.IsEmpty() ? Decorator->GetDecoratorName() : DisplayNameMetadata;
+		FString DisplayNameMetadata;
+		ScriptStruct->GetStringMetaDataHierarchical(FRigVMStruct::DisplayNameMetaName, &DisplayNameMetadata);
+		const FString DisplayName = DisplayNameMetadata.IsEmpty() ? Decorator->GetDecoratorName() : DisplayNameMetadata;
 
-	const FName DecoratorName = Controller->AddDecorator(
-		DecoratorStackNode->GetFName(),
-		*CppDecoratorStruct->GetPathName(),
-		*DisplayName,
-		DefaultValue, INDEX_NONE, true, true);
-	UE_RETURN_ON_ERROR(DecoratorName == DisplayName, TEXT("FAnimationAnimNextRuntimeTest_GraphExecuteLatent -> Unexpected decorator name"));
+		const FName DecoratorName = Controller->AddDecorator(
+			DecoratorStackNode->GetFName(),
+			*CppDecoratorStruct->GetPathName(),
+			*DisplayName,
+			DefaultValue, INDEX_NONE, true, true);
+		UE_RETURN_ON_ERROR(DecoratorName == DisplayName, TEXT("FAnimationAnimNextRuntimeTest_GraphExecuteLatent -> Unexpected decorator name"));
 
-	// Set some values on our decorator
-	URigVMPin* DecoratorPin = DecoratorStackNode->FindPin(*DisplayName);
-	UE_RETURN_ON_ERROR(DecoratorPin != nullptr, TEXT("FAnimationAnimNextRuntimeTest_GraphExecuteLatent -> Failed to find decorator pin"));
+		// Set some values on our decorator
+		URigVMPin* DecoratorPin = DecoratorStackNode->FindPin(*DisplayName);
+		UE_RETURN_ON_ERROR(DecoratorPin != nullptr, TEXT("FAnimationAnimNextRuntimeTest_GraphExecuteLatent -> Failed to find decorator pin"));
 
-	Controller->SetPinDefaultValue(DecoratorPin->GetSubPins()[1]->GetPinPath(), TEXT("78"));		// SomeInt32
-	Controller->SetPinDefaultValue(DecoratorPin->GetSubPins()[2]->GetPinPath(), TEXT("142.33"));	// SomeFloat
-	Controller->SetPinDefaultValue(DecoratorPin->GetSubPins()[5]->GetPinPath(), TEXT("1123.31"));	// SomeLatentFloat, inline value on latent pin
+		Controller->SetPinDefaultValue(DecoratorPin->GetSubPins()[1]->GetPinPath(), TEXT("78"));		// SomeInt32
+		Controller->SetPinDefaultValue(DecoratorPin->GetSubPins()[2]->GetPinPath(), TEXT("142.33"));	// SomeFloat
+		Controller->SetPinDefaultValue(DecoratorPin->GetSubPins()[5]->GetPinPath(), TEXT("1123.31"));	// SomeLatentFloat, inline value on latent pin
 
-	// Set some latent values on our decorator
-	{
-		FRigVMFunction_MathIntAdd IntAdd;
-		IntAdd.A = 10;
-		IntAdd.B = 23;
+		// Set some latent values on our decorator
+		{
+			FRigVMFunction_MathIntAdd IntAdd;
+			IntAdd.A = 10;
+			IntAdd.B = 23;
 
-		URigVMUnitNode* IntAddNode = Controller->AddUnitNodeWithDefaults(FRigVMFunction_MathIntAdd::StaticStruct(), FRigStructScope(IntAdd), FRigVMStruct::ExecuteName, FVector2D::ZeroVector, FString(), false);
-		UE_RETURN_ON_ERROR(IntAddNode != nullptr, TEXT("FAnimationAnimNextRuntimeTest_GraphExecuteLatent -> Failed to create Int add node"));
+			URigVMUnitNode* IntAddNode = Controller->AddUnitNodeWithDefaults(FRigVMFunction_MathIntAdd::StaticStruct(), FRigStructScope(IntAdd), FRigVMStruct::ExecuteName, FVector2D::ZeroVector, FString(), false);
+			UE_RETURN_ON_ERROR(IntAddNode != nullptr, TEXT("FAnimationAnimNextRuntimeTest_GraphExecuteLatent -> Failed to create Int add node"));
 
-		Controller->AddLink(
-			IntAddNode->FindPin(GET_MEMBER_NAME_STRING_CHECKED(FRigVMFunction_MathIntAdd, Result)),
-			DecoratorPin->GetSubPins()[3]);	// SomeLatentInt32
-	}
+			Controller->AddLink(
+				IntAddNode->FindPin(GET_MEMBER_NAME_STRING_CHECKED(FRigVMFunction_MathIntAdd, Result)),
+				DecoratorPin->GetSubPins()[3]);	// SomeLatentInt32
+		}
 
-	{
-		const FRigVMDispatchFactory* GetParameterFactory = FRigVMRegistry::Get().RegisterFactory(FRigVMDispatch_GetParameter::StaticStruct());
-		const FName GetParameterNotation = GetParameterFactory->GetTemplate()->GetNotation();
+		{
+			const FRigVMDispatchFactory* GetParameterFactory = FRigVMRegistry::Get().RegisterFactory(FRigVMDispatch_GetParameter::StaticStruct());
+			const FName GetParameterNotation = GetParameterFactory->GetTemplate()->GetNotation();
 
-		URigVMNode* GetParameterNode = Controller->AddTemplateNode(GetParameterNotation);
-		UE_RETURN_ON_ERROR(GetParameterNode != nullptr, TEXT("FAnimationAnimNextRuntimeTest_GraphExecuteLatent -> Failed to create GetParameter node"));
+			URigVMNode* GetParameterNode = Controller->AddTemplateNode(GetParameterNotation);
+			UE_RETURN_ON_ERROR(GetParameterNode != nullptr, TEXT("FAnimationAnimNextRuntimeTest_GraphExecuteLatent -> Failed to create GetParameter node"));
 
-		UE_RETURN_ON_ERROR(Controller->ResolveWildCardPin(GetParameterNode->FindPin(TEXT("Value")), RigVMTypeUtils::TypeIndex::Int32), TEXT("FAnimationAnimNextRuntimeTest_GraphExecuteLatent -> Failed to resolve wildcard for GetParameter node"));
-		Controller->SetPinDefaultValue(GetParameterNode->FindPin(TEXT("Parameter"))->GetPinPath(), TEXT("SomeSourceInt"));
+			UE_RETURN_ON_ERROR(Controller->ResolveWildCardPin(GetParameterNode->FindPin(TEXT("Value")), RigVMTypeUtils::TypeIndex::Int32), TEXT("FAnimationAnimNextRuntimeTest_GraphExecuteLatent -> Failed to resolve wildcard for GetParameter node"));
+			Controller->SetPinDefaultValue(GetParameterNode->FindPin(TEXT("Parameter"))->GetPinPath(), TEXT("SomeSourceInt"));
 
-		Controller->AddLink(
-			GetParameterNode->FindPin(TEXT("Value")),
-			DecoratorPin->GetSubPins()[4]);	// SomeOtherLatentInt32
-	}
+			Controller->AddLink(
+				GetParameterNode->FindPin(TEXT("Value")),
+				DecoratorPin->GetSubPins()[4]);	// SomeOtherLatentInt32
+		}
 
-	TSharedRef<FParamStack> ParamStack = MakeShared<FParamStack>();
-	FParamStack::AttachToCurrentThread(ParamStack);
+		TSharedRef<FParamStack> ParamStack = MakeShared<FParamStack>();
+		FParamStack::AttachToCurrentThread(ParamStack);
 
-	FAnimNextGraphInstancePtr GraphInstance;
-	AnimNextGraph->AllocateInstance(GraphInstance);
+		FAnimNextGraphInstancePtr GraphInstance;
+		AnimNextGraph->AllocateInstance(GraphInstance);
 
-	FParamStack::FPushedLayerHandle LayerHandle = ParamStack->PushValues(
-		"UpdateCount", (int32)0,
-		"EvaluateCount", (int32)0,
-		"SomeSourceInt", (int32)1223,
-		"SomeInt32", (int32)0,
-		"SomeFloat", 0.0f,
-		"SomeLatentInt32", (int32)0,
-		"SomeOtherLatentInt32", (int32)0,
-		"SomeLatentFloat", 0.0f
-	);
+		FParamStack::FPushedLayerHandle LayerHandle = ParamStack->PushValues(
+			"UpdateCount", (int32)0,
+			"EvaluateCount", (int32)0,
+			"SomeSourceInt", (int32)1223,
+			"SomeInt32", (int32)0,
+			"SomeFloat", 0.0f,
+			"SomeLatentInt32", (int32)0,
+			"SomeOtherLatentInt32", (int32)0,
+			"SomeLatentFloat", 0.0f
+		);
 
-	{
-		UE::AnimNext::UpdateGraph(GraphInstance, 1.0f / 30.0f);
-		(void)UE::AnimNext::EvaluateGraph(GraphInstance);
-	}
+		{
+			UE::AnimNext::UpdateGraph(GraphInstance, 1.0f / 30.0f);
+			(void)UE::AnimNext::EvaluateGraph(GraphInstance);
+		}
 
-	AddErrorIfFalse(ParamStack->GetParam<int32>("UpdateCount") == 1, "FAnimationAnimNextRuntimeTest_GraphExecuteLatent -> Unexpected update count");
-	AddErrorIfFalse(ParamStack->GetParam<int32>("EvaluateCount") == 1, "FAnimationAnimNextRuntimeTest_GraphExecuteLatent -> Unexpected evaluate count");
-	AddErrorIfFalse(ParamStack->GetParam<int32>("SomeInt32") == 78, "FAnimationAnimNextRuntimeTest_GraphExecuteLatent -> Unexpected SomeInt32 value");
-	AddErrorIfFalse(ParamStack->GetParam<float>("SomeFloat") == 142.33f, "FAnimationAnimNextRuntimeTest_GraphExecuteLatent -> Unexpected SomeFloat value");
-	AddErrorIfFalse(ParamStack->GetParam<int32>("SomeLatentInt32") == 33, "FAnimationAnimNextRuntimeTest_GraphExecuteLatent -> Unexpected SomeLatentInt32 value");
-	AddErrorIfFalse(ParamStack->GetParam<int32>("SomeOtherLatentInt32") == 1223, "FAnimationAnimNextRuntimeTest_GraphExecuteLatent -> Unexpected SomeOtherLatentInt32 value");
-	AddErrorIfFalse(ParamStack->GetParam<float>("SomeLatentFloat") == 1123.31f, "FAnimationAnimNextRuntimeTest_GraphExecuteLatent -> Unexpected SomeLatentFloat value");
+		AddErrorIfFalse(ParamStack->GetParam<int32>("UpdateCount") == 1, "FAnimationAnimNextRuntimeTest_GraphExecuteLatent -> Unexpected update count");
+		AddErrorIfFalse(ParamStack->GetParam<int32>("EvaluateCount") == 1, "FAnimationAnimNextRuntimeTest_GraphExecuteLatent -> Unexpected evaluate count");
+		AddErrorIfFalse(ParamStack->GetParam<int32>("SomeInt32") == 78, "FAnimationAnimNextRuntimeTest_GraphExecuteLatent -> Unexpected SomeInt32 value");
+		AddErrorIfFalse(ParamStack->GetParam<float>("SomeFloat") == 142.33f, "FAnimationAnimNextRuntimeTest_GraphExecuteLatent -> Unexpected SomeFloat value");
+		AddErrorIfFalse(ParamStack->GetParam<int32>("SomeLatentInt32") == 33, "FAnimationAnimNextRuntimeTest_GraphExecuteLatent -> Unexpected SomeLatentInt32 value");
+		AddErrorIfFalse(ParamStack->GetParam<int32>("SomeOtherLatentInt32") == 1223, "FAnimationAnimNextRuntimeTest_GraphExecuteLatent -> Unexpected SomeOtherLatentInt32 value");
+		AddErrorIfFalse(ParamStack->GetParam<float>("SomeLatentFloat") == 1123.31f, "FAnimationAnimNextRuntimeTest_GraphExecuteLatent -> Unexpected SomeLatentFloat value");
 
-	ParamStack->PopLayer(LayerHandle);
-	GraphInstance.Release();
+		ParamStack->PopLayer(LayerHandle);
+		GraphInstance.Release();
 
-	FParamStack::DetachFromCurrentThread();
+		FParamStack::DetachFromCurrentThread();
+	}
 
 	Tests::FUtils::CleanupAfterTests();
 
diff --git a/Engine/Plugins/Experimental/AnimNext/Source/AnimNextTestSuite/Private/AnimNextDecoratorInterfacesTest.cpp b/Engine/Plugins/Experimental/AnimNext/Source/AnimNextTestSuite/Private/AnimNextDecoratorInterfacesTest.cpp
--- a/Engine/Plugins/Experimental/AnimNext/Source/AnimNextTestSuite/Private/AnimNextDecoratorInterfacesTest.cpp
+++ b/Engine/Plugins/Experimental/AnimNext/Source/AnimNextTestSuite/Private/AnimNextDecoratorInterfacesTest.cpp
@@ -84,10 +84,13 @@ namespace UE::AnimNext
 		}
 	};
 
-	DEFINE_ANIM_DECORATOR_BEGIN(FDecoratorWithNoChildren)
-		DEFINE_ANIM_DECORATOR_IMPLEMENTS_INTERFACE(IUpdate)
-		DEFINE_ANIM_DECORATOR_IMPLEMENTS_INTERFACE(IEvaluate)
-	DEFINE_ANIM_DECORATOR_END(FDecoratorWithNoChildren)
+	// Decorator implementation boilerplate
+	#define DECORATOR_INTERFACE_ENUMERATOR(GeneratorMacro) \
+		GeneratorMacro(IEvaluate) \
+		GeneratorMacro(IUpdate) \
+
+	GENERATE_ANIM_DECORATOR_IMPLEMENTATION(FDecoratorWithNoChildren, DECORATOR_INTERFACE_ENUMERATOR)
+	#undef DECORATOR_INTERFACE_ENUMERATOR
 
 	// This decorator does not update or evaluate
 	struct FDecoratorWithOneChild : FBaseDecorator, IHierarchy
@@ -122,9 +125,12 @@ namespace UE::AnimNext
 		}
 	};
 
-	DEFINE_ANIM_DECORATOR_BEGIN(FDecoratorWithOneChild)
-		DEFINE_ANIM_DECORATOR_IMPLEMENTS_INTERFACE(IHierarchy)
-	DEFINE_ANIM_DECORATOR_END(FDecoratorWithOneChild)
+	// Decorator implementation boilerplate
+	#define DECORATOR_INTERFACE_ENUMERATOR(GeneratorMacro) \
+		GeneratorMacro(IHierarchy) \
+
+	GENERATE_ANIM_DECORATOR_IMPLEMENTATION(FDecoratorWithOneChild, DECORATOR_INTERFACE_ENUMERATOR)
+	#undef DECORATOR_INTERFACE_ENUMERATOR
 
 	struct FDecoratorWithChildren : FBaseDecorator, IHierarchy, IUpdate, IEvaluate
 	{
@@ -210,11 +216,14 @@ namespace UE::AnimNext
 		}
 	};
 
-	DEFINE_ANIM_DECORATOR_BEGIN(FDecoratorWithChildren)
-		DEFINE_ANIM_DECORATOR_IMPLEMENTS_INTERFACE(IHierarchy)
-		DEFINE_ANIM_DECORATOR_IMPLEMENTS_INTERFACE(IUpdate)
-		DEFINE_ANIM_DECORATOR_IMPLEMENTS_INTERFACE(IEvaluate)
-	DEFINE_ANIM_DECORATOR_END(FDecoratorWithChildren)
+	// Decorator implementation boilerplate
+	#define DECORATOR_INTERFACE_ENUMERATOR(GeneratorMacro) \
+		GeneratorMacro(IEvaluate) \
+		GeneratorMacro(IHierarchy) \
+		GeneratorMacro(IUpdate) \
+
+	GENERATE_ANIM_DECORATOR_IMPLEMENTATION(FDecoratorWithChildren, DECORATOR_INTERFACE_ENUMERATOR)
+	#undef DECORATOR_INTERFACE_ENUMERATOR
 }
 
 IMPLEMENT_SIMPLE_AUTOMATION_TEST(FAnimationAnimNextRuntimeTest_IHierarchy, "Animation.AnimNext.Runtime.IHierarchy", EAutomationTestFlags::EditorContext | EAutomationTestFlags::EngineFilter)
@@ -280,59 +289,59 @@ bool FAnimationAnimNextRuntimeTest_IHierarchy::RunTest(const FString& InParamete
 			NodeHandles.Add(DecoratorWriter.RegisterNode(*NodeTemplateD));
 
 			// We don't have decorator properties
-			TArray<TMap<FString, FString>> DecoratorPropertiesA;
+			TArray<TMap<FName, FString>> DecoratorPropertiesA;
 			DecoratorPropertiesA.AddDefaulted(NodeTemplateDecoratorListA.Num());
 
-			TArray<TMap<FString, FString>> DecoratorPropertiesB;
+			TArray<TMap<FName, FString>> DecoratorPropertiesB;
 			DecoratorPropertiesB.AddDefaulted(NodeTemplateDecoratorListB.Num());
 			DecoratorPropertiesB[0].Add(TEXT("Child"), ToString<FDecoratorWithOneChild::FSharedData>(TEXT("Child"), FAnimNextDecoratorHandle(NodeHandles[0])));
 
-			TArray<TMap<FString, FString>> DecoratorPropertiesC;
+			TArray<TMap<FName, FString>> DecoratorPropertiesC;
 			DecoratorPropertiesC.AddDefaulted(NodeTemplateDecoratorListC.Num());
 			DecoratorPropertiesC[0].Add(TEXT("Child"), ToString<FDecoratorWithOneChild::FSharedData>(TEXT("Child"), FAnimNextDecoratorHandle(NodeHandles[0])));
 			DecoratorPropertiesC[1].Add(TEXT("Child"), ToString<FDecoratorWithOneChild::FSharedData>(TEXT("Child"), FAnimNextDecoratorHandle(NodeHandles[1])));
 
-			TArray<TMap<FString, FString>> DecoratorPropertiesD;
+			TArray<TMap<FName, FString>> DecoratorPropertiesD;
 			DecoratorPropertiesD.AddDefaulted(NodeTemplateDecoratorListD.Num());
 			FAnimNextDecoratorHandle ChildrenHandlesD[2] = { FAnimNextDecoratorHandle(NodeHandles[0]), FAnimNextDecoratorHandle(NodeHandles[2], 1)};
 			DecoratorPropertiesD[0].Add(TEXT("Children"), ToString<FDecoratorWithChildren::FSharedData>(TEXT("Children"), ChildrenHandlesD));
 
 			DecoratorWriter.BeginNodeWriting();
 			DecoratorWriter.WriteNode(NodeHandles[0],
-				[&DecoratorPropertiesA](uint32 DecoratorIndex, const FString& PropertyName)
+				[&DecoratorPropertiesA](uint32 DecoratorIndex, FName PropertyName)
 				{
 					return DecoratorPropertiesA[DecoratorIndex][PropertyName];
 				},
-				[](uint32 DecoratorIndex, const FString& PropertyName)
+				[](uint32 DecoratorIndex, FName PropertyName)
 				{
-					return false;
+					return MAX_uint16;
 				});
 			DecoratorWriter.WriteNode(NodeHandles[1],
-				[&DecoratorPropertiesB](uint32 DecoratorIndex, const FString& PropertyName)
+				[&DecoratorPropertiesB](uint32 DecoratorIndex, FName PropertyName)
 				{
 					return DecoratorPropertiesB[DecoratorIndex][PropertyName];
 				},
-				[](uint32 DecoratorIndex, const FString& PropertyName)
+				[](uint32 DecoratorIndex, FName PropertyName)
 				{
-					return false;
+					return MAX_uint16;
 				});
 			DecoratorWriter.WriteNode(NodeHandles[2],
-				[&DecoratorPropertiesC](uint32 DecoratorIndex, const FString& PropertyName)
+				[&DecoratorPropertiesC](uint32 DecoratorIndex, FName PropertyName)
 				{
 					return DecoratorPropertiesC[DecoratorIndex][PropertyName];
 				},
-				[](uint32 DecoratorIndex, const FString& PropertyName)
+				[](uint32 DecoratorIndex, FName PropertyName)
 				{
-					return false;
+					return MAX_uint16;
 				});
 			DecoratorWriter.WriteNode(NodeHandles[3],
-				[&DecoratorPropertiesD](uint32 DecoratorIndex, const FString& PropertyName)
+				[&DecoratorPropertiesD](uint32 DecoratorIndex, FName PropertyName)
 				{
 					return DecoratorPropertiesD[DecoratorIndex][PropertyName];
 				},
-				[](uint32 DecoratorIndex, const FString& PropertyName)
+				[](uint32 DecoratorIndex, FName PropertyName)
 				{
-					return false;
+					return MAX_uint16;
 				});
 			DecoratorWriter.EndNodeWriting();
 
@@ -458,45 +467,45 @@ bool FAnimationAnimNextRuntimeTest_IUpdate::RunTest(const FString& InParameters)
 			NodeHandles.Add(DecoratorWriter.RegisterNode(*NodeTemplateB));
 
 			// We don't have decorator properties
-			TArray<TMap<FString, FString>> DecoratorPropertiesA;
+			TArray<TMap<FName, FString>> DecoratorPropertiesA;
 			DecoratorPropertiesA.AddDefaulted(NodeTemplateDecoratorListA.Num());
 
-			TArray<TMap<FString, FString>> DecoratorPropertiesB;
+			TArray<TMap<FName, FString>> DecoratorPropertiesB;
 			DecoratorPropertiesB.AddDefaulted(NodeTemplateDecoratorListB.Num());
 			DecoratorPropertiesB[0].Add(TEXT("Child"), ToString<FDecoratorWithOneChild::FSharedData>(TEXT("Child"), FAnimNextDecoratorHandle(NodeHandles[1])));
 
-			TArray<TMap<FString, FString>> DecoratorPropertiesC;
+			TArray<TMap<FName, FString>> DecoratorPropertiesC;
 			DecoratorPropertiesC.AddDefaulted(NodeTemplateDecoratorListC.Num());
 			FAnimNextDecoratorHandle ChildrenHandlesC[2] = { FAnimNextDecoratorHandle(NodeHandles[1]), FAnimNextDecoratorHandle(NodeHandles[2])};
 			DecoratorPropertiesC[0].Add(TEXT("Children"), ToString<FDecoratorWithChildren::FSharedData>(TEXT("Children"), ChildrenHandlesC));
 
 			DecoratorWriter.BeginNodeWriting();
 			DecoratorWriter.WriteNode(NodeHandles[0],
-				[&DecoratorPropertiesC](uint32 DecoratorIndex, const FString& PropertyName)
+				[&DecoratorPropertiesC](uint32 DecoratorIndex, FName PropertyName)
 				{
 					return DecoratorPropertiesC[DecoratorIndex][PropertyName];
 				},
-				[](uint32 DecoratorIndex, const FString& PropertyName)
+				[](uint32 DecoratorIndex, FName PropertyName)
 				{
-					return false;
+					return MAX_uint16;
 				});
 			DecoratorWriter.WriteNode(NodeHandles[1],
-				[&DecoratorPropertiesA](uint32 DecoratorIndex, const FString& PropertyName)
+				[&DecoratorPropertiesA](uint32 DecoratorIndex, FName PropertyName)
 				{
 					return DecoratorPropertiesA[DecoratorIndex][PropertyName];
 				},
-				[](uint32 DecoratorIndex, const FString& PropertyName)
+				[](uint32 DecoratorIndex, FName PropertyName)
 				{
-					return false;
+					return MAX_uint16;
 				});
 			DecoratorWriter.WriteNode(NodeHandles[2],
-				[&DecoratorPropertiesB](uint32 DecoratorIndex, const FString& PropertyName)
+				[&DecoratorPropertiesB](uint32 DecoratorIndex, FName PropertyName)
 				{
 					return DecoratorPropertiesB[DecoratorIndex][PropertyName];
 				},
-				[](uint32 DecoratorIndex, const FString& PropertyName)
+				[](uint32 DecoratorIndex, FName PropertyName)
 				{
-					return false;
+					return MAX_uint16;
 				});
 			DecoratorWriter.EndNodeWriting();
 
@@ -596,14 +605,14 @@ bool FAnimationAnimNextRuntimeTest_IEvaluate::RunTest(const FString& InParameter
 			NodeHandles.Add(DecoratorWriter.RegisterNode(*NodeTemplateB));
 
 			// We don't have decorator properties
-			TArray<TMap<FString, FString>> DecoratorPropertiesA;
+			TArray<TMap<FName, FString>> DecoratorPropertiesA;
 			DecoratorPropertiesA.AddDefaulted(NodeTemplateDecoratorListA.Num());
 
-			TArray<TMap<FString, FString>> DecoratorPropertiesB;
+			TArray<TMap<FName, FString>> DecoratorPropertiesB;
 			DecoratorPropertiesB.AddDefaulted(NodeTemplateDecoratorListB.Num());
 			DecoratorPropertiesB[0].Add(TEXT("Child"), ToString<FDecoratorWithOneChild::FSharedData>(TEXT("Child"), FAnimNextDecoratorHandle(NodeHandles[1])));
 
-			TArray<TMap<FString, FString>> DecoratorPropertiesC;
+			TArray<TMap<FName, FString>> DecoratorPropertiesC;
 			DecoratorPropertiesC.AddDefaulted(NodeTemplateDecoratorListC.Num());
 
 			FAnimNextDecoratorHandle ChildrenHandlesC[2] = { FAnimNextDecoratorHandle(NodeHandles[1]), FAnimNextDecoratorHandle(NodeHandles[2])};
@@ -611,31 +620,31 @@ bool FAnimationAnimNextRuntimeTest_IEvaluate::RunTest(const FString& InParameter
 
 			DecoratorWriter.BeginNodeWriting();
 			DecoratorWriter.WriteNode(NodeHandles[0],
-				[&DecoratorPropertiesC](uint32 DecoratorIndex, const FString& PropertyName)
+				[&DecoratorPropertiesC](uint32 DecoratorIndex, FName PropertyName)
 				{
 					return DecoratorPropertiesC[DecoratorIndex][PropertyName];
 				},
-				[](uint32 DecoratorIndex, const FString& PropertyName)
+				[](uint32 DecoratorIndex, FName PropertyName)
 				{
-					return false;
+					return MAX_uint16;
 				});
 			DecoratorWriter.WriteNode(NodeHandles[1],
-				[&DecoratorPropertiesA](uint32 DecoratorIndex, const FString& PropertyName)
+				[&DecoratorPropertiesA](uint32 DecoratorIndex, FName PropertyName)
 				{
 					return DecoratorPropertiesA[DecoratorIndex][PropertyName];
 				},
-				[](uint32 DecoratorIndex, const FString& PropertyName)
+				[](uint32 DecoratorIndex, FName PropertyName)
 				{
-					return false;
+					return MAX_uint16;
 				});
 			DecoratorWriter.WriteNode(NodeHandles[2],
-				[&DecoratorPropertiesB](uint32 DecoratorIndex, const FString& PropertyName)
+				[&DecoratorPropertiesB](uint32 DecoratorIndex, FName PropertyName)
 				{
 					return DecoratorPropertiesB[DecoratorIndex][PropertyName];
 				},
-				[](uint32 DecoratorIndex, const FString& PropertyName)
+				[](uint32 DecoratorIndex, FName PropertyName)
 				{
-					return false;
+					return MAX_uint16;
 				});
 			DecoratorWriter.EndNodeWriting();
 
diff --git a/Engine/Plugins/Experimental/AnimNext/Source/AnimNextUncookedOnly/Private/UncookedOnlyUtils.cpp b/Engine/Plugins/Experimental/AnimNext/Source/AnimNextUncookedOnly/Private/UncookedOnlyUtils.cpp
--- a/Engine/Plugins/Experimental/AnimNext/Source/AnimNextUncookedOnly/Private/UncookedOnlyUtils.cpp
+++ b/Engine/Plugins/Experimental/AnimNext/Source/AnimNextUncookedOnly/Private/UncookedOnlyUtils.cpp
@@ -61,6 +61,9 @@ namespace Private
 		// The AnimNext decorator
 		const FDecorator* Decorator = nullptr;
 
+		// A map from latent property names to their corresponding RigVM memory handle index
+		TMap<FName, uint16> LatentPropertyNameToIndexMap;
+
 		FDecoratorEntryMapping(const URigVMNode* InDecoratorStackNode, const URigVMPin* InDecoratorEntryPin, const FDecorator* InDecorator)
 			: DecoratorStackNode(InDecoratorStackNode)
 			, DecoratorEntryPin(InDecoratorEntryPin)
@@ -132,7 +135,7 @@ namespace Private
 		return DecoratorWriter.RegisterNode(*NodeTemplate);
 	}
 
-	FString GetDecoratorProperty(const FDecoratorStackMapping& DecoratorStack, uint32 DecoratorIndex, const FString& PropertyName, const TArray<FDecoratorStackMapping>& DecoratorStackNodes)
+	FString GetDecoratorProperty(const FDecoratorStackMapping& DecoratorStack, uint32 DecoratorIndex, FName PropertyName, const TArray<FDecoratorStackMapping>& DecoratorStackNodes)
 	{
 		const TArray<URigVMPin*>& Pins = DecoratorStack.DecoratorEntries[DecoratorIndex].DecoratorEntryPin->GetSubPins();
 		for (const URigVMPin* Pin : Pins)
@@ -142,7 +145,7 @@ namespace Private
 				continue;	// We only look for input pins
 			}
 
-			if (Pin->GetName() == PropertyName)
+			if (Pin->GetFName() == PropertyName)
 			{
 				if (Pin->GetCPPTypeObject() == FAnimNextDecoratorHandle::StaticStruct())
 				{
@@ -200,37 +203,27 @@ namespace Private
 		return FString();
 	}
 
-	bool IsDecoratorPropertyLatent(const FDecoratorStackMapping& DecoratorStack, uint32 DecoratorIndex, const FString& PropertyName, const TArray<FDecoratorStackMapping>& DecoratorStackNodes)
+	uint16 GetDecoratorLatentPropertyIndex(const FDecoratorStackMapping& DecoratorStack, uint32 DecoratorIndex, FName PropertyName)
 	{
-		const TArray<URigVMPin*>& Pins = DecoratorStack.DecoratorEntries[DecoratorIndex].DecoratorEntryPin->GetSubPins();
-		for (const URigVMPin* Pin : Pins)
+		const FDecoratorEntryMapping& Entry = DecoratorStack.DecoratorEntries[DecoratorIndex];
+		if (const uint16* RigVMIndex = Entry.LatentPropertyNameToIndexMap.Find(PropertyName))
 		{
-			if (Pin->GetDirection() != ERigVMPinDirection::Input)
-			{
-				continue;	// We only look for input pins
-			}
-
-			if (Pin->GetName() == PropertyName)
-			{
-				// Lazy pins are latent if they are connected to something
-				return Pin->IsLazy() && !Pin->GetLinks().IsEmpty();
-			}
+			return *RigVMIndex;
 		}
 
-		// Unknown property
-		return false;
+		return MAX_uint16;
 	}
 
 	void WriteDecoratorProperties(FDecoratorWriter& DecoratorWriter, const FDecoratorStackMapping& Mapping, const TArray<FDecoratorStackMapping>& DecoratorStackNodes)
 	{
 		DecoratorWriter.WriteNode(Mapping.DecoratorStackNodeHandle,
-			[&Mapping, &DecoratorStackNodes](uint32 DecoratorIndex, const FString& PropertyName)
+			[&Mapping, &DecoratorStackNodes](uint32 DecoratorIndex, FName PropertyName)
 			{
 				return GetDecoratorProperty(Mapping, DecoratorIndex, PropertyName, DecoratorStackNodes);
 			},
-			[&Mapping, &DecoratorStackNodes](uint32 DecoratorIndex, const FString& PropertyName)
+			[&Mapping](uint32 DecoratorIndex, FName PropertyName)
 			{
-				return IsDecoratorPropertyLatent(Mapping, DecoratorIndex, PropertyName, DecoratorStackNodes);
+				return GetDecoratorLatentPropertyIndex(Mapping, DecoratorIndex, PropertyName);
 			});
 	}
 
@@ -413,21 +406,24 @@ namespace Private
 		return DecoratorStackNodes;
 	}
 
-	FRigVMPinInfoArray CollectLatentPins(const TArray<FDecoratorStackMapping>& DecoratorStackNodes, TMap<FName, URigVMPin*>& LatentPinMapping)
+	FRigVMPinInfoArray CollectLatentPins(TArray<FDecoratorStackMapping>& DecoratorStackNodes, TMap<FName, URigVMPin*>& LatentPinMapping)
 	{
 		const FRigVMRegistry& Registry = FRigVMRegistry::Get();
 
 		FRigVMPinInfoArray LatentPins;
 
-		for (const FDecoratorStackMapping& DecoratorStack : DecoratorStackNodes)
+		for (FDecoratorStackMapping& DecoratorStack : DecoratorStackNodes)
 		{
-			for (const FDecoratorEntryMapping& DecoratorEntry : DecoratorStack.DecoratorEntries)
+			for (FDecoratorEntryMapping& DecoratorEntry : DecoratorStack.DecoratorEntries)
 			{
 				for (URigVMPin* Pin : DecoratorEntry.DecoratorEntryPin->GetSubPins())
 				{
 					if (Pin->IsLazy() && !Pin->GetLinks().IsEmpty())
 					{
 						// This pin has something linked to it, it is a latent pin
+						check(LatentPins.Num() < ((1 << 16) - 1));	// We reserve MAX_uint16 as an invalid value and we must fit on 15 bits when packed
+						DecoratorEntry.LatentPropertyNameToIndexMap.Add(Pin->GetFName(), (uint16)LatentPins.Num());
+
 						const FName LatentPinName(TEXT("LatentPin"), LatentPins.Num());	// Create unique latent pin names
 
 						FRigVMPinInfo PinInfo;
