Fix for shutdown crash at AnimNextEditorModule, due to static classes not returning a valid FName

diff --git a/Engine/Plugins/Experimental/AnimNext/Source/AnimNextEditor/Private/AnimNextEditorModule.cpp b/Engine/Plugins/Experimental/AnimNext/Source/AnimNextEditor/Private/AnimNextEditorModule.cpp
--- a/Engine/Plugins/Experimental/AnimNext/Source/AnimNextEditor/Private/AnimNextEditorModule.cpp
+++ b/Engine/Plugins/Experimental/AnimNext/Source/AnimNextEditor/Private/AnimNextEditorModule.cpp
@@ -200,7 +200,7 @@ class FModule : public IModule
 			FPropertyEditorModule& PropertyModule = FModuleManager::LoadModuleChecked<FPropertyEditorModule>("PropertyEditor");
 			PropertyModule.UnregisterCustomPropertyTypeLayout("AnimNextParamType");
 			PropertyModule.UnregisterCustomPropertyTypeLayout("AnimNextParam");
-			PropertyModule.UnregisterCustomPropertyTypeLayout(FNameProperty::StaticClass()->GetFName());
+			PropertyModule.UnregisterCustomPropertyTypeLayout("NameProperty");
 			PropertyModule.UnregisterCustomClassLayout("AnimNextParameterBlockParameter");
 			PropertyModule.UnregisterCustomClassLayout("AnimNextGraph_EdGraphNode");
 		}
@@ -209,9 +209,9 @@ class FModule : public IModule
 
 		FEdGraphUtilities::UnregisterVisualPinFactory(ParametersGraphPanelPinFactory);
 
-		FWorkspaceEditor::UnregisterAssetDocumentWidget(UAnimNextSchedule::StaticClass()->GetFName());
-		FWorkspaceEditor::UnregisterAssetDocumentWidget(UAnimNextParameterBlock::StaticClass()->GetFName());
-		FWorkspaceEditor::UnregisterAssetDocumentWidget(UAnimNextGraph::StaticClass()->GetFName());
+		FWorkspaceEditor::UnregisterAssetDocumentWidget("AnimNextSchedule");
+		FWorkspaceEditor::UnregisterAssetDocumentWidget("AnimNextParameterBlock");
+		FWorkspaceEditor::UnregisterAssetDocumentWidget("AnimNextGraph");
 		
 		SRigVMAssetView::UnregisterCategoryFactory("Parameters");
 		SRigVMAssetView::UnregisterCategoryFactory("Parameter Graphs");
