AnimNext: Support for multiple graphs in AnimNext graph assets

Modifies compilation of AnimNext graphs to allow for multiple 'entry points'. Note we still only have one 'rig vm' entry point to update our latent pins.
Instances now also hold the entry point that they correspond to.
Updated subgraph decorator to also support switching entry points.
Updated freezing/thawing of graph instances to use a critical section. Due to entry points being on the same asset the function is now recursive, which RW locks do not support.

Added 'self' accessor unit for use by subgraph decorators. Returns a reference to the currentlly executing AnimNext graph.

Added extraction of parameters from node pins (including decorator pins). This allows parameters to be created for (e.g.) entry points before the actual entry point is created.
Rearranged how exported parameters are gathered. Instead of gathering at save/load time, graphs are traversed during compilation and the results cached on the asset's editor data (they are not serialized however). This is to avoid traversing into pin metadata (which needs to perform an ImportText for decorator data), which cannot be called during save (if FObjectProperties are parsed, then the underlying code will call FindObject, which cannot be called during save).

Changed FindEntryForRigVMGraph to use ptr, rather than name-based matching. RigVM and Edgraphs in entries now have 'generic' names. Actual entry point names are held on the outer entry (the externally packaged object).

Moved schedule compilation to after the package is fully loaded and added preload dependencies, as subobjects could potentially not be loaded correctly in PostLoad when using the ALT.

Additionally:
- Added better tab titles so graphs can be differentiated
- Added categories to dispatch factories so they are better organized in context menus.

RigVM changes:
- Added the ability for unit nodes to display a subtitle
- Removed restiction around hidden pins in URigVMUnitNode::GetStructDefaultValue. This was to allow access to pin defaults on AnimNext graph root nodes, where the entry point's name is stored (used for subtitle display).
- Added GetDecoratorScriptStruct to URigVMNode, so certain data can be retrieved without constructing an instance of the data.
- Added generic metadata access (not just custom widgets) to URigVMPin. Also expanded the metadata access to read data from decorators.

diff --git a/Engine/Plugins/Experimental/AnimNext/Source/AnimNext/Internal/Graph/AnimNextExecuteContext.h b/Engine/Plugins/Experimental/AnimNext/Source/AnimNext/Internal/Graph/AnimNextExecuteContext.h
--- a/Engine/Plugins/Experimental/AnimNext/Source/AnimNext/Internal/Graph/AnimNextExecuteContext.h
+++ b/Engine/Plugins/Experimental/AnimNext/Source/AnimNext/Internal/Graph/AnimNextExecuteContext.h
@@ -7,6 +7,8 @@
 #include "Units/RigUnit.h"
 #include "AnimNextExecuteContext.generated.h"
 
+struct FAnimNextGraphInstance;
+
 namespace UE::AnimNext
 {
 	struct FLatentPropertyHandle;
@@ -19,6 +21,7 @@ struct FAnimNextExecuteContext : public FRigVMExecuteContext
 
 	FAnimNextExecuteContext() = default;
 
+	const FAnimNextGraphInstance& GetGraphInstance() const { check(Instance); return *Instance; }
 	const TConstArrayView<UE::AnimNext::FLatentPropertyHandle>& GetLatentHandles() const { return LatentHandles; }
 	void* GetDestinationBasePtr() const { return DestinationBasePtr; }
 	bool IsFrozen() const { return bIsFrozen; }
@@ -28,14 +31,16 @@ struct FAnimNextExecuteContext : public FRigVMExecuteContext
 		Super::Copy(InOtherContext);
 
 		const FAnimNextExecuteContext* OtherContext = (const FAnimNextExecuteContext*)InOtherContext;
+		Instance = OtherContext->Instance;
 		LatentHandles = OtherContext->LatentHandles;
 		DestinationBasePtr = OtherContext->DestinationBasePtr;
 		bIsFrozen = OtherContext->bIsFrozen;
 	}
 
 private:
-	void SetupForExecution(const TConstArrayView<UE::AnimNext::FLatentPropertyHandle>& InLatentHandles, void* InDestinationBasePtr, bool bInIsFrozen)
+	void SetupForExecution(const FAnimNextGraphInstance* InInstance, const TConstArrayView<UE::AnimNext::FLatentPropertyHandle>& InLatentHandles, void* InDestinationBasePtr, bool bInIsFrozen)
 	{
+		Instance = InInstance;
 		LatentHandles = InLatentHandles;
 		DestinationBasePtr = InDestinationBasePtr;
 		bIsFrozen = bInIsFrozen;
@@ -44,11 +49,13 @@ private:
 	// Call this to reset the context to its original state to detect stale usage, can't call it Reset due to virtual in base with that name
 	void DebugReset()
 	{
+		Instance = nullptr;
 		LatentHandles = TConstArrayView<UE::AnimNext::FLatentPropertyHandle>();
 		DestinationBasePtr = nullptr;
 		bIsFrozen = false;
 	}
 
+	const FAnimNextGraphInstance* Instance = nullptr;
 	TConstArrayView<UE::AnimNext::FLatentPropertyHandle> LatentHandles;
 	void* DestinationBasePtr = nullptr;
 	bool bIsFrozen = false;
diff --git a/Engine/Plugins/Experimental/AnimNext/Source/AnimNext/Internal/Graph/AnimNextGraph.h b/Engine/Plugins/Experimental/AnimNext/Source/AnimNext/Internal/Graph/AnimNextGraph.h
--- a/Engine/Plugins/Experimental/AnimNext/Source/AnimNext/Internal/Graph/AnimNextGraph.h
+++ b/Engine/Plugins/Experimental/AnimNext/Source/AnimNext/Internal/Graph/AnimNextGraph.h
@@ -3,7 +3,9 @@
 #pragma once
 
 #include "CoreMinimal.h"
+#include "AnimNextGraphEntryPoint.h"
 #include "AnimNextRigVMAsset.h"
+#include "RigUnit_AnimNextGraphRoot.h"
 #include "RigVMCore/RigVM.h"
 #include "DecoratorBase/DecoratorPtr.h"
 #include "DecoratorBase/DecoratorHandle.h"
@@ -71,10 +73,15 @@ public:
 	virtual TConstArrayView<UE::AnimNext::FScheduleTerm> GetTerms() const override;
 
 	// Allocates an instance of the graph
-	void AllocateInstance(FAnimNextGraphInstancePtr& Instance) const;
+	// @param	OutInstance		The instance to allocate data for
+	// @param	InEntryPoint	The entry point to use. If this is NAME_None then the default entry point for this graph is used
+	void AllocateInstance(FAnimNextGraphInstancePtr& OutInstance, FName InEntryPoint = NAME_None) const;
 
 	// Allocates an instance of the graph with the specified parent graph instance
-	void AllocateInstance(FAnimNextGraphInstance& ParentGraphInstance, FAnimNextGraphInstancePtr& Instance) const;
+	// @param	InOutParentGraphInstance	The parent graph instance to use
+	// @param	OutInstance					The instance to allocate data for
+	// @param	InEntryPoint				The entry point to use. If this is NAME_None then the default entry point for this graph is used
+	void AllocateInstance(FAnimNextGraphInstance& InOutParentGraphInstance, FAnimNextGraphInstancePtr& OutInstance, FName InEntryPoint = NAME_None) const;
 
 	// Get the parameter to use to access the reference pose
 	UE::AnimNext::FParamId GetReferencePoseParam() const { return ReferencePoseId; }
@@ -88,7 +95,7 @@ protected:
 	bool LoadFromArchiveBuffer(const TArray<uint8>& SharedDataArchiveBuffer);
 
 	// Allocates an instance of the graph with an optional parent graph instance
-	void AllocateInstanceImpl(FAnimNextGraphInstance* ParentGraphInstance, FAnimNextGraphInstancePtr& Instance) const;
+	void AllocateInstanceImpl(FAnimNextGraphInstance* InOutParentGraphInstance, FAnimNextGraphInstancePtr& OutInstance, FName InEntryPoint) const;
 
 #if WITH_EDITORONLY_DATA
 	// During graph compilation, if we have existing graph instances, we freeze them by releasing their memory before thawing them
@@ -113,22 +120,22 @@ protected:
 	friend UE::AnimNext::FModule;
 	
 #if WITH_EDITORONLY_DATA
-	mutable FRWLock GraphInstancesLock;
+	mutable FCriticalSection GraphInstancesLock;
 
 	// This is a list of live graph instances that have been allocated, used in the editor to reset instances when we re-compile/live edit
 	mutable TSet<FAnimNextGraphInstance*> GraphInstances;
 #endif
 
-	// This is the execute method definition used by this graph
+	// This is the execute method definition used by a graph to evaluate latent pins
 	UPROPERTY()
 	FAnimNextGraphEvaluatorExecuteDefinition ExecuteDefinition;
 
-	// This is a handle to the root decorator in our graph
+	// Data for each entry point in this graph
 	UPROPERTY()
-	FAnimNextEntryPointHandle RootDecoratorHandle;
+	TArray<FAnimNextGraphEntryPoint> EntryPoints;
 
-	// This is a resolved handle to the root decorator in our graph
-	FAnimNextDecoratorHandle ResolvedRootDecoratorHandle;
+	// This is a resolved handle to the root decorator in our graph, for each entry point 
+	TMap<FName, FAnimNextDecoratorHandle> ResolvedRootDecoratorHandles;
 
 	// This is the graph shared data used by the decorator system, the output of FDecoratorReader
 	// We de-serialize manually into this buffer from the archive buffer, this is never saved on disk
@@ -141,6 +148,10 @@ protected:
 	UPROPERTY()
 	TArray<TObjectPtr<UObject>> GraphReferencedObjects;
 
+	// The entry point that this graph defaults to using
+	UPROPERTY(EditAnywhere, Category = "Graph")
+	FName DefaultEntryPoint = FRigUnit_AnimNextGraphRoot::DefaultEntryPoint;
+
 	// The parameter to use to access the reference pose
 	UPROPERTY(EditAnywhere, Category = "Graph", meta=(CustomWidget = "ParamName", AllowedParamType = "FAnimNextGraphReferencePose"))
 	FName ReferencePose = DefaultReferencePoseId.GetName();
diff --git a/Engine/Plugins/Experimental/AnimNext/Source/AnimNext/Internal/Graph/AnimNextGraphEntryPoint.h b/Engine/Plugins/Experimental/AnimNext/Source/AnimNext/Internal/Graph/AnimNextGraphEntryPoint.h
new file mode 100644
--- /dev/null
+++ b/Engine/Plugins/Experimental/AnimNext/Source/AnimNext/Internal/Graph/AnimNextGraphEntryPoint.h
@@ -0,0 +1,23 @@
+﻿// Copyright Epic Games, Inc. All Rights Reserved.
+
+#pragma once
+
+#include "CoreMinimal.h"
+#include "RigUnit_AnimNextGraphEvaluator.h"
+#include "DecoratorBase/EntryPointHandle.h"
+#include "AnimNextGraphEntryPoint.generated.h"
+
+USTRUCT()
+struct FAnimNextGraphEntryPoint
+{
+	GENERATED_BODY()
+
+	// The name of the entry point
+	UPROPERTY()
+	FName EntryPointName;
+
+	// This is a handle to the root decorator for a graph
+	UPROPERTY()
+	FAnimNextEntryPointHandle RootDecoratorHandle;
+};
+
diff --git a/Engine/Plugins/Experimental/AnimNext/Source/AnimNext/Internal/Graph/RigUnit_AnimNextGraphRoot.h b/Engine/Plugins/Experimental/AnimNext/Source/AnimNext/Internal/Graph/RigUnit_AnimNextGraphRoot.h
--- a/Engine/Plugins/Experimental/AnimNext/Source/AnimNext/Internal/Graph/RigUnit_AnimNextGraphRoot.h
+++ b/Engine/Plugins/Experimental/AnimNext/Source/AnimNext/Internal/Graph/RigUnit_AnimNextGraphRoot.h
@@ -23,6 +23,7 @@ struct ANIMNEXT_API FRigUnit_AnimNextGraphRoot : public FRigUnit_AnimNextBase
 	void DummyExecute();
 
 	virtual FName GetEventName() const override { return EventName; }
+	virtual FString GetUnitSubTitle() const override { return EntryPoint.ToString(); };
 	virtual bool CanOnlyExistOnce() const override { return true; }
 
 	// The execution result
@@ -34,6 +35,13 @@ struct ANIMNEXT_API FRigUnit_AnimNextGraphRoot : public FRigUnit_AnimNextBase
 	UPROPERTY()
 	FAnimNextExecuteContext ExecuteContext;
 
+	// The name of the entry point
+	UPROPERTY(VisibleAnywhere, Category = Result, meta = (Hidden))
+	FName EntryPoint = DefaultEntryPoint;
+
 	// This unit is our graph entry point, it needs an event so we can call it
 	static FName EventName;
+
+	// Default entry point name
+	static FName DefaultEntryPoint;
 };
diff --git a/Engine/Plugins/Experimental/AnimNext/Source/AnimNext/Internal/Graph/RigUnit_AnimNextGraphSelf.h b/Engine/Plugins/Experimental/AnimNext/Source/AnimNext/Internal/Graph/RigUnit_AnimNextGraphSelf.h
new file mode 100644
--- /dev/null
+++ b/Engine/Plugins/Experimental/AnimNext/Source/AnimNext/Internal/Graph/RigUnit_AnimNextGraphSelf.h
@@ -0,0 +1,23 @@
+﻿// Copyright Epic Games, Inc. All Rights Reserved.
+
+#pragma once
+
+#include "AnimNextExecuteContext.h"
+
+#include "RigUnit_AnimNextGraphSelf.generated.h"
+
+class UAnimNextGraph;
+
+/** Get a reference to the currently executing graph */
+USTRUCT(meta=(DisplayName="Self", Category="Graph", NodeColor="0, 0, 1", Keywords="Current,This"))
+struct ANIMNEXT_API FRigUnit_AnimNextGraphSelf : public FRigUnit_AnimNextBase
+{
+	GENERATED_BODY()
+
+	RIGVM_METHOD()
+	void Execute();
+
+	// The currently-executing graph
+	UPROPERTY(meta=(Output))
+	TObjectPtr<UAnimNextGraph> Self;
+};
diff --git a/Engine/Plugins/Experimental/AnimNext/Source/AnimNext/Internal/Param/RigUnit_AnimNextParameterBeginExecution.h b/Engine/Plugins/Experimental/AnimNext/Source/AnimNext/Internal/Param/RigUnit_AnimNextParameterBeginExecution.h
--- a/Engine/Plugins/Experimental/AnimNext/Source/AnimNext/Internal/Param/RigUnit_AnimNextParameterBeginExecution.h
+++ b/Engine/Plugins/Experimental/AnimNext/Source/AnimNext/Internal/Param/RigUnit_AnimNextParameterBeginExecution.h
@@ -16,12 +16,17 @@ struct ANIMNEXT_API FRigUnit_AnimNextParameterBeginExecution : public FRigUnit_A
 	RIGVM_METHOD()
 	void Execute();
 
-	virtual FName GetEventName() const override { return EventName; }
+	virtual FString GetUnitLabel() const override { return EntryPoint.ToString(); };
+	virtual FName GetEventName() const override { return EntryPoint; }
 	virtual bool CanOnlyExistOnce() const override { return true; }
 
 	// The execution result
-	UPROPERTY(EditAnywhere, DisplayName = "Execute", Category = "BeginExecution", meta = (Output))
+	UPROPERTY(EditAnywhere, DisplayName = "Execute", Category = "Entry Point", meta = (Output))
 	FAnimNextParameterExecuteContext ExecuteContext;
 
-	static FName EventName;
+	// The name of the entry point
+	UPROPERTY(VisibleAnywhere, Category = "Entry Point", meta = (Hidden))
+	FName EntryPoint = DefaultEntryPoint;
+
+	static FName DefaultEntryPoint;
 };
diff --git a/Engine/Plugins/Experimental/AnimNext/Source/AnimNext/Internal/Param/RigVMDispatch_GetLayerParameter.h b/Engine/Plugins/Experimental/AnimNext/Source/AnimNext/Internal/Param/RigVMDispatch_GetLayerParameter.h
--- a/Engine/Plugins/Experimental/AnimNext/Source/AnimNext/Internal/Param/RigVMDispatch_GetLayerParameter.h
+++ b/Engine/Plugins/Experimental/AnimNext/Source/AnimNext/Internal/Param/RigVMDispatch_GetLayerParameter.h
@@ -14,7 +14,7 @@ namespace UE::AnimNext::UncookedOnly
 /*
  * Gets a parameter's current value from the parameter block value
  */
-USTRUCT(meta = (DisplayName = "Get Parameter from Block", NodeColor = "0.8, 0, 0.2, 1"))
+USTRUCT(meta = (DisplayName = "Get Parameter from Block", Category="Parameters", NodeColor = "0.8, 0, 0.2, 1"))
 struct ANIMNEXT_API FRigVMDispatch_GetLayerParameter : public FRigVMDispatchFactory
 {
 	GENERATED_BODY()
diff --git a/Engine/Plugins/Experimental/AnimNext/Source/AnimNext/Internal/Param/RigVMDispatch_GetParameter.h b/Engine/Plugins/Experimental/AnimNext/Source/AnimNext/Internal/Param/RigVMDispatch_GetParameter.h
--- a/Engine/Plugins/Experimental/AnimNext/Source/AnimNext/Internal/Param/RigVMDispatch_GetParameter.h
+++ b/Engine/Plugins/Experimental/AnimNext/Source/AnimNext/Internal/Param/RigVMDispatch_GetParameter.h
@@ -14,7 +14,7 @@ namespace UE::AnimNext::UncookedOnly
 /*
  * Gets a parameter's value
  */
-USTRUCT(meta = (DisplayName = "Get Parameter", NodeColor = "0.8, 0, 0.2, 1"))
+USTRUCT(meta = (DisplayName = "Get Parameter", Category="Parameters", NodeColor = "0.8, 0, 0.2, 1"))
 struct ANIMNEXT_API FRigVMDispatch_GetParameter : public FRigVMDispatchFactory
 {
 	GENERATED_BODY()
diff --git a/Engine/Plugins/Experimental/AnimNext/Source/AnimNext/Internal/Param/RigVMDispatch_SetLayerParameter.h b/Engine/Plugins/Experimental/AnimNext/Source/AnimNext/Internal/Param/RigVMDispatch_SetLayerParameter.h
--- a/Engine/Plugins/Experimental/AnimNext/Source/AnimNext/Internal/Param/RigVMDispatch_SetLayerParameter.h
+++ b/Engine/Plugins/Experimental/AnimNext/Source/AnimNext/Internal/Param/RigVMDispatch_SetLayerParameter.h
@@ -14,7 +14,7 @@ struct FUtils;
 /*
  * Sets a parameter's value
  */
-USTRUCT(meta=(DisplayName = "Set Block Parameter", NodeColor = "0.8, 0, 0.2, 1"))
+USTRUCT(meta=(DisplayName = "Set Block Parameter", Category="Parameters", NodeColor = "0.8, 0, 0.2, 1"))
 struct ANIMNEXT_API FRigVMDispatch_SetLayerParameter : public FRigVMDispatchFactory
 {
 	GENERATED_BODY()
diff --git a/Engine/Plugins/Experimental/AnimNext/Source/AnimNext/Internal/Scheduler/AnimNextSchedule.h b/Engine/Plugins/Experimental/AnimNext/Source/AnimNext/Internal/Scheduler/AnimNextSchedule.h
--- a/Engine/Plugins/Experimental/AnimNext/Source/AnimNext/Internal/Scheduler/AnimNextSchedule.h
+++ b/Engine/Plugins/Experimental/AnimNext/Source/AnimNext/Internal/Scheduler/AnimNextSchedule.h
@@ -73,6 +73,9 @@ class UAnimNextScheduleEntry_AnimNextGraph : public UAnimNextScheduleEntry
 private:
 	friend struct UE::AnimNext::UncookedOnly::FUtils;
 
+	// UObject interface
+	virtual void GetPreloadDependencies(TArray<UObject*>& OutDeps) override;
+
 	// The graph to run by default
 	UPROPERTY(EditAnywhere, Category = "Graph")
 	TObjectPtr<UAnimNextGraph> Graph = nullptr;
@@ -132,6 +135,9 @@ class UAnimNextScheduleEntry_ParamScope : public UAnimNextScheduleEntry
 private:
 	friend struct UE::AnimNext::UncookedOnly::FUtils;
 
+	// UObject interface
+	virtual void GetPreloadDependencies(TArray<UObject*>& OutDeps) override;
+
 	// The scope to use
 	UPROPERTY(EditAnywhere, Category = "Parameters", meta = (CustomWidget = "ParamName", AllowedParamType = "FAnimNextScope"))
 	FName Scope;
@@ -239,6 +245,9 @@ private:
 	virtual void GetAssetRegistryTags(FAssetRegistryTagsContext Context) const override;
 	UE_DEPRECATED(5.4, "Implement the version that takes FAssetRegistryTagsContext instead.")
 	virtual void GetAssetRegistryTags(TArray<FAssetRegistryTag>& OutTags) const override;
+	virtual void GetPreloadDependencies(TArray<UObject*>& OutDeps) override;
+
+	void HandlePackageDone(const FEndLoadPackageContext& Context);
 
 	// Compile the editor data into a compact runtime representation
 	void CompileSchedule();
diff --git a/Engine/Plugins/Experimental/AnimNext/Source/AnimNext/Private/Decorators/SubGraphHost.cpp b/Engine/Plugins/Experimental/AnimNext/Source/AnimNext/Private/Decorators/SubGraphHost.cpp
--- a/Engine/Plugins/Experimental/AnimNext/Source/AnimNext/Private/Decorators/SubGraphHost.cpp
+++ b/Engine/Plugins/Experimental/AnimNext/Source/AnimNext/Private/Decorators/SubGraphHost.cpp
@@ -3,6 +3,7 @@
 #include "Decorators/SubGraphHost.h"
 
 #include "DecoratorBase/ExecutionContext.h"
+#include "Graph/AnimNextGraphInstance.h"
 
 namespace UE::AnimNext
 {
@@ -67,14 +68,25 @@ namespace UE::AnimNext
 		const bool bHasActiveSubGraph = InstanceData->CurrentlyActiveSubGraphIndex != INDEX_NONE;
 
 		TObjectPtr<const UAnimNextGraph> CurrentActiveSubGraph;
+		FName CurrentActiveEntryPoint = NAME_None;
 		if (bHasActiveSubGraph)
 		{
-			CurrentActiveSubGraph = InstanceData->SubGraphSlots[InstanceData->CurrentlyActiveSubGraphIndex].SubGraph;
+			const FSubGraphSlot& SubGraphSlot = InstanceData->SubGraphSlots[InstanceData->CurrentlyActiveSubGraphIndex];
+			CurrentActiveSubGraph = SubGraphSlot.SubGraph;
+			CurrentActiveEntryPoint = SubGraphSlot.EntryPoint;
 		}
 
 		const TObjectPtr<const UAnimNextGraph> DesiredSubGraph = SharedData->GetSubGraph(Context, Binding);
+		const FName EntryPoint = SharedData->GetEntryPoint(Context, Binding);
 
-		if (CurrentActiveSubGraph != DesiredSubGraph || !bHasActiveSubGraph)
+		// Check for reentrancy and early-out if we are linking back to the current instance
+		FAnimNextGraphInstance& GraphInstance = Context.GetGraphInstance();
+		if(GraphInstance.UsesGraph(DesiredSubGraph) && GraphInstance.UsesEntryPoint(EntryPoint))
+		{
+			return;
+		}
+
+		if (!bHasActiveSubGraph || CurrentActiveSubGraph != DesiredSubGraph || CurrentActiveEntryPoint != EntryPoint)
 		{
 			// Find an empty slot we can use
 			int32 FreeSlotIndex = INDEX_NONE;
@@ -99,6 +111,7 @@ namespace UE::AnimNext
 			FSubGraphSlot& SubGraphSlot = InstanceData->SubGraphSlots[FreeSlotIndex];
 			SubGraphSlot.SubGraph = DesiredSubGraph;
 			SubGraphSlot.State = DesiredSubGraph ? ESlotState::ActiveWithGraph : ESlotState::ActiveWithReferencePose;
+			SubGraphSlot.EntryPoint = EntryPoint;
 
 			const int32 OldChildIndex = InstanceData->CurrentlyActiveSubGraphIndex;
 			const int32 NewChildIndex = FreeSlotIndex;
@@ -188,7 +201,7 @@ namespace UE::AnimNext
 
 			if (SubGraphEntry.State == ESlotState::ActiveWithGraph)
 			{
-				SubGraphEntry.SubGraph->AllocateInstance(Context.GetGraphInstance(), SubGraphEntry.GraphInstance);
+				SubGraphEntry.SubGraph->AllocateInstance(Context.GetGraphInstance(), SubGraphEntry.GraphInstance, SubGraphEntry.EntryPoint);
 			}
 		}
 	}
diff --git a/Engine/Plugins/Experimental/AnimNext/Source/AnimNext/Private/Decorators/SubGraphHost.h b/Engine/Plugins/Experimental/AnimNext/Source/AnimNext/Private/Decorators/SubGraphHost.h
--- a/Engine/Plugins/Experimental/AnimNext/Source/AnimNext/Private/Decorators/SubGraphHost.h
+++ b/Engine/Plugins/Experimental/AnimNext/Source/AnimNext/Private/Decorators/SubGraphHost.h
@@ -27,9 +27,14 @@ struct FAnimNextSubGraphHostDecoratorSharedData : public FAnimNextDecoratorShare
 	UPROPERTY(meta = (Hidden))
 	FAnimNextDecoratorHandle ReferencePoseChild;
 
+	/** Entry point in the Subgraph that we will use */
+	UPROPERTY(EditAnywhere, Category = "Default", meta=(CustomWidget = "ParamName", AllowedParamType = "FAnimNextEntryPoint"))
+	FName EntryPoint;
+
 	// Latent pin support boilerplate
 	#define DECORATOR_LATENT_PROPERTIES_ENUMERATOR(GeneratorMacro) \
 		GeneratorMacro(SubGraph) \
+		GeneratorMacro(EntryPoint) \
 
 	GENERATE_DECORATOR_LATENT_PROPERTIES(FAnimNextSubGraphHostDecoratorSharedData, DECORATOR_LATENT_PROPERTIES_ENUMERATOR)
 	#undef DECORATOR_LATENT_PROPERTIES_ENUMERATOR
@@ -61,6 +66,9 @@ namespace UE::AnimNext
 			// The graph instance
 			FAnimNextGraphInstancePtr GraphInstance;
 
+			// The entry point to use
+			FName EntryPoint;
+
 			// The current slot state
 			ESlotState State = ESlotState::Inactive;
 		};
diff --git a/Engine/Plugins/Experimental/AnimNext/Source/AnimNext/Private/Graph/AnimNextGraph.cpp b/Engine/Plugins/Experimental/AnimNext/Source/AnimNext/Private/Graph/AnimNextGraph.cpp
--- a/Engine/Plugins/Experimental/AnimNext/Source/AnimNext/Private/Graph/AnimNextGraph.cpp
+++ b/Engine/Plugins/Experimental/AnimNext/Source/AnimNext/Private/Graph/AnimNextGraph.cpp
@@ -15,6 +15,8 @@
 #include "Graph/AnimNextGraphInstance.h"
 #include "Serialization/MemoryReader.h"
 #include "AnimNextStats.h"
+#include "Graph/AnimNextGraphEntryPoint.h"
+#include "Graph/RigUnit_AnimNextGraphRoot.h"
 
 DEFINE_STAT(STAT_AnimNext_Graph_AllocateInstance);
 
@@ -29,22 +31,24 @@ UAnimNextGraph::UAnimNextGraph(const FObjectInitializer& ObjectInitializer)
 	ExtendedExecuteContext.SetContextPublicDataStruct(FAnimNextExecuteContext::StaticStruct());
 }
 
-void UAnimNextGraph::AllocateInstance(FAnimNextGraphInstancePtr& Instance) const
+void UAnimNextGraph::AllocateInstance(FAnimNextGraphInstancePtr& Instance, FName InEntryPoint) const
 {
-	AllocateInstanceImpl(nullptr, Instance);
+	AllocateInstanceImpl(nullptr, Instance, InEntryPoint);
 }
 
-void UAnimNextGraph::AllocateInstance(FAnimNextGraphInstance& ParentGraphInstance, FAnimNextGraphInstancePtr& Instance) const
+void UAnimNextGraph::AllocateInstance(FAnimNextGraphInstance& ParentGraphInstance, FAnimNextGraphInstancePtr& Instance, FName InEntryPoint) const
 {
-	AllocateInstanceImpl(&ParentGraphInstance, Instance);
+	AllocateInstanceImpl(&ParentGraphInstance, Instance, InEntryPoint);
 }
 
-void UAnimNextGraph::AllocateInstanceImpl(FAnimNextGraphInstance* ParentGraphInstance, FAnimNextGraphInstancePtr& Instance) const
+void UAnimNextGraph::AllocateInstanceImpl(FAnimNextGraphInstance* ParentGraphInstance, FAnimNextGraphInstancePtr& Instance, FName InEntryPoint) const
 {
 	SCOPE_CYCLE_COUNTER(STAT_AnimNext_Graph_AllocateInstance);
 
 	Instance.Release();
 
+	const FName EntryPoint = InEntryPoint == NAME_None ? DefaultEntryPoint : InEntryPoint;
+	const FAnimNextDecoratorHandle ResolvedRootDecoratorHandle = ResolvedRootDecoratorHandles.FindRef(EntryPoint);
 	if (!ResolvedRootDecoratorHandle.IsValid())
 	{
 		return;
@@ -55,6 +59,7 @@ void UAnimNextGraph::AllocateInstanceImpl(FAnimNextGraphInstance* ParentGraphIns
 
 		InstanceImpl->Graph = this;
 		InstanceImpl->ParentGraphInstance = ParentGraphInstance;
+		InstanceImpl->EntryPoint = EntryPoint;
 
 		// If we have a parent graph, use its root since we share the same root, otherwise if we have no parent, we are the root
 		InstanceImpl->RootGraphInstance = ParentGraphInstance != nullptr ? ParentGraphInstance->GetRootGraphInstance() : InstanceImpl.Get();
@@ -80,7 +85,7 @@ void UAnimNextGraph::AllocateInstanceImpl(FAnimNextGraphInstance* ParentGraphIns
 #if WITH_EDITORONLY_DATA
 	if (Instance.IsValid())
 	{
-		FRWScopeLock Lock(GraphInstancesLock, SLT_Write);
+		FScopeLock Lock(&GraphInstancesLock);
 		check(!GraphInstances.Contains(Instance.Impl.Get()));
 		GraphInstances.Add(Instance.Impl.Get());
 	}
@@ -169,7 +174,10 @@ bool UAnimNextGraph::LoadFromArchiveBuffer(const TArray<uint8>& InSharedDataArch
 	const FDecoratorReader::EErrorState ErrorState = DecoratorReader.ReadGraph(SharedDataBuffer);
 	if (ErrorState == FDecoratorReader::EErrorState::None)
 	{
-		ResolvedRootDecoratorHandle = DecoratorReader.ResolveEntryPointHandle(RootDecoratorHandle);
+		for(FAnimNextGraphEntryPoint& EntryPoint : EntryPoints)
+		{
+			ResolvedRootDecoratorHandles.Add(EntryPoint.EntryPointName, DecoratorReader.ResolveEntryPointHandle(EntryPoint.RootDecoratorHandle));
+		}
 
 		// Make sure our execute method is registered
 		FRigUnit_AnimNextGraphEvaluator::RegisterExecuteMethod(ExecuteDefinition);
@@ -178,7 +186,7 @@ bool UAnimNextGraph::LoadFromArchiveBuffer(const TArray<uint8>& InSharedDataArch
 	else
 	{
 		SharedDataBuffer.Empty(0);
-		ResolvedRootDecoratorHandle = FAnimNextDecoratorHandle();
+		ResolvedRootDecoratorHandles.Add(FRigUnit_AnimNextGraphRoot::DefaultEntryPoint, FAnimNextDecoratorHandle());
 		return false;
 	}
 }
@@ -186,9 +194,10 @@ bool UAnimNextGraph::LoadFromArchiveBuffer(const TArray<uint8>& InSharedDataArch
 #if WITH_EDITORONLY_DATA
 void UAnimNextGraph::FreezeGraphInstances()
 {
-	FRWScopeLock Lock(GraphInstancesLock, SLT_ReadOnly);
+	FScopeLock Lock(&GraphInstancesLock);
 
-	for (FAnimNextGraphInstance* GraphInstance : GraphInstances)
+	TSet<FAnimNextGraphInstance*> GraphInstancesCopy = GraphInstances;
+	for (FAnimNextGraphInstance* GraphInstance : GraphInstancesCopy)
 	{
 		GraphInstance->Freeze();
 	}
@@ -196,9 +205,10 @@ void UAnimNextGraph::FreezeGraphInstances()
 
 void UAnimNextGraph::ThawGraphInstances()
 {
-	FRWScopeLock Lock(GraphInstancesLock, SLT_ReadOnly);
+	FScopeLock Lock(&GraphInstancesLock);
 
-	for (FAnimNextGraphInstance* GraphInstance : GraphInstances)
+	TSet<FAnimNextGraphInstance*> GraphInstancesCopy = GraphInstances;
+	for (FAnimNextGraphInstance* GraphInstance : GraphInstancesCopy)
 	{
 		GraphInstance->Thaw();
 	}
diff --git a/Engine/Plugins/Experimental/AnimNext/Source/AnimNext/Private/Graph/AnimNextGraphInstance.cpp b/Engine/Plugins/Experimental/AnimNext/Source/AnimNext/Private/Graph/AnimNextGraphInstance.cpp
--- a/Engine/Plugins/Experimental/AnimNext/Source/AnimNext/Private/Graph/AnimNextGraphInstance.cpp
+++ b/Engine/Plugins/Experimental/AnimNext/Source/AnimNext/Private/Graph/AnimNextGraphInstance.cpp
@@ -40,6 +40,11 @@ const UAnimNextGraph* FAnimNextGraphInstance::GetGraph() const
 	return Graph;
 }
 
+FName FAnimNextGraphInstance::GetEntryPoint() const
+{
+	return EntryPoint;
+}
+
 UE::AnimNext::FWeakDecoratorPtr FAnimNextGraphInstance::GetGraphRootPtr() const
 {
 	return GraphInstancePtr;
@@ -60,6 +65,20 @@ bool FAnimNextGraphInstance::UsesGraph(const UAnimNextGraph* InGraph) const
 	return Graph == InGraph;
 }
 
+bool FAnimNextGraphInstance::UsesEntryPoint(FName InEntryPoint) const
+{
+	if(Graph != nullptr)
+	{
+		if(InEntryPoint == NAME_None)
+		{
+			return EntryPoint == Graph->DefaultEntryPoint;
+		}
+
+		return InEntryPoint == EntryPoint;
+	}
+	return false;
+}
+
 bool FAnimNextGraphInstance::IsRoot() const
 {
 	return this == RootGraphInstance;
@@ -110,7 +129,7 @@ void FAnimNextGraphInstance::ExecuteLatentPins(const TConstArrayView<UE::AnimNex
 	if (URigVM* VM = Graph->VM)
 	{
 		FAnimNextExecuteContext& AnimNextContext = ExtendedExecuteContext.GetPublicDataSafe<FAnimNextExecuteContext>();
-		AnimNextContext.SetupForExecution(LatentHandles, DestinationBasePtr, bIsFrozen);
+		AnimNextContext.SetupForExecution(this, LatentHandles, DestinationBasePtr, bIsFrozen);
 
 		VM->ExecuteVM(ExtendedExecuteContext, FRigUnit_AnimNextShimRoot::EventName);
 
@@ -140,7 +159,10 @@ void FAnimNextGraphInstance::Thaw()
 
 		{
 			UE::AnimNext::FExecutionContext Context(*this);
-			GraphInstancePtr = Context.AllocateNodeInstance(UE::AnimNext::FWeakDecoratorPtr(), GraphPtr->ResolvedRootDecoratorHandle);
+			if(const FAnimNextDecoratorHandle* FoundHandle = GraphPtr->ResolvedRootDecoratorHandles.Find(EntryPoint))
+			{
+				GraphInstancePtr = Context.AllocateNodeInstance(UE::AnimNext::FWeakDecoratorPtr(), *FoundHandle);
+			}
 		}
 
 		if (!IsValid())
diff --git a/Engine/Plugins/Experimental/AnimNext/Source/AnimNext/Private/Graph/AnimNextGraphInstancePtr.cpp b/Engine/Plugins/Experimental/AnimNext/Source/AnimNext/Private/Graph/AnimNextGraphInstancePtr.cpp
--- a/Engine/Plugins/Experimental/AnimNext/Source/AnimNext/Private/Graph/AnimNextGraphInstancePtr.cpp
+++ b/Engine/Plugins/Experimental/AnimNext/Source/AnimNext/Private/Graph/AnimNextGraphInstancePtr.cpp
@@ -19,10 +19,12 @@ void FAnimNextGraphInstancePtr::Release()
 	if (Impl)
 	{
 #if WITH_EDITORONLY_DATA
-		const UAnimNextGraph* Graph = Impl->GetGraph();
-		FRWScopeLock Lock(Graph->GraphInstancesLock, SLT_Write);
-		check(Graph->GraphInstances.Contains(Impl.Get()));
-		Graph->GraphInstances.Remove(Impl.Get());
+		{
+			const UAnimNextGraph* Graph = Impl->GetGraph();
+			FScopeLock Lock(&Graph->GraphInstancesLock);
+			check(Graph->GraphInstances.Contains(Impl.Get()));
+			Graph->GraphInstances.Remove(Impl.Get());
+		}
 #endif
 
 		// Destroy the graph instance
diff --git a/Engine/Plugins/Experimental/AnimNext/Source/AnimNext/Private/Graph/RigUnit_AnimNextGraphRoot.cpp b/Engine/Plugins/Experimental/AnimNext/Source/AnimNext/Private/Graph/RigUnit_AnimNextGraphRoot.cpp
--- a/Engine/Plugins/Experimental/AnimNext/Source/AnimNext/Private/Graph/RigUnit_AnimNextGraphRoot.cpp
+++ b/Engine/Plugins/Experimental/AnimNext/Source/AnimNext/Private/Graph/RigUnit_AnimNextGraphRoot.cpp
@@ -5,6 +5,7 @@
 #include UE_INLINE_GENERATED_CPP_BY_NAME(RigUnit_AnimNextGraphRoot)
 
 FName FRigUnit_AnimNextGraphRoot::EventName = TEXT("DummyExecute");
+FName FRigUnit_AnimNextGraphRoot::DefaultEntryPoint = TEXT("Root");
 
 FRigUnit_AnimNextGraphRoot_DummyExecute()
 {
diff --git a/Engine/Plugins/Experimental/AnimNext/Source/AnimNext/Private/Graph/RigUnit_AnimNextGraphSelf.cpp b/Engine/Plugins/Experimental/AnimNext/Source/AnimNext/Private/Graph/RigUnit_AnimNextGraphSelf.cpp
new file mode 100644
--- /dev/null
+++ b/Engine/Plugins/Experimental/AnimNext/Source/AnimNext/Private/Graph/RigUnit_AnimNextGraphSelf.cpp
@@ -0,0 +1,9 @@
+﻿// Copyright Epic Games, Inc. All Rights Reserved.
+
+#include "Graph/RigUnit_AnimNextGraphSelf.h"
+#include "Graph/AnimNextGraphInstance.h"
+
+FRigUnit_AnimNextGraphSelf_Execute()
+{
+	Self = const_cast<UAnimNextGraph*>(ExecuteContext.GetGraphInstance().GetGraph());
+}
diff --git a/Engine/Plugins/Experimental/AnimNext/Source/AnimNext/Private/Param/RigUnit_AnimNextParameterBeginExecution.cpp b/Engine/Plugins/Experimental/AnimNext/Source/AnimNext/Private/Param/RigUnit_AnimNextParameterBeginExecution.cpp
--- a/Engine/Plugins/Experimental/AnimNext/Source/AnimNext/Private/Param/RigUnit_AnimNextParameterBeginExecution.cpp
+++ b/Engine/Plugins/Experimental/AnimNext/Source/AnimNext/Private/Param/RigUnit_AnimNextParameterBeginExecution.cpp
@@ -5,7 +5,7 @@
 
 #include UE_INLINE_GENERATED_CPP_BY_NAME(RigUnit_AnimNextParameterBeginExecution)
 
-FName FRigUnit_AnimNextParameterBeginExecution::EventName = TEXT("Update");
+FName FRigUnit_AnimNextParameterBeginExecution::DefaultEntryPoint = TEXT("Update");
 
 FRigUnit_AnimNextParameterBeginExecution_Execute()
 {
diff --git a/Engine/Plugins/Experimental/AnimNext/Source/AnimNext/Private/Scheduler/AnimNextSchedule.cpp b/Engine/Plugins/Experimental/AnimNext/Source/AnimNext/Private/Scheduler/AnimNextSchedule.cpp
--- a/Engine/Plugins/Experimental/AnimNext/Source/AnimNext/Private/Scheduler/AnimNextSchedule.cpp
+++ b/Engine/Plugins/Experimental/AnimNext/Source/AnimNext/Private/Scheduler/AnimNextSchedule.cpp
@@ -5,23 +5,58 @@
 #include "Async/TaskGraphInterfaces.h"
 #include "EngineLogs.h"
 #include "UObject/AssetRegistryTagsContext.h"
+#include "Graph/AnimNextGraph.h"
+#include "Param/AnimNextParameterBlock.h"
 
 #if WITH_EDITOR
 TUniqueFunction<void(UAnimNextSchedule*)> UAnimNextSchedule::CompileFunction;
 TUniqueFunction<void(const UAnimNextSchedule*, FAssetRegistryTagsContext)> UAnimNextSchedule::GetAssetRegistryTagsFunction;
 #endif
 
+void UAnimNextScheduleEntry_AnimNextGraph::GetPreloadDependencies(TArray<UObject*>& OutDeps)
+{
+	Super::GetPreloadDependencies(OutDeps);
+
+	OutDeps.Add(Graph);
+}
+
+void UAnimNextScheduleEntry_ParamScope::GetPreloadDependencies(TArray<UObject*>& OutDeps)
+{
+	Super::GetPreloadDependencies(OutDeps);
+
+	for(UAnimNextParameterBlock* ParameterBlock : ParameterBlocks)
+	{
+		OutDeps.Add(ParameterBlock);
+	}
+
+	for(UAnimNextScheduleEntry* SubEntry : SubEntries)
+	{
+		SubEntry->GetPreloadDependencies(OutDeps);
+	}
+}
+
 void UAnimNextSchedule::PostLoad()
 {
 	Super::PostLoad();
 
 #if WITH_EDITOR
-	CompileSchedule();
+	// delay compilation until the package has been loaded
+	FCoreUObjectDelegates::OnEndLoadPackage.AddUObject(this, &UAnimNextSchedule::HandlePackageDone);
 #endif
 }
 
 #if WITH_EDITOR
 
+void UAnimNextSchedule::HandlePackageDone(const FEndLoadPackageContext& Context)
+{
+	if (!Context.LoadedPackages.Contains(GetPackage()))
+	{
+		return;
+	}
+
+	CompileSchedule();
+}
+
 void UAnimNextSchedule::PostEditChangeProperty(FPropertyChangedEvent& PropertyChangedEvent)
 {
 	Super::PostEditChangeProperty(PropertyChangedEvent);
@@ -59,4 +94,14 @@ void UAnimNextSchedule::GetAssetRegistryTags(FAssetRegistryTagsContext Context)
 	GetAssetRegistryTagsFunction(this, Context);
 }
 
+void UAnimNextSchedule::GetPreloadDependencies(TArray<UObject*>& OutDeps)
+{
+	Super::GetPreloadDependencies(OutDeps);
+
+	for(UAnimNextScheduleEntry* Entry : Entries)
+	{
+		Entry->GetPreloadDependencies(OutDeps);
+	}
+}
+
 #endif // #if WITH_EDITOR
\ No newline at end of file
diff --git a/Engine/Plugins/Experimental/AnimNext/Source/AnimNext/Private/Scheduler/AnimNextScheduleGraphTask.cpp b/Engine/Plugins/Experimental/AnimNext/Source/AnimNext/Private/Scheduler/AnimNextScheduleGraphTask.cpp
--- a/Engine/Plugins/Experimental/AnimNext/Source/AnimNext/Private/Scheduler/AnimNextScheduleGraphTask.cpp
+++ b/Engine/Plugins/Experimental/AnimNext/Source/AnimNext/Private/Scheduler/AnimNextScheduleGraphTask.cpp
@@ -100,7 +100,7 @@ void FAnimNextScheduleGraphTask::RunGraph(const UE::AnimNext::FScheduleContext&
 	// Allocate our graph instance data
 	if (!GraphCache.GraphInstanceData.IsValid())
 	{
-		GraphToRun->AllocateInstance(GraphCache.GraphInstanceData);
+		GraphToRun->AllocateInstance(GraphCache.GraphInstanceData, EntryPoint);
 
 		// Only do dynamic verification for dynamic graphs. Static graphs get verified at compile time. 
 		if (Graph == nullptr && DynamicGraph != NAME_None)
diff --git a/Engine/Plugins/Experimental/AnimNext/Source/AnimNext/Public/Graph/AnimNextGraphInstance.h b/Engine/Plugins/Experimental/AnimNext/Source/AnimNext/Public/Graph/AnimNextGraphInstance.h
--- a/Engine/Plugins/Experimental/AnimNext/Source/AnimNext/Public/Graph/AnimNextGraphInstance.h
+++ b/Engine/Plugins/Experimental/AnimNext/Source/AnimNext/Public/Graph/AnimNextGraphInstance.h
@@ -50,6 +50,9 @@ struct FAnimNextGraphInstance
 	// Returns the graph used by this instance or nullptr if the instance is invalid
 	const UAnimNextGraph* GetGraph() const;
 
+	// Returns the entry point in Graph that this instance corresponds to 
+	FName GetEntryPoint() const;
+	
 	// Returns a weak handle to the root decorator instance
 	UE::AnimNext::FWeakDecoratorPtr GetGraphRootPtr() const;
 
@@ -62,6 +65,9 @@ struct FAnimNextGraphInstance
 	// Check to see if this instance data matches the provided graph
 	bool UsesGraph(const UAnimNextGraph* InGraph) const;
 
+	// Check to see if this instance data matches the provided graph entry point
+	bool UsesEntryPoint(FName InEntryPoint) const;
+	
 	// Returns whether or not this graph instance is the root graph instance or false otherwise
 	bool IsRoot() const;
 
@@ -105,6 +111,9 @@ private:
 	UPROPERTY()
 	TObjectPtr<const UAnimNextGraph> Graph;
 
+	// The entry point in Graph that this instance corresponds to 
+	FName EntryPoint;
+
 	// Hard reference to the graph instance data, we own it
 	UE::AnimNext::FDecoratorPtr GraphInstancePtr;
 
diff --git a/Engine/Plugins/Experimental/AnimNext/Source/AnimNext/Public/Param/AnimNextTag.h b/Engine/Plugins/Experimental/AnimNext/Source/AnimNext/Public/Param/AnimNextTag.h
--- a/Engine/Plugins/Experimental/AnimNext/Source/AnimNext/Public/Param/AnimNextTag.h
+++ b/Engine/Plugins/Experimental/AnimNext/Source/AnimNext/Public/Param/AnimNextTag.h
@@ -17,3 +17,10 @@ struct FAnimNextScope : public FAnimNextTag
 {
 	GENERATED_BODY()
 };
+
+// Empty struct to tag named entry points
+USTRUCT(BlueprintType)
+struct FAnimNextEntryPoint : public FAnimNextTag
+{
+	GENERATED_BODY()
+};
diff --git a/Engine/Plugins/Experimental/AnimNext/Source/AnimNextEditor/Private/AnimNextEditorModule.cpp b/Engine/Plugins/Experimental/AnimNext/Source/AnimNextEditor/Private/AnimNextEditorModule.cpp
--- a/Engine/Plugins/Experimental/AnimNext/Source/AnimNextEditor/Private/AnimNextEditorModule.cpp
+++ b/Engine/Plugins/Experimental/AnimNext/Source/AnimNextEditor/Private/AnimNextEditorModule.cpp
@@ -13,7 +13,9 @@
 #include "Graph/AnimNextGraph.h"
 #include "Graph/AnimNextGraphPanelNodeFactory.h"
 #include "Graph/AnimNextGraph_EdGraphNodeCustomization.h"
+#include "Graph/AnimNextGraph_EditorData.h"
 #include "Param/AnimNextParameterBlock.h"
+#include "Param/AnimNextParameterBlock_EditorData.h"
 #include "Param/ParameterBlockParameterCustomization.h"
 #include "Param/ParameterPickerArgs.h"
 #include "Param/ParametersGraphPanelPinFactory.h"
diff --git a/Engine/Plugins/Experimental/AnimNext/Source/AnimNextEditor/Private/Graph/GraphFactory.cpp b/Engine/Plugins/Experimental/AnimNext/Source/AnimNextEditor/Private/Graph/GraphFactory.cpp
--- a/Engine/Plugins/Experimental/AnimNext/Source/AnimNextEditor/Private/Graph/GraphFactory.cpp
+++ b/Engine/Plugins/Experimental/AnimNext/Source/AnimNextEditor/Private/Graph/GraphFactory.cpp
@@ -37,7 +37,7 @@ UObject* UAnimNextGraphFactory::FactoryCreateNew(UClass* Class, UObject* InParen
 	EditorData->Initialize(/*bRecompileVM*/false);
 
 	// Add root graph
-	EditorData->AddGraph(TEXT("Root"), false);
+	EditorData->AddGraph(FRigUnit_AnimNextGraphRoot::DefaultEntryPoint, false);
 	check(EditorData->Entries.Num() > 0);
 
 	// Compile the initial skeleton
diff --git a/Engine/Plugins/Experimental/AnimNext/Source/AnimNextEditor/Private/Param/ParametersGraphPanelPinFactory.cpp b/Engine/Plugins/Experimental/AnimNext/Source/AnimNextEditor/Private/Param/ParametersGraphPanelPinFactory.cpp
--- a/Engine/Plugins/Experimental/AnimNext/Source/AnimNextEditor/Private/Param/ParametersGraphPanelPinFactory.cpp
+++ b/Engine/Plugins/Experimental/AnimNext/Source/AnimNextEditor/Private/Param/ParametersGraphPanelPinFactory.cpp
@@ -25,9 +25,12 @@ TSharedPtr<SGraphPin> FParametersGraphPanelPinFactory::CreatePin_Internal(UEdGra
 		{
 			if(ModelPin->GetCustomWidgetName() == "ParamName")
 			{
+				const FString ParamTypeString = ModelPin->GetMetaData("AllowedParamType");
+				FAnimNextParamType FilterType = FAnimNextParamType::FromString(ParamTypeString);
 				return SNew(SGraphPinParamName, InPin)
 					.ModelPin(ModelPin)
-					.GraphNode(RigNode);
+					.GraphNode(RigNode)
+					.FilterType(FilterType);
 			}
 		}
 	}
diff --git a/Engine/Plugins/Experimental/AnimNext/Source/AnimNextTestSuite/Private/AnimNextEditorTest.cpp b/Engine/Plugins/Experimental/AnimNext/Source/AnimNextTestSuite/Private/AnimNextEditorTest.cpp
--- a/Engine/Plugins/Experimental/AnimNext/Source/AnimNextTestSuite/Private/AnimNextEditorTest.cpp
+++ b/Engine/Plugins/Experimental/AnimNext/Source/AnimNextTestSuite/Private/AnimNextEditorTest.cpp
@@ -9,6 +9,10 @@
 #include "Param/AnimNextParameterBlock.h"
 #include "Param/AnimNextParameterBlockParameter.h"
 #include "Animation/AnimSequence.h"
+#include "Graph/AnimNextGraph.h"
+#include "Graph/AnimNextGraphEntry.h"
+#include "Graph/AnimNextGraph_EditorData.h"
+#include "Graph/GraphFactory.h"
 #if WITH_EDITOR
 #include "ScopedTransaction.h"
 #include "Editor.h"
@@ -19,9 +23,12 @@
 
 #if WITH_DEV_AUTOMATION_TESTS && WITH_EDITOR
 
-IMPLEMENT_SIMPLE_AUTOMATION_TEST(FAnimationAnimNextParametersEditorTest_Block, "Animation.AnimNext.Parameters.Editor.Block", EAutomationTestFlags::EditorContext | EAutomationTestFlags::EngineFilter)
+namespace UE::AnimNext::Tests
+{
+
+IMPLEMENT_SIMPLE_AUTOMATION_TEST(FEditor_Parameters_ParameterBlock, "Animation.AnimNext.Editor.Parameters.ParameterBlock", EAutomationTestFlags::EditorContext | EAutomationTestFlags::EngineFilter)
 
-bool FAnimationAnimNextParametersEditorTest_Block::RunTest(const FString& InParameters)
+bool FEditor_Parameters_ParameterBlock::RunTest(const FString& InParameters)
 {
 	using namespace UE::AnimNext;
 
@@ -131,14 +138,14 @@ bool FAnimationAnimNextParametersEditorTest_Block::RunTest(const FString& InPara
 	AddErrorIfFalse(EditorData->Entries.Num() == 1, FString::Printf(TEXT("Unexpected entry count found in parameter block (Have %d, expected 1)."), EditorData->Entries.Num()));
 	GEditor->UndoTransaction();
 
-	Tests::FUtils::CleanupAfterTests();
+	FUtils::CleanupAfterTests();
 
 	return true;
 }
 
-IMPLEMENT_SIMPLE_AUTOMATION_TEST(FAnimationAnimNextParametersEditorTest_Python, "Animation.AnimNext.Parameters.Editor.Python", EAutomationTestFlags::EditorContext | EAutomationTestFlags::EngineFilter)
+IMPLEMENT_SIMPLE_AUTOMATION_TEST(FEditor_Parameters_Python, "Animation.AnimNext.Editor.Parameters.Python", EAutomationTestFlags::EditorContext | EAutomationTestFlags::EngineFilter)
 
-bool FAnimationAnimNextParametersEditorTest_Python::RunTest(const FString& InParameters)
+bool FEditor_Parameters_Python::RunTest(const FString& InParameters)
 {
 	using namespace UE::AnimNext;
 
@@ -152,9 +159,70 @@ bool FAnimationAnimNextParametersEditorTest_Python::RunTest(const FString& InPar
 
 	IPythonScriptPlugin::Get()->ExecPythonCommand(Script);
 
-	Tests::FUtils::CleanupAfterTests();
+	FUtils::CleanupAfterTests();
+
+	return true;
+}
+
+IMPLEMENT_SIMPLE_AUTOMATION_TEST(FEditor_Graph, "Animation.AnimNext.Editor.Graph", EAutomationTestFlags::EditorContext | EAutomationTestFlags::EngineFilter)
+
+bool FEditor_Graph::RunTest(const FString& InParameters)
+{
+	using namespace UE::AnimNext;
+
+	const TStrongObjectPtr<UFactory> GraphFactory(NewObject<UAnimNextGraphFactory>());
+	UAnimNextGraph* Graph = Cast<UAnimNextGraph>(GraphFactory->FactoryCreateNew(UAnimNextGraph::StaticClass(), GetTransientPackage(), TEXT("TestAnimNextGraph"), RF_Transient, nullptr, nullptr, NAME_None));
+	if(Graph == nullptr)
+	{
+		AddError(TEXT("Could not create graph."));
+		return false;
+	}
+
+	UAnimNextGraph_EditorData* EditorData = UncookedOnly::FUtils::GetEditorData(Graph);
+	if(EditorData == nullptr)
+	{
+		AddError(TEXT("Graph has no editor data."));
+		return false;
+	}
+
+	// Add graph
+	UAnimNextGraphEntry* GraphEntry = nullptr;
+	{
+		FScopedTransaction Transaction(FText::GetEmpty());
+		GraphEntry = EditorData->AddGraph(TEXT("TestGraph"));
+		AddErrorIfFalse(Graph != nullptr, TEXT("Could not create new graph in asset."));
+	}
+
+	GEditor->UndoTransaction();
+	AddErrorIfFalse(EditorData->Entries.Num() == 1, FString::Printf(TEXT("Unexpected entry count found in graph asset (Have %d, expected 1)."), EditorData->Entries.Num()));
+
+	GEditor->RedoTransaction();
+	AddErrorIfFalse(EditorData->Entries.Num() == 2, FString::Printf(TEXT("Unexpected entry count found in graph asset (Have %d, expected 2)."), EditorData->Entries.Num()));
+
+	// RemoveEntry
+	{
+		FScopedTransaction Transaction(FText::GetEmpty());
+		AddErrorIfFalse(EditorData->RemoveEntry(GraphEntry), TEXT("Failed to remove entry."));
+	}
+
+	GEditor->UndoTransaction();
+	AddErrorIfFalse(EditorData->Entries.Num() == 2, FString::Printf(TEXT("Unexpected entry count found in graph asset (Have %d, expected 2)."), EditorData->Entries.Num()));
+
+	GEditor->RedoTransaction();
+	AddErrorIfFalse(EditorData->Entries.Num() == 1, FString::Printf(TEXT("Unexpected entry count found in graph asset (Have %d, expected 1)."), EditorData->Entries.Num()));
+
+	GEditor->UndoTransaction();
+	AddErrorIfFalse(EditorData->Entries.Num() == 2, FString::Printf(TEXT("Unexpected entry count found in graph asset (Have %d, expected 2)."), EditorData->Entries.Num()));
+
+	// FindEntry
+	AddErrorIfFalse(EditorData->FindEntry(TEXT("TestGraph")) != nullptr, TEXT("Could not find entry in asset."));
+	GEditor->UndoTransaction();
+
+	FUtils::CleanupAfterTests();
 
 	return true;
 }
 
+}
+
 #endif	// WITH_DEV_AUTOMATION_TESTS && WITH_EDITOR
\ No newline at end of file
diff --git a/Engine/Plugins/Experimental/AnimNext/Source/AnimNextTestSuite/Private/AnimNextRuntimeTest.cpp b/Engine/Plugins/Experimental/AnimNext/Source/AnimNextTestSuite/Private/AnimNextRuntimeTest.cpp
--- a/Engine/Plugins/Experimental/AnimNext/Source/AnimNextTestSuite/Private/AnimNextRuntimeTest.cpp
+++ b/Engine/Plugins/Experimental/AnimNext/Source/AnimNextTestSuite/Private/AnimNextRuntimeTest.cpp
@@ -5,6 +5,7 @@
 
 #include "DecoratorBase/DecoratorReader.h"
 #include "Graph/AnimNextGraph.h"
+#include "Graph/RigUnit_AnimNextGraphRoot.h"
 #include "Misc/AutomationTest.h"
 #include "Serialization/MemoryReader.h"
 
@@ -30,10 +31,12 @@ namespace UE::AnimNext
 		ExecuteDefinition.Hash = 0;
 		ExecuteDefinition.MethodName = TEXT("Execute_0");
 
+		FAnimNextGraphEntryPoint& EntryPoint = Graph.EntryPoints.AddDefaulted_GetRef();
+		EntryPoint.EntryPointName = FRigUnit_AnimNextGraphRoot::DefaultEntryPoint;
+		EntryPoint.RootDecoratorHandle = FAnimNextEntryPointHandle(NodeHandles[0]);
 		Graph.ExecuteDefinition = ExecuteDefinition;
 		Graph.SharedDataArchiveBuffer = SharedDataArchiveBuffer;
 		Graph.GraphReferencedObjects.Empty();
-		Graph.RootDecoratorHandle = FAnimNextEntryPointHandle(NodeHandles[0]);
 
 		// Reconstruct our graph shared data
 		FMemoryReader GraphSharedDataArchive(SharedDataArchiveBuffer);
@@ -42,7 +45,7 @@ namespace UE::AnimNext
 		const FDecoratorReader::EErrorState ErrorState = DecoratorReader.ReadGraph(Graph.SharedDataBuffer);
 		if (ErrorState == FDecoratorReader::EErrorState::None)
 		{
-			Graph.ResolvedRootDecoratorHandle = DecoratorReader.ResolveEntryPointHandle(Graph.RootDecoratorHandle);
+			Graph.ResolvedRootDecoratorHandles.Add(FRigUnit_AnimNextGraphRoot::DefaultEntryPoint, DecoratorReader.ResolveEntryPointHandle(Graph.EntryPoints[0].RootDecoratorHandle));
 
 			for (FNodeHandle& NodeHandle : NodeHandles)
 			{
@@ -56,7 +59,7 @@ namespace UE::AnimNext
 		else
 		{
 			Graph.SharedDataBuffer.Empty(0);
-			Graph.ResolvedRootDecoratorHandle = FAnimNextDecoratorHandle();
+			Graph.ResolvedRootDecoratorHandles.Add(FRigUnit_AnimNextGraphRoot::DefaultEntryPoint, FAnimNextDecoratorHandle());
 			return false;
 		}
 	}
diff --git a/Engine/Plugins/Experimental/AnimNext/Source/AnimNextUncookedOnly/Internal/AnimNextRigVMAssetEditorData.h b/Engine/Plugins/Experimental/AnimNext/Source/AnimNextUncookedOnly/Internal/AnimNextRigVMAssetEditorData.h
--- a/Engine/Plugins/Experimental/AnimNext/Source/AnimNextUncookedOnly/Internal/AnimNextRigVMAssetEditorData.h
+++ b/Engine/Plugins/Experimental/AnimNext/Source/AnimNextUncookedOnly/Internal/AnimNextRigVMAssetEditorData.h
@@ -6,6 +6,7 @@
 #include "RigVMModel/RigVMGraph.h"
 #include "RigVMCore/RigVMGraphFunctionHost.h"
 #include "RigVMBlueprint.h"
+#include "UncookedOnlyUtils.h"
 #include "Kismet/BlueprintFunctionLibrary.h"
 #include "AnimNextRigVMAssetEditorData.generated.h"
 
@@ -34,6 +35,11 @@ namespace UE::AnimNext::Editor
 	class FWorkspaceEditor;
 }
 
+namespace UE::AnimNext::Tests
+{
+	class FEditor_Graph;
+}
+
 namespace UE::AnimNext::UncookedOnly
 {
 	// A delegate for subscribing / reacting to editor data modifications.
@@ -73,6 +79,7 @@ protected:
 	friend class UAnimNextRigVMAssetLibrary;
 	friend class UAnimNextGraph_EdGraph;
 	friend class UAnimNextParameterBlock_EdGraph;
+	friend class UE::AnimNext::Tests::FEditor_Graph;
 
 	// UObject interface
 	virtual void Serialize(FArchive& Ar) override;
@@ -164,9 +171,6 @@ protected:
 		return CastChecked<EntryClassType>(CreateNewSubEntry(InEditorData, EntryClassType::StaticClass()));
 	}
 
-	// Helper that generates a fixed UEdGraph name counterpart to a a URigVMGraph 
-	static FName MakeEdGraphName(FName InRigVMGraphName);
-
 	// Get all the entries for this asset
 	TConstArrayView<TObjectPtr<UAnimNextRigVMAssetEntry>> GetAllEntries() const { return Entries; } 
 
@@ -190,6 +194,19 @@ protected:
 		}
 	}
 
+	// Returns all nodes in all graphs of the specified class
+	template<class T>
+	void GetAllNodesOfClass(TArray<T*>& OutNodes)
+	{
+		ForEachEntryOfType<IAnimNextRigVMGraphInterface>([&OutNodes](IAnimNextRigVMGraphInterface* InGraphInterface)
+		{
+			TArray<T*> GraphNodes;
+			InGraphInterface->GetEdGraph()->GetNodesOfClass<T>(GraphNodes);
+			OutNodes.Append(GraphNodes);
+			return true;
+		});
+	}
+	
 	// Find an entry by name
 	ANIMNEXTUNCOOKEDONLY_API UAnimNextRigVMAssetEntry* FindEntry(FName InName) const;
 
@@ -213,9 +230,6 @@ protected:
 	// Find an entry that corresponds to the specified RigVMGraph. This uses the name of the graph to match the entry 
 	UAnimNextRigVMAssetEntry* FindEntryForRigVMGraph(URigVMGraph* InRigVMGraph) const;
 
-	// Find an entry that corresponds to the specified EdGraph
-	UAnimNextRigVMAssetEntry* FindEntryForEdGraph(URigVMEdGraph* InEdGraph) const;
-
 	/** All entries in this asset - not saved, either serialized or discovered at load time */
 	UPROPERTY(transient)
 	TArray<TObjectPtr<UAnimNextRigVMAssetEntry>> Entries;
@@ -259,6 +273,9 @@ protected:
 	// Delegate to subscribe to modifications to this editor data
 	UE::AnimNext::UncookedOnly::FOnEditorDataModified ModifiedDelegate;
 
+	// Cached exports, generated on compilation
+	FAnimNextParameterProviderAssetRegistryExports CachedExports;
+	
 	bool bAutoRecompileVM = true;
 	bool bErrorsDuringCompilation = false;
 	bool bSuspendModelNotificationsForSelf = false;
diff --git a/Engine/Plugins/Experimental/AnimNext/Source/AnimNextUncookedOnly/Internal/Graph/AnimNextGraphEntry.h b/Engine/Plugins/Experimental/AnimNext/Source/AnimNextUncookedOnly/Internal/Graph/AnimNextGraphEntry.h
--- a/Engine/Plugins/Experimental/AnimNext/Source/AnimNextUncookedOnly/Internal/Graph/AnimNextGraphEntry.h
+++ b/Engine/Plugins/Experimental/AnimNext/Source/AnimNextUncookedOnly/Internal/Graph/AnimNextGraphEntry.h
@@ -27,7 +27,8 @@ class UAnimNextGraphEntry : public UAnimNextRigVMAssetEntry, public IAnimNextRig
 
 	// UAnimNextRigVMAssetEntry interface
 	virtual FName GetEntryName() const override;
-
+	virtual void SetEntryName(FName InName, bool bSetupUndoRedo = true) override;
+	
 	// IAnimNextRigVMGraphInterface interface
 	virtual URigVMGraph* GetRigVMGraph() const override;
 	virtual URigVMEdGraph* GetEdGraph() const override;
diff --git a/Engine/Plugins/Experimental/AnimNext/Source/AnimNextUncookedOnly/Internal/Graph/AnimNextGraph_EdGraphSchema.h b/Engine/Plugins/Experimental/AnimNext/Source/AnimNextUncookedOnly/Internal/Graph/AnimNextGraph_EdGraphSchema.h
--- a/Engine/Plugins/Experimental/AnimNext/Source/AnimNextUncookedOnly/Internal/Graph/AnimNextGraph_EdGraphSchema.h
+++ b/Engine/Plugins/Experimental/AnimNext/Source/AnimNextUncookedOnly/Internal/Graph/AnimNextGraph_EdGraphSchema.h
@@ -14,7 +14,8 @@ class UAnimNextGraph_EdGraphSchema : public URigVMEdGraphSchema
 
 	// URigVMEdGraphSchema interface
 	virtual TSubclassOf<URigVMEdGraphNode> GetGraphNodeClass(const URigVMEdGraph* InGraph) const override { return UAnimNextGraph_EdGraphNode::StaticClass(); }
-
+	
 	// UEdGraphSchema interface
 	virtual void TrySetDefaultValue(UEdGraphPin& InPin, const FString& InNewDefaultValue, bool bMarkAsModified) const override;
+	virtual void GetGraphDisplayInformation(const UEdGraph& Graph, /*out*/ FGraphDisplayInfo& DisplayInfo) const override;
 };
\ No newline at end of file
diff --git a/Engine/Plugins/Experimental/AnimNext/Source/AnimNextUncookedOnly/Internal/Param/AnimNextParameterBlockParameter.h b/Engine/Plugins/Experimental/AnimNext/Source/AnimNextUncookedOnly/Internal/Param/AnimNextParameterBlockParameter.h
--- a/Engine/Plugins/Experimental/AnimNext/Source/AnimNextUncookedOnly/Internal/Param/AnimNextParameterBlockParameter.h
+++ b/Engine/Plugins/Experimental/AnimNext/Source/AnimNextUncookedOnly/Internal/Param/AnimNextParameterBlockParameter.h
@@ -16,13 +16,18 @@ namespace UE::AnimNext::Editor
 	class FParameterBlockParameterCustomization;
 }
 
+namespace UE::AnimNext::Tests
+{
+	class FEditor_Parameters_ParameterBlock;
+}
+
 UCLASS(MinimalAPI, Category = "Parameters")
 class UAnimNextParameterBlockParameter : public UAnimNextRigVMAssetEntry, public IAnimNextRigVMParameterInterface
 {
 	GENERATED_BODY()
 
 	friend class UAnimNextParameterBlock_EditorData;
-	friend class FAnimationAnimNextParametersEditorTest_Block;
+	friend class UE::AnimNext::Tests::FEditor_Parameters_ParameterBlock;
 	friend class UE::AnimNext::Editor::FParameterBlockParameterCustomization;
 
 	// UAnimNextRigVMAssetEntry interface
diff --git a/Engine/Plugins/Experimental/AnimNext/Source/AnimNextUncookedOnly/Internal/Param/AnimNextParameterBlock_EdGraphSchema.h b/Engine/Plugins/Experimental/AnimNext/Source/AnimNextUncookedOnly/Internal/Param/AnimNextParameterBlock_EdGraphSchema.h
--- a/Engine/Plugins/Experimental/AnimNext/Source/AnimNextUncookedOnly/Internal/Param/AnimNextParameterBlock_EdGraphSchema.h
+++ b/Engine/Plugins/Experimental/AnimNext/Source/AnimNextUncookedOnly/Internal/Param/AnimNextParameterBlock_EdGraphSchema.h
@@ -13,5 +13,8 @@ class UAnimNextParameterBlock_EdGraphSchema : public URigVMEdGraphSchema
 	GENERATED_BODY()
 
 	// URigVMEdGraphSchema interface
-	virtual TSubclassOf<URigVMEdGraphNode> GetGraphNodeClass(const URigVMEdGraph* InGraph) const override { return UAnimNextParameterBlock_EdGraphNode::StaticClass(); }	
+	virtual TSubclassOf<URigVMEdGraphNode> GetGraphNodeClass(const URigVMEdGraph* InGraph) const override { return UAnimNextParameterBlock_EdGraphNode::StaticClass(); }
+
+	// UEdGraphSchema interface
+	virtual void GetGraphDisplayInformation(const UEdGraph& Graph, /*out*/ FGraphDisplayInfo& DisplayInfo) const override;
 };
\ No newline at end of file
diff --git a/Engine/Plugins/Experimental/AnimNext/Source/AnimNextUncookedOnly/Internal/Param/AnimNextParameterBlock_EditorData.h b/Engine/Plugins/Experimental/AnimNext/Source/AnimNextUncookedOnly/Internal/Param/AnimNextParameterBlock_EditorData.h
--- a/Engine/Plugins/Experimental/AnimNext/Source/AnimNextUncookedOnly/Internal/Param/AnimNextParameterBlock_EditorData.h
+++ b/Engine/Plugins/Experimental/AnimNext/Source/AnimNextUncookedOnly/Internal/Param/AnimNextParameterBlock_EditorData.h
@@ -14,7 +14,7 @@ class UAnimNextParameterBlock;
 class UAnimNextParameterBlockParameter;
 class UAnimNextParameterBlock_Controller;
 enum class ERigVMGraphNotifType : uint8;
-class FAnimationAnimNextParametersEditorTest_Block;
+
 
 namespace UE::AnimNext::UncookedOnly
 {
@@ -34,6 +34,11 @@ namespace UE::AnimNext::Editor
 	class FModule;
 }
 
+namespace UE::AnimNext::Tests
+{
+	class FEditor_Parameters_ParameterBlock;
+}
+
 UCLASS()
 class UAnimNextParameterBlockLibrary_Schema : public URigVMSchema
 {
@@ -74,7 +79,7 @@ class UAnimNextParameterBlock_EditorData : public UAnimNextRigVMAssetEditorData
 	friend class UE::AnimNext::Editor::SRigVMAssetView;
 	friend class UE::AnimNext::Editor::SParameterPicker;
 	friend class UE::AnimNext::Editor::FParameterBlockTabSummoner;
-	friend class FAnimationAnimNextParametersEditorTest_Block;
+	friend class UE::AnimNext::Tests::FEditor_Parameters_ParameterBlock;
 	friend class UE::AnimNext::Editor::SRigVMAssetViewRow;
 	friend class UE::AnimNext::Editor::FParameterBlockParameterCustomization;
 
diff --git a/Engine/Plugins/Experimental/AnimNext/Source/AnimNextUncookedOnly/Internal/UncookedOnlyUtils.h b/Engine/Plugins/Experimental/AnimNext/Source/AnimNextUncookedOnly/Internal/UncookedOnlyUtils.h
--- a/Engine/Plugins/Experimental/AnimNext/Source/AnimNextUncookedOnly/Internal/UncookedOnlyUtils.h
+++ b/Engine/Plugins/Experimental/AnimNext/Source/AnimNextUncookedOnly/Internal/UncookedOnlyUtils.h
@@ -3,10 +3,8 @@
 #pragma once
 
 #include "CoreMinimal.h"
-#include "Param/AnimNextParameterBlock_EditorData.h"
-#include "Graph/AnimNextGraph_EditorData.h"
-#include "AnimNextRigVMAssetEditorData.h"
 #include "IAnimNextRigVMGraphInterface.h"
+#include "AssetRegistry/AssetData.h"
 #include "Param/ParamTypeHandle.h"
 #include "RigVMCore/RigVMTemplate.h"
 
@@ -25,6 +23,7 @@ class UAnimNextGraph_EdGraph;
 struct FEdGraphPinType;
 class UAnimNextRigVMAsset;
 class UAnimNextRigVMAssetEditorData;
+class UAnimNextRigVMAssetEntry;
 
 namespace UE
 {
@@ -141,32 +140,13 @@ struct ANIMNEXTUNCOOKEDONLY_API FUtils
 	static FRigVMTemplateArgumentType GetRigVMArgTypeFromParamType(const FAnimNextParamType& InParamType);
 
 	/** Set up a simple animation graph */
-	static void SetupAnimGraph(URigVMController* InController);
+	static void SetupAnimGraph(UAnimNextRigVMAssetEntry* InEntry, URigVMController* InController);
 	
 	/** Set up a simple parameter graph */
 	static void SetupParameterGraph(URigVMController* InController);
 	
-	/** Set up a binding graph given the type to set */
-	static void SetupBindingGraph(URigVMController* InController, FName InParameterName, const FAnimNextParamType& InParamType);
-
-	/** Set up a binding graph for a simple literal value given the type to set */
-	static void SetupBindingGraphForLiteral(URigVMController* InController, FName InParameterName, const FAnimNextParamType& InParamType);
-
 	/** Converts the Verse-tag-like snake_case_parameter_name to a period-separated display name similar to a gameplay tag */
 	static FText GetParameterDisplayNameText(FName InParameterName);
-
-	/** Returns all nodes in all graphs of the specified class */
-	template<class T>
-	static void GetAllNodesOfClass(const UAnimNextRigVMAssetEditorData* InEditorData, TArray<T*>& OutNodes)
-	{
-		InEditorData->ForEachEntryOfType<IAnimNextRigVMGraphInterface>([&OutNodes](IAnimNextRigVMGraphInterface* InGraphInterface)
-		{
-			TArray<T*> GraphNodes;
-			InGraphInterface->GetEdGraph()->GetNodesOfClass<T>(GraphNodes);
-			OutNodes.Append(GraphNodes);
-			return true;
-		});
-	}
 	
 	// Gets the parameters that are exported to the asset registry for an asset
 	static bool GetExportedParametersForAsset(const FAssetData& InAsset, FAnimNextParameterProviderAssetRegistryExports& OutExports);
@@ -174,7 +154,10 @@ struct ANIMNEXTUNCOOKEDONLY_API FUtils
 	// Gets all the parameters that are exported to the asset registry
 	static bool GetExportedParametersFromAssetRegistry(FAnimNextParameterProviderAssetRegistryExports& OutExports);
 
-	// Gets the parameters that are used by a RigVM graph
+	// Gets the exported parameters that are used by a RigVM asset
+	static void GetAssetParameters(const UAnimNextRigVMAssetEditorData* EditorData, FAnimNextParameterProviderAssetRegistryExports& OutExports);
+
+	// Gets the exported parameters that are used by a RigVM graph
 	static void GetGraphParameters(const URigVMGraph* Graph, FAnimNextParameterProviderAssetRegistryExports& OutExports);
 
 	// Gets the parameters that are exported to the asset registry by a schedule
diff --git a/Engine/Plugins/Experimental/AnimNext/Source/AnimNextUncookedOnly/Private/AnimNextRigVMAssetEditorData.cpp b/Engine/Plugins/Experimental/AnimNext/Source/AnimNextUncookedOnly/Private/AnimNextRigVMAssetEditorData.cpp
--- a/Engine/Plugins/Experimental/AnimNext/Source/AnimNextUncookedOnly/Private/AnimNextRigVMAssetEditorData.cpp
+++ b/Engine/Plugins/Experimental/AnimNext/Source/AnimNextUncookedOnly/Private/AnimNextRigVMAssetEditorData.cpp
@@ -10,6 +10,7 @@
 #include "IAnimNextRigVMParameterInterface.h"
 #include "UncookedOnlyUtils.h"
 #include "Misc/TransactionObjectEvent.h"
+#include "Param/AnimNextTag.h"
 #include "RigVMModel/RigVMFunctionLibrary.h"
 #include "RigVMModel/RigVMNotifications.h"
 #include "RigVMModel/Nodes/RigVMCollapseNode.h"
@@ -44,7 +45,7 @@ void UAnimNextRigVMAssetEditorData::ReconstructAllNodes()
 	}
 
 	TArray<URigVMEdGraphNode*> AllNodes;
-	UE::AnimNext::UncookedOnly::FUtils::GetAllNodesOfClass(this, AllNodes);
+	GetAllNodesOfClass(AllNodes);
 
 	for (URigVMEdGraphNode* Node : AllNodes)
 	{
@@ -152,23 +153,8 @@ void UAnimNextRigVMAssetEditorData::GetAssetRegistryTags(FAssetRegistryTagsConte
 {
 	Super::GetAssetRegistryTags(Context);
 
-	FAnimNextParameterProviderAssetRegistryExports ExportParameters;
-	ExportParameters.Parameters.Reserve(Entries.Num());
-
-	for(const UAnimNextRigVMAssetEntry* Entry : Entries)
-	{
-		if(const IAnimNextRigVMParameterInterface* ParameterInterface = Cast<IAnimNextRigVMParameterInterface>(Entry))
-		{
-			ExportParameters.Parameters.Emplace(Entry->GetEntryName(), ParameterInterface->GetParamType(), EAnimNextParameterFlags::Bound);
-		}
-		else if(const IAnimNextRigVMGraphInterface* GraphInterface = Cast<IAnimNextRigVMGraphInterface>(Entry))
-		{
-			UE::AnimNext::UncookedOnly::FUtils::GetGraphParameters(GraphInterface->GetRigVMGraph(), ExportParameters);
-		}
-	}
-
 	FString TagValue;
-	FAnimNextParameterProviderAssetRegistryExports::StaticStruct()->ExportText(TagValue, &ExportParameters, nullptr, nullptr, PPF_None, nullptr);
+	FAnimNextParameterProviderAssetRegistryExports::StaticStruct()->ExportText(TagValue, &CachedExports, nullptr, nullptr, PPF_None, nullptr);
 	Context.AddTag(FAssetRegistryTag(UE::AnimNext::ExportsAnimNextAssetRegistryTag, TagValue, FAssetRegistryTag::TT_Hidden));
 }
 
@@ -428,7 +414,7 @@ void UAnimNextRigVMAssetEditorData::HandleRigVMGraphAdded(const FRigVMClient* In
 			GetOuter() != GetTransientPackage())
 		{
 			CreateEdGraph(RigVMGraph, true);
-			RecompileVM();
+			RequestAutoVMRecompilation();
 		}
 		
 #if WITH_EDITOR
@@ -782,29 +768,14 @@ UAnimNextRigVMAssetEntry* UAnimNextRigVMAssetEditorData::FindEntryForRigVMGraph(
 {
 	for(UAnimNextRigVMAssetEntry* Entry : Entries)
 	{
-		if(Entry->GetEntryName() == InRigVMGraph->GetFName())
-		{
-			return Entry;
-		}
-	}
-
-	return nullptr;
-}
-
-UAnimNextRigVMAssetEntry* UAnimNextRigVMAssetEditorData::FindEntryForEdGraph(URigVMEdGraph* InEdGraph) const
-{
-	for(UAnimNextRigVMAssetEntry* Entry : Entries)
-	{
-		if(MakeEdGraphName(Entry->GetEntryName()) == InEdGraph->GetFName())
+		if(IAnimNextRigVMGraphInterface* GraphInterface = Cast<IAnimNextRigVMGraphInterface>(Entry))
 		{
-			return Entry;
+			if(GraphInterface->GetRigVMGraph() == InRigVMGraph)
+			{
+				return Entry;
+			}
 		}
 	}
 
 	return nullptr;
 }
-
-FName UAnimNextRigVMAssetEditorData::MakeEdGraphName(FName InRigVMGraphName)
-{
-	return *FString::Printf(TEXT("%s_EdGraph"), *InRigVMGraphName.ToString());
-}
\ No newline at end of file
diff --git a/Engine/Plugins/Experimental/AnimNext/Source/AnimNextUncookedOnly/Private/Graph/AnimNextGraphEntry.cpp b/Engine/Plugins/Experimental/AnimNext/Source/AnimNextUncookedOnly/Private/Graph/AnimNextGraphEntry.cpp
--- a/Engine/Plugins/Experimental/AnimNext/Source/AnimNextUncookedOnly/Private/Graph/AnimNextGraphEntry.cpp
+++ b/Engine/Plugins/Experimental/AnimNext/Source/AnimNextUncookedOnly/Private/Graph/AnimNextGraphEntry.cpp
@@ -2,12 +2,42 @@
 
 #include "Graph/AnimNextGraphEntry.h"
 #include "Graph/AnimNextGraph_EdGraph.h"
+#include "Graph/RigUnit_AnimNextGraphRoot.h"
 
 FName UAnimNextGraphEntry::GetEntryName() const
 {
 	return GraphName;
 }
 
+void UAnimNextGraphEntry::SetEntryName(FName InName, bool bSetupUndoRedo)
+{
+	if(bSetupUndoRedo)
+	{
+		Modify();
+	};
+
+	GraphName = InName;
+
+	// Forward to entry point node
+	URigVMController* Controller = GetImplementingOuter<IRigVMClientHost>()->GetController(Graph);
+	for(URigVMNode* Node : Graph->GetNodes())
+	{
+		if(URigVMUnitNode* UnitNode = Cast<URigVMUnitNode>(Node))
+		{
+			if(UnitNode->GetScriptStruct() == FRigUnit_AnimNextGraphRoot::StaticStruct())
+			{
+				URigVMPin* EntryPointPin = UnitNode->FindPin(GET_MEMBER_NAME_STRING_CHECKED(FRigUnit_AnimNextGraphRoot, EntryPoint));
+				check(EntryPointPin);
+				check(EntryPointPin->GetDirection() == ERigVMPinDirection::Hidden);
+
+				Controller->SetPinDefaultValue(EntryPointPin->GetPinPath(), InName.ToString());
+			}
+		}
+	}
+	
+	BroadcastModified();
+}
+
 URigVMGraph* UAnimNextGraphEntry::GetRigVMGraph() const
 {
 	return Graph;
diff --git a/Engine/Plugins/Experimental/AnimNext/Source/AnimNextUncookedOnly/Private/Graph/AnimNextGraph_EdGraphSchema.cpp b/Engine/Plugins/Experimental/AnimNext/Source/AnimNextUncookedOnly/Private/Graph/AnimNextGraph_EdGraphSchema.cpp
--- a/Engine/Plugins/Experimental/AnimNext/Source/AnimNextUncookedOnly/Private/Graph/AnimNextGraph_EdGraphSchema.cpp
+++ b/Engine/Plugins/Experimental/AnimNext/Source/AnimNextUncookedOnly/Private/Graph/AnimNextGraph_EdGraphSchema.cpp
@@ -1,8 +1,13 @@
 ﻿// Copyright Epic Games, Inc. All Rights Reserved.
 
 #include "Graph/AnimNextGraph_EdGraphSchema.h"
+
+#include "AnimNextRigVMAsset.h"
+#include "AnimNextRigVMAssetEntry.h"
 #include "Editor.h"
 
+#define LOCTEXT_NAMESPACE "AnimNextGraph_EdGraphSchema"
+
 void UAnimNextGraph_EdGraphSchema::TrySetDefaultValue(UEdGraphPin& InPin, const FString& InNewDefaultValue, bool bMarkAsModified) const
 {
 #if WITH_EDITOR
@@ -40,3 +45,16 @@ void UAnimNextGraph_EdGraphSchema::TrySetDefaultValue(UEdGraphPin& InPin, const
 		}
 	}
 }
+
+void UAnimNextGraph_EdGraphSchema::GetGraphDisplayInformation(const UEdGraph& Graph, /*out*/ FGraphDisplayInfo& DisplayInfo) const
+{
+	Super::GetGraphDisplayInformation(Graph, DisplayInfo);
+	
+	if(UAnimNextRigVMAssetEntry* AssetEntry = Cast<UAnimNextRigVMAssetEntry>(Graph.GetOuter()))
+	{
+		DisplayInfo.DisplayName = FText::Format(LOCTEXT("GraphTabTitleFormat", "{0}: {1}"), FText::FromName(AssetEntry->GetEntryName()), FText::FromName(AssetEntry->GetTypedOuter<UAnimNextRigVMAsset>()->GetFName()));
+		DisplayInfo.Tooltip = FText::Format(LOCTEXT("GraphTabTooltipFormat", "{0} in:\n{1}"), FText::FromName(AssetEntry->GetEntryName()), FText::FromString(AssetEntry->GetTypedOuter<UAnimNextRigVMAsset>()->GetPathName()));
+	}
+}
+
+#undef LOCTEXT_NAMESPACE
\ No newline at end of file
diff --git a/Engine/Plugins/Experimental/AnimNext/Source/AnimNextUncookedOnly/Private/Graph/AnimNextGraph_EditorData.cpp b/Engine/Plugins/Experimental/AnimNext/Source/AnimNextUncookedOnly/Private/Graph/AnimNextGraph_EditorData.cpp
--- a/Engine/Plugins/Experimental/AnimNext/Source/AnimNextUncookedOnly/Private/Graph/AnimNextGraph_EditorData.cpp
+++ b/Engine/Plugins/Experimental/AnimNext/Source/AnimNextUncookedOnly/Private/Graph/AnimNextGraph_EditorData.cpp
@@ -48,7 +48,7 @@ void UAnimNextGraph_EditorData::PostLoad()
 					UAnimNextGraphEntry* GraphEntry = CastChecked<UAnimNextGraphEntry>(FoundEntry);
 					GraphEntry->EdGraph = Graph;
 
-					Graph->Rename(*MakeEdGraphName(FoundRigVMGraph->GetFName()).ToString(), FoundEntry, REN_ForceNoResetLoaders | REN_DoNotDirty | REN_DontCreateRedirectors | REN_NonTransactional);
+					Graph->Rename(nullptr, FoundEntry, REN_ForceNoResetLoaders | REN_DoNotDirty | REN_DontCreateRedirectors | REN_NonTransactional);
 					Graph->Initialize(this);
 				}
 			}
@@ -58,6 +58,7 @@ void UAnimNextGraph_EditorData::PostLoad()
 
 void UAnimNextGraph_EditorData::RecompileVM()
 {
+	UE::AnimNext::UncookedOnly::FUtils::GetAssetParameters(this, CachedExports);
 	UE::AnimNext::UncookedOnly::FUtils::Compile(GetTypedOuter<UAnimNextGraph>());
 }
 
@@ -125,9 +126,18 @@ UEdGraph* UAnimNextGraph_EditorData::CreateEdGraph(URigVMGraph* InRigVMGraph, bo
 	UAnimNextGraphEntry* Entry = Cast<UAnimNextGraphEntry>(FindEntryForRigVMGraph(InRigVMGraph));
 	if(Entry == nullptr)
 	{
-		return nullptr;
+		// Not found, we could be adding a new entry, in which case the graph wont be assigned yet
+		check(Entries.Num() > 0);
+		check(Cast<IAnimNextRigVMGraphInterface>(Entries.Last()) != nullptr);
+		check(Cast<IAnimNextRigVMGraphInterface>(Entries.Last())->GetRigVMGraph() == nullptr);
+		Entry = Cast<UAnimNextGraphEntry>(FindEntryForRigVMGraph(nullptr));
 	}
 
+	if(Entry == nullptr)
+	{
+		return nullptr;
+	}
+	
 	if(bForce)
 	{
 		RemoveEdGraph(InRigVMGraph);
@@ -136,7 +146,7 @@ UEdGraph* UAnimNextGraph_EditorData::CreateEdGraph(URigVMGraph* InRigVMGraph, bo
 	FString GraphName = InRigVMGraph->GetName();
 	check(!GraphName.IsEmpty());
 
-	UAnimNextGraph_EdGraph* RigFunctionGraph = NewObject<UAnimNextGraph_EdGraph>(Entry, MakeEdGraphName(*GraphName), RF_Transactional);
+	UAnimNextGraph_EdGraph* RigFunctionGraph = NewObject<UAnimNextGraph_EdGraph>(Entry, NAME_None, RF_Transactional);
 	RigFunctionGraph->Schema = UAnimNextGraph_EdGraphSchema::StaticClass();
 	RigFunctionGraph->bAllowDeletion = true;
 	RigFunctionGraph->bIsFunctionDefinition = false;
@@ -191,7 +201,7 @@ UAnimNextGraphEntry* UAnimNextGraph_EditorData::AddGraph(FName InName, bool bSet
 	int32 NameNumber = InName.GetNumber() + 1;
 	while(bAlreadyExists)
 	{
-		NewGraphName = FName(InName, NameNumber);
+		NewGraphName = FName(InName, NameNumber++);
 		bAlreadyExists =  Entries.ContainsByPredicate(DuplicateNamePredicate);
 	}
 
@@ -210,12 +220,12 @@ UAnimNextGraphEntry* UAnimNextGraph_EditorData::AddGraph(FName InName, bool bSet
 	{
 		TGuardValue<bool> EnablePythonPrint(bSuspendPythonMessagesForRigVMClient, !bPrintPythonCommand);
 		TGuardValue<bool> DisableAutoCompile(bAutoRecompileVM, false);
-		URigVMGraph* NewGraph = RigVMClient.AddModel(NewGraphName, bSetupUndoRedo);
+		URigVMGraph* NewGraph = RigVMClient.AddModel(URigVMGraph::StaticClass()->GetFName(), bSetupUndoRedo);
 		ensure(NewGraph);
 		NewEntry->Graph = NewGraph;
 
 		URigVMController* Controller = RigVMClient.GetController(NewGraph);
-		UE::AnimNext::UncookedOnly::FUtils::SetupAnimGraph(Controller);
+		UE::AnimNext::UncookedOnly::FUtils::SetupAnimGraph(NewEntry, Controller);
 	}
 
 	BroadcastModified();
diff --git a/Engine/Plugins/Experimental/AnimNext/Source/AnimNextUncookedOnly/Private/Param/AnimNextParameterBlock_EdGraphSchema.cpp b/Engine/Plugins/Experimental/AnimNext/Source/AnimNextUncookedOnly/Private/Param/AnimNextParameterBlock_EdGraphSchema.cpp
--- a/Engine/Plugins/Experimental/AnimNext/Source/AnimNextUncookedOnly/Private/Param/AnimNextParameterBlock_EdGraphSchema.cpp
+++ b/Engine/Plugins/Experimental/AnimNext/Source/AnimNextUncookedOnly/Private/Param/AnimNextParameterBlock_EdGraphSchema.cpp
@@ -2,3 +2,20 @@
 
 #include "Param/AnimNextParameterBlock_EdGraphSchema.h"
 
+#include "AnimNextRigVMAsset.h"
+#include "AnimNextRigVMAssetEntry.h"
+
+#define LOCTEXT_NAMESPACE "AnimNextGraph_EdGraphSchema"
+
+void UAnimNextParameterBlock_EdGraphSchema::GetGraphDisplayInformation(const UEdGraph& Graph, /*out*/ FGraphDisplayInfo& DisplayInfo) const
+{
+	Super::GetGraphDisplayInformation(Graph, DisplayInfo);
+	
+	if(UAnimNextRigVMAssetEntry* AssetEntry = Cast<UAnimNextRigVMAssetEntry>(Graph.GetOuter()))
+	{
+		DisplayInfo.DisplayName = FText::Format(LOCTEXT("GraphTabTitleFormat", "{0}: {1}"), FText::FromName(AssetEntry->GetEntryName()), FText::FromName(AssetEntry->GetTypedOuter<UAnimNextRigVMAsset>()->GetFName()));
+		DisplayInfo.Tooltip = FText::Format(LOCTEXT("GraphTabTooltipFormat", "{0} in:\n{1}"), FText::FromName(AssetEntry->GetEntryName()), FText::FromString(AssetEntry->GetTypedOuter<UAnimNextRigVMAsset>()->GetPathName()));
+	}
+}
+
+#undef LOCTEXT_NAMESPACE
\ No newline at end of file
diff --git a/Engine/Plugins/Experimental/AnimNext/Source/AnimNextUncookedOnly/Private/Param/AnimNextParameterBlock_EditorData.cpp b/Engine/Plugins/Experimental/AnimNext/Source/AnimNextUncookedOnly/Private/Param/AnimNextParameterBlock_EditorData.cpp
--- a/Engine/Plugins/Experimental/AnimNext/Source/AnimNextUncookedOnly/Private/Param/AnimNextParameterBlock_EditorData.cpp
+++ b/Engine/Plugins/Experimental/AnimNext/Source/AnimNextUncookedOnly/Private/Param/AnimNextParameterBlock_EditorData.cpp
@@ -106,14 +106,14 @@ UAnimNextParameterBlockGraph* UAnimNextParameterBlock_EditorData::AddGraph(FName
 	{
 		TGuardValue<bool> EnablePythonPrint(bSuspendPythonMessagesForRigVMClient, !bPrintPythonCommand);
 		TGuardValue<bool> DisableAutoCompile(bAutoRecompileVM, false);
-		URigVMGraph* NewGraph = RigVMClient.AddModel(NewGraphName, bSetupUndoRedo);
+		URigVMGraph* NewGraph = RigVMClient.AddModel(URigVMGraph::StaticClass()->GetFName(), bSetupUndoRedo);
 		ensure(NewGraph);
 		NewEntry->Graph = NewGraph;
 
 		URigVMController* Controller = RigVMClient.GetController(NewGraph);
 		UE::AnimNext::UncookedOnly::FUtils::SetupParameterGraph(Controller);
 	}
-	
+
 	BroadcastModified();
 
 	return NewEntry;
@@ -165,7 +165,7 @@ void UAnimNextParameterBlock_EditorData::PostLoad()
 					UAnimNextParameterBlockGraph* GraphEntry = CastChecked<UAnimNextParameterBlockGraph>(FoundEntry);
 					GraphEntry->EdGraph = Graph;
 
-					Graph->Rename(*MakeEdGraphName(FoundRigVMGraph->GetFName()).ToString(), FoundEntry, REN_ForceNoResetLoaders | REN_DoNotDirty | REN_DontCreateRedirectors | REN_NonTransactional);
+					Graph->Rename(nullptr, FoundEntry, REN_ForceNoResetLoaders | REN_DoNotDirty | REN_DontCreateRedirectors | REN_NonTransactional);
 					Graph->Initialize(this);
 				}
 			}
@@ -202,6 +202,7 @@ void UAnimNextParameterBlock_EditorData::PostLoad()
 
 void UAnimNextParameterBlock_EditorData::RecompileVM()
 {
+	UE::AnimNext::UncookedOnly::FUtils::GetAssetParameters(this, CachedExports);
 	UE::AnimNext::UncookedOnly::FUtils::CompileVM(GetTypedOuter<UAnimNextParameterBlock>());
 }
 
@@ -237,6 +238,15 @@ UEdGraph* UAnimNextParameterBlock_EditorData::CreateEdGraph(URigVMGraph* InRigVM
 	}
 
 	UAnimNextParameterBlockGraph* Entry = Cast<UAnimNextParameterBlockGraph>(FindEntryForRigVMGraph(InRigVMGraph));
+	if(Entry == nullptr)
+	{
+		// Not found, we could be adding a new entry, in which case the graph wont be assigned yet
+		check(Entries.Num() > 0);
+		check(Cast<IAnimNextRigVMGraphInterface>(Entries.Last()) != nullptr);
+		check(Cast<IAnimNextRigVMGraphInterface>(Entries.Last())->GetRigVMGraph() == nullptr);
+		Entry = Cast<UAnimNextParameterBlockGraph>(Entries.Last());
+	}
+
 	if(Entry == nullptr)
 	{
 		return nullptr;
@@ -250,7 +260,7 @@ UEdGraph* UAnimNextParameterBlock_EditorData::CreateEdGraph(URigVMGraph* InRigVM
 	FString GraphName = InRigVMGraph->GetName();
 	check(!GraphName.IsEmpty());
 
-	UAnimNextParameterBlock_EdGraph* RigFunctionGraph = NewObject<UAnimNextParameterBlock_EdGraph>(Entry, MakeEdGraphName(*GraphName), RF_Transactional);
+	UAnimNextParameterBlock_EdGraph* RigFunctionGraph = NewObject<UAnimNextParameterBlock_EdGraph>(Entry, NAME_None, RF_Transactional);
 	RigFunctionGraph->Schema = UAnimNextParameterBlock_EdGraphSchema::StaticClass();
 	RigFunctionGraph->bAllowDeletion = true;
 	RigFunctionGraph->bIsFunctionDefinition = false;
diff --git a/Engine/Plugins/Experimental/AnimNext/Source/AnimNextUncookedOnly/Private/UncookedOnlyUtils.cpp b/Engine/Plugins/Experimental/AnimNext/Source/AnimNextUncookedOnly/Private/UncookedOnlyUtils.cpp
--- a/Engine/Plugins/Experimental/AnimNext/Source/AnimNextUncookedOnly/Private/UncookedOnlyUtils.cpp
+++ b/Engine/Plugins/Experimental/AnimNext/Source/AnimNextUncookedOnly/Private/UncookedOnlyUtils.cpp
@@ -44,6 +44,7 @@
 #include "AnimNextRigVMAsset.h"
 #include "AnimNextRigVMAssetEditorData.h"
 #include "AnimNextRigVMAssetEntry.h"
+#include "Graph/AnimNextGraphEntryPoint.h"
 
 namespace UE::AnimNext::UncookedOnly
 {
@@ -88,6 +89,18 @@ namespace Private
 		{}
 	};
 
+	struct FDecoratorGraph
+	{
+		FName EntryPoint;
+		URigVMNode* RootNode;
+		TArray<FDecoratorStackMapping> DecoratorStackNodes;
+
+		explicit FDecoratorGraph(URigVMNode* InRootNode)
+			: EntryPoint(*InRootNode->FindPin(GET_MEMBER_NAME_STRING_CHECKED(FRigUnit_AnimNextGraphRoot, EntryPoint))->GetDefaultValue())
+			, RootNode(InRootNode)
+		{}
+	};
+
 	template<typename DecoratorAction>
 	void ForEachDecoratorInStack(const URigVMNode* DecoratorStackNode, const DecoratorAction& Action)
 	{
@@ -227,11 +240,11 @@ namespace Private
 			});
 	}
 
-	const URigVMUnitNode* FindRootNode(const TArray<URigVMNode*>& VMNodes)
+	URigVMUnitNode* FindRootNode(const TArray<URigVMNode*>& VMNodes)
 	{
-		for (const URigVMNode* VMNode : VMNodes)
+		for (URigVMNode* VMNode : VMNodes)
 		{
-			if (const URigVMUnitNode* VMUnitNode = Cast<URigVMUnitNode>(VMNode))
+			if (URigVMUnitNode* VMUnitNode = Cast<URigVMUnitNode>(VMNode))
 			{
 				const UScriptStruct* ScriptStruct = VMUnitNode->GetScriptStruct();
 				if (ScriptStruct == FRigUnit_AnimNextGraphRoot::StaticStruct())
@@ -244,7 +257,7 @@ namespace Private
 		return nullptr;
 	}
 
-	void AddMissingInputLinks(const URigVMPin* DecoratorPin, UAnimNextGraph_Controller* VMController)
+	void AddMissingInputLinks(const URigVMPin* DecoratorPin, URigVMController* VMController)
 	{
 		const TArray<URigVMPin*>& Pins = DecoratorPin->GetSubPins();
 		for (URigVMPin* Pin : Pins)
@@ -302,7 +315,7 @@ namespace Private
 		}
 	}
 
-	void AddMissingInputLinks(const URigVMGraph* VMGraph, UAnimNextGraph_Controller* VMController)
+	void AddMissingInputLinks(const URigVMGraph* VMGraph, URigVMController* VMController)
 	{
 		const TArray<URigVMNode*> VMNodes = VMGraph->GetNodes();	// Copy since we might add new nodes
 		for (URigVMNode* VMNode : VMNodes)
@@ -324,12 +337,10 @@ namespace Private
 		}
 	}
 
-	TArray<FDecoratorStackMapping> CollectDecoratorStacks(const URigVMGraph* VMGraph, UAnimNextGraph_Controller* VMController)
+	FDecoratorGraph CollectGraphInfo(const URigVMGraph* VMGraph, URigVMController* VMController)
 	{
 		const TArray<URigVMNode*>& VMNodes = VMGraph->GetNodes();
-		const URigVMUnitNode* VMRootNode = FindRootNode(VMNodes);
-
-		TArray<FDecoratorStackMapping> DecoratorStackNodes;
+		URigVMUnitNode* VMRootNode = FindRootNode(VMNodes);
 
 		if (VMRootNode == nullptr)
 		{
@@ -340,6 +351,8 @@ namespace Private
 		// Make sure we don't have empty input pins
 		AddMissingInputLinks(VMGraph, VMController);
 
+		FDecoratorGraph DecoratorGraph(VMRootNode);
+
 		TArray<const URigVMNode*> NodesToVisit;
 		NodesToVisit.Add(VMRootNode);
 
@@ -360,7 +373,7 @@ namespace Private
 							Mapping.DecoratorEntries.Add(FDecoratorEntryMapping(DecoratorStackNode, DecoratorPin, Decorator));
 						});
 
-					DecoratorStackNodes.Add(MoveTemp(Mapping));
+					DecoratorGraph.DecoratorStackNodes.Add(MoveTemp(Mapping));
 				}
 			}
 
@@ -368,7 +381,7 @@ namespace Private
 			NodesToVisit.Append(SourceNodes);
 		}
 
-		if (DecoratorStackNodes.IsEmpty())
+		if (DecoratorGraph.DecoratorStackNodes.IsEmpty())
 		{
 			// If the graph is empty, add a dummy node that just pushes a reference pose
 			URigVMUnitNode* VMNode = VMController->AddUnitNode(FRigUnit_AnimNextDecoratorStack::StaticStruct(), FRigVMStruct::ExecuteName, FVector2D(0.0f, 0.0f), FString(), false);
@@ -400,18 +413,14 @@ namespace Private
 					Mapping.DecoratorEntries.Add(FDecoratorEntryMapping(DecoratorStackNode, DecoratorPin, Decorator));
 				});
 
-			DecoratorStackNodes.Add(MoveTemp(Mapping));
+			DecoratorGraph.DecoratorStackNodes.Add(MoveTemp(Mapping));
 		}
 
-		return DecoratorStackNodes;
+		return DecoratorGraph;
 	}
 
-	FRigVMPinInfoArray CollectLatentPins(TArray<FDecoratorStackMapping>& DecoratorStackNodes, TMap<FName, URigVMPin*>& LatentPinMapping)
+	void CollectLatentPins(TArray<FDecoratorStackMapping>& DecoratorStackNodes, FRigVMPinInfoArray& OutLatentPins, TMap<FName, URigVMPin*>& OutLatentPinMapping)
 	{
-		const FRigVMRegistry& Registry = FRigVMRegistry::Get();
-
-		FRigVMPinInfoArray LatentPins;
-
 		for (FDecoratorStackMapping& DecoratorStack : DecoratorStackNodes)
 		{
 			for (FDecoratorEntryMapping& DecoratorEntry : DecoratorStack.DecoratorEntries)
@@ -421,10 +430,10 @@ namespace Private
 					if (Pin->IsLazy() && !Pin->GetLinks().IsEmpty())
 					{
 						// This pin has something linked to it, it is a latent pin
-						check(LatentPins.Num() < ((1 << 16) - 1));	// We reserve MAX_uint16 as an invalid value and we must fit on 15 bits when packed
-						DecoratorEntry.LatentPropertyNameToIndexMap.Add(Pin->GetFName(), (uint16)LatentPins.Num());
+						check(OutLatentPins.Num() < ((1 << 16) - 1));	// We reserve MAX_uint16 as an invalid value and we must fit on 15 bits when packed
+						DecoratorEntry.LatentPropertyNameToIndexMap.Add(Pin->GetFName(), (uint16)OutLatentPins.Num());
 
-						const FName LatentPinName(TEXT("LatentPin"), LatentPins.Num());	// Create unique latent pin names
+						const FName LatentPinName(TEXT("LatentPin"), OutLatentPins.Num());	// Create unique latent pin names
 
 						FRigVMPinInfo PinInfo;
 						PinInfo.Name = LatentPinName;
@@ -434,18 +443,16 @@ namespace Private
 						PinInfo.Direction = ERigVMPinDirection::Input;
 						PinInfo.bIsLazy = true;
 
-						LatentPins.Pins.Emplace(PinInfo);
+						OutLatentPins.Pins.Emplace(PinInfo);
 
 						const TArray<URigVMLink*>& PinLinks = Pin->GetLinks();
 						check(PinLinks.Num() == 1);
 
-						LatentPinMapping.Add(LatentPinName, PinLinks[0]->GetSourcePin());
+						OutLatentPinMapping.Add(LatentPinName, PinLinks[0]->GetSourcePin());
 					}
 				}
 			}
 		}
-
-		return LatentPins;
 	}
 
 	FAnimNextGraphEvaluatorExecuteDefinition GetGraphEvaluatorExecuteMethod(const FRigVMPinInfoArray& LatentPins)
@@ -511,6 +518,13 @@ void FUtils::Compile(UAnimNextGraph* InGraph)
 	RecreateVM(InGraph);
 
 	InGraph->VMRuntimeSettings = EditorData->VMRuntimeSettings;
+	InGraph->EntryPoints.Empty();
+	InGraph->ResolvedRootDecoratorHandles.Empty();
+	InGraph->ExecuteDefinition = FAnimNextGraphEvaluatorExecuteDefinition();
+	InGraph->SharedDataBuffer.Empty();
+	InGraph->GraphReferencedObjects.Empty();
+	InGraph->RequiredParametersHash = 0;
+	InGraph->RequiredParameters.Empty();
 
 	EditorData->CompileLog.Messages.Reset();
 	EditorData->CompileLog.NumErrors = EditorData->CompileLog.NumWarnings = 0;
@@ -522,27 +536,60 @@ void FUtils::Compile(UAnimNextGraph* InGraph)
 	{
 		return;
 	}
-	
-	// We use a temporary graph model to build our final graph that we'll compile
-	URigVMGraph* VMTempGraph = CastChecked<URigVMGraph>(StaticDuplicateObject(VMRootGraph, GetTransientPackage(), VMClient->GetUniqueName(TEXT("TempRigVMGraph"))));
 
-	UAnimNextGraph_Controller* TempController = CastChecked<UAnimNextGraph_Controller>(VMClient->GetOrCreateController(VMTempGraph));
+	TArray<URigVMGraph*> VMTempGraphs;
+	for(const URigVMGraph* SourceGraph : VMClient->GetAllModels(false, false))
+	{
+		// We use a temporary graph models to build our final graphs that we'll compile
+		URigVMGraph* VMTempGraph = CastChecked<URigVMGraph>(StaticDuplicateObject(SourceGraph, GetTransientPackage(), NAME_None, RF_Transient));
+		VMTempGraphs.Add(VMTempGraph);
+	}
+
+	if(VMTempGraphs.Num() == 0)
+	{
+		return;
+	}
 
-	// Gather our decorator stacks
-	TArray<Private::FDecoratorStackMapping> DecoratorStackNodes = Private::CollectDecoratorStacks(VMTempGraph, TempController);
-	check(!DecoratorStackNodes.IsEmpty());
+	UAnimNextGraph_Controller* TempController = CastChecked<UAnimNextGraph_Controller>(VMClient->GetOrCreateController(VMTempGraphs[0]));
 
-	// Add our runtime shim root node
-	URigVMUnitNode* TempShimRootNode = TempController->AddUnitNode(FRigUnit_AnimNextShimRoot::StaticStruct(), FRigVMStruct::ExecuteName, FVector2D::ZeroVector, FString(), false);
+	FDecoratorWriter DecoratorWriter;
 
-	// Add our graph evaluator node
+	FRigVMPinInfoArray LatentPins;
 	TMap<FName, URigVMPin*> LatentPinMapping;
-	const FRigVMPinInfoArray LatentPins = Private::CollectLatentPins(DecoratorStackNodes, LatentPinMapping);
+	TArray<Private::FDecoratorGraph> DecoratorGraphs;
+
+	// Build entry points and extract their required latent pins
+	for(const URigVMGraph* VMTempGraph : VMTempGraphs)
+	{
+		// Gather our decorator stacks
+		Private::FDecoratorGraph& DecoratorGraph = DecoratorGraphs.Add_GetRef(Private::CollectGraphInfo(VMTempGraph, TempController->GetControllerForGraph(VMTempGraph)));
+		check(!DecoratorGraph.DecoratorStackNodes.IsEmpty());
+
+		FAnimNextGraphEntryPoint& EntryPoint = InGraph->EntryPoints.AddDefaulted_GetRef();
+		EntryPoint.EntryPointName = DecoratorGraph.EntryPoint;
+
+		// Extract latent pins for this graph
+		Private::CollectLatentPins(DecoratorGraph.DecoratorStackNodes, LatentPins, LatentPinMapping);
+
+		// Iterate over every decorator stack and register our node templates
+		for (Private::FDecoratorStackMapping& NodeMapping : DecoratorGraph.DecoratorStackNodes)
+		{
+			NodeMapping.DecoratorStackNodeHandle = Private::RegisterDecoratorNodeTemplate(DecoratorWriter, NodeMapping.DecoratorStackNode);
+		}
+
+		// Find our root node handle, if we have any stack nodes, the first one is our root stack
+		if (DecoratorGraph.DecoratorStackNodes.Num() != 0)
+		{
+			EntryPoint.RootDecoratorHandle = FAnimNextEntryPointHandle(DecoratorGraph.DecoratorStackNodes[0].DecoratorStackNodeHandle);
+		}
+	}
 
 	// We need a unique method name to match our unique argument list
-	const FAnimNextGraphEvaluatorExecuteDefinition ExecuteDefinition = Private::GetGraphEvaluatorExecuteMethod(LatentPins);
+	InGraph->ExecuteDefinition = Private::GetGraphEvaluatorExecuteMethod(LatentPins);
 
-	URigVMUnitNode* GraphEvaluatorNode = TempController->AddUnitNodeWithPins(FRigUnit_AnimNextGraphEvaluator::StaticStruct(), LatentPins, *ExecuteDefinition.MethodName, FVector2D::ZeroVector, FString(), false);
+	// Add our runtime shim root node
+	URigVMUnitNode* TempShimRootNode = TempController->AddUnitNode(FRigUnit_AnimNextShimRoot::StaticStruct(), FRigUnit_AnimNextShimRoot::EventName, FVector2D::ZeroVector, FString(), false);
+	URigVMUnitNode* GraphEvaluatorNode = TempController->AddUnitNodeWithPins(FRigUnit_AnimNextGraphEvaluator::StaticStruct(), LatentPins, *InGraph->ExecuteDefinition.MethodName, FVector2D::ZeroVector, FString(), false);
 
 	// Link our shim and evaluator nodes together using the execution context
 	TempController->AddLink(
@@ -558,60 +605,38 @@ void FUtils::Compile(UAnimNextGraph* InGraph)
 			GraphEvaluatorNode->FindPin(LatentPin.Name.ToString()),
 			false);
 	}
-
-	FDecoratorWriter DecoratorWriter;
-
-	// Iterate over every decorator stack and register our node templates
-	for (Private::FDecoratorStackMapping& NodeMapping : DecoratorStackNodes)
-	{
-		NodeMapping.DecoratorStackNodeHandle = Private::RegisterDecoratorNodeTemplate(DecoratorWriter, NodeMapping.DecoratorStackNode);
-	}
-
+	
 	// Write our node shared data
 	DecoratorWriter.BeginNodeWriting();
 
-	for (const Private::FDecoratorStackMapping& NodeMapping : DecoratorStackNodes)
+	for(Private::FDecoratorGraph& DecoratorGraph : DecoratorGraphs)
 	{
-		Private::WriteDecoratorProperties(DecoratorWriter, NodeMapping, DecoratorStackNodes);
+		for (const Private::FDecoratorStackMapping& NodeMapping : DecoratorGraph.DecoratorStackNodes)
+		{
+			Private::WriteDecoratorProperties(DecoratorWriter, NodeMapping, DecoratorGraph.DecoratorStackNodes);
+		}
 	}
 
 	DecoratorWriter.EndNodeWriting();
 
-	// Find our root node handle, if we have any stack nodes, the first one is our root stack
-	FAnimNextEntryPointHandle RootDecoratorHandle;
-	if (DecoratorStackNodes.Num() != 0)
-	{
-		RootDecoratorHandle = FAnimNextEntryPointHandle(DecoratorStackNodes[0].DecoratorStackNodeHandle);
-	}
-
 	// Cache our compiled metadata
-	InGraph->ExecuteDefinition = ExecuteDefinition;
 	InGraph->SharedDataArchiveBuffer = DecoratorWriter.GetGraphSharedData();
 	InGraph->GraphReferencedObjects = DecoratorWriter.GetGraphReferencedObjects();
-	InGraph->RootDecoratorHandle = RootDecoratorHandle;
 
 	// Populate our runtime metadata
 	InGraph->LoadFromArchiveBuffer(InGraph->SharedDataArchiveBuffer);
 
-	// Remove our old root node
-	if (URigVMNode* const* RootNode = VMTempGraph->GetNodes().FindByPredicate(
-		[](URigVMNode* Node)
-		{
-			if (URigVMUnitNode* UnitNode = Cast<URigVMUnitNode>(Node))
-			{
-				return UnitNode->GetScriptStruct() == FRigUnit_AnimNextGraphRoot::StaticStruct();
-			}
-
-			return false;
-		}))
+	// Remove our old root nodes
+	for(Private::FDecoratorGraph& DecoratorGraph : DecoratorGraphs)
 	{
-		TempController->RemoveNode(*RootNode, false, false);
+		URigVMController* GraphController = TempController->GetControllerForGraph(DecoratorGraph.RootNode->GetGraph());
+		GraphController->RemoveNode(DecoratorGraph.RootNode, false, false);
 	}
 
 	URigVMCompiler* Compiler = URigVMCompiler::StaticClass()->GetDefaultObject<URigVMCompiler>();
 	EditorData->VMCompileSettings.SetExecuteContextStruct(EditorData->RigVMClient.GetExecuteContextStruct());
 	const FRigVMCompileSettings Settings = (EditorData->bCompileInDebugMode) ? FRigVMCompileSettings::Fast(EditorData->VMCompileSettings.GetExecuteContextStruct()) : EditorData->VMCompileSettings;
-	Compiler->Compile(Settings, { VMTempGraph }, TempController, InGraph->VM, InGraph->ExtendedExecuteContext, TArray<FRigVMExternalVariable>(), & EditorData->PinToOperandMap);
+	Compiler->Compile(Settings, VMTempGraphs, TempController, InGraph->VM, InGraph->ExtendedExecuteContext, TArray<FRigVMExternalVariable>(), & EditorData->PinToOperandMap);
 
 	// Initialize right away, in packaged builds we initialize during PostLoad
 	InGraph->VM->Initialize(InGraph->ExtendedExecuteContext);
@@ -630,14 +655,14 @@ void FUtils::Compile(UAnimNextGraph* InGraph)
 		EditorData->RigVMCompiledEvent.Broadcast(InGraph, InGraph->VM, InGraph->ExtendedExecuteContext);
 	}
 
-	VMClient->RemoveController(VMTempGraph);
+	for(URigVMGraph* VMTempGraph : VMTempGraphs)
+	{
+		VMClient->RemoveController(VMTempGraph);
+	}
 
 	// Now that the graph has been re-compiled, re-allocate the previous live instances
 	InGraph->ThawGraphInstances();
 
-	// Gather and store the required external parameters
-	InGraph->RequiredParameters.Empty();
-
 	FAssetData AssetData(InGraph);
 	FAnimNextParameterProviderAssetRegistryExports Exports;
 	GetExportedParametersForAsset(AssetData, Exports);
@@ -826,7 +851,7 @@ void FUtils::CompileVM(UAnimNextParameterBlock* InParameterBlock)
 	FRigVMExtendedExecuteContext& CDOContext = InParameterBlock->GetRigVMExtendedExecuteContext();
 	const FRigVMCompileSettings Settings = (EditorData->bCompileInDebugMode) ? FRigVMCompileSettings::Fast(EditorData->VMCompileSettings.GetExecuteContextStruct()) : EditorData->VMCompileSettings;
 	URigVMController* RootController = VMClient->GetOrCreateController(RootGraph);
-	Compiler->Compile(Settings, { RootGraph }, RootController, InParameterBlock->VM, CDOContext, InParameterBlock->GetExternalVariables(), &EditorData->PinToOperandMap);
+	Compiler->Compile(Settings, VMClient->GetAllModels(false, false), RootController, InParameterBlock->VM, CDOContext, InParameterBlock->GetExternalVariables(), &EditorData->PinToOperandMap);
 
 	InParameterBlock->VM->Initialize(CDOContext);
 	InParameterBlock->GenerateUserDefinedDependenciesData(CDOContext);
@@ -1238,7 +1263,7 @@ FRigVMTemplateArgumentType FUtils::GetRigVMArgTypeFromParamType(const FAnimNextP
 	return ArgType;
 }
 
-void FUtils::SetupAnimGraph(URigVMController* InController)
+void FUtils::SetupAnimGraph(UAnimNextRigVMAssetEntry* InEntry, URigVMController* InController)
 {
 	// Clear the graph
 	InController->RemoveNodes(InController->GetGraph()->GetNodes());
@@ -1248,6 +1273,11 @@ void FUtils::SetupAnimGraph(URigVMController* InController)
 	URigVMPin* BeginExecutePin = MainEntryPointNode->FindPin(GET_MEMBER_NAME_STRING_CHECKED(FRigUnit_AnimNextGraphRoot, Result));
 	check(BeginExecutePin);
 	check(BeginExecutePin->GetDirection() == ERigVMPinDirection::Input);
+
+	URigVMPin* EntryPointPin = MainEntryPointNode->FindPin(GET_MEMBER_NAME_STRING_CHECKED(FRigUnit_AnimNextGraphRoot, EntryPoint));
+	check(EntryPointPin);
+	check(EntryPointPin->GetDirection() == ERigVMPinDirection::Hidden);
+	InController->SetPinDefaultValue(EntryPointPin->GetPinPath(), InEntry->GetEntryName().ToString());
 }
 
 void FUtils::SetupParameterGraph(URigVMController* InController)
@@ -1259,31 +1289,6 @@ void FUtils::SetupParameterGraph(URigVMController* InController)
 	InController->AddUnitNode(FRigUnit_AnimNextParameterBeginExecution::StaticStruct(), FRigVMStruct::ExecuteName, FVector2D(-200.0f, 0.0f), FString(), false);
 }
 
-void FUtils::SetupBindingGraph(URigVMController* InController, FName InParameterName, const FAnimNextParamType& InParamType)
-{
-	SetupBindingGraphForLiteral(InController, InParameterName, InParamType);
-}
-
-void FUtils::SetupBindingGraphForLiteral(URigVMController* InController, FName InParameterName, const FAnimNextParamType& InParamType)
-{
-	FRigVMTemplateArgumentType ArgType = GetRigVMArgTypeFromParamType(InParamType);
-	TRigVMTypeIndex TypeIndex = FRigVMRegistry::Get().GetTypeIndex(ArgType);
-
-	// Clear the graph
-	InController->RemoveNodes(InController->GetGraph()->GetNodes());
-
-	// Add new nodes for a simple literal binding
-	URigVMUnitNode* EntryPointNode = InController->AddUnitNode(FRigUnit_AnimNextParameterBeginExecution::StaticStruct(), FRigVMStruct::ExecuteName, FVector2D(-200.0f, 0.0f), FString(), false);
-
-	const FName FactoryName = FRigVMDispatch_SetLayerParameter().GetFactoryName();
-	FRigVMDispatch_SetLayerParameter* Factory = static_cast<FRigVMDispatch_SetLayerParameter*>(FRigVMRegistry::Get().FindDispatchFactory(FactoryName));
-	URigVMTemplateNode* SetParameterNode = InController->AddTemplateNode(Factory->GetTemplate()->GetNotation(), FVector2D(200.0f, 0.0f));
-	InController->SetPinDefaultValue(SetParameterNode->FindPin(FRigVMDispatch_SetLayerParameter::ParameterName.ToString())->GetPinPath(), InParameterName.ToString());
-	InController->ResolveWildCardPin(SetParameterNode->FindPin(FRigVMDispatch_SetLayerParameter::ValueName.ToString()), TypeIndex);
-
-	InController->AddLink(EntryPointNode->FindPin(GET_MEMBER_NAME_STRING_CHECKED(FRigUnit_AnimNextParameterBeginExecution, ExecuteContext)), SetParameterNode->FindPin(FRigVMDispatch_SetLayerParameter::ExecuteContextName.ToString()));
-}
-
 FText FUtils::GetParameterDisplayNameText(FName InParameterName)
 {
 	FString NameAsString = InParameterName.ToString();
@@ -1360,6 +1365,28 @@ bool FUtils::GetExportedParametersFromAssetRegistry(FAnimNextParameterProviderAs
 	return OutExports.Parameters.Num() > 0;
 }
 
+void FUtils::GetAssetParameters(const UAnimNextRigVMAssetEditorData* EditorData, FAnimNextParameterProviderAssetRegistryExports& OutExports)
+{
+	OutExports.Parameters.Reset();
+	OutExports.Parameters.Reserve(EditorData->Entries.Num());
+
+	for(const UAnimNextRigVMAssetEntry* Entry : EditorData->Entries)
+	{
+		if(const IAnimNextRigVMParameterInterface* ParameterInterface = Cast<IAnimNextRigVMParameterInterface>(Entry))
+		{
+			// TODO: Public/private symbols would influence whether this would be exposed to the asset registry here
+			OutExports.Parameters.Emplace(Entry->GetEntryName(), ParameterInterface->GetParamType(), EAnimNextParameterFlags::Bound);
+		}
+		else if(const IAnimNextRigVMGraphInterface* GraphInterface = Cast<IAnimNextRigVMGraphInterface>(Entry))
+		{
+			// TODO: Public/private symbols would influence whether this would be exposed to the asset registry here
+			OutExports.Parameters.Emplace(Entry->GetEntryName(), FAnimNextParamType::GetType<FAnimNextEntryPoint>(), EAnimNextParameterFlags::Bound);
+
+			GetGraphParameters(GraphInterface->GetRigVMGraph(), OutExports);
+		}
+	}
+}
+
 void FUtils::GetGraphParameters(const URigVMGraph* Graph, FAnimNextParameterProviderAssetRegistryExports& OutExports)
 {
 	const TArray<URigVMNode*>& Nodes = Graph->GetNodes();
@@ -1382,20 +1409,46 @@ void FUtils::GetGraphParameters(const URigVMGraph* Graph, FAnimNextParameterProv
 			{
 				if (const URigVMPin* NamePin = TemplateNode->FindPin(FRigVMDispatch_GetParameter::ParameterName.ToString()))
 				{
-					if (const URigVMPin* ValuePin = TemplateNode->FindPin(FRigVMDispatch_GetParameter::ValueName.ToString()))
+					const FString PinDefaultValue = NamePin->GetDefaultValue();
+					if(!PinDefaultValue.IsEmpty() && PinDefaultValue != TEXT("None"))
 					{
-						FAnimNextParamType Type = FAnimNextParamType::FromRigVMTemplateArgument(ValuePin->GetTemplateArgumentType());
-						EAnimNextParameterFlags Flags = EAnimNextParameterFlags::NoFlags;
-						if (bReadParameter)
+						if (const URigVMPin* ValuePin = TemplateNode->FindPin(FRigVMDispatch_GetParameter::ValueName.ToString()))
 						{
-							Flags |= EAnimNextParameterFlags::Read;
+							FAnimNextParamType Type = FAnimNextParamType::FromRigVMTemplateArgument(ValuePin->GetTemplateArgumentType());
+							EAnimNextParameterFlags Flags = EAnimNextParameterFlags::NoFlags;
+							if (bReadParameter)
+							{
+								Flags |= EAnimNextParameterFlags::Read;
+							}
+							
+							if (bWriteParameter)
+							{
+								Flags |= EAnimNextParameterFlags::Write;
+							}
+							OutExports.Parameters.Emplace(FName(PinDefaultValue), Type, Flags);
 						}
-						
-						if (bWriteParameter)
+					}
+				}
+			}
+			else
+			{
+				for(const URigVMPin* Pin : TemplateNode->GetAllPinsRecursively())
+				{
+					if(Pin->GetCPPType() == TEXT("FName"))
+					{
+						if(Pin->GetCustomWidgetName() == "ParamName")
 						{
-							Flags |= EAnimNextParameterFlags::Write;
+							const FString PinDefaultValue = Pin->GetDefaultValue();
+							if(!PinDefaultValue.IsEmpty() && PinDefaultValue != TEXT("None"))
+							{
+								const FString ParamTypeString = Pin->GetMetaData("AllowedParamType");
+								FAnimNextParamType Type = FAnimNextParamType::FromString(ParamTypeString);
+								if(Type.IsValid())
+								{
+									OutExports.Parameters.Emplace(FName(Pin->GetDefaultValue()), Type, EAnimNextParameterFlags::Read);
+								}
+							}
 						}
-						OutExports.Parameters.Emplace(FName(NamePin->GetDefaultValue()), Type, Flags);
 					}
 				}
 			}
diff --git a/Engine/Plugins/Runtime/RigVM/Source/RigVM/Public/RigVMCore/RigVMStruct.h b/Engine/Plugins/Runtime/RigVM/Source/RigVM/Public/RigVMCore/RigVMStruct.h
--- a/Engine/Plugins/Runtime/RigVM/Source/RigVM/Public/RigVMCore/RigVMStruct.h
+++ b/Engine/Plugins/Runtime/RigVM/Source/RigVM/Public/RigVMCore/RigVMStruct.h
@@ -168,6 +168,7 @@ struct RIGVM_API FRigVMStruct
 	virtual FName GetEventName() const { return NAME_None; }
 	virtual bool CanOnlyExistOnce() const { return false; }
 	virtual FString GetUnitLabel() const { return FString(); };
+	virtual FString GetUnitSubTitle() const { return FString(); };
 
 public:
 
diff --git a/Engine/Plugins/Runtime/RigVM/Source/RigVMDeveloper/Private/EdGraph/RigVMEdGraphNode.cpp b/Engine/Plugins/Runtime/RigVM/Source/RigVMDeveloper/Private/EdGraph/RigVMEdGraphNode.cpp
--- a/Engine/Plugins/Runtime/RigVM/Source/RigVMDeveloper/Private/EdGraph/RigVMEdGraphNode.cpp
+++ b/Engine/Plugins/Runtime/RigVM/Source/RigVMDeveloper/Private/EdGraph/RigVMEdGraphNode.cpp
@@ -70,6 +70,7 @@ FText URigVMEdGraphNode::GetNodeTitle(ENodeTitleType::Type TitleType) const
 					{
 						const FRigVMStruct* RigVMStruct = (const FRigVMStruct*)StructOnScope->GetStructMemory();
 						NodeTitle = FText::FromString(RigVMStruct->GetUnitLabel());
+						SubTitle = RigVMStruct->GetUnitSubTitle();
 					}
 				}
 			}
diff --git a/Engine/Plugins/Runtime/RigVM/Source/RigVMDeveloper/Private/RigVMModel/Nodes/RigVMUnitNode.cpp b/Engine/Plugins/Runtime/RigVM/Source/RigVMDeveloper/Private/RigVMModel/Nodes/RigVMUnitNode.cpp
--- a/Engine/Plugins/Runtime/RigVM/Source/RigVMDeveloper/Private/RigVMModel/Nodes/RigVMUnitNode.cpp
+++ b/Engine/Plugins/Runtime/RigVM/Source/RigVMDeveloper/Private/RigVMModel/Nodes/RigVMUnitNode.cpp
@@ -315,10 +315,6 @@ FString URigVMUnitNode::GetStructDefaultValue() const
 	TArray<FString> PinDefaultValues;
 	for (URigVMPin* Pin : GetPins())
 	{
-		if (Pin->GetDirection() == ERigVMPinDirection::Hidden)
-		{
-			continue;
-		}
 		FString PinDefaultValue = Pin->GetDefaultValue();
 		if (Pin->IsStringType())
 		{
diff --git a/Engine/Plugins/Runtime/RigVM/Source/RigVMDeveloper/Private/RigVMModel/RigVMNode.cpp b/Engine/Plugins/Runtime/RigVM/Source/RigVMDeveloper/Private/RigVMModel/RigVMNode.cpp
--- a/Engine/Plugins/Runtime/RigVM/Source/RigVMDeveloper/Private/RigVMModel/RigVMNode.cpp
+++ b/Engine/Plugins/Runtime/RigVM/Source/RigVMDeveloper/Private/RigVMModel/RigVMNode.cpp
@@ -510,6 +510,25 @@ TSharedPtr<FStructOnScope> URigVMNode::GetDecoratorInstance(const URigVMPin* InD
 	return EmptyScope;
 }
 
+UScriptStruct* URigVMNode::GetDecoratorScriptStruct(const FName& InName) const
+{
+	return GetDecoratorScriptStruct(FindPin(InName.ToString()));
+}
+
+UScriptStruct* URigVMNode::GetDecoratorScriptStruct(const URigVMPin* InDecoratorPin) const
+{
+	if(const URigVMPin* RootPin = FindDecorator(InDecoratorPin))
+	{
+		check(RootPin->IsStruct());
+
+		UScriptStruct* ScriptStruct = RootPin->GetScriptStruct();
+		check(ScriptStruct->IsChildOf(FRigVMDecorator::StaticStruct()));
+		return ScriptStruct;
+	}
+
+	return nullptr;
+}
+
 URigVMLibraryNode* URigVMNode::FindFunctionForNode() const  
 {
 	const UObject* Subject = this;
diff --git a/Engine/Plugins/Runtime/RigVM/Source/RigVMDeveloper/Private/RigVMModel/RigVMPin.cpp b/Engine/Plugins/Runtime/RigVM/Source/RigVMDeveloper/Private/RigVMModel/RigVMPin.cpp
--- a/Engine/Plugins/Runtime/RigVM/Source/RigVMDeveloper/Private/RigVMModel/RigVMPin.cpp
+++ b/Engine/Plugins/Runtime/RigVM/Source/RigVMDeveloper/Private/RigVMModel/RigVMPin.cpp
@@ -1128,35 +1128,84 @@ FName URigVMPin::GetCustomWidgetName() const
 #if WITH_EDITOR
 	if(CustomWidgetName.IsNone())
 	{
-		if(const URigVMUnitNode* UnitNode = Cast<URigVMUnitNode>(GetNode()))
+		return FName(GetMetaData(FRigVMStruct::CustomWidgetMetaName));
+	}
+#endif
+	return CustomWidgetName;
+}
+
+FString URigVMPin::GetMetaData(FName InKey) const
+{
+	if (IsArrayElement())
+	{
+		return GetParentPin()->GetMetaData(InKey);
+	}
+
+#if WITH_EDITOR
+	if(const URigVMUnitNode* UnitNode = Cast<URigVMUnitNode>(GetNode()))
+	{
+		if(IsDecoratorPin())
 		{
-			if(const UScriptStruct* Struct = UnitNode->GetScriptStruct())
+			if(const UScriptStruct* Struct = GetDecoratorScriptStruct())
 			{
 				if(const FProperty* Property = Struct->FindPropertyByName(GetFName()))
 				{
-					const FString MetaData = Property->GetMetaData(FRigVMStruct::CustomWidgetMetaName);
+					const FString MetaData = Property->GetMetaData(InKey);
 					if(!MetaData.IsEmpty())
 					{
 						return *MetaData;
 					}
 				}
+				else
+				{
+					// Possible the pin was programmatically generated from the decorator's shared struct
+					TSharedPtr<FStructOnScope> DecoratorScope = GetDecoratorInstance();
+					if(DecoratorScope.IsValid())
+					{
+						const FRigVMDecorator* VMDecorator = (FRigVMDecorator*)DecoratorScope->GetStructMemory();
+						Struct = VMDecorator->GetDecoratorSharedDataStruct();
+						Property = Struct->FindPropertyByName(GetFName());
+						if(Property)
+						{
+							const FString MetaData = Property->GetMetaData(InKey);
+							if(!MetaData.IsEmpty())
+							{
+								return *MetaData;
+							}
+						}
+					}
+				}
 			}
 		}
-		else if(const URigVMTemplateNode* TemplateNode = Cast<URigVMTemplateNode>(GetNode()))
+		else
 		{
-			if(const FRigVMTemplate* Template = TemplateNode->GetTemplate())
+			if(const UScriptStruct* Struct = UnitNode->GetScriptStruct())
 			{
-				const FString MetaData = Template->GetArgumentMetaData(GetFName(), FRigVMStruct::CustomWidgetMetaName);
-				if(!MetaData.IsEmpty())
+				if(const FProperty* Property = Struct->FindPropertyByName(GetFName()))
 				{
-					return *MetaData;
+					const FString MetaData = Property->GetMetaData(InKey);
+					if(!MetaData.IsEmpty())
+					{
+						return *MetaData;
+					}
 				}
 			}
 		}
 	}
+	else if(const URigVMTemplateNode* TemplateNode = Cast<URigVMTemplateNode>(GetNode()))
+	{
+		if(const FRigVMTemplate* Template = TemplateNode->GetTemplate())
+		{
+			const FString MetaData = Template->GetArgumentMetaData(GetFName(), InKey);
+			if(!MetaData.IsEmpty())
+			{
+				return *MetaData;
+			}
+		}
+	}
 #endif
-	
-	return CustomWidgetName;
+
+	return FString();
 }
 
 FText URigVMPin::GetToolTipText() const
@@ -1438,6 +1487,16 @@ TSharedPtr<FStructOnScope> URigVMPin::GetDecoratorInstance(bool bUseDefaultValue
 	return EmptyDecorator;
 }
 
+UScriptStruct* URigVMPin::GetDecoratorScriptStruct() const
+{
+	if(const URigVMNode* Node = GetNode())
+	{
+		return Node->GetDecoratorScriptStruct(GetRootPin());
+	}
+
+	return nullptr;
+}
+
 void URigVMPin::UpdateTypeInformationIfRequired() const
 {
 	if (CPPTypeObject == nullptr)
diff --git a/Engine/Plugins/Runtime/RigVM/Source/RigVMDeveloper/Public/RigVMModel/RigVMNode.h b/Engine/Plugins/Runtime/RigVM/Source/RigVMDeveloper/Public/RigVMModel/RigVMNode.h
--- a/Engine/Plugins/Runtime/RigVM/Source/RigVMDeveloper/Public/RigVMModel/RigVMNode.h
+++ b/Engine/Plugins/Runtime/RigVM/Source/RigVMDeveloper/Public/RigVMModel/RigVMNode.h
@@ -312,6 +312,10 @@ public:
 
 	TSharedPtr<FStructOnScope> GetDecoratorInstance(const URigVMPin* InDecoratorPin, bool bUseDefaultValueFromPin = true) const;
 
+	UScriptStruct* GetDecoratorScriptStruct(const FName& InName) const;
+
+	UScriptStruct* GetDecoratorScriptStruct(const URigVMPin* InDecoratorPin) const;
+	
 	virtual bool IsOutDated() const
 	{
 		return false;
diff --git a/Engine/Plugins/Runtime/RigVM/Source/RigVMDeveloper/Public/RigVMModel/RigVMPin.h b/Engine/Plugins/Runtime/RigVM/Source/RigVMDeveloper/Public/RigVMModel/RigVMPin.h
--- a/Engine/Plugins/Runtime/RigVM/Source/RigVMDeveloper/Public/RigVMModel/RigVMPin.h
+++ b/Engine/Plugins/Runtime/RigVM/Source/RigVMDeveloper/Public/RigVMModel/RigVMPin.h
@@ -305,6 +305,10 @@ public:
 	// Returns the default value clamped with the limit meta values defined by the UPROPERTY in URigVMUnitNodes 
 	FString ClampDefaultValueFromMetaData(const FString& InDefaultValue) const;
 
+	// Returns the keyed metadata associated with this pin, if any
+	UFUNCTION(BlueprintCallable, Category = RigVMPin)
+	FString GetMetaData(FName InKey) const;
+
 	// Returns the name of a custom widget to be used
 	// for editing the Pin.
 	UFUNCTION(BlueprintCallable, Category = RigVMPin)
@@ -469,6 +473,9 @@ public:
 	// Returns the decorator backing up this pin
 	TSharedPtr<FStructOnScope> GetDecoratorInstance(bool bUseDefaultValueFromPin = true) const;
 
+	// Returns the struct of the decorator backing up this pin
+	UScriptStruct* GetDecoratorScriptStruct() const;
+
 private:
 
 	void UpdateTypeInformationIfRequired() const;
