CIS Fix: Non-unity

diff --git a/Engine/Plugins/Experimental/AnimNext/Source/AnimNextEditor/Private/Param/ParameterBlockParameterCustomization.cpp b/Engine/Plugins/Experimental/AnimNext/Source/AnimNextEditor/Private/Param/ParameterBlockParameterCustomization.cpp
--- a/Engine/Plugins/Experimental/AnimNext/Source/AnimNextEditor/Private/Param/ParameterBlockParameterCustomization.cpp
+++ b/Engine/Plugins/Experimental/AnimNext/Source/AnimNextEditor/Private/Param/ParameterBlockParameterCustomization.cpp
@@ -16,6 +16,7 @@
 #include "ISinglePropertyView.h"
 #include "IStructureDataProvider.h"
 #include "UncookedOnlyUtils.h"
+#include "Param/AnimNextParameterBlock_EditorData.h"
 
 #define LOCTEXT_NAMESPACE "ParamTypePropertyCustomization"
 
