CIS fix: Static analysis

diff --git a/Engine/Plugins/Experimental/AnimNext/Source/AnimNext/Private/Graph/AnimNextGraph.cpp b/Engine/Plugins/Experimental/AnimNext/Source/AnimNext/Private/Graph/AnimNextGraph.cpp
--- a/Engine/Plugins/Experimental/AnimNext/Source/AnimNext/Private/Graph/AnimNextGraph.cpp
+++ b/Engine/Plugins/Experimental/AnimNext/Source/AnimNext/Private/Graph/AnimNextGraph.cpp
@@ -47,7 +47,7 @@ void UAnimNextGraph::AllocateInstanceImpl(FAnimNextGraphInstance* ParentGraphIns
 
 	Instance.Release();
 
-	const FName EntryPoint = InEntryPoint == NAME_None ? DefaultEntryPoint : InEntryPoint;
+	const FName EntryPoint = (InEntryPoint == NAME_None) ? DefaultEntryPoint : InEntryPoint;
 	const FAnimNextDecoratorHandle ResolvedRootDecoratorHandle = ResolvedRootDecoratorHandles.FindRef(EntryPoint);
 	if (!ResolvedRootDecoratorHandle.IsValid())
 	{
