AnimNext RefPose now uses a ref counted handle, so if the RefPose is updated, due to chamges to bone visibility or other factors, it will be updated without leaving dangling pointers in the graphs

diff --git a/Engine/Plugins/Experimental/AnimNext/Source/AnimNext/Internal/Graph/AnimNext_LODPose.h b/Engine/Plugins/Experimental/AnimNext/Source/AnimNext/Internal/Graph/AnimNext_LODPose.h
--- a/Engine/Plugins/Experimental/AnimNext/Source/AnimNext/Internal/Graph/AnimNext_LODPose.h
+++ b/Engine/Plugins/Experimental/AnimNext/Source/AnimNext/Internal/Graph/AnimNext_LODPose.h
@@ -14,12 +14,12 @@ struct FAnimNextGraphReferencePose
 
 	FAnimNextGraphReferencePose() = default;
 
-	explicit FAnimNextGraphReferencePose(const UE::AnimNext::FReferencePose* InReferencePose)
+	explicit FAnimNextGraphReferencePose(UE::AnimNext::FDataHandle& InReferencePose)
 		: ReferencePose(InReferencePose)
 	{
 	}
 
-	const UE::AnimNext::FReferencePose* ReferencePose = nullptr;
+	UE::AnimNext::FDataHandle ReferencePose;
 };
 
 USTRUCT(BlueprintType, meta = (DisplayName = "LODPose"))
diff --git a/Engine/Plugins/Experimental/AnimNext/Source/AnimNext/Private/Component/AnimNextMeshComponent.cpp b/Engine/Plugins/Experimental/AnimNext/Source/AnimNext/Private/Component/AnimNextMeshComponent.cpp
--- a/Engine/Plugins/Experimental/AnimNext/Source/AnimNext/Private/Component/AnimNextMeshComponent.cpp
+++ b/Engine/Plugins/Experimental/AnimNext/Source/AnimNext/Private/Component/AnimNextMeshComponent.cpp
@@ -1,4 +1,4 @@
-﻿// Copyright Epic Games, Inc. All Rights Reserved.
+// Copyright Epic Games, Inc. All Rights Reserved.
 
 #include "AnimNextMeshComponent.h"
 #include "GenerationTools.h"
@@ -43,6 +43,5 @@ void UAnimNextMeshComponent::CompleteAndDispatch(TConstArrayView<FBoneIndexType>
 FAnimNextGraphReferencePose UAnimNextMeshComponent::GetReferencePose()
 {
 	UE::AnimNext::FDataHandle RefPoseHandle = UE::AnimNext::FDataRegistry::Get()->GetOrGenerateReferencePose(this);
-	const UE::AnimNext::FReferencePose& RefPose = RefPoseHandle.GetRef<UE::AnimNext::FReferencePose>();
-	return FAnimNextGraphReferencePose(&RefPose);
+	return FAnimNextGraphReferencePose(RefPoseHandle);
 }
\ No newline at end of file
diff --git a/Engine/Plugins/Experimental/AnimNext/Source/AnimNext/Private/Graph/AnimGraph/AnimNode_AnimNextGraph.cpp b/Engine/Plugins/Experimental/AnimNext/Source/AnimNext/Private/Graph/AnimGraph/AnimNode_AnimNextGraph.cpp
--- a/Engine/Plugins/Experimental/AnimNext/Source/AnimNext/Private/Graph/AnimGraph/AnimNode_AnimNextGraph.cpp
+++ b/Engine/Plugins/Experimental/AnimNext/Source/AnimNext/Private/Graph/AnimGraph/AnimNode_AnimNextGraph.cpp
@@ -135,11 +135,11 @@ void FAnimNode_AnimNextGraph::Evaluate_AnyThread(FPoseContext& Output)
 		check(SkeletalMeshComponent != nullptr);
 
 		FDataHandle RefPoseHandle = FDataRegistry::Get()->GetOrGenerateReferencePose(SkeletalMeshComponent);
-		const UE::AnimNext::FReferencePose& RefPose = RefPoseHandle.GetRef<UE::AnimNext::FReferencePose>();
-		FAnimNextGraphReferencePose GraphReferencePose(&RefPose);
+		FAnimNextGraphReferencePose GraphReferencePose(RefPoseHandle);
 
 		const int32 LODLevel = Output.AnimInstanceProxy->GetLODLevel();
 
+		const UE::AnimNext::FReferencePose& RefPose = RefPoseHandle.GetRef<UE::AnimNext::FReferencePose>();
 		FAnimNextGraphLODPose ResultPose(FLODPoseHeap(RefPose, LODLevel, true, Output.ExpectsAdditivePose()));
 
 		FAnimGraphParamStackScope Scope(Output);
diff --git a/Engine/Plugins/Experimental/AnimNext/Source/AnimNext/Private/Param/ParamHelpers.cpp b/Engine/Plugins/Experimental/AnimNext/Source/AnimNext/Private/Param/ParamHelpers.cpp
--- a/Engine/Plugins/Experimental/AnimNext/Source/AnimNext/Private/Param/ParamHelpers.cpp
+++ b/Engine/Plugins/Experimental/AnimNext/Source/AnimNext/Private/Param/ParamHelpers.cpp
@@ -325,7 +325,11 @@ void FParamHelpers::Copy(const FParamTypeHandle& InTypeHandle, TConstArrayView<u
 		SimpleCopy(sizeof(FAnimNextGraphLODPose), alignof(FAnimNextGraphLODPose));
 		break;
 	case FParamTypeHandle::EParamType::AnimNextGraphReferencePose:
-		SimpleCopy(sizeof(FAnimNextGraphReferencePose), alignof(FAnimNextGraphReferencePose));
+		{
+			const FAnimNextGraphReferencePose* SourceRefPose = reinterpret_cast<const FAnimNextGraphReferencePose*>(InSourceMemory.GetData());
+			FAnimNextGraphReferencePose* TargetSourceRefPose = reinterpret_cast<FAnimNextGraphReferencePose*>(InTargetMemory.GetData());
+			*TargetSourceRefPose = *SourceRefPose;
+		}
 		break;
 	case FParamTypeHandle::EParamType::Custom:
 		{
@@ -516,7 +520,12 @@ void FParamHelpers::Destroy(const FParamTypeHandle& InTypeHandle, TArrayView<uin
 	case FParamTypeHandle::EParamType::AnimNextMeshComponent:
 	case FParamTypeHandle::EParamType::AnimSequence:
 	case FParamTypeHandle::EParamType::AnimNextGraphLODPose:
+		break;
 	case FParamTypeHandle::EParamType::AnimNextGraphReferencePose:
+		{
+			FAnimNextGraphReferencePose* RefPose = reinterpret_cast<FAnimNextGraphReferencePose*>(InMemory.GetData());
+			RefPose->~FAnimNextGraphReferencePose();
+		}
 		break;
 	case FParamTypeHandle::EParamType::Custom:
 		{
diff --git a/Engine/Plugins/Experimental/AnimNext/Source/AnimNext/Private/Scheduler/AnimNextScheduleGraphTask.cpp b/Engine/Plugins/Experimental/AnimNext/Source/AnimNext/Private/Scheduler/AnimNextScheduleGraphTask.cpp
--- a/Engine/Plugins/Experimental/AnimNext/Source/AnimNext/Private/Scheduler/AnimNextScheduleGraphTask.cpp
+++ b/Engine/Plugins/Experimental/AnimNext/Source/AnimNext/Private/Scheduler/AnimNextScheduleGraphTask.cpp
@@ -110,7 +110,7 @@ void FAnimNextScheduleGraphTask::RunGraph(const UE::AnimNext::FScheduleContext&
 	}
 
 	const FAnimNextGraphReferencePose* GraphReferencePose = ParamStack.GetParamPtr<FAnimNextGraphReferencePose>(GraphToRun->GetReferencePoseParam());
-	if(GraphReferencePose == nullptr)
+	if(GraphReferencePose == nullptr || !GraphReferencePose->ReferencePose.IsValid())
 	{
 		return;
 	}
@@ -145,12 +145,14 @@ void FAnimNextScheduleGraphTask::RunGraph(const UE::AnimNext::FScheduleContext&
 	{
 		return;
 	}
+	
+	const UE::AnimNext::FReferencePose& RefPose = GraphReferencePose->ReferencePose.GetRef<UE::AnimNext::FReferencePose>();
 
 	// Create or update our result pose
 	// TODO: Currently forcing additive flag to false here
-	if (OutputPose->LODPose.ShouldPrepareForLOD(*GraphReferencePose->ReferencePose, *GraphLODLevel, false))
+	if (OutputPose->LODPose.ShouldPrepareForLOD(RefPose, *GraphLODLevel, false))
 	{
-		OutputPose->LODPose.PrepareForLOD(*GraphReferencePose->ReferencePose, *GraphLODLevel, true, false);
+		OutputPose->LODPose.PrepareForLOD(RefPose, *GraphLODLevel, true, false);
 	}
 
 	check(OutputPose->LODPose.LODLevel == *GraphLODLevel);
@@ -168,7 +170,7 @@ void FAnimNextScheduleGraphTask::RunGraph(const UE::AnimNext::FScheduleContext&
 	{
 		const FEvaluationProgram EvaluationProgram = UE::AnimNext::EvaluateGraph(GraphCache.GraphInstanceData);
 
-		FEvaluationVM EvaluationVM(EEvaluationFlags::All, *GraphReferencePose->ReferencePose, *GraphLODLevel);
+		FEvaluationVM EvaluationVM(EEvaluationFlags::All, RefPose, *GraphLODLevel);
 		bool bHasValidOutput = false;
 
 		if (!EvaluationProgram.IsEmpty())
diff --git a/Engine/Plugins/Experimental/AnimNext/Source/AnimNext/Private/Scheduler/Ports/AnimNextSchedulePort_AnimNextMeshComponentPose.cpp b/Engine/Plugins/Experimental/AnimNext/Source/AnimNext/Private/Scheduler/Ports/AnimNextSchedulePort_AnimNextMeshComponentPose.cpp
--- a/Engine/Plugins/Experimental/AnimNext/Source/AnimNext/Private/Scheduler/Ports/AnimNextSchedulePort_AnimNextMeshComponentPose.cpp
+++ b/Engine/Plugins/Experimental/AnimNext/Source/AnimNext/Private/Scheduler/Ports/AnimNextSchedulePort_AnimNextMeshComponentPose.cpp
@@ -1,4 +1,4 @@
-﻿// Copyright Epic Games, Inc. All Rights Reserved.
+// Copyright Epic Games, Inc. All Rights Reserved.
 
 #include "AnimNextSchedulePort_AnimNextMeshComponentPose.h"
 
@@ -56,6 +56,8 @@ void UAnimNextSchedulePort_AnimNextMeshComponentPose::Run(const UE::AnimNext::FS
 		return;
 	}
 
+	const UE::AnimNext::FReferencePose& RefPose = GraphReferencePose->ReferencePose.GetRef<UE::AnimNext::FReferencePose>();
+
 	FMemMark MemMark(FMemStack::Get());
 
 	TArray<FTransform, TMemStackAllocator<>> LocalSpaceTransforms;
@@ -66,8 +68,8 @@ void UAnimNextSchedulePort_AnimNextMeshComponentPose::Run(const UE::AnimNext::FS
 
 	// Convert and dispatch to renderer
 	Component->CompleteAndDispatch(
-		GraphReferencePose->ReferencePose->GetParentIndices(),
-		GraphReferencePose->ReferencePose->GetLODBoneIndexToMeshBoneIndexMap(InputPose->LODPose.LODLevel),
+		RefPose.GetParentIndices(),
+		RefPose.GetLODBoneIndexToMeshBoneIndexMap(InputPose->LODPose.LODLevel),
 		LocalSpaceTransforms);
 }
 
