Fix mesh space additive retargeting to use the correct rotation fixup

diff --git a/Engine/Plugins/Experimental/AnimNext/Source/AnimNext/Private/DecompressionTools.cpp b/Engine/Plugins/Experimental/AnimNext/Source/AnimNext/Private/DecompressionTools.cpp
--- a/Engine/Plugins/Experimental/AnimNext/Source/AnimNext/Private/DecompressionTools.cpp
+++ b/Engine/Plugins/Experimental/AnimNext/Source/AnimNext/Private/DecompressionTools.cpp
@@ -375,18 +375,11 @@ void FDecompressionTools::DecompressPose(FLODPose& OutAnimationPoseData,
 			// Retarget the root onto the target skeleton (correcting for differences in rest poses)
 			if (SkeletonRemapping.RequiresReferencePoseRetarget())
 			{
-				const int32 TargetSkeletonBoneIndex = 0;
-
-				if (DecompressionContext.IsAdditiveAnimation())
-				{
-					RootAtom.SetRotation(SkeletonRemapping.RetargetAdditiveRotationToTargetSkeleton(TargetSkeletonBoneIndex, RootAtom.GetRotation()));
-					if (TargetSkeleton->GetBoneTranslationRetargetingMode(TargetSkeletonBoneIndex, OutAnimationPoseData.GetDisableRetargeting()) != EBoneTranslationRetargetingMode::Skeleton)
-					{
-						RootAtom.SetTranslation(SkeletonRemapping.RetargetAdditiveTranslationToTargetSkeleton(TargetSkeletonBoneIndex, RootAtom.GetTranslation()));
-					}
-				}
-				else
+				// Root bone does not require fix-up for additive animations as there is no parent delta rotation to account for
+				if (!DecompressionContext.IsAdditiveAnimation())
 				{
+					const int32 TargetSkeletonBoneIndex = 0;
+
 					RootAtom.SetRotation(SkeletonRemapping.RetargetBoneRotationToTargetSkeleton(TargetSkeletonBoneIndex, RootAtom.GetRotation()));
 					if (TargetSkeleton->GetBoneTranslationRetargetingMode(TargetSkeletonBoneIndex, OutAnimationPoseData.GetDisableRetargeting()) != EBoneTranslationRetargetingMode::Skeleton)
 					{
@@ -430,10 +423,16 @@ void FDecompressionTools::DecompressPose(FLODPose& OutAnimationPoseData,
 
 		if (DecompressionContext.IsAdditiveAnimation())
 		{
-			for (int32 LODBoneIndex = bFirstTrackIsRootBone ? 1 : 0; LODBoneIndex < LODNumBones; ++LODBoneIndex)
+			for (int32 LODBoneIndex = (bFirstTrackIsRootBone ? 1 : 0); LODBoneIndex < LODNumBones; ++LODBoneIndex)
 			{
 				const int32 TargetSkeletonBoneIndex = ReferencePose.GetSkeletonBoneIndexFromLODBoneIndex(LODBoneIndex);
-				OutAnimationPoseData.LocalTransformsView[LODBoneIndex].SetRotation(SkeletonRemapping.RetargetAdditiveRotationToTargetSkeleton(TargetSkeletonBoneIndex, OutAnimationPoseData.LocalTransformsView[LODBoneIndex].GetRotation()));
+
+				// Mesh space additives do not require fix-up
+				if (DecompressionContext.GetAdditiveType() == AAT_LocalSpaceBase)
+				{
+					OutAnimationPoseData.LocalTransformsView[LODBoneIndex].SetRotation(SkeletonRemapping.RetargetAdditiveRotationToTargetSkeleton(TargetSkeletonBoneIndex, OutAnimationPoseData.LocalTransformsView[LODBoneIndex].GetRotation()));
+				}
+
 				if (TargetSkeleton->GetBoneTranslationRetargetingMode(TargetSkeletonBoneIndex, OutAnimationPoseData.GetDisableRetargeting()) != EBoneTranslationRetargetingMode::Skeleton)
 				{
 					OutAnimationPoseData.LocalTransformsView[LODBoneIndex].SetTranslation(SkeletonRemapping.RetargetAdditiveTranslationToTargetSkeleton(TargetSkeletonBoneIndex, OutAnimationPoseData.LocalTransformsView[LODBoneIndex].GetTranslation()));
diff --git a/Engine/Source/Runtime/Engine/Private/Animation/AnimationDecompression.cpp b/Engine/Source/Runtime/Engine/Private/Animation/AnimationDecompression.cpp
--- a/Engine/Source/Runtime/Engine/Private/Animation/AnimationDecompression.cpp
+++ b/Engine/Source/Runtime/Engine/Private/Animation/AnimationDecompression.cpp
@@ -156,18 +156,11 @@ void DecompressPose(FCompactPose& OutPose,
 			// Retarget the root onto the target skeleton (correcting for differences in rest poses)
 			if (SkeletonRemapping.RequiresReferencePoseRetarget())
 			{
-				const int32 TargetSkeletonBoneIndex = 0;
-
-				if (DecompressionContext.IsAdditiveAnimation())
-				{
-					RootAtom.SetRotation(SkeletonRemapping.RetargetAdditiveRotationToTargetSkeleton(TargetSkeletonBoneIndex, RootAtom.GetRotation()));
-					if (TargetSkeleton->GetBoneTranslationRetargetingMode(TargetSkeletonBoneIndex, RequiredBones.GetDisableRetargeting()) != EBoneTranslationRetargetingMode::Skeleton)
-					{
-						RootAtom.SetTranslation(SkeletonRemapping.RetargetAdditiveTranslationToTargetSkeleton(TargetSkeletonBoneIndex, RootAtom.GetTranslation()));
-					}
-				}
-				else
+				// Root bone does not require fix-up for additive animations as there is no parent delta rotation to account for
+				if (!DecompressionContext.IsAdditiveAnimation())
 				{
+					const int32 TargetSkeletonBoneIndex = 0;
+
 					RootAtom.SetRotation(SkeletonRemapping.RetargetBoneRotationToTargetSkeleton(TargetSkeletonBoneIndex, RootAtom.GetRotation()));
 					if (TargetSkeleton->GetBoneTranslationRetargetingMode(TargetSkeletonBoneIndex, RequiredBones.GetDisableRetargeting()) != EBoneTranslationRetargetingMode::Skeleton)
 					{
@@ -196,7 +189,13 @@ void DecompressPose(FCompactPose& OutPose,
 			for (FCompactPoseBoneIndex BoneIndex(bFirstTrackIsRootBone ? 1 : 0); BoneIndex < NumCompactBones; ++BoneIndex)
 			{
 				const int32 TargetSkeletonBoneIndex = RequiredBones.GetSkeletonIndex(BoneIndex);
-				OutPose[BoneIndex].SetRotation(SkeletonRemapping.RetargetAdditiveRotationToTargetSkeleton(TargetSkeletonBoneIndex, OutPose[BoneIndex].GetRotation()));
+
+				// Mesh space additives do not require fix-up
+				if (DecompressionContext.GetAdditiveType() == AAT_LocalSpaceBase)
+				{
+					OutPose[BoneIndex].SetRotation(SkeletonRemapping.RetargetAdditiveRotationToTargetSkeleton(TargetSkeletonBoneIndex, OutPose[BoneIndex].GetRotation()));
+				}
+
 				if (TargetSkeleton->GetBoneTranslationRetargetingMode(TargetSkeletonBoneIndex, RequiredBones.GetDisableRetargeting()) != EBoneTranslationRetargetingMode::Skeleton)
 				{
 					OutPose[BoneIndex].SetTranslation(SkeletonRemapping.RetargetAdditiveTranslationToTargetSkeleton(TargetSkeletonBoneIndex, OutPose[BoneIndex].GetTranslation()));
diff --git a/Engine/Source/Runtime/Engine/Private/Animation/SkeletonRemapping.cpp b/Engine/Source/Runtime/Engine/Private/Animation/SkeletonRemapping.cpp
--- a/Engine/Source/Runtime/Engine/Private/Animation/SkeletonRemapping.cpp
+++ b/Engine/Source/Runtime/Engine/Private/Animation/SkeletonRemapping.cpp
@@ -120,7 +120,7 @@ void FSkeletonRemapping::GenerateMapping()
 	// from local space (D) to component space (C), given its parent (P), we have:
 	//
 	//		C * P = P * D
-	//		C = P * D * P⁻¹
+	//		C = P * D * P⁻¹		(this uses the sandwich product to rotate the rotation axis from local to mesh space)
 	//
 	// Setting the source and target component space rotations to be equal (Cs = Ct) then gives us:
 	//
@@ -144,14 +144,17 @@ void FSkeletonRemapping::GenerateMapping()
 	//
 	//		Lt = Pt⁻¹ * Ps * Ls * Rs⁻¹ * Ps⁻¹ * Pt * Rt
 	//
-	// Finally, factoring out the constant terms (which we precompute here) gives us:
+	// Finally, factoring out the constant terms (which we pre-compute here) gives us:
 	//
 	//		Lt = Q0 * Ls * Q1
 	//		Q0 = Pt⁻¹ * Ps
 	//		Q1 = Rs⁻¹ * Ps⁻¹ * Pt * Rt
 	//
 	// Note that when remapping additive animations, we drop the rest pose terms, but we still need to convert between the
-	// source and target rotation frames.  Dropping Rs and Rt from the equations above gives us:
+	// source and target rotation frames. The terms drop because an additive rotation needs to be applied to a base one to
+	// become a local space rotation. To that end, we can use any base as we are interested in the delta between source/target.
+	// Using the bind pose as our base quickly cancels out the terms.
+	// Dropping Rs and Rt from the equations above gives us:
 	//
 	//		Lt = Pt⁻¹ * Ps * Ls * Ps⁻¹ * Pt			(for additive animations)
 	//
@@ -159,6 +162,9 @@ void FSkeletonRemapping::GenerateMapping()
 	//
 	//		Lt = Q0 * Ls * Q0⁻¹						(for additive animations)
 	//
+	// For mesh space additive animations, the source rotations already represent a mesh space delta and as such no fix-up needs
+	// to applied for those.
+
 	RetargetingTable[0] = MakeTuple(FQuat::Identity, SourceLocalTransforms[0].GetRotation().Inverse() * TargetLocalTransforms[0].GetRotation());
 	for (int32 TargetBoneIndex = 1; TargetBoneIndex < TargetNumBones; ++TargetBoneIndex)
 	{
