Added EAllowShrinking support to TTransformArrayAoS and TTransformArraySoA.

diff --git a/Engine/Plugins/Experimental/AnimNext/Source/AnimNext/Public/TransformArray.h b/Engine/Plugins/Experimental/AnimNext/Source/AnimNext/Public/TransformArray.h
--- a/Engine/Plugins/Experimental/AnimNext/Source/AnimNext/Public/TransformArray.h
+++ b/Engine/Plugins/Experimental/AnimNext/Source/AnimNext/Public/TransformArray.h
@@ -39,14 +39,24 @@ struct TTransformArrayAoS
 		Transforms.Reset(NumTransforms);
 	}
 
-	void SetNum(int32 NumTransforms, bool bAllowShrinking = true)
+	void SetNum(int32 NumTransforms, EAllowShrinking AllowShrinking = EAllowShrinking::Yes)
 	{
-		Transforms.SetNum(NumTransforms, bAllowShrinking);
+		Transforms.SetNum(NumTransforms, AllowShrinking);
+	}
+	UE_ALLOWSHRINKING_BOOL_DEPRECATED("SetNum")
+	FORCEINLINE void SetNum(int32 NumTransforms, bool bAllowShrinking)
+	{
+		SetNum(NumTransforms, bAllowShrinking ? EAllowShrinking::Yes : EAllowShrinking::No);
 	}
 
-	void SetNumUninitialized(int32 NumTransforms, bool bAllowShrinking = true)
+	void SetNumUninitialized(int32 NumTransforms, EAllowShrinking AllowShrinking = EAllowShrinking::Yes)
+	{
+		Transforms.SetNumUninitialized(NumTransforms, AllowShrinking);
+	}
+	UE_ALLOWSHRINKING_BOOL_DEPRECATED("SetNumUninitialized")
+	FORCEINLINE void SetNumUninitialized(int32 NumTransforms, bool bAllowShrinking)
 	{
-		Transforms.SetNumUninitialized(NumTransforms, bAllowShrinking);
+		SetNumUninitialized(NumTransforms, bAllowShrinking ? EAllowShrinking::Yes : EAllowShrinking::No);
 	}
 
 	inline void SetIdentity(bool bAdditiveIdentity = false)
@@ -314,21 +324,31 @@ struct TTransformArraySoA
 		AllocatedMemory.Reset(NumTransforms * TransformSize);
 	}
 
-	void SetNum(int32 NumTransforms, bool bAllowShrinking = true)
+	void SetNum(int32 NumTransforms, EAllowShrinking AllowShrinking = EAllowShrinking::Yes)
 	{
 		constexpr int32 TransformSize = sizeof(FVector) + sizeof(FQuat) + sizeof(FVector);
-		AllocatedMemory.SetNum(NumTransforms * TransformSize, bAllowShrinking);
+		AllocatedMemory.SetNum(NumTransforms * TransformSize, AllowShrinking);
 
 		UpdateViews(AllocatedMemory.GetData(), NumTransforms);
 	}
+	UE_ALLOWSHRINKING_BOOL_DEPRECATED("SetNum")
+	FORCEINLINE void SetNum(int32 NumTransforms, bool bAllowShrinking)
+	{
+		SetNum(NumTransforms, bAllowShrinking ? EAllowShrinking::Yes : EAllowShrinking::No);
+	}
 
-	void SetNumUninitialized(int32 NumTransforms, bool bAllowShrinking = true)
+	void SetNumUninitialized(int32 NumTransforms, EAllowShrinking AllowShrinking = EAllowShrinking::Yes)
 	{
 		constexpr int32 TransformSize = sizeof(FVector) + sizeof(FQuat) + sizeof(FVector);
-		AllocatedMemory.SetNumUninitialized(NumTransforms * TransformSize, bAllowShrinking);
+		AllocatedMemory.SetNumUninitialized(NumTransforms * TransformSize, AllowShrinking);
 
 		UpdateViews(AllocatedMemory.GetData(), NumTransforms);
 	}
+	UE_ALLOWSHRINKING_BOOL_DEPRECATED("SetNumUninitialized")
+	FORCEINLINE void SetNumUninitialized(int32 NumTransforms, bool bAllowShrinking)
+	{
+		SetNumUninitialized(NumTransforms, bAllowShrinking ? EAllowShrinking::Yes : EAllowShrinking::No);
+	}
 
 	inline void SetIdentity(bool bAdditiveIdentity = false)
 	{
