Migrate away from the 'decorator' term in AnimNext in favor of 'trait'

#changelist validated
#virtualized

diff --git a/Engine/Plugins/Experimental/AnimNext/Config/BaseAnimNext.ini b/Engine/Plugins/Experimental/AnimNext/Config/BaseAnimNext.ini
--- a/Engine/Plugins/Experimental/AnimNext/Config/BaseAnimNext.ini
+++ b/Engine/Plugins/Experimental/AnimNext/Config/BaseAnimNext.ini
@@ -2,3 +2,16 @@
 +ExposedClasses=(Class="/Script/Engine.World",AccessorName="UE_World")
 +ExposedClasses=(Class="/Script/Engine.CharacterMovementComponent",AccessorName="UE_CharacterMovementComponent")
 +ExposedClasses=(Class="/Script/AnimNext.AnimNextMeshComponent",AccessorName="UE_AnimNextMeshComponent")
+
+[CoreRedirects]
++StructRedirects=(OldName="/Script/AnimNext.AnimNextDecoratorHandle", NewName="/Script/AnimNext.AnimNextTraitHandle")
++StructRedirects=(OldName="/Script/AnimNext.RigUnit_AnimNextDecoratorStack", NewName="/Script/AnimNext.RigUnit_AnimNextTraitStack")
++StructRedirects=(OldName="/Script/AnimNext.AnimNextBlendByBoolDecoratorSharedData", NewName="/Script/AnimNext.AnimNextBlendByBoolTraitSharedData")
++StructRedirects=(OldName="/Script/AnimNext.AnimNextBlendInertializerDecoratorSharedData", NewName="/Script/AnimNext.AnimNextBlendInertializerTraitSharedData")
++StructRedirects=(OldName="/Script/AnimNext.AnimNextBlendSmootherDecoratorSharedData", NewName="/Script/AnimNext.AnimNextBlendSmootherTraitSharedData")
++StructRedirects=(OldName="/Script/AnimNext.AnimNextBlendSmootherPerBoneDecoratorSharedData", NewName="/Script/AnimNext.AnimNextBlendSmootherPerBoneTraitSharedData")
++StructRedirects=(OldName="/Script/AnimNext.AnimNextBlendTwoWayDecoratorSharedData", NewName="/Script/AnimNext.AnimNextBlendTwoWayTraitSharedData")
++StructRedirects=(OldName="/Script/AnimNext.AnimNextReferencePoseDecoratorSharedData", NewName="/Script/AnimNext.AnimNextReferencePoseTraitSharedData")
++StructRedirects=(OldName="/Script/AnimNext.AnimNextSequencePlayerDecoratorSharedData", NewName="/Script/AnimNext.AnimNextSequencePlayerTraitSharedData")
++StructRedirects=(OldName="/Script/AnimNext.AnimNextSubGraphHostDecoratorSharedData", NewName="/Script/AnimNext.AnimNextSubGraphHostTraitSharedData")
++StructRedirects=(OldName="/Script/AnimNext.AnimNextSynchronizeUsingGroupsDecoratorSharedData", NewName="/Script/AnimNext.AnimNextSynchronizeUsingGroupsTraitSharedData")
diff --git a/Engine/Plugins/Experimental/AnimNext/Source/AnimNext/Internal/Graph/AnimNextGraph.h b/Engine/Plugins/Experimental/AnimNext/Source/AnimNext/Internal/Graph/AnimNextGraph.h
--- a/Engine/Plugins/Experimental/AnimNext/Source/AnimNext/Internal/Graph/AnimNextGraph.h
+++ b/Engine/Plugins/Experimental/AnimNext/Source/AnimNext/Internal/Graph/AnimNextGraph.h
@@ -7,9 +7,9 @@
 #include "AnimNextRigVMAsset.h"
 #include "RigUnit_AnimNextGraphRoot.h"
 #include "RigVMCore/RigVM.h"
-#include "DecoratorBase/DecoratorPtr.h"
-#include "DecoratorBase/DecoratorHandle.h"
-#include "DecoratorBase/EntryPointHandle.h"
+#include "TraitCore/TraitPtr.h"
+#include "TraitCore/TraitHandle.h"
+#include "TraitCore/EntryPointHandle.h"
 #include "Graph/GraphInstanceComponent.h"
 #include "Graph/RigUnit_AnimNextGraphEvaluator.h"
 #include "Param/ParamId.h"
@@ -134,10 +134,10 @@ protected:
 	UPROPERTY()
 	TArray<FAnimNextGraphEntryPoint> EntryPoints;
 
-	// This is a resolved handle to the root decorator in our graph, for each entry point 
-	TMap<FName, FAnimNextDecoratorHandle> ResolvedRootDecoratorHandles;
+	// This is a resolved handle to the root trait in our graph, for each entry point 
+	TMap<FName, FAnimNextTraitHandle> ResolvedRootTraitHandles;
 
-	// This is the graph shared data used by the decorator system, the output of FDecoratorReader
+	// This is the graph shared data used by the trait system, the output of FTraitReader
 	// We de-serialize manually into this buffer from the archive buffer, this is never saved on disk
 	TArray<uint8> SharedDataBuffer;
 
@@ -172,7 +172,7 @@ protected:
 	TArray<FAnimNextParam> RequiredParameters;
 
 #if WITH_EDITORONLY_DATA
-	// This buffer holds the output of the FDecoratorWriter post compilation
+	// This buffer holds the output of the FTraitWriter post compilation
 	// We serialize it manually and it is discarded at runtime
 	TArray<uint8> SharedDataArchiveBuffer;
 #endif
diff --git a/Engine/Plugins/Experimental/AnimNext/Source/AnimNext/Internal/Graph/AnimNextGraphEntryPoint.h b/Engine/Plugins/Experimental/AnimNext/Source/AnimNext/Internal/Graph/AnimNextGraphEntryPoint.h
--- a/Engine/Plugins/Experimental/AnimNext/Source/AnimNext/Internal/Graph/AnimNextGraphEntryPoint.h
+++ b/Engine/Plugins/Experimental/AnimNext/Source/AnimNext/Internal/Graph/AnimNextGraphEntryPoint.h
@@ -1,10 +1,11 @@
-﻿// Copyright Epic Games, Inc. All Rights Reserved.
+// Copyright Epic Games, Inc. All Rights Reserved.
 
 #pragma once
 
 #include "CoreMinimal.h"
-#include "RigUnit_AnimNextGraphEvaluator.h"
-#include "DecoratorBase/EntryPointHandle.h"
+
+#include "TraitCore/EntryPointHandle.h"
+
 #include "AnimNextGraphEntryPoint.generated.h"
 
 USTRUCT()
@@ -16,8 +17,8 @@ struct FAnimNextGraphEntryPoint
 	UPROPERTY()
 	FName EntryPointName;
 
-	// This is a handle to the root decorator for a graph
+	// This is a handle to the root trait for a graph
 	UPROPERTY()
-	FAnimNextEntryPointHandle RootDecoratorHandle;
+	FAnimNextEntryPointHandle RootTraitHandle;
 };
 
diff --git a/Engine/Plugins/Experimental/AnimNext/Source/AnimNext/Internal/Graph/RigDecorator_AnimNextCppTrait.h b/Engine/Plugins/Experimental/AnimNext/Source/AnimNext/Internal/Graph/RigDecorator_AnimNextCppTrait.h
new file mode 100644
--- /dev/null
+++ b/Engine/Plugins/Experimental/AnimNext/Source/AnimNext/Internal/Graph/RigDecorator_AnimNextCppTrait.h
@@ -0,0 +1,40 @@
+// Copyright Epic Games, Inc. All Rights Reserved.
+
+#pragma once
+
+#include "RigVMCore/RigVMDecorator.h"
+
+#include "RigDecorator_AnimNextCppTrait.generated.h"
+
+namespace UE::AnimNext { struct FTrait; }
+
+/**
+ * AnimNext RigDecorator for all C++ traits.
+ * The trait shared data UScriptStruct determines which properties are exposed.
+ */
+USTRUCT(BlueprintType)
+struct ANIMNEXT_API FRigDecorator_AnimNextCppDecorator : public FRigVMDecorator
+{
+	GENERATED_BODY()
+
+	// The struct the trait exposes with its shared data. Each one of its properties will be added as a pin.
+	UPROPERTY(meta = (Hidden))
+	UScriptStruct* DecoratorSharedDataStruct = nullptr;
+
+#if WITH_EDITOR
+	virtual void GetProgrammaticPins(URigVMController* InController, int32 InParentPinIndex, const FString& InDefaultValue, struct FRigVMPinInfoArray& OutPinArray) const override;
+	
+	virtual UScriptStruct* GetDecoratorSharedDataStruct() const override { return DecoratorSharedDataStruct; }
+
+	const UE::AnimNext::FTrait* GetTrait() const;
+#endif
+};
+
+USTRUCT(BlueprintType)
+struct FAnimNextCppDecoratorWrapper
+{
+	GENERATED_BODY()
+
+	UPROPERTY()
+	FRigDecorator_AnimNextCppDecorator CppDecorator;
+};
diff --git a/Engine/Plugins/Experimental/AnimNext/Source/AnimNext/Internal/Graph/RigUnit_AnimNextGraphRoot.h b/Engine/Plugins/Experimental/AnimNext/Source/AnimNext/Internal/Graph/RigUnit_AnimNextGraphRoot.h
--- a/Engine/Plugins/Experimental/AnimNext/Source/AnimNext/Internal/Graph/RigUnit_AnimNextGraphRoot.h
+++ b/Engine/Plugins/Experimental/AnimNext/Source/AnimNext/Internal/Graph/RigUnit_AnimNextGraphRoot.h
@@ -3,7 +3,7 @@
 #pragma once
 
 #include "AnimNextExecuteContext.h"
-#include "DecoratorBase/DecoratorHandle.h"
+#include "TraitCore/TraitHandle.h"
 
 #include "RigUnit_AnimNextGraphRoot.generated.h"
 
@@ -28,7 +28,7 @@ struct ANIMNEXT_API FRigUnit_AnimNextGraphRoot : public FRigUnit_AnimNextBase
 
 	// The execution result
 	UPROPERTY(EditAnywhere, Category = Result, meta = (Input))
-	FAnimNextDecoratorHandle Result;
+	FAnimNextTraitHandle Result;
 
 	// In order for this node to be considered an executable RigUnit, it needs a pin to derive from FRigVMExecuteContext
 	// We keep it hidden it since we don't need it
diff --git a/Engine/Plugins/Experimental/AnimNext/Source/AnimNext/Internal/TraitCore/NodeTemplateBuilder.h b/Engine/Plugins/Experimental/AnimNext/Source/AnimNext/Internal/TraitCore/NodeTemplateBuilder.h
new file mode 100644
--- /dev/null
+++ b/Engine/Plugins/Experimental/AnimNext/Source/AnimNext/Internal/TraitCore/NodeTemplateBuilder.h
@@ -0,0 +1,43 @@
+// Copyright Epic Games, Inc. All Rights Reserved.
+
+#pragma once
+
+#include "CoreMinimal.h"
+#include "TraitCore/TraitUID.h"
+
+namespace UE::AnimNext
+{
+	struct FNodeTemplate;
+
+	/**
+	  * FNodeTemplateBuilder
+	  * 
+	  * Utility to help construct node templates.
+	  */
+	struct ANIMNEXT_API FNodeTemplateBuilder
+	{
+		FNodeTemplateBuilder() = default;
+
+		// Adds the specified trait type to the node template
+		void AddTrait(FTraitUID TraitUID);
+
+		// Returns a node template for the provided list of traits
+		// The node template will be built into the provided buffer and a pointer to it is returned
+		FNodeTemplate* BuildNodeTemplate(TArray<uint8>& NodeTemplateBuffer) const;
+
+		// Returns a node template for the provided list of traits
+		// The node template will be built into the provided buffer and a pointer to it is returned
+		static FNodeTemplate* BuildNodeTemplate(const TArray<FTraitUID>& InTraitUIDs, TArray<uint8>& NodeTemplateBuffer);
+
+		// Resets the node template builder
+		void Reset();
+
+	private:
+		static uint32 GetNodeTemplateUID(const TArray<FTraitUID>& InTraitUIDs);
+		static void AppendTemplateTrait(
+			const TArray<FTraitUID>& InTraitUIDs, int32 TraitIndex,
+			TArray<uint8>& NodeTemplateBuffer);
+
+		TArray<FTraitUID> TraitUIDs;	// The list of traits to use when building the node template
+	};
+}
diff --git a/Engine/Plugins/Experimental/AnimNext/Source/AnimNext/Private/DataRegistry.cpp b/Engine/Plugins/Experimental/AnimNext/Source/AnimNext/Private/DataRegistry.cpp
--- a/Engine/Plugins/Experimental/AnimNext/Source/AnimNext/Private/DataRegistry.cpp
+++ b/Engine/Plugins/Experimental/AnimNext/Source/AnimNext/Private/DataRegistry.cpp
@@ -166,7 +166,7 @@ FDataHandle FDataRegistry::GetRegisteredData(const FName& Id) const
 
 void FDataRegistry::FreeAllocatedBlock(Private::FAllocatedBlock* AllocatedBlock)
 {
-	FRWScopeLock Lock(DataTypeDefsLock, SLT_Write);
+	FRWScopeLock ScopedAllocatedBlocksLock(AllocatedBlocksLock, SLT_Write);
 
 	if (ensure(AllocatedBlock != nullptr && AllocatedBlocks.Find(AllocatedBlock)))
 	{
@@ -174,15 +174,18 @@ void FDataRegistry::FreeAllocatedBlock(Private::FAllocatedBlock* AllocatedBlock)
 		{
 			void* Memory = AllocatedBlock->Memory;
 
-			FDataTypeDef* TypeDef = DataTypeDefs.Find(AllocatedBlock->TypeHandle);
-			if (ensure(TypeDef != nullptr))
 			{
-				TypeDef->DestroyTypeFn((uint8*)AllocatedBlock->Memory, AllocatedBlock->NumElem);
-
-				FMemory::Free(AllocatedBlock->Memory); // TODO : This should come from preallocated chunks, use malloc / free for now
-				AllocatedBlock->Memory = nullptr;
+				FRWScopeLock ScopedDataTypeDefsLock(DataTypeDefsLock, SLT_ReadOnly);
+				FDataTypeDef* TypeDef = DataTypeDefs.Find(AllocatedBlock->TypeHandle);
+				if (ensure(TypeDef != nullptr))
+				{
+					TypeDef->DestroyTypeFn((uint8*)AllocatedBlock->Memory, AllocatedBlock->NumElem);
+				}
 			}
 
+			FMemory::Free(AllocatedBlock->Memory); // TODO : This should come from preallocated chunks, use malloc / free for now
+			AllocatedBlock->Memory = nullptr;
+
 			AllocatedBlocks.Remove(AllocatedBlock);
 			delete AllocatedBlock; // TODO : avoid memory fragmentation
 		}
diff --git a/Engine/Plugins/Experimental/AnimNext/Source/AnimNext/Private/EvaluationVM/EvaluationProgram.cpp b/Engine/Plugins/Experimental/AnimNext/Source/AnimNext/Private/EvaluationVM/EvaluationProgram.cpp
--- a/Engine/Plugins/Experimental/AnimNext/Source/AnimNext/Private/EvaluationVM/EvaluationProgram.cpp
+++ b/Engine/Plugins/Experimental/AnimNext/Source/AnimNext/Private/EvaluationVM/EvaluationProgram.cpp
@@ -46,7 +46,7 @@ namespace UE::AnimNext
 				{
 					PropertyValue.Reset();
 
-					const bool bIsDefaultValue = !Property->ExportText_InContainer(0, PropertyValue, Task.Get(), nullptr, nullptr, PPF_None);
+					const bool bIsDefaultValue = !Property->ExportText_InContainer(0, PropertyValue, Task.Get(), nullptr, nullptr, PPF_SerializedAsImportText);
 
 					if (bIsDefaultValue)
 					{
diff --git a/Engine/Plugins/Experimental/AnimNext/Source/AnimNext/Private/Graph/AnimGraph/AnimNode_AnimNextGraph.cpp b/Engine/Plugins/Experimental/AnimNext/Source/AnimNext/Private/Graph/AnimGraph/AnimNode_AnimNextGraph.cpp
--- a/Engine/Plugins/Experimental/AnimNext/Source/AnimNext/Private/Graph/AnimGraph/AnimNode_AnimNextGraph.cpp
+++ b/Engine/Plugins/Experimental/AnimNext/Source/AnimNext/Private/Graph/AnimGraph/AnimNode_AnimNextGraph.cpp
@@ -17,8 +17,8 @@
 #include "Param/ParamStack.h"
 #include "AnimGraphParamStackScope.h"
 #include "Scheduler/ScheduleContext.h"
-#include "DecoratorInterfaces/IEvaluate.h"
-#include "DecoratorInterfaces/IUpdate.h"
+#include "TraitInterfaces/IEvaluate.h"
+#include "TraitInterfaces/IUpdate.h"
 #include "EvaluationVM/EvaluationVM.h"
 
 #include UE_INLINE_GENERATED_CPP_BY_NAME(AnimNode_AnimNextGraph)
diff --git a/Engine/Plugins/Experimental/AnimNext/Source/AnimNext/Private/Graph/AnimNextGraph.cpp b/Engine/Plugins/Experimental/AnimNext/Source/AnimNext/Private/Graph/AnimNextGraph.cpp
--- a/Engine/Plugins/Experimental/AnimNext/Source/AnimNext/Private/Graph/AnimNextGraph.cpp
+++ b/Engine/Plugins/Experimental/AnimNext/Source/AnimNext/Private/Graph/AnimNextGraph.cpp
@@ -8,8 +8,8 @@
 #include "UObject/ObjectSaveContext.h"
 #include "UObject/Package.h"
 #include "Param/ParamTypeHandle.h"
-#include "DecoratorBase/DecoratorReader.h"
-#include "DecoratorBase/ExecutionContext.h"
+#include "TraitCore/TraitReader.h"
+#include "TraitCore/ExecutionContext.h"
 #include "Graph/AnimNext_LODPose.h"
 #include "Graph/GC_GraphInstanceComponent.h"
 #include "Graph/AnimNextGraphInstance.h"
@@ -48,8 +48,8 @@ void UAnimNextGraph::AllocateInstanceImpl(FAnimNextGraphInstance* ParentGraphIns
 	Instance.Release();
 
 	const FName EntryPoint = (InEntryPoint == NAME_None) ? DefaultEntryPoint : InEntryPoint;
-	const FAnimNextDecoratorHandle ResolvedRootDecoratorHandle = ResolvedRootDecoratorHandles.FindRef(EntryPoint);
-	if (!ResolvedRootDecoratorHandle.IsValid())
+	const FAnimNextTraitHandle ResolvedRootTraitHandle = ResolvedRootTraitHandles.FindRef(EntryPoint);
+	if (!ResolvedRootTraitHandle.IsValid())
 	{
 		return;
 	}
@@ -73,7 +73,7 @@ void UAnimNextGraph::AllocateInstanceImpl(FAnimNextGraphInstance* ParentGraphIns
 
 	{
 		UE::AnimNext::FExecutionContext Context(Instance);
-		Instance.Impl->GraphInstancePtr = Context.AllocateNodeInstance(UE::AnimNext::FWeakDecoratorPtr(), ResolvedRootDecoratorHandle);
+		Instance.Impl->GraphInstancePtr = Context.AllocateNodeInstance(UE::AnimNext::FWeakTraitPtr(), ResolvedRootTraitHandle);
 	}
 
 	if (!Instance.IsValid())
@@ -169,14 +169,14 @@ bool UAnimNextGraph::LoadFromArchiveBuffer(const TArray<uint8>& InSharedDataArch
 
 	// Reconstruct our graph shared data
 	FMemoryReader GraphSharedDataArchive(InSharedDataArchiveBuffer);
-	FDecoratorReader DecoratorReader(GraphReferencedObjects, GraphSharedDataArchive);
+	FTraitReader TraitReader(GraphReferencedObjects, GraphSharedDataArchive);
 
-	const FDecoratorReader::EErrorState ErrorState = DecoratorReader.ReadGraph(SharedDataBuffer);
-	if (ErrorState == FDecoratorReader::EErrorState::None)
+	const FTraitReader::EErrorState ErrorState = TraitReader.ReadGraph(SharedDataBuffer);
+	if (ErrorState == FTraitReader::EErrorState::None)
 	{
 		for(FAnimNextGraphEntryPoint& EntryPoint : EntryPoints)
 		{
-			ResolvedRootDecoratorHandles.Add(EntryPoint.EntryPointName, DecoratorReader.ResolveEntryPointHandle(EntryPoint.RootDecoratorHandle));
+			ResolvedRootTraitHandles.Add(EntryPoint.EntryPointName, TraitReader.ResolveEntryPointHandle(EntryPoint.RootTraitHandle));
 		}
 
 		// Make sure our execute method is registered
@@ -186,7 +186,7 @@ bool UAnimNextGraph::LoadFromArchiveBuffer(const TArray<uint8>& InSharedDataArch
 	else
 	{
 		SharedDataBuffer.Empty(0);
-		ResolvedRootDecoratorHandles.Add(FRigUnit_AnimNextGraphRoot::DefaultEntryPoint, FAnimNextDecoratorHandle());
+		ResolvedRootTraitHandles.Add(FRigUnit_AnimNextGraphRoot::DefaultEntryPoint, FAnimNextTraitHandle());
 		return false;
 	}
 }
diff --git a/Engine/Plugins/Experimental/AnimNext/Source/AnimNext/Private/Graph/AnimNextGraphInstance.cpp b/Engine/Plugins/Experimental/AnimNext/Source/AnimNext/Private/Graph/AnimNextGraphInstance.cpp
--- a/Engine/Plugins/Experimental/AnimNext/Source/AnimNext/Private/Graph/AnimNextGraphInstance.cpp
+++ b/Engine/Plugins/Experimental/AnimNext/Source/AnimNext/Private/Graph/AnimNextGraphInstance.cpp
@@ -3,7 +3,7 @@
 #include "Graph/AnimNextGraphInstance.h"
 
 #include "AnimNextStats.h"
-#include "DecoratorBase/ExecutionContext.h"
+#include "TraitCore/ExecutionContext.h"
 #include "Graph/AnimNextGraph.h"
 #include "Graph/GC_GraphInstanceComponent.h"
 #include "Graph/RigUnit_AnimNextShimRoot.h"
@@ -45,7 +45,7 @@ FName FAnimNextGraphInstance::GetEntryPoint() const
 	return EntryPoint;
 }
 
-UE::AnimNext::FWeakDecoratorPtr FAnimNextGraphInstance::GetGraphRootPtr() const
+UE::AnimNext::FWeakTraitPtr FAnimNextGraphInstance::GetGraphRootPtr() const
 {
 	return GraphInstancePtr;
 }
@@ -159,9 +159,9 @@ void FAnimNextGraphInstance::Thaw()
 
 		{
 			UE::AnimNext::FExecutionContext Context(*this);
-			if(const FAnimNextDecoratorHandle* FoundHandle = GraphPtr->ResolvedRootDecoratorHandles.Find(EntryPoint))
+			if(const FAnimNextTraitHandle* FoundHandle = GraphPtr->ResolvedRootTraitHandles.Find(EntryPoint))
 			{
-				GraphInstancePtr = Context.AllocateNodeInstance(UE::AnimNext::FWeakDecoratorPtr(), *FoundHandle);
+				GraphInstancePtr = Context.AllocateNodeInstance(UE::AnimNext::FWeakTraitPtr(), *FoundHandle);
 			}
 		}
 
diff --git a/Engine/Plugins/Experimental/AnimNext/Source/AnimNext/Private/Graph/AnimNextGraphInstancePtr.cpp b/Engine/Plugins/Experimental/AnimNext/Source/AnimNext/Private/Graph/AnimNextGraphInstancePtr.cpp
--- a/Engine/Plugins/Experimental/AnimNext/Source/AnimNext/Private/Graph/AnimNextGraphInstancePtr.cpp
+++ b/Engine/Plugins/Experimental/AnimNext/Source/AnimNext/Private/Graph/AnimNextGraphInstancePtr.cpp
@@ -45,9 +45,9 @@ const UAnimNextGraph* FAnimNextGraphInstancePtr::GetGraph() const
 	return Impl ? Impl->GetGraph() : nullptr;
 }
 
-UE::AnimNext::FWeakDecoratorPtr FAnimNextGraphInstancePtr::GetGraphRootPtr() const
+UE::AnimNext::FWeakTraitPtr FAnimNextGraphInstancePtr::GetGraphRootPtr() const
 {
-	return Impl ? Impl->GetGraphRootPtr() : UE::AnimNext::FWeakDecoratorPtr();
+	return Impl ? Impl->GetGraphRootPtr() : UE::AnimNext::FWeakTraitPtr();
 }
 
 FAnimNextGraphInstance* FAnimNextGraphInstancePtr::GetImpl() const
diff --git a/Engine/Plugins/Experimental/AnimNext/Source/AnimNext/Private/Graph/GC_GraphInstanceComponent.cpp b/Engine/Plugins/Experimental/AnimNext/Source/AnimNext/Private/Graph/GC_GraphInstanceComponent.cpp
--- a/Engine/Plugins/Experimental/AnimNext/Source/AnimNext/Private/Graph/GC_GraphInstanceComponent.cpp
+++ b/Engine/Plugins/Experimental/AnimNext/Source/AnimNext/Private/Graph/GC_GraphInstanceComponent.cpp
@@ -2,42 +2,42 @@
 
 #include "Graph/GC_GraphInstanceComponent.h"
 
-#include "DecoratorBase/ExecutionContext.h"
-#include "DecoratorInterfaces/IGarbageCollection.h"
+#include "TraitCore/ExecutionContext.h"
+#include "TraitInterfaces/IGarbageCollection.h"
 
 namespace UE::AnimNext
 {
-	void FGCGraphInstanceComponent::Register(FAnimNextGraphInstance& GraphInstance, const FWeakDecoratorPtr& DecoratorPtr)
+	void FGCGraphInstanceComponent::Register(FAnimNextGraphInstance& GraphInstance, const FWeakTraitPtr& TraitPtr)
 	{
-		DecoratorsWithReferences.Add(FEntry(GraphInstance, DecoratorPtr));
+		TraitsWithReferences.Add(FEntry(GraphInstance, TraitPtr));
 	}
 
-	void FGCGraphInstanceComponent::Unregister(const FWeakDecoratorPtr& DecoratorPtr)
+	void FGCGraphInstanceComponent::Unregister(const FWeakTraitPtr& TraitPtr)
 	{
-		const int32 EntryIndex = DecoratorsWithReferences.IndexOfByPredicate(
-			[&DecoratorPtr](const FEntry& Entry)
+		const int32 EntryIndex = TraitsWithReferences.IndexOfByPredicate(
+			[&TraitPtr](const FEntry& Entry)
 			{
-				return Entry.DecoratorPtr == DecoratorPtr;
+				return Entry.TraitPtr == TraitPtr;
 			});
 
 		if (ensure(EntryIndex != INDEX_NONE))
 		{
-			DecoratorsWithReferences.RemoveAtSwap(EntryIndex);
+			TraitsWithReferences.RemoveAtSwap(EntryIndex);
 		}
 	}
 
 	void FGCGraphInstanceComponent::AddReferencedObjects(FReferenceCollector& Collector) const
 	{
 		FExecutionContext Context;
-		TDecoratorBinding<IGarbageCollection> GCDecorator;
+		TTraitBinding<IGarbageCollection> GCTrait;
 
 		// TODO: If we kept the entries sorted by graph instance, we could re-use the execution context
-		for (const FEntry& Entry : DecoratorsWithReferences)
+		for (const FEntry& Entry : TraitsWithReferences)
 		{
 			Context.BindTo(Entry.GraphInstance);
-			ensure(Context.GetInterface(Entry.DecoratorPtr, GCDecorator));
+			ensure(Context.GetInterface(Entry.TraitPtr, GCTrait));
 
-			GCDecorator.AddReferencedObjects(Context, Collector);
+			GCTrait.AddReferencedObjects(Context, Collector);
 		}
 	}
 }
diff --git a/Engine/Plugins/Experimental/AnimNext/Source/AnimNext/Private/Graph/GC_GraphInstanceComponent.h b/Engine/Plugins/Experimental/AnimNext/Source/AnimNext/Private/Graph/GC_GraphInstanceComponent.h
--- a/Engine/Plugins/Experimental/AnimNext/Source/AnimNext/Private/Graph/GC_GraphInstanceComponent.h
+++ b/Engine/Plugins/Experimental/AnimNext/Source/AnimNext/Private/Graph/GC_GraphInstanceComponent.h
@@ -3,7 +3,7 @@
 #pragma once
 
 #include "CoreMinimal.h"
-#include "DecoratorBase/DecoratorPtr.h"
+#include "TraitCore/TraitPtr.h"
 #include "Graph/GraphInstanceComponent.h"
 
 struct FAnimNextGraphInstance;
@@ -20,12 +20,12 @@ namespace UE::AnimNext
 	{
 		DECLARE_ANIM_GRAPH_INSTANCE_COMPONENT(FGCGraphInstanceComponent)
 
-		// Registers the provided decorator with the GC system
+		// Registers the provided trait with the GC system
 		// Once registered, IGarbageCollection::AddReferencedObjects will be called on it during GC
-		void Register(FAnimNextGraphInstance& GraphInstance, const FWeakDecoratorPtr& DecoratorPtr);
+		void Register(FAnimNextGraphInstance& GraphInstance, const FWeakTraitPtr& TraitPtr);
 
-		// Unregisters the provided decorator from the GC system
-		void Unregister(const FWeakDecoratorPtr& DecoratorPtr);
+		// Unregisters the provided trait from the GC system
+		void Unregister(const FWeakTraitPtr& TraitPtr);
 
 		// Called during garbage collection to collect strong object references
 		void AddReferencedObjects(FReferenceCollector& Collector) const;
@@ -33,17 +33,17 @@ namespace UE::AnimNext
 	private:
 		struct FEntry
 		{
-			FEntry(FAnimNextGraphInstance& InGraphInstance, const FWeakDecoratorPtr& InDecoratorPtr)
+			FEntry(FAnimNextGraphInstance& InGraphInstance, const FWeakTraitPtr& InTraitPtr)
 				: GraphInstance(InGraphInstance)
-				, DecoratorPtr(InDecoratorPtr)
+				, TraitPtr(InTraitPtr)
 			{
 			}
 
 			FAnimNextGraphInstance& GraphInstance;
-			FWeakDecoratorPtr DecoratorPtr;
+			FWeakTraitPtr TraitPtr;
 		};
 
-		// List of decorator handles that contain UObject references and implement IGarbageCollection
-		TArray<FEntry> DecoratorsWithReferences;
+		// List of trait handles that contain UObject references and implement IGarbageCollection
+		TArray<FEntry> TraitsWithReferences;
 	};
 }
diff --git a/Engine/Plugins/Experimental/AnimNext/Source/AnimNext/Private/Graph/RigUnit_AnimNextGraphEvaluator.cpp b/Engine/Plugins/Experimental/AnimNext/Source/AnimNext/Private/Graph/RigUnit_AnimNextGraphEvaluator.cpp
--- a/Engine/Plugins/Experimental/AnimNext/Source/AnimNext/Private/Graph/RigUnit_AnimNextGraphEvaluator.cpp
+++ b/Engine/Plugins/Experimental/AnimNext/Source/AnimNext/Private/Graph/RigUnit_AnimNextGraphEvaluator.cpp
@@ -2,7 +2,7 @@
 
 #include "Graph/RigUnit_AnimNextGraphEvaluator.h"
 #include "Context.h"
-#include "DecoratorBase/LatentPropertyHandle.h"
+#include "TraitCore/LatentPropertyHandle.h"
 
 #include UE_INLINE_GENERATED_CPP_BY_NAME(RigUnit_AnimNextGraphEvaluator)
 
diff --git a/Engine/Plugins/Experimental/AnimNext/Source/AnimNext/Private/Graph/SyncGroup_GraphInstanceComponent.cpp b/Engine/Plugins/Experimental/AnimNext/Source/AnimNext/Private/Graph/SyncGroup_GraphInstanceComponent.cpp
--- a/Engine/Plugins/Experimental/AnimNext/Source/AnimNext/Private/Graph/SyncGroup_GraphInstanceComponent.cpp
+++ b/Engine/Plugins/Experimental/AnimNext/Source/AnimNext/Private/Graph/SyncGroup_GraphInstanceComponent.cpp
@@ -1,14 +1,14 @@
 // Copyright Epic Games, Inc. All Rights Reserved.
 
 #include "Graph/SyncGroup_GraphInstanceComponent.h"
-#include "DecoratorInterfaces/IGroupSynchronization.h"
+#include "TraitInterfaces/IGroupSynchronization.h"
 
 namespace UE::AnimNext
 {
-	void FSyncGroupGraphInstanceComponent::RegisterWithGroup(FName GroupName, EAnimGroupRole::Type GroupRole, const FWeakDecoratorPtr& DecoratorPtr, const FDecoratorUpdateState& DecoratorState)
+	void FSyncGroupGraphInstanceComponent::RegisterWithGroup(FName GroupName, EAnimGroupRole::Type GroupRole, const FWeakTraitPtr& TraitPtr, const FTraitUpdateState& TraitState)
 	{
 		FSyncGroupState& GroupState = SyncGroupMap.FindOrAdd(GroupName);
-		GroupState.Members.Add(FSyncGroupMember{ DecoratorState, DecoratorPtr, GroupRole });
+		GroupState.Members.Add(FSyncGroupMember{ TraitState, TraitPtr, GroupRole });
 	}
 
 	void FSyncGroupGraphInstanceComponent::PreUpdate(FExecutionContext& Context)
@@ -19,7 +19,7 @@ namespace UE::AnimNext
 
 	void FSyncGroupGraphInstanceComponent::PostUpdate(FExecutionContext& Context)
 	{
-		TDecoratorBinding<IGroupSynchronization> GroupSyncDecorator;
+		TTraitBinding<IGroupSynchronization> GroupSyncTrait;
 
 		// Now that we have discovered all groups and their memberships, we can perform synchronization
 		for (const TTuple<FName, FSyncGroupState>& It : SyncGroupMap)
@@ -42,10 +42,10 @@ namespace UE::AnimNext
 				case EAnimGroupRole::CanBeLeader:
 				case EAnimGroupRole::TransitionLeader:
 					// Highest weight is the leader
-					if (GroupMember.DecoratorState.GetTotalWeight() > LeaderTotalWeight)
+					if (GroupMember.TraitState.GetTotalWeight() > LeaderTotalWeight)
 					{
 						LeaderIndex = MemberIndex;
-						LeaderTotalWeight = GroupMember.DecoratorState.GetTotalWeight();
+						LeaderTotalWeight = GroupMember.TraitState.GetTotalWeight();
 					}
 					break;
 				case EAnimGroupRole::AlwaysLeader:
@@ -73,10 +73,10 @@ namespace UE::AnimNext
 			{
 				const FSyncGroupMember& GroupLeader = GroupState.Members[LeaderIndex];
 
-				Context.BindTo(GroupLeader.DecoratorPtr);
-				ensure(Context.GetInterface(GroupLeader.DecoratorPtr, GroupSyncDecorator));
+				Context.BindTo(GroupLeader.TraitPtr);
+				ensure(Context.GetInterface(GroupLeader.TraitPtr, GroupSyncTrait));
 
-				LeaderProgressRatio = GroupSyncDecorator.AdvanceBy(Context, GroupLeader.DecoratorState.GetDeltaTime());
+				LeaderProgressRatio = GroupSyncTrait.AdvanceBy(Context, GroupLeader.TraitState.GetDeltaTime());
 			}
 
 			// Advance every follower to the same progress ratio as the leader
@@ -89,10 +89,10 @@ namespace UE::AnimNext
 
 				const FSyncGroupMember& GroupMember = GroupState.Members[MemberIndex];
 
-				Context.BindTo(GroupMember.DecoratorPtr);
-				ensure(Context.GetInterface(GroupMember.DecoratorPtr, GroupSyncDecorator));
+				Context.BindTo(GroupMember.TraitPtr);
+				ensure(Context.GetInterface(GroupMember.TraitPtr, GroupSyncTrait));
 
-				GroupSyncDecorator.AdvanceToRatio(Context, LeaderProgressRatio);
+				GroupSyncTrait.AdvanceToRatio(Context, LeaderProgressRatio);
 			}
 		}
 	}
diff --git a/Engine/Plugins/Experimental/AnimNext/Source/AnimNext/Private/Graph/SyncGroup_GraphInstanceComponent.h b/Engine/Plugins/Experimental/AnimNext/Source/AnimNext/Private/Graph/SyncGroup_GraphInstanceComponent.h
--- a/Engine/Plugins/Experimental/AnimNext/Source/AnimNext/Private/Graph/SyncGroup_GraphInstanceComponent.h
+++ b/Engine/Plugins/Experimental/AnimNext/Source/AnimNext/Private/Graph/SyncGroup_GraphInstanceComponent.h
@@ -4,8 +4,8 @@
 
 #include "CoreMinimal.h"
 #include "Animation/AnimationAsset.h"
-#include "DecoratorBase/DecoratorPtr.h"
-#include "DecoratorInterfaces/IUpdate.h"
+#include "TraitCore/TraitPtr.h"
+#include "TraitInterfaces/IUpdate.h"
 #include "Graph/GraphInstanceComponent.h"
 
 namespace UE::AnimNext
@@ -19,7 +19,7 @@ namespace UE::AnimNext
 	{
 		DECLARE_ANIM_GRAPH_INSTANCE_COMPONENT(FSyncGroupGraphInstanceComponent)
 
-		void RegisterWithGroup(FName GroupName, EAnimGroupRole::Type GroupRole, const FWeakDecoratorPtr& DecoratorPtr, const FDecoratorUpdateState& DecoratorState);
+		void RegisterWithGroup(FName GroupName, EAnimGroupRole::Type GroupRole, const FWeakTraitPtr& TraitPtr, const FTraitUpdateState& TraitState);
 
 		// FGraphInstanceComponent impl
 		virtual void PreUpdate(FExecutionContext& Context) override;
@@ -28,8 +28,8 @@ namespace UE::AnimNext
 	private:
 		struct FSyncGroupMember
 		{
-			FDecoratorUpdateState				DecoratorState;
-			FWeakDecoratorPtr					DecoratorPtr;
+			FTraitUpdateState					TraitState;
+			FWeakTraitPtr						TraitPtr;
 			TEnumAsByte<EAnimGroupRole::Type>	GroupRole;
 		};
 
diff --git a/Engine/Plugins/Experimental/AnimNext/Source/AnimNext/Private/Module.cpp b/Engine/Plugins/Experimental/AnimNext/Source/AnimNext/Private/Module.cpp
--- a/Engine/Plugins/Experimental/AnimNext/Source/AnimNext/Private/Module.cpp
+++ b/Engine/Plugins/Experimental/AnimNext/Source/AnimNext/Private/Module.cpp
@@ -8,8 +8,8 @@
 #include "Modules/ModuleManager.h"
 #include "Misc/CoreDelegates.h"
 #include "DataRegistry.h"
-#include "DecoratorBase/DecoratorRegistry.h"
-#include "DecoratorBase/NodeTemplateRegistry.h"
+#include "TraitCore/TraitRegistry.h"
+#include "TraitCore/NodeTemplateRegistry.h"
 #include "Graph/AnimNextGraph.h"
 #include "RigVMCore/RigVMRegistry.h"
 #include "RigVMRuntimeDataRegistry.h"
@@ -25,9 +25,9 @@
 #include "HAL/IConsoleManager.h"
 #include "UObject/UObjectIterator.h"
 
-#include "DecoratorBase/DecoratorTemplate.h"
-#include "DecoratorBase/NodeDescription.h"
-#include "DecoratorBase/NodeTemplate.h"
+#include "TraitCore/TraitTemplate.h"
+#include "TraitCore/NodeDescription.h"
+#include "TraitCore/NodeTemplate.h"
 #endif
 
 namespace UE::AnimNext
@@ -54,7 +54,7 @@ public:
 		FObjectProxyFactory::Init();
 		FExternalParameterRegistry::Init();
 		FDataRegistry::Init();
-		FDecoratorRegistry::Init();
+		FTraitRegistry::Init();
 		FNodeTemplateRegistry::Init();
 		FScheduler::Init();
 		FRigVMRuntimeDataRegistry::Init();
@@ -83,7 +83,7 @@ public:
 		FRigVMRuntimeDataRegistry::Destroy();
 		FScheduler::Destroy();
 		FNodeTemplateRegistry::Destroy();
-		FDecoratorRegistry::Destroy();
+		FTraitRegistry::Destroy();
 		FDataRegistry::Destroy();
 		FObjectProxyFactory::Destroy();
 		FExternalParameterRegistry::Destroy();
@@ -110,7 +110,7 @@ public:
 		LogAnimation.SetVerbosity(ELogVerbosity::All);
 
 		const FNodeTemplateRegistry& NodeTemplateRegistry = FNodeTemplateRegistry::Get();
-		const FDecoratorRegistry& DecoratorRegistry = FDecoratorRegistry::Get();
+		const FTraitRegistry& TraitRegistry = FTraitRegistry::Get();
 
 		UE_LOG(LogAnimation, Log, TEXT("===== AnimNext Node Templates ====="));
 		UE_LOG(LogAnimation, Log, TEXT("Template Buffer Size: %u bytes"), NodeTemplateRegistry.TemplateBuffer.GetAllocatedSize());
@@ -120,35 +120,35 @@ public:
 			const FNodeTemplateRegistryHandle Handle = It.Value();
 			const FNodeTemplate* NodeTemplate = NodeTemplateRegistry.Find(Handle);
 
-			const uint32 NumDecorators = NodeTemplate->GetNumDecorators();
+			const uint32 NumTraits = NodeTemplate->GetNumTraits();
 
-			UE_LOG(LogAnimation, Log, TEXT("[%x] has %u decorators ..."), NodeTemplate->GetUID(), NumDecorators);
+			UE_LOG(LogAnimation, Log, TEXT("[%x] has %u traits ..."), NodeTemplate->GetUID(), NumTraits);
 			UE_LOG(LogAnimation, Log, TEXT("    Template Size: %u bytes"), NodeTemplate->GetNodeTemplateSize());
 			UE_LOG(LogAnimation, Log, TEXT("    Shared Data Size: %u bytes"), NodeTemplate->GetNodeSharedDataSize());
 			UE_LOG(LogAnimation, Log, TEXT("    Instance Data Size: %u bytes"), NodeTemplate->GetNodeInstanceDataSize());
-			UE_LOG(LogAnimation, Log, TEXT("    Decorators ..."));
+			UE_LOG(LogAnimation, Log, TEXT("    Traits ..."));
 
-			const FDecoratorTemplate* DecoratorTemplates = NodeTemplate->GetDecorators();
-			for (uint32 DecoratorIndex = 0; DecoratorIndex < NumDecorators; ++DecoratorIndex)
+			const FTraitTemplate* TraitTemplates = NodeTemplate->GetTraits();
+			for (uint32 TraitIndex = 0; TraitIndex < NumTraits; ++TraitIndex)
 			{
-				const FDecoratorTemplate* DecoratorTemplate = DecoratorTemplates + DecoratorIndex;
-				const FDecorator* Decorator = DecoratorRegistry.Find(DecoratorTemplate->GetRegistryHandle());
-				const FString DecoratorName = Decorator != nullptr ? Decorator->GetDecoratorName() : TEXT("<Unknown>");
+				const FTraitTemplate* TraitTemplate = TraitTemplates + TraitIndex;
+				const FTrait* Trait = TraitRegistry.Find(TraitTemplate->GetRegistryHandle());
+				const FString TraitName = Trait != nullptr ? Trait->GetTraitName() : TEXT("<Unknown>");
 
-				const uint32 NextDecoratorIndex = DecoratorIndex + 1;
-				const uint32 EndOfNextDecoratorSharedData = NextDecoratorIndex < NumDecorators ? DecoratorTemplates[NextDecoratorIndex].GetNodeSharedOffset() : NodeTemplate->GetNodeSharedDataSize();
-				const uint32 DecoratorSharedDataSize = EndOfNextDecoratorSharedData - DecoratorTemplate->GetNodeSharedOffset();
+				const uint32 NextTraitIndex = TraitIndex + 1;
+				const uint32 EndOfNextTraitSharedData = NextTraitIndex < NumTraits ? TraitTemplates[NextTraitIndex].GetNodeSharedOffset() : NodeTemplate->GetNodeSharedDataSize();
+				const uint32 TraitSharedDataSize = EndOfNextTraitSharedData - TraitTemplate->GetNodeSharedOffset();
 
-				const uint32 EndOfNextDecoratorInstanceData = NextDecoratorIndex < NumDecorators ? DecoratorTemplates[NextDecoratorIndex].GetNodeInstanceOffset() : NodeTemplate->GetNodeInstanceDataSize();
-				const uint32 DecoratorInstanceDataSize = EndOfNextDecoratorInstanceData - DecoratorTemplate->GetNodeInstanceOffset();
+				const uint32 EndOfNextTraitInstanceData = NextTraitIndex < NumTraits ? TraitTemplates[NextTraitIndex].GetNodeInstanceOffset() : NodeTemplate->GetNodeInstanceDataSize();
+				const uint32 TraitInstanceDataSize = EndOfNextTraitInstanceData - TraitTemplate->GetNodeInstanceOffset();
 
-				UE_LOG(LogAnimation, Log, TEXT("            %u: [%x] %s (%s)"), DecoratorIndex, DecoratorTemplate->GetUID().GetUID(), *DecoratorName, DecoratorTemplate->GetMode() == EDecoratorMode::Base ? TEXT("Base") : TEXT("Additive"));
-				UE_LOG(LogAnimation, Log, TEXT("                Shared Data: [Offset: %u bytes, Size: %u bytes]"), DecoratorTemplate->GetNodeSharedOffset(), DecoratorSharedDataSize);
-				if (DecoratorTemplate->HasLatentProperties() && Decorator != nullptr)
+				UE_LOG(LogAnimation, Log, TEXT("            %u: [%x] %s (%s)"), TraitIndex, TraitTemplate->GetUID().GetUID(), *TraitName, TraitTemplate->GetMode() == ETraitMode::Base ? TEXT("Base") : TEXT("Additive"));
+				UE_LOG(LogAnimation, Log, TEXT("                Shared Data: [Offset: %u bytes, Size: %u bytes]"), TraitTemplate->GetNodeSharedOffset(), TraitSharedDataSize);
+				if (TraitTemplate->HasLatentProperties() && Trait != nullptr)
 				{
-					UE_LOG(LogAnimation, Log, TEXT("                Shared Data Latent Property Handles: [Offset: %u bytes, Count: %u]"), DecoratorTemplate->GetNodeSharedLatentPropertyHandlesOffset(), Decorator->GetNumLatentDecoratorProperties());
+					UE_LOG(LogAnimation, Log, TEXT("                Shared Data Latent Property Handles: [Offset: %u bytes, Count: %u]"), TraitTemplate->GetNodeSharedLatentPropertyHandlesOffset(), Trait->GetNumLatentTraitProperties());
 				}
-				UE_LOG(LogAnimation, Log, TEXT("                Instance Data: [Offset: %u bytes, Size: %u bytes]"), DecoratorTemplate->GetNodeInstanceOffset(), DecoratorInstanceDataSize);
+				UE_LOG(LogAnimation, Log, TEXT("                Instance Data: [Offset: %u bytes, Size: %u bytes]"), TraitTemplate->GetNodeInstanceOffset(), TraitInstanceDataSize);
 			}
 		}
 
@@ -181,7 +181,7 @@ public:
 		AnimGraphs.Sort(FCompareObjectNames());
 
 		const FNodeTemplateRegistry& NodeTemplateRegistry = FNodeTemplateRegistry::Get();
-		const FDecoratorRegistry& DecoratorRegistry = FDecoratorRegistry::Get();
+		const FTraitRegistry& TraitRegistry = FTraitRegistry::Get();
 		const bool bDetailedOutput = true;
 
 		UE_LOG(LogAnimation, Log, TEXT("===== AnimNext Animation Graphs ====="));
@@ -222,34 +222,34 @@ public:
 					const FNodeDescription* NodeDesc = reinterpret_cast<const FNodeDescription*>(&AnimGraph->SharedDataBuffer[NodeOffset]);
 					const FNodeTemplate* NodeTemplate = NodeTemplateRegistry.Find(NodeDesc->GetTemplateHandle());
 
-					const uint32 NumDecorators = NodeTemplate->GetNumDecorators();
+					const uint32 NumTraits = NodeTemplate->GetNumTraits();
 
-					UE_LOG(LogAnimation, Log, TEXT("        Node %u: [Template %x with %u decorators]"), NodeDesc->GetUID().GetNodeIndex(), NodeTemplate->GetUID(), NumDecorators);
+					UE_LOG(LogAnimation, Log, TEXT("        Node %u: [Template %x with %u traits]"), NodeDesc->GetUID().GetNodeIndex(), NodeTemplate->GetUID(), NumTraits);
 					UE_LOG(LogAnimation, Log, TEXT("            Shared Data: [Offset: %u bytes, Size: %u bytes]"), NodeOffset, NodeTemplate->GetNodeSharedDataSize());
 					UE_LOG(LogAnimation, Log, TEXT("            Instance Data Size: %u bytes"), NodeDesc->GetNodeInstanceDataSize());
-					UE_LOG(LogAnimation, Log, TEXT("            Decorators ..."));
+					UE_LOG(LogAnimation, Log, TEXT("            Traits ..."));
 
-					const FDecoratorTemplate* DecoratorTemplates = NodeTemplate->GetDecorators();
-					for (uint32 DecoratorIndex = 0; DecoratorIndex < NumDecorators; ++DecoratorIndex)
+					const FTraitTemplate* TraitTemplates = NodeTemplate->GetTraits();
+					for (uint32 TraitIndex = 0; TraitIndex < NumTraits; ++TraitIndex)
 					{
-						const FDecoratorTemplate* DecoratorTemplate = DecoratorTemplates + DecoratorIndex;
-						const FDecorator* Decorator = DecoratorRegistry.Find(DecoratorTemplate->GetRegistryHandle());
-						const FString DecoratorName = Decorator != nullptr ? Decorator->GetDecoratorName() : TEXT("<Unknown>");
+						const FTraitTemplate* TraitTemplate = TraitTemplates + TraitIndex;
+						const FTrait* Trait = TraitRegistry.Find(TraitTemplate->GetRegistryHandle());
+						const FString TraitName = Trait != nullptr ? Trait->GetTraitName() : TEXT("<Unknown>");
 
-						const uint32 NextDecoratorIndex = DecoratorIndex + 1;
-						const uint32 EndOfNextDecoratorSharedData = NextDecoratorIndex < NumDecorators ? DecoratorTemplates[NextDecoratorIndex].GetNodeSharedOffset() : NodeTemplate->GetNodeSharedDataSize();
-						const uint32 DecoratorSharedDataSize = EndOfNextDecoratorSharedData - DecoratorTemplate->GetNodeSharedOffset();
+						const uint32 NextTraitIndex = TraitIndex + 1;
+						const uint32 EndOfNextTraitSharedData = NextTraitIndex < NumTraits ? TraitTemplates[NextTraitIndex].GetNodeSharedOffset() : NodeTemplate->GetNodeSharedDataSize();
+						const uint32 TraitSharedDataSize = EndOfNextTraitSharedData - TraitTemplate->GetNodeSharedOffset();
 
-						const uint32 EndOfNextDecoratorInstanceData = NextDecoratorIndex < NumDecorators ? DecoratorTemplates[NextDecoratorIndex].GetNodeInstanceOffset() : NodeTemplate->GetNodeInstanceDataSize();
-						const uint32 DecoratorInstanceDataSize = EndOfNextDecoratorInstanceData - DecoratorTemplate->GetNodeInstanceOffset();
+						const uint32 EndOfNextTraitInstanceData = NextTraitIndex < NumTraits ? TraitTemplates[NextTraitIndex].GetNodeInstanceOffset() : NodeTemplate->GetNodeInstanceDataSize();
+						const uint32 TraitInstanceDataSize = EndOfNextTraitInstanceData - TraitTemplate->GetNodeInstanceOffset();
 
-						UE_LOG(LogAnimation, Log, TEXT("                    %u: [%x] %s (%s)"), DecoratorIndex, DecoratorTemplate->GetUID().GetUID(), *DecoratorName, DecoratorTemplate->GetMode() == EDecoratorMode::Base ? TEXT("Base") : TEXT("Additive"));
-						UE_LOG(LogAnimation, Log, TEXT("                        Shared Data: [Offset: %u bytes, Size: %u bytes]"), DecoratorTemplate->GetNodeSharedOffset(), DecoratorSharedDataSize);
-						if (DecoratorTemplate->HasLatentProperties() && Decorator != nullptr)
+						UE_LOG(LogAnimation, Log, TEXT("                    %u: [%x] %s (%s)"), TraitIndex, TraitTemplate->GetUID().GetUID(), *TraitName, TraitTemplate->GetMode() == ETraitMode::Base ? TEXT("Base") : TEXT("Additive"));
+						UE_LOG(LogAnimation, Log, TEXT("                        Shared Data: [Offset: %u bytes, Size: %u bytes]"), TraitTemplate->GetNodeSharedOffset(), TraitSharedDataSize);
+						if (TraitTemplate->HasLatentProperties() && Trait != nullptr)
 						{
-							UE_LOG(LogAnimation, Log, TEXT("                        Shared Data Latent Property Handles: [Offset: %u bytes, Count: %u]"), DecoratorTemplate->GetNodeSharedLatentPropertyHandlesOffset(), Decorator->GetNumLatentDecoratorProperties());
+							UE_LOG(LogAnimation, Log, TEXT("                        Shared Data Latent Property Handles: [Offset: %u bytes, Count: %u]"), TraitTemplate->GetNodeSharedLatentPropertyHandlesOffset(), Trait->GetNumLatentTraitProperties());
 						}
-						UE_LOG(LogAnimation, Log, TEXT("                        Instance Data: [Offset: %u bytes, Size: %u bytes]"), DecoratorTemplate->GetNodeInstanceOffset(), DecoratorInstanceDataSize);
+						UE_LOG(LogAnimation, Log, TEXT("                        Instance Data: [Offset: %u bytes, Size: %u bytes]"), TraitTemplate->GetNodeInstanceOffset(), TraitInstanceDataSize);
 					}
 
 					NodeOffset += NodeTemplate->GetNodeSharedDataSize();
diff --git a/Engine/Plugins/Experimental/AnimNext/Source/AnimNext/Private/Scheduler/AnimNextScheduleGraphTask.cpp b/Engine/Plugins/Experimental/AnimNext/Source/AnimNext/Private/Scheduler/AnimNextScheduleGraphTask.cpp
--- a/Engine/Plugins/Experimental/AnimNext/Source/AnimNext/Private/Scheduler/AnimNextScheduleGraphTask.cpp
+++ b/Engine/Plugins/Experimental/AnimNext/Source/AnimNext/Private/Scheduler/AnimNextScheduleGraphTask.cpp
@@ -5,8 +5,8 @@
 #include "Scheduler/ScheduleContext.h"
 #include "Graph/AnimNextGraph.h"
 #include "Context.h"
-#include "DecoratorInterfaces/IEvaluate.h"
-#include "DecoratorInterfaces/IUpdate.h"
+#include "TraitInterfaces/IEvaluate.h"
+#include "TraitInterfaces/IUpdate.h"
 #include "EvaluationVM/EvaluationVM.h"
 #include "Graph/AnimNextExecuteContext.h"
 #include "Graph/AnimNext_LODPose.h"
diff --git a/Engine/Plugins/Experimental/AnimNext/Source/AnimNext/Private/TraitCore/ExecutionContext.cpp b/Engine/Plugins/Experimental/AnimNext/Source/AnimNext/Private/TraitCore/ExecutionContext.cpp
new file mode 100644
--- /dev/null
+++ b/Engine/Plugins/Experimental/AnimNext/Source/AnimNext/Private/TraitCore/ExecutionContext.cpp
@@ -0,0 +1,461 @@
+// Copyright Epic Games, Inc. All Rights Reserved.
+
+#include "TraitCore/ExecutionContext.h"
+
+#include "TraitCore/Trait.h"
+#include "TraitCore/TraitRegistry.h"
+#include "TraitCore/TraitTemplate.h"
+#include "TraitCore/NodeDescription.h"
+#include "TraitCore/NodeInstance.h"
+#include "TraitCore/NodeTemplate.h"
+#include "TraitCore/NodeTemplateRegistry.h"
+#include "Graph/AnimNextGraph.h"
+#include "Graph/AnimNextGraphInstance.h"
+
+namespace UE::AnimNext
+{
+	FExecutionContext::FExecutionContext()
+		: NodeTemplateRegistry(FNodeTemplateRegistry::Get())
+		, TraitRegistry(FTraitRegistry::Get())
+		, GraphInstance(nullptr)
+	{
+	}
+
+	FExecutionContext::FExecutionContext(FAnimNextGraphInstancePtr& InGraphInstance)
+		: FExecutionContext()
+	{
+		BindTo(InGraphInstance);
+	}
+
+	FExecutionContext::FExecutionContext(FAnimNextGraphInstance& InGraphInstance)
+		: FExecutionContext()
+	{
+		BindTo(InGraphInstance);
+	}
+
+	void FExecutionContext::BindTo(FAnimNextGraphInstancePtr& InGraphInstance)
+	{
+		if (FAnimNextGraphInstance* InstanceImpl = InGraphInstance.GetImpl())
+		{
+			BindTo(*InstanceImpl);
+		}
+	}
+
+	void FExecutionContext::BindTo(FAnimNextGraphInstance& InGraphInstance)
+	{
+		if (GraphInstance == &InGraphInstance)
+		{
+			return;	// Already bound to this graph instance, nothing to do
+		}
+
+		if (const UAnimNextGraph* Graph = InGraphInstance.GetGraph())
+		{
+			GraphInstance = &InGraphInstance;
+			GraphSharedData = Graph->SharedDataBuffer;
+		}
+	}
+
+	void FExecutionContext::BindTo(const FWeakTraitPtr& TraitPtr)
+	{
+		if (const FNodeInstance* NodeInstance = TraitPtr.GetNodeInstance())
+		{
+			BindTo(NodeInstance->GetOwner());
+		}
+	}
+
+	bool FExecutionContext::IsBound() const
+	{
+		return GraphInstance != nullptr;
+	}
+
+	bool FExecutionContext::IsBoundTo(const FAnimNextGraphInstancePtr& InGraphInstance) const
+	{
+		return GraphInstance == InGraphInstance.GetImpl();
+	}
+
+	bool FExecutionContext::IsBoundTo(const FAnimNextGraphInstance& InGraphInstance) const
+	{
+		return GraphInstance == &InGraphInstance;
+	}
+
+	FTraitPtr FExecutionContext::AllocateNodeInstance(const FWeakTraitPtr& ParentBinding, FAnimNextTraitHandle ChildTraitHandle) const
+	{
+		if (!ensure(ChildTraitHandle.IsValid()))
+		{
+			return FTraitPtr();	// Attempting to allocate a node using an invalid trait handle
+		}
+
+		if (!ensure(IsBound()))
+		{
+			return FTraitPtr();	// The execution context must be bound to a valid graph instance
+		}
+
+		const FNodeHandle ChildNodeHandle = ChildTraitHandle.GetNodeHandle();
+		const FNodeDescription& NodeDesc = GetNodeDescription(ChildNodeHandle);
+
+		const FNodeTemplate* NodeTemplate = GetNodeTemplate(NodeDesc);
+		if (!ensure(NodeTemplate != nullptr))
+		{
+			return FTraitPtr();	// Node template wasn't found, node descriptor is perhaps corrupted
+		}
+
+		const uint32 ChildTraitIndex = ChildTraitHandle.GetTraitIndex();
+
+		if (ChildTraitIndex >= NodeTemplate->GetNumTraits())
+		{
+			return FTraitPtr();	// The requested trait index doesn't exist on that node descriptor
+		}
+
+		// If the trait we wish to allocate lives in the parent node, return a weak handle to it
+		// We use a weak handle to avoid issues when multiple base traits live within the same node
+		// When this happens, a trait can end up pointing to another within the same node causing
+		// the reference count to never reach zero when all other handles are released
+		if (FNodeInstance* ParentNodeInstance = ParentBinding.GetNodeInstance())
+		{
+			if (ParentNodeInstance->GetNodeHandle() == ChildNodeHandle)
+			{
+				return FTraitPtr(ParentNodeInstance, FTraitPtr::IS_WEAK_BIT, ChildTraitIndex);
+			}
+		}
+
+		// We need to allocate a new node instance
+		const FTraitTemplate* TraitDescs = NodeTemplate->GetTraits();
+
+		const uint32 InstanceSize = NodeDesc.GetNodeInstanceDataSize();
+		uint8* NodeInstanceBuffer = reinterpret_cast<uint8*>(FMemory::Malloc(InstanceSize, 16));
+		FNodeInstance* NodeInstance = new(NodeInstanceBuffer) FNodeInstance(*GraphInstance, ChildNodeHandle);
+
+		// Start construction with the bottom trait
+		const FTraitTemplate* StartDesc = TraitDescs;
+		const FTraitTemplate* EndDesc = TraitDescs + NodeTemplate->GetNumTraits();
+
+		for (const FTraitTemplate* TraitDesc = StartDesc; TraitDesc != EndDesc; ++TraitDesc)
+		{
+			const FTrait* Trait = GetTrait(*TraitDesc);
+			if (Trait == nullptr)
+			{
+				continue;	// Trait hasn't been loaded or registered, skip it
+			}
+
+			const uint32 TraitIndex = TraitDesc - TraitDescs;
+
+			FWeakTraitPtr TraitPtr(NodeInstance, TraitIndex);
+
+			const FAnimNextTraitSharedData* SharedData = TraitDesc->GetTraitDescription(NodeDesc);
+			FTraitInstanceData* InstanceData = TraitDesc->GetTraitInstance(*NodeInstance);
+
+			const FTraitBinding Binding(nullptr, TraitDesc, &NodeDesc, TraitPtr);
+			Trait->ConstructTraitInstance(*this, Binding);
+		}
+
+		return FTraitPtr(NodeInstance, ChildTraitIndex);
+	}
+
+	void FExecutionContext::ReleaseNodeInstance(FTraitPtr& NodePtr) const
+	{
+		if (!NodePtr.IsValid())
+		{
+			return;
+		}
+
+		FNodeInstance* NodeInstance = NodePtr.GetNodeInstance();
+
+		if (!ensure(IsBoundTo(NodeInstance->GetOwner())))
+		{
+			return;	// The execution context isn't bound to the right graph instance
+		}
+
+		// Reset the handle here to simplify the multiple return statements below
+		NodePtr.PackedPointerAndFlags = 0;
+		NodePtr.TraitIndex = 0;
+
+		if (NodeInstance->RemoveReference())
+		{
+			return;	// Node instance still has references, we can't release it
+		}
+
+		const FNodeDescription& NodeDesc = GetNodeDescription(NodeInstance->GetNodeHandle());
+
+		const FNodeTemplate* NodeTemplate = GetNodeTemplate(NodeDesc);
+		if (!ensure(NodeTemplate != nullptr))
+		{
+			return;	// Node template wasn't found, node descriptor is perhaps corrupted (we'll leak the node memory)
+		}
+
+		const FTraitTemplate* TraitDescs = NodeTemplate->GetTraits();
+
+		// Start destruction with the top trait
+		const FTraitTemplate* StartDesc = TraitDescs + NodeTemplate->GetNumTraits() - 1;
+		const FTraitTemplate* EndDesc = TraitDescs - 1;
+		for (const FTraitTemplate* TraitDesc = StartDesc; TraitDesc != EndDesc; --TraitDesc)
+		{
+			if (const FTrait* Trait = GetTrait(*TraitDesc))
+			{
+				const uint32 TraitIndex = TraitDesc - TraitDescs;
+
+				FWeakTraitPtr TraitPtr(NodeInstance, TraitIndex);
+
+				const FTraitBinding Binding(nullptr, TraitDesc, &NodeDesc, TraitPtr);
+				Trait->DestructTraitInstance(*this, Binding);
+			}
+		}
+
+		NodeInstance->~FNodeInstance();
+		FMemory::Free(NodeInstance);
+	}
+
+	bool FExecutionContext::GetInterfaceImpl(FTraitInterfaceUID InterfaceUID, const FWeakTraitPtr& TraitPtr, FTraitBinding& InterfaceBinding) const
+	{
+		if (!TraitPtr.IsValid())
+		{
+			return false;
+		}
+
+		// TODO: Revisit this to avoid querying each trait over and over
+		//
+		// If we have D traits on a node and I unique interfaces per trait then finding an interface
+		// has O(D*I) complexity: for each trait we have to test every interface
+		// However, if the I interfaces are not unique per trait and instead are unique for the system,
+		// the below code will still have the same complexity. We'll end up testing the same interface multiple
+		// times, once for each trait. Because we will have a small set of interfaces per node that
+		// multiple traits implement (e.g IUpdate, IEvaluate), we can do better.
+		// 
+		// It would be much cheaper if instead we could test if we implement the interface first and then
+		// look for a list of traits (in stack order) that implement it. To do this, when we build the node
+		// template, we could aggregate all interfaces that the traits implement (sorted by their interface id
+		// for determinism or perhaps some other criteria to put popular/hot interfaces first). Then, in the node
+		// template we store a mapping of InterfaceUID to InterfaceIndex (local to that node template). The interface
+		// index can then be used to index into a list of trait indices that implement it.
+		// Code would look like:
+		//    * uint32 InterfaceIndex = NodeTemplate.GetInterfaceUIDs().FindIndex(InterfaceUID);	// INDEX_NONE means that none of the traits implement the interface
+		//    * uint16 TraitIndicesStartOffset = NodeTemplate.InterfaceIndexToTraitIndices(InterfaceIndex);		// An offset relative to the NodeTemplate*
+		//    * const uint8* TraitIndicesForInterface = ((const uint8*NodeTemplate) + TraitIndicesStartOffset;
+		//    * uint8 NumTraitsWithInterface = TraitIndicesForInterface[0];		// Reserve first index for the trait count or if we use a guard value at the end of the array we have to reserve a value
+		//    * uint8 TopTraitIndex = TraitIndicesForInterface[1];				// First index is top of stack
+		// Super trait lookup would be similar:
+		//    * uint8 SuperTraitIndex = TraitIndicesForInterface[TraitIndicesForInterface.FindIndex(CurrentTraitIndex) + 1];	// Need to check if we exceed the max count (bottom of stack)
+		// Once we have the trait index, we use the trait template as we do now
+		//
+		// The above would have a number of advantages:
+		//    * Querying for an interface that the node doesn't implement would be much cheaper (we test each interface once)
+		//    * InterfaceIndex query can easily be implemented in SIMD, we look for a uint32 in a list sequentially
+		//    * Bulk/interleaved querying of interfaces would be much easier (e.g. query for IEvaluate for 4 nodes)
+		//    * A lot less branching is involved, improving throughput
+		//    * We can handle multiple base traits by checking the trait index found which has comparable cost
+		//    * We could store the interface offset from the trait ptr base to avoid the call to GetInterface()
+		//      Instead of testing to find our interface again to return a custom static_cast, we'd store the interface/trait offset
+		//      in the node template alongside the trait index. We can then query for the FTrait* and add the offset
+		//      improving bulk query feasibility.
+		//
+		// However, querying for a Super interface might be slower since we have to start the search from the start.
+		// In contrast, we now start searching at the current trait. In order to keep that cost down and amortize
+		// the search, we would have to cache the InterfaceIndex and the current index in the trait list.
+		// We could use 8-bits for each entry allowing for a max of 256 interfaces per node and 256 traits per node (existing limitation).
+		//
+		// Adding new data to the trait binding isn't ideal. It's current size on 64-bit systems is:
+		//    * 8 bytes for interface pointer (to forward function calls to)
+		//    * 8 bytes for trait template (to get the offsets for shared/instance/latent data)
+		//    * 8 bytes for node description pointer (base of shared data)
+		//    * 8 bytes for node instance pointer (base of instance data, in weak trait ptr)
+		//    * 4 bytes for trait index (in weak trait ptr)
+		//    * 4 bytes of padding (in weak trait ptr)
+		// We current use 4 bytes for the trait index but in practice we only need 1 byte. We use 4 bytes because we pay
+		// for padding anyway. We could use 2 extra bytes for our mapping and still have 5 bytes to space in padding (due to alignment).
+		// In practice, trait bindings always live on the stack and they shouldn't be getting copied/moved around.
+		// We populate them in place when we query, then we use it from its storage on the stack that we first populated.
+		//
+		// To be able to support all of this, instead of having FTrait::GetInterface() implemented by derived types
+		// we would need instead to implement FTrait::GetInterfaceUIDs(). Then on load, when we build a node template and
+		// finalize it, we can build the list of unique interface UIDs and sort it and we can build the mapping structures.
+		// On load, as long as we fit within 64KB we are fine. We have to update FNodeTemplate::Finalize and FNodeTemplateBuilder::BuildNodeTemplate.
+		// Perhaps the builder should use the same code and build in a buffer on the stack and copy the final size out.
+		// We must build the mapping on load because the interfaces we implement are only known at runtime (e.g. can be changed through defines/configuration).
+		// We can still implement GetInterface() in terms of GetInterfaceUIDs if it returns the interface offsets as well.
+
+		FNodeInstance* NodeInstance = TraitPtr.GetNodeInstance();
+
+		if (!ensure(IsBoundTo(NodeInstance->GetOwner())))
+		{
+			return false;	// The execution context isn't bound to the right graph instance
+		}
+
+		const FNodeDescription& NodeDesc = GetNodeDescription(NodeInstance->GetNodeHandle());
+
+		const FNodeTemplate* NodeTemplate = GetNodeTemplate(NodeDesc);
+		if (!ensure(NodeTemplate != nullptr))
+		{
+			return false;	// Node template wasn't found, node descriptor is perhaps corrupted
+		}
+
+		const FTraitTemplate* TraitDescs = NodeTemplate->GetTraits();
+
+		// We only search within the partial stack of the provided trait
+		const FTraitTemplate* CurrentTraitDesc = TraitDescs + TraitPtr.GetTraitIndex();
+		const FTraitTemplate* BaseTraitDesc = CurrentTraitDesc->GetMode() == ETraitMode::Base ? CurrentTraitDesc : (CurrentTraitDesc - CurrentTraitDesc->GetAdditiveTraitIndex());
+
+		// Start searching with the top additive trait towards our base trait
+		const FTraitTemplate* StartDesc = BaseTraitDesc + BaseTraitDesc->GetNumAdditiveTraits();
+		const FTraitTemplate* EndDesc = BaseTraitDesc - 1;
+		for (const FTraitTemplate* TraitDesc = StartDesc; TraitDesc != EndDesc; --TraitDesc)
+		{
+			const FTrait* Trait = GetTrait(*TraitDesc);
+			if (Trait == nullptr)
+			{
+				continue;	// Trait hasn't been loaded or registered, skip it
+			}
+
+			if (const ITraitInterface* Interface = Trait->GetTraitInterface(InterfaceUID))
+			{
+				const FAnimNextTraitSharedData* SharedData = TraitDesc->GetTraitDescription(NodeDesc);
+				FTraitInstanceData* InstanceData = TraitDesc->GetTraitInstance(*NodeInstance);
+
+				const uint32 TraitIndex = TraitDesc - TraitDescs;
+				FWeakTraitPtr InterfaceTraitPtr(NodeInstance, TraitIndex);
+
+				InterfaceBinding = FTraitBinding(Interface, TraitDesc, &NodeDesc, InterfaceTraitPtr);
+				return true;
+			}
+		}
+
+		// We failed to find a trait that handles this interface
+		return false;
+	}
+
+	bool FExecutionContext::GetInterfaceSuperImpl(FTraitInterfaceUID InterfaceUID, const FWeakTraitPtr& TraitPtr, FTraitBinding& SuperBinding) const
+	{
+		if (!TraitPtr.IsValid())
+		{
+			return false;	// Trait pointer isn't valid, can't find a 'super'
+		}
+
+		FNodeInstance* NodeInstance = TraitPtr.GetNodeInstance();
+
+		if (!ensure(IsBoundTo(NodeInstance->GetOwner())))
+		{
+			return false;	// The execution context isn't bound to the right graph instance
+		}
+
+		const FNodeDescription& NodeDesc = GetNodeDescription(NodeInstance->GetNodeHandle());
+
+		const FNodeTemplate* NodeTemplate = GetNodeTemplate(NodeDesc);
+		if (!ensure(NodeTemplate != nullptr))
+		{
+			return false;	// Node template wasn't found, node descriptor is perhaps corrupted
+		}
+
+		const FTraitTemplate* TraitDescs = NodeTemplate->GetTraits();
+
+		// We only search within the partial stack of the provided trait
+		const FTraitTemplate* CurrentTraitDesc = TraitDescs + TraitPtr.GetTraitIndex();
+		if (CurrentTraitDesc->GetMode() == ETraitMode::Base)
+		{
+			return false;	// We've reached a base trait, we don't forward to any parent trait we might have
+		}
+
+		const FTraitTemplate* BaseTraitDesc = CurrentTraitDesc - CurrentTraitDesc->GetAdditiveTraitIndex();
+
+		// Start searching with the next additive trait towards our base trait
+		const FTraitTemplate* StartDesc = CurrentTraitDesc - 1;
+		const FTraitTemplate* EndDesc = BaseTraitDesc - 1;
+		for (const FTraitTemplate* TraitDesc = StartDesc; TraitDesc != EndDesc; --TraitDesc)
+		{
+			const FTrait* Trait = GetTrait(*TraitDesc);
+			if (Trait == nullptr)
+			{
+				continue;	// Trait hasn't been loaded or registered, skip it
+			}
+
+			if (const ITraitInterface* Interface = Trait->GetTraitInterface(InterfaceUID))
+			{
+				const FAnimNextTraitSharedData* SharedData = TraitDesc->GetTraitDescription(NodeDesc);
+				FTraitInstanceData* InstanceData = TraitDesc->GetTraitInstance(*NodeInstance);
+
+				const uint32 TraitIndex = TraitDesc - TraitDescs;
+				FWeakTraitPtr SuperPtr(NodeInstance, TraitIndex);
+
+				SuperBinding = FTraitBinding(Interface, TraitDesc, &NodeDesc, SuperPtr);
+				return true;
+			}
+		}
+
+		// We failed to find a trait that handles this interface
+		return false;
+	}
+
+	void FExecutionContext::SnapshotLatentProperties(const FWeakTraitPtr& TraitPtr, bool bIsFrozen) const
+	{
+		if (!TraitPtr.IsValid())
+		{
+			return;	// Nothing to do
+		}
+
+		FNodeInstance* NodeInstance = TraitPtr.GetNodeInstance();
+
+		if (!ensure(IsBoundTo(NodeInstance->GetOwner())))
+		{
+			return;	// The execution context isn't bound to the right graph instance
+		}
+
+		const FNodeDescription& NodeDesc = GetNodeDescription(NodeInstance->GetNodeHandle());
+
+		const FNodeTemplate* NodeTemplate = GetNodeTemplate(NodeDesc);
+		if (!ensure(NodeTemplate != nullptr))
+		{
+			return;	// Node template wasn't found, node descriptor is perhaps corrupted
+		}
+
+		const FTraitTemplate* TraitDescs = NodeTemplate->GetTraits();
+
+		// We only snapshot the partial stack the specified trait lives in
+		const FTraitTemplate* CurrentTraitDesc = TraitDescs + TraitPtr.GetTraitIndex();
+		const FTraitTemplate* BaseTraitDesc = CurrentTraitDesc->GetMode() == ETraitMode::Base ? CurrentTraitDesc : (CurrentTraitDesc - CurrentTraitDesc->GetAdditiveTraitIndex());
+
+		const FLatentPropertiesHeader& LatentHeader = BaseTraitDesc->GetTraitLatentPropertiesHeader(NodeDesc);
+		if (!LatentHeader.bHasValidLatentProperties)
+		{
+			return;	// All latent properties are inline, nothing to snapshot
+		}
+		else if (bIsFrozen && LatentHeader.bCanAllPropertiesFreeze)
+		{
+			return;	// We are frozen and all latent properties support freezing, nothing to snapshot
+		}
+
+		const FLatentPropertyHandle* LatentHandles = BaseTraitDesc->GetTraitLatentPropertyHandles(NodeDesc);
+		const uint32 NumLatentHandles = BaseTraitDesc->GetNumSubStackLatentPropreties();
+
+		GraphInstance->ExecuteLatentPins(TConstArrayView<FLatentPropertyHandle>(LatentHandles, NumLatentHandles), NodeInstance, bIsFrozen);
+	}
+
+	FGraphInstanceComponent* FExecutionContext::TryGetComponent(int32 ComponentNameHash, FName ComponentName) const
+	{
+		return GraphInstance->TryGetComponent(ComponentNameHash, ComponentName);
+	}
+
+	FGraphInstanceComponent& FExecutionContext::AddComponent(int32 ComponentNameHash, FName ComponentName, TSharedPtr<FGraphInstanceComponent>&& Component) const
+	{
+		return GraphInstance->AddComponent(ComponentNameHash, ComponentName, MoveTemp(Component));
+	}
+
+	GraphInstanceComponentMapType::TConstIterator FExecutionContext::GetComponentIterator() const
+	{
+		return GraphInstance->GetComponentIterator();
+	}
+
+	const FNodeDescription& FExecutionContext::GetNodeDescription(FNodeHandle NodeHandle) const
+	{
+		check(NodeHandle.IsValid());
+		return *reinterpret_cast<const FNodeDescription*>(&GraphSharedData[NodeHandle.GetSharedOffset()]);
+	}
+
+	const FNodeTemplate* FExecutionContext::GetNodeTemplate(const FNodeDescription& NodeDesc) const
+	{
+		check(NodeDesc.GetTemplateHandle().IsValid());
+		return NodeTemplateRegistry.Find(NodeDesc.GetTemplateHandle());
+	}
+
+	const FTrait* FExecutionContext::GetTrait(const FTraitTemplate& Template) const
+	{
+		check(Template.GetRegistryHandle().IsValid());
+		return TraitRegistry.Find(Template.GetRegistryHandle());
+	}
+}
diff --git a/Engine/Plugins/Experimental/AnimNext/Source/AnimNext/Private/TraitCore/TraitBinding.cpp b/Engine/Plugins/Experimental/AnimNext/Source/AnimNext/Private/TraitCore/TraitBinding.cpp
new file mode 100644
--- /dev/null
+++ b/Engine/Plugins/Experimental/AnimNext/Source/AnimNext/Private/TraitCore/TraitBinding.cpp
@@ -0,0 +1,13 @@
+// Copyright Epic Games, Inc. All Rights Reserved.
+
+#include "TraitCore/TraitBinding.h"
+
+#include "TraitCore/ITraitInterface.h"
+
+namespace UE::AnimNext
+{
+	FTraitInterfaceUID FTraitBinding::GetInterfaceUID() const
+	{
+		return Interface != nullptr ? Interface->GetInterfaceUID() : FTraitInterfaceUID();
+	}
+}
diff --git a/Engine/Plugins/Experimental/AnimNext/Source/AnimNext/Private/TraitInterfaces/IGarbageColllection.cpp b/Engine/Plugins/Experimental/AnimNext/Source/AnimNext/Private/TraitInterfaces/IGarbageColllection.cpp
new file mode 100644
--- /dev/null
+++ b/Engine/Plugins/Experimental/AnimNext/Source/AnimNext/Private/TraitInterfaces/IGarbageColllection.cpp
@@ -0,0 +1,30 @@
+// Copyright Epic Games, Inc. All Rights Reserved.
+
+#include "TraitInterfaces/IGarbageCollection.h"
+
+#include "TraitCore/ExecutionContext.h"
+#include "Graph/GC_GraphInstanceComponent.h"
+
+namespace UE::AnimNext
+{
+	void IGarbageCollection::RegisterWithGC(const FExecutionContext& Context, const FTraitBinding& Binding)
+	{
+		FGCGraphInstanceComponent& Component = Context.GetComponent<FGCGraphInstanceComponent>();
+		Component.Register(Context.GetGraphInstance(), Binding.GetTraitPtr());
+	}
+
+	void IGarbageCollection::UnregisterWithGC(const FExecutionContext& Context, const FTraitBinding& Binding)
+	{
+		FGCGraphInstanceComponent& Component = Context.GetComponent<FGCGraphInstanceComponent>();
+		Component.Unregister(Binding.GetTraitPtr());
+	}
+
+	void IGarbageCollection::AddReferencedObjects(const FExecutionContext& Context, const TTraitBinding<IGarbageCollection>& Binding, FReferenceCollector& Collector) const
+	{
+		TTraitBinding<IGarbageCollection> SuperBinding;
+		if (Context.GetInterfaceSuper(Binding, SuperBinding))
+		{
+			SuperBinding.AddReferencedObjects(Context, Collector);
+		}
+	}
+}
diff --git a/Engine/Plugins/Experimental/AnimNext/Source/AnimNext/Private/TraitInterfaces/IHierarchy.cpp b/Engine/Plugins/Experimental/AnimNext/Source/AnimNext/Private/TraitInterfaces/IHierarchy.cpp
new file mode 100644
--- /dev/null
+++ b/Engine/Plugins/Experimental/AnimNext/Source/AnimNext/Private/TraitInterfaces/IHierarchy.cpp
@@ -0,0 +1,28 @@
+// Copyright Epic Games, Inc. All Rights Reserved.
+
+#include "TraitInterfaces/IHierarchy.h"
+
+#include "TraitCore/ExecutionContext.h"
+
+namespace UE::AnimNext
+{
+	uint32 IHierarchy::GetNumChildren(const FExecutionContext& Context, const TTraitBinding<IHierarchy>& Binding) const
+	{
+		TTraitBinding<IHierarchy> SuperBinding;
+		if (Context.GetInterfaceSuper(Binding, SuperBinding))
+		{
+			return SuperBinding.GetNumChildren(Context);
+		}
+
+		return 0;
+	}
+
+	void IHierarchy::GetChildren(const FExecutionContext& Context, const TTraitBinding<IHierarchy>& Binding, FChildrenArray& Children) const
+	{
+		TTraitBinding<IHierarchy> SuperBinding;
+		if (Context.GetInterfaceSuper(Binding, SuperBinding))
+		{
+			SuperBinding.GetChildren(Context, Children);
+		}
+	}
+}
diff --git a/Engine/Plugins/Experimental/AnimNext/Source/AnimNext/Public/DataRegistry.h b/Engine/Plugins/Experimental/AnimNext/Source/AnimNext/Public/DataRegistry.h
--- a/Engine/Plugins/Experimental/AnimNext/Source/AnimNext/Public/DataRegistry.h
+++ b/Engine/Plugins/Experimental/AnimNext/Source/AnimNext/Public/DataRegistry.h
@@ -76,27 +76,32 @@ public:
 	{
 		FParamTypeHandle ParamTypeHandle = FParamTypeHandle::GetHandle<DataType>();
 
-		const FDataTypeDef* TypeDef = nullptr;
+		bool bIsTypeDefValid = false;
+		FDataTypeDef TypeDef;
 		{
 			FRWScopeLock Lock(DataTypeDefsLock, SLT_ReadOnly);
-			TypeDef = DataTypeDefs.Find(ParamTypeHandle);
+			if (const FDataTypeDef* TypeDefPtr = DataTypeDefs.Find(ParamTypeHandle))
+			{
+				bIsTypeDefValid = true;
+				TypeDef = *TypeDefPtr;
+			}
 		}
 
-		if (TypeDef == nullptr)
+		if (!bIsTypeDefValid)
 		{
 			TypeDef = RegisterDataType_Impl<DataType>(DEFAULT_BLOCK_SIZE);
 			// TODO : Log if we allocate more than DEFAULT_BLOCK_SIZE elements of that type
 		}
 
-		if (ensure(TypeDef != nullptr && TypeDef->ParamTypeHandle.IsValid()))
+		if (ensure(TypeDef.ParamTypeHandle.IsValid()))
 		{
-			const int32 ElementSize = TypeDef->ElementSize;
-			const int32 ElementAlign = TypeDef->ElementAlign;
+			const int32 ElementSize = TypeDef.ElementSize;
+			const int32 ElementAlign = TypeDef.ElementAlign;
 			const int32 AlignedSize = Align(ElementSize, ElementAlign);
 
 			const int32 BufferSize = NumElements * AlignedSize;
 
-			uint8* Memory = (uint8*)FMemory::Malloc(BufferSize, TypeDef->ElementAlign);    // TODO : This should come from preallocated chunks, use malloc / free for now
+			uint8* Memory = (uint8*)FMemory::Malloc(BufferSize, TypeDef.ElementAlign);    // TODO : This should come from preallocated chunks, use malloc / free for now
 
 			Private::FAllocatedBlock* AllocatedBlock = new Private::FAllocatedBlock(Memory, NumElements, ParamTypeHandle); // TODO : avoid memory fragmentation
 			AllocatedBlock->AddRef();
@@ -181,7 +186,7 @@ private:
 
 	// Registers a type and sets the allocation block size
 	template<typename DataType>
-	FDataTypeDef* RegisterDataType_Impl(int32 AllocationBlockSize)
+	FDataTypeDef RegisterDataType_Impl(int32 AllocationBlockSize)
 	{
 		FParamTypeHandle ParamTypeHandle = FParamTypeHandle::GetHandle<DataType>();
 		check(ParamTypeHandle.IsValid());
@@ -199,15 +204,14 @@ private:
 			}
 		};
 
-		FDataTypeDef* AddedDef = nullptr;
 		{
 			FRWScopeLock WriteLock(DataTypeDefsLock, SLT_Write);
 
-			AddedDef = &DataTypeDefs.FindOrAdd(ParamTypeHandle, { ParamTypeHandle, DestroyFn, ElementSize, ElementAlign, AllocationBlockSize });
+			FDataTypeDef* AddedDef = &DataTypeDefs.FindOrAdd(ParamTypeHandle, { ParamTypeHandle, DestroyFn, ElementSize, ElementAlign, AllocationBlockSize });
 			check(AddedDef->ParamTypeHandle == ParamTypeHandle); // check we have not added two different types with the same ID
-		}
 
-		return AddedDef;
+			return *AddedDef;
+		}
 	}
 
 	void OnLODRequiredBonesUpdate(USkeletalMeshComponent* SkeletalMeshComponent, int32 LODLevel, const TArray<FBoneIndexType>& LODRequiredBones);
diff --git a/Engine/Plugins/Experimental/AnimNext/Source/AnimNext/Public/Graph/AnimNextGraphInstance.h b/Engine/Plugins/Experimental/AnimNext/Source/AnimNext/Public/Graph/AnimNextGraphInstance.h
--- a/Engine/Plugins/Experimental/AnimNext/Source/AnimNext/Public/Graph/AnimNextGraphInstance.h
+++ b/Engine/Plugins/Experimental/AnimNext/Source/AnimNext/Public/Graph/AnimNextGraphInstance.h
@@ -3,7 +3,7 @@
 #pragma once
 
 #include "CoreMinimal.h"
-#include "DecoratorBase/DecoratorPtr.h"
+#include "TraitCore/TraitPtr.h"
 #include "RigVMCore/RigVMExecuteContext.h"
 
 #include "AnimNextGraphInstance.generated.h"
@@ -53,8 +53,8 @@ struct FAnimNextGraphInstance
 	// Returns the entry point in Graph that this instance corresponds to 
 	FName GetEntryPoint() const;
 	
-	// Returns a weak handle to the root decorator instance
-	UE::AnimNext::FWeakDecoratorPtr GetGraphRootPtr() const;
+	// Returns a weak handle to the root trait instance
+	UE::AnimNext::FWeakTraitPtr GetGraphRootPtr() const;
 
 	// Returns the parent graph instance that owns us or nullptr for the root graph instance or if we are invalid
 	FAnimNextGraphInstance* GetParentGraphInstance() const;
@@ -115,7 +115,7 @@ private:
 	FName EntryPoint;
 
 	// Hard reference to the graph instance data, we own it
-	UE::AnimNext::FDecoratorPtr GraphInstancePtr;
+	UE::AnimNext::FTraitPtr GraphInstancePtr;
 
 	// The graph instance that owns us
 	FAnimNextGraphInstance* ParentGraphInstance = nullptr;
diff --git a/Engine/Plugins/Experimental/AnimNext/Source/AnimNext/Public/Graph/AnimNextGraphInstancePtr.h b/Engine/Plugins/Experimental/AnimNext/Source/AnimNext/Public/Graph/AnimNextGraphInstancePtr.h
--- a/Engine/Plugins/Experimental/AnimNext/Source/AnimNext/Public/Graph/AnimNextGraphInstancePtr.h
+++ b/Engine/Plugins/Experimental/AnimNext/Source/AnimNext/Public/Graph/AnimNextGraphInstancePtr.h
@@ -3,7 +3,7 @@
 #pragma once
 
 #include "CoreMinimal.h"
-#include "DecoratorBase/DecoratorPtr.h"
+#include "TraitCore/TraitPtr.h"
 
 #include "AnimNextGraphInstancePtr.generated.h"
 
@@ -50,8 +50,8 @@ struct ANIMNEXT_API FAnimNextGraphInstancePtr
 	// Returns the graph used by this instance or nullptr if the instance is invalid
 	const UAnimNextGraph* GetGraph() const;
 
-	// Returns a weak handle to the root decorator instance
-	UE::AnimNext::FWeakDecoratorPtr GetGraphRootPtr() const;
+	// Returns a weak handle to the root trait instance
+	UE::AnimNext::FWeakTraitPtr GetGraphRootPtr() const;
 
 	// Returns the graph instance implementation
 	FAnimNextGraphInstance* GetImpl() const;
diff --git a/Engine/Plugins/Experimental/AnimNext/Source/AnimNext/Public/TraitCore/ExecutionContext.h b/Engine/Plugins/Experimental/AnimNext/Source/AnimNext/Public/TraitCore/ExecutionContext.h
new file mode 100644
--- /dev/null
+++ b/Engine/Plugins/Experimental/AnimNext/Source/AnimNext/Public/TraitCore/ExecutionContext.h
@@ -0,0 +1,191 @@
+// Copyright Epic Games, Inc. All Rights Reserved.
+
+#pragma once
+
+#include "CoreMinimal.h"
+#include "TraitCore/TraitBinding.h"
+#include "TraitCore/TraitHandle.h"
+#include "TraitCore/TraitPtr.h"
+#include "TraitCore/TraitInterfaceUID.h"
+#include "TraitCore/LatentPropertyHandle.h"
+#include "TraitCore/NodeHandle.h"
+#include "Graph/AnimNextGraphInstancePtr.h"
+
+struct FAnimNextGraphInstance;
+
+namespace UE::AnimNext
+{
+	struct FNodeDescription;
+	struct FNodeInstance;
+	struct FNodeTemplateRegistry;
+	struct FNodeTemplate;
+	struct FTrait;
+	struct FTraitRegistry;
+	struct FTraitTemplate;
+
+	/**
+	 * Execution Context
+	 * 
+	 * The execution context aims to centralize the trait query API.
+	 * It is meant to be bound to a graph instance and re-used by the nodes/traits within.
+	 */
+	struct ANIMNEXT_API FExecutionContext
+	{
+		// Creates an unbound execution context
+		FExecutionContext();
+
+		// Creates an execution context and binds it to the specified graph instance
+		explicit FExecutionContext(FAnimNextGraphInstancePtr& InGraphInstance);
+
+		// Creates an execution context and binds it to the specified graph instance
+		explicit FExecutionContext(FAnimNextGraphInstance& InGraphInstance);
+
+		// Binds the execution context to the specified graph instance if it differs from the currently bound instance
+		void BindTo(FAnimNextGraphInstancePtr& InGraphInstance);
+
+		// Binds the execution context to the specified graph instance if it differs from the currently bound instance
+		void BindTo(FAnimNextGraphInstance& InGraphInstance);
+
+		// Binds the execution context to the graph instance that owns the specified trait if it differs from the currently bound instance
+		void BindTo(const FWeakTraitPtr& TraitPtr);
+
+		// Returns whether or not this execution context is bound to a graph instance
+		bool IsBound() const;
+
+		// Returns whether or not this execution context is bound to the specified graph instance
+		bool IsBoundTo(const FAnimNextGraphInstancePtr& InGraphInstance) const;
+
+		// Returns whether or not this execution context is bound to the specified graph instance
+		bool IsBoundTo(const FAnimNextGraphInstance& InGraphInstance) const;
+
+		// Queries a node for a trait that implements the specified interface.
+		// If no such trait exists, nullptr is returned.
+		template<class TraitInterface>
+		bool GetInterface(const FWeakTraitPtr& TraitPtr, TTraitBinding<TraitInterface>& InterfaceBinding) const;
+
+		// Queries a node for a trait that implements the specified interface.
+		// If no such trait exists, nullptr is returned.
+		template<class TraitInterface>
+		bool GetInterface(const FTraitBinding& Binding, TTraitBinding<TraitInterface>& InterfaceBinding) const;
+
+		// Queries a node for a trait lower on the stack that implements the specified interface.
+		// If no such trait exists, nullptr is returned.
+		template<class TraitInterface>
+		bool GetInterfaceSuper(const FWeakTraitPtr& TraitPtr, TTraitBinding<TraitInterface>& SuperBinding) const;
+
+		// Queries a node for a trait lower on the stack that implements the specified interface.
+		// If no such trait exists, nullptr is returned.
+		template<class TraitInterface>
+		bool GetInterfaceSuper(const FTraitBinding& Binding, TTraitBinding<TraitInterface>& SuperBinding) const;
+
+		// Allocates a new node instance from a trait handle
+		// If the desired trait lives in the current parent, a weak handle to it will be returned
+		FTraitPtr AllocateNodeInstance(const FTraitBinding& ParentBinding, FAnimNextTraitHandle ChildTraitHandle) const;
+
+		// Allocates a new node instance from a trait handle
+		// If the desired trait lives in the current parent, a weak handle to it will be returned
+		FTraitPtr AllocateNodeInstance(const FWeakTraitPtr& ParentBinding, FAnimNextTraitHandle ChildTraitHandle) const;
+
+		// Decrements the reference count of the provided node pointer and releases it if
+		// there are no more references remaining, reseting the pointer in the process
+		void ReleaseNodeInstance(FTraitPtr& NodePtr) const;
+
+		// Takes a snapshot of all latent properties on the provided node sub-stack (all traits on the sub-stack of the provided one)
+		// Properties can be marked as always updating or as supporting freezing (e.g. when a branch of the graph blends out)
+		// A freezable property does not update when a snapshot is taken of a frozen node
+		void SnapshotLatentProperties(const FWeakTraitPtr& TraitPtr, bool bIsFrozen) const;
+
+		// Returns a typed graph instance component, creating it lazily the first time it is queried
+		template<class ComponentType>
+		ComponentType& GetComponent() const;
+
+		// Returns a typed graph instance component pointer if found or nullptr otherwise
+		template<class ComponentType>
+		ComponentType* TryGetComponent() const;
+
+		// Returns const iterators to the graph instance component container
+		GraphInstanceComponentMapType::TConstIterator GetComponentIterator() const;
+
+		// Returns the bound graph instance
+		FAnimNextGraphInstance& GetGraphInstance() const { return *GraphInstance; }
+
+	private:
+		// No copy or move
+		FExecutionContext(const FExecutionContext&) = delete;
+		FExecutionContext& operator=(const FExecutionContext&) = delete;
+
+		bool GetInterfaceImpl(FTraitInterfaceUID InterfaceUID, const FWeakTraitPtr& TraitPtr, FTraitBinding& InterfaceBinding) const;
+		bool GetInterfaceSuperImpl(FTraitInterfaceUID InterfaceUID, const FWeakTraitPtr& TraitPtr, FTraitBinding& SuperBinding) const;
+		FGraphInstanceComponent* TryGetComponent(int32 ComponentNameHash, FName ComponentName) const;
+		FGraphInstanceComponent& AddComponent(int32 ComponentNameHash, FName ComponentName, TSharedPtr<FGraphInstanceComponent>&& Component) const;
+
+		const FNodeDescription& GetNodeDescription(FNodeHandle NodeHandle) const;
+		const FNodeTemplate* GetNodeTemplate(const FNodeDescription& NodeDesc) const;
+		const FTrait* GetTrait(const FTraitTemplate& TraitDesc) const;
+
+		// Cached references to the registries we need
+		const FNodeTemplateRegistry& NodeTemplateRegistry;
+		const FTraitRegistry& TraitRegistry;
+
+		// Cached properties for the currently executing graph
+		FAnimNextGraphInstance* GraphInstance = nullptr;
+		TArrayView<const uint8> GraphSharedData;
+	};
+
+	//////////////////////////////////////////////////////////////////////////
+	// Inline implementations
+
+	template<class TraitInterface>
+	inline bool FExecutionContext::GetInterface(const FWeakTraitPtr& TraitPtr, TTraitBinding<TraitInterface>& InterfaceBinding) const
+	{
+		constexpr FTraitInterfaceUID InterfaceUID = TraitInterface::InterfaceUID;
+		return GetInterfaceImpl(InterfaceUID, TraitPtr, InterfaceBinding);
+	}
+
+	template<class TraitInterface>
+	inline bool FExecutionContext::GetInterface(const FTraitBinding& Binding, TTraitBinding<TraitInterface>& InterfaceBinding) const
+	{
+		return GetInterface<TraitInterface>(Binding.GetTraitPtr(), InterfaceBinding);
+	}
+
+	template<class TraitInterface>
+	inline bool FExecutionContext::GetInterfaceSuper(const FWeakTraitPtr& TraitPtr, TTraitBinding<TraitInterface>& SuperBinding) const
+	{
+		constexpr FTraitInterfaceUID InterfaceUID = TraitInterface::InterfaceUID;
+		return GetInterfaceSuperImpl(InterfaceUID, TraitPtr, SuperBinding);
+	}
+
+	template<class TraitInterface>
+	inline bool FExecutionContext::GetInterfaceSuper(const FTraitBinding& Binding, TTraitBinding<TraitInterface>& SuperBinding) const
+	{
+		return GetInterfaceSuper<TraitInterface>(Binding.GetTraitPtr(), SuperBinding);
+	}
+
+	inline FTraitPtr FExecutionContext::AllocateNodeInstance(const FTraitBinding& ParentBinding, FAnimNextTraitHandle ChildTraitHandle) const
+	{
+		return AllocateNodeInstance(ParentBinding.GetTraitPtr(), ChildTraitHandle);
+	}
+
+	template<class ComponentType>
+	ComponentType& FExecutionContext::GetComponent() const
+	{
+		const FName ComponentName = ComponentType::StaticComponentName();
+		const int32 ComponentNameHash = GetTypeHash(ComponentName);
+
+		if (FGraphInstanceComponent* Component = TryGetComponent(ComponentNameHash, ComponentName))
+		{
+			return *static_cast<ComponentType*>(Component);
+		}
+
+		return static_cast<ComponentType&>(AddComponent(ComponentNameHash, ComponentName, MakeShared<ComponentType>()));
+	}
+
+	template<class ComponentType>
+	ComponentType* FExecutionContext::TryGetComponent() const
+	{
+		const FName ComponentName = ComponentType::StaticComponentName();
+		const int32 ComponentNameHash = GetTypeHash(ComponentName);
+
+		return static_cast<ComponentType*>(TryGetComponent(ComponentNameHash, ComponentName));
+	}
+}
diff --git a/Engine/Plugins/Experimental/AnimNext/Source/AnimNext/Public/TraitCore/ExecutionContextProxy.h b/Engine/Plugins/Experimental/AnimNext/Source/AnimNext/Public/TraitCore/ExecutionContextProxy.h
new file mode 100644
--- /dev/null
+++ b/Engine/Plugins/Experimental/AnimNext/Source/AnimNext/Public/TraitCore/ExecutionContextProxy.h
@@ -0,0 +1,137 @@
+// Copyright Epic Games, Inc. All Rights Reserved.
+
+#pragma once
+
+#include "CoreMinimal.h"
+#include "TraitCore/ExecutionContext.h"
+
+namespace UE::AnimNext
+{
+	/**
+	 * FExecutionContextProxy
+	 *
+	 * Wraps an execution context and forwards all calls to it by mirroring its API.
+	 */
+	struct FExecutionContextProxy
+	{
+		explicit FExecutionContextProxy(const FExecutionContext& InExecutionContext)
+			: ExecutionContext(InExecutionContext)
+		{
+		}
+
+		// We safely coerce to our wrapped execution context to allow identical usage
+		operator const FExecutionContext& () const
+		{
+			return ExecutionContext;
+		}
+
+		// Returns whether or not this execution context is bound to a graph instance
+		bool IsBound() const
+		{
+			return ExecutionContext.IsBound();
+		}
+
+		// Returns whether or not this execution context is bound to the specified graph instance
+		bool IsBoundTo(const FAnimNextGraphInstancePtr& InGraphInstance) const
+		{
+			return ExecutionContext.IsBoundTo(InGraphInstance);
+		}
+
+		// Returns whether or not this execution context is bound to the specified graph instance
+		bool IsBoundTo(const FAnimNextGraphInstance& InGraphInstance) const
+		{
+			return ExecutionContext.IsBoundTo(InGraphInstance);
+		}
+
+		// Queries a node for a trait that implements the specified interface.
+		// If no such trait exists, nullptr is returned.
+		template<class TraitInterface>
+		bool GetInterface(const FWeakTraitPtr& TraitPtr, TTraitBinding<TraitInterface>& InterfaceBinding) const
+		{
+			return ExecutionContext.GetInterface<TraitInterface>(TraitPtr, InterfaceBinding);
+		}
+
+		// Queries a node for a trait that implements the specified interface.
+		// If no such trait exists, nullptr is returned.
+		template<class TraitInterface>
+		bool GetInterface(const FTraitBinding& Binding, TTraitBinding<TraitInterface>& InterfaceBinding) const
+		{
+			return ExecutionContext.GetInterface<TraitInterface>(Binding, InterfaceBinding);
+		}
+
+		// Queries a node for a trait lower on the stack that implements the specified interface.
+		// If no such trait exists, nullptr is returned.
+		template<class TraitInterface>
+		bool GetInterfaceSuper(const FWeakTraitPtr& TraitPtr, TTraitBinding<TraitInterface>& SuperBinding) const
+		{
+			return ExecutionContext.GetInterfaceSuper<TraitInterface>(TraitPtr, SuperBinding);
+		}
+
+		// Queries a node for a trait lower on the stack that implements the specified interface.
+		// If no such trait exists, nullptr is returned.
+		template<class TraitInterface>
+		bool GetInterfaceSuper(const FTraitBinding& Binding, TTraitBinding<TraitInterface>& SuperBinding) const
+		{
+			return ExecutionContext.GetInterfaceSuper<TraitInterface>(Binding, SuperBinding);
+		}
+
+		// Allocates a new node instance from a trait handle
+		// If the desired trait lives in the current parent, a weak handle to it will be returned
+		FTraitPtr AllocateNodeInstance(const FTraitBinding& ParentBinding, FAnimNextTraitHandle ChildTraitHandle) const
+		{
+			return ExecutionContext.AllocateNodeInstance(ParentBinding, ChildTraitHandle);
+		}
+
+		// Allocates a new node instance from a trait handle
+		// If the desired trait lives in the current parent, a weak handle to it will be returned
+		FTraitPtr AllocateNodeInstance(const FWeakTraitPtr& ParentBinding, FAnimNextTraitHandle ChildTraitHandle) const
+		{
+			return ExecutionContext.AllocateNodeInstance(ParentBinding, ChildTraitHandle);
+		}
+
+		// Decrements the reference count of the provided node pointer and releases it if
+		// there are no more references remaining, reseting the pointer in the process
+		void ReleaseNodeInstance(FTraitPtr& NodePtr) const
+		{
+			return ExecutionContext.ReleaseNodeInstance(NodePtr);
+		}
+
+		// Takes a snapshot of all latent properties on the provided node sub-stack (all traits on the sub-stack of the provided one)
+		// Properties can be marked as always updating or as supporting freezing (e.g. when a branch of the graph blends out)
+		// A freezable property does not update when a snapshot is taken of a frozen node
+		void SnapshotLatentProperties(const FWeakTraitPtr& TraitPtr, bool bIsFrozen) const
+		{
+			ExecutionContext.SnapshotLatentProperties(TraitPtr, bIsFrozen);
+		}
+
+		// Returns a typed graph instance component, creating it lazily the first time it is queried
+		template<class ComponentType>
+		ComponentType& GetComponent() const
+		{
+			return ExecutionContext.GetComponent<ComponentType>();
+		}
+
+		// Returns a typed graph instance component pointer if found or nullptr otherwise
+		template<class ComponentType>
+		ComponentType* TryGetComponent() const
+		{
+			return ExecutionContext.TryGetComponent<ComponentType>();
+		}
+
+		// Returns const iterators to the graph instance component container
+		GraphInstanceComponentMapType::TConstIterator GetComponentIterator() const
+		{
+			return ExecutionContext.GetComponentIterator();
+		}
+
+		// Returns the bound graph instance
+		FAnimNextGraphInstance& GetGraphInstance() const
+		{
+			return ExecutionContext.GetGraphInstance();
+		}
+
+	private:
+		// The execution context that we wrap
+		const FExecutionContext& ExecutionContext;
+	};
+}
diff --git a/Engine/Plugins/Experimental/AnimNext/Source/AnimNext/Public/TraitCore/ITraitInterface.h b/Engine/Plugins/Experimental/AnimNext/Source/AnimNext/Public/TraitCore/ITraitInterface.h
new file mode 100644
--- /dev/null
+++ b/Engine/Plugins/Experimental/AnimNext/Source/AnimNext/Public/TraitCore/ITraitInterface.h
@@ -0,0 +1,36 @@
+// Copyright Epic Games, Inc. All Rights Reserved.
+
+#pragma once
+
+#include "CoreMinimal.h"
+#include "TraitCore/TraitBinding.h"			// Derived types will need to implement the binding specialization
+#include "TraitCore/TraitInterfaceUID.h"
+
+// Helper macros
+#define DECLARE_ANIM_TRAIT_INTERFACE(InterfaceName, InterfaceNameHash) \
+	/* Globally unique UID for this interface */ \
+	static constexpr UE::AnimNext::FTraitInterfaceUID InterfaceUID = UE::AnimNext::FTraitInterfaceUID(InterfaceNameHash, TEXT(#InterfaceName)); \
+	virtual UE::AnimNext::FTraitInterfaceUID GetInterfaceUID() const override { return InterfaceUID; }
+
+namespace UE::AnimNext
+{
+	struct FExecutionContext;	// Derived types will have functions that accept the execution context
+
+	/**
+	 * ITraitInterface
+	 * 
+	 * Base type for all trait interfaces. Used for type safety.
+	 */
+	struct ANIMNEXT_API ITraitInterface
+	{
+		virtual ~ITraitInterface() {}
+
+		// The globally unique UID for this interface
+		// Derived types will have their own InterfaceUID member that hides/aliases/shadows this one
+		// @see DECLARE_ANIM_TRAIT_INTERFACE
+		static constexpr FTraitInterfaceUID InterfaceUID = FTraitInterfaceUID(0xcdc27733, TEXT("ITraitInterface"));
+
+		// Returns the globally unique UID for this interface
+		virtual FTraitInterfaceUID GetInterfaceUID() const { return InterfaceUID; };
+	};
+}
