Migrate away from the 'decorator' term in AnimNext in favor of 'trait'

diff --git a/Engine/Plugins/Experimental/AnimNext/Source/AnimNext/Public/TraitCore/Trait.h b/Engine/Plugins/Experimental/AnimNext/Source/AnimNext/Public/TraitCore/Trait.h
new file mode 100644
--- /dev/null
+++ b/Engine/Plugins/Experimental/AnimNext/Source/AnimNext/Public/TraitCore/Trait.h
@@ -0,0 +1,330 @@
+// Copyright Epic Games, Inc. All Rights Reserved.
+
+#pragma once
+
+#include "CoreMinimal.h"
+
+#include "TraitCore/TraitBinding.h"
+#include "TraitCore/TraitHandle.h"			// Derived types are likely to refer to other traits as children
+#include "TraitCore/TraitInstanceData.h"
+#include "TraitCore/TraitMode.h"
+#include "TraitCore/TraitPtr.h"
+#include "TraitCore/TraitSharedData.h"
+#include "TraitCore/TraitUID.h"
+#include "TraitCore/ITraitInterface.h"
+#include "TraitCore/LatentPropertyHandle.h"
+
+#include <type_traits>
+
+class FArchive;
+
+// Helper macros
+// In the trait class declaration, this macro declares the Super alias and base functions we override
+#define DECLARE_ANIM_TRAIT(TraitName, TraitNameHash, SuperTraitName) \
+	using TraitSuper = SuperTraitName; \
+	/* FTrait impl */ \
+	static constexpr UE::AnimNext::FTraitUID TraitUID = UE::AnimNext::FTraitUID(TraitNameHash, TEXT(#TraitName)); \
+	virtual UE::AnimNext::FTraitUID GetTraitUID() const override { return TraitUID; } \
+	virtual FString GetTraitName() const override { return TEXT(#TraitName); } \
+	static const UE::AnimNext::FTraitMemoryLayout TraitMemoryDescription; \
+	virtual UE::AnimNext::FTraitMemoryLayout GetTraitMemoryDescription() const override { return TraitMemoryDescription; } \
+	virtual UScriptStruct* GetTraitSharedDataStruct() const override { return FSharedData::StaticStruct(); } \
+	virtual void ConstructTraitInstance(const UE::AnimNext::FExecutionContext& Context, const UE::AnimNext::FTraitBinding& Binding) const override; \
+	virtual void DestructTraitInstance(const UE::AnimNext::FExecutionContext& Context, const UE::AnimNext::FTraitBinding& Binding) const override; \
+	virtual const UE::AnimNext::ITraitInterface* GetTraitInterface(UE::AnimNext::FTraitInterfaceUID InterfaceUID) const override; \
+	virtual TConstArrayView<FTraitInterfaceUID> GetTraitInterfaces() const override; \
+	virtual uint32 GetNumLatentTraitProperties() const override { return -FSharedData::GetLatentPropertyIndex(~(size_t)0); } \
+	virtual FTraitLatentPropertyMemoryLayout GetLatentPropertyMemoryLayout(FName PropertyName, uint32 PropertyIndex) const override; \
+	static_assert(std::is_base_of<FAnimNextTraitSharedData, FSharedData>::value, "Trait shared data must derive from FAnimNextTraitSharedData"); \
+	static_assert(std::is_base_of<FTraitInstanceData, FInstanceData>::value, "Trait instance data must derive from FTraitInstanceData");
+
+#define DECLARE_ABSTRACT_ANIM_TRAIT(TraitName, TraitNameHash, SuperTraitName) \
+	using TraitSuper = SuperTraitName; \
+	/* FTrait impl */ \
+	static constexpr UE::AnimNext::FTraitUID TraitUID = UE::AnimNext::FTraitUID(TraitNameHash, TEXT(#TraitName)); \
+	virtual UE::AnimNext::FTraitUID GetTraitUID() const override { return TraitUID; } \
+	virtual FString GetTraitName() const override { return TEXT(#TraitName); }
+
+// In the trait cpp, these three macros implement the base functionality
+// 
+// Usage is as follow:
+// #define TRAIT_INTERFACE_ENUMERATOR(GeneratorMacro) \
+//		GeneratorMacro(IHierarchy) \
+//		GeneratorMacro(IUpdate) \
+//
+// GENERATE_ANIM_TRAIT_IMPLEMENTATION(FMyTrait, TRAIT_INTERFACE_ENUMERATOR)
+// #undef TRAIT_INTERFACE_ENUMERATOR
+
+// Implements various parts of FTrait
+#define ANIM_NEXT_IMPL_DEFINE_ANIM_TRAIT(TraitName) \
+	const UE::AnimNext::FTraitMemoryLayout TraitName::TraitMemoryDescription = \
+		UE::AnimNext::FTraitMemoryLayout{ sizeof(TraitName), alignof(TraitName), sizeof(TraitName::FSharedData), alignof(TraitName::FSharedData), sizeof(TraitName::FInstanceData), alignof(TraitName::FInstanceData) }; \
+	FTraitLatentPropertyMemoryLayout TraitName::GetLatentPropertyMemoryLayout(FName PropertyName, uint32 PropertyIndex) const \
+	{ \
+		/* Thread safe cache initialization */ \
+		static TArray<FTraitLatentPropertyMemoryLayout> CachedLatentPropertyMemoryLayouts = [this](){ TArray<FTraitLatentPropertyMemoryLayout> Result; Result.SetNum(GetNumLatentTraitProperties()); return Result; }(); \
+		return GetLatentPropertyMemoryLayoutImpl(PropertyName, PropertyIndex, CachedLatentPropertyMemoryLayouts); \
+	} \
+	void TraitName::ConstructTraitInstance(const UE::AnimNext::FExecutionContext& Context, const UE::AnimNext::FTraitBinding& Binding) const \
+	{ \
+		FInstanceData* Data = new(Binding.GetInstanceData<FInstanceData>()) FInstanceData(); \
+		Data->Construct(Context, Binding); \
+	} \
+	void TraitName::DestructTraitInstance(const UE::AnimNext::FExecutionContext& Context, const UE::AnimNext::FTraitBinding& Binding) const \
+	{ \
+		FInstanceData* Data = Binding.GetInstanceData<FInstanceData>(); \
+		Data->Destruct(Context, Binding); \
+		Data->~FInstanceData(); \
+	}
+
+// Helper that handles the GetTraitInterface() details for each interface specified by the generator macro
+#define ANIM_NEXT_IMPL_GET_INTERFACE_IMPL_FOR_INTERFACE(InterfaceName) \
+	if (InInterfaceUID == InterfaceName::InterfaceUID) \
+	{ \
+		return static_cast<const InterfaceName*>(this); \
+	}
+
+// Implements GetTraitInterface()
+#define ANIM_NEXT_IMPL_DEFINE_ANIM_TRAIT_GET_INTERFACE(TraitName, InterfaceEnumeratorMacro) \
+	const UE::AnimNext::ITraitInterface* TraitName::GetTraitInterface(UE::AnimNext::FTraitInterfaceUID InInterfaceUID) const \
+	{ \
+		InterfaceEnumeratorMacro(ANIM_NEXT_IMPL_GET_INTERFACE_IMPL_FOR_INTERFACE) \
+		/* Forward to base implementation */ \
+		return TraitSuper::GetTraitInterface(InInterfaceUID); \
+	}
+
+// Helper that handles the GetTraitInterfaces() details for each interface specified by the generator macro
+#define ANIM_NEXT_IMPL_GET_INTERFACES_IMPL_FOR_INTERFACE(InterfaceName) InterfaceName::InterfaceUID,
+
+// Implements GetTraitInterfaces()
+#define ANIM_NEXT_IMPL_DEFINE_ANIM_TRAIT_GET_INTERFACES(TraitName, InterfaceEnumeratorMacro) \
+	TConstArrayView<UE::AnimNext::FTraitInterfaceUID> TraitName::GetTraitInterfaces() const \
+	{ \
+		/* Thread safe cache initialization */ \
+		static TArray<UE::AnimNext::FTraitInterfaceUID> CachedInterfaceList = FTrait::BuildTraitInterfaceList( \
+		TraitSuper::GetTraitInterfaces(), \
+		{ \
+			InterfaceEnumeratorMacro(ANIM_NEXT_IMPL_GET_INTERFACES_IMPL_FOR_INTERFACE) \
+		}); \
+		return CachedInterfaceList; \
+	}
+
+/**
+  * This macro defines the necessary boilerplate for implementing FTrait. See above for usage example.
+  */
+#define GENERATE_ANIM_TRAIT_IMPLEMENTATION(TraitName, InterfaceEnumeratorMacro) \
+	ANIM_NEXT_IMPL_DEFINE_ANIM_TRAIT(TraitName) \
+	ANIM_NEXT_IMPL_DEFINE_ANIM_TRAIT_GET_INTERFACE(TraitName, InterfaceEnumeratorMacro) \
+	ANIM_NEXT_IMPL_DEFINE_ANIM_TRAIT_GET_INTERFACES(TraitName, InterfaceEnumeratorMacro)
+
+// Allows a trait to auto-register and unregister within the current execution scope
+// The trait must be found in the current scope without a namespace qualification
+#define AUTO_REGISTER_ANIM_TRAIT(TraitName) \
+	UE::AnimNext::FTraitStaticInitHook TraitName##Hook( \
+		[](void* DestPtr, UE::AnimNext::FTraitMemoryLayout& MemoryDesc) -> UE::AnimNext::FTrait* \
+		{ \
+			MemoryDesc = TraitName::TraitMemoryDescription; \
+			return DestPtr != nullptr ? new(DestPtr) TraitName() : nullptr; \
+		});
+
+namespace UE::AnimNext
+{
+	struct FTrait;
+	struct FTraitMemoryLayout;
+	class FTraitReader;
+	class FTraitWriter;
+	struct FExecutionContext;
+
+	// A function pointer to a shim to construct a trait into the desired memory location
+	// When called with a nullptr DestPtr, the function returns nullptr and only populates the
+	// memory description output argument. This allows the caller to determine how much space
+	// to reserve and how to properly align it. This is similar in spirit to various Windows SDK functions.
+	using TraitConstructorFunc = FTrait * (*)(void* DestPtr, FTraitMemoryLayout& MemoryDesc);
+
+	/**
+	 * FTraitMemoryLayout
+	 * 
+	 * Encapsulates size/alignment details for a trait.
+	 */
+	struct FTraitMemoryLayout
+	{
+		// The size in bytes of an instance of the trait class which derives from FTrait
+		uint32 TraitSize = 0;
+
+		// The alignment in bytes of an instance of the trait class which derives from FTrait
+		uint32 TraitAlignment = 1;
+
+		// The size in bytes of the shared data for the trait which derives from FAnimNextTraitSharedData
+		uint32 SharedDataSize = 0;
+
+		// The alignment in bytes of the shared data for the trait which derives from FAnimNextTraitSharedData
+		uint32 SharedDataAlignment = 1;
+
+		// The size in bytes of the instance data for the trait which derives from FTraitInstanceData
+		uint32 InstanceDataSize = 0;
+
+		// The alignment in bytes of the instance data for the trait which derives from FTraitInstanceData
+		uint32 InstanceDataAlignment = 1;
+	};
+
+	/**
+	 * FTraitLatentPropertyMemoryLayout
+	 *
+	 * Encapsulates size/alignment details for a latent property.
+	 */
+	struct FTraitLatentPropertyMemoryLayout
+	{
+		// The size in bytes of the latent property
+		uint32 Size = 0;
+
+		// The alignment in bytes of the latent property
+		uint32 Alignment = 1;
+	};
+
+	/**
+	 * FTrait
+	 * 
+	 * Base class for all traits.
+	 * A trait can implement any number of interfaces based on ITraitInterface.
+	 * A trait may derive from another trait.
+	 * A trait should implement GetInterface(..) and test against the interfaces that it supports.
+	 * 
+	 * Traits should NOT have any internal state, hence why all API functions are 'const'.
+	 * The reason for this is that at runtime, a single instance of every trait exists.
+	 * That single instance is used by all instances of a trait on a node and concurrently
+	 * on all worker threads.
+	 * 
+	 * Traits can have shared read-only data that all instances of a graph can use (e.g. hard-coded properties).
+	 * Shared data must derive from FAnimNextTraitSharedData.
+	 * Traits can have instance data (e.g. blend weight).
+	 * Instance data must derive from FTraitInstanceData.
+	 */
+	struct ANIMNEXT_API FTrait
+	{
+		virtual ~FTrait() {}
+
+		// Empty shared/instance data types
+		// Derived types must define an alias for these
+		using FSharedData = FAnimNextTraitSharedData;
+		using FInstanceData = FTraitInstanceData;
+
+		// The globally unique UID for this trait
+		// Derived types will have their own TraitUID member that hides/aliases/shadows this one
+		// @see DECLARE_ANIM_TRAIT
+		static constexpr FTraitUID TraitUID = FTraitUID(0x4d97bae7, TEXT("FTrait"));
+
+		// Returns the globally unique UID for this trait
+		virtual FTraitUID GetTraitUID() const { return TraitUID; };
+
+		// Returns the trait name
+		virtual FString GetTraitName() const { return TEXT("FTrait"); }
+
+		// Returns the memory requirements of the derived trait instance
+		virtual FTraitMemoryLayout GetTraitMemoryDescription() const = 0;
+
+		// Returns the UScriptStruct associated with the shared data for the trait
+		virtual UScriptStruct* GetTraitSharedDataStruct() const { return FSharedData::StaticStruct(); }
+
+		// Called when a new instance of the trait is created or destroyed
+		// Derived types must override this and forward to the instance data constructor/destructor
+		virtual void ConstructTraitInstance(const FExecutionContext& Context, const FTraitBinding& Binding) const = 0;
+		virtual void DestructTraitInstance(const FExecutionContext& Context, const FTraitBinding& Binding) const = 0;
+
+		// Returns the trait mode.
+		virtual ETraitMode GetTraitMode() const = 0;
+
+		// Returns a pointer to the specified interface if it is supported.
+		// Derived types must override this.
+		virtual const ITraitInterface* GetTraitInterface(FTraitInterfaceUID InterfaceUID) const
+		{
+			// TODO:
+			// if/else sequence with static_casts to get the right v-table
+			// could be implemented with two tables: one of UIDs, another with matching offsets to 'this'
+			// we could scan the first table with SIMD, 4x UIDs at a time with 'cmpeq' to generate a mask
+			// we can mode the mask into a general register, if non-zero, we have a match
+			// using the mask, we can easily compute the UID offset in our 4x entry by counting leading/trailing zeroes
+			// using the UID offset, we can load and apply the correct offset
+			// may or may not be faster, but it shifts the burden from code cache to data cache and we can better control locality
+			// we could store the tables contiguous with one another, offsets could be 16 bit or maybe 8 bit (multiple of pointer size)
+			// we could store the tables contiguous with the tables of other traits for better cache locality
+			// by using tables, it means the lookup code can live in a single place and remain hot
+			// it means we can test 4x UIDs at a time, or interleave and test 8x or 16x
+			// it means we can quickly early out if none of the interfaces match (common case?) since we don't need to test
+			// all of them one by one
+			// SIMD code path also opens the door for cheap bulk interface queries where we query up to 4x interface UIDs and
+			// return 4x interface offsets (caller can generate pointers easily)
+
+			// Base class doesn't implement any interfaces
+			// Derived types must implement this
+			return nullptr;
+		}
+
+		// Returns a list of interfaces that this trait supports
+		virtual TConstArrayView<FTraitInterfaceUID> GetTraitInterfaces() const { return TConstArrayView<FTraitInterfaceUID>(); }
+
+		// The number of latent property properties in the shared data of this trait
+		virtual uint32 GetNumLatentTraitProperties() const { return 0; }
+
+		// Returns the memory layout of the specified latent property
+		virtual FTraitLatentPropertyMemoryLayout GetLatentPropertyMemoryLayout(FName PropertyName, uint32 PropertyIndex) const { return FTraitLatentPropertyMemoryLayout(); }
+
+		// Called to serialize trait shared data
+		virtual void SerializeTraitSharedData(FArchive& Ar, FAnimNextTraitSharedData& SharedData) const;
+
+#if WITH_EDITOR
+		// Takes the editor properties as authored in the graph and converts them into an instance of the FAnimNextTraitSharedData
+		// derived type using UE reflection.
+		// Traits can override this function to control how editor only properties are coerced into the runtime shared data
+		// instance.
+		virtual void SaveTraitSharedData(const TFunction<FString(FName PropertyName)>& GetTraitProperty, FAnimNextTraitSharedData& OutSharedData) const;
+
+		// Takes the editor properties as authored in the graph and returns the latent property metadata using UE reflection.
+		TArray<FLatentPropertyMetadata> GetLatentPropertyHandles(
+			bool bFilterEditorOnly,
+			const TFunction<uint16(FName PropertyName)>& GetTraitLatentPropertyIndex) const;
+#endif
+
+	protected:
+		// Implements GetLatentPropertyMemoryLayout() by allowing a map to be provided for caching purposes to speed up repeated queries
+		FTraitLatentPropertyMemoryLayout GetLatentPropertyMemoryLayoutImpl(
+			FName PropertyName,
+			uint32 PropertyIndex,
+			TArray<FTraitLatentPropertyMemoryLayout>& LatentPropertyMemoryLayouts) const;
+
+		// Builds a list of interfaces with the provided super interfaces and current interfaces as an initializer list
+		static TArray<FTraitInterfaceUID> BuildTraitInterfaceList(
+			const TConstArrayView<FTraitInterfaceUID>& SuperInterfaces,
+			std::initializer_list<FTraitInterfaceUID> InterfaceList);
+	};
+
+	// Base class for base traits that are standalone
+	struct ANIMNEXT_API FBaseTrait : FTrait
+	{
+		DECLARE_ABSTRACT_ANIM_TRAIT(FBaseTrait, 0xbb4b35d0, FTrait)
+
+		virtual ETraitMode GetTraitMode() const override { return ETraitMode::Base; }
+	};
+
+	// Base class for additive traits that override behavior of other traits
+	struct ANIMNEXT_API FAdditiveTrait : FTrait
+	{
+		DECLARE_ABSTRACT_ANIM_TRAIT(FAdditiveTrait, 0x0dbd1b3f, FTrait)
+
+		virtual ETraitMode GetTraitMode() const override { return ETraitMode::Additive; }
+	};
+
+	/**
+	 * FTraitStaticInitHook
+	 *
+	 * Allows traits to automatically register/unregister within the current scope.
+	 * This can be used during static init.
+	 */
+	struct ANIMNEXT_API FTraitStaticInitHook final
+	{
+		explicit FTraitStaticInitHook(TraitConstructorFunc InTraitConstructor);
+		~FTraitStaticInitHook();
+
+	private:
+		TraitConstructorFunc TraitConstructor;
+	};
+}
diff --git a/Engine/Plugins/Experimental/AnimNext/Source/AnimNext/Public/TraitCore/TraitBinding.h b/Engine/Plugins/Experimental/AnimNext/Source/AnimNext/Public/TraitCore/TraitBinding.h
new file mode 100644
--- /dev/null
+++ b/Engine/Plugins/Experimental/AnimNext/Source/AnimNext/Public/TraitCore/TraitBinding.h
@@ -0,0 +1,153 @@
+// Copyright Epic Games, Inc. All Rights Reserved.
+
+#pragma once
+
+#include "CoreMinimal.h"
+#include "TraitCore/TraitPtr.h"
+#include "TraitCore/TraitInterfaceUID.h"
+#include "TraitCore/TraitTemplate.h"
+
+#include <type_traits>
+
+struct FAnimNextTraitSharedData;
+
+namespace UE::AnimNext
+{
+	struct FTraitInstanceData;
+	struct ITraitInterface;
+	struct FNodeDescription;
+
+	/**
+	 * FTraitBinding
+	 * 
+	 * Base class for all trait bindings.
+	 * A trait binding contains untyped data about a specific trait instance.
+	 */
+	struct ANIMNEXT_API FTraitBinding
+	{
+		// Creates an empty binding.
+		FTraitBinding() = default;
+
+		// Returns whether or not this binding is valid.
+		bool IsValid() const { return TraitTemplate != nullptr; }
+
+		// Queries a node for a pointer to its trait shared data.
+		// If the trait handle is invalid, a null pointer is returned.
+		template<class SharedDataType>
+		const SharedDataType* GetSharedData() const
+		{
+			static_assert(std::is_base_of<FAnimNextTraitSharedData, SharedDataType>::value, "Trait shared data must derive from FAnimNextTraitSharedData");
+
+			if (!IsValid())
+			{
+				return nullptr;
+			}
+
+			return static_cast<const SharedDataType*>(TraitTemplate->GetTraitDescription(*NodeDescription));
+		}
+
+		// Queries a node for a pointer to its trait instance data.
+		// If the trait handle is invalid, a null pointer is returned.
+		template<class InstanceDataType>
+		InstanceDataType* GetInstanceData() const
+		{
+			static_assert(std::is_base_of<FTraitInstanceData, InstanceDataType>::value, "Trait instance data must derive from FTraitInstanceData");
+
+			if (!IsValid())
+			{
+				return nullptr;
+			}
+
+			return static_cast<InstanceDataType*>(TraitTemplate->GetTraitInstance(*TraitPtr.GetNodeInstance()));
+		}
+
+		// Queries a node for a pointer to its trait latent properties.
+		// If the trait handle is invalid or if we have no latent properties, a null pointer is returned.
+		const FLatentPropertyHandle* GetLatentPropertyHandles() const
+		{
+			if (!IsValid() || !TraitTemplate->HasLatentProperties())
+			{
+				return nullptr;
+			}
+
+			return TraitTemplate->GetTraitLatentPropertyHandles(*NodeDescription);
+		}
+
+		// Returns a pointer to the latent property specified by the provided handle or nullptr if the binding/handle are invalid
+		template<typename PropertyType>
+		const PropertyType* GetLatentProperty(FLatentPropertyHandle Handle) const
+		{
+			if (!IsValid())
+			{
+				return nullptr;
+			}
+
+			if (!Handle.IsOffsetValid())
+			{
+				return nullptr;
+			}
+
+			const uint8* NodeInstance = (const uint8*)TraitPtr.GetNodeInstance();
+			return (const PropertyType*)(NodeInstance + Handle.GetLatentPropertyOffset());
+		}
+
+		// Returns the trait pointer we are bound to.
+		FWeakTraitPtr GetTraitPtr() const { return TraitPtr; }
+
+		// Returns the trait interface UID when bound, an invalid UID otherwise.
+		FTraitInterfaceUID GetInterfaceUID() const;
+
+		// Equality and inequality tests
+		bool operator==(const FTraitBinding& RHS) const { return TraitPtr == RHS.TraitPtr && Interface == RHS.Interface; }
+		bool operator!=(const FTraitBinding& RHS) const { return TraitPtr != RHS.TraitPtr || Interface != RHS.Interface; }
+
+	protected:
+		// Creates a valid binding
+		FTraitBinding(const ITraitInterface* InInterface, const FTraitTemplate* InTraitTemplate, const FNodeDescription* InNodeDescription, FWeakTraitPtr InTraitPtr)
+			: Interface(InInterface)
+			, TraitTemplate(InTraitTemplate)
+			, NodeDescription(InNodeDescription)
+			, TraitPtr(InTraitPtr)
+		{}
+
+		// Performs a naked cast to the desired interface type
+		template<class TraitInterfaceType>
+		const TraitInterfaceType* GetInterfaceTyped() const
+		{
+			static_assert(std::is_base_of<ITraitInterface, TraitInterfaceType>::value, "Trait interface data must derive from ITraitInterface");
+			return static_cast<const TraitInterfaceType*>(Interface);
+		}
+
+		// A pointer to the bound interface or nullptr if we are bound to a trait but none of its interfaces
+		const ITraitInterface*				Interface = nullptr;
+
+		// A pointer to the trait template that implements the interface we are bound to or nullptr if we are invalid
+		const FTraitTemplate*				TraitTemplate = nullptr;
+
+		// A pointer to the node shared data we are bound to
+		const FNodeDescription*				NodeDescription = nullptr;
+
+		// A weak handle to the trait instance data we are bound to
+		FWeakTraitPtr						TraitPtr;
+
+		friend struct FExecutionContext;
+	};
+
+	/**
+	 * TTraitBinding
+	 * 
+	 * A templated proxy for trait interfaces. It is meant to be specialized per interface
+	 * in order to allow a clean API and avoid human error. It wraps the necessary information
+	 * to bind a trait to a specific interface. See existing interfaces for examples.
+	 * 
+	 * Here, we forward declare the template which every interface must specialize. Because we
+	 * rely on specialization, it must be defined within the UE::AnimNext namespace where the
+	 * declaration exists.
+	 * 
+	 * Specializations must derive from FTraitBinding to provide the necessary machinery.
+	 * 
+	 * @see IUpdate, IEvaluate, IHierarchy
+	 */
+	template<class TraitInterfaceType>
+	struct TTraitBinding;
+}
diff --git a/Engine/Plugins/Experimental/AnimNext/Source/AnimNext/Public/TraitCore/TraitInterfaceUID.h b/Engine/Plugins/Experimental/AnimNext/Source/AnimNext/Public/TraitCore/TraitInterfaceUID.h
new file mode 100644
--- /dev/null
+++ b/Engine/Plugins/Experimental/AnimNext/Source/AnimNext/Public/TraitCore/TraitInterfaceUID.h
@@ -0,0 +1,67 @@
+// Copyright Epic Games, Inc. All Rights Reserved.
+
+#pragma once
+
+#include "CoreMinimal.h"
+
+namespace UE::AnimNext
+{
+	// Type alias for a raw trait UID, not typesafe
+	using FTraitInterfaceUIDRaw = uint32;
+
+	/**
+	 * FTraitInterfaceUID
+	 *
+	 * Encapsulates an interface global UID.
+	 * The string is exposed in non-shipping builds for logging and debugging purposes.
+	 * The UID should be generated from the provided string using FNV1a with 32 bits.
+	 *
+	 * The whole struct is meant to be 'constexpr' to allow inlining in interface queries.
+	 */
+	struct FTraitInterfaceUID final
+	{
+		// Constructs an invalid UID
+		constexpr FTraitInterfaceUID()
+			: UID(INVALID_UID)
+#if !UE_BUILD_SHIPPING
+			, InterfaceName(TEXT("<Invalid Interface UID>"))
+#endif
+		{}
+
+		// Constructs an interface UID
+		explicit constexpr FTraitInterfaceUID(FTraitInterfaceUIDRaw InUID, const TCHAR* InInterfaceName = TEXT("<Unknown Trait Interface Name>"))
+			: UID(InUID)
+#if !UE_BUILD_SHIPPING
+			, InterfaceName(InInterfaceName)
+#endif
+		{
+		}
+
+#if !UE_BUILD_SHIPPING
+		// Returns a literal string to the interface name
+		constexpr const TCHAR* GetInterfaceName() const { return InterfaceName; }
+#endif
+
+		// Returns the interface global UID
+		constexpr FTraitInterfaceUIDRaw GetUID() const { return UID; }
+
+		// Returns whether this UID is valid or not
+		constexpr bool IsValid() const { return UID != INVALID_UID; }
+
+	private:
+		static constexpr FTraitInterfaceUIDRaw INVALID_UID = 0;
+
+		FTraitInterfaceUIDRaw		UID;
+
+#if !UE_BUILD_SHIPPING
+		const TCHAR*				InterfaceName;
+#endif
+	};
+
+	// Compares for equality and inequality
+	constexpr bool operator==(FTraitInterfaceUID LHS, FTraitInterfaceUID RHS) { return LHS.GetUID() == RHS.GetUID(); }
+	constexpr bool operator!=(FTraitInterfaceUID LHS, FTraitInterfaceUID RHS) { return LHS.GetUID() != RHS.GetUID(); }
+
+	// For sorting
+	constexpr bool operator<(FTraitInterfaceUID LHS, FTraitInterfaceUID RHS) { return LHS.GetUID() < RHS.GetUID(); }
+}
diff --git a/Engine/Plugins/Experimental/AnimNext/Source/AnimNext/Public/TraitInterfaces/IGarbageCollection.h b/Engine/Plugins/Experimental/AnimNext/Source/AnimNext/Public/TraitInterfaces/IGarbageCollection.h
new file mode 100644
--- /dev/null
+++ b/Engine/Plugins/Experimental/AnimNext/Source/AnimNext/Public/TraitInterfaces/IGarbageCollection.h
@@ -0,0 +1,48 @@
+// Copyright Epic Games, Inc. All Rights Reserved.
+
+#pragma once
+
+#include "CoreMinimal.h"
+#include "TraitCore/ITraitInterface.h"
+
+class FReferenceCollector;
+
+namespace UE::AnimNext
+{
+	/**
+	 * IGarbageCollection
+	 * 
+	 * This interface exposes garbage collection reference tracking.
+	 */
+	struct ANIMNEXT_API IGarbageCollection : ITraitInterface
+	{
+		DECLARE_ANIM_TRAIT_INTERFACE(IGarbageCollection, 0x231a2017)
+
+		// Registers the provided binding for GC callback
+		// Once registered, AddReferencedObjects is called during GC to collect references
+		static void RegisterWithGC(const FExecutionContext& Context, const FTraitBinding& Binding);
+
+		// Unregisters the provided binding from GC callback
+		static void UnregisterWithGC(const FExecutionContext& Context, const FTraitBinding& Binding);
+
+		// Called when garbage collection requests hard/strong object references
+		// @see UObject::AddReferencedObjects
+		virtual void AddReferencedObjects(const FExecutionContext& Context, const TTraitBinding<IGarbageCollection>& Binding, FReferenceCollector& Collector) const;
+	};
+
+	/**
+	 * Specialization for trait binding.
+	 */
+	template<>
+	struct TTraitBinding<IGarbageCollection> : FTraitBinding
+	{
+		// @see IGarbageCollection::AddReferencedObjects
+		void AddReferencedObjects(const FExecutionContext& Context, FReferenceCollector& Collector) const
+		{
+			GetInterface()->AddReferencedObjects(Context, *this, Collector);
+		}
+
+	protected:
+		const IGarbageCollection* GetInterface() const { return GetInterfaceTyped<IGarbageCollection>(); }
+	};
+}
diff --git a/Engine/Plugins/Experimental/AnimNext/Source/AnimNext/Public/TraitInterfaces/IHierarchy.h b/Engine/Plugins/Experimental/AnimNext/Source/AnimNext/Public/TraitInterfaces/IHierarchy.h
new file mode 100644
--- /dev/null
+++ b/Engine/Plugins/Experimental/AnimNext/Source/AnimNext/Public/TraitInterfaces/IHierarchy.h
@@ -0,0 +1,55 @@
+// Copyright Epic Games, Inc. All Rights Reserved.
+
+#pragma once
+
+#include "CoreMinimal.h"
+#include "Misc/MemStack.h"
+#include "TraitCore/ITraitInterface.h"
+
+namespace UE::AnimNext
+{
+	// An array of children pointers
+	// We reserve a small amount inline and spill on the memstack
+	using FChildrenArray = TArray<FWeakTraitPtr, TInlineAllocator<8, TMemStackAllocator<>>>;
+
+	/**
+	 * IHierarchy
+	 * 
+	 * This interface exposes hierarchy traversal information to navigate the graph.
+	 */
+	struct ANIMNEXT_API IHierarchy : ITraitInterface
+	{
+		DECLARE_ANIM_TRAIT_INTERFACE(IHierarchy, 0x846d8a37)
+
+		// Returns the number of children
+		// Includes inactive children
+		virtual uint32 GetNumChildren(const FExecutionContext& Context, const TTraitBinding<IHierarchy>& Binding) const;
+
+		// Appends weak handles to any children we wish to traverse.
+		// Traits are responsible for allocating and releasing child instance data.
+		// Empty handles and duplicates can be appended.
+		virtual void GetChildren(const FExecutionContext& Context, const TTraitBinding<IHierarchy>& Binding, FChildrenArray& Children) const;
+	};
+
+	/**
+	 * Specialization for trait binding.
+	 */
+	template<>
+	struct TTraitBinding<IHierarchy> : FTraitBinding
+	{
+		// @see IHierarchy::GetNumChildren
+		uint32 GetNumChildren(const FExecutionContext& Context) const
+		{
+			return GetInterface()->GetNumChildren(Context, *this);
+		}
+
+		// @see IHierarchy::GetChildren
+		void GetChildren(const FExecutionContext& Context, FChildrenArray& Children) const
+		{
+			GetInterface()->GetChildren(Context, *this, Children);
+		}
+
+	protected:
+		const IHierarchy* GetInterface() const { return GetInterfaceTyped<IHierarchy>(); }
+	};
+}
diff --git a/Engine/Plugins/Experimental/AnimNext/Source/AnimNextEditor/Private/Graph/AnimNextGraph_EdGraphNodeCustomization.cpp b/Engine/Plugins/Experimental/AnimNext/Source/AnimNextEditor/Private/Graph/AnimNextGraph_EdGraphNodeCustomization.cpp
--- a/Engine/Plugins/Experimental/AnimNext/Source/AnimNextEditor/Private/Graph/AnimNextGraph_EdGraphNodeCustomization.cpp
+++ b/Engine/Plugins/Experimental/AnimNext/Source/AnimNextEditor/Private/Graph/AnimNextGraph_EdGraphNodeCustomization.cpp
@@ -5,7 +5,7 @@
 #include "DetailLayoutBuilder.h"
 #include "DetailCategoryBuilder.h"
 #include "Graph/AnimNextGraph_EdGraphNode.h"
-#include "Graph/RigDecorator_AnimNextCppDecorator.h"
+#include "Graph/RigDecorator_AnimNextCppTrait.h"
 #include "Logging/LogScopedVerbosityOverride.h"
 #include "InstancedPropertyBagStructureDataProvider.h"
 #include "RigVMModel/RigVMController.h"
@@ -29,7 +29,7 @@ void FAnimNextGraph_EdGraphNodeCustomization::CustomizeDetails(IDetailLayoutBuil
 	{
 		if (URigVMNode* ModelNode = EdGraphNode->GetModelNode())
 		{
-			// If there are any decorators in the node, treat it as a decorator stack
+			// If there are any decorators in the node, treat it as a trait stack
 			const TArray<URigVMPin*> DecoratorPins = ModelNode->GetDecoratorPins();
 			if (DecoratorPins.Num() > 0)
 			{
@@ -56,7 +56,7 @@ void FAnimNextGraph_EdGraphNodeCustomization::CustomizeDetails(IDetailLayoutBuil
 							FRigVMPinDefaultValueImportErrorContext ErrorPipe;
 							LOG_SCOPE_VERBOSITY_OVERRIDE(LogExec, ELogVerbosity::Verbose);
 							const FString DefaultValue = DecoratorPin->GetDefaultValue();
-							DecoratorSharedInstanceData->ImportText(*DefaultValue, ScopedSharedData->GetStructMemory(), nullptr, PPF_None, &ErrorPipe, DecoratorSharedInstanceData->GetName());
+							DecoratorSharedInstanceData->ImportText(*DefaultValue, ScopedSharedData->GetStructMemory(), nullptr, PPF_SerializedAsImportText, &ErrorPipe, DecoratorSharedInstanceData->GetName());
 						}
 						else // or use the decorator itself for simple decorators
 						{
diff --git a/Engine/Plugins/Experimental/AnimNext/Source/AnimNextEditor/Private/Graph/SAnimNextGraphNode.cpp b/Engine/Plugins/Experimental/AnimNext/Source/AnimNextEditor/Private/Graph/SAnimNextGraphNode.cpp
--- a/Engine/Plugins/Experimental/AnimNext/Source/AnimNextEditor/Private/Graph/SAnimNextGraphNode.cpp
+++ b/Engine/Plugins/Experimental/AnimNext/Source/AnimNextEditor/Private/Graph/SAnimNextGraphNode.cpp
@@ -4,7 +4,7 @@
 #include "Graph/AnimNextGraph_EdGraphNode.h"
 #include "SGraphPin.h"
 #include "RigVMModel/Nodes/RigVMDispatchNode.h"
-#include "DecoratorBase/DecoratorHandle.h"
+#include "TraitCore/TraitHandle.h"
 #include "Widgets/SBoxPanel.h"
 #include "Widgets/Input/SButton.h"
 #include "Widgets/Layout/SSpacer.h"
@@ -162,8 +162,8 @@ void SAnimNextGraphNode::UpdatePinTreeView()
 			PinInfo.bIsContainer = false;
 		}
 
-		// Decorator handle pins have special packing that are not user editable, hide their innards
-		if (ModelPin->GetCPPTypeObject() == FAnimNextDecoratorHandle::StaticStruct())
+		// Trait handle pins have special packing that are not user editable, hide their innards
+		if (ModelPin->GetCPPTypeObject() == FAnimNextTraitHandle::StaticStruct())
 		{
 			PinInfo.bHasChildren = false;
 			PinInfo.bIsContainer = false;
@@ -236,7 +236,7 @@ void SAnimNextGraphNode::UpdatePinTreeView()
 		
 		if(UEdGraphPin* InputEdGraphPin = RigGraphNode->FindPin(ModelPin->GetPinPath(), EEdGraphPinDirection::EGPD_Input))
 		{
-			// Make sure to handle input pins that we explicitly hide, like decorator handles
+			// Make sure to handle input pins that we explicitly hide, like trait handles
 			if (!InputEdGraphPin->bHidden)
 			{
 				if (const int32* PinIndexPtr = EdGraphPinToInputPin.Find(InputEdGraphPin))
diff --git a/Engine/Plugins/Experimental/AnimNext/Source/AnimNextTestSuite/Private/AnimNextRuntimeTest.cpp b/Engine/Plugins/Experimental/AnimNext/Source/AnimNextTestSuite/Private/AnimNextRuntimeTest.cpp
--- a/Engine/Plugins/Experimental/AnimNext/Source/AnimNextTestSuite/Private/AnimNextRuntimeTest.cpp
+++ b/Engine/Plugins/Experimental/AnimNext/Source/AnimNextTestSuite/Private/AnimNextRuntimeTest.cpp
@@ -3,7 +3,7 @@
 
 #include "AnimNextRuntimeTest.h"
 
-#include "DecoratorBase/DecoratorReader.h"
+#include "TraitCore/TraitReader.h"
 #include "Graph/AnimNextGraph.h"
 #include "Graph/RigUnit_AnimNextGraphRoot.h"
 #include "Misc/AutomationTest.h"
@@ -33,23 +33,23 @@ namespace UE::AnimNext
 
 		FAnimNextGraphEntryPoint& EntryPoint = Graph.EntryPoints.AddDefaulted_GetRef();
 		EntryPoint.EntryPointName = FRigUnit_AnimNextGraphRoot::DefaultEntryPoint;
-		EntryPoint.RootDecoratorHandle = FAnimNextEntryPointHandle(NodeHandles[0]);
+		EntryPoint.RootTraitHandle = FAnimNextEntryPointHandle(NodeHandles[0]);
 		Graph.ExecuteDefinition = ExecuteDefinition;
 		Graph.SharedDataArchiveBuffer = SharedDataArchiveBuffer;
 		Graph.GraphReferencedObjects.Empty();
 
 		// Reconstruct our graph shared data
 		FMemoryReader GraphSharedDataArchive(SharedDataArchiveBuffer);
-		FDecoratorReader DecoratorReader(Graph.GraphReferencedObjects, GraphSharedDataArchive);
+		FTraitReader TraitReader(Graph.GraphReferencedObjects, GraphSharedDataArchive);
 
-		const FDecoratorReader::EErrorState ErrorState = DecoratorReader.ReadGraph(Graph.SharedDataBuffer);
-		if (ErrorState == FDecoratorReader::EErrorState::None)
+		const FTraitReader::EErrorState ErrorState = TraitReader.ReadGraph(Graph.SharedDataBuffer);
+		if (ErrorState == FTraitReader::EErrorState::None)
 		{
-			Graph.ResolvedRootDecoratorHandles.Add(FRigUnit_AnimNextGraphRoot::DefaultEntryPoint, DecoratorReader.ResolveEntryPointHandle(Graph.EntryPoints[0].RootDecoratorHandle));
+			Graph.ResolvedRootTraitHandles.Add(FRigUnit_AnimNextGraphRoot::DefaultEntryPoint, TraitReader.ResolveEntryPointHandle(Graph.EntryPoints[0].RootTraitHandle));
 
 			for (FNodeHandle& NodeHandle : NodeHandles)
 			{
-				NodeHandle = DecoratorReader.ResolveNodeHandle(NodeHandle);
+				NodeHandle = TraitReader.ResolveNodeHandle(NodeHandle);
 			}
 
 			// Make sure our execute method is registered
@@ -59,7 +59,7 @@ namespace UE::AnimNext
 		else
 		{
 			Graph.SharedDataBuffer.Empty(0);
-			Graph.ResolvedRootDecoratorHandles.Add(FRigUnit_AnimNextGraphRoot::DefaultEntryPoint, FAnimNextDecoratorHandle());
+			Graph.ResolvedRootTraitHandles.Add(FRigUnit_AnimNextGraphRoot::DefaultEntryPoint, FAnimNextTraitHandle());
 			return false;
 		}
 	}
diff --git a/Engine/Plugins/Experimental/AnimNext/Source/AnimNextTestSuite/Private/AnimNextRuntimeTest.h b/Engine/Plugins/Experimental/AnimNext/Source/AnimNextTestSuite/Private/AnimNextRuntimeTest.h
--- a/Engine/Plugins/Experimental/AnimNext/Source/AnimNextTestSuite/Private/AnimNextRuntimeTest.h
+++ b/Engine/Plugins/Experimental/AnimNext/Source/AnimNextTestSuite/Private/AnimNextRuntimeTest.h
@@ -9,9 +9,9 @@
 #include "UObject/Class.h"
 #include "UObject/UnrealType.h"
 
-#include "DecoratorBase/DecoratorUID.h"
-#include "DecoratorBase/NodeTemplate.h"
-#include "DecoratorBase/NodeTemplateRegistry.h"
+#include "TraitCore/TraitUID.h"
+#include "TraitCore/NodeTemplate.h"
+#include "TraitCore/NodeTemplateRegistry.h"
 
 class UAnimNextGraph;
 
@@ -20,10 +20,10 @@ namespace UE::AnimNext
 	struct FNodeHandle;
 
 	// Converts a property value into its string representation using UE reflection
-	template<class DecoratorSharedDataType, typename PropertyType>
+	template<class TraitSharedDataType, typename PropertyType>
 	static FString ToString(const FString& PropertyName, PropertyType PropertyValue)
 	{
-		const UScriptStruct* SharedDataStruct = DecoratorSharedDataType::StaticStruct();
+		const UScriptStruct* SharedDataStruct = TraitSharedDataType::StaticStruct();
 		if (const FProperty* Property = SharedDataStruct->FindPropertyByName(*PropertyName))
 		{
 			void* PropertyDefaults = Property->AllocateAndInitializeValue();
