Migrate away from the 'decorator' term in AnimNext in favor of 'trait'

diff --git a/Engine/Plugins/Experimental/AnimNext/Source/AnimNextTestSuite/Private/AnimNextTraitGraphTest.cpp b/Engine/Plugins/Experimental/AnimNext/Source/AnimNextTestSuite/Private/AnimNextTraitGraphTest.cpp
new file mode 100644
--- /dev/null
+++ b/Engine/Plugins/Experimental/AnimNext/Source/AnimNextTestSuite/Private/AnimNextTraitGraphTest.cpp
@@ -0,0 +1,448 @@
+// Copyright Epic Games, Inc. All Rights Reserved.
+
+#include "AnimNextTraitGraphTest.h"
+
+#include "AnimNextRuntimeTest.h"
+#include "AnimNextTest.h"
+#include "AssetToolsModule.h"
+#include "Context.h"
+#include "UncookedOnlyUtils.h"
+#include "TraitCore/TraitRegistry.h"
+#include "TraitInterfaces/IEvaluate.h"
+#include "TraitInterfaces/IUpdate.h"
+#include "Editor/Transactor.h"
+#include "Graph/AnimNextExecuteContext.h"
+#include "Graph/AnimNextGraph.h"
+#include "Graph/AnimNextGraph_EditorData.h"
+#include "Graph/GraphFactory.h"
+#include "Graph/RigDecorator_AnimNextCppTrait.h"
+#include "Graph/RigUnit_AnimNextGraphRoot.h"
+#include "Graph/RigUnit_AnimNextTraitStack.h"
+#include "Misc/AutomationTest.h"
+#include "Param/ParamStack.h"
+#include "Param/RigVMDispatch_GetParameter.h"
+#include "RigVMFunctions/Math/RigVMFunction_MathInt.h"
+#include "RigVMFunctions/Math/RigVMFunction_MathFloat.h"
+
+#if WITH_DEV_AUTOMATION_TESTS
+
+//****************************************************************************
+// AnimNext Runtime Trait Graph Tests
+//****************************************************************************
+
+namespace UE::AnimNext
+{
+	struct FTestTrait final : FBaseTrait, IEvaluate, IUpdate
+	{
+		DECLARE_ANIM_TRAIT(FTestTrait, 0x41cecb7f, FBaseTrait)
+
+		using FSharedData = FTestTraitSharedData;
+
+		// IUpdate impl
+		virtual void PostUpdate(FUpdateTraversalContext& Context, const TTraitBinding<IUpdate>& Binding, const FTraitUpdateState& TraitState) const override
+		{
+			IUpdate::PostUpdate(Context, Binding, TraitState);
+
+			const FSharedData* SharedData = Binding.GetSharedData<FSharedData>();
+			UE::AnimNext::FParamStack& ParamStack = UE::AnimNext::FParamStack::Get();
+
+			int32& UpdateCount = ParamStack.GetMutableParam<int32>("UpdateCount");
+			UpdateCount++;
+
+			ParamStack.GetMutableParam<int32>("SomeInt32") = SharedData->SomeInt32;
+			ParamStack.GetMutableParam<float>("SomeFloat") = SharedData->SomeFloat;
+
+			ParamStack.GetMutableParam<int32>("SomeLatentInt32") = SharedData->GetSomeLatentInt32(Context, Binding);				// MathAdd with constants, latent
+			ParamStack.GetMutableParam<int32>("SomeOtherLatentInt32") = SharedData->GetSomeOtherLatentInt32(Context, Binding);		// GetParameter, latent
+			ParamStack.GetMutableParam<float>("SomeLatentFloat") = SharedData->GetSomeLatentFloat(Context, Binding);				// Inline value, not latent
+		}
+
+		// IEvaluate impl
+		virtual void PostEvaluate(FEvaluateTraversalContext& Context, const TTraitBinding<IEvaluate>& Binding) const override
+		{
+			IEvaluate::PostEvaluate(Context, Binding);
+
+			UE::AnimNext::FParamStack& ParamStack = UE::AnimNext::FParamStack::Get();
+
+			int32& EvaluateCount = ParamStack.GetMutableParam<int32>("EvaluateCount");
+			EvaluateCount++;
+		}
+	};
+
+	// Trait implementation boilerplate
+	#define TRAIT_INTERFACE_ENUMERATOR(GeneratorMacro) \
+		GeneratorMacro(IEvaluate) \
+		GeneratorMacro(IUpdate) \
+
+	GENERATE_ANIM_TRAIT_IMPLEMENTATION(FTestTrait, TRAIT_INTERFACE_ENUMERATOR)
+	#undef TRAIT_INTERFACE_ENUMERATOR
+}
+
+IMPLEMENT_SIMPLE_AUTOMATION_TEST(FAnimationAnimNextRuntimeTest_GraphAddTrait, "Animation.AnimNext.Runtime.Graph.AddTrait", EAutomationTestFlags::EditorContext | EAutomationTestFlags::EngineFilter)
+
+bool FAnimationAnimNextRuntimeTest_GraphAddTrait::RunTest(const FString& InParameters)
+{
+	using namespace UE::AnimNext;
+
+	{
+		AUTO_REGISTER_ANIM_TRAIT(FTestTrait)
+
+		FScopedClearNodeTemplateRegistry ScopedClearNodeTemplateRegistry;
+
+		UFactory* GraphFactory = NewObject<UAnimNextGraphFactory>();
+		UAnimNextGraph* AnimNextGraph = CastChecked<UAnimNextGraph>(GraphFactory->FactoryCreateNew(UAnimNextGraph::StaticClass(), GetTransientPackage(), TEXT("TestAnimNextGraph"), RF_Transient, nullptr, nullptr, NAME_None));
+		UE_RETURN_ON_ERROR(AnimNextGraph != nullptr, "FAnimationAnimNextRuntimeTest_GraphAddTrait -> Failed to create animation graph");
+
+		UAnimNextGraph_EditorData* EditorData = UncookedOnly::FUtils::GetEditorData(AnimNextGraph);
+		UE_RETURN_ON_ERROR(EditorData != nullptr, "FAnimationAnimNextRuntimeTest_GraphAddTrait -> Failed to find animation graph editor data");
+
+		URigVMController* Controller = EditorData->GetRigVMClient()->GetController(EditorData->GetRigVMClient()->GetDefaultModel());
+		UE_RETURN_ON_ERROR(Controller != nullptr, "FAnimationAnimNextRuntimeTest_GraphAddTrait -> Failed to get RigVM controller");
+
+		// Create an empty trait stack node
+		URigVMUnitNode* DecoratorStackNode = Controller->AddUnitNode(FRigUnit_AnimNextTraitStack::StaticStruct(), FRigVMStruct::ExecuteName, FVector2D(0.0f, 0.0f), FString(), false);
+		UE_RETURN_ON_ERROR(DecoratorStackNode != nullptr, "FAnimationAnimNextRuntimeTest_GraphAddDecorator -> Failed to create decorator stack node");
+
+		// Add a trait
+		const UScriptStruct* CppDecoratorStruct = FRigDecorator_AnimNextCppDecorator::StaticStruct();
+		UE_RETURN_ON_ERROR(CppDecoratorStruct != nullptr, "FAnimationAnimNextRuntimeTest_GraphAddDecorator -> Failed to get find Cpp decorator static struct");
+
+		const FTrait* Trait = FTraitRegistry::Get().Find(FTestTrait::TraitUID);
+		UE_RETURN_ON_ERROR(Trait != nullptr, "FAnimationAnimNextRuntimeTest_GraphAddTrait -> Failed to find test trait");
+
+		UScriptStruct* ScriptStruct = Trait->GetTraitSharedDataStruct();
+		UE_RETURN_ON_ERROR(ScriptStruct != nullptr, "FAnimationAnimNextRuntimeTest_GraphAddTrait -> Failed to find trait shared data struct");
+
+		FString DefaultValue;
+		{
+			const FRigDecorator_AnimNextCppDecorator DefaultCppDecoratorStructInstance;
+			FRigDecorator_AnimNextCppDecorator CppDecoratorStructInstance;
+			CppDecoratorStructInstance.DecoratorSharedDataStruct = ScriptStruct;
+
+			UE_RETURN_ON_ERROR(CppDecoratorStructInstance.CanBeAddedToNode(DecoratorStackNode, nullptr), "FAnimationAnimNextRuntimeTest_GraphAddTrait -> Trait cannot be added to decorator stack node");
+
+			const FProperty* Prop = FAnimNextCppDecoratorWrapper::StaticStruct()->FindPropertyByName(GET_MEMBER_NAME_STRING_CHECKED(FAnimNextCppDecoratorWrapper, CppDecorator));
+			UE_RETURN_ON_ERROR(Prop != nullptr, "FAnimationAnimNextRuntimeTest_GraphAddTrait -> Failed to find wrapper property");
+
+			Prop->ExportText_Direct(DefaultValue, &CppDecoratorStructInstance, &DefaultCppDecoratorStructInstance, nullptr, PPF_None);
+		}
+
+		FString DisplayNameMetadata;
+		ScriptStruct->GetStringMetaDataHierarchical(FRigVMStruct::DisplayNameMetaName, &DisplayNameMetadata);
+		const FString DisplayName = DisplayNameMetadata.IsEmpty() ? Trait->GetTraitName() : DisplayNameMetadata;
+
+		const FName DecoratorName = Controller->AddDecorator(
+			DecoratorStackNode->GetFName(),
+			*CppDecoratorStruct->GetPathName(),
+			*DisplayName,
+			DefaultValue, INDEX_NONE, true, true);
+		UE_RETURN_ON_ERROR(DecoratorName == DisplayName, TEXT("FAnimationAnimNextRuntimeTest_GraphAddTrait -> Unexpected decorator name"));
+
+		URigVMPin* DecoratorPin = DecoratorStackNode->FindPin(*DisplayName);
+		UE_RETURN_ON_ERROR(DecoratorPin != nullptr, TEXT("FAnimationAnimNextRuntimeTest_GraphAddTrait -> Failed to find decorator pin"));
+
+		// Our first pin is the hard coded output result, decorator pins follow
+		UE_RETURN_ON_ERROR(DecoratorStackNode->GetPins().Num() == 2, TEXT("FAnimationAnimNextRuntimeTest_GraphAddTrait -> Unexpected number of pins"));
+		UE_RETURN_ON_ERROR(DecoratorPin->IsDecoratorPin() == true, TEXT("FAnimationAnimNextRuntimeTest_GraphAddTrait -> Unexpected pin type"));
+		UE_RETURN_ON_ERROR(DecoratorPin->GetFName() == DecoratorName, TEXT("FAnimationAnimNextRuntimeTest_GraphAddTrait -> Unexpected pin name"));
+		UE_RETURN_ON_ERROR(DecoratorPin->GetCPPTypeObject() == FRigDecorator_AnimNextCppDecorator::StaticStruct(), TEXT("FAnimationAnimNextRuntimeTest_GraphAddTrait -> Unexpected pin type"));
+
+		// Our first sub-pin is the hard coded script struct member that parametrizes the decorator, dynamic decorator sub-pins follow
+		UE_RETURN_ON_ERROR(DecoratorPin->GetSubPins().Num() == 6, TEXT("FAnimationAnimNextRuntimeTest_GraphAddTrait -> Unexpected decorator sub pins"));
+
+		// SomeInt32
+		UE_RETURN_ON_ERROR(DecoratorPin->GetSubPins()[1]->GetCPPType() == TEXT("int32"), TEXT("FAnimationAnimNextRuntimeTest_GraphAddTrait -> Unexpected decorator pin type"));
+		UE_RETURN_ON_ERROR(DecoratorPin->GetSubPins()[1]->GetDefaultValue() == TEXT("3"), TEXT("FAnimationAnimNextRuntimeTest_GraphAddTrait -> Unexpected decorator pin value"));
+		UE_RETURN_ON_ERROR(!DecoratorPin->GetSubPins()[1]->IsLazy(), TEXT("FAnimationAnimNextRuntimeTest_GraphAddTrait -> Expected non-lazy decorator pin"));
+
+		// SomeFloat
+		UE_RETURN_ON_ERROR(DecoratorPin->GetSubPins()[2]->GetCPPType() == TEXT("float"), TEXT("FAnimationAnimNextRuntimeTest_GraphAddTrait -> Unexpected decorator pin type"));
+		UE_RETURN_ON_ERROR(DecoratorPin->GetSubPins()[2]->GetDefaultValue() == TEXT("34.000000"), TEXT("FAnimationAnimNextRuntimeTest_GraphAddTrait -> Unexpected decorator pin value"));
+		UE_RETURN_ON_ERROR(!DecoratorPin->GetSubPins()[2]->IsLazy(), TEXT("FAnimationAnimNextRuntimeTest_GraphAddTrait -> Expected non-lazy decorator pin"));
+
+		// SomeLatentInt32
+		UE_RETURN_ON_ERROR(DecoratorPin->GetSubPins()[3]->GetCPPType() == TEXT("int32"), TEXT("FAnimationAnimNextRuntimeTest_GraphAddTrait -> Unexpected decorator pin type"));
+		UE_RETURN_ON_ERROR(DecoratorPin->GetSubPins()[3]->GetDefaultValue() == TEXT("3"), TEXT("FAnimationAnimNextRuntimeTest_GraphAddTrait -> Unexpected decorator pin value"));
+		UE_RETURN_ON_ERROR(DecoratorPin->GetSubPins()[3]->IsLazy(), TEXT("FAnimationAnimNextRuntimeTest_GraphAddTrait -> Expected lazy decorator pin"));
+
+		// SomeOtherLatentInt32
+		UE_RETURN_ON_ERROR(DecoratorPin->GetSubPins()[4]->GetCPPType() == TEXT("int32"), TEXT("FAnimationAnimNextRuntimeTest_GraphAddTrait -> Unexpected decorator pin type"));
+		UE_RETURN_ON_ERROR(DecoratorPin->GetSubPins()[4]->GetDefaultValue() == TEXT("3"), TEXT("FAnimationAnimNextRuntimeTest_GraphAddTrait -> Unexpected decorator pin value"));
+		UE_RETURN_ON_ERROR(DecoratorPin->GetSubPins()[4]->IsLazy(), TEXT("FAnimationAnimNextRuntimeTest_GraphAddTrait -> Expected lazy decorator pin"));
+
+		// SomeLatentFloat
+		UE_RETURN_ON_ERROR(DecoratorPin->GetSubPins()[5]->GetCPPType() == TEXT("float"), TEXT("FAnimationAnimNextRuntimeTest_GraphAddTrait -> Unexpected decorator pin type"));
+		UE_RETURN_ON_ERROR(DecoratorPin->GetSubPins()[5]->GetDefaultValue() == TEXT("34.000000"), TEXT("FAnimationAnimNextRuntimeTest_GraphAddTrait -> Unexpected decorator pin value"));
+		UE_RETURN_ON_ERROR(DecoratorPin->GetSubPins()[5]->IsLazy(), TEXT("FAnimationAnimNextRuntimeTest_GraphAddTrait -> Expected lazy decorator pin"));
+	}
+
+	Tests::FUtils::CleanupAfterTests();
+
+	return true;
+}
+
+IMPLEMENT_SIMPLE_AUTOMATION_TEST(FAnimationAnimNextRuntimeTest_GraphExecute, "Animation.AnimNext.Runtime.Graph.Execute", EAutomationTestFlags::EditorContext | EAutomationTestFlags::EngineFilter)
+
+bool FAnimationAnimNextRuntimeTest_GraphExecute::RunTest(const FString& InParameters)
+{
+	using namespace UE::AnimNext;
+
+	{
+		AUTO_REGISTER_ANIM_TRAIT(FTestTrait)
+
+		FScopedClearNodeTemplateRegistry ScopedClearNodeTemplateRegistry;
+
+		UFactory* GraphFactory = NewObject<UAnimNextGraphFactory>();
+		UAnimNextGraph* AnimNextGraph = CastChecked<UAnimNextGraph>(GraphFactory->FactoryCreateNew(UAnimNextGraph::StaticClass(), GetTransientPackage(), TEXT("TestAnimNextGraph"), RF_Transient, nullptr, nullptr, NAME_None));
+		UE_RETURN_ON_ERROR(AnimNextGraph != nullptr, "FAnimationAnimNextRuntimeTest_GraphExecute -> Failed to create animation graph");
+
+		UAnimNextGraph_EditorData* EditorData = UncookedOnly::FUtils::GetEditorData(AnimNextGraph);
+		UE_RETURN_ON_ERROR(EditorData != nullptr, "FAnimationAnimNextRuntimeTest_GraphExecute -> Failed to find animation graph editor data");
+
+		URigVMController* Controller = EditorData->GetRigVMClient()->GetController(EditorData->GetRigVMClient()->GetDefaultModel());
+		UE_RETURN_ON_ERROR(Controller != nullptr, "FAnimationAnimNextRuntimeTest_GraphExecute -> Failed to get RigVM controller");
+
+		// Find graph entry point
+		URigVMNode* MainEntryPointNode = Controller->GetGraph()->FindNodeByName(FRigUnit_AnimNextGraphRoot::StaticStruct()->GetFName());
+		UE_RETURN_ON_ERROR(MainEntryPointNode != nullptr, "FAnimationAnimNextRuntimeTest_GraphExecute -> Failed to find main entry point node");
+
+		URigVMPin* BeginExecutePin = MainEntryPointNode->FindPin(GET_MEMBER_NAME_STRING_CHECKED(FRigUnit_AnimNextGraphRoot, Result));
+		UE_RETURN_ON_ERROR(BeginExecutePin != nullptr && BeginExecutePin->GetDirection() == ERigVMPinDirection::Input, "FAnimationAnimNextRuntimeTest_GraphExecute -> Failed to create entry point");
+
+		// Create an empty decorator stack node
+		URigVMUnitNode* DecoratorStackNode = Controller->AddUnitNode(FRigUnit_AnimNextTraitStack::StaticStruct(), FRigVMStruct::ExecuteName, FVector2D(0.0f, 0.0f), FString(), false);
+		UE_RETURN_ON_ERROR(DecoratorStackNode != nullptr, "FAnimationAnimNextRuntimeTest_GraphExecute -> Failed to create decorator stack node");
+
+		// Link our stack result to our entry point
+		Controller->AddLink(DecoratorStackNode->GetPins()[0], MainEntryPointNode->GetPins()[0]);
+
+		// Add a trait
+		const UScriptStruct* CppDecoratorStruct = FRigDecorator_AnimNextCppDecorator::StaticStruct();
+		UE_RETURN_ON_ERROR(CppDecoratorStruct != nullptr, "FAnimationAnimNextRuntimeTest_GraphExecute -> Failed to get find Cpp decorator static struct");
+
+		const FTrait* Trait = FTraitRegistry::Get().Find(FTestTrait::TraitUID);
+		UE_RETURN_ON_ERROR(Trait != nullptr, "FAnimationAnimNextRuntimeTest_GraphExecute -> Failed to find test trait");
+
+		UScriptStruct* ScriptStruct = Trait->GetTraitSharedDataStruct();
+		UE_RETURN_ON_ERROR(ScriptStruct != nullptr, "FAnimationAnimNextRuntimeTest_GraphExecute -> Failed to find trait shared data struct");
+
+		FString DefaultValue;
+		{
+			const FRigDecorator_AnimNextCppDecorator DefaultCppDecoratorStructInstance;
+			FRigDecorator_AnimNextCppDecorator CppDecoratorStructInstance;
+			CppDecoratorStructInstance.DecoratorSharedDataStruct = ScriptStruct;
+
+			UE_RETURN_ON_ERROR(CppDecoratorStructInstance.CanBeAddedToNode(DecoratorStackNode, nullptr), "FAnimationAnimNextRuntimeTest_GraphExecute -> Trait cannot be added to decorator stack node");
+
+			const FProperty* Prop = FAnimNextCppDecoratorWrapper::StaticStruct()->FindPropertyByName(GET_MEMBER_NAME_STRING_CHECKED(FAnimNextCppDecoratorWrapper, CppDecorator));
+			UE_RETURN_ON_ERROR(Prop != nullptr, "FAnimationAnimNextRuntimeTest_GraphExecute -> Failed to find wrapper property");
+
+			Prop->ExportText_Direct(DefaultValue, &CppDecoratorStructInstance, &DefaultCppDecoratorStructInstance, nullptr, PPF_None);
+		}
+
+		FString DisplayNameMetadata;
+		ScriptStruct->GetStringMetaDataHierarchical(FRigVMStruct::DisplayNameMetaName, &DisplayNameMetadata);
+		const FString DisplayName = DisplayNameMetadata.IsEmpty() ? Trait->GetTraitName() : DisplayNameMetadata;
+
+		const FName DecoratorName = Controller->AddDecorator(
+			DecoratorStackNode->GetFName(),
+			*CppDecoratorStruct->GetPathName(),
+			*DisplayName,
+			DefaultValue, INDEX_NONE, true, true);
+		UE_RETURN_ON_ERROR(DecoratorName == DisplayName, TEXT("FAnimationAnimNextRuntimeTest_GraphExecute -> Unexpected decorator name"));
+
+		URigVMPin* DecoratorPin = DecoratorStackNode->FindPin(*DisplayName);
+		UE_RETURN_ON_ERROR(DecoratorPin != nullptr, TEXT("FAnimationAnimNextRuntimeTest_GraphExecute -> Failed to find decorator pin"));
+
+		// Set some values on our decorator
+		Controller->SetPinDefaultValue(DecoratorPin->GetSubPins()[1]->GetPinPath(), TEXT("78"));
+		Controller->SetPinDefaultValue(DecoratorPin->GetSubPins()[2]->GetPinPath(), TEXT("142.33"));
+
+		TSharedRef<FParamStack> ParamStack = MakeShared<FParamStack>();
+		FParamStack::AttachToCurrentThread(ParamStack);
+
+		FAnimNextGraphInstancePtr GraphInstance;
+		AnimNextGraph->AllocateInstance(GraphInstance);
+
+		FParamStack::FPushedLayerHandle LayerHandle = ParamStack->PushValues(
+			"UpdateCount", (int32)0,
+			"EvaluateCount", (int32)0,
+			"SomeInt32", (int32)0,
+			"SomeFloat", 0.0f,
+			"SomeLatentInt32", (int32)0,
+			"SomeOtherLatentInt32", (int32)0,
+			"SomeLatentFloat", 0.0f
+		);
+
+		{
+			UE::AnimNext::UpdateGraph(GraphInstance, 1.0f / 30.0f);
+			(void)UE::AnimNext::EvaluateGraph(GraphInstance);
+		}
+
+		AddErrorIfFalse(ParamStack->GetParam<int32>("UpdateCount") == 1, "FAnimationAnimNextRuntimeTest_GraphExecute -> Unexpected update count");
+		AddErrorIfFalse(ParamStack->GetParam<int32>("EvaluateCount") == 1, "FAnimationAnimNextRuntimeTest_GraphExecute -> Unexpected evaluate count");
+		AddErrorIfFalse(ParamStack->GetParam<int32>("SomeInt32") == 78, "FAnimationAnimNextRuntimeTest_GraphExecute -> Unexpected SomeInt32 value");
+		AddErrorIfFalse(ParamStack->GetParam<float>("SomeFloat") == 142.33f, "FAnimationAnimNextRuntimeTest_GraphExecute -> Unexpected SomeFloat value");
+		AddErrorIfFalse(ParamStack->GetParam<int32>("SomeLatentInt32") == 3, "FAnimationAnimNextRuntimeTest_GraphExecute -> Unexpected SomeLatentInt32 value");
+		AddErrorIfFalse(ParamStack->GetParam<int32>("SomeOtherLatentInt32") == 3, "FAnimationAnimNextRuntimeTest_GraphExecute -> Unexpected SomeOtherLatentInt32 value");
+		AddErrorIfFalse(ParamStack->GetParam<float>("SomeLatentFloat") == 34.0f, "FAnimationAnimNextRuntimeTest_GraphExecute -> Unexpected SomeLatentFloat value");
+
+		ParamStack->PopLayer(LayerHandle);
+		GraphInstance.Release();
+
+		FParamStack::DetachFromCurrentThread();
+	}
+
+	Tests::FUtils::CleanupAfterTests();
+
+	return true;
+}
+
+IMPLEMENT_SIMPLE_AUTOMATION_TEST(FAnimationAnimNextRuntimeTest_GraphExecuteLatent, "Animation.AnimNext.Runtime.Graph.ExecuteLatent", EAutomationTestFlags::EditorContext | EAutomationTestFlags::EngineFilter)
+
+bool FAnimationAnimNextRuntimeTest_GraphExecuteLatent::RunTest(const FString& InParameters)
+{
+	using namespace UE::AnimNext;
+
+	{
+		AUTO_REGISTER_ANIM_TRAIT(FTestTrait)
+		FScopedClearNodeTemplateRegistry ScopedClearNodeTemplateRegistry;
+
+		UFactory* GraphFactory = NewObject<UAnimNextGraphFactory>();
+		UAnimNextGraph* AnimNextGraph = CastChecked<UAnimNextGraph>(GraphFactory->FactoryCreateNew(UAnimNextGraph::StaticClass(), GetTransientPackage(), TEXT("TestAnimNextGraph"), RF_Transient, nullptr, nullptr, NAME_None));
+		UE_RETURN_ON_ERROR(AnimNextGraph != nullptr, "FAnimationAnimNextRuntimeTest_GraphExecuteLatent -> Failed to create animation graph");
+
+		UAnimNextGraph_EditorData* EditorData = UncookedOnly::FUtils::GetEditorData(AnimNextGraph);
+		UE_RETURN_ON_ERROR(EditorData != nullptr, "FAnimationAnimNextRuntimeTest_GraphExecuteLatent -> Failed to find animation graph editor data");
+
+		URigVMController* Controller = EditorData->GetRigVMClient()->GetController(EditorData->GetRigVMClient()->GetDefaultModel());
+		UE_RETURN_ON_ERROR(Controller != nullptr, "FAnimationAnimNextRuntimeTest_GraphExecuteLatent -> Failed to get RigVM controller");
+
+		// Find graph entry point
+		URigVMNode* MainEntryPointNode = Controller->GetGraph()->FindNodeByName(FRigUnit_AnimNextGraphRoot::StaticStruct()->GetFName());
+		UE_RETURN_ON_ERROR(MainEntryPointNode != nullptr, "FAnimationAnimNextRuntimeTest_GraphExecuteLatent -> Failed to find main entry point node");
+
+		URigVMPin* BeginExecutePin = MainEntryPointNode->FindPin(GET_MEMBER_NAME_STRING_CHECKED(FRigUnit_AnimNextGraphRoot, Result));
+		UE_RETURN_ON_ERROR(BeginExecutePin != nullptr && BeginExecutePin->GetDirection() == ERigVMPinDirection::Input, "FAnimationAnimNextRuntimeTest_GraphExecuteLatent -> Failed to create entry point");
+
+		// Create an empty decorator stack node
+		URigVMUnitNode* DecoratorStackNode = Controller->AddUnitNode(FRigUnit_AnimNextTraitStack::StaticStruct(), FRigVMStruct::ExecuteName, FVector2D(0.0f, 0.0f), FString(), false);
+		UE_RETURN_ON_ERROR(DecoratorStackNode != nullptr, "FAnimationAnimNextRuntimeTest_GraphExecuteLatent -> Failed to create decorator stack node");
+
+		// Link our stack result to our entry point
+		Controller->AddLink(DecoratorStackNode->GetPins()[0], MainEntryPointNode->GetPins()[0]);
+
+		// Add a trait
+		const UScriptStruct* CppDecoratorStruct = FRigDecorator_AnimNextCppDecorator::StaticStruct();
+		UE_RETURN_ON_ERROR(CppDecoratorStruct != nullptr, "FAnimationAnimNextRuntimeTest_GraphExecuteLatent -> Failed to get find Cpp decorator static struct");
+
+		const FTrait* Trait = FTraitRegistry::Get().Find(FTestTrait::TraitUID);
+		UE_RETURN_ON_ERROR(Trait != nullptr, "FAnimationAnimNextRuntimeTest_GraphExecuteLatent -> Failed to find test trait");
+
+		UScriptStruct* ScriptStruct = Trait->GetTraitSharedDataStruct();
+		UE_RETURN_ON_ERROR(ScriptStruct != nullptr, "FAnimationAnimNextRuntimeTest_GraphExecuteLatent -> Failed to find trait shared data struct");
+
+		FString DefaultValue;
+		{
+			const FRigDecorator_AnimNextCppDecorator DefaultCppDecoratorStructInstance;
+			FRigDecorator_AnimNextCppDecorator CppDecoratorStructInstance;
+			CppDecoratorStructInstance.DecoratorSharedDataStruct = ScriptStruct;
+
+			UE_RETURN_ON_ERROR(CppDecoratorStructInstance.CanBeAddedToNode(DecoratorStackNode, nullptr), "FAnimationAnimNextRuntimeTest_GraphExecuteLatent -> Trait cannot be added to decorator stack node");
+
+			const FProperty* Prop = FAnimNextCppDecoratorWrapper::StaticStruct()->FindPropertyByName(GET_MEMBER_NAME_STRING_CHECKED(FAnimNextCppDecoratorWrapper, CppDecorator));
+			UE_RETURN_ON_ERROR(Prop != nullptr, "FAnimationAnimNextRuntimeTest_GraphExecuteLatent -> Failed to find wrapper property");
+
+			Prop->ExportText_Direct(DefaultValue, &CppDecoratorStructInstance, &DefaultCppDecoratorStructInstance, nullptr, PPF_None);
+		}
+
+		FString DisplayNameMetadata;
+		ScriptStruct->GetStringMetaDataHierarchical(FRigVMStruct::DisplayNameMetaName, &DisplayNameMetadata);
+		const FString DisplayName = DisplayNameMetadata.IsEmpty() ? Trait->GetTraitName() : DisplayNameMetadata;
+
+		const FName DecoratorName = Controller->AddDecorator(
+			DecoratorStackNode->GetFName(),
+			*CppDecoratorStruct->GetPathName(),
+			*DisplayName,
+			DefaultValue, INDEX_NONE, true, true);
+		UE_RETURN_ON_ERROR(DecoratorName == DisplayName, TEXT("FAnimationAnimNextRuntimeTest_GraphExecuteLatent -> Unexpected decorator name"));
+
+		// Set some values on our decorator
+		URigVMPin* DecoratorPin = DecoratorStackNode->FindPin(*DisplayName);
+		UE_RETURN_ON_ERROR(DecoratorPin != nullptr, TEXT("FAnimationAnimNextRuntimeTest_GraphExecuteLatent -> Failed to find decorator pin"));
+
+		Controller->SetPinDefaultValue(DecoratorPin->GetSubPins()[1]->GetPinPath(), TEXT("78"));		// SomeInt32
+		Controller->SetPinDefaultValue(DecoratorPin->GetSubPins()[2]->GetPinPath(), TEXT("142.33"));	// SomeFloat
+		Controller->SetPinDefaultValue(DecoratorPin->GetSubPins()[5]->GetPinPath(), TEXT("1123.31"));	// SomeLatentFloat, inline value on latent pin
+
+		// Set some latent values on our decorator
+		{
+			FRigVMFunction_MathIntAdd IntAdd;
+			IntAdd.A = 10;
+			IntAdd.B = 23;
+
+			URigVMUnitNode* IntAddNode = Controller->AddUnitNodeWithDefaults(FRigVMFunction_MathIntAdd::StaticStruct(), FRigStructScope(IntAdd), FRigVMStruct::ExecuteName, FVector2D::ZeroVector, FString(), false);
+			UE_RETURN_ON_ERROR(IntAddNode != nullptr, TEXT("FAnimationAnimNextRuntimeTest_GraphExecuteLatent -> Failed to create Int add node"));
+
+			Controller->AddLink(
+				IntAddNode->FindPin(GET_MEMBER_NAME_STRING_CHECKED(FRigVMFunction_MathIntAdd, Result)),
+				DecoratorPin->GetSubPins()[3]);	// SomeLatentInt32
+		}
+
+		{
+			const FRigVMDispatchFactory* GetParameterFactory = FRigVMRegistry::Get().RegisterFactory(FRigVMDispatch_GetParameter::StaticStruct());
+			const FName GetParameterNotation = GetParameterFactory->GetTemplate()->GetNotation();
+
+			URigVMNode* GetParameterNode = Controller->AddTemplateNode(GetParameterNotation);
+			UE_RETURN_ON_ERROR(GetParameterNode != nullptr, TEXT("FAnimationAnimNextRuntimeTest_GraphExecuteLatent -> Failed to create GetParameter node"));
+
+			UE_RETURN_ON_ERROR(Controller->ResolveWildCardPin(GetParameterNode->FindPin(TEXT("Value")), RigVMTypeUtils::TypeIndex::Int32), TEXT("FAnimationAnimNextRuntimeTest_GraphExecuteLatent -> Failed to resolve wildcard for GetParameter node"));
+			Controller->SetPinDefaultValue(GetParameterNode->FindPin(TEXT("Parameter"))->GetPinPath(), TEXT("SomeSourceInt"));
+
+			Controller->AddLink(
+				GetParameterNode->FindPin(TEXT("Value")),
+				DecoratorPin->GetSubPins()[4]);	// SomeOtherLatentInt32
+		}
+
+		TSharedRef<FParamStack> ParamStack = MakeShared<FParamStack>();
+		FParamStack::AttachToCurrentThread(ParamStack);
+
+		FAnimNextGraphInstancePtr GraphInstance;
+		AnimNextGraph->AllocateInstance(GraphInstance);
+
+		FParamStack::FPushedLayerHandle LayerHandle = ParamStack->PushValues(
+			"UpdateCount", (int32)0,
+			"EvaluateCount", (int32)0,
+			"SomeSourceInt", (int32)1223,
+			"SomeInt32", (int32)0,
+			"SomeFloat", 0.0f,
+			"SomeLatentInt32", (int32)0,
+			"SomeOtherLatentInt32", (int32)0,
+			"SomeLatentFloat", 0.0f
+		);
+
+		{
+			UE::AnimNext::UpdateGraph(GraphInstance, 1.0f / 30.0f);
+			(void)UE::AnimNext::EvaluateGraph(GraphInstance);
+		}
+
+		AddErrorIfFalse(ParamStack->GetParam<int32>("UpdateCount") == 1, "FAnimationAnimNextRuntimeTest_GraphExecuteLatent -> Unexpected update count");
+		AddErrorIfFalse(ParamStack->GetParam<int32>("EvaluateCount") == 1, "FAnimationAnimNextRuntimeTest_GraphExecuteLatent -> Unexpected evaluate count");
+		AddErrorIfFalse(ParamStack->GetParam<int32>("SomeInt32") == 78, "FAnimationAnimNextRuntimeTest_GraphExecuteLatent -> Unexpected SomeInt32 value");
+		AddErrorIfFalse(ParamStack->GetParam<float>("SomeFloat") == 142.33f, "FAnimationAnimNextRuntimeTest_GraphExecuteLatent -> Unexpected SomeFloat value");
+		AddErrorIfFalse(ParamStack->GetParam<int32>("SomeLatentInt32") == 33, "FAnimationAnimNextRuntimeTest_GraphExecuteLatent -> Unexpected SomeLatentInt32 value");
+		AddErrorIfFalse(ParamStack->GetParam<int32>("SomeOtherLatentInt32") == 1223, "FAnimationAnimNextRuntimeTest_GraphExecuteLatent -> Unexpected SomeOtherLatentInt32 value");
+		AddErrorIfFalse(ParamStack->GetParam<float>("SomeLatentFloat") == 1123.31f, "FAnimationAnimNextRuntimeTest_GraphExecuteLatent -> Unexpected SomeLatentFloat value");
+
+		ParamStack->PopLayer(LayerHandle);
+		GraphInstance.Release();
+
+		FParamStack::DetachFromCurrentThread();
+	}
+
+	Tests::FUtils::CleanupAfterTests();
+
+	return true;
+}
+
+#endif
diff --git a/Engine/Plugins/Experimental/AnimNext/Source/AnimNextTestSuite/Private/AnimNextTraitGraphTest.h b/Engine/Plugins/Experimental/AnimNext/Source/AnimNextTestSuite/Private/AnimNextTraitGraphTest.h
new file mode 100644
--- /dev/null
+++ b/Engine/Plugins/Experimental/AnimNext/Source/AnimNextTestSuite/Private/AnimNextTraitGraphTest.h
@@ -0,0 +1,39 @@
+// Copyright Epic Games, Inc. All Rights Reserved.
+
+#pragma once
+
+#include "CoreMinimal.h"
+
+#include "TraitCore/TraitSharedData.h"
+#include "Graph/AnimNextGraph.h"
+
+#include "AnimNextTraitGraphTest.generated.h"
+
+USTRUCT()
+struct FTestTraitSharedData : public FAnimNextTraitSharedData
+{
+	GENERATED_BODY()
+
+	UPROPERTY(meta = (Input, Inline))
+	int32 SomeInt32 = 3;
+
+	UPROPERTY(meta = (Input, Inline))
+	float SomeFloat = 34.0f;
+
+	UPROPERTY(meta = (Input))
+	int32 SomeLatentInt32 = 3;			// MathAdd with constants, latent
+
+	UPROPERTY(meta = (Input))
+	int32 SomeOtherLatentInt32 = 3;		// GetParameter, latent
+
+	UPROPERTY(meta = (Input))
+	float SomeLatentFloat = 34.0f;		// Inline value, not latent
+
+	#define TRAIT_LATENT_PROPERTIES_ENUMERATOR(GeneratorMacro) \
+		GeneratorMacro(SomeLatentInt32) \
+		GeneratorMacro(SomeOtherLatentInt32) \
+		GeneratorMacro(SomeLatentFloat) \
+
+	GENERATE_TRAIT_LATENT_PROPERTIES(FTestTraitSharedData, TRAIT_LATENT_PROPERTIES_ENUMERATOR)
+	#undef TRAIT_LATENT_PROPERTIES_ENUMERATOR
+};
diff --git a/Engine/Plugins/Experimental/AnimNext/Source/AnimNextTestSuite/Private/AnimNextTraitInterfacesTest.cpp b/Engine/Plugins/Experimental/AnimNext/Source/AnimNextTestSuite/Private/AnimNextTraitInterfacesTest.cpp
new file mode 100644
--- /dev/null
+++ b/Engine/Plugins/Experimental/AnimNext/Source/AnimNextTestSuite/Private/AnimNextTraitInterfacesTest.cpp
@@ -0,0 +1,691 @@
+// Copyright Epic Games, Inc. All Rights Reserved.
+
+#include "AnimNextTraitInterfacesTest.h"
+#include "AnimNextRuntimeTest.h"
+#include "AnimNextTest.h"
+
+#include "Misc/AutomationTest.h"
+
+#if WITH_DEV_AUTOMATION_TESTS
+
+#include "Serialization/MemoryReader.h"
+
+#include "TraitCore/Trait.h"
+#include "TraitCore/TraitReader.h"
+#include "TraitCore/TraitRegistry.h"
+#include "TraitCore/TraitWriter.h"
+#include "TraitCore/ExecutionContext.h"
+#include "TraitCore/ITraitInterface.h"
+#include "TraitCore/NodeInstance.h"
+#include "TraitCore/NodeTemplateBuilder.h"
+#include "TraitCore/NodeTemplateRegistry.h"
+#include "TraitInterfaces/IEvaluate.h"
+#include "TraitInterfaces/IHierarchy.h"
+#include "TraitInterfaces/IUpdate.h"
+#include "Graph/AnimNextGraph.h"
+#include "Graph/GraphFactory.h"
+
+//****************************************************************************
+// AnimNext Runtime TraitInterfaces Tests
+//****************************************************************************
+
+namespace UE::AnimNext
+{
+	namespace Private
+	{
+		static TArray<FTraitUID>* UpdatedTraits = nullptr;
+		static TArray<FTraitUID>* EvaluatedTraits = nullptr;
+	}
+
+	struct FTraitWithNoChildren : FBaseTrait, IUpdate, IEvaluate
+	{
+		DECLARE_ANIM_TRAIT(FTraitWithNoChildren, 0x79080a53, FBaseTrait)
+
+		// IUpdate impl
+		virtual void PreUpdate(FUpdateTraversalContext& Context, const TTraitBinding<IUpdate>& Binding, const FTraitUpdateState& TraitState) const override
+		{
+			if (Private::UpdatedTraits != nullptr)
+			{
+				Private::UpdatedTraits->Add(FTraitWithNoChildren::TraitUID);
+			}
+
+			IUpdate::PreUpdate(Context, Binding, TraitState);
+		}
+
+		virtual void PostUpdate(FUpdateTraversalContext& Context, const TTraitBinding<IUpdate>& Binding, const FTraitUpdateState& TraitState) const override
+		{
+			if (Private::UpdatedTraits != nullptr)
+			{
+				Private::UpdatedTraits->Add(FTraitWithNoChildren::TraitUID);
+			}
+
+			IUpdate::PostUpdate(Context, Binding, TraitState);
+		}
+
+		// IEvaluate impl
+		virtual void PreEvaluate(FEvaluateTraversalContext& Context, const TTraitBinding<IEvaluate>& Binding) const override
+		{
+			if (Private::EvaluatedTraits != nullptr)
+			{
+				Private::EvaluatedTraits->Add(FTraitWithNoChildren::TraitUID);
+			}
+
+			IEvaluate::PreEvaluate(Context, Binding);
+		}
+
+		virtual void PostEvaluate(FEvaluateTraversalContext& Context, const TTraitBinding<IEvaluate>& Binding) const override
+		{
+			if (Private::EvaluatedTraits != nullptr)
+			{
+				Private::EvaluatedTraits->Add(FTraitWithNoChildren::TraitUID);
+			}
+
+			IEvaluate::PostEvaluate(Context, Binding);
+		}
+	};
+
+	// Trait implementation boilerplate
+	#define TRAIT_INTERFACE_ENUMERATOR(GeneratorMacro) \
+		GeneratorMacro(IEvaluate) \
+		GeneratorMacro(IUpdate) \
+
+	GENERATE_ANIM_TRAIT_IMPLEMENTATION(FTraitWithNoChildren, TRAIT_INTERFACE_ENUMERATOR)
+	#undef TRAIT_INTERFACE_ENUMERATOR
+
+	// This trait does not update or evaluate
+	struct FTraitWithOneChild : FBaseTrait, IHierarchy
+	{
+		DECLARE_ANIM_TRAIT(FTraitWithOneChild, 0xab680b35, FBaseTrait)
+
+		using FSharedData = FTraitWithOneChildSharedData;
+
+		struct FInstanceData : FTrait::FInstanceData
+		{
+			FTraitPtr Child;
+
+			void Construct(const FExecutionContext& Context, const FTraitBinding& Binding)
+			{
+				Child = Context.AllocateNodeInstance(Binding.GetTraitPtr(), Binding.GetSharedData<FSharedData>()->Child);
+			}
+		};
+
+		// IHierarchy impl
+		virtual uint32 GetNumChildren(const FExecutionContext& Context, const TTraitBinding<IHierarchy>& Binding) const override
+		{
+			return 1;
+		}
+
+		virtual void GetChildren(const FExecutionContext& Context, const TTraitBinding<IHierarchy>& Binding, FChildrenArray& Children) const override
+		{
+			const FInstanceData* InstanceData = Binding.GetInstanceData<FInstanceData>();
+
+			Children.Add(InstanceData->Child);
+
+			IHierarchy::GetChildren(Context, Binding, Children);
+		}
+	};
+
+	// Trait implementation boilerplate
+	#define TRAIT_INTERFACE_ENUMERATOR(GeneratorMacro) \
+		GeneratorMacro(IHierarchy) \
+
+	GENERATE_ANIM_TRAIT_IMPLEMENTATION(FTraitWithOneChild, TRAIT_INTERFACE_ENUMERATOR)
+	#undef TRAIT_INTERFACE_ENUMERATOR
+
+	struct FTraitWithChildren : FBaseTrait, IHierarchy, IUpdate, IEvaluate
+	{
+		DECLARE_ANIM_TRAIT(FTraitWithChildren, 0x4b296948, FBaseTrait)
+
+		using FSharedData = FTraitWithChildrenSharedData;
+
+		struct FInstanceData : FTrait::FInstanceData
+		{
+			FTraitPtr Children[2];
+
+			void Construct(const FExecutionContext& Context, const FTraitBinding& Binding)
+			{
+				Children[0] = Context.AllocateNodeInstance(Binding.GetTraitPtr(), Binding.GetSharedData<FSharedData>()->Children[0]);
+				Children[1] = Context.AllocateNodeInstance(Binding.GetTraitPtr(), Binding.GetSharedData<FSharedData>()->Children[1]);
+			}
+		};
+
+		// IHierarchy impl
+		virtual uint32 GetNumChildren(const FExecutionContext& Context, const TTraitBinding<IHierarchy>& Binding) const override
+		{
+			return 2;
+		}
+
+		virtual void GetChildren(const FExecutionContext& Context, const TTraitBinding<IHierarchy>& Binding, FChildrenArray& Children) const override
+		{
+			const FInstanceData* InstanceData = Binding.GetInstanceData<FInstanceData>();
+
+			Children.Add(InstanceData->Children[0]);
+			Children.Add(InstanceData->Children[1]);
+
+			IHierarchy::GetChildren(Context, Binding, Children);
+		}
+
+		// IUpdate impl
+		virtual void PreUpdate(FUpdateTraversalContext& Context, const TTraitBinding<IUpdate>& Binding, const FTraitUpdateState& TraitState) const override
+		{
+			if (Private::UpdatedTraits != nullptr)
+			{
+				Private::UpdatedTraits->Add(FTraitWithChildren::TraitUID);
+			}
+
+			IUpdate::PreUpdate(Context, Binding, TraitState);
+		}
+
+		virtual void QueueChildrenForTraversal(FUpdateTraversalContext& Context, const TTraitBinding<IUpdate>& Binding, const FTraitUpdateState& TraitState, FUpdateTraversalQueue& TraversalQueue) const override
+		{
+			const FInstanceData* InstanceData = Binding.GetInstanceData<FInstanceData>();
+
+			TraversalQueue.Push(InstanceData->Children[0], TraitState);
+			TraversalQueue.Push(InstanceData->Children[1], TraitState);
+		}
+
+		virtual void PostUpdate(FUpdateTraversalContext& Context, const TTraitBinding<IUpdate>& Binding, const FTraitUpdateState& TraitState) const override
+		{
+			if (Private::UpdatedTraits != nullptr)
+			{
+				Private::UpdatedTraits->Add(FTraitWithChildren::TraitUID);
+			}
+
+			IUpdate::PostUpdate(Context, Binding, TraitState);
+		}
+
+		// IEvaluate impl
+		virtual void PreEvaluate(FEvaluateTraversalContext& Context, const TTraitBinding<IEvaluate>& Binding) const override
+		{
+			if (Private::EvaluatedTraits != nullptr)
+			{
+				Private::EvaluatedTraits->Add(FTraitWithChildren::TraitUID);
+			}
+
+			IEvaluate::PreEvaluate(Context, Binding);
+		}
+
+		virtual void PostEvaluate(FEvaluateTraversalContext& Context, const TTraitBinding<IEvaluate>& Binding) const override
+		{
+			if (Private::EvaluatedTraits != nullptr)
+			{
+				Private::EvaluatedTraits->Add(FTraitWithChildren::TraitUID);
+			}
+
+			IEvaluate::PostEvaluate(Context, Binding);
+		}
+	};
+
+	// Trait implementation boilerplate
+	#define TRAIT_INTERFACE_ENUMERATOR(GeneratorMacro) \
+		GeneratorMacro(IEvaluate) \
+		GeneratorMacro(IHierarchy) \
+		GeneratorMacro(IUpdate) \
+
+	GENERATE_ANIM_TRAIT_IMPLEMENTATION(FTraitWithChildren, TRAIT_INTERFACE_ENUMERATOR)
+	#undef TRAIT_INTERFACE_ENUMERATOR
+}
+
+IMPLEMENT_SIMPLE_AUTOMATION_TEST(FAnimationAnimNextRuntimeTest_IHierarchy, "Animation.AnimNext.Runtime.IHierarchy", EAutomationTestFlags::EditorContext | EAutomationTestFlags::EngineFilter)
+
+bool FAnimationAnimNextRuntimeTest_IHierarchy::RunTest(const FString& InParameters)
+{
+	using namespace UE::AnimNext;
+
+	{
+		AUTO_REGISTER_ANIM_TRAIT(FTraitWithNoChildren)
+		AUTO_REGISTER_ANIM_TRAIT(FTraitWithOneChild)
+		AUTO_REGISTER_ANIM_TRAIT(FTraitWithChildren)
+
+		UFactory* GraphFactory = NewObject<UAnimNextGraphFactory>();
+		UAnimNextGraph* AnimNextGraph = CastChecked<UAnimNextGraph>(GraphFactory->FactoryCreateNew(UAnimNextGraph::StaticClass(), GetTransientPackage(), TEXT("TestAnimNextGraph"), RF_Transient, nullptr, nullptr, NAME_None));
+		UE_RETURN_ON_ERROR(AnimNextGraph != nullptr, "FAnimationAnimNextRuntimeTest_IHierarchy -> Failed to create animation graph");
+
+		FScopedClearNodeTemplateRegistry ScopedClearNodeTemplateRegistry;
+		FNodeTemplateRegistry& Registry = FNodeTemplateRegistry::Get();
+
+		// We create a few node templates
+		// Template A has a single trait with no children
+		TArray<FTraitUID> NodeTemplateTraitListA;
+		NodeTemplateTraitListA.Add(FTraitWithNoChildren::TraitUID);
+
+		// Template B has a single trait with one child
+		TArray<FTraitUID> NodeTemplateTraitListB;
+		NodeTemplateTraitListB.Add(FTraitWithOneChild::TraitUID);
+
+		// Template C has two traits, each with one child
+		TArray<FTraitUID> NodeTemplateTraitListC;
+		NodeTemplateTraitListC.Add(FTraitWithOneChild::TraitUID);
+		NodeTemplateTraitListC.Add(FTraitWithOneChild::TraitUID);
+
+		// Template D has a single trait with children
+		TArray<FTraitUID> NodeTemplateTraitListD;
+		NodeTemplateTraitListD.Add(FTraitWithChildren::TraitUID);
+
+		// Populate our node template registry
+		TArray<uint8> NodeTemplateBufferA, NodeTemplateBufferB, NodeTemplateBufferC, NodeTemplateBufferD;
+		const FNodeTemplate* NodeTemplateA = FNodeTemplateBuilder::BuildNodeTemplate(NodeTemplateTraitListA, NodeTemplateBufferA);
+		const FNodeTemplate* NodeTemplateB = FNodeTemplateBuilder::BuildNodeTemplate(NodeTemplateTraitListB, NodeTemplateBufferB);
+		const FNodeTemplate* NodeTemplateC = FNodeTemplateBuilder::BuildNodeTemplate(NodeTemplateTraitListC, NodeTemplateBufferC);
+		const FNodeTemplate* NodeTemplateD = FNodeTemplateBuilder::BuildNodeTemplate(NodeTemplateTraitListD, NodeTemplateBufferD);
+
+		// Build our graph, it as follow (each node template has a single node instance):
+		// NodeA has no children
+		// NodeB has one child: NodeA
+		// NodeC has two children: NodeA and NodeB (but both traits are base, only NodeB will be referenced)
+		// NodeD has two children: NodeA and NodeC
+
+		TArray<FNodeHandle> NodeHandles;
+
+		// Write our graph
+		TArray<uint8> GraphSharedDataArchiveBuffer;
+		TArray<TObjectPtr<UObject>> GraphReferencedObjects;
+		{
+			FTraitWriter TraitWriter;
+
+			NodeHandles.Add(TraitWriter.RegisterNode(*NodeTemplateA));
+			NodeHandles.Add(TraitWriter.RegisterNode(*NodeTemplateB));
+			NodeHandles.Add(TraitWriter.RegisterNode(*NodeTemplateC));
+			NodeHandles.Add(TraitWriter.RegisterNode(*NodeTemplateD));
+
+			// We don't have trait properties
+			TArray<TMap<FName, FString>> TraitPropertiesA;
+			TraitPropertiesA.AddDefaulted(NodeTemplateTraitListA.Num());
+
+			TArray<TMap<FName, FString>> TraitPropertiesB;
+			TraitPropertiesB.AddDefaulted(NodeTemplateTraitListB.Num());
+			TraitPropertiesB[0].Add(TEXT("Child"), ToString<FTraitWithOneChild::FSharedData>(TEXT("Child"), FAnimNextTraitHandle(NodeHandles[0])));
+
+			TArray<TMap<FName, FString>> TraitPropertiesC;
+			TraitPropertiesC.AddDefaulted(NodeTemplateTraitListC.Num());
+			TraitPropertiesC[0].Add(TEXT("Child"), ToString<FTraitWithOneChild::FSharedData>(TEXT("Child"), FAnimNextTraitHandle(NodeHandles[0])));
+			TraitPropertiesC[1].Add(TEXT("Child"), ToString<FTraitWithOneChild::FSharedData>(TEXT("Child"), FAnimNextTraitHandle(NodeHandles[1])));
+
+			TArray<TMap<FName, FString>> TraitPropertiesD;
+			TraitPropertiesD.AddDefaulted(NodeTemplateTraitListD.Num());
+			FAnimNextTraitHandle ChildrenHandlesD[2] = { FAnimNextTraitHandle(NodeHandles[0]), FAnimNextTraitHandle(NodeHandles[2], 1)};
+			TraitPropertiesD[0].Add(TEXT("Children"), ToString<FTraitWithChildren::FSharedData>(TEXT("Children"), ChildrenHandlesD));
+
+			TraitWriter.BeginNodeWriting();
+			TraitWriter.WriteNode(NodeHandles[0],
+				[&TraitPropertiesA](uint32 TraitIndex, FName PropertyName)
+				{
+					return TraitPropertiesA[TraitIndex][PropertyName];
+				},
+				[](uint32 TraitIndex, FName PropertyName)
+				{
+					return MAX_uint16;
+				});
+			TraitWriter.WriteNode(NodeHandles[1],
+				[&TraitPropertiesB](uint32 TraitIndex, FName PropertyName)
+				{
+					return TraitPropertiesB[TraitIndex][PropertyName];
+				},
+				[](uint32 TraitIndex, FName PropertyName)
+				{
+					return MAX_uint16;
+				});
+			TraitWriter.WriteNode(NodeHandles[2],
+				[&TraitPropertiesC](uint32 TraitIndex, FName PropertyName)
+				{
+					return TraitPropertiesC[TraitIndex][PropertyName];
+				},
+				[](uint32 TraitIndex, FName PropertyName)
+				{
+					return MAX_uint16;
+				});
+			TraitWriter.WriteNode(NodeHandles[3],
+				[&TraitPropertiesD](uint32 TraitIndex, FName PropertyName)
+				{
+					return TraitPropertiesD[TraitIndex][PropertyName];
+				},
+				[](uint32 TraitIndex, FName PropertyName)
+				{
+					return MAX_uint16;
+				});
+			TraitWriter.EndNodeWriting();
+
+			AddErrorIfFalse(TraitWriter.GetErrorState() == FTraitWriter::EErrorState::None, "FAnimationAnimNextRuntimeTest_IHierarchy -> Failed to write traits");
+			GraphSharedDataArchiveBuffer = TraitWriter.GetGraphSharedData();
+			GraphReferencedObjects = TraitWriter.GetGraphReferencedObjects();
+		}
+
+		// Read our graph
+		FTestUtils::LoadFromArchiveBuffer(*AnimNextGraph, NodeHandles, GraphSharedDataArchiveBuffer);
+
+		FAnimNextGraphInstancePtr GraphInstance;
+		AnimNextGraph->AllocateInstance(GraphInstance);
+
+		FExecutionContext Context(GraphInstance);
+
+		{
+			FMemMark Mark(FMemStack::Get());
+
+			FWeakTraitPtr NullPtr;								// Empty, no parent
+			FAnimNextTraitHandle RootHandle(NodeHandles[3]);	// Point to NodeD, first base trait
+
+			FTraitPtr NodeDPtr = Context.AllocateNodeInstance(NullPtr, RootHandle);
+			AddErrorIfFalse(NodeDPtr.IsValid(), "FAnimationAnimNextRuntimeTest_IHierarchy -> Failed to allocate root node instance");
+
+			TTraitBinding<IHierarchy> HierarchyBindingNodeD;
+			AddErrorIfFalse(Context.GetInterface(NodeDPtr, HierarchyBindingNodeD), "FAnimationAnimNextRuntimeTest_IHierarchy -> IHierarchy not found");
+
+			FChildrenArray ChildrenNodeD;
+			HierarchyBindingNodeD.GetChildren(Context, ChildrenNodeD);
+
+			AddErrorIfFalse(ChildrenNodeD.Num() == 2, "FAnimationAnimNextRuntimeTest_IHierarchy -> Expected 2 children");
+			AddErrorIfFalse(ChildrenNodeD[0].IsValid() && ChildrenNodeD[0].GetNodeInstance()->GetNodeHandle() == NodeHandles[0], "FAnimationAnimNextRuntimeTest_IHierarchy -> Expected child: NodeA");
+			AddErrorIfFalse(ChildrenNodeD[1].IsValid() && ChildrenNodeD[1].GetNodeInstance()->GetNodeHandle() == NodeHandles[2], "FAnimationAnimNextRuntimeTest_IHierarchy -> Expected child: NodeC");
+
+			{
+				TTraitBinding<IHierarchy> HierarchyBindingNodeC;
+				AddErrorIfFalse(Context.GetInterface(ChildrenNodeD[1], HierarchyBindingNodeC), "FAnimationAnimNextRuntimeTest_IHierarchy -> IHierarchy not found");
+
+				FChildrenArray ChildrenNodeC;
+				HierarchyBindingNodeC.GetChildren(Context, ChildrenNodeC);
+
+				AddErrorIfFalse(ChildrenNodeC.Num() == 1, "FAnimationAnimNextRuntimeTest_IHierarchy -> Expected 2 children");
+				AddErrorIfFalse(ChildrenNodeC[0].IsValid() && ChildrenNodeC[0].GetNodeInstance()->GetNodeHandle() == NodeHandles[1], "FAnimationAnimNextRuntimeTest_IHierarchy -> Expected child: NodeB");
+
+				{
+					TTraitBinding<IHierarchy> HierarchyBindingNodeB;
+					AddErrorIfFalse(Context.GetInterface(ChildrenNodeC[0], HierarchyBindingNodeB), "FAnimationAnimNextRuntimeTest_IHierarchy -> IHierarchy not found");
+
+					FChildrenArray ChildrenNodeB;
+					HierarchyBindingNodeB.GetChildren(Context, ChildrenNodeB);
+
+					AddErrorIfFalse(ChildrenNodeB.Num() == 1, "FAnimationAnimNextRuntimeTest_IHierarchy -> Expected 1 child");
+					AddErrorIfFalse(ChildrenNodeB[0].IsValid() && ChildrenNodeB[0].GetNodeInstance()->GetNodeHandle() == NodeHandles[0], "FAnimationAnimNextRuntimeTest_IHierarchy -> Expected child: NodeA");
+				}
+			}
+		}
+
+		Registry.Unregister(NodeTemplateA);
+		Registry.Unregister(NodeTemplateB);
+		Registry.Unregister(NodeTemplateC);
+		Registry.Unregister(NodeTemplateD);
+
+		AddErrorIfFalse(Registry.GetNum() == 0, "FAnimationAnimNextRuntimeTest_IHierarchy -> Registry should contain 0 templates");
+	}
+
+	Tests::FUtils::CleanupAfterTests();
+	
+	return true;
+}
+
+IMPLEMENT_SIMPLE_AUTOMATION_TEST(FAnimationAnimNextRuntimeTest_IUpdate, "Animation.AnimNext.Runtime.IUpdate", EAutomationTestFlags::EditorContext | EAutomationTestFlags::EngineFilter)
+
+bool FAnimationAnimNextRuntimeTest_IUpdate::RunTest(const FString& InParameters)
+{
+	using namespace UE::AnimNext;
+	{
+		AUTO_REGISTER_ANIM_TRAIT(FTraitWithNoChildren)
+		AUTO_REGISTER_ANIM_TRAIT(FTraitWithOneChild)
+		AUTO_REGISTER_ANIM_TRAIT(FTraitWithChildren)
+
+		UFactory* GraphFactory = NewObject<UAnimNextGraphFactory>();
+		UAnimNextGraph* AnimNextGraph = CastChecked<UAnimNextGraph>(GraphFactory->FactoryCreateNew(UAnimNextGraph::StaticClass(), GetTransientPackage(), TEXT("TestAnimNextGraph"), RF_Transient, nullptr, nullptr, NAME_None));
+		UE_RETURN_ON_ERROR(AnimNextGraph != nullptr, "FAnimationAnimNextRuntimeTest_IUpdate -> Failed to create animation graph");
+
+		FScopedClearNodeTemplateRegistry ScopedClearNodeTemplateRegistry;
+		FNodeTemplateRegistry& Registry = FNodeTemplateRegistry::Get();
+
+		// We create a few node templates
+		// Template A has a single trait with no children
+		TArray<FTraitUID> NodeTemplateTraitListA;
+		NodeTemplateTraitListA.Add(FTraitWithNoChildren::TraitUID);
+
+		// Template B has a single trait with one child, it doesn't update
+		TArray<FTraitUID> NodeTemplateTraitListB;
+		NodeTemplateTraitListB.Add(FTraitWithOneChild::TraitUID);
+
+		// Template C has a single trait with children
+		TArray<FTraitUID> NodeTemplateTraitListC;
+		NodeTemplateTraitListC.Add(FTraitWithChildren::TraitUID);
+
+		// Populate our node template registry
+		TArray<uint8> NodeTemplateBufferA, NodeTemplateBufferB, NodeTemplateBufferC;
+		const FNodeTemplate* NodeTemplateA = FNodeTemplateBuilder::BuildNodeTemplate(NodeTemplateTraitListA, NodeTemplateBufferA);
+		const FNodeTemplate* NodeTemplateB = FNodeTemplateBuilder::BuildNodeTemplate(NodeTemplateTraitListB, NodeTemplateBufferB);
+		const FNodeTemplate* NodeTemplateC = FNodeTemplateBuilder::BuildNodeTemplate(NodeTemplateTraitListC, NodeTemplateBufferC);
+
+		// Build our graph, it as follow (each node template has a single node instance):
+		// NodeA has no children
+		// NodeB has one child: NodeA (it doesn't update)
+		// NodeC (root) has two children: NodeA and NodeB
+
+		TArray<FNodeHandle> NodeHandles;
+
+		// Write our graph
+		TArray<uint8> GraphSharedDataArchiveBuffer;
+		TArray<TObjectPtr<UObject>> GraphReferencedObjects;
+		{
+			FTraitWriter TraitWriter;
+
+			NodeHandles.Add(TraitWriter.RegisterNode(*NodeTemplateC));	// Root node
+			NodeHandles.Add(TraitWriter.RegisterNode(*NodeTemplateA));
+			NodeHandles.Add(TraitWriter.RegisterNode(*NodeTemplateB));
+
+			// We don't have trait properties
+			TArray<TMap<FName, FString>> TraitPropertiesA;
+			TraitPropertiesA.AddDefaulted(NodeTemplateTraitListA.Num());
+
+			TArray<TMap<FName, FString>> TraitPropertiesB;
+			TraitPropertiesB.AddDefaulted(NodeTemplateTraitListB.Num());
+			TraitPropertiesB[0].Add(TEXT("Child"), ToString<FTraitWithOneChild::FSharedData>(TEXT("Child"), FAnimNextTraitHandle(NodeHandles[1])));
+
+			TArray<TMap<FName, FString>> TraitPropertiesC;
+			TraitPropertiesC.AddDefaulted(NodeTemplateTraitListC.Num());
+			FAnimNextTraitHandle ChildrenHandlesC[2] = { FAnimNextTraitHandle(NodeHandles[1]), FAnimNextTraitHandle(NodeHandles[2])};
+			TraitPropertiesC[0].Add(TEXT("Children"), ToString<FTraitWithChildren::FSharedData>(TEXT("Children"), ChildrenHandlesC));
+
+			TraitWriter.BeginNodeWriting();
+			TraitWriter.WriteNode(NodeHandles[0],
+				[&TraitPropertiesC](uint32 TraitIndex, FName PropertyName)
+				{
+					return TraitPropertiesC[TraitIndex][PropertyName];
+				},
+				[](uint32 TraitIndex, FName PropertyName)
+				{
+					return MAX_uint16;
+				});
+			TraitWriter.WriteNode(NodeHandles[1],
+				[&TraitPropertiesA](uint32 TraitIndex, FName PropertyName)
+				{
+					return TraitPropertiesA[TraitIndex][PropertyName];
+				},
+				[](uint32 TraitIndex, FName PropertyName)
+				{
+					return MAX_uint16;
+				});
+			TraitWriter.WriteNode(NodeHandles[2],
+				[&TraitPropertiesB](uint32 TraitIndex, FName PropertyName)
+				{
+					return TraitPropertiesB[TraitIndex][PropertyName];
+				},
+				[](uint32 TraitIndex, FName PropertyName)
+				{
+					return MAX_uint16;
+				});
+			TraitWriter.EndNodeWriting();
+
+			AddErrorIfFalse(TraitWriter.GetErrorState() == FTraitWriter::EErrorState::None, "FAnimationAnimNextRuntimeTest_IUpdate -> Failed to write traits");
+			GraphSharedDataArchiveBuffer = TraitWriter.GetGraphSharedData();
+			GraphReferencedObjects = TraitWriter.GetGraphReferencedObjects();
+		}
+
+		// Read our graph
+		FTestUtils::LoadFromArchiveBuffer(*AnimNextGraph, NodeHandles, GraphSharedDataArchiveBuffer);
+
+		FAnimNextGraphInstancePtr GraphInstance;
+		AnimNextGraph->AllocateInstance(GraphInstance);
+
+		FExecutionContext Context(GraphInstance);
+
+		{
+			TArray<FTraitUID> UpdatedTraits;
+			Private::UpdatedTraits = &UpdatedTraits;
+
+			// Call pre/post update on our graph
+			UpdateGraph(GraphInstance, 0.0333f);
+
+			AddErrorIfFalse(UpdatedTraits.Num() == 6, "FAnimationAnimNextRuntimeTest_IUpdate -> Expected 6 nodes to have been visited during the update traversal");
+			AddErrorIfFalse(UpdatedTraits[0] == FTraitWithChildren::TraitUID, "FAnimationAnimNextRuntimeTest_IUpdate -> Unexpected update order");		// NodeC
+			AddErrorIfFalse(UpdatedTraits[1] == FTraitWithNoChildren::TraitUID, "FAnimationAnimNextRuntimeTest_IUpdate -> Unexpected update order");		// NodeA
+			AddErrorIfFalse(UpdatedTraits[2] == FTraitWithNoChildren::TraitUID, "FAnimationAnimNextRuntimeTest_IUpdate -> Unexpected update order");		// NodeA
+			AddErrorIfFalse(UpdatedTraits[3] == FTraitWithNoChildren::TraitUID, "FAnimationAnimNextRuntimeTest_IUpdate -> Unexpected update order");		// NodeB -> NodeA
+			AddErrorIfFalse(UpdatedTraits[4] == FTraitWithNoChildren::TraitUID, "FAnimationAnimNextRuntimeTest_IUpdate -> Unexpected update order");		// NodeB -> NodeA
+			AddErrorIfFalse(UpdatedTraits[5] == FTraitWithChildren::TraitUID, "FAnimationAnimNextRuntimeTest_IUpdate -> Unexpected update order");		// NodeC
+
+			Private::UpdatedTraits = nullptr;
+		}
+
+		Registry.Unregister(NodeTemplateA);
+		Registry.Unregister(NodeTemplateB);
+		Registry.Unregister(NodeTemplateC);
+
+		AddErrorIfFalse(Registry.GetNum() == 0, "FAnimationAnimNextRuntimeTest_IUpdate -> Registry should contain 0 templates");
+	}
+	Tests::FUtils::CleanupAfterTests();
+
+	return true;
+}
+
+IMPLEMENT_SIMPLE_AUTOMATION_TEST(FAnimationAnimNextRuntimeTest_IEvaluate, "Animation.AnimNext.Runtime.IEvaluate", EAutomationTestFlags::EditorContext | EAutomationTestFlags::EngineFilter)
+
+bool FAnimationAnimNextRuntimeTest_IEvaluate::RunTest(const FString& InParameters)
+{
+	using namespace UE::AnimNext;
+	{
+		AUTO_REGISTER_ANIM_TRAIT(FTraitWithNoChildren)
+		AUTO_REGISTER_ANIM_TRAIT(FTraitWithOneChild)
+		AUTO_REGISTER_ANIM_TRAIT(FTraitWithChildren)
+
+		UFactory* GraphFactory = NewObject<UAnimNextGraphFactory>();
+		UAnimNextGraph* AnimNextGraph = CastChecked<UAnimNextGraph>(GraphFactory->FactoryCreateNew(UAnimNextGraph::StaticClass(), GetTransientPackage(), TEXT("TestAnimNextGraph"), RF_Transient, nullptr, nullptr, NAME_None));
+		UE_RETURN_ON_ERROR(AnimNextGraph != nullptr, "FAnimationAnimNextRuntimeTest_IEvaluate -> Failed to create animation graph");
+
+		FScopedClearNodeTemplateRegistry ScopedClearNodeTemplateRegistry;
+		FNodeTemplateRegistry& Registry = FNodeTemplateRegistry::Get();
+
+		// We create a few node templates
+		// Template A has a single trait with no children
+		TArray<FTraitUID> NodeTemplateTraitListA;
+		NodeTemplateTraitListA.Add(FTraitWithNoChildren::TraitUID);
+
+		// Template B has a single trait with one child, it doesn't evaluate
+		TArray<FTraitUID> NodeTemplateTraitListB;
+		NodeTemplateTraitListB.Add(FTraitWithOneChild::TraitUID);
+
+		// Template C has a single trait with children
+		TArray<FTraitUID> NodeTemplateTraitListC;
+		NodeTemplateTraitListC.Add(FTraitWithChildren::TraitUID);
+
+		// Populate our node template registry
+		TArray<uint8> NodeTemplateBufferA, NodeTemplateBufferB, NodeTemplateBufferC;
+		const FNodeTemplate* NodeTemplateA = FNodeTemplateBuilder::BuildNodeTemplate(NodeTemplateTraitListA, NodeTemplateBufferA);
+		const FNodeTemplate* NodeTemplateB = FNodeTemplateBuilder::BuildNodeTemplate(NodeTemplateTraitListB, NodeTemplateBufferB);
+		const FNodeTemplate* NodeTemplateC = FNodeTemplateBuilder::BuildNodeTemplate(NodeTemplateTraitListC, NodeTemplateBufferC);
+
+		// Build our graph, it as follow (each node template has a single node instance):
+		// NodeA has no children
+		// NodeB has one child: NodeA (it doesn't evaluate)
+		// NodeC (root) has two children: NodeA and NodeB
+
+		TArray<FNodeHandle> NodeHandles;
+
+		// Write our graph
+		TArray<uint8> GraphSharedDataArchiveBuffer;
+		TArray<TObjectPtr<UObject>> GraphReferencedObjects;
+		{
+			FTraitWriter TraitWriter;
+
+			NodeHandles.Add(TraitWriter.RegisterNode(*NodeTemplateC));	// Root node
+			NodeHandles.Add(TraitWriter.RegisterNode(*NodeTemplateA));
+			NodeHandles.Add(TraitWriter.RegisterNode(*NodeTemplateB));
+
+			// We don't have trait properties
+			TArray<TMap<FName, FString>> TraitPropertiesA;
+			TraitPropertiesA.AddDefaulted(NodeTemplateTraitListA.Num());
+
+			TArray<TMap<FName, FString>> TraitPropertiesB;
+			TraitPropertiesB.AddDefaulted(NodeTemplateTraitListB.Num());
+			TraitPropertiesB[0].Add(TEXT("Child"), ToString<FTraitWithOneChild::FSharedData>(TEXT("Child"), FAnimNextTraitHandle(NodeHandles[1])));
+
+			TArray<TMap<FName, FString>> TraitPropertiesC;
+			TraitPropertiesC.AddDefaulted(NodeTemplateTraitListC.Num());
+
+			FAnimNextTraitHandle ChildrenHandlesC[2] = { FAnimNextTraitHandle(NodeHandles[1]), FAnimNextTraitHandle(NodeHandles[2])};
+			TraitPropertiesC[0].Add(TEXT("Children"), ToString<FTraitWithChildren::FSharedData>(TEXT("Children"), ChildrenHandlesC));
+
+			TraitWriter.BeginNodeWriting();
+			TraitWriter.WriteNode(NodeHandles[0],
+				[&TraitPropertiesC](uint32 TraitIndex, FName PropertyName)
+				{
+					return TraitPropertiesC[TraitIndex][PropertyName];
+				},
+				[](uint32 TraitIndex, FName PropertyName)
+				{
+					return MAX_uint16;
+				});
+			TraitWriter.WriteNode(NodeHandles[1],
+				[&TraitPropertiesA](uint32 TraitIndex, FName PropertyName)
+				{
+					return TraitPropertiesA[TraitIndex][PropertyName];
+				},
+				[](uint32 TraitIndex, FName PropertyName)
+				{
+					return MAX_uint16;
+				});
+			TraitWriter.WriteNode(NodeHandles[2],
+				[&TraitPropertiesB](uint32 TraitIndex, FName PropertyName)
+				{
+					return TraitPropertiesB[TraitIndex][PropertyName];
+				},
+				[](uint32 TraitIndex, FName PropertyName)
+				{
+					return MAX_uint16;
+				});
+			TraitWriter.EndNodeWriting();
+
+			AddErrorIfFalse(TraitWriter.GetErrorState() == FTraitWriter::EErrorState::None, "FAnimationAnimNextRuntimeTest_IEvaluate -> Failed to write traits");
+			GraphSharedDataArchiveBuffer = TraitWriter.GetGraphSharedData();
+			GraphReferencedObjects = TraitWriter.GetGraphReferencedObjects();
+		}
+
+		// Read our graph
+		FTestUtils::LoadFromArchiveBuffer(*AnimNextGraph, NodeHandles, GraphSharedDataArchiveBuffer);
+
+		FAnimNextGraphInstancePtr GraphInstance;
+		AnimNextGraph->AllocateInstance(GraphInstance);
+
+		{
+			TArray<FTraitUID> EvaluatedTraits;
+			Private::EvaluatedTraits = &EvaluatedTraits;
+
+			// Call pre/post evaluate on our graph
+			(void)EvaluateGraph(GraphInstance);
+
+			AddErrorIfFalse(EvaluatedTraits.Num() == 6, "FAnimationAnimNextRuntimeTest_IEvaluate -> Expected 6 nodes to have been visited during the evaluate traversal");
+			AddErrorIfFalse(EvaluatedTraits[0] == FTraitWithChildren::TraitUID, "FAnimationAnimNextRuntimeTest_IEvaluate -> Unexpected evaluate order");		// NodeC
+			AddErrorIfFalse(EvaluatedTraits[1] == FTraitWithNoChildren::TraitUID, "FAnimationAnimNextRuntimeTest_IEvaluate -> Unexpected evaluate order");			// NodeA
+			AddErrorIfFalse(EvaluatedTraits[2] == FTraitWithNoChildren::TraitUID, "FAnimationAnimNextRuntimeTest_IEvaluate -> Unexpected evaluate order");			// NodeA
+			AddErrorIfFalse(EvaluatedTraits[3] == FTraitWithNoChildren::TraitUID, "FAnimationAnimNextRuntimeTest_IEvaluate -> Unexpected evaluate order");			// NodeB -> NodeA
+			AddErrorIfFalse(EvaluatedTraits[4] == FTraitWithNoChildren::TraitUID, "FAnimationAnimNextRuntimeTest_IEvaluate -> Unexpected evaluate order");			// NodeB -> NodeA
+			AddErrorIfFalse(EvaluatedTraits[5] == FTraitWithChildren::TraitUID, "FAnimationAnimNextRuntimeTest_IEvaluate -> Unexpected evaluate order");		// NodeC
+
+			Private::EvaluatedTraits = nullptr;
+		}
+
+		Registry.Unregister(NodeTemplateA);
+		Registry.Unregister(NodeTemplateB);
+		Registry.Unregister(NodeTemplateC);
+
+		AddErrorIfFalse(Registry.GetNum() == 0, "FAnimationAnimNextRuntimeTest_IEvaluate -> Registry should contain 0 templates");
+	}
+	Tests::FUtils::CleanupAfterTests();
+
+	return true;
+}
+
+#endif
diff --git a/Engine/Plugins/Experimental/AnimNext/Source/AnimNextTestSuite/Private/AnimNextTraitInterfacesTest.h b/Engine/Plugins/Experimental/AnimNext/Source/AnimNextTestSuite/Private/AnimNextTraitInterfacesTest.h
new file mode 100644
--- /dev/null
+++ b/Engine/Plugins/Experimental/AnimNext/Source/AnimNextTestSuite/Private/AnimNextTraitInterfacesTest.h
@@ -0,0 +1,28 @@
+// Copyright Epic Games, Inc. All Rights Reserved.
+
+#pragma once
+
+#include "CoreMinimal.h"
+
+#include "TraitCore/TraitHandle.h"
+#include "TraitCore/TraitSharedData.h"
+
+#include "AnimNextTraitInterfacesTest.generated.h"
+
+USTRUCT()
+struct FTraitWithOneChildSharedData : public FAnimNextTraitSharedData
+{
+	GENERATED_BODY()
+
+	UPROPERTY()
+	FAnimNextTraitHandle Child;
+};
+
+USTRUCT()
+struct FTraitWithChildrenSharedData : public FAnimNextTraitSharedData
+{
+	GENERATED_BODY()
+
+	UPROPERTY()
+	FAnimNextTraitHandle Children[2];
+};
diff --git a/Engine/Plugins/Experimental/AnimNext/Source/AnimNextUncookedOnly/Internal/Graph/AnimNextGraph_EdGraphNode.h b/Engine/Plugins/Experimental/AnimNext/Source/AnimNextUncookedOnly/Internal/Graph/AnimNextGraph_EdGraphNode.h
--- a/Engine/Plugins/Experimental/AnimNext/Source/AnimNextUncookedOnly/Internal/Graph/AnimNextGraph_EdGraphNode.h
+++ b/Engine/Plugins/Experimental/AnimNext/Source/AnimNextUncookedOnly/Internal/Graph/AnimNextGraph_EdGraphNode.h
@@ -8,7 +8,7 @@
 #include "AnimNextGraph_EdGraphNode.generated.h"
 
 // EdGraphNode representation for AnimNext nodes
-// A node can hold a decorator stack or a decorator entry
+// A node can hold a trait stack or a trait entry
 UCLASS(MinimalAPI)
 class UAnimNextGraph_EdGraphNode : public URigVMEdGraphNode
 {
@@ -26,10 +26,10 @@ public:
 	//////////////////////////////////////////////////////////////////////////
 	// Our implementation
 
-	// Returns whether this node is a decorator stack or not
-	ANIMNEXTUNCOOKEDONLY_API bool IsDecoratorStack() const;
+	// Returns whether this node is a trait stack or not
+	ANIMNEXTUNCOOKEDONLY_API bool IsTraitStack() const;
 
 private:
-	// Populates the SubMenu with entries for each decorator that can be added through the context menu
-	void BuildAddDecoratorContextMenu(class UToolMenu* SubMenu);
+	// Populates the SubMenu with entries for each trait that can be added through the context menu
+	void BuildAddTraitContextMenu(class UToolMenu* SubMenu);
 };
diff --git a/Engine/Plugins/Experimental/AnimNext/Source/AnimNextUncookedOnly/Internal/Graph/AnimNextGraph_EditorData.h b/Engine/Plugins/Experimental/AnimNext/Source/AnimNextUncookedOnly/Internal/Graph/AnimNextGraph_EditorData.h
--- a/Engine/Plugins/Experimental/AnimNext/Source/AnimNextUncookedOnly/Internal/Graph/AnimNextGraph_EditorData.h
+++ b/Engine/Plugins/Experimental/AnimNext/Source/AnimNextUncookedOnly/Internal/Graph/AnimNextGraph_EditorData.h
@@ -13,7 +13,7 @@
 
 class UAnimNextGraph;
 enum class ERigVMGraphNotifType : uint8;
-class FAnimationAnimNextRuntimeTest_GraphAddDecorator;
+class FAnimationAnimNextRuntimeTest_GraphAddTrait;
 class FAnimationAnimNextRuntimeTest_GraphExecute;
 class FAnimationAnimNextRuntimeTest_GraphExecuteLatent;
 
@@ -65,7 +65,7 @@ class UAnimNextGraph_EditorData : public UAnimNextRigVMAssetEditorData
 	friend class UE::AnimNext::Editor::SAnimNextGraphView;
 	friend struct FAnimNextGraphSchemaAction_RigUnit;
 	friend struct FAnimNextGraphSchemaAction_DispatchFactory;
-	friend class FAnimationAnimNextRuntimeTest_GraphAddDecorator;
+	friend class FAnimationAnimNextRuntimeTest_GraphAddTrait;
 	friend class FAnimationAnimNextRuntimeTest_GraphExecute;
 	friend class FAnimationAnimNextRuntimeTest_GraphExecuteLatent;
 	
diff --git a/Engine/Plugins/Experimental/AnimNext/Source/AnimNextUncookedOnly/Private/Graph/AnimNextGraph_EdGraphNode.cpp b/Engine/Plugins/Experimental/AnimNext/Source/AnimNextUncookedOnly/Private/Graph/AnimNextGraph_EdGraphNode.cpp
--- a/Engine/Plugins/Experimental/AnimNext/Source/AnimNextUncookedOnly/Private/Graph/AnimNextGraph_EdGraphNode.cpp
+++ b/Engine/Plugins/Experimental/AnimNext/Source/AnimNextUncookedOnly/Private/Graph/AnimNextGraph_EdGraphNode.cpp
@@ -2,11 +2,11 @@
 
 #include "Graph/AnimNextGraph_EdGraphNode.h"
 
-#include "DecoratorBase/DecoratorHandle.h"
-#include "Graph/RigUnit_AnimNextDecoratorStack.h"
-#include "Graph/RigDecorator_AnimNextCppDecorator.h"
+#include "TraitCore/TraitHandle.h"
+#include "Graph/RigUnit_AnimNextTraitStack.h"
+#include "Graph/RigDecorator_AnimNextCppTrait.h"
 #include "RigVMModel/RigVMController.h"
-#include "DecoratorBase/DecoratorRegistry.h"
+#include "TraitCore/TraitRegistry.h"
 #include "ToolMenu.h"
 
 #define LOCTEXT_NAMESPACE "AnimNextGraph_EdGraphNode"
@@ -15,17 +15,17 @@ void UAnimNextGraph_EdGraphNode::GetNodeContextMenuActions(UToolMenu* Menu, clas
 {
 	URigVMEdGraphNode::GetNodeContextMenuActions(Menu, Context);
 
-	if (IsDecoratorStack())
+	if (IsTraitStack())
 	{
-		FToolMenuSection& Section = Menu->AddSection("AnimNextDecoratorNodeActions", LOCTEXT("AnimNextDecoratorNodeActionsMenuHeader", "Anim Next Decorator Actions"));
+		FToolMenuSection& Section = Menu->AddSection("AnimNextTraitNodeActions", LOCTEXT("AnimNextTraitNodeActionsMenuHeader", "Anim Next Trait Actions"));
 
 		UAnimNextGraph_EdGraphNode* NonConstThis = const_cast<UAnimNextGraph_EdGraphNode*>(this);
 
 		Section.AddSubMenu(
-			TEXT("AddDecoratorMenu"),
-			LOCTEXT("AddDecoratorMenu", "Add Decorator"),
-			LOCTEXT("AddDecoratorMenuTooltip", "Add the chosen decorator to currently selected node"),
-			FNewToolMenuDelegate::CreateUObject(NonConstThis, &UAnimNextGraph_EdGraphNode::BuildAddDecoratorContextMenu));
+			TEXT("AddTraitMenu"),
+			LOCTEXT("AddTraitMenu", "Add Trait"),
+			LOCTEXT("AddTraitMenuTooltip", "Add the chosen trait to currently selected node"),
+			FNewToolMenuDelegate::CreateUObject(NonConstThis, &UAnimNextGraph_EdGraphNode::BuildAddTraitContextMenu));
 	}
 }
 
@@ -33,11 +33,11 @@ void UAnimNextGraph_EdGraphNode::ConfigurePin(UEdGraphPin* EdGraphPin, const URi
 {
 	Super::ConfigurePin(EdGraphPin, ModelPin);
 
-	// Decorator handles always remain as a RigVM input pins so that we can still link things to them even if they are hidden
+	// Trait handles always remain as a RigVM input pins so that we can still link things to them even if they are hidden
 	// We handle visibility for those explicitly here
 	const bool bIsInputPin = ModelPin->GetDirection() == ERigVMPinDirection::Input;
-	const bool bIsDecoratorHandle = ModelPin->GetCPPTypeObject() == FAnimNextDecoratorHandle::StaticStruct();
-	if (bIsInputPin && bIsDecoratorHandle)
+	const bool bIsTraitHandle = ModelPin->GetCPPTypeObject() == FAnimNextTraitHandle::StaticStruct();
+	if (bIsInputPin && bIsTraitHandle)
 	{
 		if (const URigVMPin* DecoratorPin = ModelPin->GetParentPin())
 		{
@@ -48,42 +48,42 @@ void UAnimNextGraph_EdGraphNode::ConfigurePin(UEdGraphPin* EdGraphPin, const URi
 				TSharedPtr<FStructOnScope> DecoratorScope = DecoratorPin->GetDecoratorInstance();
 				const FRigDecorator_AnimNextCppDecorator* VMDecorator = (const FRigDecorator_AnimNextCppDecorator*)DecoratorScope->GetStructMemory();
 
-				const UScriptStruct* DecoratorStruct = VMDecorator->GetDecoratorSharedDataStruct();
-				check(DecoratorStruct != nullptr);
+				const UScriptStruct* TraitStruct = VMDecorator->GetDecoratorSharedDataStruct();
+				check(TraitStruct != nullptr);
 
-				const FProperty* PinProperty = DecoratorStruct->FindPropertyByName(ModelPin->GetFName());
+				const FProperty* PinProperty = TraitStruct->FindPropertyByName(ModelPin->GetFName());
 				EdGraphPin->bHidden = PinProperty->HasMetaData(FRigVMStruct::HiddenMetaName);
 			}
 		}
 	}
 }
 
-bool UAnimNextGraph_EdGraphNode::IsDecoratorStack() const
+bool UAnimNextGraph_EdGraphNode::IsTraitStack() const
 {
 	if (const URigVMUnitNode* VMNode = Cast<URigVMUnitNode>(GetModelNode()))
 	{
 		const UScriptStruct* ScriptStruct = VMNode->GetScriptStruct();
-		return ScriptStruct == FRigUnit_AnimNextDecoratorStack::StaticStruct();
+		return ScriptStruct == FRigUnit_AnimNextTraitStack::StaticStruct();
 	}
 
 	return false;
 }
 
-void UAnimNextGraph_EdGraphNode::BuildAddDecoratorContextMenu(UToolMenu* SubMenu)
+void UAnimNextGraph_EdGraphNode::BuildAddTraitContextMenu(UToolMenu* SubMenu)
 {
 	using namespace UE::AnimNext;
 
-	const FDecoratorRegistry& DecoratorRegistry = FDecoratorRegistry::Get();
-	TArray<const FDecorator*> Decorators = DecoratorRegistry.GetDecorators();
+	const FTraitRegistry& TraitRegistry = FTraitRegistry::Get();
+	TArray<const FTrait*> Traits = TraitRegistry.GetTraits();
 
 	URigVMController* VMController = GetController();
 	URigVMNode* VMNode = GetModelNode();
 
 	const UScriptStruct* CppDecoratorStruct = FRigDecorator_AnimNextCppDecorator::StaticStruct();
 
-	for (const FDecorator* Decorator : Decorators)
+	for (const FTrait* Trait : Traits)
 	{
-		UScriptStruct* ScriptStruct = Decorator->GetDecoratorSharedDataStruct();
+		UScriptStruct* ScriptStruct = Trait->GetTraitSharedDataStruct();
 
 		FString DefaultValue;
 		{
@@ -93,26 +93,26 @@ void UAnimNextGraph_EdGraphNode::BuildAddDecoratorContextMenu(UToolMenu* SubMenu
 
 			if (!CppDecoratorStructInstance.CanBeAddedToNode(VMNode, nullptr))
 			{
-				continue;	// This decorator isn't supported on this node
+				continue;	// This trait isn't supported on this node
 			}
 
 			const FProperty* Prop = FAnimNextCppDecoratorWrapper::StaticStruct()->FindPropertyByName(GET_MEMBER_NAME_STRING_CHECKED(FAnimNextCppDecoratorWrapper, CppDecorator));
-			Prop->ExportText_Direct(DefaultValue, &CppDecoratorStructInstance, &DefaultCppDecoratorStructInstance, nullptr, PPF_None);
+			Prop->ExportText_Direct(DefaultValue, &CppDecoratorStructInstance, &DefaultCppDecoratorStructInstance, nullptr, PPF_SerializedAsImportText);
 		}
 
 		FString DisplayNameMetadata;
 		ScriptStruct->GetStringMetaDataHierarchical(FRigVMStruct::DisplayNameMetaName, &DisplayNameMetadata);
-		const FString DisplayName = DisplayNameMetadata.IsEmpty() ? Decorator->GetDecoratorName() : DisplayNameMetadata;
+		const FString DisplayName = DisplayNameMetadata.IsEmpty() ? Trait->GetTraitName() : DisplayNameMetadata;
 
 		const FText ToolTip = ScriptStruct->GetToolTipText();
 
-		FToolMenuEntry DecoratorEntry = FToolMenuEntry::InitMenuEntry(
-			*Decorator->GetDecoratorName(),
+		FToolMenuEntry TraitEntry = FToolMenuEntry::InitMenuEntry(
+			*Trait->GetTraitName(),
 			FText::FromString(DisplayName),
 			ToolTip,
 			FSlateIcon(),
 			FUIAction(FExecuteAction::CreateLambda(
-				[this, Decorator, VMController, VMNode, CppDecoratorStruct, DefaultValue, DisplayName]()
+				[this, Trait, VMController, VMNode, CppDecoratorStruct, DefaultValue, DisplayName]()
 				{
 					VMController->AddDecorator(
 						VMNode->GetFName(),
@@ -123,7 +123,7 @@ void UAnimNextGraph_EdGraphNode::BuildAddDecoratorContextMenu(UToolMenu* SubMenu
 			)
 		);
 
-		SubMenu->AddMenuEntry(NAME_None, DecoratorEntry);
+		SubMenu->AddMenuEntry(NAME_None, TraitEntry);
 	}
 }
 
diff --git a/Engine/Plugins/Experimental/AnimNext/Source/AnimNextUncookedOnly/Private/UncookedOnlyUtils.cpp b/Engine/Plugins/Experimental/AnimNext/Source/AnimNextUncookedOnly/Private/UncookedOnlyUtils.cpp
--- a/Engine/Plugins/Experimental/AnimNext/Source/AnimNextUncookedOnly/Private/UncookedOnlyUtils.cpp
+++ b/Engine/Plugins/Experimental/AnimNext/Source/AnimNextUncookedOnly/Private/UncookedOnlyUtils.cpp
@@ -10,22 +10,22 @@
 #include "Graph/AnimNextGraph.h"
 #include "Graph/AnimNextGraph_Controller.h"
 #include "Graph/AnimNextGraph_EditorData.h"
-#include "Graph/RigUnit_AnimNextDecoratorStack.h"
+#include "Graph/RigUnit_AnimNextTraitStack.h"
 #include "Graph/RigUnit_AnimNextGraphRoot.h"
 #include "Graph/RigUnit_AnimNextGraphEvaluator.h"
 #include "Graph/RigUnit_AnimNextShimRoot.h"
 #include "Param/AnimNextParameterBlock.h"
 #include "Param/AnimNextParameterBlock_EditorData.h"
-#include "Graph/RigDecorator_AnimNextCppDecorator.h"
+#include "Graph/RigDecorator_AnimNextCppTrait.h"
 #include "Param/RigUnit_AnimNextParameterBeginExecution.h"
 #include "Param/RigVMDispatch_GetParameter.h"
 #include "Param/RigVMDispatch_SetLayerParameter.h"
 #include "IAnimNextRigVMParameterInterface.h"
-#include "DecoratorBase/DecoratorReader.h"
-#include "DecoratorBase/DecoratorWriter.h"
-#include "DecoratorBase/NodeTemplateBuilder.h"
-#include "DecoratorBase/DecoratorRegistry.h"
-#include "DecoratorBase/Decorator.h"
+#include "TraitCore/TraitReader.h"
+#include "TraitCore/TraitWriter.h"
+#include "TraitCore/NodeTemplateBuilder.h"
+#include "TraitCore/TraitRegistry.h"
+#include "TraitCore/Trait.h"
 #include "Graph/RigUnit_AnimNextBeginExecution.h"
 #include "Serialization/MemoryReader.h"
 #include "RigVMRuntimeDataRegistry.h"
@@ -50,8 +50,8 @@ namespace UE::AnimNext::UncookedOnly
 {
 namespace Private
 {
-	// Represents a decorator entry on a node
-	struct FDecoratorEntryMapping
+	// Represents a trait entry on a node
+	struct FTraitEntryMapping
 	{
 		// The RigVM node that hosts this RigVM decorator
 		const URigVMNode* DecoratorStackNode = nullptr;
@@ -59,50 +59,50 @@ namespace Private
 		// The RigVM decorator pin on our host node
 		const URigVMPin* DecoratorEntryPin = nullptr;
 
-		// The AnimNext decorator
-		const FDecorator* Decorator = nullptr;
+		// The AnimNext trait
+		const FTrait* Trait = nullptr;
 
 		// A map from latent property names to their corresponding RigVM memory handle index
 		TMap<FName, uint16> LatentPropertyNameToIndexMap;
 
-		FDecoratorEntryMapping(const URigVMNode* InDecoratorStackNode, const URigVMPin* InDecoratorEntryPin, const FDecorator* InDecorator)
+		FTraitEntryMapping(const URigVMNode* InDecoratorStackNode, const URigVMPin* InDecoratorEntryPin, const FTrait* InTrait)
 			: DecoratorStackNode(InDecoratorStackNode)
 			, DecoratorEntryPin(InDecoratorEntryPin)
-			, Decorator(InDecorator)
+			, Trait(InTrait)
 		{}
 	};
 
-	// Represents a node that contains a decorator list
-	struct FDecoratorStackMapping
+	// Represents a node that contains a trait list
+	struct FTraitStackMapping
 	{
 		// The RigVM node that hosts the RigVM decorators
 		const URigVMNode* DecoratorStackNode = nullptr;
 
-		// The decorator list on this node
-		TArray<FDecoratorEntryMapping> DecoratorEntries;
+		// The trait list on this node
+		TArray<FTraitEntryMapping> TraitEntries;
 
 		// The node handle assigned to this RigVM node
-		FNodeHandle DecoratorStackNodeHandle;
+		FNodeHandle TraitStackNodeHandle;
 
-		explicit FDecoratorStackMapping(const URigVMNode* InDecoratorStackNode)
+		explicit FTraitStackMapping(const URigVMNode* InDecoratorStackNode)
 			: DecoratorStackNode(InDecoratorStackNode)
 		{}
 	};
 
-	struct FDecoratorGraph
+	struct FTraitGraph
 	{
 		FName EntryPoint;
 		URigVMNode* RootNode;
-		TArray<FDecoratorStackMapping> DecoratorStackNodes;
+		TArray<FTraitStackMapping> TraitStackNodes;
 
-		explicit FDecoratorGraph(URigVMNode* InRootNode)
+		explicit FTraitGraph(URigVMNode* InRootNode)
 			: EntryPoint(*InRootNode->FindPin(GET_MEMBER_NAME_STRING_CHECKED(FRigUnit_AnimNextGraphRoot, EntryPoint))->GetDefaultValue())
 			, RootNode(InRootNode)
 		{}
 	};
 
-	template<typename DecoratorAction>
-	void ForEachDecoratorInStack(const URigVMNode* DecoratorStackNode, const DecoratorAction& Action)
+	template<typename TraitAction>
+	void ForEachTraitInStack(const URigVMNode* DecoratorStackNode, const TraitAction& Action)
 	{
 		const TArray<URigVMPin*>& Pins = DecoratorStackNode->GetPins();
 		for (URigVMPin* Pin : Pins)
@@ -117,40 +117,40 @@ namespace Private
 				TSharedPtr<FStructOnScope> DecoratorScope = Pin->GetDecoratorInstance();
 				FRigDecorator_AnimNextCppDecorator* VMDecorator = (FRigDecorator_AnimNextCppDecorator*)DecoratorScope->GetStructMemory();
 
-				if (const FDecorator* Decorator = VMDecorator->GetDecorator())
+				if (const FTrait* Trait = VMDecorator->GetTrait())
 				{
-					Action(DecoratorStackNode, Pin, Decorator);
+					Action(DecoratorStackNode, Pin, Trait);
 				}
 			}
 		}
 	}
 
-	TArray<FDecoratorUID> GetDecoratorUIDs(const URigVMNode* DecoratorStackNode)
+	TArray<FTraitUID> GetTraitUIDs(const URigVMNode* DecoratorStackNode)
 	{
-		TArray<FDecoratorUID> Decorators;
+		TArray<FTraitUID> Traits;
 
-		ForEachDecoratorInStack(DecoratorStackNode,
-			[&Decorators](const URigVMNode* DecoratorStackNode, const URigVMPin* DecoratorPin, const FDecorator* Decorator)
+		ForEachTraitInStack(DecoratorStackNode,
+			[&Traits](const URigVMNode* DecoratorStackNode, const URigVMPin* DecoratorPin, const FTrait* Trait)
 			{
-				Decorators.Add(Decorator->GetDecoratorUID());
+				Traits.Add(Trait->GetTraitUID());
 			});
 
-		return Decorators;
+		return Traits;
 	}
 
-	FNodeHandle RegisterDecoratorNodeTemplate(FDecoratorWriter& DecoratorWriter, const URigVMNode* DecoratorStackNode)
+	FNodeHandle RegisterTraitNodeTemplate(FTraitWriter& TraitWriter, const URigVMNode* DecoratorStackNode)
 	{
-		const TArray<FDecoratorUID> DecoratorUIDs = GetDecoratorUIDs(DecoratorStackNode);
+		const TArray<FTraitUID> TraitUIDs = GetTraitUIDs(DecoratorStackNode);
 
 		TArray<uint8> NodeTemplateBuffer;
-		const FNodeTemplate* NodeTemplate = FNodeTemplateBuilder::BuildNodeTemplate(DecoratorUIDs, NodeTemplateBuffer);
+		const FNodeTemplate* NodeTemplate = FNodeTemplateBuilder::BuildNodeTemplate(TraitUIDs, NodeTemplateBuffer);
 
-		return DecoratorWriter.RegisterNode(*NodeTemplate);
+		return TraitWriter.RegisterNode(*NodeTemplate);
 	}
 
-	FString GetDecoratorProperty(const FDecoratorStackMapping& DecoratorStack, uint32 DecoratorIndex, FName PropertyName, const TArray<FDecoratorStackMapping>& DecoratorStackNodes)
+	FString GetTraitProperty(const FTraitStackMapping& TraitStack, uint32 TraitIndex, FName PropertyName, const TArray<FTraitStackMapping>& TraitStackNodes)
 	{
-		const TArray<URigVMPin*>& Pins = DecoratorStack.DecoratorEntries[DecoratorIndex].DecoratorEntryPin->GetSubPins();
+		const TArray<URigVMPin*>& Pins = TraitStack.TraitEntries[TraitIndex].DecoratorEntryPin->GetSubPins();
 		for (const URigVMPin* Pin : Pins)
 		{
 			if (Pin->GetDirection() != ERigVMPinDirection::Input)
@@ -160,9 +160,9 @@ namespace Private
 
 			if (Pin->GetFName() == PropertyName)
 			{
-				if (Pin->GetCPPTypeObject() == FAnimNextDecoratorHandle::StaticStruct())
+				if (Pin->GetCPPTypeObject() == FAnimNextTraitHandle::StaticStruct())
 				{
-					// Decorator handle pins don't have a value, just an optional link
+					// Trait handle pins don't have a value, just an optional link
 					const TArray<URigVMLink*>& PinLinks = Pin->GetLinks();
 					if (!PinLinks.IsEmpty())
 					{
@@ -172,32 +172,32 @@ namespace Private
 						const URigVMNode* SourceNode = PinLinks[0]->GetSourceNode();
 
 						FNodeHandle SourceNodeHandle;
-						int32 SourceDecoratorIndex = INDEX_NONE;
+						int32 SourceTraitIndex = INDEX_NONE;
 
-						const FDecoratorStackMapping* SourceDecoratorStack = DecoratorStackNodes.FindByPredicate([SourceNode](const FDecoratorStackMapping& Mapping) { return Mapping.DecoratorStackNode == SourceNode; });
-						if (SourceDecoratorStack != nullptr)
+						const FTraitStackMapping* SourceTraitStack = TraitStackNodes.FindByPredicate([SourceNode](const FTraitStackMapping& Mapping) { return Mapping.DecoratorStackNode == SourceNode; });
+						if (SourceTraitStack != nullptr)
 						{
-							SourceNodeHandle = SourceDecoratorStack->DecoratorStackNodeHandle;
+							SourceNodeHandle = SourceTraitStack->TraitStackNodeHandle;
 
 							// If the source pin is null, we are a node where the result pin lives on the stack node instead of a decorator sub-pin
-							// If this is the case, we bind to the first decorator index since we only allowed a single base decorator per stack
-							// Otherwise we lookup the decorator index we are linked to
+							// If this is the case, we bind to the first trait index since we only allowed a single base trait per stack
+							// Otherwise we lookup the trait index we are linked to
 							const URigVMPin* SourceDecoratorPin = PinLinks[0]->GetSourcePin()->GetParentPin();
-							SourceDecoratorIndex = SourceDecoratorPin != nullptr ? SourceDecoratorStack->DecoratorStackNode->GetDecoratorPins().IndexOfByKey(SourceDecoratorPin) : 0;
+							SourceTraitIndex = SourceDecoratorPin != nullptr ? SourceTraitStack->DecoratorStackNode->GetDecoratorPins().IndexOfByKey(SourceDecoratorPin) : 0;
 						}
 
 						if (SourceNodeHandle.IsValid())
 						{
-							check(SourceDecoratorIndex != INDEX_NONE);
+							check(SourceTraitIndex != INDEX_NONE);
 
-							const FAnimNextDecoratorHandle DecoratorHandle(SourceNodeHandle, SourceDecoratorIndex);
-							const FAnimNextDecoratorHandle DefaultDecoratorHandle;
+							const FAnimNextTraitHandle TraitHandle(SourceNodeHandle, SourceTraitIndex);
+							const FAnimNextTraitHandle DefaultTraitHandle;
 
-							// We need an instance of a decorator handle property to be able to serialize it into text, grab it from the root
+							// We need an instance of a trait handle property to be able to serialize it into text, grab it from the root
 							const FProperty* Property = FRigUnit_AnimNextGraphRoot::StaticStruct()->FindPropertyByName(GET_MEMBER_NAME_STRING_CHECKED(FRigUnit_AnimNextGraphRoot, Result));
 
 							FString PropertyValue;
-							Property->ExportText_Direct(PropertyValue, &DecoratorHandle, &DefaultDecoratorHandle, nullptr, PPF_None);
+							Property->ExportText_Direct(PropertyValue, &TraitHandle, &DefaultTraitHandle, nullptr, PPF_SerializedAsImportText);
 
 							return PropertyValue;
 						}
@@ -216,9 +216,9 @@ namespace Private
 		return FString();
 	}
 
-	uint16 GetDecoratorLatentPropertyIndex(const FDecoratorStackMapping& DecoratorStack, uint32 DecoratorIndex, FName PropertyName)
+	uint16 GetTraitLatentPropertyIndex(const FTraitStackMapping& TraitStack, uint32 TraitIndex, FName PropertyName)
 	{
-		const FDecoratorEntryMapping& Entry = DecoratorStack.DecoratorEntries[DecoratorIndex];
+		const FTraitEntryMapping& Entry = TraitStack.TraitEntries[TraitIndex];
 		if (const uint16* RigVMIndex = Entry.LatentPropertyNameToIndexMap.Find(PropertyName))
 		{
 			return *RigVMIndex;
@@ -227,16 +227,16 @@ namespace Private
 		return MAX_uint16;
 	}
 
-	void WriteDecoratorProperties(FDecoratorWriter& DecoratorWriter, const FDecoratorStackMapping& Mapping, const TArray<FDecoratorStackMapping>& DecoratorStackNodes)
+	void WriteTraitProperties(FTraitWriter& TraitWriter, const FTraitStackMapping& Mapping, const TArray<FTraitStackMapping>& TraitStackNodes)
 	{
-		DecoratorWriter.WriteNode(Mapping.DecoratorStackNodeHandle,
-			[&Mapping, &DecoratorStackNodes](uint32 DecoratorIndex, FName PropertyName)
+		TraitWriter.WriteNode(Mapping.TraitStackNodeHandle,
+			[&Mapping, &TraitStackNodes](uint32 TraitIndex, FName PropertyName)
 			{
-				return GetDecoratorProperty(Mapping, DecoratorIndex, PropertyName, DecoratorStackNodes);
+				return GetTraitProperty(Mapping, TraitIndex, PropertyName, TraitStackNodes);
 			},
-			[&Mapping](uint32 DecoratorIndex, FName PropertyName)
+			[&Mapping](uint32 TraitIndex, FName PropertyName)
 			{
-				return GetDecoratorLatentPropertyIndex(Mapping, DecoratorIndex, PropertyName);
+				return GetTraitLatentPropertyIndex(Mapping, TraitIndex, PropertyName);
 			});
 	}
 
@@ -268,9 +268,9 @@ namespace Private
 				continue;	// We only look for hidden or input pins
 			}
 
-			if (Pin->GetCPPTypeObject() != FAnimNextDecoratorHandle::StaticStruct())
+			if (Pin->GetCPPTypeObject() != FAnimNextTraitHandle::StaticStruct())
 			{
-				continue;	// We only look for decorator handle pins
+				continue;	// We only look for trait handle pins
 			}
 
 			const TArray<URigVMLink*>& PinLinks = Pin->GetLinks();
@@ -281,34 +281,34 @@ namespace Private
 
 			// Add a dummy node that will output a reference pose to ensure every link is valid.
 			// RigVM doesn't let us link two decorators on a same node together or linking a child back to a parent
-			// as this would create a cycle in the RigVM graph. The AnimNext graph decorators do support it
+			// as this would create a cycle in the RigVM graph. The AnimNext graph traits do support it
 			// and so perhaps we could have a merging pass later on to remove useless dummy nodes like this.
 
-			URigVMUnitNode* VMReferencePoseNode = VMController->AddUnitNode(FRigUnit_AnimNextDecoratorStack::StaticStruct(), FRigVMStruct::ExecuteName, FVector2D(0.0f, 0.0f), FString(), false);
+			URigVMUnitNode* VMReferencePoseNode = VMController->AddUnitNode(FRigUnit_AnimNextTraitStack::StaticStruct(), FRigVMStruct::ExecuteName, FVector2D(0.0f, 0.0f), FString(), false);
 			check(VMReferencePoseNode != nullptr);
 
 			const UScriptStruct* CppDecoratorStruct = FRigDecorator_AnimNextCppDecorator::StaticStruct();
 
 			FString DefaultValue;
 			{
-				const UE::AnimNext::FDecoratorUID ReferencePoseDecoratorUID(0xc03d6afc);	// Decorator header is private, reference by UID directly
-				const FDecorator* Decorator = FDecoratorRegistry::Get().Find(ReferencePoseDecoratorUID);
-				check(Decorator != nullptr);
+				const UE::AnimNext::FTraitUID ReferencePoseTraitUID(0x7508ab89);	// Trait header is private, reference by UID directly
+				const FTrait* Trait = FTraitRegistry::Get().Find(ReferencePoseTraitUID);
+				check(Trait != nullptr);
 
 				const FRigDecorator_AnimNextCppDecorator DefaultCppDecoratorStructInstance;
 				FRigDecorator_AnimNextCppDecorator CppDecoratorStructInstance;
-				CppDecoratorStructInstance.DecoratorSharedDataStruct = Decorator->GetDecoratorSharedDataStruct();
+				CppDecoratorStructInstance.DecoratorSharedDataStruct = Trait->GetTraitSharedDataStruct();
 
 				const FProperty* Prop = FAnimNextCppDecoratorWrapper::StaticStruct()->FindPropertyByName(GET_MEMBER_NAME_STRING_CHECKED(FAnimNextCppDecoratorWrapper, CppDecorator));
 				check(Prop != nullptr);
 
-				Prop->ExportText_Direct(DefaultValue, &CppDecoratorStructInstance, &DefaultCppDecoratorStructInstance, nullptr, PPF_None);
+				Prop->ExportText_Direct(DefaultValue, &CppDecoratorStructInstance, &DefaultCppDecoratorStructInstance, nullptr, PPF_SerializedAsImportText);
 			}
 
 			const FName ReferencePoseDecoratorName = VMController->AddDecorator(VMReferencePoseNode->GetFName(), *CppDecoratorStruct->GetPathName(), TEXT("ReferencePose"), DefaultValue, INDEX_NONE, false, false);
 			check(!ReferencePoseDecoratorName.IsNone());
 
-			URigVMPin* OutputPin = VMReferencePoseNode->FindPin(GET_MEMBER_NAME_STRING_CHECKED(FRigUnit_AnimNextDecoratorStack, Result));
+			URigVMPin* OutputPin = VMReferencePoseNode->FindPin(GET_MEMBER_NAME_STRING_CHECKED(FRigUnit_AnimNextTraitStack, Result));
 			check(OutputPin != nullptr);
 
 			ensure(VMController->AddLink(OutputPin, Pin, false));
@@ -323,13 +323,13 @@ namespace Private
 			if (const URigVMUnitNode* VMUnitNode = Cast<URigVMUnitNode>(VMNode))
 			{
 				const UScriptStruct* ScriptStruct = VMUnitNode->GetScriptStruct();
-				if (ScriptStruct != FRigUnit_AnimNextDecoratorStack::StaticStruct())
+				if (ScriptStruct != FRigUnit_AnimNextTraitStack::StaticStruct())
 				{
-					continue;	// Skip non-decorator nodes
+					continue;	// Skip non-trait nodes
 				}
 
-				ForEachDecoratorInStack(VMNode,
-					[VMController](const URigVMNode* DecoratorStackNode, const URigVMPin* DecoratorPin, const FDecorator* Decorator)
+				ForEachTraitInStack(VMNode,
+					[VMController](const URigVMNode* DecoratorStackNode, const URigVMPin* DecoratorPin, const FTrait* Trait)
 					{
 						AddMissingInputLinks(DecoratorPin, VMController);
 					});
@@ -337,7 +337,7 @@ namespace Private
 		}
 	}
 
-	FDecoratorGraph CollectGraphInfo(const URigVMGraph* VMGraph, URigVMController* VMController)
+	FTraitGraph CollectGraphInfo(const URigVMGraph* VMGraph, URigVMController* VMController)
 	{
 		const TArray<URigVMNode*>& VMNodes = VMGraph->GetNodes();
 		URigVMUnitNode* VMRootNode = FindRootNode(VMNodes);
@@ -351,7 +351,7 @@ namespace Private
 		// Make sure we don't have empty input pins
 		AddMissingInputLinks(VMGraph, VMController);
 
-		FDecoratorGraph DecoratorGraph(VMRootNode);
+		FTraitGraph TraitGraph(VMRootNode);
 
 		TArray<const URigVMNode*> NodesToVisit;
 		NodesToVisit.Add(VMRootNode);
@@ -364,16 +364,16 @@ namespace Private
 			if (const URigVMUnitNode* VMUnitNode = Cast<URigVMUnitNode>(VMNode))
 			{
 				const UScriptStruct* ScriptStruct = VMUnitNode->GetScriptStruct();
-				if (ScriptStruct == FRigUnit_AnimNextDecoratorStack::StaticStruct())
+				if (ScriptStruct == FRigUnit_AnimNextTraitStack::StaticStruct())
 				{
-					FDecoratorStackMapping Mapping(VMNode);
-					ForEachDecoratorInStack(VMNode,
-						[&Mapping](const URigVMNode* DecoratorStackNode, const URigVMPin* DecoratorPin, const FDecorator* Decorator)
+					FTraitStackMapping Mapping(VMNode);
+					ForEachTraitInStack(VMNode,
+						[&Mapping](const URigVMNode* DecoratorStackNode, const URigVMPin* DecoratorPin, const FTrait* Trait)
 						{
-							Mapping.DecoratorEntries.Add(FDecoratorEntryMapping(DecoratorStackNode, DecoratorPin, Decorator));
+							Mapping.TraitEntries.Add(FTraitEntryMapping(DecoratorStackNode, DecoratorPin, Trait));
 						});
 
-					DecoratorGraph.DecoratorStackNodes.Add(MoveTemp(Mapping));
+					TraitGraph.TraitStackNodes.Add(MoveTemp(Mapping));
 				}
 			}
 
@@ -381,57 +381,57 @@ namespace Private
 			NodesToVisit.Append(SourceNodes);
 		}
 
-		if (DecoratorGraph.DecoratorStackNodes.IsEmpty())
+		if (TraitGraph.TraitStackNodes.IsEmpty())
 		{
 			// If the graph is empty, add a dummy node that just pushes a reference pose
-			URigVMUnitNode* VMNode = VMController->AddUnitNode(FRigUnit_AnimNextDecoratorStack::StaticStruct(), FRigVMStruct::ExecuteName, FVector2D(0.0f, 0.0f), FString(), false);
+			URigVMUnitNode* VMNode = VMController->AddUnitNode(FRigUnit_AnimNextTraitStack::StaticStruct(), FRigVMStruct::ExecuteName, FVector2D(0.0f, 0.0f), FString(), false);
 
 			const UScriptStruct* CppDecoratorStruct = FRigDecorator_AnimNextCppDecorator::StaticStruct();
 
 			FString DefaultValue;
 			{
-				const UE::AnimNext::FDecoratorUID ReferencePoseDecoratorUID(0xc03d6afc);	// Decorator header is private, reference by UID directly
-				const FDecorator* Decorator = FDecoratorRegistry::Get().Find(ReferencePoseDecoratorUID);
-				check(Decorator != nullptr);
+				const UE::AnimNext::FTraitUID ReferencePoseTraitUID(0x7508ab89);	// Trait header is private, reference by UID directly
+				const FTrait* Trait = FTraitRegistry::Get().Find(ReferencePoseTraitUID);
+				check(Trait != nullptr);
 
 				const FRigDecorator_AnimNextCppDecorator DefaultCppDecoratorStructInstance;
 				FRigDecorator_AnimNextCppDecorator CppDecoratorStructInstance;
-				CppDecoratorStructInstance.DecoratorSharedDataStruct = Decorator->GetDecoratorSharedDataStruct();
+				CppDecoratorStructInstance.DecoratorSharedDataStruct = Trait->GetTraitSharedDataStruct();
 
 				const FProperty* Prop = FAnimNextCppDecoratorWrapper::StaticStruct()->FindPropertyByName(GET_MEMBER_NAME_STRING_CHECKED(FAnimNextCppDecoratorWrapper, CppDecorator));
 				check(Prop != nullptr);
 
-				Prop->ExportText_Direct(DefaultValue, &CppDecoratorStructInstance, &DefaultCppDecoratorStructInstance, nullptr, PPF_None);
+				Prop->ExportText_Direct(DefaultValue, &CppDecoratorStructInstance, &DefaultCppDecoratorStructInstance, nullptr, PPF_SerializedAsImportText);
 			}
 
 			VMController->AddDecorator(VMNode->GetFName(), *CppDecoratorStruct->GetPathName(), TEXT("ReferencePose"), DefaultValue, INDEX_NONE, false, false);
 
-			FDecoratorStackMapping Mapping(VMNode);
-			ForEachDecoratorInStack(VMNode,
-				[&Mapping](const URigVMNode* DecoratorStackNode, const URigVMPin* DecoratorPin, const FDecorator* Decorator)
+			FTraitStackMapping Mapping(VMNode);
+			ForEachTraitInStack(VMNode,
+				[&Mapping](const URigVMNode* DecoratorStackNode, const URigVMPin* DecoratorPin, const FTrait* Trait)
 				{
-					Mapping.DecoratorEntries.Add(FDecoratorEntryMapping(DecoratorStackNode, DecoratorPin, Decorator));
+					Mapping.TraitEntries.Add(FTraitEntryMapping(DecoratorStackNode, DecoratorPin, Trait));
 				});
 
-			DecoratorGraph.DecoratorStackNodes.Add(MoveTemp(Mapping));
+			TraitGraph.TraitStackNodes.Add(MoveTemp(Mapping));
 		}
 
-		return DecoratorGraph;
+		return TraitGraph;
 	}
 
-	void CollectLatentPins(TArray<FDecoratorStackMapping>& DecoratorStackNodes, FRigVMPinInfoArray& OutLatentPins, TMap<FName, URigVMPin*>& OutLatentPinMapping)
+	void CollectLatentPins(TArray<FTraitStackMapping>& TraitStackNodes, FRigVMPinInfoArray& OutLatentPins, TMap<FName, URigVMPin*>& OutLatentPinMapping)
 	{
-		for (FDecoratorStackMapping& DecoratorStack : DecoratorStackNodes)
+		for (FTraitStackMapping& TraitStack : TraitStackNodes)
 		{
-			for (FDecoratorEntryMapping& DecoratorEntry : DecoratorStack.DecoratorEntries)
+			for (FTraitEntryMapping& TraitEntry : TraitStack.TraitEntries)
 			{
-				for (URigVMPin* Pin : DecoratorEntry.DecoratorEntryPin->GetSubPins())
+				for (URigVMPin* Pin : TraitEntry.DecoratorEntryPin->GetSubPins())
 				{
 					if (Pin->IsLazy() && !Pin->GetLinks().IsEmpty())
 					{
 						// This pin has something linked to it, it is a latent pin
 						check(OutLatentPins.Num() < ((1 << 16) - 1));	// We reserve MAX_uint16 as an invalid value and we must fit on 15 bits when packed
-						DecoratorEntry.LatentPropertyNameToIndexMap.Add(Pin->GetFName(), (uint16)OutLatentPins.Num());
+						TraitEntry.LatentPropertyNameToIndexMap.Add(Pin->GetFName(), (uint16)OutLatentPins.Num());
 
 						const FName LatentPinName(TEXT("LatentPin"), OutLatentPins.Num());	// Create unique latent pin names
 
@@ -504,7 +504,7 @@ void FUtils::Compile(UAnimNextGraph* InGraph)
 
 	
 	// Before we re-compile a graph, we need to release and live instances since we need the metadata we are about to replace
-	// to call decorator destructors etc
+	// to call trait destructors etc
 	InGraph->FreezeGraphInstances();
 
 	EditorData->bErrorsDuringCompilation = false;
@@ -519,7 +519,7 @@ void FUtils::Compile(UAnimNextGraph* InGraph)
 
 	InGraph->VMRuntimeSettings = EditorData->VMRuntimeSettings;
 	InGraph->EntryPoints.Empty();
-	InGraph->ResolvedRootDecoratorHandles.Empty();
+	InGraph->ResolvedRootTraitHandles.Empty();
 	InGraph->ExecuteDefinition = FAnimNextGraphEvaluatorExecuteDefinition();
 	InGraph->SharedDataBuffer.Empty();
 	InGraph->GraphReferencedObjects.Empty();
@@ -552,35 +552,35 @@ void FUtils::Compile(UAnimNextGraph* InGraph)
 
 	UAnimNextGraph_Controller* TempController = CastChecked<UAnimNextGraph_Controller>(VMClient->GetOrCreateController(VMTempGraphs[0]));
 
-	FDecoratorWriter DecoratorWriter;
+	FTraitWriter TraitWriter;
 
 	FRigVMPinInfoArray LatentPins;
 	TMap<FName, URigVMPin*> LatentPinMapping;
-	TArray<Private::FDecoratorGraph> DecoratorGraphs;
+	TArray<Private::FTraitGraph> TraitGraphs;
 
 	// Build entry points and extract their required latent pins
 	for(const URigVMGraph* VMTempGraph : VMTempGraphs)
 	{
-		// Gather our decorator stacks
-		Private::FDecoratorGraph& DecoratorGraph = DecoratorGraphs.Add_GetRef(Private::CollectGraphInfo(VMTempGraph, TempController->GetControllerForGraph(VMTempGraph)));
-		check(!DecoratorGraph.DecoratorStackNodes.IsEmpty());
+		// Gather our trait stacks
+		Private::FTraitGraph& TraitGraph = TraitGraphs.Add_GetRef(Private::CollectGraphInfo(VMTempGraph, TempController->GetControllerForGraph(VMTempGraph)));
+		check(!TraitGraph.TraitStackNodes.IsEmpty());
 
 		FAnimNextGraphEntryPoint& EntryPoint = InGraph->EntryPoints.AddDefaulted_GetRef();
-		EntryPoint.EntryPointName = DecoratorGraph.EntryPoint;
+		EntryPoint.EntryPointName = TraitGraph.EntryPoint;
 
 		// Extract latent pins for this graph
-		Private::CollectLatentPins(DecoratorGraph.DecoratorStackNodes, LatentPins, LatentPinMapping);
+		Private::CollectLatentPins(TraitGraph.TraitStackNodes, LatentPins, LatentPinMapping);
 
-		// Iterate over every decorator stack and register our node templates
-		for (Private::FDecoratorStackMapping& NodeMapping : DecoratorGraph.DecoratorStackNodes)
+		// Iterate over every trait stack and register our node templates
+		for (Private::FTraitStackMapping& NodeMapping : TraitGraph.TraitStackNodes)
 		{
-			NodeMapping.DecoratorStackNodeHandle = Private::RegisterDecoratorNodeTemplate(DecoratorWriter, NodeMapping.DecoratorStackNode);
+			NodeMapping.TraitStackNodeHandle = Private::RegisterTraitNodeTemplate(TraitWriter, NodeMapping.DecoratorStackNode);
 		}
 
 		// Find our root node handle, if we have any stack nodes, the first one is our root stack
-		if (DecoratorGraph.DecoratorStackNodes.Num() != 0)
+		if (TraitGraph.TraitStackNodes.Num() != 0)
 		{
-			EntryPoint.RootDecoratorHandle = FAnimNextEntryPointHandle(DecoratorGraph.DecoratorStackNodes[0].DecoratorStackNodeHandle);
+			EntryPoint.RootTraitHandle = FAnimNextEntryPointHandle(TraitGraph.TraitStackNodes[0].TraitStackNodeHandle);
 		}
 	}
 
@@ -607,30 +607,30 @@ void FUtils::Compile(UAnimNextGraph* InGraph)
 	}
 	
 	// Write our node shared data
-	DecoratorWriter.BeginNodeWriting();
+	TraitWriter.BeginNodeWriting();
 
-	for(Private::FDecoratorGraph& DecoratorGraph : DecoratorGraphs)
+	for(Private::FTraitGraph& TraitGraph : TraitGraphs)
 	{
-		for (const Private::FDecoratorStackMapping& NodeMapping : DecoratorGraph.DecoratorStackNodes)
+		for (const Private::FTraitStackMapping& NodeMapping : TraitGraph.TraitStackNodes)
 		{
-			Private::WriteDecoratorProperties(DecoratorWriter, NodeMapping, DecoratorGraph.DecoratorStackNodes);
+			Private::WriteTraitProperties(TraitWriter, NodeMapping, TraitGraph.TraitStackNodes);
 		}
 	}
 
-	DecoratorWriter.EndNodeWriting();
+	TraitWriter.EndNodeWriting();
 
 	// Cache our compiled metadata
-	InGraph->SharedDataArchiveBuffer = DecoratorWriter.GetGraphSharedData();
-	InGraph->GraphReferencedObjects = DecoratorWriter.GetGraphReferencedObjects();
+	InGraph->SharedDataArchiveBuffer = TraitWriter.GetGraphSharedData();
+	InGraph->GraphReferencedObjects = TraitWriter.GetGraphReferencedObjects();
 
 	// Populate our runtime metadata
 	InGraph->LoadFromArchiveBuffer(InGraph->SharedDataArchiveBuffer);
 
 	// Remove our old root nodes
-	for(Private::FDecoratorGraph& DecoratorGraph : DecoratorGraphs)
+	for(Private::FTraitGraph& TraitGraph : TraitGraphs)
 	{
-		URigVMController* GraphController = TempController->GetControllerForGraph(DecoratorGraph.RootNode->GetGraph());
-		GraphController->RemoveNode(DecoratorGraph.RootNode, false, false);
+		URigVMController* GraphController = TempController->GetControllerForGraph(TraitGraph.RootNode->GetGraph());
+		GraphController->RemoveNode(TraitGraph.RootNode, false, false);
 	}
 
 	URigVMCompiler* Compiler = URigVMCompiler::StaticClass()->GetDefaultObject<URigVMCompiler>();
