CIS: Fix missing include

diff --git a/Engine/Plugins/Experimental/AnimNext/Source/AnimNextEditor/Private/Param/SGraphPinParamName.cpp b/Engine/Plugins/Experimental/AnimNext/Source/AnimNextEditor/Private/Param/SGraphPinParamName.cpp
--- a/Engine/Plugins/Experimental/AnimNext/Source/AnimNextEditor/Private/Param/SGraphPinParamName.cpp
+++ b/Engine/Plugins/Experimental/AnimNext/Source/AnimNextEditor/Private/Param/SGraphPinParamName.cpp
@@ -10,6 +10,9 @@
 #include "Param/RigVMDispatch_GetLayerParameter.h"
 #include "Param/RigVMDispatch_GetParameter.h"
 #include "Param/RigVMDispatch_SetLayerParameter.h"
+#include "RigVMModel/Nodes/RigVMTemplateNode.h"
+#include "RigVMModel/RigVMPin.h"
+#include "RigVMCore/RigVMRegistry.h"
 
 #define LOCTEXT_NAMESPACE "SGraphPinParamName"
 
