AnimNext: Merging parameter blocks and graphs - part 3

Removed parameter block asset. Merged functionality into AnimNext graphs. Added redirectors to redirect existing assets onto the new type.
Added a custom version to deal with serialization differences between the two assets (skipping serializing the SharedDataArchiveBuffer which didnt exist in older param block assets).
Settled on naming for different graph types: 'animation graphs' and 'event graphs' (for what were parameter graphs).
Combined compilation logic for the two assets, split graph compilation based on schema.
Renamed all mentions of parameter blocks appropriately.
Also re-colored the content browser icons for all assets to be more consistent

RigVM changes:
Removed transient flags from URigVMGraph::SchemaClass - if a graph overrides the schema it needs to persist as nothing else can reset it.

diff --git a/Engine/Plugins/Experimental/AnimNext/Config/BaseAnimNext.ini b/Engine/Plugins/Experimental/AnimNext/Config/BaseAnimNext.ini
--- a/Engine/Plugins/Experimental/AnimNext/Config/BaseAnimNext.ini
+++ b/Engine/Plugins/Experimental/AnimNext/Config/BaseAnimNext.ini
@@ -4,6 +4,7 @@
 +ExposedClasses=(Class="/Script/AnimNext.AnimNextMeshComponent",AccessorName="UE_AnimNextMeshComponent")
 
 [CoreRedirects]
+; Traits rename
 +StructRedirects=(OldName="/Script/AnimNext.AnimNextDecoratorHandle", NewName="/Script/AnimNext.AnimNextTraitHandle")
 +StructRedirects=(OldName="/Script/AnimNext.RigUnit_AnimNextDecoratorStack", NewName="/Script/AnimNext.RigUnit_AnimNextTraitStack")
 +StructRedirects=(OldName="/Script/AnimNext.AnimNextBlendByBoolDecoratorSharedData", NewName="/Script/AnimNext.AnimNextBlendByBoolTraitSharedData")
@@ -15,3 +16,17 @@
 +StructRedirects=(OldName="/Script/AnimNext.AnimNextSequencePlayerDecoratorSharedData", NewName="/Script/AnimNext.AnimNextSequencePlayerTraitSharedData")
 +StructRedirects=(OldName="/Script/AnimNext.AnimNextSubGraphHostDecoratorSharedData", NewName="/Script/AnimNext.AnimNextSubGraphHostTraitSharedData")
 +StructRedirects=(OldName="/Script/AnimNext.AnimNextSynchronizeUsingGroupsDecoratorSharedData", NewName="/Script/AnimNext.AnimNextSynchronizeUsingGroupsTraitSharedData")
+
+; Parameter block/graph merge
++ClassRedirects=(OldName="/Script/AnimNext.AnimNextParameterBlock", NewName="/Script/AnimNext.AnimNextGraph")
++ClassRedirects=(OldName="/Script/AnimNextUncookedOnly.AnimNextParameterBlock_EditorData", NewName="/Script/AnimNextUncookedOnly.AnimNextGraph_EditorData")
++ClassRedirects=(OldName="/Script/AnimNextUncookedOnly.AnimNextParameterBlock_ParamSchema", NewName="/Script/AnimNextUncookedOnly.AnimNextGraph_EventGraphSchema")
++ClassRedirects=(OldName="/Script/AnimNextUncookedOnly.AnimNextParameterBlock_EdGraph", NewName="/Script/AnimNextUncookedOnly.AnimNextGraph_EdGraph")
++ClassRedirects=(OldName="/Script/AnimNextUncookedOnly.AnimNextParameterBlock_EdGraphNode", NewName="/Script/AnimNextUncookedOnly.AnimNextGraph_EdGraphNode")
++ClassRedirects=(OldName="/Script/AnimNextUncookedOnly.AnimNextParameterBlock_EdGraphSchema", NewName="/Script/AnimNextUncookedOnly.AnimNextGraph_EdGraphSchema")
++ClassRedirects=(OldName="/Script/AnimNextUncookedOnly.AnimNextParameterBlockParameter", NewName="/Script/AnimNextUncookedOnly.AnimNextGraph_Parameter")
++ClassRedirects=(OldName="/Script/AnimNextUncookedOnly.AnimNextParameterBlockGraph", NewName="/Script/AnimNextUncookedOnly.AnimNextGraph_EventGraph")
++ClassRedirects=(OldName="/Script/AnimNextUncookedOnly.AnimNextGraphEntry", NewName="/Script/AnimNextUncookedOnly.AnimNextGraph_AnimationGraph")
++ClassRedirects=(OldName="/Script/AnimNextUncookedOnly.AnimNextGraph_TraitSchema", NewName="/Script/AnimNextUncookedOnly.AnimNextGraph_AnimationGraphSchema")
++PropertyRedirects=(OldName="/Script/AnimNext.AnimNextScheduleEntry_ParamScope.ParameterBlocks", NewName="Parameters")
++PropertyRedirects=(OldName="/Script/AnimNext.AnimNextScheduleParamScopeEntryTask.ParameterBlocks", NewName="Parameters")
\ No newline at end of file
diff --git a/Engine/Plugins/Experimental/AnimNext/Source/AnimNext/Internal/AnimNextExecuteContext.h b/Engine/Plugins/Experimental/AnimNext/Source/AnimNext/Internal/AnimNextExecuteContext.h
--- a/Engine/Plugins/Experimental/AnimNext/Source/AnimNext/Internal/AnimNextExecuteContext.h
+++ b/Engine/Plugins/Experimental/AnimNext/Source/AnimNext/Internal/AnimNextExecuteContext.h
@@ -57,6 +57,6 @@ protected:
 	TVariant<FAnimNextGraphContextData, FAnimNextParamContextData> ContextData;
 
 	friend struct FAnimNextGraphInstance;
-	friend class UAnimNextParameterBlock;
+	friend class UAnimNextGraph;
 };
 
diff --git a/Engine/Plugins/Experimental/AnimNext/Source/AnimNext/Internal/Graph/AnimGraph/AnimNode_AnimNextParameters.h b/Engine/Plugins/Experimental/AnimNext/Source/AnimNext/Internal/Graph/AnimGraph/AnimNode_AnimNextParameters.h
--- a/Engine/Plugins/Experimental/AnimNext/Source/AnimNext/Internal/Graph/AnimGraph/AnimNode_AnimNextParameters.h
+++ b/Engine/Plugins/Experimental/AnimNext/Source/AnimNext/Internal/Graph/AnimGraph/AnimNode_AnimNextParameters.h
@@ -6,13 +6,13 @@
 #include "Animation/AnimNodeBase.h"
 #include "AnimNode_AnimNextParameters.generated.h"
 
+class UAnimNextGraph;
+
 namespace UE::AnimNext
 {
-	struct FParameterBlockProxy;
+	struct FParametersProxy;
 }
 
-class UAnimNextParameterBlock;
-
 USTRUCT(BlueprintInternalUseOnly)
 struct FAnimNode_AnimNextParameters : public FAnimNode_Base
 {
@@ -32,13 +32,13 @@ private:
 	FPoseLink Source;
 
 	UPROPERTY(EditAnywhere, Category = Settings)
-	TObjectPtr<UAnimNextParameterBlock> Parameters;
+	TObjectPtr<UAnimNextGraph> Parameters;
 
-	// Cache previous param block so we know when it changes via pin
-	TObjectPtr<UAnimNextParameterBlock> PreviousParameters;
+	// Cache previous parameters so we know when it changes via pin
+	TObjectPtr<UAnimNextGraph> PreviousParameters;
 
 	// Cached proxy
-	TUniquePtr<UE::AnimNext::FParameterBlockProxy> ParametersProxy;
+	TUniquePtr<UE::AnimNext::FParametersProxy> ParametersProxy;
 
 private:
 	// FAnimNode_Base
diff --git a/Engine/Plugins/Experimental/AnimNext/Source/AnimNext/Internal/Graph/AnimNextGraph.h b/Engine/Plugins/Experimental/AnimNext/Source/AnimNext/Internal/Graph/AnimNextGraph.h
--- a/Engine/Plugins/Experimental/AnimNext/Source/AnimNext/Internal/Graph/AnimNextGraph.h
+++ b/Engine/Plugins/Experimental/AnimNext/Source/AnimNext/Internal/Graph/AnimNextGraph.h
@@ -35,6 +35,7 @@ namespace UE::AnimNext
 	struct FExecutionContext;
 	class FModule;
 	struct FTestUtils;
+	struct FParametersProxy;
 }
 
 namespace UE::AnimNext::UncookedOnly
@@ -45,6 +46,7 @@ namespace UE::AnimNext::UncookedOnly
 namespace UE::AnimNext::Editor
 {
 	class FGraphEditor;
+	class FParameterCustomization;
 }
 
 namespace UE::AnimNext::Graph
@@ -89,6 +91,9 @@ public:
 	// Get the parameter to use to access the current LOD
 	UE::AnimNext::FParamId GetCurrentLODParam() const { return CurrentLODId; }
 
+	// Update the parameter layer, if any
+	void UpdateLayer(UE::AnimNext::FParamStackLayerHandle& InHandle, float InDeltaTime) const;
+
 protected:
 
 	// Loads the graph data from the provided archive buffer and returns true on success, false otherwise
@@ -108,6 +113,7 @@ protected:
 
 	friend class UAnimNextGraphFactory;
 	friend class UAnimNextGraph_EditorData;
+	friend class UAnimNextGraph_Parameter;
 	friend struct UE::AnimNext::UncookedOnly::FUtils;
 	friend class UE::AnimNext::Editor::FGraphEditor;
 	friend struct UE::AnimNext::FTestUtils;
@@ -118,7 +124,9 @@ protected:
 	friend class UAnimNextSchedule;
 	friend struct FAnimNextScheduleGraphTask;
 	friend UE::AnimNext::FModule;
-	
+	friend class UE::AnimNext::Editor::FParameterCustomization;
+	friend struct UE::AnimNext::FParametersProxy;
+
 #if WITH_EDITORONLY_DATA
 	mutable FCriticalSection GraphInstancesLock;
 
@@ -171,6 +179,10 @@ protected:
 	UPROPERTY()
 	TArray<FAnimNextParam> RequiredParameters;
 
+	// Property bag that holds state for this graph 
+	UPROPERTY()
+	FInstancedPropertyBag PropertyBag;
+
 #if WITH_EDITORONLY_DATA
 	// This buffer holds the output of the FTraitWriter post compilation
 	// We serialize it manually and it is discarded at runtime
diff --git a/Engine/Plugins/Experimental/AnimNext/Source/AnimNext/Internal/Param/AnimNextParamContextData.h b/Engine/Plugins/Experimental/AnimNext/Source/AnimNext/Internal/Param/AnimNextParamContextData.h
--- a/Engine/Plugins/Experimental/AnimNext/Source/AnimNext/Internal/Param/AnimNextParamContextData.h
+++ b/Engine/Plugins/Experimental/AnimNext/Source/AnimNext/Internal/Param/AnimNextParamContextData.h
@@ -36,6 +36,6 @@ private:
 	// Parameter scope context
 	UE::AnimNext::FParamStackLayerHandle* LayerHandle = nullptr;
 
-	friend class UAnimNextParameterBlock;
+	friend class UAnimNextGraph;
 	friend struct FAnimNextExecuteContext;
 };
diff --git a/Engine/Plugins/Experimental/AnimNext/Source/AnimNext/Internal/Param/AnimNextParameterBlock.h b/Engine/Plugins/Experimental/AnimNext/Source/AnimNext/Internal/Param/AnimNextParameterBlock.h
deleted file mode 100644
diff --git a/Engine/Plugins/Experimental/AnimNext/Source/AnimNext/Internal/Param/ExternalParameterRegistry.h b/Engine/Plugins/Experimental/AnimNext/Source/AnimNext/Internal/Param/ExternalParameterRegistry.h
--- a/Engine/Plugins/Experimental/AnimNext/Source/AnimNext/Internal/Param/ExternalParameterRegistry.h
+++ b/Engine/Plugins/Experimental/AnimNext/Source/AnimNext/Internal/Param/ExternalParameterRegistry.h
@@ -7,7 +7,7 @@
 #include "Param/IParameterSourceFactory.h"
 
 class UAnimNextSchedule;
-class UAnimNextParameterBlockParameter;
+class UAnimNextGraph_Parameter;
 class AActor;
 class UActorComponent;
 
@@ -49,7 +49,7 @@ struct FExternalParameterRegistry
 
 private:
 	friend class ::UAnimNextSchedule;
-	friend class ::UAnimNextParameterBlockParameter;
+	friend class ::UAnimNextGraph_Parameter;
 	friend class UE::AnimNext::Editor::SParameterPicker;
 	friend struct UE::AnimNext::UncookedOnly::FUtils;
 
diff --git a/Engine/Plugins/Experimental/AnimNext/Source/AnimNext/Internal/Param/RigVMDispatch_GetLayerParameter.h b/Engine/Plugins/Experimental/AnimNext/Source/AnimNext/Internal/Param/RigVMDispatch_GetLayerParameter.h
--- a/Engine/Plugins/Experimental/AnimNext/Source/AnimNext/Internal/Param/RigVMDispatch_GetLayerParameter.h
+++ b/Engine/Plugins/Experimental/AnimNext/Source/AnimNext/Internal/Param/RigVMDispatch_GetLayerParameter.h
@@ -12,9 +12,9 @@ namespace UE::AnimNext::UncookedOnly
 }
 
 /*
- * Gets a parameter's current value from the parameter block value
+ * Gets a parameter's current value from the graph
  */
-USTRUCT(meta = (DisplayName = "Get Parameter from Block", Category="Parameters", NodeColor = "0.8, 0, 0.2, 1"))
+USTRUCT(meta = (DisplayName = "Get Parameter from Graph", Category="Parameters", NodeColor = "0.8, 0, 0.2, 1"))
 struct ANIMNEXT_API FRigVMDispatch_GetLayerParameter : public FRigVMDispatchFactory
 {
 	GENERATED_BODY()
diff --git a/Engine/Plugins/Experimental/AnimNext/Source/AnimNext/Internal/Param/RigVMDispatch_SetLayerParameter.h b/Engine/Plugins/Experimental/AnimNext/Source/AnimNext/Internal/Param/RigVMDispatch_SetLayerParameter.h
--- a/Engine/Plugins/Experimental/AnimNext/Source/AnimNext/Internal/Param/RigVMDispatch_SetLayerParameter.h
+++ b/Engine/Plugins/Experimental/AnimNext/Source/AnimNext/Internal/Param/RigVMDispatch_SetLayerParameter.h
@@ -14,7 +14,7 @@ struct FUtils;
 /*
  * Sets a parameter's value
  */
-USTRUCT(meta=(DisplayName = "Set Block Parameter", Category="Parameters", NodeColor = "0.8, 0, 0.2, 1"))
+USTRUCT(meta=(DisplayName = "Set Graph Parameter", Category="Parameters", NodeColor = "0.8, 0, 0.2, 1"))
 struct ANIMNEXT_API FRigVMDispatch_SetLayerParameter : public FRigVMDispatchFactory
 {
 	GENERATED_BODY()
diff --git a/Engine/Plugins/Experimental/AnimNext/Source/AnimNext/Internal/Scheduler/AnimNextSchedule.h b/Engine/Plugins/Experimental/AnimNext/Source/AnimNext/Internal/Scheduler/AnimNextSchedule.h
--- a/Engine/Plugins/Experimental/AnimNext/Source/AnimNext/Internal/Scheduler/AnimNextSchedule.h
+++ b/Engine/Plugins/Experimental/AnimNext/Source/AnimNext/Internal/Scheduler/AnimNextSchedule.h
@@ -13,7 +13,6 @@
 #include "AnimNextSchedule.generated.h"
 
 class UAnimNextGraph;
-class UAnimNextParameterBlock;
 class UAnimNextSchedule;
 class UAnimNextSchedulerWorldSubsystem;
 class UAnimNextComponent;
@@ -144,7 +143,7 @@ private:
 
 	// Parameters to apply in this scope
 	UPROPERTY(EditAnywhere, Category = "Parameters")
-	TArray<TObjectPtr<UAnimNextParameterBlock>> ParameterBlocks;
+	TArray<TObjectPtr<UAnimNextGraph>> Parameters;
 
 	// Entries that are part of this scope
 	UPROPERTY(EditAnywhere, Category = "Parameters", Instanced)
diff --git a/Engine/Plugins/Experimental/AnimNext/Source/AnimNext/Internal/Scheduler/AnimNextScheduleGraphTask.h b/Engine/Plugins/Experimental/AnimNext/Source/AnimNext/Internal/Scheduler/AnimNextScheduleGraphTask.h
--- a/Engine/Plugins/Experimental/AnimNext/Source/AnimNext/Internal/Scheduler/AnimNextScheduleGraphTask.h
+++ b/Engine/Plugins/Experimental/AnimNext/Source/AnimNext/Internal/Scheduler/AnimNextScheduleGraphTask.h
@@ -7,7 +7,6 @@
 #include "AnimNextScheduleGraphTask.generated.h"
 
 class UAnimNextGraph;
-class UAnimNextParameterBlock;
 struct FAnimNextParam;
 
 namespace UE::AnimNext
diff --git a/Engine/Plugins/Experimental/AnimNext/Source/AnimNext/Internal/Scheduler/AnimNextScheduleParamScopeTask.h b/Engine/Plugins/Experimental/AnimNext/Source/AnimNext/Internal/Scheduler/AnimNextScheduleParamScopeTask.h
--- a/Engine/Plugins/Experimental/AnimNext/Source/AnimNext/Internal/Scheduler/AnimNextScheduleParamScopeTask.h
+++ b/Engine/Plugins/Experimental/AnimNext/Source/AnimNext/Internal/Scheduler/AnimNextScheduleParamScopeTask.h
@@ -6,7 +6,7 @@
 #include "AnimNextScheduleParamScopeTask.generated.h"
 
 class UAnimNextSchedule;
-class UAnimNextParameterBlock;
+class UAnimNextGraph;
 class UAnimNextSchedulerWorldSubsystem;
 class UAnimNextComponent;
 
@@ -55,9 +55,9 @@ private:
 	UPROPERTY()
 	FName Scope;
 
-	// Parameters to apply in this scope
+	// Parameters to apply in this scopes
 	UPROPERTY()
-	TArray<TObjectPtr<UAnimNextParameterBlock>> ParameterBlocks;
+	TArray<TObjectPtr<UAnimNextGraph>> Parameters;
 };
 
 USTRUCT()
diff --git a/Engine/Plugins/Experimental/AnimNext/Source/AnimNext/Internal/Scheduler/Scheduler.h b/Engine/Plugins/Experimental/AnimNext/Source/AnimNext/Internal/Scheduler/Scheduler.h
--- a/Engine/Plugins/Experimental/AnimNext/Source/AnimNext/Internal/Scheduler/Scheduler.h
+++ b/Engine/Plugins/Experimental/AnimNext/Source/AnimNext/Internal/Scheduler/Scheduler.h
@@ -7,7 +7,6 @@
 
 enum class EAnimNextScheduleInitMethod : uint8;
 class UAnimNextGraph;
-class UAnimNextParameterBlock;
 class UAnimNextSchedule;
 
 namespace UE::AnimNext
diff --git a/Engine/Plugins/Experimental/AnimNext/Source/AnimNext/Private/Graph/AnimGraph/AnimNode_AnimNextParameters.cpp b/Engine/Plugins/Experimental/AnimNext/Source/AnimNext/Private/Graph/AnimGraph/AnimNode_AnimNextParameters.cpp
--- a/Engine/Plugins/Experimental/AnimNext/Source/AnimNext/Private/Graph/AnimGraph/AnimNode_AnimNextParameters.cpp
+++ b/Engine/Plugins/Experimental/AnimNext/Source/AnimNext/Private/Graph/AnimGraph/AnimNode_AnimNextParameters.cpp
@@ -2,9 +2,9 @@
 
 #include "Graph/AnimGraph/AnimNode_AnimNextParameters.h"
 #include "Param/ParamStack.h"
-#include "Param/AnimNextParameterBlock.h"
+#include "Graph/AnimNextGraph.h"
 #include "AnimGraphParamStackScope.h"
-#include "Param/ParameterBlockProxy.h"
+#include "Param/ParametersProxy.h"
 
 #include UE_INLINE_GENERATED_CPP_BY_NAME(AnimNode_AnimNextParameters)
 
@@ -57,16 +57,16 @@ void FAnimNode_AnimNextParameters::Update_AnyThread(const FAnimationUpdateContex
 
 	GetEvaluateGraphExposedInputs().Execute(Context);
 
-	UAnimNextParameterBlock* CurrentParameters = Parameters;
+	UAnimNextGraph* CurrentParameters = Parameters;
 
-	// Reconstruct param block's cached layer if required
+	// Reconstruct parameter's cached layer if required
 	if (CurrentParameters != PreviousParameters || !ParametersProxy.IsValid())
 	{
 		ParametersProxy.Reset();
 
 		if (CurrentParameters)
 		{
-			ParametersProxy = MakeUnique<FParameterBlockProxy>(CurrentParameters);
+			ParametersProxy = MakeUnique<FParametersProxy>(CurrentParameters);
 		}
 
 		PreviousParameters = CurrentParameters;
diff --git a/Engine/Plugins/Experimental/AnimNext/Source/AnimNext/Private/Graph/AnimNextGraph.cpp b/Engine/Plugins/Experimental/AnimNext/Source/AnimNext/Private/Graph/AnimNextGraph.cpp
--- a/Engine/Plugins/Experimental/AnimNext/Source/AnimNext/Private/Graph/AnimNextGraph.cpp
+++ b/Engine/Plugins/Experimental/AnimNext/Source/AnimNext/Private/Graph/AnimNextGraph.cpp
@@ -15,10 +15,15 @@
 #include "Graph/AnimNextGraphInstance.h"
 #include "Serialization/MemoryReader.h"
 #include "AnimNextStats.h"
+#include "RigVMRuntimeDataRegistry.h"
 #include "Graph/AnimNextGraphEntryPoint.h"
+#include "Graph/RigUnit_AnimNextBeginExecution.h"
 #include "Graph/RigUnit_AnimNextGraphRoot.h"
+#include "UObject/LinkerLoad.h"
+#include "UObject/ObjectResource.h"
 
 DEFINE_STAT(STAT_AnimNext_Graph_AllocateInstance);
+DEFINE_STAT(STAT_AnimNext_Graph_UpdateParamLayer);
 
 #include UE_INLINE_GENERATED_CPP_BY_NAME(AnimNextGraph)
 
@@ -92,6 +97,34 @@ void UAnimNextGraph::AllocateInstanceImpl(FAnimNextGraphInstance* ParentGraphIns
 #endif
 }
 
+void UAnimNextGraph::UpdateLayer(UE::AnimNext::FParamStackLayerHandle& InHandle, float InDeltaTime) const
+{
+	SCOPE_CYCLE_COUNTER(STAT_AnimNext_Graph_UpdateParamLayer);
+	
+	if (VM)
+	{
+		if (TSharedPtr<UE::AnimNext::FRigVMRuntimeData> RuntimeData = UE::AnimNext::FRigVMRuntimeDataRegistry::FindOrAddLocalRuntimeData(VM, GetRigVMExtendedExecuteContext()).Pin())
+		{
+			FRigVMExtendedExecuteContext& Context = RuntimeData->Context;
+
+			check(Context.VMHash == VM->GetVMHash());
+
+			FAnimNextExecuteContext& AnimNextContext = Context.GetPublicDataSafe<FAnimNextExecuteContext>();
+
+			// Parameter setup
+			AnimNextContext.SetContextData<FAnimNextParamContextData>(InHandle);
+
+			// RigVM setup
+			AnimNextContext.SetDeltaTime(InDeltaTime);
+
+			VM->ExecuteVM(Context, FRigUnit_AnimNextBeginExecution::EventName);
+
+			// Reset the context to avoid issues if we forget to reset it the next time we use it
+			AnimNextContext.DebugReset<FAnimNextParamContextData>();
+		}
+	}
+}
+
 void UAnimNextGraph::PostLoad()
 {
 	using namespace UE::AnimNext;
@@ -108,23 +141,36 @@ void UAnimNextGraph::Serialize(FArchive& Ar)
 
 	if (Ar.IsLoading())
 	{
-		int32 SharedDataArchiveBufferSize = 0;
-		Ar << SharedDataArchiveBufferSize;
+		if(Ar.CustomVer(FFortniteMainBranchObjectVersion::GUID) < FFortniteMainBranchObjectVersion::AnimNextCombineParameterBlocksAndGraphs)
+		{
+			// Skip over shared archive buffer if we are loading from an older version
+			if (const FLinkerLoad* Linker = GetLinker())
+			{
+				const int32 LinkerIndex = GetLinkerIndex();
+				const FObjectExport& Export = Linker->ExportMap[LinkerIndex];
+				Ar.Seek(Export.SerialOffset + Export.SerialSize);
+			}
+		}
+		else
+		{
+			int32 SharedDataArchiveBufferSize = 0;
+			Ar << SharedDataArchiveBufferSize;
 
 #if !WITH_EDITORONLY_DATA
-		// When editor data isn't present, we don't persist the archive buffer as it is only needed on load
-		// to populate the graph shared data
-		TArray<uint8> SharedDataArchiveBuffer;
+			// When editor data isn't present, we don't persist the archive buffer as it is only needed on load
+			// to populate the graph shared data
+			TArray<uint8> SharedDataArchiveBuffer;
 #endif
 
-		SharedDataArchiveBuffer.SetNumUninitialized(SharedDataArchiveBufferSize);
-		Ar.Serialize(SharedDataArchiveBuffer.GetData(), SharedDataArchiveBufferSize);
+			SharedDataArchiveBuffer.SetNumUninitialized(SharedDataArchiveBufferSize);
+			Ar.Serialize(SharedDataArchiveBuffer.GetData(), SharedDataArchiveBufferSize);
 
-		if (Ar.IsLoadingFromCookedPackage())
-		{
-			// If we are cooked, we populate our graph shared data otherwise in the editor we'll compile on load
-			// and re-populate everything then to account for changes in code/content
-			LoadFromArchiveBuffer(SharedDataArchiveBuffer);
+			if (Ar.IsLoadingFromCookedPackage())
+			{
+				// If we are cooked, we populate our graph shared data otherwise in the editor we'll compile on load
+				// and re-populate everything then to account for changes in code/content
+				LoadFromArchiveBuffer(SharedDataArchiveBuffer);
+			}
 		}
 	}
 	else if (Ar.IsSaving())
diff --git a/Engine/Plugins/Experimental/AnimNext/Source/AnimNext/Private/Param/AnimNextParameterBlock.cpp b/Engine/Plugins/Experimental/AnimNext/Source/AnimNext/Private/Param/AnimNextParameterBlock.cpp
deleted file mode 100644
diff --git a/Engine/Plugins/Experimental/AnimNext/Source/AnimNext/Private/Param/ParameterBlockProxy.cpp b/Engine/Plugins/Experimental/AnimNext/Source/AnimNext/Private/Param/ParameterBlockProxy.cpp
deleted file mode 100644
diff --git a/Engine/Plugins/Experimental/AnimNext/Source/AnimNext/Private/Param/ParameterBlockProxy.h b/Engine/Plugins/Experimental/AnimNext/Source/AnimNext/Private/Param/ParameterBlockProxy.h
deleted file mode 100644
diff --git a/Engine/Plugins/Experimental/AnimNext/Source/AnimNext/Private/Param/ParametersProxy.cpp b/Engine/Plugins/Experimental/AnimNext/Source/AnimNext/Private/Param/ParametersProxy.cpp
new file mode 100644
--- /dev/null
+++ b/Engine/Plugins/Experimental/AnimNext/Source/AnimNext/Private/Param/ParametersProxy.cpp
@@ -0,0 +1,38 @@
+// Copyright Epic Games, Inc. All Rights Reserved.
+
+#include "Param/ParametersProxy.h"
+#include "Graph/AnimNextGraph.h"
+#include "Param/ParamStack.h"
+
+namespace UE::AnimNext
+{
+
+FParametersProxy::FParametersProxy(UAnimNextGraph* InGraph)
+	: Graph(InGraph)
+	, PropertyBag(InGraph->PropertyBag)
+	, LayerHandle(FParamStack::MakeReferenceLayer(PropertyBag))
+{
+	check(Graph);
+}
+
+void FParametersProxy::Update(float DeltaTime)
+{
+#if WITH_EDITOR	// Layout should only be changing in editor
+	const FInstancedPropertyBag* HandlePropertyBag = LayerHandle.As<FInstancedPropertyBag>();
+	if(HandlePropertyBag == nullptr || HandlePropertyBag->GetPropertyBagStruct() != Graph->PropertyBag.GetPropertyBagStruct())
+	{
+		PropertyBag = Graph->PropertyBag;
+		LayerHandle = FParamStack::MakeReferenceLayer(PropertyBag);
+	}
+#endif
+
+	Graph->UpdateLayer(LayerHandle, DeltaTime);
+}
+
+void FParametersProxy::AddReferencedObjects(FReferenceCollector& Collector)
+{
+	Collector.AddReferencedObject(Graph);
+	PropertyBag.AddStructReferencedObjects(Collector);
+}
+
+}
diff --git a/Engine/Plugins/Experimental/AnimNext/Source/AnimNext/Private/Param/ParametersProxy.h b/Engine/Plugins/Experimental/AnimNext/Source/AnimNext/Private/Param/ParametersProxy.h
new file mode 100644
--- /dev/null
+++ b/Engine/Plugins/Experimental/AnimNext/Source/AnimNext/Private/Param/ParametersProxy.h
@@ -0,0 +1,36 @@
+// Copyright Epic Games, Inc. All Rights Reserved.
+
+#pragma once
+
+#include "Param/IParameterSource.h"
+#include "PropertyBag.h"
+#include "UObject/StrongObjectPtr.h"
+
+class UAnimNextGraph;
+
+namespace UE::AnimNext
+{
+
+// Proxy struct used to reference parameter instance data
+struct FParametersProxy : public IParameterSource
+{
+	FParametersProxy() = delete;
+
+	FParametersProxy(UAnimNextGraph* InGraph);
+
+	// IParameterSource interface
+	virtual void Update(float DeltaTime) override;
+	virtual const FParamStackLayerHandle& GetLayerHandle() const override { return LayerHandle; }
+	virtual void AddReferencedObjects(FReferenceCollector& Collector) override;
+
+	// The object that this proxy wraps
+	TObjectPtr<UAnimNextGraph> Graph;
+
+	// Copy of the parameter data
+	FInstancedPropertyBag PropertyBag;
+
+	// Layer handle - must be updated if PropertyBag changes layout
+	FParamStackLayerHandle LayerHandle;
+};
+
+}
\ No newline at end of file
diff --git a/Engine/Plugins/Experimental/AnimNext/Source/AnimNext/Private/Scheduler/AnimNextSchedule.cpp b/Engine/Plugins/Experimental/AnimNext/Source/AnimNext/Private/Scheduler/AnimNextSchedule.cpp
--- a/Engine/Plugins/Experimental/AnimNext/Source/AnimNext/Private/Scheduler/AnimNextSchedule.cpp
+++ b/Engine/Plugins/Experimental/AnimNext/Source/AnimNext/Private/Scheduler/AnimNextSchedule.cpp
@@ -6,7 +6,6 @@
 #include "EngineLogs.h"
 #include "UObject/AssetRegistryTagsContext.h"
 #include "Graph/AnimNextGraph.h"
-#include "Param/AnimNextParameterBlock.h"
 
 #if WITH_EDITOR
 TUniqueFunction<void(UAnimNextSchedule*)> UAnimNextSchedule::CompileFunction;
@@ -24,9 +23,9 @@ void UAnimNextScheduleEntry_ParamScope::GetPreloadDependencies(TArray<UObject*>&
 {
 	Super::GetPreloadDependencies(OutDeps);
 
-	for(UAnimNextParameterBlock* ParameterBlock : ParameterBlocks)
+	for(UAnimNextGraph* Graph : Parameters)
 	{
-		OutDeps.Add(ParameterBlock);
+		OutDeps.Add(Graph);
 	}
 
 	for(UAnimNextScheduleEntry* SubEntry : SubEntries)
diff --git a/Engine/Plugins/Experimental/AnimNext/Source/AnimNext/Private/Scheduler/AnimNextSchedulerWorldSubsystem.h b/Engine/Plugins/Experimental/AnimNext/Source/AnimNext/Private/Scheduler/AnimNextSchedulerWorldSubsystem.h
--- a/Engine/Plugins/Experimental/AnimNext/Source/AnimNext/Private/Scheduler/AnimNextSchedulerWorldSubsystem.h
+++ b/Engine/Plugins/Experimental/AnimNext/Source/AnimNext/Private/Scheduler/AnimNextSchedulerWorldSubsystem.h
@@ -15,7 +15,6 @@
 #include "AnimNextSchedulerWorldSubsystem.generated.h"
 
 class UAnimNextSchedule;
-class UAnimNextParameterBlock;
 
 namespace UE::AnimNext
 {
diff --git a/Engine/Plugins/Experimental/AnimNext/Source/AnimNext/Private/Scheduler/ScheduleInstanceData.cpp b/Engine/Plugins/Experimental/AnimNext/Source/AnimNext/Private/Scheduler/ScheduleInstanceData.cpp
--- a/Engine/Plugins/Experimental/AnimNext/Source/AnimNext/Private/Scheduler/ScheduleInstanceData.cpp
+++ b/Engine/Plugins/Experimental/AnimNext/Source/AnimNext/Private/Scheduler/ScheduleInstanceData.cpp
@@ -9,7 +9,7 @@
 #include "AnimNextStats.h"
 #include "Param/ExternalParameterRegistry.h"
 #include "Param/IParameterSourceFactory.h"
-#include "Param/ParameterBlockProxy.h"
+#include "Param/ParametersProxy.h"
 #include "Param/PropertyBagProxy.h"
 #include "Scheduler/AnimNextScheduleExternalParamTask.h"
 
@@ -33,12 +33,12 @@ FScheduleInstanceData::FScheduleInstanceData(const FScheduleContext& InScheduleC
 	{
 		const FAnimNextScheduleParamScopeEntryTask& EntryTask = InSchedule->ParamScopeEntryTasks[EntryIndex];
 		FScopeCache& ScopeCache = ScopeCaches[EntryTask.ParamScopeIndex];
-		ScopeCache.ParameterSources.Reserve(EntryTask.ParameterBlocks.Num());
-		for(UAnimNextParameterBlock* ParameterBlock : EntryTask.ParameterBlocks)
+		ScopeCache.ParameterSources.Reserve(EntryTask.Parameters.Num());
+		for(UAnimNextGraph* Graph : EntryTask.Parameters)
 		{
-			if(ParameterBlock)
+			if(Graph)
 			{
-				ScopeCache.ParameterSources.Emplace(MakeUnique<FParameterBlockProxy>(ParameterBlock));
+				ScopeCache.ParameterSources.Emplace(MakeUnique<FParametersProxy>(Graph));
 			}
 		}
 
diff --git a/Engine/Plugins/Experimental/AnimNext/Source/AnimNext/Public/AnimNextStats.h b/Engine/Plugins/Experimental/AnimNext/Source/AnimNext/Public/AnimNextStats.h
--- a/Engine/Plugins/Experimental/AnimNext/Source/AnimNext/Public/AnimNextStats.h
+++ b/Engine/Plugins/Experimental/AnimNext/Source/AnimNext/Public/AnimNextStats.h
@@ -30,7 +30,7 @@ DECLARE_CYCLE_STAT_EXTERN(TEXT("AnimNext: Param Stack Get Param"), STAT_AnimNext
 DECLARE_CYCLE_STAT_EXTERN(TEXT("AnimNext: Param Stack Adapter"), STAT_AnimNext_ParamStack_Adapter, STATGROUP_AnimNext, ANIMNEXT_API);
 DECLARE_CYCLE_STAT_EXTERN(TEXT("AnimNext: Param Stack Coalesce"), STAT_AnimNext_ParamStack_Coalesce, STATGROUP_AnimNext, ANIMNEXT_API);
 DECLARE_CYCLE_STAT_EXTERN(TEXT("AnimNext: Param Stack Decoalesce"), STAT_AnimNext_ParamStack_Decoalesce, STATGROUP_AnimNext, ANIMNEXT_API);
-DECLARE_CYCLE_STAT_EXTERN(TEXT("AnimNext: Run Param Block"), STAT_AnimNext_ParamBlock_UpdateLayer, STATGROUP_AnimNext, ANIMNEXT_API);
+DECLARE_CYCLE_STAT_EXTERN(TEXT("AnimNext: Update Param Layer"), STAT_AnimNext_Graph_UpdateParamLayer, STATGROUP_AnimNext, ANIMNEXT_API);
 
 DECLARE_CYCLE_STAT_EXTERN(TEXT("AnimNext: Generate Reference Pose"), STAT_AnimNext_GenerateReferencePose, STATGROUP_AnimNext, ANIMNEXT_API);
 DECLARE_CYCLE_STAT_EXTERN(TEXT("AnimNext: Remap Pose"), STAT_AnimNext_RemapPose, STATGROUP_AnimNext, ANIMNEXT_API);
diff --git a/Engine/Plugins/Experimental/AnimNext/Source/AnimNext/Public/Param/ParamStackLayerHandle.h b/Engine/Plugins/Experimental/AnimNext/Source/AnimNext/Public/Param/ParamStackLayerHandle.h
--- a/Engine/Plugins/Experimental/AnimNext/Source/AnimNext/Public/Param/ParamStackLayerHandle.h
+++ b/Engine/Plugins/Experimental/AnimNext/Source/AnimNext/Public/Param/ParamStackLayerHandle.h
@@ -8,8 +8,6 @@
 #include "Param/ParamId.h"
 #include "Param/ParamTypeHandle.h"
 
-class UAnimNextParameterBlock;
-
 namespace UE::AnimNext
 {
 	struct FParamStackLayer;
@@ -169,7 +167,6 @@ public:
 
 private:
 	friend struct FParamStack;
-	friend class UAnimNextParameterBlock;
 
 	// Set parameter values
 	ANIMNEXT_API FParamResult SetValuesInternal(TConstArrayView<Private::FParamEntry> InParams);
diff --git a/Engine/Plugins/Experimental/AnimNext/Source/AnimNextEditor/Internal/Param/ParameterBlockFactory.h b/Engine/Plugins/Experimental/AnimNext/Source/AnimNextEditor/Internal/Param/ParameterBlockFactory.h
deleted file mode 100644
diff --git a/Engine/Plugins/Experimental/AnimNext/Source/AnimNextEditor/Internal/Param/ParameterPickerArgs.h b/Engine/Plugins/Experimental/AnimNext/Source/AnimNextEditor/Internal/Param/ParameterPickerArgs.h
--- a/Engine/Plugins/Experimental/AnimNext/Source/AnimNextEditor/Internal/Param/ParameterPickerArgs.h
+++ b/Engine/Plugins/Experimental/AnimNext/Source/AnimNextEditor/Internal/Param/ParameterPickerArgs.h
@@ -33,15 +33,15 @@ struct FParameterToAdd
 	FName Name;
 };
 
-// A parameter asset, optionally bound in a block
+// A parameter asset, optionally bound in a graph
 struct FParameterBindingReference
 {
 	FParameterBindingReference() = default;
 
-	FParameterBindingReference(FName InParameter, const FAnimNextParamType& InType, const FAssetData& InBlock = FAssetData())
+	FParameterBindingReference(FName InParameter, const FAnimNextParamType& InType, const FAssetData& InGraph = FAssetData())
 		: Parameter(InParameter)
 		, Type(InType)
-		, Block(InBlock)
+		, Graph(InGraph)
 	{
 	}
 
@@ -54,14 +54,14 @@ struct FParameterBindingReference
 	// Asset (first found in asset registry) that the parameter is used in
 	FAssetData Asset;
 
-	// Optional block asset that the parameter is bound in
-	FAssetData Block;
+	// Optional graph asset that the parameter is bound in
+	FAssetData Graph;
 };
 
-// Delegate called when a parameter has been picked. Block argument is invalid when an unbound parameter is chosen.
+// Delegate called when a parameter has been picked. Graph argument is invalid when an unbound parameter is chosen.
 DECLARE_DELEGATE_OneParam(FOnGetParameterBindings, TArray<FParameterBindingReference>& /*OutParameterBindings*/);
 
-// Delegate called when a parameter has been picked. Block argument is invalid when an unbound parameter is chosen.
+// Delegate called when a parameter has been picked. Graph argument is invalid when an unbound parameter is chosen.
 DECLARE_DELEGATE_OneParam(FOnParameterPicked, const FParameterBindingReference& /*InParameterBinding*/);
 
 // Delegate called when a parameter is due to be added.
@@ -109,17 +109,17 @@ struct FParameterPickerArgs
 	// Whether we allow selecting multiple parameters or just one
 	bool bMultiSelect = true;
 
-	// Whether we should show parameters that are bound in a parameter block
+	// Whether we should show parameters that are bound in a graph
 	bool bShowBoundParameters = true;
 
-	// Whether we should show parameters that are not bound in a parameter block (if bShowBoundParameters is false this will show all parameters)
+	// Whether we should show parameters that are not bound in a graph (if bShowBoundParameters is false this will show all parameters)
 	bool bShowUnboundParameters = true;	
 
 	// Whether we should show parameters that are built in
 	bool bShowBuiltInParameters = true;
 
-	// Whether we should show the block alongside bound parameters
-	bool bShowBlocks = true;
+	// Whether we should show the source graph alongside bound parameters
+	bool bShowSourceGraph = true;
 
 	// Whether we should allow new parameters to be created by this widget
 	bool bAllowNew = true;
diff --git a/Engine/Plugins/Experimental/AnimNext/Source/AnimNextEditor/Private/AnimNextEditorModule.cpp b/Engine/Plugins/Experimental/AnimNext/Source/AnimNextEditor/Private/AnimNextEditorModule.cpp
--- a/Engine/Plugins/Experimental/AnimNext/Source/AnimNextEditor/Private/AnimNextEditorModule.cpp
+++ b/Engine/Plugins/Experimental/AnimNext/Source/AnimNextEditor/Private/AnimNextEditorModule.cpp
@@ -14,9 +14,7 @@
 #include "Graph/AnimNextGraphPanelNodeFactory.h"
 #include "Graph/AnimNextGraph_EdGraphNodeCustomization.h"
 #include "Graph/AnimNextGraph_EditorData.h"
-#include "Param/AnimNextParameterBlock.h"
-#include "Param/AnimNextParameterBlock_EditorData.h"
-#include "Param/ParameterBlockParameterCustomization.h"
+#include "Param/ParameterCustomization.h"
 #include "Param/ParameterPickerArgs.h"
 #include "Param/ParametersGraphPanelPinFactory.h"
 #include "Param/ParamNamePropertyCustomization.h"
@@ -59,8 +57,8 @@ class FModule : public IModule
 			FOnGetPropertyTypeCustomizationInstance::CreateLambda([] { return MakeShared<FParamNamePropertyTypeCustomization>(); }),
 			Identifier);
 
-		PropertyModule.RegisterCustomClassLayout("AnimNextParameterBlockParameter", 
-			FOnGetDetailCustomizationInstance::CreateLambda([] { return MakeShared<FParameterBlockParameterCustomization>(); }));
+		PropertyModule.RegisterCustomClassLayout("AnimNextGraph_Parameter", 
+			FOnGetDetailCustomizationInstance::CreateLambda([] { return MakeShared<FParameterCustomization>(); }));
 
 		PropertyModule.RegisterCustomClassLayout("AnimNextGraph_EdGraphNode",
 			FOnGetDetailCustomizationInstance::CreateLambda([] { return MakeShared<FAnimNextGraph_EdGraphNodeCustomization>(); }));
@@ -71,16 +69,6 @@ class FModule : public IModule
 		ParametersGraphPanelPinFactory = MakeShared<FParametersGraphPanelPinFactory>();
 		FEdGraphUtilities::RegisterVisualPinFactory(ParametersGraphPanelPinFactory);
 
-		FWorkspaceEditor::RegisterAssetDocumentWidget(UAnimNextParameterBlock::StaticClass()->GetFName(), [](TSharedRef<FWorkspaceEditor> InEditor, UObject* InAsset)
-		{
-			UAnimNextParameterBlock* ParameterBlock = CastChecked<UAnimNextParameterBlock>(InAsset);
-			UAnimNextParameterBlock_EditorData* EditorData = UncookedOnly::FUtils::GetEditorData(ParameterBlock);
-			return SNew(SRigVMAssetView, EditorData)
-				.OnSelectionChanged(&InEditor.Get(), &FWorkspaceEditor::SetSelectedObjects)
-				.OnOpenGraph(&InEditor.Get(), &FWorkspaceEditor::OnOpenGraph)
-				.OnDeleteEntries(&InEditor.Get(), &FWorkspaceEditor::OnDeleteEntries);
-		});
-
 		FWorkspaceEditor::RegisterAssetDocumentWidget(UAnimNextSchedule::StaticClass()->GetFName(), [](TSharedRef<FWorkspaceEditor> InEditor, UObject* InAsset)
 		{
 			UAnimNextSchedule* Schedule = CastChecked<UAnimNextSchedule>(InAsset);
@@ -108,7 +96,7 @@ class FModule : public IModule
 
 		SRigVMAssetView::RegisterCategoryFactory("Parameters", [](UAnimNextRigVMAssetEditorData* InEditorData)
 		{
-			UAnimNextParameterBlock_EditorData* EditorData = CastChecked<UAnimNextParameterBlock_EditorData>(InEditorData);
+			UAnimNextGraph_EditorData* EditorData = CastChecked<UAnimNextGraph_EditorData>(InEditorData);
 			return SNew(SSimpleComboButton)
 				.Text(LOCTEXT("AddParameterButton", "Add Parameter"))
 				.Icon(FAppStyle::Get().GetBrush("Icons.Plus"))
@@ -119,13 +107,13 @@ class FModule : public IModule
 					
 					FParameterPickerArgs Args;
 					Args.bMultiSelect = false;
-					Args.bShowBlocks = false;
+					Args.bShowSourceGraph = false;
 					Args.bShowBoundParameters = false;
 					Args.bShowBuiltInParameters = false; // Built-In parameters disabled for MVP
 					Args.OnFilterParameter = FOnFilterParameter::CreateLambda([EditorData, AssetData](const FParameterBindingReference& InParameterBinding)
 					{
-						// Skip params that are already bound in this block
-						if(InParameterBinding.Block == AssetData)
+						// Skip params that are already bound in this graph
+						if(InParameterBinding.Graph == AssetData)
 						{
 							return EFilterParameterResult::Exclude;
 						}
@@ -160,18 +148,18 @@ class FModule : public IModule
 				});
 		});
 
-		SRigVMAssetView::RegisterCategoryFactory("Parameter Graphs", [](UAnimNextRigVMAssetEditorData* InEditorData)
+		SRigVMAssetView::RegisterCategoryFactory("Event Graphs", [](UAnimNextRigVMAssetEditorData* InEditorData)
 		{
-			UAnimNextParameterBlock_EditorData* EditorData = CastChecked<UAnimNextParameterBlock_EditorData>(InEditorData);
+			UAnimNextGraph_EditorData* EditorData = CastChecked<UAnimNextGraph_EditorData>(InEditorData);
 			return SNew(SSimpleButton)
-				.Text(LOCTEXT("AddGraphButton", "Add Graph"))
+				.Text(LOCTEXT("AddEventGraphButton", "Add Event Graph"))
 				.Icon(FAppStyle::Get().GetBrush("Icons.Plus"))
 				.OnClicked_Lambda([EditorData]()
 				{
-					FScopedTransaction Transaction(LOCTEXT("AddGraph", "Add Graph"));
+					FScopedTransaction Transaction(LOCTEXT("AddEventGraph", "Add Event Graph"));
 
 					// Create a new entry for the graph
-					EditorData->AddGraph(TEXT("NewGraph"));
+					EditorData->AddEventGraph(TEXT("NewGraph"));
 
 					return FReply::Handled();
 				});
@@ -181,14 +169,14 @@ class FModule : public IModule
 		{
 			UAnimNextGraph_EditorData* EditorData = CastChecked<UAnimNextGraph_EditorData>(InEditorData);
 			return SNew(SSimpleButton)
-				.Text(LOCTEXT("AddGraphButton", "Add Graph"))
+				.Text(LOCTEXT("AddGraphButton", "Add Animation Graph"))
 				.Icon(FAppStyle::Get().GetBrush("Icons.Plus"))
 				.OnClicked_Lambda([EditorData]()
 				{
-					FScopedTransaction Transaction(LOCTEXT("AddGraph", "Add Graph"));
+					FScopedTransaction Transaction(LOCTEXT("AddAnimationGraph", "Add Animation Graph"));
 
 					// Create a new entry for the graph
-					EditorData->AddGraph(TEXT("NewGraph"));
+					EditorData->AddAnimationGraph(TEXT("NewGraph"));
 
 					return FReply::Handled();
 				});
@@ -203,7 +191,7 @@ class FModule : public IModule
 			PropertyModule.UnregisterCustomPropertyTypeLayout("AnimNextParamType");
 			PropertyModule.UnregisterCustomPropertyTypeLayout("AnimNextParam");
 			PropertyModule.UnregisterCustomPropertyTypeLayout("NameProperty");
-			PropertyModule.UnregisterCustomClassLayout("AnimNextParameterBlockParameter");
+			PropertyModule.UnregisterCustomClassLayout("AnimNextGraph_Parameter");
 			PropertyModule.UnregisterCustomClassLayout("AnimNextGraph_EdGraphNode");
 		}
 
@@ -212,7 +200,6 @@ class FModule : public IModule
 		FEdGraphUtilities::UnregisterVisualPinFactory(ParametersGraphPanelPinFactory);
 
 		FWorkspaceEditor::UnregisterAssetDocumentWidget("AnimNextSchedule");
-		FWorkspaceEditor::UnregisterAssetDocumentWidget("AnimNextParameterBlock");
 		FWorkspaceEditor::UnregisterAssetDocumentWidget("AnimNextGraph");
 		
 		SRigVMAssetView::UnregisterCategoryFactory("Parameters");
diff --git a/Engine/Plugins/Experimental/AnimNext/Source/AnimNextEditor/Private/Common/SRigVMAssetView.cpp b/Engine/Plugins/Experimental/AnimNext/Source/AnimNextEditor/Private/Common/SRigVMAssetView.cpp
--- a/Engine/Plugins/Experimental/AnimNext/Source/AnimNextEditor/Private/Common/SRigVMAssetView.cpp
+++ b/Engine/Plugins/Experimental/AnimNext/Source/AnimNextEditor/Private/Common/SRigVMAssetView.cpp
@@ -2,15 +2,16 @@
 
 #include "SRigVMAssetView.h"
 
+#include "AnimNextRigVMAsset.h"
+#include "AnimNextRigVMAssetEditorData.h"
 #include "AnimNextRigVMAssetEntry.h"
-#include "Param/AnimNextParameterBlock.h"
-#include "Param/AnimNextParameterBlock_EditorData.h"
 #include "Widgets/Layout/SBox.h"
 #include "Widgets/Input/SSearchBox.h"
 #include "DetailLayoutBuilder.h"
 #include "ISourceControlModule.h"
 #include "ISourceControlProvider.h"
 #include "EditorUtils.h"
+#include "IAnimNextRigVMParameterInterface.h"
 #include "InstancedPropertyBagStructureDataProvider.h"
 #include "UncookedOnlyUtils.h"
 #include "RigVMAssetViewMenuContext.h"
@@ -22,8 +23,6 @@
 #include "Widgets/Text/SInlineEditableTextBlock.h"
 #include "ToolMenus.h"
 #include "ScopedTransaction.h"
-#include "Param/AnimNextParameterBlockParameter.h"
-#include "Param/AnimNextParameterBlockGraph.h"
 #include "Framework/Application/SlateApplication.h"
 #include "PropertyEditorModule.h"
 #include "ISinglePropertyView.h"
@@ -35,7 +34,7 @@
 namespace UE::AnimNext::Editor
 {
 
-namespace ParameterBlockView
+namespace RigVMAssetView
 {
 
 static FName ContextMenuName(TEXT("AnimNext.RigVMAssetView.ContextMenu"));
@@ -205,7 +204,7 @@ void SRigVMAssetView::UnregisterCategoryFactory(FName InCategory)
 
 void SRigVMAssetView::Construct(const FArguments& InArgs, UAnimNextRigVMAssetEditorData* InEditorData)
 {
-	using namespace ParameterBlockView;
+	using namespace RigVMAssetView;
 	
 	check(InEditorData);
 
@@ -233,7 +232,7 @@ void SRigVMAssetView::Construct(const FArguments& InArgs, UAnimNextRigVMAssetEdi
 	OnOpenGraphDelegate = InArgs._OnOpenGraph;
 	OnDeleteEntriesDelegate = InArgs._OnDeleteEntries;
 
-	EditorData->ModifiedDelegate.AddSP(this, &SRigVMAssetView::HandleBlockModified);
+	EditorData->ModifiedDelegate.AddSP(this, &SRigVMAssetView::HandleAssetModified);
 
 	ChildSlot
 	[
@@ -442,7 +441,7 @@ void SRigVMAssetView::RefreshEntries()
 
 void SRigVMAssetView::RefreshFilter()
 {
-	using namespace ParameterBlockView;
+	using namespace RigVMAssetView;
 
 	FilteredEntries = Categories;
 
@@ -457,8 +456,8 @@ void SRigVMAssetView::RefreshFilter()
 	{
 		if (Entry->PassesFilter(FilterTextAsString))
 		{
-			const UObject* BlockEntry = Entry->WeakEntry.Get();
-			const UClass* EntryClass = BlockEntry->GetClass();
+			const UObject* AssetEntry = Entry->WeakEntry.Get();
+			const UClass* EntryClass = AssetEntry->GetClass();
 			const FString& CategoryMetaData = EntryClass->GetMetaData(NAME_Category);
 			TSharedRef<FRigVMAssetViewEntry> CategoryEntry = GetCategoryEntry(*CategoryMetaData);
 			CategoryEntry->Children.Add(Entry);
@@ -481,7 +480,7 @@ void SRigVMAssetView::BindCommands()
 		FCanExecuteAction::CreateSP(this, &SRigVMAssetView::HasValidSingleSelection));
 }
 
-void SRigVMAssetView::HandleBlockModified(UAnimNextRigVMAssetEditorData* InEditorData)
+void SRigVMAssetView::HandleAssetModified(UAnimNextRigVMAssetEditorData* InEditorData)
 {
 	check(InEditorData == EditorData);
 
@@ -490,7 +489,7 @@ void SRigVMAssetView::HandleBlockModified(UAnimNextRigVMAssetEditorData* InEdito
 
 TSharedRef<SWidget> SRigVMAssetView::HandleGetContextContent()
 {
-	using namespace ParameterBlockView;
+	using namespace RigVMAssetView;
 	
 	UToolMenus* ToolMenus = UToolMenus::Get();
 
@@ -639,7 +638,7 @@ class SRigVMAssetViewRow : public SMultiColumnTableRow<TSharedRef<FRigVMAssetVie
 
 	virtual TSharedRef<SWidget> GenerateWidgetForColumn(const FName& InColumnName) override
 	{
-		using namespace ParameterBlockView;
+		using namespace RigVMAssetView;
 		
 		if(InColumnName == Column_RevisionControl)
 		{
diff --git a/Engine/Plugins/Experimental/AnimNext/Source/AnimNextEditor/Private/Common/SRigVMAssetView.h b/Engine/Plugins/Experimental/AnimNext/Source/AnimNextEditor/Private/Common/SRigVMAssetView.h
--- a/Engine/Plugins/Experimental/AnimNext/Source/AnimNextEditor/Private/Common/SRigVMAssetView.h
+++ b/Engine/Plugins/Experimental/AnimNext/Source/AnimNextEditor/Private/Common/SRigVMAssetView.h
@@ -56,8 +56,8 @@ private:
 	// Bind input commands
 	void BindCommands();
 
-	// Handle modifications to the block
-	void HandleBlockModified(UAnimNextRigVMAssetEditorData* InEditorData);
+	// Handle modifications to the asset
+	void HandleAssetModified(UAnimNextRigVMAssetEditorData* InEditorData);
 
 	// Get the content for the context menu
 	TSharedRef<SWidget> HandleGetContextContent();
diff --git a/Engine/Plugins/Experimental/AnimNext/Source/AnimNextEditor/Private/EditorUtils.cpp b/Engine/Plugins/Experimental/AnimNext/Source/AnimNextEditor/Private/EditorUtils.cpp
--- a/Engine/Plugins/Experimental/AnimNext/Source/AnimNextEditor/Private/EditorUtils.cpp
+++ b/Engine/Plugins/Experimental/AnimNext/Source/AnimNextEditor/Private/EditorUtils.cpp
@@ -1,15 +1,13 @@
 // Copyright Epic Games, Inc. All Rights Reserved.
 
 #include "EditorUtils.h"
+
+#include "AnimNextRigVMAssetEntry.h"
 #include "Graph/AnimNextGraph_EditorData.h"
 #include "Param/ParamType.h"
 #include "Kismet2/Kismet2NameValidators.h"
-#include "Param/AnimNextParameterBlock_EditorData.h"
-#include "Param/AnimNextParameterBlock.h"
+#include "Graph/AnimNextGraph.h"
 #include "PropertyBagDetails.h"
-#include "Graph/AnimNextGraphEntry.h"
-#include "IAnimNextRigVMGraphInterface.h"
-#include "IAnimNextRigVMParameterInterface.h"
 #include "Workspace/AnimNextWorkspace.h"
 #include "UncookedOnlyUtils.h"
 #include "Param/RigVMDispatch_GetParameter.h"
@@ -21,7 +19,7 @@ namespace UE::AnimNext::Editor
 
 void FUtils::GetAllEntryNames(const UAnimNextRigVMAssetEditorData* InEditorData, TSet<FName>& OutNames)
 {
-	for(const  UAnimNextRigVMAssetEntry* Entry : InEditorData->GetAllEntries())
+	for(const UAnimNextRigVMAssetEntry* Entry : InEditorData->GetAllEntries())
 	{
 		OutNames.Add(Entry->GetEntryName());
 	}
@@ -237,111 +235,6 @@ FAnimNextParamType FUtils::GetParameterTypeFromMetaData(const FStringView& InStr
 	return FAnimNextParamType(); 
 }
 
-FName FUtils::ValidateName(const UAnimNextParameterBlock_EditorData* InEditorData, const FString& InName)
-{
-	struct FNameValidator : public INameValidatorInterface
-	{
-		FNameValidator(const UAnimNextParameterBlock_EditorData* InEditorData)
-			: EditorData(InEditorData)
-		{
-			GetAllEntryNames(EditorData, Names);
-		}
-		
-		virtual EValidatorResult IsValid (const FName& Name, bool bOriginal = false) override
-		{
-			EValidatorResult ValidatorResult = EValidatorResult::AlreadyInUse;
-
-			if(Name == NAME_None)
-			{
-				ValidatorResult = EValidatorResult::EmptyName;
-			}
-			else if(Name.ToString().Len() > MaxNameLength)
-			{
-				ValidatorResult = EValidatorResult::TooLong;
-			}
-			else
-			{
-				// If it is in the names list then it is already in use.
-				if(!Names.Contains(Name))
-				{
-					ValidatorResult = EValidatorResult::Ok;
-
-					// Check for collision with an existing object.
-					if (UObject* ExistingObject = StaticFindObject(/*Class=*/ nullptr, const_cast<UAnimNextParameterBlock_EditorData*>(EditorData), *Name.ToString(), true))
-					{
-						ValidatorResult = EValidatorResult::AlreadyInUse;
-					}
-				}
-			}
-			
-			return ValidatorResult;
-		}
-		
-		virtual EValidatorResult IsValid (const FString& Name, bool bOriginal = false) override
-		{
-			// Converting a string that is too large for an FName will cause an assert, so verify the length
-			if(Name.Len() >= NAME_SIZE)
-			{
-				return EValidatorResult::TooLong;
-			}
-			else if (!FName::IsValidXName(Name, UE_BLUEPRINT_INVALID_NAME_CHARACTERS))
-			{
-				return EValidatorResult::ContainsInvalidCharacters;
-			}
-
-			// If not defined in name table, not current graph name
-			return IsValid( FName(*Name) );
-		}
-
-		/** Name set to validate */
-		TSet<FName> Names;
-		/** The editor data to check for validity within */
-		const UAnimNextParameterBlock_EditorData* EditorData;
-	};
-	
-	FString Name = InName;
-	if (Name.StartsWith(TEXT("RigUnit_")))
-	{
-		Name.RightChopInline(8, EAllowShrinking::No);
-	}
-
-	FNameValidator NameValidator(InEditorData);
-
-	// Clean up BaseName to not contain any invalid characters, which will mean we can never find a legal name no matter how many numbers we add
-	if (NameValidator.IsValid(Name) == EValidatorResult::ContainsInvalidCharacters)
-	{
-		for (TCHAR& TestChar : Name)
-		{
-			for (TCHAR BadChar : UE_BLUEPRINT_INVALID_NAME_CHARACTERS)
-			{
-				if (TestChar == BadChar)
-				{
-					TestChar = TEXT('_');
-					break;
-				}
-			}
-		}
-	}
-	
-	int32 Count = 0;
-	FString BaseName = Name;
-	while (NameValidator.IsValid(Name) != EValidatorResult::Ok)
-	{
-		// Calculate the number of digits in the number, adding 2 (1 extra to correctly count digits, another to account for the '_' that will be added to the name
-		int32 CountLength = Count > 0 ? (int32)log((double)Count) + 2 : 2;
-
-		// If the length of the final string will be too long, cut off the end so we can fit the number
-		if (CountLength + BaseName.Len() > MaxNameLength)
-		{
-			BaseName.LeftInline(MaxNameLength - CountLength);
-		}
-		Name = FString::Printf(TEXT("%s_%d"), *BaseName, Count);
-		Count++;
-	}
-
-	return *Name;
-}
-
 namespace Private
 {
 
diff --git a/Engine/Plugins/Experimental/AnimNext/Source/AnimNextEditor/Private/EditorUtils.h b/Engine/Plugins/Experimental/AnimNext/Source/AnimNextEditor/Private/EditorUtils.h
--- a/Engine/Plugins/Experimental/AnimNext/Source/AnimNextEditor/Private/EditorUtils.h
+++ b/Engine/Plugins/Experimental/AnimNext/Source/AnimNextEditor/Private/EditorUtils.h
@@ -9,11 +9,7 @@
 class UAnimNextGraph;
 class UAnimNextRigVMAssetEditorData;
 struct FAnimNextParamType;
-class UAnimNextParameterBlock;
-class UAnimNextParameterBlockBinding;
-class UAnimNextParameterBlock_EditorData;
 class URigVMController;
-struct FAnimNextParameterBlockAssetRegistryExports;
 struct FAnimNextWorkspaceAssetRegistryExports;
 
 struct FAnimNextParameterProviderAssetRegistryExports;
@@ -29,8 +25,6 @@ struct FUtils
 
 	static FAnimNextParamType GetParameterTypeFromMetaData(const FStringView& InStringView);
 
-	static FName ValidateName(const UAnimNextParameterBlock_EditorData* InEditorData, const FString& InName);
-
 	static void GetFilteredVariableTypeTree(TArray<TSharedPtr<UEdGraphSchema_K2::FPinTypeTreeInfo>>& TypeTree, ETypeTreeFilter TypeTreeFilter);
 
 	static FName GetNewParameterName(const TCHAR* InBaseName, TArrayView<FName> InAdditionalExistingNames);
diff --git a/Engine/Plugins/Experimental/AnimNext/Source/AnimNextEditor/Private/Graph/AnimNextGraphAssetDefinition.cpp b/Engine/Plugins/Experimental/AnimNext/Source/AnimNextEditor/Private/Graph/AnimNextGraphAssetDefinition.cpp
--- a/Engine/Plugins/Experimental/AnimNext/Source/AnimNextEditor/Private/Graph/AnimNextGraphAssetDefinition.cpp
+++ b/Engine/Plugins/Experimental/AnimNext/Source/AnimNextEditor/Private/Graph/AnimNextGraphAssetDefinition.cpp
@@ -24,3 +24,22 @@ EAssetCommandResult UAssetDefinition_AnimNextGraph::OpenAssets(const FAssetOpenA
 
 	return EAssetCommandResult::Handled;
 }
+
+
+FText UAssetDefinition_AnimNextParameter::GetObjectDisplayNameText(UObject* Object) const
+{
+	UAnimNextRigVMAssetEntry* Parameter = CastChecked<UAnimNextRigVMAssetEntry>(Object);
+	return Parameter->GetDisplayName();
+}
+
+FText UAssetDefinition_AnimNextAnimationGraph::GetObjectDisplayNameText(UObject* Object) const
+{
+	UAnimNextRigVMAssetEntry* Parameter = CastChecked<UAnimNextRigVMAssetEntry>(Object);
+	return Parameter->GetDisplayName();
+}
+
+FText UAssetDefinition_AnimNextEventGraph::GetObjectDisplayNameText(UObject* Object) const
+{
+	UAnimNextRigVMAssetEntry* Parameter = CastChecked<UAnimNextRigVMAssetEntry>(Object);
+	return Parameter->GetDisplayName();
+}
\ No newline at end of file
diff --git a/Engine/Plugins/Experimental/AnimNext/Source/AnimNextEditor/Private/Graph/AnimNextGraphAssetDefinition.h b/Engine/Plugins/Experimental/AnimNext/Source/AnimNextEditor/Private/Graph/AnimNextGraphAssetDefinition.h
--- a/Engine/Plugins/Experimental/AnimNext/Source/AnimNextEditor/Private/Graph/AnimNextGraphAssetDefinition.h
+++ b/Engine/Plugins/Experimental/AnimNext/Source/AnimNextEditor/Private/Graph/AnimNextGraphAssetDefinition.h
@@ -4,6 +4,9 @@
 
 #include "Graph/AnimNextGraph.h"
 #include "AssetDefinitionDefault.h"
+#include "Graph/AnimNextGraph_EventGraph.h"
+#include "Graph/AnimNextGraph_Parameter.h"
+#include "Graph/AnimNextGraph_AnimationGraph.h"
 #include "AnimNextGraphAssetDefinition.generated.h"
 
 #define LOCTEXT_NAMESPACE "AnimNextAssetDefinitions"
@@ -16,7 +19,7 @@ class UAssetDefinition_AnimNextGraph : public UAssetDefinitionDefault
 public:
 	// UAssetDefinition interface
 	virtual FText GetAssetDisplayName() const override { return LOCTEXT("AnimNextGraph", "AnimNext Graph"); }
-	virtual FLinearColor GetAssetColor() const override { return FLinearColor(FColor(128,128,64)); }
+	virtual FLinearColor GetAssetColor() const override { return FLinearColor(FColor(128,96,48)); }
 	virtual TSoftClassPtr<UObject> GetAssetClass() const override { return UAnimNextGraph::StaticClass(); }
 	virtual EAssetCommandResult OpenAssets(const FAssetOpenArgs& OpenArgs) const override;
 	virtual TConstArrayView<FAssetCategoryPath> GetAssetCategories() const override
@@ -26,4 +29,43 @@ public:
 	}
 };
 
+UCLASS()
+class UAssetDefinition_AnimNextParameter : public UAssetDefinitionDefault
+{
+	GENERATED_BODY()
+
+public:
+	// UAssetDefinition interface
+	virtual FText GetAssetDisplayName() const override { return LOCTEXT("AnimNextParameter", "Parameter"); }
+	virtual FLinearColor GetAssetColor() const override { return FLinearColor(FColor(128,64,64)); }
+	virtual TSoftClassPtr<UObject> GetAssetClass() const override { return UAnimNextGraph_Parameter::StaticClass(); }
+	virtual FText GetObjectDisplayNameText(UObject* Object) const override;
+};
+
+UCLASS()
+class UAssetDefinition_AnimNextAnimationGraph : public UAssetDefinitionDefault
+{
+	GENERATED_BODY()
+
+public:
+	// UAssetDefinition interface
+	virtual FText GetAssetDisplayName() const override { return LOCTEXT("AnimNextAnimationGraph", "Animation Graph"); }
+	virtual FLinearColor GetAssetColor() const override { return FLinearColor(FColor(128,64,64)); }
+	virtual TSoftClassPtr<UObject> GetAssetClass() const override { return UAnimNextGraph_AnimationGraph::StaticClass(); }
+	virtual FText GetObjectDisplayNameText(UObject* Object) const override;
+};
+
+UCLASS()
+class UAssetDefinition_AnimNextEventGraph : public UAssetDefinitionDefault
+{
+	GENERATED_BODY()
+
+public:
+	// UAssetDefinition interface
+	virtual FText GetAssetDisplayName() const override { return LOCTEXT("AnimNextEventGraph", "Event Graph"); }
+	virtual FLinearColor GetAssetColor() const override { return FLinearColor(FColor(128,64,64)); }
+	virtual TSoftClassPtr<UObject> GetAssetClass() const override { return UAnimNextGraph_EventGraph::StaticClass(); }
+	virtual FText GetObjectDisplayNameText(UObject* Object) const override;
+};
+
 #undef LOCTEXT_NAMESPACE
\ No newline at end of file
diff --git a/Engine/Plugins/Experimental/AnimNext/Source/AnimNextEditor/Private/Graph/AnimNextGraphDocumentSummoner.cpp b/Engine/Plugins/Experimental/AnimNext/Source/AnimNextEditor/Private/Graph/AnimNextGraphDocumentSummoner.cpp
deleted file mode 100644
diff --git a/Engine/Plugins/Experimental/AnimNext/Source/AnimNextEditor/Private/Graph/AnimNextGraphDocumentSummoner.h b/Engine/Plugins/Experimental/AnimNext/Source/AnimNextEditor/Private/Graph/AnimNextGraphDocumentSummoner.h
deleted file mode 100644
diff --git a/Engine/Plugins/Experimental/AnimNext/Source/AnimNextEditor/Private/Graph/AnimationGraphDocumentSummoner.cpp b/Engine/Plugins/Experimental/AnimNext/Source/AnimNextEditor/Private/Graph/AnimationGraphDocumentSummoner.cpp
new file mode 100644
--- /dev/null
+++ b/Engine/Plugins/Experimental/AnimNext/Source/AnimNextEditor/Private/Graph/AnimationGraphDocumentSummoner.cpp
@@ -0,0 +1,34 @@
+﻿// Copyright Epic Games, Inc. All Rights Reserved.
+
+#include "AnimationGraphDocumentSummoner.h"
+#include "Common/SActionMenu.h"
+#include "Graph/AnimNextGraph_EdGraph.h"
+#include "Workspace/AnimNextWorkspaceEditor.h"
+
+namespace UE::AnimNext::Editor
+{
+
+FAnimationGraphDocumentSummoner::FAnimationGraphDocumentSummoner(FName InIdentifier, TSharedPtr<FWorkspaceEditor> InHostingApp)
+	: FGraphDocumentSummoner(InIdentifier, InHostingApp)
+{
+}
+
+FActionMenuContent FAnimationGraphDocumentSummoner::OnCreateGraphActionMenu(UEdGraph* InGraph, const FVector2D& InNodePosition, const TArray<UEdGraphPin*>& InDraggedPins, bool bAutoExpand, SGraphEditor::FActionMenuClosed InOnMenuClosed) const
+{
+	TSharedRef<SActionMenu> ActionMenu = SNew(SActionMenu, InGraph)
+		.AutoExpandActionMenu(bAutoExpand)
+		.NewNodePosition(InNodePosition)
+		.DraggedFromPins(InDraggedPins)
+		.OnClosedCallback(InOnMenuClosed);
+
+	TSharedPtr<SWidget> FilterTextBox = StaticCastSharedRef<SWidget>(ActionMenu->GetFilterTextBox());
+	return FActionMenuContent(StaticCastSharedRef<SWidget>(ActionMenu), FilterTextBox);
+}
+
+bool FAnimationGraphDocumentSummoner::IsPayloadSupported(TSharedRef<FTabPayload> Payload) const
+{
+	UObject* Object = Payload->IsValid() ? FTabPayload_UObject::CastChecked<UObject>(Payload) : nullptr;
+	return Object && Object->IsA<UAnimNextGraph_EdGraph>();
+}
+
+}
diff --git a/Engine/Plugins/Experimental/AnimNext/Source/AnimNextEditor/Private/Graph/AnimationGraphDocumentSummoner.h b/Engine/Plugins/Experimental/AnimNext/Source/AnimNextEditor/Private/Graph/AnimationGraphDocumentSummoner.h
new file mode 100644
--- /dev/null
+++ b/Engine/Plugins/Experimental/AnimNext/Source/AnimNextEditor/Private/Graph/AnimationGraphDocumentSummoner.h
@@ -0,0 +1,22 @@
+﻿// Copyright Epic Games, Inc. All Rights Reserved.
+
+#pragma once
+
+#include "CoreMinimal.h"
+#include "Workspace/GraphDocumentSummoner.h"
+
+namespace UE::AnimNext::Editor
+{
+
+class FAnimationGraphDocumentSummoner : public FGraphDocumentSummoner
+{
+public:
+	explicit FAnimationGraphDocumentSummoner(FName InIdentifier, TSharedPtr<FWorkspaceEditor> InHostingApp);
+
+private:
+	// FGraphDocumentSummoner interface
+	virtual FActionMenuContent OnCreateGraphActionMenu(UEdGraph* InGraph, const FVector2D& InNodePosition, const TArray<UEdGraphPin*>& InDraggedPins, bool bAutoExpand, SGraphEditor::FActionMenuClosed InOnMenuClosed) const override;
+	virtual bool IsPayloadSupported(TSharedRef<FTabPayload> Payload) const override;
+};
+
+}
diff --git a/Engine/Plugins/Experimental/AnimNext/Source/AnimNextEditor/Private/Graph/GraphFactory.cpp b/Engine/Plugins/Experimental/AnimNext/Source/AnimNextEditor/Private/Graph/GraphFactory.cpp
--- a/Engine/Plugins/Experimental/AnimNext/Source/AnimNextEditor/Private/Graph/GraphFactory.cpp
+++ b/Engine/Plugins/Experimental/AnimNext/Source/AnimNextEditor/Private/Graph/GraphFactory.cpp
@@ -36,10 +36,6 @@ UObject* UAnimNextGraphFactory::FactoryCreateNew(UClass* Class, UObject* InParen
 	NewGraph->EditorData = EditorData;
 	EditorData->Initialize(/*bRecompileVM*/false);
 
-	// Add root graph
-	EditorData->AddGraph(FRigUnit_AnimNextGraphRoot::DefaultEntryPoint, false);
-	check(EditorData->Entries.Num() > 0);
-
 	// Compile the initial skeleton
 	UE::AnimNext::UncookedOnly::FUtils::Compile(NewGraph);
 	check(!EditorData->bErrorsDuringCompilation);
diff --git a/Engine/Plugins/Experimental/AnimNext/Source/AnimNextEditor/Private/Param/AnimNextParameterBlockEditor.cpp b/Engine/Plugins/Experimental/AnimNext/Source/AnimNextEditor/Private/Param/AnimNextParameterBlockEditor.cpp
deleted file mode 100644
diff --git a/Engine/Plugins/Experimental/AnimNext/Source/AnimNextEditor/Private/Param/AnimNextParameterBlockEditor.h b/Engine/Plugins/Experimental/AnimNext/Source/AnimNextEditor/Private/Param/AnimNextParameterBlockEditor.h
deleted file mode 100644
diff --git a/Engine/Plugins/Experimental/AnimNext/Source/AnimNextEditor/Private/Param/EventGraphDocumentSummoner.cpp b/Engine/Plugins/Experimental/AnimNext/Source/AnimNextEditor/Private/Param/EventGraphDocumentSummoner.cpp
new file mode 100644
--- /dev/null
+++ b/Engine/Plugins/Experimental/AnimNext/Source/AnimNextEditor/Private/Param/EventGraphDocumentSummoner.cpp
@@ -0,0 +1,36 @@
+﻿// Copyright Epic Games, Inc. All Rights Reserved.
+
+#include "EventGraphDocumentSummoner.h"
+#include "Common/SActionMenu.h"
+#include "Graph/AnimNextGraph.h"
+#include "Graph/AnimNextGraph_EdGraph.h"
+
+#include "Workspace/AnimNextWorkspaceEditor.h"
+
+namespace UE::AnimNext::Editor
+{
+
+FEventGraphDocumentSummoner::FEventGraphDocumentSummoner(FName InIdentifier, TSharedPtr<FWorkspaceEditor> InHostingApp)
+	: FGraphDocumentSummoner(InIdentifier, InHostingApp)
+{
+}
+
+FActionMenuContent FEventGraphDocumentSummoner::OnCreateGraphActionMenu(UEdGraph* InGraph, const FVector2D& InNodePosition, const TArray<UEdGraphPin*>& InDraggedPins, bool bAutoExpand, SGraphEditor::FActionMenuClosed InOnMenuClosed) const
+{
+	TSharedRef<SActionMenu> ActionMenu = SNew(SActionMenu, InGraph)
+		.AutoExpandActionMenu(bAutoExpand)
+		.NewNodePosition(InNodePosition)
+		.DraggedFromPins(InDraggedPins)
+		.OnClosedCallback(InOnMenuClosed);
+
+	TSharedPtr<SWidget> FilterTextBox = StaticCastSharedRef<SWidget>(ActionMenu->GetFilterTextBox());
+	return FActionMenuContent(StaticCastSharedRef<SWidget>(ActionMenu), FilterTextBox);
+}
+
+bool FEventGraphDocumentSummoner::IsPayloadSupported(TSharedRef<FTabPayload> Payload) const
+{
+	UObject* Object = Payload->IsValid() ? FTabPayload_UObject::CastChecked<UObject>(Payload) : nullptr;
+	return Object && Object->IsA<UAnimNextGraph_EdGraph>();
+}
+
+}
diff --git a/Engine/Plugins/Experimental/AnimNext/Source/AnimNextEditor/Private/Param/EventGraphDocumentSummoner.h b/Engine/Plugins/Experimental/AnimNext/Source/AnimNextEditor/Private/Param/EventGraphDocumentSummoner.h
new file mode 100644
--- /dev/null
+++ b/Engine/Plugins/Experimental/AnimNext/Source/AnimNextEditor/Private/Param/EventGraphDocumentSummoner.h
@@ -0,0 +1,22 @@
+﻿// Copyright Epic Games, Inc. All Rights Reserved.
+
+#pragma once
+
+#include "CoreMinimal.h"
+#include "Workspace/GraphDocumentSummoner.h"
+
+namespace UE::AnimNext::Editor
+{
+
+class FEventGraphDocumentSummoner : public FGraphDocumentSummoner
+{
+public:
+	explicit FEventGraphDocumentSummoner(FName InIdentifier, TSharedPtr<FWorkspaceEditor> InHostingApp);
+
+private:
+	// FGraphDocumentSummoner interface
+	virtual FActionMenuContent OnCreateGraphActionMenu(UEdGraph* InGraph, const FVector2D& InNodePosition, const TArray<UEdGraphPin*>& InDraggedPins, bool bAutoExpand, SGraphEditor::FActionMenuClosed InOnMenuClosed) const override;
+	virtual bool IsPayloadSupported(TSharedRef<FTabPayload> Payload) const override;
+};
+
+}
diff --git a/Engine/Plugins/Experimental/AnimNext/Source/AnimNextEditor/Private/Param/ParamNamePropertyCustomization.cpp b/Engine/Plugins/Experimental/AnimNext/Source/AnimNextEditor/Private/Param/ParamNamePropertyCustomization.cpp
--- a/Engine/Plugins/Experimental/AnimNext/Source/AnimNextEditor/Private/Param/ParamNamePropertyCustomization.cpp
+++ b/Engine/Plugins/Experimental/AnimNext/Source/AnimNextEditor/Private/Param/ParamNamePropertyCustomization.cpp
@@ -53,7 +53,7 @@ void FParamNamePropertyTypeCustomization::CustomizeHeader(TSharedRef<IPropertyHa
 	FAnimNextParamType FilterType = FAnimNextParamType::FromString(ParamTypeString);
 	
 	FParameterPickerArgs PickerArgs;
-	PickerArgs.bShowBlocks = false;
+	PickerArgs.bShowSourceGraph = false;
 	PickerArgs.bMultiSelect = false;
 	PickerArgs.OnFilterParameterType = FOnFilterParameterType::CreateLambda([FilterType](const FAnimNextParamType& InParameterType)
 	{
diff --git a/Engine/Plugins/Experimental/AnimNext/Source/AnimNextEditor/Private/Param/ParamPropertyCustomization.cpp b/Engine/Plugins/Experimental/AnimNext/Source/AnimNextEditor/Private/Param/ParamPropertyCustomization.cpp
--- a/Engine/Plugins/Experimental/AnimNext/Source/AnimNextEditor/Private/Param/ParamPropertyCustomization.cpp
+++ b/Engine/Plugins/Experimental/AnimNext/Source/AnimNextEditor/Private/Param/ParamPropertyCustomization.cpp
@@ -43,7 +43,7 @@ void FParamPropertyTypeCustomization::CustomizeHeader(TSharedRef<IPropertyHandle
 	FAnimNextParamType FilterType = FAnimNextParamType::FromString(ParamTypeString);
 
 	FParameterPickerArgs PickerArgs;
-	PickerArgs.bShowBlocks = false;
+	PickerArgs.bShowSourceGraph = false;
 	PickerArgs.bMultiSelect = false;
 	PickerArgs.OnFilterParameterType = FOnFilterParameterType::CreateLambda([FilterType](const FAnimNextParamType& InParameterType)
 	{
diff --git a/Engine/Plugins/Experimental/AnimNext/Source/AnimNextEditor/Private/Param/ParameterBlockEditorMode.cpp b/Engine/Plugins/Experimental/AnimNext/Source/AnimNextEditor/Private/Param/ParameterBlockEditorMode.cpp
deleted file mode 100644
diff --git a/Engine/Plugins/Experimental/AnimNext/Source/AnimNextEditor/Private/Param/ParameterBlockEditorMode.h b/Engine/Plugins/Experimental/AnimNext/Source/AnimNextEditor/Private/Param/ParameterBlockEditorMode.h
deleted file mode 100644
diff --git a/Engine/Plugins/Experimental/AnimNext/Source/AnimNextEditor/Private/Param/ParameterBlockFactory.cpp b/Engine/Plugins/Experimental/AnimNext/Source/AnimNextEditor/Private/Param/ParameterBlockFactory.cpp
deleted file mode 100644
diff --git a/Engine/Plugins/Experimental/AnimNext/Source/AnimNextEditor/Private/Param/ParameterBlockGraphDocumentSummoner.cpp b/Engine/Plugins/Experimental/AnimNext/Source/AnimNextEditor/Private/Param/ParameterBlockGraphDocumentSummoner.cpp
deleted file mode 100644
diff --git a/Engine/Plugins/Experimental/AnimNext/Source/AnimNextEditor/Private/Param/ParameterBlockGraphDocumentSummoner.h b/Engine/Plugins/Experimental/AnimNext/Source/AnimNextEditor/Private/Param/ParameterBlockGraphDocumentSummoner.h
deleted file mode 100644
diff --git a/Engine/Plugins/Experimental/AnimNext/Source/AnimNextEditor/Private/Param/ParameterBlockParameterCustomization.cpp b/Engine/Plugins/Experimental/AnimNext/Source/AnimNextEditor/Private/Param/ParameterBlockParameterCustomization.cpp
deleted file mode 100644
diff --git a/Engine/Plugins/Experimental/AnimNext/Source/AnimNextEditor/Private/Param/ParameterBlockParameterCustomization.h b/Engine/Plugins/Experimental/AnimNext/Source/AnimNextEditor/Private/Param/ParameterBlockParameterCustomization.h
deleted file mode 100644
diff --git a/Engine/Plugins/Experimental/AnimNext/Source/AnimNextEditor/Private/Param/ParameterCustomization.cpp b/Engine/Plugins/Experimental/AnimNext/Source/AnimNextEditor/Private/Param/ParameterCustomization.cpp
new file mode 100644
--- /dev/null
+++ b/Engine/Plugins/Experimental/AnimNext/Source/AnimNextEditor/Private/Param/ParameterCustomization.cpp
@@ -0,0 +1,97 @@
+// Copyright Epic Games, Inc. All Rights Reserved.
+
+#include "Param/ParameterCustomization.h"
+
+#include "DetailLayoutBuilder.h"
+#include "DetailWidgetRow.h"
+#include "EdGraphSchema_K2.h"
+#include "InstancedPropertyBagStructureDataProvider.h"
+#include "UncookedOnlyUtils.h"
+#include "PropertyHandle.h"
+#include "Graph/AnimNextGraph.h"
+#include "Graph/AnimNextGraph_Parameter.h"
+#include "Graph/AnimNextGraph_EditorData.h"
+
+#define LOCTEXT_NAMESPACE "ParamTypePropertyCustomization"
+
+namespace UE::AnimNext::Editor
+{
+
+void FParameterCustomization::CustomizeDetails(IDetailLayoutBuilder& DetailBuilder)
+{
+	TArray<TWeakObjectPtr<UObject>> Objects;
+	DetailBuilder.GetObjectsBeingCustomized(Objects);
+
+	if (Objects.IsEmpty() || Objects.Num() > 1)
+	{
+		return;
+	}
+
+	if (UAnimNextGraph_Parameter* Parameter = Cast<UAnimNextGraph_Parameter>(Objects[0].Get()))
+	{
+		IDetailCategoryBuilder& ParameterCategory = DetailBuilder.EditCategory(TEXT("Parameter"), FText::GetEmpty(), ECategoryPriority::Important);
+
+		IDetailCategoryBuilder& DefaultValueCategory = DetailBuilder.EditCategory(TEXT("DefaultValue"), FText::GetEmpty(), ECategoryPriority::Default);
+
+		TSharedRef< SWidget > ColumnWidget = SNullWidget::NullWidget;
+
+		if (UAnimNextGraph_EditorData* EditorData = Cast<UAnimNextGraph_EditorData>(Parameter->GetOuter()))
+		{
+			const FName ParameterName = Parameter->GetEntryName();
+			if (UAnimNextRigVMAssetEntry* AssetEntry = EditorData->FindEntry(ParameterName)) 
+			{
+				if (UAnimNextGraph* ReferencedGraph = UE::AnimNext::UncookedOnly::FUtils::GetGraph(EditorData))
+				{
+					FAddPropertyParams AddPropertyParams;
+					TArray<IDetailPropertyRow*> DetailPropertyRows;
+
+					if (ReferencedGraph->PropertyBag.FindPropertyDescByName(ParameterName))
+					{
+						IDetailPropertyRow* DetailPropertyRow = DefaultValueCategory.AddExternalStructureProperty(MakeShared<FInstancePropertyBagStructureDataProvider>(ReferencedGraph->PropertyBag), ParameterName, EPropertyLocation::Default, AddPropertyParams);
+						if (TSharedPtr<IPropertyHandle> Handle = DetailPropertyRow->GetPropertyHandle(); Handle.IsValid())
+						{
+							Handle->SetOnChildPropertyValuePreChange(FSimpleDelegate::CreateLambda([this, ReferencedGraph]()
+							{
+								ReferencedGraph->Modify(); // needed to enable the transaction when we modify the PropertyBag
+							}));
+							Handle->SetOnPropertyValueChanged(FSimpleDelegate::CreateLambda([this, ReferencedGraph, ParameterName]()
+							{
+								if (UAnimNextGraph_EditorData* EditorData = Cast<UAnimNextGraph_EditorData>(ReferencedGraph->EditorData))
+								{
+									if (UAnimNextRigVMAssetEntry* AssetEntry = EditorData->FindEntry(ParameterName))
+									{
+										AssetEntry->MarkPackageDirty();
+									}
+								}
+							}));
+						}
+					}
+				}
+			}
+		}
+	}
+}
+
+void FParameterCustomization::CustomizeDetails(const TSharedPtr<IDetailLayoutBuilder>& DetailBuilder) 
+{
+	CustomizeDetails(*DetailBuilder);
+}
+
+FText FParameterCustomization::GetName() const
+{
+	return FText();
+}
+
+void FParameterCustomization::SetName(const FText& InNewText, ETextCommit::Type InCommitType)
+{
+}
+
+bool FParameterCustomization::OnVerifyNameChanged(const FText& InText, FText& OutErrorMessage)
+{
+	return true;
+}
+
+
+}
+
+#undef LOCTEXT_NAMESPACE
diff --git a/Engine/Plugins/Experimental/AnimNext/Source/AnimNextEditor/Private/Param/ParameterCustomization.h b/Engine/Plugins/Experimental/AnimNext/Source/AnimNextEditor/Private/Param/ParameterCustomization.h
new file mode 100644
--- /dev/null
+++ b/Engine/Plugins/Experimental/AnimNext/Source/AnimNextEditor/Private/Param/ParameterCustomization.h
@@ -0,0 +1,30 @@
+// Copyright Epic Games, Inc. All Rights Reserved.
+
+#pragma once
+
+#include "IDetailCustomization.h"
+#include "IPropertyTypeCustomization.h"
+#include "PropertyEditorModule.h"
+
+namespace UE::AnimNext::Editor
+{
+
+class FParameterCustomization : public IDetailCustomization
+{
+private:
+
+	/** Called when details should be customized */
+	virtual void CustomizeDetails(IDetailLayoutBuilder& DetailBuilder) override;
+
+	/**
+	 * Called when details should be customized, this version allows the customization to store a weak ptr to the layout builder.
+	 * This allows property changes to trigger a force refresh of the detail panel.
+	 */
+	virtual void CustomizeDetails(const TSharedPtr<IDetailLayoutBuilder>& DetailBuilder) override;
+
+	FText GetName() const;
+	void SetName(const FText& InNewText, ETextCommit::Type InCommitType);
+	bool OnVerifyNameChanged(const FText& InText, FText& OutErrorMessage);
+};
+
+}
diff --git a/Engine/Plugins/Experimental/AnimNext/Source/AnimNextEditor/Private/Param/ParametersAssetDefinitions.cpp b/Engine/Plugins/Experimental/AnimNext/Source/AnimNextEditor/Private/Param/ParametersAssetDefinitions.cpp
deleted file mode 100644
diff --git a/Engine/Plugins/Experimental/AnimNext/Source/AnimNextEditor/Private/Param/ParametersAssetDefinitions.h b/Engine/Plugins/Experimental/AnimNext/Source/AnimNextEditor/Private/Param/ParametersAssetDefinitions.h
deleted file mode 100644
diff --git a/Engine/Plugins/Experimental/AnimNext/Source/AnimNextEditor/Private/Param/SAddParametersDialog.cpp b/Engine/Plugins/Experimental/AnimNext/Source/AnimNextEditor/Private/Param/SAddParametersDialog.cpp
--- a/Engine/Plugins/Experimental/AnimNext/Source/AnimNextEditor/Private/Param/SAddParametersDialog.cpp
+++ b/Engine/Plugins/Experimental/AnimNext/Source/AnimNextEditor/Private/Param/SAddParametersDialog.cpp
@@ -16,7 +16,6 @@
 #include "Framework/Application/SlateApplication.h"
 #include "Framework/Docking/TabManager.h"
 #include "Widgets/Layout/SUniformGridPanel.h"
-#include "Widgets/Layout/SBorder.h"
 #include "Widgets/Input/SButton.h"
 #include "Widgets/Layout/SBox.h"
 #include "Widgets/Text/SInlineEditableTextBlock.h"
@@ -56,7 +55,6 @@ void SAddParametersDialog::Construct(const FArguments& InArgs)
 {
 	using namespace AddParametersDialog;
 
-	TargetBlock = InArgs._Block;
 	OnFilterParameterType = InArgs._OnFilterParameterType;
 
 	SWindow::Construct(SWindow::FArguments()
@@ -147,7 +145,7 @@ void SAddParametersDialog::Construct(const FArguments& InArgs)
 									return FText::Format(LOCTEXT("AddParametersButtonTooltip_InvalidEntry", "A parameter to add is not valid: {0}"), Reason);
 								}
 							}
-							return LOCTEXT("AddParametersButtonTooltip", "Add the selected parameters to the current parameter block");
+							return LOCTEXT("AddParametersButtonTooltip", "Add the selected parameters to the current graph");
 						})
 						.OnClicked_Lambda([this]()
 						{
diff --git a/Engine/Plugins/Experimental/AnimNext/Source/AnimNextEditor/Private/Param/SAddParametersDialog.h b/Engine/Plugins/Experimental/AnimNext/Source/AnimNextEditor/Private/Param/SAddParametersDialog.h
--- a/Engine/Plugins/Experimental/AnimNext/Source/AnimNextEditor/Private/Param/SAddParametersDialog.h
+++ b/Engine/Plugins/Experimental/AnimNext/Source/AnimNextEditor/Private/Param/SAddParametersDialog.h
@@ -10,7 +10,7 @@
 
 class SWrapBox;
 class UAnimNextParameterLibrary;
-class UAnimNextParameterBlock_EditorData;
+class UAnimNextGraph_EditorData;
 
 namespace UE::AnimNext::Editor
 {
@@ -21,12 +21,9 @@ class SAddParametersDialog : public SWindow
 {
 public:
 	SLATE_BEGIN_ARGS(SAddParametersDialog)
-		: _Block(nullptr),
-		_AllowMultiple(true)
+		: _AllowMultiple(true)
 	{}
 
-	SLATE_ARGUMENT(UAnimNextParameterBlock_EditorData*, Block)
-
 	/** Whether we allow multiple parameters to be added or just one at a time */
 	SLATE_ARGUMENT(bool, AllowMultiple)
 
@@ -60,8 +57,6 @@ private:
 
 	TArray<TSharedRef<FParameterToAdd>> Entries;
 
-	UAnimNextParameterBlock_EditorData* TargetBlock = nullptr;
-	
 	FOnFilterParameterType OnFilterParameterType;
 
 	bool bCancelPressed = false;
diff --git a/Engine/Plugins/Experimental/AnimNext/Source/AnimNextEditor/Private/Param/SGraphPinParamName.cpp b/Engine/Plugins/Experimental/AnimNext/Source/AnimNextEditor/Private/Param/SGraphPinParamName.cpp
--- a/Engine/Plugins/Experimental/AnimNext/Source/AnimNextEditor/Private/Param/SGraphPinParamName.cpp
+++ b/Engine/Plugins/Experimental/AnimNext/Source/AnimNextEditor/Private/Param/SGraphPinParamName.cpp
@@ -31,12 +31,12 @@ void SGraphPinParamName::Construct(const FArguments& InArgs, UEdGraphPin* InPin)
 TSharedRef<SWidget> SGraphPinParamName::GetDefaultValueWidget()
 {
 	FParameterPickerArgs Args;
-	Args.bShowBlocks = false;
+	Args.bShowSourceGraph = false;
 	Args.bMultiSelect = false;
 
 	CachedType = UncookedOnly::FUtils::GetParameterTypeFromName(FName(*GraphPinObj->DefaultValue));
 
-	// Check whether this is a Set/Get parameter from block node, and if so only show bound parameters
+	// Check whether this is a Set/Get parameter node, and if so only show bound parameters
 	if (ModelPin)
 	{
 		if (const URigVMTemplateNode* TemplateNode = Cast<URigVMTemplateNode>(ModelPin->GetOuter()))
diff --git a/Engine/Plugins/Experimental/AnimNext/Source/AnimNextEditor/Private/Param/SLinkParametersDialog.cpp b/Engine/Plugins/Experimental/AnimNext/Source/AnimNextEditor/Private/Param/SLinkParametersDialog.cpp
--- a/Engine/Plugins/Experimental/AnimNext/Source/AnimNextEditor/Private/Param/SLinkParametersDialog.cpp
+++ b/Engine/Plugins/Experimental/AnimNext/Source/AnimNextEditor/Private/Param/SLinkParametersDialog.cpp
@@ -95,7 +95,7 @@ void SLinkParametersDialog::Construct(const FArguments& InArgs)
 						.HAlign(HAlign_Center)
 						.ButtonStyle(&FAppStyle::Get().GetWidgetStyle<FButtonStyle>("PrimaryButton"))
 						.Text(LOCTEXT("LinkButton", "Link"))
-						.ToolTipText(LOCTEXT("LinkButtonTooltip", "Link the selected parameters to the current parameter block"))
+						.ToolTipText(LOCTEXT("LinkButtonTooltip", "Link the selected parameters to the current graph"))
 						.OnClicked_Lambda([this]()
 						{
 							RequestDestroyWindow();
diff --git a/Engine/Plugins/Experimental/AnimNext/Source/AnimNextEditor/Private/Param/SParameterPicker.cpp b/Engine/Plugins/Experimental/AnimNext/Source/AnimNextEditor/Private/Param/SParameterPicker.cpp
--- a/Engine/Plugins/Experimental/AnimNext/Source/AnimNextEditor/Private/Param/SParameterPicker.cpp
+++ b/Engine/Plugins/Experimental/AnimNext/Source/AnimNextEditor/Private/Param/SParameterPicker.cpp
@@ -2,12 +2,10 @@
 
 #include "SParameterPicker.h"
 
-#include "Param/AnimNextParameterBlock.h"
 #include "UncookedOnlyUtils.h"
 #include "AssetRegistry/AssetRegistryModule.h"
 #include "Kismet2/BlueprintEditorUtils.h"
 #include "Param/ParamType.h"
-#include "Param/AnimNextParameterBlock_EditorData.h"
 #include "DetailLayoutBuilder.h"
 #include "EditorUtils.h"
 #include "SAddParametersDialog.h"
@@ -20,6 +18,7 @@
 #include "Framework/Application/SlateApplication.h"
 #include "ScopedTransaction.h"
 #include "SSimpleButton.h"
+#include "Graph/AnimNextGraph.h"
 #include "Param/ExternalParameterRegistry.h"
 
 #define LOCTEXT_NAMESPACE "SParameterPicker"
@@ -30,7 +29,7 @@ namespace UE::AnimNext::Editor
 namespace ParameterPicker
 {
 static FName Column_Parameter(TEXT("Parameter"));
-static FName Column_Block(TEXT("Block"));
+static FName Column_Graph(TEXT("Graph"));
 static FName Column_Type(TEXT("Type"));
 static FName Column_New(TEXT("New"));
 }
@@ -232,12 +231,12 @@ void SParameterPicker::Construct(const FArguments& InArgs)
 	];
 
 	
-	if (Args.bShowBlocks)
+	if (Args.bShowSourceGraph)
 	{
 		HeaderRow->AddColumn(
-			SHeaderRow::Column(Column_Block)
-			.DefaultLabel(LOCTEXT("BlockColumnHeader", "Block"))
-			.ToolTipText(LOCTEXT("BlockColumnHeaderTooltip", "The parameter block that has a binding to the parameter"))
+			SHeaderRow::Column(Column_Graph)
+			.DefaultLabel(LOCTEXT("GraphColumnHeader", "Graph"))
+			.ToolTipText(LOCTEXT("GraphColumnHeaderTooltip", "The graph that has a binding to the parameter"))
 			.FillWidth(0.33f));
 	}
 
@@ -305,23 +304,23 @@ void SParameterPicker::RefreshEntries()
 		Entries.Add(MakeShared<FParameterPickerEntry>(FParameterBindingReference(NAME_None, FAnimNextParamType())));
 	}
 	
-	// Find all blocks and their bound parameters
+	// Find all graphs and their bound parameters
 	if(Args.bShowBoundParameters)
 	{
-		ARFilter.ClassPaths = { UAnimNextParameterBlock::StaticClass()->GetClassPathName() };
+		ARFilter.ClassPaths = { UAnimNextGraph::StaticClass()->GetClassPathName() };
 		
-		TArray<FAssetData> BlockAssets;
-		AssetRegistry.GetAssets(ARFilter, BlockAssets);
+		TArray<FAssetData> GraphAssets;
+		AssetRegistry.GetAssets(ARFilter, GraphAssets);
 
-		for(const FAssetData& BlockAsset : BlockAssets)
+		for(const FAssetData& GraphAsset : GraphAssets)
 		{
 			FAnimNextParameterProviderAssetRegistryExports Exports;
-			if(UncookedOnly::FUtils::GetExportedParametersForAsset(BlockAsset, Exports))
+			if(UncookedOnly::FUtils::GetExportedParametersForAsset(GraphAsset, Exports))
 			{
 				for(const FAnimNextParameterAssetRegistryExportEntry& Export : Exports.Parameters)
 				{
-					BoundParameters.Add({ Export.Name, BlockAsset });
-					FParameterBindingReference NewReference(Export.Name, Export.Type, BlockAsset);
+					BoundParameters.Add({ Export.Name, GraphAsset });
+					FParameterBindingReference NewReference(Export.Name, Export.Type, GraphAsset);
 					if(!Args.OnFilterParameter.IsBound() || Args.OnFilterParameter.Execute(NewReference) == EFilterParameterResult::Include)
 					{
 						FAnimNextParamType ParamType = UE::AnimNext::UncookedOnly::FUtils::GetParameterTypeFromName(Export.Name);
@@ -601,9 +600,9 @@ class SParameterPickerRow : public SMultiColumnTableRow<TSharedRef<FParameterPic
 					})
 				];
 		}
-		else if(InColumnName == Column_Block)
+		else if(InColumnName == Column_Graph)
 		{
-			if(Entry->Binding.Block.IsValid())
+			if(Entry->Binding.Graph.IsValid())
 			{
 				return
 					SNew(SBox)
@@ -611,8 +610,8 @@ class SParameterPickerRow : public SMultiColumnTableRow<TSharedRef<FParameterPic
 					[
 						SNew(STextBlock)
 						.Font(IDetailLayoutBuilder::GetDetailFont())
-						.Text(FText::FromName(Entry->Binding.Block.AssetName))
-						.ToolTipText(FText::FromName(Entry->Binding.Block.PackageName))
+						.Text(FText::FromName(Entry->Binding.Graph.AssetName))
+						.ToolTipText(FText::FromName(Entry->Binding.Graph.PackageName))
 					];
 			}
 		}
diff --git a/Engine/Plugins/Experimental/AnimNext/Source/AnimNextEditor/Private/Scheduler/ScheduleAssetDefinition.h b/Engine/Plugins/Experimental/AnimNext/Source/AnimNextEditor/Private/Scheduler/ScheduleAssetDefinition.h
--- a/Engine/Plugins/Experimental/AnimNext/Source/AnimNextEditor/Private/Scheduler/ScheduleAssetDefinition.h
+++ b/Engine/Plugins/Experimental/AnimNext/Source/AnimNextEditor/Private/Scheduler/ScheduleAssetDefinition.h
@@ -16,7 +16,7 @@ class UAssetDefinition_AnimNextSchedule : public UAssetDefinitionDefault
 public:
 	// UAssetDefinition interface
 	virtual FText GetAssetDisplayName() const override { return LOCTEXT("AnimNextSchedule", "AnimNext Schedule"); }
-	virtual FLinearColor GetAssetColor() const override { return FLinearColor(FColor(128, 64, 64)); }
+	virtual FLinearColor GetAssetColor() const override { return FLinearColor(FColor(150,96,48)); }
 	virtual TSoftClassPtr<UObject> GetAssetClass() const override { return UAnimNextSchedule::StaticClass(); }
 	virtual EAssetCommandResult OpenAssets(const FAssetOpenArgs& OpenArgs) const override;
 	virtual TConstArrayView<FAssetCategoryPath> GetAssetCategories() const override
diff --git a/Engine/Plugins/Experimental/AnimNext/Source/AnimNextEditor/Private/Workspace/AnimNextWorkspace.cpp b/Engine/Plugins/Experimental/AnimNext/Source/AnimNextEditor/Private/Workspace/AnimNextWorkspace.cpp
--- a/Engine/Plugins/Experimental/AnimNext/Source/AnimNextEditor/Private/Workspace/AnimNextWorkspace.cpp
+++ b/Engine/Plugins/Experimental/AnimNext/Source/AnimNextEditor/Private/Workspace/AnimNextWorkspace.cpp
@@ -3,7 +3,6 @@
 #include "AnimNextWorkspace.h"
 
 #include "Graph/AnimNextGraph.h"
-#include "Param/AnimNextParameterBlock.h"
 #include "Scheduler/AnimNextSchedule.h"
 #include "UObject/AssetRegistryTagsContext.h"
 
@@ -178,7 +177,6 @@ bool UAnimNextWorkspace::RemoveAssets(TConstArrayView<FAssetData> InAssets, bool
 const TArray<FTopLevelAssetPath>& UAnimNextWorkspace::GetSupportedAssetClassPaths()
 {
 	static TArray<FTopLevelAssetPath> SupportedAssets = {
-		UAnimNextParameterBlock::StaticClass()->GetClassPathName(),
 		UAnimNextGraph::StaticClass()->GetClassPathName(),
 		UAnimNextSchedule::StaticClass()->GetClassPathName(),
 	};
diff --git a/Engine/Plugins/Experimental/AnimNext/Source/AnimNextEditor/Private/Workspace/AnimNextWorkspaceAssetDefinition.h b/Engine/Plugins/Experimental/AnimNext/Source/AnimNextEditor/Private/Workspace/AnimNextWorkspaceAssetDefinition.h
--- a/Engine/Plugins/Experimental/AnimNext/Source/AnimNextEditor/Private/Workspace/AnimNextWorkspaceAssetDefinition.h
+++ b/Engine/Plugins/Experimental/AnimNext/Source/AnimNextEditor/Private/Workspace/AnimNextWorkspaceAssetDefinition.h
@@ -16,7 +16,7 @@ class UAssetDefinition_AnimNextWorkspace : public UAssetDefinitionDefault
 public:
 	// UAssetDefinition interface
 	virtual FText GetAssetDisplayName() const override { return LOCTEXT("AnimNextWorkspace", "AnimNext Workspace"); }
-	virtual FLinearColor GetAssetColor() const override { return FLinearColor(FColor(128,64,64)); }
+	virtual FLinearColor GetAssetColor() const override { return FLinearColor(FColor(170,96,48)); }
 	virtual TSoftClassPtr<UObject> GetAssetClass() const override { return UAnimNextWorkspace::StaticClass(); }
 	virtual EAssetCommandResult OpenAssets(const FAssetOpenArgs& OpenArgs) const override;
 	virtual TConstArrayView<FAssetCategoryPath> GetAssetCategories() const override
diff --git a/Engine/Plugins/Experimental/AnimNext/Source/AnimNextEditor/Private/Workspace/AnimNextWorkspaceEditor.cpp b/Engine/Plugins/Experimental/AnimNext/Source/AnimNextEditor/Private/Workspace/AnimNextWorkspaceEditor.cpp
--- a/Engine/Plugins/Experimental/AnimNext/Source/AnimNextEditor/Private/Workspace/AnimNextWorkspaceEditor.cpp
+++ b/Engine/Plugins/Experimental/AnimNext/Source/AnimNextEditor/Private/Workspace/AnimNextWorkspaceEditor.cpp
@@ -11,25 +11,21 @@
 #include "EdGraphNode_Comment.h"
 #include "EditorUtils.h"
 #include "ExternalPackageHelper.h"
-#include "Framework/Commands/GenericCommands.h"
 #include "RigVMModel/RigVMController.h"
 #include "Kismet2/BlueprintEditorUtils.h"
-#include "ScopedTransaction.h"
 #include "Widgets/Docking/SDockTab.h"
 #include "WorkflowOrientedApp/WorkflowUObjectDocuments.h"
 #include "SWorkspacePicker.h"
-#include "Dialog/SCustomDialog.h"
-#include "Graph/AnimNextGraphDocumentSummoner.h"
-#include "Param/AnimNextParameterBlock.h"
-#include "Param/AnimNextParameterBlock_EditorData.h"
-#include "Param/ParameterBlockGraphDocumentSummoner.h"
+#include "Graph/AnimationGraphDocumentSummoner.h"
+#include "Param/EventGraphDocumentSummoner.h"
 #include "Scheduler/AnimNextSchedule.h"
 #include "Framework/Application/SlateApplication.h"
 #include "Graph/AnimNextGraph.h"
-#include "Graph/AnimNextGraph_EdGraphNode.h"
-#include "GraphEditAction.h"
 #include "IAnimNextRigVMGraphInterface.h"
 #include "UncookedOnlyUtils.h"
+#include "EdGraph/RigVMEdGraph.h"
+#include "RigVMModel/RigVMClient.h"
+#include "AnimNextRigVMAssetEditorData.h"
 
 #define LOCTEXT_NAMESPACE "AnimNextWorkspaceEditor"
 
@@ -47,8 +43,8 @@ namespace WorkspaceTabs
 	const FName WorkspaceView("WorkspaceView");
 	const FName LeftAssetDocument("LeftAssetDocument");
 	const FName MiddleAssetDocument("MiddleAssetDocument");
-	const FName AnimNextGraphDocument("AnimNextGraphDocument");
-	const FName ParameterBlockGraphDocument("ParameterBlockGraphDocument");
+	const FName AnimationGraphDocument("AnimationGraphDocument");
+	const FName EventGraphDocument("EventGraphDocument");
 }
 
 const FName WorkspaceAppIdentifier("AnimNextWorkspaceEditor");
@@ -70,19 +66,18 @@ void FWorkspaceEditor::InitEditor(const EToolkitMode::Type InMode, const TShared
 	DocumentManager = MakeShared<FDocumentTracker>(NAME_None);
 	DocumentManager->Initialize(SharedThis(this));
 
-	TSharedRef<FParameterBlockGraphDocumentSummoner> ParameterGraphDocumentSummoner = MakeShared<FParameterBlockGraphDocumentSummoner>(WorkspaceTabs::ParameterBlockGraphDocument, SharedThis(this));
+	TSharedRef<FEventGraphDocumentSummoner> ParameterGraphDocumentSummoner = MakeShared<FEventGraphDocumentSummoner>(WorkspaceTabs::EventGraphDocument, SharedThis(this));
 	ParameterGraphDocumentSummoner->OnSaveGraphState().BindSP(this, &FWorkspaceEditor::HandleSaveGraphState);
 	DocumentManager->RegisterDocumentFactory(ParameterGraphDocumentSummoner);
 
-	TSharedRef<FAnimNextGraphDocumentSummoner> GraphDocumentSummoner = MakeShared<FAnimNextGraphDocumentSummoner>(WorkspaceTabs::AnimNextGraphDocument, SharedThis(this));
-	GraphDocumentSummoner->OnSaveGraphState().BindSP(this, &FWorkspaceEditor::HandleSaveGraphState);
-	DocumentManager->RegisterDocumentFactory(GraphDocumentSummoner);
+	TSharedRef<FAnimationGraphDocumentSummoner> AnimationGraphDocumentSummoner = MakeShared<FAnimationGraphDocumentSummoner>(WorkspaceTabs::AnimationGraphDocument, SharedThis(this));
+	AnimationGraphDocumentSummoner->OnSaveGraphState().BindSP(this, &FWorkspaceEditor::HandleSaveGraphState);
+	DocumentManager->RegisterDocumentFactory(AnimationGraphDocumentSummoner);
 
 	TSharedRef<FAssetDocumentSummoner> LeftAssetDocumentSummoner = MakeShared<FAssetDocumentSummoner>(WorkspaceTabs::LeftAssetDocument, SharedThis(this));
 	LeftAssetDocumentSummoner->SetAllowedAssetClassPaths(
 	{
 		UAnimNextGraph::StaticClass()->GetClassPathName(),
-		UAnimNextParameterBlock::StaticClass()->GetClassPathName(),
 	});
 	LeftAssetDocumentSummoner->OnSaveDocumentState().BindSP(this, &FWorkspaceEditor::HandleSaveDocumentState);
 	DocumentManager->RegisterDocumentFactory(LeftAssetDocumentSummoner);
diff --git a/Engine/Plugins/Experimental/AnimNext/Source/AnimNextEditor/Private/Workspace/AnimNextWorkspaceEditor.h b/Engine/Plugins/Experimental/AnimNext/Source/AnimNextEditor/Private/Workspace/AnimNextWorkspaceEditor.h
--- a/Engine/Plugins/Experimental/AnimNext/Source/AnimNextEditor/Private/Workspace/AnimNextWorkspaceEditor.h
+++ b/Engine/Plugins/Experimental/AnimNext/Source/AnimNextEditor/Private/Workspace/AnimNextWorkspaceEditor.h
@@ -39,8 +39,8 @@ namespace WorkspaceTabs
 	extern const FName WorkspaceView;
 	extern const FName LeftAssetDocument;
 	extern const FName MiddleAssetDocument;
-	extern const FName AnimNextGraphDocument;
-	extern const FName ParameterBlockGraphDocument;
+	extern const FName AnimationGraphDocument;
+	extern const FName EventGraphDocument;
 }
 
 class FWorkspaceEditor : public FWorkflowCentricApplication
diff --git a/Engine/Plugins/Experimental/AnimNext/Source/AnimNextEditor/Private/Workspace/AnimNextWorkspaceEditorMode.cpp b/Engine/Plugins/Experimental/AnimNext/Source/AnimNextEditor/Private/Workspace/AnimNextWorkspaceEditorMode.cpp
--- a/Engine/Plugins/Experimental/AnimNext/Source/AnimNextEditor/Private/Workspace/AnimNextWorkspaceEditorMode.cpp
+++ b/Engine/Plugins/Experimental/AnimNext/Source/AnimNextEditor/Private/Workspace/AnimNextWorkspaceEditorMode.cpp
@@ -26,7 +26,7 @@ FWorkspaceEditorMode::FWorkspaceEditorMode(TSharedRef<FWorkspaceEditor> InHostin
 	TabFactories.RegisterFactory(MakeShared<FDetailsTabSummoner>(WorkspaceEditor, FOnDetailsViewCreated::CreateSP(&WorkspaceEditor.Get(), &FWorkspaceEditor::HandleDetailsViewCreated)));
 	TabFactories.RegisterFactory(MakeShared<FWorkspaceTabSummoner>(WorkspaceEditor));
 
-	TabLayout = FTabManager::NewLayout("Standalone_AnimNextWorkspaceEditor_Layout_v1.1")
+	TabLayout = FTabManager::NewLayout("Standalone_AnimNextWorkspaceEditor_Layout_v1.2")
 		->AddArea
 		(
 			FTabManager::NewPrimaryArea()
@@ -49,8 +49,8 @@ FWorkspaceEditorMode::FWorkspaceEditorMode(TSharedRef<FWorkspaceEditor> InHostin
 					FTabManager::NewStack()
 					->SetSizeCoefficient(0.5f)
 					->SetHideTabWell(false)
-					->AddTab(WorkspaceTabs::ParameterBlockGraphDocument, ETabState::ClosedTab)
-					->AddTab(WorkspaceTabs::AnimNextGraphDocument, ETabState::ClosedTab)
+					->AddTab(WorkspaceTabs::EventGraphDocument, ETabState::ClosedTab)
+					->AddTab(WorkspaceTabs::AnimationGraphDocument, ETabState::ClosedTab)
 					->AddTab(WorkspaceTabs::MiddleAssetDocument, ETabState::ClosedTab)
 				)
 				->Split
diff --git a/Engine/Plugins/Experimental/AnimNext/Source/AnimNextTestSuite/Private/AnimNextEditorTest.cpp b/Engine/Plugins/Experimental/AnimNext/Source/AnimNextTestSuite/Private/AnimNextEditorTest.cpp
--- a/Engine/Plugins/Experimental/AnimNext/Source/AnimNextTestSuite/Private/AnimNextEditorTest.cpp
+++ b/Engine/Plugins/Experimental/AnimNext/Source/AnimNextTestSuite/Private/AnimNextEditorTest.cpp
@@ -3,15 +3,12 @@
 #include "CoreMinimal.h"
 #include "AnimNextTest.h"
 #include "UncookedOnlyUtils.h"
-#include "Param/AnimNextParameterBlock_EditorData.h"
 #include "Misc/AutomationTest.h"
-#include "Param/ParameterBlockFactory.h"
-#include "Param/AnimNextParameterBlock.h"
-#include "Param/AnimNextParameterBlockParameter.h"
 #include "Animation/AnimSequence.h"
 #include "Graph/AnimNextGraph.h"
-#include "Graph/AnimNextGraphEntry.h"
+#include "Graph/AnimNextGraph_Parameter.h"
 #include "Graph/AnimNextGraph_EditorData.h"
+#include "Graph/AnimNextGraph_AnimationGraph.h"
 #include "Graph/GraphFactory.h"
 #if WITH_EDITOR
 #include "ScopedTransaction.h"
@@ -26,59 +23,59 @@
 namespace UE::AnimNext::Tests
 {
 
-IMPLEMENT_SIMPLE_AUTOMATION_TEST(FEditor_Parameters_ParameterBlock, "Animation.AnimNext.Editor.Parameters.ParameterBlock", EAutomationTestFlags::EditorContext | EAutomationTestFlags::EngineFilter)
+IMPLEMENT_SIMPLE_AUTOMATION_TEST(FEditor_Parameters, "Animation.AnimNext.Editor.Parameters", EAutomationTestFlags::EditorContext | EAutomationTestFlags::EngineFilter)
 
-bool FEditor_Parameters_ParameterBlock::RunTest(const FString& InParameters)
+bool FEditor_Parameters::RunTest(const FString& InParameters)
 {
 	using namespace UE::AnimNext;
 
-	const TStrongObjectPtr<UFactory> BlockFactory(NewObject<UAnimNextParameterBlockFactory>());
-	UAnimNextParameterBlock* Block = Cast<UAnimNextParameterBlock>(BlockFactory->FactoryCreateNew(UAnimNextParameterBlock::StaticClass(), GetTransientPackage(), TEXT("TestAnimNextParameterBlock"), RF_Transient, nullptr, nullptr, NAME_None));
-	if(Block == nullptr)
+	const TStrongObjectPtr<UFactory> GraphFactory(NewObject<UAnimNextGraphFactory>());
+	UAnimNextGraph* Graph = Cast<UAnimNextGraph>(GraphFactory->FactoryCreateNew(UAnimNextGraph::StaticClass(), GetTransientPackage(), TEXT("TestAnimNextGraph"), RF_Transient, nullptr, nullptr, NAME_None));
+	if(Graph == nullptr)
 	{
-		AddError(TEXT("Could not create parameter block."));
+		AddError(TEXT("Could not create graph."));
 		return false;
 	}
 
-	UAnimNextParameterBlock_EditorData* EditorData = UncookedOnly::FUtils::GetEditorData(Block);
+	UAnimNextGraph_EditorData* EditorData = UncookedOnly::FUtils::GetEditorData(Graph);
 	if(EditorData == nullptr)
 	{
-		AddError(TEXT("Parameter block has no editor data."));
+		AddError(TEXT("Graph has no editor data."));
 		return false;
 	}
 
 	static FName TestParameterName = TEXT("TestParam");
 	
 	// AddParameter
-	UAnimNextParameterBlockParameter* Parameter = nullptr;
+	UAnimNextGraph_Parameter* Parameter = nullptr;
 	{
-		FScopedTransaction Transaction(FText::GetEmpty());		
+		FScopedTransaction Transaction(FText::GetEmpty());
 		Parameter = EditorData->AddParameter(TestParameterName, FAnimNextParamType::GetType<bool>());
 
-		if (AddErrorIfFalse(Parameter != nullptr, TEXT("Could not create new parameter in block.")))
+		if (AddErrorIfFalse(Parameter != nullptr, TEXT("Could not create new parameter in graph.")))
 		{
-			AddErrorIfFalse(Parameter->GetParamType() == FAnimNextParamType::GetType<bool>(), TEXT("Incorrect parameter type found"));			
+			AddErrorIfFalse(Parameter->GetParamType() == FAnimNextParamType::GetType<bool>(), TEXT("Incorrect parameter type found"));
 		}
 	}
 
-	AddExpectedError(TEXT("UAnimNextParameterBlock_EditorData::AddParameter: A parameter already exists for the supplied parameter name."));
+	AddExpectedError(TEXT("UAnimNextGraph_EditorData::AddParameter: A parameter already exists for the supplied parameter name."));
 	AddErrorIfFalse(EditorData->AddParameter(TestParameterName, FAnimNextParamType::GetType<bool>()) == nullptr, TEXT("Expected duplicate parameter name argument to fail"));
 
 	GEditor->UndoTransaction();
-	AddErrorIfFalse(EditorData->Entries.Num() == 0, FString::Printf(TEXT("Unexpected entry count found in parameter block (Have %d, expected 0)."), EditorData->Entries.Num()));
+	AddErrorIfFalse(EditorData->Entries.Num() == 0, FString::Printf(TEXT("Unexpected entry count found in graph (Have %d, expected 0)."), EditorData->Entries.Num()));
 
 	GEditor->RedoTransaction();
-	AddErrorIfFalse(EditorData->Entries.Num() == 1, FString::Printf(TEXT("Unexpected entry count found in parameter block (Have %d, expected 1)."), EditorData->Entries.Num()));
+	AddErrorIfFalse(EditorData->Entries.Num() == 1, FString::Printf(TEXT("Unexpected entry count found in graph (Have %d, expected 1)."), EditorData->Entries.Num()));
 
 	// Failure cases
-	AddExpectedError(TEXT("UAnimNextParameterBlock_EditorData::AddParameter: Invalid parameter name supplied."));
+	AddExpectedError(TEXT("UAnimNextGraph_EditorData::AddParameter: Invalid parameter name supplied."));
 	AddErrorIfFalse(EditorData->AddParameter(NAME_None, FAnimNextParamType::GetType<bool>()) == nullptr, TEXT("Expected invalid argument to fail"));
 
 	auto TestParameterType = [this, EditorData](FAnimNextParamType InType)
 	{
-		UAnimNextParameterBlockParameter* TypedParameter = EditorData->AddParameter(TEXT("TestParam0"), InType);
+		UAnimNextGraph_Parameter* TypedParameter = EditorData->AddParameter(TEXT("TestParam0"), InType);
 		const bool bValidParameter = TypedParameter != nullptr;
-		if (bValidParameter && AddErrorIfFalse(bValidParameter, FString::Printf(TEXT("Could not create new parameter of type %s in block."), *InType.ToString())))
+		if (bValidParameter && AddErrorIfFalse(bValidParameter, FString::Printf(TEXT("Could not create new parameter of type %s in graph."), *InType.ToString())))
 		{
 			AddErrorIfFalse(TypedParameter->GetParamType() == InType, TEXT("Incorrect parameter type found"));
 			EditorData->RemoveEntry(TypedParameter);
@@ -111,31 +108,31 @@ bool FEditor_Parameters_ParameterBlock::RunTest(const FString& InParameters)
 	}
 
 	GEditor->UndoTransaction();
-	AddErrorIfFalse(EditorData->Entries.Num() == 1, FString::Printf(TEXT("Unexpected entry count found in parameter block (Have %d, expected 1)."), EditorData->Entries.Num()));
+	AddErrorIfFalse(EditorData->Entries.Num() == 1, FString::Printf(TEXT("Unexpected entry count found in graph (Have %d, expected 1)."), EditorData->Entries.Num()));
 
 	GEditor->RedoTransaction();
-	AddErrorIfFalse(EditorData->Entries.Num() == 0, FString::Printf(TEXT("Unexpected entry count found in parameter block (Have %d, expected 0)."), EditorData->Entries.Num()));
+	AddErrorIfFalse(EditorData->Entries.Num() == 0, FString::Printf(TEXT("Unexpected entry count found in graph (Have %d, expected 0)."), EditorData->Entries.Num()));
 
 	GEditor->UndoTransaction();
-	AddErrorIfFalse(EditorData->Entries.Num() == 1, FString::Printf(TEXT("Unexpected entry count found in parameter block (Have %d, expected 1)."), EditorData->Entries.Num()));
+	AddErrorIfFalse(EditorData->Entries.Num() == 1, FString::Printf(TEXT("Unexpected entry count found in graph (Have %d, expected 1)."), EditorData->Entries.Num()));
 
 	// FindEntry
-	AddErrorIfFalse(EditorData->FindEntry(TestParameterName) != nullptr, TEXT("Could not find entry in block."));
+	AddErrorIfFalse(EditorData->FindEntry(TestParameterName) != nullptr, TEXT("Could not find entry in graph."));
 	GEditor->UndoTransaction();
 
 	// Add graph
-	UAnimNextParameterBlockGraph* Graph = nullptr;
+	UAnimNextGraph_EventGraph* EventGraph = nullptr;
 	{
 		FScopedTransaction Transaction(FText::GetEmpty());
-		Graph = EditorData->AddGraph(TEXT("TestGraph"));
-		AddErrorIfFalse(Graph != nullptr, TEXT("Could not create new graph in block."));
+		EventGraph = EditorData->AddEventGraph(TEXT("TestGraph"));
+		AddErrorIfFalse(EventGraph != nullptr, TEXT("Could not create new event graph in graph."));
 	}
 
 	GEditor->UndoTransaction();
-	AddErrorIfFalse(EditorData->Entries.Num() == 0, FString::Printf(TEXT("Unexpected entry count found in parameter block (Have %d, expected 0)."), EditorData->Entries.Num()));
+	AddErrorIfFalse(EditorData->Entries.Num() == 0, FString::Printf(TEXT("Unexpected entry count found in graph (Have %d, expected 0)."), EditorData->Entries.Num()));
 
 	GEditor->RedoTransaction();
-	AddErrorIfFalse(EditorData->Entries.Num() == 1, FString::Printf(TEXT("Unexpected entry count found in parameter block (Have %d, expected 1)."), EditorData->Entries.Num()));
+	AddErrorIfFalse(EditorData->Entries.Num() == 1, FString::Printf(TEXT("Unexpected entry count found in graph (Have %d, expected 1)."), EditorData->Entries.Num()));
 	GEditor->UndoTransaction();
 
 	FUtils::CleanupAfterTests();
@@ -151,10 +148,11 @@ bool FEditor_Parameters_Python::RunTest(const FString& InParameters)
 
 	const TCHAR* Script = TEXT(
 		"asset_tools = unreal.AssetToolsHelpers.get_asset_tools()\n"
-		"block = unreal.AssetTools.create_asset(asset_tools, asset_name = \"TestBlock\", package_path = \"/Game/\", asset_class = unreal.AnimNextParameterBlock, factory = unreal.AnimNextParameterBlockFactory())\n"
-		"block.add_parameter(name = \"TestParam\", value_type = unreal.PropertyBagPropertyType.BOOL, container_type = unreal.PropertyBagContainerType.NONE)\n"
-		"block.add_graph(name = \"TestGraph\")\n"
-		"unreal.EditorAssetLibrary.delete_loaded_asset(block)\n"
+		"graph = unreal.AssetTools.create_asset(asset_tools, asset_name = \"TestGraph\", package_path = \"/Game/\", asset_class = unreal.AnimNextGraph, factory = unreal.AnimNextGraphFactory())\n"
+		"graph.add_parameter(name = \"TestParam\", value_type = unreal.PropertyBagPropertyType.BOOL, container_type = unreal.PropertyBagContainerType.NONE)\n"
+		"graph.add_event_graph(name = \"TestEventGraph\")\n"
+		"graph.add_animation_graph(name = \"TestAnimationGraph\")\n"
+		"unreal.EditorAssetLibrary.delete_loaded_asset(graph)\n"
 	);
 
 	IPythonScriptPlugin::Get()->ExecPythonCommand(Script);
@@ -164,7 +162,7 @@ bool FEditor_Parameters_Python::RunTest(const FString& InParameters)
 	return true;
 }
 
-IMPLEMENT_SIMPLE_AUTOMATION_TEST(FEditor_Graph, "Animation.AnimNext.Editor.Graph", EAutomationTestFlags::EditorContext | EAutomationTestFlags::EngineFilter)
+IMPLEMENT_SIMPLE_AUTOMATION_TEST(FEditor_Graph, "Animation.AnimNext.Editor.AnimationGraph", EAutomationTestFlags::EditorContext | EAutomationTestFlags::EngineFilter)
 
 bool FEditor_Graph::RunTest(const FString& InParameters)
 {
@@ -186,18 +184,18 @@ bool FEditor_Graph::RunTest(const FString& InParameters)
 	}
 
 	// Add graph
-	UAnimNextGraphEntry* GraphEntry = nullptr;
+	UAnimNextGraph_AnimationGraph* GraphEntry = nullptr;
 	{
 		FScopedTransaction Transaction(FText::GetEmpty());
-		GraphEntry = EditorData->AddGraph(TEXT("TestGraph"));
-		AddErrorIfFalse(Graph != nullptr, TEXT("Could not create new graph in asset."));
+		GraphEntry = EditorData->AddAnimationGraph(TEXT("TestGraph"));
+		AddErrorIfFalse(Graph != nullptr, TEXT("Could not create new animation graph in asset."));
 	}
 
 	GEditor->UndoTransaction();
-	AddErrorIfFalse(EditorData->Entries.Num() == 1, FString::Printf(TEXT("Unexpected entry count found in graph asset (Have %d, expected 1)."), EditorData->Entries.Num()));
+	AddErrorIfFalse(EditorData->Entries.Num() == 0, FString::Printf(TEXT("Unexpected entry count found in graph asset (Have %d, expected 0)."), EditorData->Entries.Num()));
 
 	GEditor->RedoTransaction();
-	AddErrorIfFalse(EditorData->Entries.Num() == 2, FString::Printf(TEXT("Unexpected entry count found in graph asset (Have %d, expected 2)."), EditorData->Entries.Num()));
+	AddErrorIfFalse(EditorData->Entries.Num() == 1, FString::Printf(TEXT("Unexpected entry count found in graph asset (Have %d, expected 1)."), EditorData->Entries.Num()));
 
 	// RemoveEntry
 	{
@@ -206,13 +204,13 @@ bool FEditor_Graph::RunTest(const FString& InParameters)
 	}
 
 	GEditor->UndoTransaction();
-	AddErrorIfFalse(EditorData->Entries.Num() == 2, FString::Printf(TEXT("Unexpected entry count found in graph asset (Have %d, expected 2)."), EditorData->Entries.Num()));
+	AddErrorIfFalse(EditorData->Entries.Num() == 1, FString::Printf(TEXT("Unexpected entry count found in graph asset (Have %d, expected 1)."), EditorData->Entries.Num()));
 
 	GEditor->RedoTransaction();
-	AddErrorIfFalse(EditorData->Entries.Num() == 1, FString::Printf(TEXT("Unexpected entry count found in graph asset (Have %d, expected 1)."), EditorData->Entries.Num()));
+	AddErrorIfFalse(EditorData->Entries.Num() == 0, FString::Printf(TEXT("Unexpected entry count found in graph asset (Have %d, expected 0)."), EditorData->Entries.Num()));
 
 	GEditor->UndoTransaction();
-	AddErrorIfFalse(EditorData->Entries.Num() == 2, FString::Printf(TEXT("Unexpected entry count found in graph asset (Have %d, expected 2)."), EditorData->Entries.Num()));
+	AddErrorIfFalse(EditorData->Entries.Num() == 1, FString::Printf(TEXT("Unexpected entry count found in graph asset (Have %d, expected 1)."), EditorData->Entries.Num()));
 
 	// FindEntry
 	AddErrorIfFalse(EditorData->FindEntry(TEXT("TestGraph")) != nullptr, TEXT("Could not find entry in asset."));
diff --git a/Engine/Plugins/Experimental/AnimNext/Source/AnimNextTestSuite/Private/AnimNextTraitGraphTest.cpp b/Engine/Plugins/Experimental/AnimNext/Source/AnimNextTestSuite/Private/AnimNextTraitGraphTest.cpp
--- a/Engine/Plugins/Experimental/AnimNext/Source/AnimNextTestSuite/Private/AnimNextTraitGraphTest.cpp
+++ b/Engine/Plugins/Experimental/AnimNext/Source/AnimNextTestSuite/Private/AnimNextTraitGraphTest.cpp
@@ -96,6 +96,8 @@ bool FAnimationAnimNextRuntimeTest_GraphAddTrait::RunTest(const FString& InParam
 		UAnimNextGraph_EditorData* EditorData = UncookedOnly::FUtils::GetEditorData(AnimNextGraph);
 		UE_RETURN_ON_ERROR(EditorData != nullptr, "FAnimationAnimNextRuntimeTest_GraphAddTrait -> Failed to find animation graph editor data");
 
+		UE_RETURN_ON_ERROR(EditorData->AddAnimationGraph(FRigUnit_AnimNextGraphRoot::DefaultEntryPoint, false) != nullptr, "FAnimationAnimNextRuntimeTest_GraphAddTrait -> Failed to add animation graph");
+
 		URigVMController* Controller = EditorData->GetRigVMClient()->GetController(EditorData->GetRigVMClient()->GetDefaultModel());
 		UE_RETURN_ON_ERROR(Controller != nullptr, "FAnimationAnimNextRuntimeTest_GraphAddTrait -> Failed to get RigVM controller");
 
@@ -199,6 +201,8 @@ bool FAnimationAnimNextRuntimeTest_GraphExecute::RunTest(const FString& InParame
 		UAnimNextGraph_EditorData* EditorData = UncookedOnly::FUtils::GetEditorData(AnimNextGraph);
 		UE_RETURN_ON_ERROR(EditorData != nullptr, "FAnimationAnimNextRuntimeTest_GraphExecute -> Failed to find animation graph editor data");
 
+		UE_RETURN_ON_ERROR(EditorData->AddAnimationGraph(FRigUnit_AnimNextGraphRoot::DefaultEntryPoint, false) != nullptr, "FAnimationAnimNextRuntimeTest_GraphExecute -> Failed to add animation graph");
+
 		URigVMController* Controller = EditorData->GetRigVMClient()->GetController(EditorData->GetRigVMClient()->GetDefaultModel());
 		UE_RETURN_ON_ERROR(Controller != nullptr, "FAnimationAnimNextRuntimeTest_GraphExecute -> Failed to get RigVM controller");
 
@@ -315,6 +319,8 @@ bool FAnimationAnimNextRuntimeTest_GraphExecuteLatent::RunTest(const FString& In
 		UAnimNextGraph_EditorData* EditorData = UncookedOnly::FUtils::GetEditorData(AnimNextGraph);
 		UE_RETURN_ON_ERROR(EditorData != nullptr, "FAnimationAnimNextRuntimeTest_GraphExecuteLatent -> Failed to find animation graph editor data");
 
+		UE_RETURN_ON_ERROR(EditorData->AddAnimationGraph(FRigUnit_AnimNextGraphRoot::DefaultEntryPoint, false) != nullptr, "FAnimationAnimNextRuntimeTest_GraphExecuteLatent -> Failed to add animation graph");
+
 		URigVMController* Controller = EditorData->GetRigVMClient()->GetController(EditorData->GetRigVMClient()->GetDefaultModel());
 		UE_RETURN_ON_ERROR(Controller != nullptr, "FAnimationAnimNextRuntimeTest_GraphExecuteLatent -> Failed to get RigVM controller");
 
diff --git a/Engine/Plugins/Experimental/AnimNext/Source/AnimNextUncookedOnly/Internal/AnimNextRigVMAssetEditorData.h b/Engine/Plugins/Experimental/AnimNext/Source/AnimNextUncookedOnly/Internal/AnimNextRigVMAssetEditorData.h
--- a/Engine/Plugins/Experimental/AnimNext/Source/AnimNextUncookedOnly/Internal/AnimNextRigVMAssetEditorData.h
+++ b/Engine/Plugins/Experimental/AnimNext/Source/AnimNextUncookedOnly/Internal/AnimNextRigVMAssetEditorData.h
@@ -13,7 +13,6 @@
 enum class ERigVMGraphNotifType : uint8;
 class UAnimNextRigVMAssetEntry;
 class UAnimNextRigVMAssetEditorData;
-class UAnimNextParameterBlock_EdGraph;
 class UAnimNextGraph_EdGraph;
 
 namespace UE::AnimNext::UncookedOnly
@@ -25,12 +24,10 @@ namespace UE::AnimNext::UncookedOnly
 namespace UE::AnimNext::Editor
 {
 	struct FUtils;
-	class FParameterBlockEditor;
 	class SRigVMAssetView;
 	class SParameterPicker;
-	class FParameterBlockTabSummoner;
 	class SRigVMAssetViewRow;
-	class FParameterBlockParameterCustomization;
+	class FParameterCustomization;
 	class FModule;
 	class FWorkspaceEditor;
 }
@@ -38,6 +35,7 @@ namespace UE::AnimNext::Editor
 namespace UE::AnimNext::Tests
 {
 	class FEditor_Graph;
+	class FEditor_Parameters;
 }
 
 namespace UE::AnimNext::UncookedOnly
@@ -78,8 +76,9 @@ protected:
 	friend class UAnimNextRigVMAssetEntry;
 	friend class UAnimNextRigVMAssetLibrary;
 	friend class UAnimNextGraph_EdGraph;
-	friend class UAnimNextParameterBlock_EdGraph;
 	friend class UE::AnimNext::Tests::FEditor_Graph;
+	friend class UE::AnimNext::Tests::FEditor_Parameters;
+	friend class UE::AnimNext::Editor::FParameterCustomization;
 
 	// UObject interface
 	virtual void Serialize(FArchive& Ar) override;
diff --git a/Engine/Plugins/Experimental/AnimNext/Source/AnimNextUncookedOnly/Internal/AnimNextRigVMAssetEntry.h b/Engine/Plugins/Experimental/AnimNext/Source/AnimNextUncookedOnly/Internal/AnimNextRigVMAssetEntry.h
--- a/Engine/Plugins/Experimental/AnimNext/Source/AnimNextUncookedOnly/Internal/AnimNextRigVMAssetEntry.h
+++ b/Engine/Plugins/Experimental/AnimNext/Source/AnimNextUncookedOnly/Internal/AnimNextRigVMAssetEntry.h
@@ -15,7 +15,7 @@ namespace UE::AnimNext::Editor
 	struct FRigVMAssetViewEntry;
 }
 
-/** Base class that defines an entry in a parameter block, e.g. a parameter binding */
+/** Base class that defines an entry in a graph, e.g. a parameter binding */
 UCLASS(MinimalAPI, BlueprintType, Abstract)
 class UAnimNextRigVMAssetEntry : public UObject
 {
diff --git a/Engine/Plugins/Experimental/AnimNext/Source/AnimNextUncookedOnly/Internal/Graph/AnimNextGraphEntry.h b/Engine/Plugins/Experimental/AnimNext/Source/AnimNextUncookedOnly/Internal/Graph/AnimNextGraphEntry.h
deleted file mode 100644
diff --git a/Engine/Plugins/Experimental/AnimNext/Source/AnimNextUncookedOnly/Internal/Graph/AnimNextGraph_AnimationGraph.h b/Engine/Plugins/Experimental/AnimNext/Source/AnimNextUncookedOnly/Internal/Graph/AnimNextGraph_AnimationGraph.h
new file mode 100644
--- /dev/null
+++ b/Engine/Plugins/Experimental/AnimNext/Source/AnimNextUncookedOnly/Internal/Graph/AnimNextGraph_AnimationGraph.h
@@ -0,0 +1,60 @@
+﻿// Copyright Epic Games, Inc. All Rights Reserved.
+
+#pragma once
+
+#include "CoreMinimal.h"
+#include "AnimNextRigVMAssetEntry.h"
+#include "IAnimNextRigVMGraphInterface.h"
+#include "AnimNextGraph_AnimationGraph.generated.h"
+
+class UAnimNextGraph_EditorData;
+class UAnimNextGraph_EdGraph;
+enum class ERigVMGraphNotifType : uint8;
+
+namespace UE::AnimNext::Editor
+{
+	struct FUtils;
+}
+
+/** A single entry in an AnimNext graph asset */
+UCLASS(MinimalAPI, Category = "Animation Graphs")
+class UAnimNextGraph_AnimationGraph : public UAnimNextRigVMAssetEntry, public IAnimNextRigVMGraphInterface
+{
+	GENERATED_BODY()
+
+	friend class UAnimNextGraph_EditorData;
+	friend struct UE::AnimNext::Editor::FUtils;	
+
+	// UAnimNextRigVMAssetEntry interface
+	virtual FName GetEntryName() const override;
+	virtual void SetEntryName(FName InName, bool bSetupUndoRedo = true) override;
+	
+	// IAnimNextRigVMGraphInterface interface
+	virtual URigVMGraph* GetRigVMGraph() const override;
+	virtual URigVMEdGraph* GetEdGraph() const override;
+	virtual void SetRigVMGraph(URigVMGraph* InGraph) override;
+	virtual void SetEdGraph(URigVMEdGraph* InGraph) override;
+
+protected:
+	/** The name of the graph */
+	UPROPERTY(VisibleAnywhere, Category = AnimationGraph)
+	FName GraphName;
+
+	/** RigVM graph */
+	UPROPERTY()
+	TObjectPtr<URigVMGraph> Graph;
+
+	/** Editor graph */
+	UPROPERTY()
+	TObjectPtr<UAnimNextGraph_EdGraph> EdGraph;
+};
+
+// Old deprecated entry
+UCLASS()
+class UAnimNextGraphEntry : public UAnimNextRigVMAssetEntry
+{
+	GENERATED_BODY()
+
+	// UAnimNextRigVMAssetEntry interface
+	virtual FName GetEntryName() const override { return NAME_None; }
+};
\ No newline at end of file
diff --git a/Engine/Plugins/Experimental/AnimNext/Source/AnimNextUncookedOnly/Internal/Graph/AnimNextGraph_AnimationGraphSchema.h b/Engine/Plugins/Experimental/AnimNext/Source/AnimNextUncookedOnly/Internal/Graph/AnimNextGraph_AnimationGraphSchema.h
new file mode 100644
--- /dev/null
+++ b/Engine/Plugins/Experimental/AnimNext/Source/AnimNextUncookedOnly/Internal/Graph/AnimNextGraph_AnimationGraphSchema.h
@@ -0,0 +1,17 @@
+﻿// Copyright Epic Games, Inc. All Rights Reserved.
+
+#pragma once
+
+#include "CoreMinimal.h"
+#include "AnimNextRigVMAssetSchema.h"
+#include "AnimNextGraph_AnimationGraphSchema.generated.h"
+
+UCLASS()
+class UAnimNextGraph_AnimationGraphSchema : public UAnimNextRigVMAssetSchema
+{
+	GENERATED_BODY()
+
+	// URigVMSchema interface
+	virtual bool SupportsUnitFunction(URigVMController* InController, const FRigVMFunction* InUnitFunction) const override;
+	virtual bool SupportsDispatchFactory(URigVMController* InController, const FRigVMDispatchFactory* InDispatchFactory) const override;
+};
diff --git a/Engine/Plugins/Experimental/AnimNext/Source/AnimNextUncookedOnly/Internal/Graph/AnimNextGraph_EdGraphSchema.h b/Engine/Plugins/Experimental/AnimNext/Source/AnimNextUncookedOnly/Internal/Graph/AnimNextGraph_EdGraphSchema.h
--- a/Engine/Plugins/Experimental/AnimNext/Source/AnimNextUncookedOnly/Internal/Graph/AnimNextGraph_EdGraphSchema.h
+++ b/Engine/Plugins/Experimental/AnimNext/Source/AnimNextUncookedOnly/Internal/Graph/AnimNextGraph_EdGraphSchema.h
@@ -16,6 +16,5 @@ class UAnimNextGraph_EdGraphSchema : public URigVMEdGraphSchema
 	virtual TSubclassOf<URigVMEdGraphNode> GetGraphNodeClass(const URigVMEdGraph* InGraph) const override { return UAnimNextGraph_EdGraphNode::StaticClass(); }
 	
 	// UEdGraphSchema interface
-	virtual void TrySetDefaultValue(UEdGraphPin& InPin, const FString& InNewDefaultValue, bool bMarkAsModified) const override;
 	virtual void GetGraphDisplayInformation(const UEdGraph& Graph, /*out*/ FGraphDisplayInfo& DisplayInfo) const override;
 };
\ No newline at end of file
diff --git a/Engine/Plugins/Experimental/AnimNext/Source/AnimNextUncookedOnly/Internal/Graph/AnimNextGraph_EditorData.h b/Engine/Plugins/Experimental/AnimNext/Source/AnimNextUncookedOnly/Internal/Graph/AnimNextGraph_EditorData.h
--- a/Engine/Plugins/Experimental/AnimNext/Source/AnimNextUncookedOnly/Internal/Graph/AnimNextGraph_EditorData.h
+++ b/Engine/Plugins/Experimental/AnimNext/Source/AnimNextUncookedOnly/Internal/Graph/AnimNextGraph_EditorData.h
@@ -29,21 +29,23 @@ namespace UE::AnimNext::Editor
 	struct FUtils;
 }
 
-UCLASS()
-class UAnimNextTraitGraph_Schema : public URigVMSchema
-{
-	GENERATED_BODY()
-};
-
 // Script-callable editor API hoisted onto UAnimNextGraph
 UCLASS()
 class ANIMNEXTUNCOOKEDONLY_API UAnimNextGraphLibrary : public UBlueprintFunctionLibrary
 {
 	GENERATED_BODY()
 
-	/** Adds a graph to an AnimNext Graph asset */
+	/** Adds an animation graph to an AnimNext Graph asset */
 	UFUNCTION(BlueprintCallable, Category = "AnimNext|Graph", meta=(ScriptMethod))
-	static UAnimNextGraphEntry* AddGraph(UAnimNextGraph* InGraph, FName InName, bool bSetupUndoRedo = true, bool bPrintPythonCommand = true);
+	static UAnimNextGraph_AnimationGraph* AddAnimationGraph(UAnimNextGraph* InGraph, FName InName, bool bSetupUndoRedo = true, bool bPrintPythonCommand = true);
+
+	/** Adds a parameter to an AnimNext Graph asset */
+	UFUNCTION(BlueprintCallable, Category = "AnimNext|Graph", meta=(ScriptMethod))
+	static UAnimNextGraph_Parameter* AddParameter(UAnimNextGraph* InGraph, FName InName, EPropertyBagPropertyType InValueType, EPropertyBagContainerType InContainerType = EPropertyBagContainerType::None, const UObject* InValueTypeObject = nullptr, bool bSetupUndoRedo = true, bool bPrintPythonCommand = true);
+
+	/** Adds an event graph to an AnimNext Graph asset */
+	UFUNCTION(BlueprintCallable, Category = "AnimNext|Graph", meta=(ScriptMethod))
+	static UAnimNextGraph_EventGraph* AddEventGraph(UAnimNextGraph* InGraph, FName InName, bool bSetupUndoRedo = true, bool bPrintPythonCommand = true);
 };
 
 /** Editor data for AnimNext graphs */
@@ -66,8 +68,14 @@ class UAnimNextGraph_EditorData : public UAnimNextRigVMAssetEditorData
 	friend class FAnimationAnimNextRuntimeTest_GraphExecuteLatent;
 	
 public:
-	/** Adds a graph to this asset */
-	ANIMNEXTUNCOOKEDONLY_API UAnimNextGraphEntry* AddGraph(FName InName, bool bSetupUndoRedo = true, bool bPrintPythonCommand = true);
+	/** Adds an animation graph to this asset */
+	ANIMNEXTUNCOOKEDONLY_API UAnimNextGraph_AnimationGraph* AddAnimationGraph(FName InName, bool bSetupUndoRedo = true, bool bPrintPythonCommand = true);
+
+	/** Adds a parameter to this asset */
+	ANIMNEXTUNCOOKEDONLY_API UAnimNextGraph_Parameter* AddParameter(FName InName, FAnimNextParamType InType, bool bSetupUndoRedo = true, bool bPrintPythonCommand = true);
+
+	/** Adds an event graph to this asset */
+	ANIMNEXTUNCOOKEDONLY_API UAnimNextGraph_EventGraph* AddEventGraph(FName InName, bool bSetupUndoRedo = true, bool bPrintPythonCommand = true);
 
 private:
 	// UObject interface
diff --git a/Engine/Plugins/Experimental/AnimNext/Source/AnimNextUncookedOnly/Internal/Graph/AnimNextGraph_EventGraph.h b/Engine/Plugins/Experimental/AnimNext/Source/AnimNextUncookedOnly/Internal/Graph/AnimNextGraph_EventGraph.h
new file mode 100644
--- /dev/null
+++ b/Engine/Plugins/Experimental/AnimNext/Source/AnimNextUncookedOnly/Internal/Graph/AnimNextGraph_EventGraph.h
@@ -0,0 +1,53 @@
+﻿// Copyright Epic Games, Inc. All Rights Reserved.
+
+#pragma once
+
+#include "CoreMinimal.h"
+#include "AnimNextRigVMAssetEntry.h"
+#include "IAnimNextRigVMGraphInterface.h"
+#include "AnimNextGraph_EventGraph.generated.h"
+
+class UAnimNextGraph_EditorData;
+class UAnimNextGraph_EdGraph;
+
+UCLASS(MinimalAPI, Category = "Event Graphs")
+class UAnimNextGraph_EventGraph : public UAnimNextRigVMAssetEntry, public IAnimNextRigVMGraphInterface
+{
+	GENERATED_BODY()
+
+	friend class UAnimNextGraph_EditorData;
+
+	// UAnimNextRigVMAssetEntry interface
+	virtual FName GetEntryName() const override { return GraphName; }
+	virtual void SetEntryName(FName InName, bool bSetupUndoRedo = true) override;
+	virtual FText GetDisplayName() const override;
+	virtual FText GetDisplayNameTooltip() const override;
+
+	// IAnimNextRigVMGraphInterface interface
+	virtual URigVMGraph* GetRigVMGraph() const override;
+	virtual URigVMEdGraph* GetEdGraph() const override;
+	virtual void SetRigVMGraph(URigVMGraph* InGraph) override;
+	virtual void SetEdGraph(URigVMEdGraph* InGraph) override;
+
+	/** The name of the graph */
+	UPROPERTY(VisibleAnywhere, Category = EventGraph)
+	FName GraphName;
+
+	/** Graph */
+	UPROPERTY()
+	TObjectPtr<URigVMGraph> Graph;
+
+	/** Graph */
+	UPROPERTY()
+	TObjectPtr<UAnimNextGraph_EdGraph> EdGraph;
+};
+
+// Old deprecated entry
+UCLASS()
+class UAnimNextParameterBlockGraph : public UAnimNextRigVMAssetEntry
+{
+	GENERATED_BODY()
+
+	// UAnimNextRigVMAssetEntry interface
+	virtual FName GetEntryName() const override { return NAME_None; }
+};
\ No newline at end of file
diff --git a/Engine/Plugins/Experimental/AnimNext/Source/AnimNextUncookedOnly/Internal/Graph/AnimNextGraph_EventGraphSchema.h b/Engine/Plugins/Experimental/AnimNext/Source/AnimNextUncookedOnly/Internal/Graph/AnimNextGraph_EventGraphSchema.h
new file mode 100644
--- /dev/null
+++ b/Engine/Plugins/Experimental/AnimNext/Source/AnimNextUncookedOnly/Internal/Graph/AnimNextGraph_EventGraphSchema.h
@@ -0,0 +1,17 @@
+﻿// Copyright Epic Games, Inc. All Rights Reserved.
+
+#pragma once
+
+#include "CoreMinimal.h"
+#include "AnimNextRigVMAssetSchema.h"
+#include "AnimNextGraph_EventGraphSchema.generated.h"
+
+UCLASS()
+class UAnimNextGraph_EventGraphSchema : public UAnimNextRigVMAssetSchema
+{
+	GENERATED_BODY()
+
+	// URigVMSchema interface
+	virtual bool SupportsUnitFunction(URigVMController* InController, const FRigVMFunction* InUnitFunction) const override;
+};
+
diff --git a/Engine/Plugins/Experimental/AnimNext/Source/AnimNextUncookedOnly/Internal/Graph/AnimNextGraph_Parameter.h b/Engine/Plugins/Experimental/AnimNext/Source/AnimNextUncookedOnly/Internal/Graph/AnimNextGraph_Parameter.h
new file mode 100644
--- /dev/null
+++ b/Engine/Plugins/Experimental/AnimNext/Source/AnimNextUncookedOnly/Internal/Graph/AnimNextGraph_Parameter.h
@@ -0,0 +1,64 @@
+// Copyright Epic Games, Inc. All Rights Reserved.
+
+#pragma once
+
+#include "CoreMinimal.h"
+#include "AnimNextRigVMAssetEntry.h"
+#include "IAnimNextRigVMParameterInterface.h"
+#include "Param/ParamType.h"
+#include "AnimNextGraph_Parameter.generated.h"
+
+class UAnimNextGraph_EditorData;
+
+namespace UE::AnimNext::Editor
+{
+	class FParameterCustomization;
+}
+
+namespace UE::AnimNext::Tests
+{
+	class FEditor_Parameters;
+}
+
+UCLASS(MinimalAPI, Category = "Parameters")
+class UAnimNextGraph_Parameter : public UAnimNextRigVMAssetEntry, public IAnimNextRigVMParameterInterface
+{
+	GENERATED_BODY()
+
+	friend class UAnimNextGraph_EditorData;
+	friend class UE::AnimNext::Tests::FEditor_Parameters;
+	friend class UE::AnimNext::Editor::FParameterCustomization;
+
+	// UAnimNextRigVMAssetEntry interface
+	virtual FName GetEntryName() const override;
+	virtual void SetEntryName(FName InName, bool bSetupUndoRedo = true) override;
+	virtual FText GetDisplayName() const override;
+	virtual FText GetDisplayNameTooltip() const override;
+
+	// IAnimNextRigVMParameterInterface interface
+	virtual FAnimNextParamType GetParamType() const override;
+	virtual bool SetParamType(const FAnimNextParamType& InType, bool bSetupUndoRedo = true) override;
+	virtual FInstancedPropertyBag& GetPropertyBag() const override;
+	
+	/** Parameter name we reference */
+	UPROPERTY(VisibleAnywhere, Category = Parameter)
+	FName ParameterName;
+
+	/** The parameter's type */
+	UPROPERTY(EditAnywhere, Category = "Parameter", AssetRegistrySearchable)
+	FAnimNextParamType Type = FAnimNextParamType::GetType<bool>();
+	
+	/** Comment to display in editor */
+	UPROPERTY(EditAnywhere, Category = "Parameter", meta=(MultiLine))
+	FString Comment;
+};
+
+// Old deprecated class
+UCLASS()
+class UAnimNextParameterBlockParameter : public UAnimNextRigVMAssetEntry
+{
+	GENERATED_BODY()
+
+	// UAnimNextRigVMAssetEntry interface
+	virtual FName GetEntryName() const { return NAME_None; }
+};
\ No newline at end of file
diff --git a/Engine/Plugins/Experimental/AnimNext/Source/AnimNextUncookedOnly/Internal/Graph/AnimNextGraph_TraitSchema.h b/Engine/Plugins/Experimental/AnimNext/Source/AnimNextUncookedOnly/Internal/Graph/AnimNextGraph_TraitSchema.h
deleted file mode 100644
diff --git a/Engine/Plugins/Experimental/AnimNext/Source/AnimNextUncookedOnly/Internal/IAnimNextRigVMGraphInterface.h b/Engine/Plugins/Experimental/AnimNext/Source/AnimNextUncookedOnly/Internal/IAnimNextRigVMGraphInterface.h
--- a/Engine/Plugins/Experimental/AnimNext/Source/AnimNextUncookedOnly/Internal/IAnimNextRigVMGraphInterface.h
+++ b/Engine/Plugins/Experimental/AnimNext/Source/AnimNextUncookedOnly/Internal/IAnimNextRigVMGraphInterface.h
@@ -25,4 +25,13 @@ public:
 
 	// Get the Editor graph
 	virtual URigVMEdGraph* GetEdGraph() const = 0;
+
+private:
+	friend class UAnimNextGraph_EditorData;
+
+	// Set the RigVM graph
+	virtual void SetRigVMGraph(URigVMGraph* InGraph) = 0;
+
+	// Get the Editor graph
+	virtual void SetEdGraph(URigVMEdGraph* InGraph) = 0;
 };
\ No newline at end of file
diff --git a/Engine/Plugins/Experimental/AnimNext/Source/AnimNextUncookedOnly/Internal/Param/AnimNextParameterBlockGraph.h b/Engine/Plugins/Experimental/AnimNext/Source/AnimNextUncookedOnly/Internal/Param/AnimNextParameterBlockGraph.h
deleted file mode 100644
diff --git a/Engine/Plugins/Experimental/AnimNext/Source/AnimNextUncookedOnly/Internal/Param/AnimNextParameterBlockParameter.h b/Engine/Plugins/Experimental/AnimNext/Source/AnimNextUncookedOnly/Internal/Param/AnimNextParameterBlockParameter.h
deleted file mode 100644
diff --git a/Engine/Plugins/Experimental/AnimNext/Source/AnimNextUncookedOnly/Internal/Param/AnimNextParameterBlock_EdGraph.h b/Engine/Plugins/Experimental/AnimNext/Source/AnimNextUncookedOnly/Internal/Param/AnimNextParameterBlock_EdGraph.h
deleted file mode 100644
diff --git a/Engine/Plugins/Experimental/AnimNext/Source/AnimNextUncookedOnly/Internal/Param/AnimNextParameterBlock_EdGraphNode.h b/Engine/Plugins/Experimental/AnimNext/Source/AnimNextUncookedOnly/Internal/Param/AnimNextParameterBlock_EdGraphNode.h
deleted file mode 100644
diff --git a/Engine/Plugins/Experimental/AnimNext/Source/AnimNextUncookedOnly/Internal/Param/AnimNextParameterBlock_EdGraphSchema.h b/Engine/Plugins/Experimental/AnimNext/Source/AnimNextUncookedOnly/Internal/Param/AnimNextParameterBlock_EdGraphSchema.h
deleted file mode 100644
diff --git a/Engine/Plugins/Experimental/AnimNext/Source/AnimNextUncookedOnly/Internal/Param/AnimNextParameterBlock_EditorData.h b/Engine/Plugins/Experimental/AnimNext/Source/AnimNextUncookedOnly/Internal/Param/AnimNextParameterBlock_EditorData.h
deleted file mode 100644
diff --git a/Engine/Plugins/Experimental/AnimNext/Source/AnimNextUncookedOnly/Internal/Param/AnimNextParameterBlock_ParamSchema.h b/Engine/Plugins/Experimental/AnimNext/Source/AnimNextUncookedOnly/Internal/Param/AnimNextParameterBlock_ParamSchema.h
deleted file mode 100644
diff --git a/Engine/Plugins/Experimental/AnimNext/Source/AnimNextUncookedOnly/Internal/Param/IAnimNextParameterBlockReferenceInterface.h b/Engine/Plugins/Experimental/AnimNext/Source/AnimNextUncookedOnly/Internal/Param/IAnimNextParameterBlockReferenceInterface.h
deleted file mode 100644
diff --git a/Engine/Plugins/Experimental/AnimNext/Source/AnimNextUncookedOnly/Internal/UncookedOnlyUtils.h b/Engine/Plugins/Experimental/AnimNext/Source/AnimNextUncookedOnly/Internal/UncookedOnlyUtils.h
--- a/Engine/Plugins/Experimental/AnimNext/Source/AnimNextUncookedOnly/Internal/UncookedOnlyUtils.h
+++ b/Engine/Plugins/Experimental/AnimNext/Source/AnimNextUncookedOnly/Internal/UncookedOnlyUtils.h
@@ -17,8 +17,6 @@ class UAnimNextGraph_EditorData;
 class UAnimNextGraph_EdGraph;
 class URigVMController;
 class URigVMGraph;
-class UAnimNextParameterBlock;
-class UAnimNextParameterBlock_EditorData;
 class UAnimNextGraph_EdGraph;
 struct FEdGraphPinType;
 class UAnimNextRigVMAsset;
@@ -87,11 +85,17 @@ namespace UE::AnimNext::UncookedOnly
 struct ANIMNEXTUNCOOKEDONLY_API FUtils
 {
 	static void Compile(UAnimNextGraph* InGraph);
-	
+
+	static void CompileVM(UAnimNextGraph* InGraph);
+
+	static void CompileStruct(UAnimNextGraph* InGraph);
+
 	static UAnimNextGraph_EditorData* GetEditorData(const UAnimNextGraph* InAnimNextGraph);
-	
+
 	static UAnimNextGraph* GetGraph(const UAnimNextGraph_EditorData* InEditorData);
-	
+
+	static FInstancedPropertyBag* GetPropertyBag(UAnimNextGraph* InAnimNextGraph);
+
 	static void RecreateVM(UAnimNextGraph* InGraph);
 
 	/**
@@ -100,20 +104,6 @@ struct ANIMNEXTUNCOOKEDONLY_API FUtils
 	 **/
 	static FParamTypeHandle GetParameterHandleFromPin(const FEdGraphPinType& InPinType);
 
-	static void Compile(UAnimNextParameterBlock* InParameterBlock);
-
-	static void CompileVM(UAnimNextParameterBlock* InParameterBlock);
-
-	static void CompileStruct(UAnimNextParameterBlock* InParameterBlock);
-	
-	static UAnimNextParameterBlock_EditorData* GetEditorData(const UAnimNextParameterBlock* InParameterBlock);
-
-	static UAnimNextParameterBlock* GetBlock(const UAnimNextParameterBlock_EditorData* InEditorData);
-
-	static FInstancedPropertyBag* GetPropertyBag(UAnimNextParameterBlock* ReferencedBlock);
-
-	static void RecreateVM(UAnimNextParameterBlock* InParameterBlock);
-
 	static UAnimNextRigVMAsset* GetAsset(UAnimNextRigVMAssetEditorData* InEditorData);
 
 	static UAnimNextRigVMAssetEditorData* GetEditorData(UAnimNextRigVMAsset* InAsset);
@@ -142,8 +132,8 @@ struct ANIMNEXTUNCOOKEDONLY_API FUtils
 	/** Set up a simple animation graph */
 	static void SetupAnimGraph(UAnimNextRigVMAssetEntry* InEntry, URigVMController* InController);
 	
-	/** Set up a simple parameter graph */
-	static void SetupParameterGraph(URigVMController* InController);
+	/** Set up a simple event graph */
+	static void SetupEventGraph(URigVMController* InController);
 	
 	/** Converts the Verse-tag-like snake_case_parameter_name to a period-separated display name similar to a gameplay tag */
 	static FText GetParameterDisplayNameText(FName InParameterName);
diff --git a/Engine/Plugins/Experimental/AnimNext/Source/AnimNextUncookedOnly/Private/AnimNextRigVMAssetEditorData.cpp b/Engine/Plugins/Experimental/AnimNext/Source/AnimNextUncookedOnly/Private/AnimNextRigVMAssetEditorData.cpp
--- a/Engine/Plugins/Experimental/AnimNext/Source/AnimNextUncookedOnly/Private/AnimNextRigVMAssetEditorData.cpp
+++ b/Engine/Plugins/Experimental/AnimNext/Source/AnimNextUncookedOnly/Private/AnimNextRigVMAssetEditorData.cpp
@@ -75,6 +75,7 @@ void UAnimNextRigVMAssetEditorData::Serialize(FArchive& Ar)
 void UAnimNextRigVMAssetEditorData::Initialize(bool bRecompileVM)
 {
 	RigVMClient.bDefaultModelCanBeRemoved = true;
+	RigVMClient.SetDefaultSchemaClass(UAnimNextRigVMAssetSchema::StaticClass());
 	RigVMClient.SetControllerClass(GetControllerClass());
 	RigVMClient.SetOuterClientHost(this, GET_MEMBER_NAME_CHECKED(UAnimNextRigVMAssetEditorData, RigVMClient));
 	{
diff --git a/Engine/Plugins/Experimental/AnimNext/Source/AnimNextUncookedOnly/Private/Graph/AnimGraph/AnimGraphNodeBinding_AnimNextParameters.cpp b/Engine/Plugins/Experimental/AnimNext/Source/AnimNextUncookedOnly/Private/Graph/AnimGraph/AnimGraphNodeBinding_AnimNextParameters.cpp
--- a/Engine/Plugins/Experimental/AnimNext/Source/AnimNextUncookedOnly/Private/Graph/AnimGraph/AnimGraphNodeBinding_AnimNextParameters.cpp
+++ b/Engine/Plugins/Experimental/AnimNext/Source/AnimNextUncookedOnly/Private/Graph/AnimGraph/AnimGraphNodeBinding_AnimNextParameters.cpp
@@ -550,7 +550,7 @@ TSharedRef<SWidget> UAnimGraphNodeBinding_AnimNextParameters::MakePropertyBindin
 				});
 				Args.NewParameterType = FParamTypeHandle::FromProperty(PropertyToBindTo).GetType();
 				Args.bMultiSelect = false;
-				Args.bShowBlocks = false;
+				Args.bShowSourceGraph = false;
 			
 				MenuBuilder.AddWidget(
 					SNew(SBox)
diff --git a/Engine/Plugins/Experimental/AnimNext/Source/AnimNextUncookedOnly/Private/Graph/AnimGraph/AnimGraphNode_AnimNextParameters.cpp b/Engine/Plugins/Experimental/AnimNext/Source/AnimNextUncookedOnly/Private/Graph/AnimGraph/AnimGraphNode_AnimNextParameters.cpp
--- a/Engine/Plugins/Experimental/AnimNext/Source/AnimNextUncookedOnly/Private/Graph/AnimGraph/AnimGraphNode_AnimNextParameters.cpp
+++ b/Engine/Plugins/Experimental/AnimNext/Source/AnimNextUncookedOnly/Private/Graph/AnimGraph/AnimGraphNode_AnimNextParameters.cpp
@@ -7,12 +7,12 @@
 
 FText UAnimGraphNode_AnimNextParameters::GetNodeTitle(ENodeTitleType::Type TitleType) const
 {
-	return LOCTEXT("NodeTitle", "AnimNext Parameter Block");
+	return LOCTEXT("NodeTitle", "AnimNext Parameters");
 }
 
 FText UAnimGraphNode_AnimNextParameters::GetTooltipText() const
 {
-	return LOCTEXT("NodeTooltip", "Pushes parameter blocks that can be used by all nodes leaf-wards");
+	return LOCTEXT("NodeTooltip", "Pushes parameters that can be used by all nodes leaf-wards");
 }
 
 FText UAnimGraphNode_AnimNextParameters::GetMenuCategory() const
diff --git a/Engine/Plugins/Experimental/AnimNext/Source/AnimNextUncookedOnly/Private/Graph/AnimNextGraphEntry.cpp b/Engine/Plugins/Experimental/AnimNext/Source/AnimNextUncookedOnly/Private/Graph/AnimNextGraphEntry.cpp
deleted file mode 100644
diff --git a/Engine/Plugins/Experimental/AnimNext/Source/AnimNextUncookedOnly/Private/Graph/AnimNextGraph_AnimationGraph.cpp b/Engine/Plugins/Experimental/AnimNext/Source/AnimNextUncookedOnly/Private/Graph/AnimNextGraph_AnimationGraph.cpp
new file mode 100644
--- /dev/null
+++ b/Engine/Plugins/Experimental/AnimNext/Source/AnimNextUncookedOnly/Private/Graph/AnimNextGraph_AnimationGraph.cpp
@@ -0,0 +1,59 @@
+﻿// Copyright Epic Games, Inc. All Rights Reserved.
+
+#include "Graph/AnimNextGraph_AnimationGraph.h"
+#include "Graph/AnimNextGraph_EdGraph.h"
+#include "Graph/RigUnit_AnimNextGraphRoot.h"
+
+FName UAnimNextGraph_AnimationGraph::GetEntryName() const
+{
+	return GraphName;
+}
+
+void UAnimNextGraph_AnimationGraph::SetEntryName(FName InName, bool bSetupUndoRedo)
+{
+	if(bSetupUndoRedo)
+	{
+		Modify();
+	};
+
+	GraphName = InName;
+
+	// Forward to entry point node
+	URigVMController* Controller = GetImplementingOuter<IRigVMClientHost>()->GetController(Graph);
+	for(URigVMNode* Node : Graph->GetNodes())
+	{
+		if(URigVMUnitNode* UnitNode = Cast<URigVMUnitNode>(Node))
+		{
+			if(UnitNode->GetScriptStruct() == FRigUnit_AnimNextGraphRoot::StaticStruct())
+			{
+				URigVMPin* EntryPointPin = UnitNode->FindPin(GET_MEMBER_NAME_STRING_CHECKED(FRigUnit_AnimNextGraphRoot, EntryPoint));
+				check(EntryPointPin);
+				check(EntryPointPin->GetDirection() == ERigVMPinDirection::Hidden);
+
+				Controller->SetPinDefaultValue(EntryPointPin->GetPinPath(), InName.ToString());
+			}
+		}
+	}
+	
+	BroadcastModified();
+}
+
+URigVMGraph* UAnimNextGraph_AnimationGraph::GetRigVMGraph() const
+{
+	return Graph;
+}
+
+URigVMEdGraph* UAnimNextGraph_AnimationGraph::GetEdGraph() const
+{
+	return EdGraph;
+}
+
+void UAnimNextGraph_AnimationGraph::SetRigVMGraph(URigVMGraph* InGraph)
+{
+	Graph = InGraph;
+}
+
+void UAnimNextGraph_AnimationGraph::SetEdGraph(URigVMEdGraph* InGraph)
+{
+	EdGraph = CastChecked<UAnimNextGraph_EdGraph>(InGraph);
+}
\ No newline at end of file
diff --git a/Engine/Plugins/Experimental/AnimNext/Source/AnimNextUncookedOnly/Private/Graph/AnimNextGraph_AnimationGraphSchema.cpp b/Engine/Plugins/Experimental/AnimNext/Source/AnimNextUncookedOnly/Private/Graph/AnimNextGraph_AnimationGraphSchema.cpp
new file mode 100644
--- /dev/null
+++ b/Engine/Plugins/Experimental/AnimNext/Source/AnimNextUncookedOnly/Private/Graph/AnimNextGraph_AnimationGraphSchema.cpp
@@ -0,0 +1,44 @@
+﻿// Copyright Epic Games, Inc. All Rights Reserved.
+
+#include "Graph/AnimNextGraph_AnimationGraphSchema.h"
+
+#include "AnimNextExecuteContext.h"
+#include "Graph/RigUnit_AnimNextBase.h"
+#include "Param/RigVMDispatch_GetParameter.h"
+
+bool UAnimNextGraph_AnimationGraphSchema::SupportsUnitFunction(URigVMController* InController, const FRigVMFunction* InUnitFunction) const
+{
+	if(InUnitFunction)
+	{
+		if(const UScriptStruct* FunctionExecuteContextStruct = InUnitFunction->GetExecuteContextStruct())
+		{
+			if(FunctionExecuteContextStruct == FAnimNextExecuteContext::StaticStruct())
+			{
+				// Only allow nodes that are children of FRigUnit_AnimNextBase
+				if(InUnitFunction->Struct)
+				{
+					return InUnitFunction->Struct->IsChildOf(FRigUnit_AnimNextBase::StaticStruct());
+				}
+			}
+		}
+	}
+
+	return Super::SupportsUnitFunction(InController, InUnitFunction);
+}
+
+bool UAnimNextGraph_AnimationGraphSchema::SupportsDispatchFactory(URigVMController* InController, const FRigVMDispatchFactory* InDispatchFactory) const
+{
+	if(InDispatchFactory)
+	{
+		if(const UScriptStruct* DispatchExecuteContextStruct = InDispatchFactory->GetExecuteContextStruct())
+		{
+			if(DispatchExecuteContextStruct == FAnimNextExecuteContext::StaticStruct())
+			{
+				// We only support the FRigVMDispatch_GetParameter at the moment.
+				return (InDispatchFactory->GetScriptStruct() == FRigVMDispatch_GetParameter::StaticStruct());
+			}
+		}
+	}
+	
+	return Super::SupportsDispatchFactory(InController, InDispatchFactory);
+}
\ No newline at end of file
diff --git a/Engine/Plugins/Experimental/AnimNext/Source/AnimNextUncookedOnly/Private/Graph/AnimNextGraph_EdGraphSchema.cpp b/Engine/Plugins/Experimental/AnimNext/Source/AnimNextUncookedOnly/Private/Graph/AnimNextGraph_EdGraphSchema.cpp
--- a/Engine/Plugins/Experimental/AnimNext/Source/AnimNextUncookedOnly/Private/Graph/AnimNextGraph_EdGraphSchema.cpp
+++ b/Engine/Plugins/Experimental/AnimNext/Source/AnimNextUncookedOnly/Private/Graph/AnimNextGraph_EdGraphSchema.cpp
@@ -4,48 +4,9 @@
 
 #include "AnimNextRigVMAsset.h"
 #include "AnimNextRigVMAssetEntry.h"
-#include "Editor.h"
 
 #define LOCTEXT_NAMESPACE "AnimNextGraph_EdGraphSchema"
 
-void UAnimNextGraph_EdGraphSchema::TrySetDefaultValue(UEdGraphPin& InPin, const FString& InNewDefaultValue, bool bMarkAsModified) const
-{
-#if WITH_EDITOR
-	if (GEditor)
-	{
-		GEditor->CancelTransaction(0);
-	}
-#endif
-
-	FString UseDefaultValue;
-	TObjectPtr<UObject> UseDefaultObject = nullptr;
-	FText UseDefaultText;
-
-	GetDefault<UEdGraphSchema_K2>()->GetPinDefaultValuesFromString(InPin.PinType, InPin.GetOwningNodeUnchecked(), InNewDefaultValue, UseDefaultValue, UseDefaultObject, UseDefaultText, /*bPreserveTextIdentity*/false);
-
-	// Check the default value and make it an error if it's bogus
-	if (GetDefault<UEdGraphSchema_K2>()->DefaultValueSimpleValidation(InPin.PinType, InPin.PinName, UseDefaultValue, UseDefaultObject, UseDefaultText))
-	{
-		InPin.DefaultObject = UseDefaultObject;
-		InPin.DefaultValue = UseDefaultValue;
-		InPin.DefaultTextValue = UseDefaultText;
-
-		UEdGraphNode* Node = InPin.GetOwningNode();
-		Node->PinDefaultValueChanged(&InPin);
-
-		// If the default value is manually set then treat it as if the value was reset to default and remove the orphaned InPin
-		if (InPin.bOrphanedPin && InPin.DoesDefaultValueMatchAutogenerated())
-		{
-			Node->PinConnectionListChanged(&InPin);
-		}
-
-		if (bMarkAsModified)
-		{
-			Node->MarkPackageDirty();
-		}
-	}
-}
-
 void UAnimNextGraph_EdGraphSchema::GetGraphDisplayInformation(const UEdGraph& Graph, /*out*/ FGraphDisplayInfo& DisplayInfo) const
 {
 	Super::GetGraphDisplayInformation(Graph, DisplayInfo);
diff --git a/Engine/Plugins/Experimental/AnimNext/Source/AnimNextUncookedOnly/Private/Graph/AnimNextGraph_EditorData.cpp b/Engine/Plugins/Experimental/AnimNext/Source/AnimNextUncookedOnly/Private/Graph/AnimNextGraph_EditorData.cpp
--- a/Engine/Plugins/Experimental/AnimNext/Source/AnimNextUncookedOnly/Private/Graph/AnimNextGraph_EditorData.cpp
+++ b/Engine/Plugins/Experimental/AnimNext/Source/AnimNextUncookedOnly/Private/Graph/AnimNextGraph_EditorData.cpp
@@ -5,9 +5,12 @@
 #include "UncookedOnlyUtils.h"
 #include "Curves/CurveFloat.h"
 #include "Graph/AnimNextGraph.h"
-#include "Graph/AnimNextGraphEntry.h"
+#include "Graph/AnimNextGraph_AnimationGraph.h"
 #include "Graph/AnimNextGraph_EdGraphSchema.h"
-#include "Graph/AnimNextGraph_TraitSchema.h"
+#include "Graph/AnimNextGraph_EventGraph.h"
+#include "Graph/AnimNextGraph_EventGraphSchema.h"
+#include "Graph/AnimNextGraph_Parameter.h"
+#include "Graph/AnimNextGraph_AnimationGraphSchema.h"
 #include "RigVMModel/RigVMFunctionLibrary.h"
 #include "RigVMModel/RigVMNotifications.h"
 #include "RigVMModel/Nodes/RigVMCollapseNode.h"
@@ -16,54 +19,111 @@
 
 void UAnimNextGraph_EditorData::PostLoad()
 {
+	Super::PostLoad();
+
+	auto FindEntryForRigVMGraph = [this](URigVMGraph* InRigVMGraph)
+	{
+		UAnimNextRigVMAssetEntry* FoundEntry = nullptr;
+		for(UAnimNextRigVMAssetEntry* Entry : Entries)
+		{
+			if(IAnimNextRigVMGraphInterface* GraphEntry = Cast<IAnimNextRigVMGraphInterface>(Entry))
+			{
+				if(InRigVMGraph == GraphEntry->GetRigVMGraph())
+				{
+					FoundEntry = Entry;
+					break;
+				}
+			}
+		}
+		return FoundEntry;
+	};
+
 	if(GetLinkerCustomVersion(FFortniteMainBranchObjectVersion::GUID) < FFortniteMainBranchObjectVersion::AnimNextCombineGraphContexts)
 	{
+		// Must preload entries so their data is populated or we cannot find the appropriate entries for graphs
+		for(UAnimNextRigVMAssetEntry* Entry : Entries) 
+		{
+			Entry->GetLinker()->Preload(Entry);
+		}
+
 		TArray<URigVMGraph*> AllModels = RigVMClient.GetAllModels(false, true);
 		for(URigVMGraph* Graph : AllModels)
 		{
 			Graph->SetExecuteContextStruct(FAnimNextExecuteContext::StaticStruct());
-			Graph->SetSchemaClass(UAnimNextGraph_TraitSchema::StaticClass());
+			if(UAnimNextRigVMAssetEntry* FoundEntry = FindEntryForRigVMGraph(Graph))
+			{
+				if(FoundEntry->IsA(UAnimNextGraph_AnimationGraph::StaticClass()))
+				{
+					Graph->SetSchemaClass(UAnimNextGraph_AnimationGraphSchema::StaticClass());
+				}
+				else
+				{
+					Graph->SetSchemaClass(UAnimNextGraph_EventGraphSchema::StaticClass());
+				}
+			}
+			else
+			{
+				Graph->SetSchemaClass(UAnimNextGraph_AnimationGraphSchema::StaticClass());
+			}
 		}
 	}
 
-	Super::PostLoad();
-	
 	if(GetLinkerCustomVersion(FFortniteMainBranchObjectVersion::GUID) < FFortniteMainBranchObjectVersion::AnimNextMoveGraphsToEntries)
 	{
+		// Must preload entries so their data is populated or we cannot find the appropriate entries for graphs
+		for(UAnimNextRigVMAssetEntry* Entry : Entries) 
+		{
+			Entry->GetLinker()->Preload(Entry);
+		}
+		
 		for(TObjectPtr<UAnimNextGraph_EdGraph> Graph : Graphs_DEPRECATED)
 		{
-			// Must preload entries so their data is populated or we cannot find the appropriate entries for graphs
-			for(UAnimNextRigVMAssetEntry* Entry : Entries) 
-			{
-				Entry->GetLinker()->Preload(Entry);
-			}
-			
 			URigVMGraph* FoundRigVMGraph = GetRigVMGraphForEditorObject(Graph);
 			if(FoundRigVMGraph)
 			{
-				UAnimNextRigVMAssetEntry* FoundEntry = nullptr;
-				for(UAnimNextRigVMAssetEntry* Entry : Entries)
+				if(UAnimNextRigVMAssetEntry* FoundEntry = FindEntryForRigVMGraph(FoundRigVMGraph))
 				{
-					if(UAnimNextGraphEntry* GraphEntry = Cast<UAnimNextGraphEntry>(Entry))
+					if(UAnimNextGraph_AnimationGraph* AnimationGraphEntry = Cast<UAnimNextGraph_AnimationGraph>(FoundEntry))
 					{
-						if(FoundRigVMGraph == static_cast<IAnimNextRigVMGraphInterface*>(GraphEntry)->GetRigVMGraph())
-						{
-							FoundEntry = Entry;
-							break;
-						}
+						AnimationGraphEntry->EdGraph = Graph;
+					}
+					else if(UAnimNextGraph_EventGraph* EventGraphEntry = Cast<UAnimNextGraph_EventGraph>(FoundEntry))
+					{
+						EventGraphEntry->EdGraph = Graph;
 					}
-				}
-
-				if(FoundEntry)
-				{
-					UAnimNextGraphEntry* GraphEntry = CastChecked<UAnimNextGraphEntry>(FoundEntry);
-					GraphEntry->EdGraph = Graph;
 
 					Graph->Rename(nullptr, FoundEntry, REN_ForceNoResetLoaders | REN_DoNotDirty | REN_DontCreateRedirectors | REN_NonTransactional);
 					Graph->Initialize(this);
 				}
 			}
 		}
+
+		// We used to add a default model that is no longer needed
+		URigVMGraph* DefaultModel = RigVMClient.GetDefaultModel();
+		if(DefaultModel && DefaultModel->GetName() == TEXT("RigVMGraph"))
+		{
+			bool bFound = false;
+			for(UAnimNextRigVMAssetEntry* Entry : Entries)
+			{
+				if(UAnimNextGraph_EventGraph* EventGraphEntry = Cast<UAnimNextGraph_EventGraph>(Entry))
+				{
+					if(DefaultModel == static_cast<IAnimNextRigVMGraphInterface*>(EventGraphEntry)->GetRigVMGraph())
+					{
+						bFound = true;
+						break;
+					}
+				}
+			}
+
+			if(!bFound)
+			{
+				TGuardValue<bool> DisablePythonPrint(bSuspendPythonMessagesForRigVMClient, false);
+				TGuardValue<bool> DisableAutoCompile(bAutoRecompileVM, false);
+				RigVMClient.RemoveModel(DefaultModel->GetNodePath(), false);
+			}
+		}
+
+		RecompileVM();
 	}
 }
 
@@ -97,7 +157,9 @@ TConstArrayView<TSubclassOf<UAnimNextRigVMAssetEntry>> UAnimNextGraph_EditorData
 {
 	static const TSubclassOf<UAnimNextRigVMAssetEntry> Classes[] =
 	{
-		UAnimNextGraphEntry::StaticClass(),
+		UAnimNextGraph_AnimationGraph::StaticClass(),
+		UAnimNextGraph_EventGraph::StaticClass(),
+		UAnimNextGraph_Parameter::StaticClass(),
 	};
 	
 	return Classes;
@@ -134,14 +196,14 @@ UEdGraph* UAnimNextGraph_EditorData::CreateEdGraph(URigVMGraph* InRigVMGraph, bo
 		return nullptr;
 	}
 
-	UAnimNextGraphEntry* Entry = Cast<UAnimNextGraphEntry>(FindEntryForRigVMGraph(InRigVMGraph));
+	IAnimNextRigVMGraphInterface* Entry = Cast<IAnimNextRigVMGraphInterface>(FindEntryForRigVMGraph(InRigVMGraph));
 	if(Entry == nullptr)
 	{
 		// Not found, we could be adding a new entry, in which case the graph wont be assigned yet
 		check(Entries.Num() > 0);
 		check(Cast<IAnimNextRigVMGraphInterface>(Entries.Last()) != nullptr);
 		check(Cast<IAnimNextRigVMGraphInterface>(Entries.Last())->GetRigVMGraph() == nullptr);
-		Entry = Cast<UAnimNextGraphEntry>(FindEntryForRigVMGraph(nullptr));
+		Entry = Cast<IAnimNextRigVMGraphInterface>(FindEntryForRigVMGraph(nullptr));
 	}
 
 	if(Entry == nullptr)
@@ -157,21 +219,21 @@ UEdGraph* UAnimNextGraph_EditorData::CreateEdGraph(URigVMGraph* InRigVMGraph, bo
 	FString GraphName = InRigVMGraph->GetName();
 	check(!GraphName.IsEmpty());
 
-	UAnimNextGraph_EdGraph* RigFunctionGraph = NewObject<UAnimNextGraph_EdGraph>(Entry, NAME_None, RF_Transactional);
+	UAnimNextGraph_EdGraph* RigFunctionGraph = NewObject<UAnimNextGraph_EdGraph>(CastChecked<UObject>(Entry), NAME_None, RF_Transactional);
 	RigFunctionGraph->Schema = UAnimNextGraph_EdGraphSchema::StaticClass();
 	RigFunctionGraph->bAllowDeletion = true;
 	RigFunctionGraph->bIsFunctionDefinition = false;
 	RigFunctionGraph->ModelNodePath = InRigVMGraph->GetNodePath();
 	RigFunctionGraph->Initialize(this);
 
-	Entry->EdGraph = RigFunctionGraph;
-	if(Entry->Graph == nullptr)
+	Entry->SetEdGraph(RigFunctionGraph);
+	if(Entry->GetRigVMGraph() == nullptr)
 	{
-		Entry->Graph = InRigVMGraph;
+		Entry->SetRigVMGraph(InRigVMGraph);
 	}
 	else
 	{
-		check(Entry->Graph == InRigVMGraph);
+		check(Entry->GetRigVMGraph() == InRigVMGraph);
 	}
 
 	return RigFunctionGraph;
@@ -179,7 +241,7 @@ UEdGraph* UAnimNextGraph_EditorData::CreateEdGraph(URigVMGraph* InRigVMGraph, bo
 
 bool UAnimNextGraph_EditorData::RemoveEdGraph(URigVMGraph* InModel)
 {
-	if(UAnimNextGraphEntry* Entry = Cast<UAnimNextGraphEntry>(FindEntryForRigVMGraph(InModel)))
+	if(UAnimNextGraph_AnimationGraph* Entry = Cast<UAnimNextGraph_AnimationGraph>(FindEntryForRigVMGraph(InModel)))
 	{
 		RigVMClient.DestroyObject(Entry->EdGraph);
 		Entry->EdGraph = nullptr;
@@ -188,16 +250,119 @@ bool UAnimNextGraph_EditorData::RemoveEdGraph(URigVMGraph* InModel)
 	return false;
 }
 
-UAnimNextGraphEntry* UAnimNextGraphLibrary::AddGraph(UAnimNextGraph* InGraph, FName InName, bool bSetupUndoRedo, bool bPrintPythonCommand)
+UAnimNextGraph_Parameter* UAnimNextGraphLibrary::AddParameter(UAnimNextGraph* InGraph, FName InName, EPropertyBagPropertyType InValueType,
+	EPropertyBagContainerType InContainerType, const UObject* InValueTypeObject, bool bSetupUndoRedo, bool bPrintPythonCommand)
+{
+	return UE::AnimNext::UncookedOnly::FUtils::GetEditorData(InGraph)->AddParameter(InName, FAnimNextParamType(InValueType, InContainerType, InValueTypeObject), bSetupUndoRedo, bPrintPythonCommand);
+}
+
+UAnimNextGraph_Parameter* UAnimNextGraph_EditorData::AddParameter(FName InName, FAnimNextParamType InType, bool bSetupUndoRedo, bool bPrintPythonCommand)
+{
+	if(InName == NAME_None)
+	{
+		ReportError(TEXT("UAnimNextGraph_EditorData::AddParameter: Invalid parameter name supplied."));
+		return nullptr;
+	}
+
+	// Check for duplicate parameter
+	const bool bAlreadyExists = Entries.ContainsByPredicate([InName](const UAnimNextRigVMAssetEntry* InEntry)
+	{
+		if(const UAnimNextGraph_Parameter* Parameter = Cast<UAnimNextGraph_Parameter>(InEntry))
+		{
+			return Parameter->ParameterName == InName;
+		}
+		return false;
+	});
+
+	if(bAlreadyExists)
+	{
+		ReportError(TEXT("UAnimNextGraph_EditorData::AddParameter: A parameter already exists for the supplied parameter name."));
+		return nullptr;
+	}
+
+	UAnimNextGraph_Parameter* NewEntry = CreateNewSubEntry<UAnimNextGraph_Parameter>(this);
+	NewEntry->ParameterName = InName;
+	NewEntry->Type = InType;
+
+	if(bSetupUndoRedo)
+	{
+		NewEntry->Modify();
+		Modify();
+	}
+	
+	Entries.Add(NewEntry);
+
+	BroadcastModified();
+
+	return NewEntry;
+}
+
+UAnimNextGraph_EventGraph* UAnimNextGraphLibrary::AddEventGraph(UAnimNextGraph* InGraph, FName InName, bool bSetupUndoRedo, bool bPrintPythonCommand)
+{
+	return UE::AnimNext::UncookedOnly::FUtils::GetEditorData(InGraph)->AddEventGraph(InName, bSetupUndoRedo, bPrintPythonCommand);
+}
+
+UAnimNextGraph_EventGraph* UAnimNextGraph_EditorData::AddEventGraph(FName InName, bool bSetupUndoRedo, bool bPrintPythonCommand)
+{
+	if(InName == NAME_None)
+	{
+		ReportError(TEXT("UAnimNextGraph_EditorData::AddEventGraph: Invalid graph name supplied."));
+		return nullptr;
+	}
+
+	// Check for duplicate name
+	FName NewGraphName = InName;
+	auto DuplicateNamePredicate = [&NewGraphName](const UAnimNextRigVMAssetEntry* InEntry)
+	{
+		return InEntry->GetEntryName() == NewGraphName;
+	};
+
+	bool bAlreadyExists = Entries.ContainsByPredicate(DuplicateNamePredicate);
+	int32 NameNumber = InName.GetNumber() + 1;
+	while(bAlreadyExists)
+	{
+		NewGraphName = FName(InName, NameNumber++);
+		bAlreadyExists =  Entries.ContainsByPredicate(DuplicateNamePredicate);
+	}
+
+	UAnimNextGraph_EventGraph* NewEntry = CreateNewSubEntry<UAnimNextGraph_EventGraph>(this);
+	NewEntry->GraphName = NewGraphName;
+
+	if(bSetupUndoRedo)
+	{
+		NewEntry->Modify();
+		Modify();
+	}
+
+	Entries.Add(NewEntry);
+
+	// Add new graph
+	{
+		TGuardValue<bool> EnablePythonPrint(bSuspendPythonMessagesForRigVMClient, !bPrintPythonCommand);
+		TGuardValue<bool> DisableAutoCompile(bAutoRecompileVM, false);
+		URigVMGraph* NewGraph = RigVMClient.AddModel(URigVMGraph::StaticClass()->GetFName(), UAnimNextGraph_EventGraphSchema::StaticClass(), bSetupUndoRedo);
+		ensure(NewGraph);
+		NewEntry->Graph = NewGraph;
+
+		URigVMController* Controller = RigVMClient.GetController(NewGraph);
+		UE::AnimNext::UncookedOnly::FUtils::SetupEventGraph(Controller);
+	}
+
+	BroadcastModified();
+
+	return NewEntry;
+}
+
+UAnimNextGraph_AnimationGraph* UAnimNextGraphLibrary::AddAnimationGraph(UAnimNextGraph* InGraph, FName InName, bool bSetupUndoRedo, bool bPrintPythonCommand)
 {
-	return UE::AnimNext::UncookedOnly::FUtils::GetEditorData(InGraph)->AddGraph(InName, bSetupUndoRedo, bPrintPythonCommand);
+	return UE::AnimNext::UncookedOnly::FUtils::GetEditorData(InGraph)->AddAnimationGraph(InName, bSetupUndoRedo, bPrintPythonCommand);
 }
 
-UAnimNextGraphEntry* UAnimNextGraph_EditorData::AddGraph(FName InName, bool bSetupUndoRedo, bool bPrintPythonCommand)
+UAnimNextGraph_AnimationGraph* UAnimNextGraph_EditorData::AddAnimationGraph(FName InName, bool bSetupUndoRedo, bool bPrintPythonCommand)
 {
 	if(InName == NAME_None)
 	{
-		ReportError(TEXT("UAnimNextGraph_EditorData::AddGraph: Invalid graph name supplied."));
+		ReportError(TEXT("UAnimNextGraph_EditorData::AddAnimationGraph: Invalid graph name supplied."));
 		return nullptr;
 	}
 
@@ -216,7 +381,7 @@ UAnimNextGraphEntry* UAnimNextGraph_EditorData::AddGraph(FName InName, bool bSet
 		bAlreadyExists =  Entries.ContainsByPredicate(DuplicateNamePredicate);
 	}
 
-	UAnimNextGraphEntry* NewEntry = CreateNewSubEntry<UAnimNextGraphEntry>(this);
+	UAnimNextGraph_AnimationGraph* NewEntry = CreateNewSubEntry<UAnimNextGraph_AnimationGraph>(this);
 	NewEntry->GraphName = NewGraphName;
 
 	if(bSetupUndoRedo)
@@ -231,7 +396,7 @@ UAnimNextGraphEntry* UAnimNextGraph_EditorData::AddGraph(FName InName, bool bSet
 	{
 		TGuardValue<bool> EnablePythonPrint(bSuspendPythonMessagesForRigVMClient, !bPrintPythonCommand);
 		TGuardValue<bool> DisableAutoCompile(bAutoRecompileVM, false);
-		URigVMGraph* NewGraph = RigVMClient.AddModel(URigVMGraph::StaticClass()->GetFName(), UAnimNextGraph_TraitSchema::StaticClass(), bSetupUndoRedo);
+		URigVMGraph* NewGraph = RigVMClient.AddModel(URigVMGraph::StaticClass()->GetFName(), UAnimNextGraph_AnimationGraphSchema::StaticClass(), bSetupUndoRedo);
 		ensure(NewGraph);
 		NewEntry->Graph = NewGraph;
 
diff --git a/Engine/Plugins/Experimental/AnimNext/Source/AnimNextUncookedOnly/Private/Graph/AnimNextGraph_EventGraph.cpp b/Engine/Plugins/Experimental/AnimNext/Source/AnimNextUncookedOnly/Private/Graph/AnimNextGraph_EventGraph.cpp
new file mode 100644
--- /dev/null
+++ b/Engine/Plugins/Experimental/AnimNext/Source/AnimNextUncookedOnly/Private/Graph/AnimNextGraph_EventGraph.cpp
@@ -0,0 +1,46 @@
+﻿// Copyright Epic Games, Inc. All Rights Reserved.
+
+#include "Graph/AnimNextGraph_EventGraph.h"
+#include "Graph/AnimNextGraph_EdGraph.h"
+
+FText UAnimNextGraph_EventGraph::GetDisplayName() const
+{
+	return FText::FromName(GraphName);
+}
+
+FText UAnimNextGraph_EventGraph::GetDisplayNameTooltip() const
+{
+	return FText::FromName(GraphName);
+}
+
+void UAnimNextGraph_EventGraph::SetEntryName(FName InName, bool bSetupUndoRedo)
+{
+	if(bSetupUndoRedo)
+	{
+		Modify();
+	}
+
+	GraphName = InName;
+
+	BroadcastModified();
+}
+
+URigVMGraph* UAnimNextGraph_EventGraph::GetRigVMGraph() const
+{
+	return Graph;
+}
+
+URigVMEdGraph* UAnimNextGraph_EventGraph::GetEdGraph() const
+{
+	return EdGraph;
+}
+
+void UAnimNextGraph_EventGraph::SetRigVMGraph(URigVMGraph* InGraph)
+{
+	Graph = InGraph;
+}
+
+void UAnimNextGraph_EventGraph::SetEdGraph(URigVMEdGraph* InGraph)
+{
+	EdGraph = CastChecked<UAnimNextGraph_EdGraph>(InGraph);
+}
\ No newline at end of file
diff --git a/Engine/Plugins/Experimental/AnimNext/Source/AnimNextUncookedOnly/Private/Graph/AnimNextGraph_EventGraphSchema.cpp b/Engine/Plugins/Experimental/AnimNext/Source/AnimNextUncookedOnly/Private/Graph/AnimNextGraph_EventGraphSchema.cpp
new file mode 100644
--- /dev/null
+++ b/Engine/Plugins/Experimental/AnimNext/Source/AnimNextUncookedOnly/Private/Graph/AnimNextGraph_EventGraphSchema.cpp
@@ -0,0 +1,23 @@
+﻿// Copyright Epic Games, Inc. All Rights Reserved.
+
+#include "Graph/AnimNextGraph_EventGraphSchema.h"
+
+#include "AnimNextExecuteContext.h"
+#include "Param/RigUnit_AnimNextParameterBase.h"
+
+bool UAnimNextGraph_EventGraphSchema::SupportsUnitFunction(URigVMController* InController, const FRigVMFunction* InUnitFunction) const
+{
+	if(const UScriptStruct* FunctionExecuteContextStruct = InUnitFunction->GetExecuteContextStruct())
+	{
+		if(FunctionExecuteContextStruct == FAnimNextExecuteContext::StaticStruct())
+		{
+			// Only allow nodes that are children of FRigUnit_AnimNextParameterBase
+			if(InUnitFunction->Struct)
+			{
+				return InUnitFunction->Struct->IsChildOf(FRigUnit_AnimNextParameterBase::StaticStruct());
+			}
+		}
+	}
+
+	return Super::SupportsUnitFunction(InController, InUnitFunction);
+}
diff --git a/Engine/Plugins/Experimental/AnimNext/Source/AnimNextUncookedOnly/Private/Graph/AnimNextGraph_Parameter.cpp b/Engine/Plugins/Experimental/AnimNext/Source/AnimNextUncookedOnly/Private/Graph/AnimNextGraph_Parameter.cpp
new file mode 100644
--- /dev/null
+++ b/Engine/Plugins/Experimental/AnimNext/Source/AnimNextUncookedOnly/Private/Graph/AnimNextGraph_Parameter.cpp
@@ -0,0 +1,75 @@
+// Copyright Epic Games, Inc. All Rights Reserved.
+
+#include "Graph/AnimNextGraph_Parameter.h"
+#include "Graph/AnimNextGraph.h"
+#include "Param/ExternalParameterRegistry.h"
+
+FAnimNextParamType UAnimNextGraph_Parameter::GetParamType() const
+{
+	using namespace UE::AnimNext;
+
+	// Look in built-in parameters
+	IParameterSourceFactory::FParameterInfo Info;
+	if(FExternalParameterRegistry::FindParameterInfo(ParameterName, Info))
+	{
+		return Info.Type;
+	}
+
+	return Type;
+}
+
+FName UAnimNextGraph_Parameter::GetEntryName() const
+{
+	return ParameterName;
+}
+
+bool UAnimNextGraph_Parameter::SetParamType(const FAnimNextParamType& InType, bool bSetupUndoRedo)
+{
+	if(bSetupUndoRedo)
+	{
+		Modify();
+	}
+	
+	Type = InType;
+
+	BroadcastModified();
+
+	return true;
+}
+
+FInstancedPropertyBag& UAnimNextGraph_Parameter::GetPropertyBag() const
+{
+	// TODO: move property bag for defaults onto this entry!
+	UAnimNextGraph* Asset = GetTypedOuter<UAnimNextGraph>();
+	check(Asset);
+	return Asset->PropertyBag;
+}
+
+void UAnimNextGraph_Parameter::SetEntryName(FName InName, bool bSetupUndoRedo)
+{
+	if(bSetupUndoRedo)
+	{
+		Modify();
+	}
+
+	ParameterName = InName;
+	BroadcastModified();
+}
+
+FText UAnimNextGraph_Parameter::GetDisplayName() const
+{
+	return FText::FromName(ParameterName);
+}
+
+FText UAnimNextGraph_Parameter::GetDisplayNameTooltip() const
+{
+	using namespace UE::AnimNext;
+
+	IParameterSourceFactory::FParameterInfo Info;
+	if(FExternalParameterRegistry::FindParameterInfo(ParameterName, Info))
+	{
+		return Info.Tooltip;
+	}
+	
+	return FText::FromName(ParameterName);
+}
\ No newline at end of file
diff --git a/Engine/Plugins/Experimental/AnimNext/Source/AnimNextUncookedOnly/Private/Graph/AnimNextGraph_TraitSchema.cpp b/Engine/Plugins/Experimental/AnimNext/Source/AnimNextUncookedOnly/Private/Graph/AnimNextGraph_TraitSchema.cpp
deleted file mode 100644
diff --git a/Engine/Plugins/Experimental/AnimNext/Source/AnimNextUncookedOnly/Private/Param/AnimNextParameterBlockGraph.cpp b/Engine/Plugins/Experimental/AnimNext/Source/AnimNextUncookedOnly/Private/Param/AnimNextParameterBlockGraph.cpp
deleted file mode 100644
diff --git a/Engine/Plugins/Experimental/AnimNext/Source/AnimNextUncookedOnly/Private/Param/AnimNextParameterBlockParameter.cpp b/Engine/Plugins/Experimental/AnimNext/Source/AnimNextUncookedOnly/Private/Param/AnimNextParameterBlockParameter.cpp
deleted file mode 100644
diff --git a/Engine/Plugins/Experimental/AnimNext/Source/AnimNextUncookedOnly/Private/Param/AnimNextParameterBlock_EdGraph.cpp b/Engine/Plugins/Experimental/AnimNext/Source/AnimNextUncookedOnly/Private/Param/AnimNextParameterBlock_EdGraph.cpp
deleted file mode 100644
diff --git a/Engine/Plugins/Experimental/AnimNext/Source/AnimNextUncookedOnly/Private/Param/AnimNextParameterBlock_EdGraphSchema.cpp b/Engine/Plugins/Experimental/AnimNext/Source/AnimNextUncookedOnly/Private/Param/AnimNextParameterBlock_EdGraphSchema.cpp
deleted file mode 100644
diff --git a/Engine/Plugins/Experimental/AnimNext/Source/AnimNextUncookedOnly/Private/Param/AnimNextParameterBlock_EditorData.cpp b/Engine/Plugins/Experimental/AnimNext/Source/AnimNextUncookedOnly/Private/Param/AnimNextParameterBlock_EditorData.cpp
deleted file mode 100644
diff --git a/Engine/Plugins/Experimental/AnimNext/Source/AnimNextUncookedOnly/Private/Param/AnimNextParameterBlock_ParamSchema.cpp b/Engine/Plugins/Experimental/AnimNext/Source/AnimNextUncookedOnly/Private/Param/AnimNextParameterBlock_ParamSchema.cpp
deleted file mode 100644
diff --git a/Engine/Plugins/Experimental/AnimNext/Source/AnimNextUncookedOnly/Private/UncookedOnlyUtils.cpp b/Engine/Plugins/Experimental/AnimNext/Source/AnimNextUncookedOnly/Private/UncookedOnlyUtils.cpp
--- a/Engine/Plugins/Experimental/AnimNext/Source/AnimNextUncookedOnly/Private/UncookedOnlyUtils.cpp
+++ b/Engine/Plugins/Experimental/AnimNext/Source/AnimNextUncookedOnly/Private/UncookedOnlyUtils.cpp
@@ -14,8 +14,6 @@
 #include "Graph/RigUnit_AnimNextGraphRoot.h"
 #include "Graph/RigUnit_AnimNextGraphEvaluator.h"
 #include "Graph/RigUnit_AnimNextShimRoot.h"
-#include "Param/AnimNextParameterBlock.h"
-#include "Param/AnimNextParameterBlock_EditorData.h"
 #include "Graph/RigDecorator_AnimNextCppTrait.h"
 #include "Param/RigUnit_AnimNextParameterBeginExecution.h"
 #include "Param/RigVMDispatch_GetParameter.h"
@@ -45,6 +43,7 @@
 #include "AnimNextRigVMAssetEditorData.h"
 #include "AnimNextRigVMAssetEntry.h"
 #include "Graph/AnimNextGraphEntryPoint.h"
+#include "Graph/AnimNextGraph_AnimationGraphSchema.h"
 
 namespace UE::AnimNext::UncookedOnly
 {
@@ -492,6 +491,14 @@ void FUtils::Compile(UAnimNextGraph* InGraph)
 {
 	check(InGraph);
 
+	CompileStruct(InGraph);
+	CompileVM(InGraph);
+}
+
+void FUtils::CompileVM(UAnimNextGraph* InGraph)
+{
+	check(InGraph);
+
 	UAnimNextGraph_EditorData* EditorData = GetEditorData(InGraph);
 
 	if (EditorData->bIsCompiling)
@@ -500,8 +507,6 @@ void FUtils::Compile(UAnimNextGraph* InGraph)
 	}
 
 	TGuardValue<bool> CompilingGuard(EditorData->bIsCompiling, true);
-
-
 	
 	// Before we re-compile a graph, we need to release and live instances since we need the metadata we are about to replace
 	// to call trait destructors etc
@@ -537,20 +542,21 @@ void FUtils::Compile(UAnimNextGraph* InGraph)
 		return;
 	}
 
-	TArray<URigVMGraph*> VMTempGraphs;
+	TArray<URigVMGraph*> TempGraphs;
 	for(const URigVMGraph* SourceGraph : VMClient->GetAllModels(false, false))
 	{
 		// We use a temporary graph models to build our final graphs that we'll compile
-		URigVMGraph* VMTempGraph = CastChecked<URigVMGraph>(StaticDuplicateObject(SourceGraph, GetTransientPackage(), NAME_None, RF_Transient));
-		VMTempGraphs.Add(VMTempGraph);
+		URigVMGraph* TempGraph = CastChecked<URigVMGraph>(StaticDuplicateObject(SourceGraph, GetTransientPackage(), NAME_None, RF_Transient));
+		TempGraph->SetFlags(RF_Transient);
+		TempGraphs.Add(TempGraph);
 	}
 
-	if(VMTempGraphs.Num() == 0)
+	if(TempGraphs.Num() == 0)
 	{
 		return;
 	}
 
-	UAnimNextGraph_Controller* TempController = CastChecked<UAnimNextGraph_Controller>(VMClient->GetOrCreateController(VMTempGraphs[0]));
+	UAnimNextGraph_Controller* TempController = CastChecked<UAnimNextGraph_Controller>(VMClient->GetOrCreateController(TempGraphs[0]));
 
 	FTraitWriter TraitWriter;
 
@@ -559,53 +565,59 @@ void FUtils::Compile(UAnimNextGraph* InGraph)
 	TArray<Private::FTraitGraph> TraitGraphs;
 
 	// Build entry points and extract their required latent pins
-	for(const URigVMGraph* VMTempGraph : VMTempGraphs)
+	for(const URigVMGraph* TempGraph : TempGraphs)
 	{
-		// Gather our trait stacks
-		Private::FTraitGraph& TraitGraph = TraitGraphs.Add_GetRef(Private::CollectGraphInfo(VMTempGraph, TempController->GetControllerForGraph(VMTempGraph)));
-		check(!TraitGraph.TraitStackNodes.IsEmpty());
+		if(TempGraph->GetSchemaClass() == UAnimNextGraph_AnimationGraphSchema::StaticClass())
+		{
+			// Gather our trait stacks
+			Private::FTraitGraph& TraitGraph = TraitGraphs.Add_GetRef(Private::CollectGraphInfo(TempGraph, TempController->GetControllerForGraph(TempGraph)));
+			check(!TraitGraph.TraitStackNodes.IsEmpty());
 
-		FAnimNextGraphEntryPoint& EntryPoint = InGraph->EntryPoints.AddDefaulted_GetRef();
-		EntryPoint.EntryPointName = TraitGraph.EntryPoint;
+			FAnimNextGraphEntryPoint& EntryPoint = InGraph->EntryPoints.AddDefaulted_GetRef();
+			EntryPoint.EntryPointName = TraitGraph.EntryPoint;
 
-		// Extract latent pins for this graph
-		Private::CollectLatentPins(TraitGraph.TraitStackNodes, LatentPins, LatentPinMapping);
+			// Extract latent pins for this graph
+			Private::CollectLatentPins(TraitGraph.TraitStackNodes, LatentPins, LatentPinMapping);
 
-		// Iterate over every trait stack and register our node templates
-		for (Private::FTraitStackMapping& NodeMapping : TraitGraph.TraitStackNodes)
-		{
-			NodeMapping.TraitStackNodeHandle = Private::RegisterTraitNodeTemplate(TraitWriter, NodeMapping.DecoratorStackNode);
-		}
+			// Iterate over every trait stack and register our node templates
+			for (Private::FTraitStackMapping& NodeMapping : TraitGraph.TraitStackNodes)
+			{
+				NodeMapping.TraitStackNodeHandle = Private::RegisterTraitNodeTemplate(TraitWriter, NodeMapping.DecoratorStackNode);
+			}
 
-		// Find our root node handle, if we have any stack nodes, the first one is our root stack
-		if (TraitGraph.TraitStackNodes.Num() != 0)
-		{
-			EntryPoint.RootTraitHandle = FAnimNextEntryPointHandle(TraitGraph.TraitStackNodes[0].TraitStackNodeHandle);
+			// Find our root node handle, if we have any stack nodes, the first one is our root stack
+			if (TraitGraph.TraitStackNodes.Num() != 0)
+			{
+				EntryPoint.RootTraitHandle = FAnimNextEntryPointHandle(TraitGraph.TraitStackNodes[0].TraitStackNodeHandle);
+			}
 		}
 	}
 
-	// We need a unique method name to match our unique argument list
-	InGraph->ExecuteDefinition = Private::GetGraphEvaluatorExecuteMethod(LatentPins);
-
-	// Add our runtime shim root node
-	URigVMUnitNode* TempShimRootNode = TempController->AddUnitNode(FRigUnit_AnimNextShimRoot::StaticStruct(), FRigUnit_AnimNextShimRoot::EventName, FVector2D::ZeroVector, FString(), false);
-	URigVMUnitNode* GraphEvaluatorNode = TempController->AddUnitNodeWithPins(FRigUnit_AnimNextGraphEvaluator::StaticStruct(), LatentPins, *InGraph->ExecuteDefinition.MethodName, FVector2D::ZeroVector, FString(), false);
+	if(LatentPins.Num() > 0)
+	{
+		// We need a unique method name to match our unique argument list
+		InGraph->ExecuteDefinition = Private::GetGraphEvaluatorExecuteMethod(LatentPins);
 
-	// Link our shim and evaluator nodes together using the execution context
-	TempController->AddLink(
-		TempShimRootNode->FindPin(GET_MEMBER_NAME_STRING_CHECKED(FRigUnit_AnimNextShimRoot, ExecuteContext)),
-		GraphEvaluatorNode->FindPin(GET_MEMBER_NAME_STRING_CHECKED(FRigUnit_AnimNextGraphEvaluator, ExecuteContext)),
-		false);
+		// Add our runtime shim root node
+		URigVMUnitNode* TempShimRootNode = TempController->AddUnitNode(FRigUnit_AnimNextShimRoot::StaticStruct(), FRigUnit_AnimNextShimRoot::EventName, FVector2D::ZeroVector, FString(), false);
+		URigVMUnitNode* GraphEvaluatorNode = TempController->AddUnitNodeWithPins(FRigUnit_AnimNextGraphEvaluator::StaticStruct(), LatentPins, *InGraph->ExecuteDefinition.MethodName, FVector2D::ZeroVector, FString(), false);
 
-	// Link our latent pins
-	for (const FRigVMPinInfo& LatentPin : LatentPins)
-	{
+		// Link our shim and evaluator nodes together using the execution context
 		TempController->AddLink(
-			LatentPinMapping[LatentPin.Name],
-			GraphEvaluatorNode->FindPin(LatentPin.Name.ToString()),
+			TempShimRootNode->FindPin(GET_MEMBER_NAME_STRING_CHECKED(FRigUnit_AnimNextShimRoot, ExecuteContext)),
+			GraphEvaluatorNode->FindPin(GET_MEMBER_NAME_STRING_CHECKED(FRigUnit_AnimNextGraphEvaluator, ExecuteContext)),
 			false);
+
+		// Link our latent pins
+		for (const FRigVMPinInfo& LatentPin : LatentPins)
+		{
+			TempController->AddLink(
+				LatentPinMapping[LatentPin.Name],
+				GraphEvaluatorNode->FindPin(LatentPin.Name.ToString()),
+				false);
+		}
 	}
-	
+
 	// Write our node shared data
 	TraitWriter.BeginNodeWriting();
 
@@ -636,10 +648,14 @@ void FUtils::Compile(UAnimNextGraph* InGraph)
 	URigVMCompiler* Compiler = URigVMCompiler::StaticClass()->GetDefaultObject<URigVMCompiler>();
 	EditorData->VMCompileSettings.SetExecuteContextStruct(FAnimNextExecuteContext::StaticStruct());
 	const FRigVMCompileSettings Settings = (EditorData->bCompileInDebugMode) ? FRigVMCompileSettings::Fast(EditorData->VMCompileSettings.GetExecuteContextStruct()) : EditorData->VMCompileSettings;
-	Compiler->Compile(Settings, VMTempGraphs, TempController, InGraph->VM, InGraph->ExtendedExecuteContext, TArray<FRigVMExternalVariable>(), & EditorData->PinToOperandMap);
+	Compiler->Compile(Settings, TempGraphs, TempController, InGraph->VM, InGraph->ExtendedExecuteContext, TArray<FRigVMExternalVariable>(), & EditorData->PinToOperandMap);
 
 	// Initialize right away, in packaged builds we initialize during PostLoad
 	InGraph->VM->Initialize(InGraph->ExtendedExecuteContext);
+	InGraph->GenerateUserDefinedDependenciesData(InGraph->ExtendedExecuteContext);
+
+	// Notable difference with vanilla RigVM host behavior - we init the VM here at the moment as we only have one 'instance'
+	InGraph->InitializeVM(FRigUnit_AnimNextBeginExecution::EventName);
 
 	if (EditorData->bErrorsDuringCompilation)
 	{
@@ -655,9 +671,9 @@ void FUtils::Compile(UAnimNextGraph* InGraph)
 		EditorData->RigVMCompiledEvent.Broadcast(InGraph, InGraph->VM, InGraph->ExtendedExecuteContext);
 	}
 
-	for(URigVMGraph* VMTempGraph : VMTempGraphs)
+	for(URigVMGraph* TempGraph : TempGraphs)
 	{
-		VMClient->RemoveController(VMTempGraph);
+		VMClient->RemoveController(TempGraph);
 	}
 
 	// Now that the graph has been re-compiled, re-allocate the previous live instances
@@ -811,81 +827,11 @@ FParamTypeHandle FUtils::GetParameterHandleFromPin(const FEdGraphPinType& InPinT
 	return FAnimNextParamType(ValueType, ContainerType, ValueTypeObject).GetHandle();
 }
 
-void FUtils::CompileVM(UAnimNextParameterBlock* InParameterBlock)
-{
-	check(InParameterBlock);
-
-	UAnimNextParameterBlock_EditorData* EditorData = FUtils::GetEditorData(InParameterBlock);
-	if(EditorData->bIsCompiling)
-	{
-		return;
-	}
-
-	TGuardValue<bool> CompilingGuard(EditorData->bIsCompiling, true);
-
-	EditorData->bErrorsDuringCompilation = false;
-
-	EditorData->RigGraphDisplaySettings.MinMicroSeconds = EditorData->RigGraphDisplaySettings.LastMinMicroSeconds = DBL_MAX;
-	EditorData->RigGraphDisplaySettings.MaxMicroSeconds = EditorData->RigGraphDisplaySettings.LastMaxMicroSeconds = (double)INDEX_NONE;
-	
-	TGuardValue<bool> ReentrantGuardSelf(EditorData->bSuspendModelNotificationsForSelf, true);
-	TGuardValue<bool> ReentrantGuardOthers(EditorData->bSuspendModelNotificationsForOthers, true);
-
-	RecreateVM(InParameterBlock);
-
-	InParameterBlock->VMRuntimeSettings = EditorData->VMRuntimeSettings;
-
-	EditorData->CompileLog.Messages.Reset();
-	EditorData->CompileLog.NumErrors = EditorData->CompileLog.NumWarnings = 0;
-
-	FRigVMClient* VMClient = EditorData->GetRigVMClient();
-	URigVMGraph* RootGraph = VMClient->GetDefaultModel();
-
-	if(RootGraph == nullptr)
-	{
-		return;
-	}
-
-	URigVMCompiler* Compiler = URigVMCompiler::StaticClass()->GetDefaultObject<URigVMCompiler>();
-	EditorData->VMCompileSettings.SetExecuteContextStruct(FAnimNextExecuteContext::StaticStruct());
-	FRigVMExtendedExecuteContext& CDOContext = InParameterBlock->GetRigVMExtendedExecuteContext();
-	const FRigVMCompileSettings Settings = (EditorData->bCompileInDebugMode) ? FRigVMCompileSettings::Fast(EditorData->VMCompileSettings.GetExecuteContextStruct()) : EditorData->VMCompileSettings;
-	URigVMController* RootController = VMClient->GetOrCreateController(RootGraph);
-	Compiler->Compile(Settings, VMClient->GetAllModels(false, false), RootController, InParameterBlock->VM, CDOContext, InParameterBlock->GetExternalVariables(), &EditorData->PinToOperandMap);
-
-	InParameterBlock->VM->Initialize(CDOContext);
-	InParameterBlock->GenerateUserDefinedDependenciesData(CDOContext);
-
-	// Notable difference with vanilla RigVM host behavior - we init the VM here at the moment as we only have one 'instance'
-	InParameterBlock->InitializeVM(FRigUnit_AnimNextBeginExecution::EventName);
-
-	if (EditorData->bErrorsDuringCompilation)
-	{
-		if(Settings.SurpressErrors)
-		{
-			Settings.Reportf(EMessageSeverity::Info, InParameterBlock,
-				TEXT("Compilation Errors may be suppressed for ControlRigBlueprint: %s. See VM Compile Setting in Class Settings for more Details"), *InParameterBlock->GetName());
-		}
-		EditorData->bVMRecompilationRequired = false;
-		if(InParameterBlock->VM)
-		{
-			EditorData->RigVMCompiledEvent.Broadcast(InParameterBlock, InParameterBlock->VM, InParameterBlock->GetRigVMExtendedExecuteContext());
-		}
-		return;
-	}
-
-//	InitializeArchetypeInstances();
-
-#if WITH_EDITOR
-//	RefreshBreakpoints(EditorData);
-#endif
-}
-
-void FUtils::CompileStruct(UAnimNextParameterBlock* InParameterBlock)
+void FUtils::CompileStruct(UAnimNextGraph* InGraph)
 {
-	check(InParameterBlock);
+	check(InGraph);
 
-	UAnimNextParameterBlock_EditorData* EditorData = GetEditorData(InParameterBlock);
+	UAnimNextGraph_EditorData* EditorData = GetEditorData(InGraph);
 	if(EditorData->bIsCompiling)
 	{
 		return;
@@ -896,7 +842,7 @@ void FUtils::CompileStruct(UAnimNextParameterBlock* InParameterBlock)
 	TArray<FPropertyBagPropertyDesc> PropertyDescs;
 	PropertyDescs.Reserve(EditorData->Entries.Num());
 	
-	// Gather all parameters in this block
+	// Gather all parameters in this asset
 	for(const UAnimNextRigVMAssetEntry* Entry : EditorData->Entries)
 	{
 		if(const IAnimNextRigVMParameterInterface* Binding = Cast<IAnimNextRigVMParameterInterface>(Entry))
@@ -912,9 +858,9 @@ void FUtils::CompileStruct(UAnimNextParameterBlock* InParameterBlock)
 		// find any existing IDs for old properties with name-matching
 		for(FPropertyBagPropertyDesc& NewDesc : PropertyDescs)
 		{
-			if(InParameterBlock->PropertyBag.GetPropertyBagStruct())
+			if(InGraph->PropertyBag.GetPropertyBagStruct())
 			{
-				for(const FPropertyBagPropertyDesc& ExistingDesc : InParameterBlock->PropertyBag.GetPropertyBagStruct()->GetPropertyDescs())
+				for(const FPropertyBagPropertyDesc& ExistingDesc : InGraph->PropertyBag.GetPropertyBagStruct()->GetPropertyDescs())
 				{
 					if(ExistingDesc.Name == NewDesc.Name)
 					{
@@ -927,30 +873,12 @@ void FUtils::CompileStruct(UAnimNextParameterBlock* InParameterBlock)
 
 		// Create new property bag and migrate
 		const UPropertyBag* NewBagStruct = UPropertyBag::GetOrCreateFromDescs(PropertyDescs);
-		InParameterBlock->PropertyBag.MigrateToNewBagStruct(NewBagStruct);
+		InGraph->PropertyBag.MigrateToNewBagStruct(NewBagStruct);
 	}
 	else
 	{
-		InParameterBlock->PropertyBag.Reset();
-	}
-}
-
-void FUtils::Compile(UAnimNextParameterBlock* InParameterBlock)
-{
-	check(InParameterBlock);
-
-	CompileStruct(InParameterBlock);
-	CompileVM(InParameterBlock);
-}
-
-void FUtils::RecreateVM(UAnimNextParameterBlock* InParameterBlock)
-{
-	if(InParameterBlock->VM == nullptr)
-	{
-		InParameterBlock->VM = NewObject<URigVM>(InParameterBlock, TEXT("VM"), RF_NoFlags);
+		InGraph->PropertyBag.Reset();
 	}
-	InParameterBlock->VM->Reset(InParameterBlock->GetRigVMExtendedExecuteContext());
-	InParameterBlock->RigVM = InParameterBlock->VM; // Local serialization
 }
 
 UAnimNextRigVMAsset* FUtils::GetAsset(UAnimNextRigVMAssetEditorData* InEditorData)
@@ -965,23 +893,9 @@ UAnimNextRigVMAssetEditorData* FUtils::GetEditorData(UAnimNextRigVMAsset* InAsse
 	return CastChecked<UAnimNextRigVMAssetEditorData>(InAsset->EditorData);
 }
 
-UAnimNextParameterBlock_EditorData* FUtils::GetEditorData(const UAnimNextParameterBlock* InParameterBlock)
-{
-	check(InParameterBlock);
-
-	return CastChecked<UAnimNextParameterBlock_EditorData>(InParameterBlock->EditorData);
-}
-
-UAnimNextParameterBlock* FUtils::GetBlock(const UAnimNextParameterBlock_EditorData* InEditorData)
-{
-	check(InEditorData);
-
-	return CastChecked<UAnimNextParameterBlock>(InEditorData->GetOuter());
-}
-
-FInstancedPropertyBag* FUtils::GetPropertyBag(UAnimNextParameterBlock* ReferencedBlock)
+FInstancedPropertyBag* FUtils::GetPropertyBag(UAnimNextGraph* InAnimNextGraph)
 {
-	FInstancedPropertyBag* InstancedPropertyBag =&ReferencedBlock->PropertyBag;
+	FInstancedPropertyBag* InstancedPropertyBag = &InAnimNextGraph->PropertyBag;
 
 	return InstancedPropertyBag;
 }
@@ -1280,7 +1194,7 @@ void FUtils::SetupAnimGraph(UAnimNextRigVMAssetEntry* InEntry, URigVMController*
 	InController->SetPinDefaultValue(EntryPointPin->GetPinPath(), InEntry->GetEntryName().ToString());
 }
 
-void FUtils::SetupParameterGraph(URigVMController* InController)
+void FUtils::SetupEventGraph(URigVMController* InController)
 {
 	// Clear the graph
 	InController->RemoveNodes(InController->GetGraph()->GetNodes());
@@ -1878,7 +1792,7 @@ void FUtils::CompileSchedule(UAnimNextSchedule* InSchedule)
 				ParamScopeEntryTask.ParamParentScopeIndex = ParentScopeIndex;
 				ParamScopeEntryTask.TickFunctionIndex = InSchedule->NumTickFunctions;
 				ParamScopeEntryTask.Scope = ParamScopeTaskEntry->Scope;
-				ParamScopeEntryTask.ParameterBlocks = ParamScopeTaskEntry->ParameterBlocks;
+				ParamScopeEntryTask.Parameters = ParamScopeTaskEntry->Parameters;
 				int32 ParamScopeTaskEntryIndex = InSchedule->ParamScopeEntryTasks.Add(ParamScopeEntryTask);
 
 				Emit(EAnimNextScheduleScheduleOpcode::RunParamScopeEntry, ParamScopeTaskEntryIndex);
@@ -2024,12 +1938,12 @@ void FUtils::CompileSchedule(UAnimNextSchedule* InSchedule)
 				}
 				else if (UAnimNextScheduleEntry_ParamScope* ParamScopeTaskEntry = Cast<UAnimNextScheduleEntry_ParamScope>(Entry))
 				{
-					for(UAnimNextParameterBlock* ParameterBlock : ParamScopeTaskEntry->ParameterBlocks)
+					for(UAnimNextGraph* Parameters : ParamScopeTaskEntry->Parameters)
 					{
-						if(ParameterBlock)
+						if(Parameters)
 						{
 							FAnimNextParameterProviderAssetRegistryExports Exports;
-							if(UncookedOnly::FUtils::GetExportedParametersForAsset(FAssetData(ParameterBlock), Exports))
+							if(UncookedOnly::FUtils::GetExportedParametersForAsset(FAssetData(Parameters), Exports))
 							{
 								TArray<FAnimNextParam> RequiredParameters;
 								RequiredParameters.Reserve(Exports.Parameters.Num());
diff --git a/Engine/Plugins/Runtime/RigVM/Source/RigVMDeveloper/Private/RigVMModel/RigVMGraph.cpp b/Engine/Plugins/Runtime/RigVM/Source/RigVMDeveloper/Private/RigVMModel/RigVMGraph.cpp
--- a/Engine/Plugins/Runtime/RigVM/Source/RigVMDeveloper/Private/RigVMModel/RigVMGraph.cpp
+++ b/Engine/Plugins/Runtime/RigVM/Source/RigVMDeveloper/Private/RigVMModel/RigVMGraph.cpp
@@ -407,13 +407,11 @@ TArray<FRigVMGraphVariableDescription> URigVMGraph::GetOutputArguments() const
 
 URigVMSchema* URigVMGraph::GetSchema() const
 {
-	check(SchemaClass);
-	return SchemaClass->GetDefaultObject<URigVMSchema>();
+	return SchemaClass ? SchemaClass->GetDefaultObject<URigVMSchema>() : nullptr;
 }
 
 void URigVMGraph::SetSchemaClass(TSubclassOf<URigVMSchema> InSchemaClass)
 {
-	check(InSchemaClass);
 	SchemaClass = InSchemaClass;
 }
 
diff --git a/Engine/Plugins/Runtime/RigVM/Source/RigVMDeveloper/Public/RigVMModel/RigVMGraph.h b/Engine/Plugins/Runtime/RigVM/Source/RigVMDeveloper/Public/RigVMModel/RigVMGraph.h
--- a/Engine/Plugins/Runtime/RigVM/Source/RigVMDeveloper/Public/RigVMModel/RigVMGraph.h
+++ b/Engine/Plugins/Runtime/RigVM/Source/RigVMDeveloper/Public/RigVMModel/RigVMGraph.h
@@ -217,7 +217,7 @@ private:
 	UPROPERTY()
 	TArray<FRigVMGraphVariableDescription> LocalVariables;
 
-	UPROPERTY(transient, DuplicateTransient)
+	UPROPERTY()
 	TSubclassOf<URigVMSchema> SchemaClass;
 
 	bool IsNameAvailable(const FString& InName) const;
diff --git a/Engine/Source/Runtime/Core/Public/UObject/FortniteMainBranchObjectVersion.h b/Engine/Source/Runtime/Core/Public/UObject/FortniteMainBranchObjectVersion.h
--- a/Engine/Source/Runtime/Core/Public/UObject/FortniteMainBranchObjectVersion.h
+++ b/Engine/Source/Runtime/Core/Public/UObject/FortniteMainBranchObjectVersion.h
@@ -420,6 +420,9 @@ struct FFortniteMainBranchObjectVersion
 		// Combine graph contexts for AnimNext graphs
 		AnimNextCombineGraphContexts,
 
+		// Combine parameter blocks and graphs
+		AnimNextCombineParameterBlocksAndGraphs,
+
 		// -----<new versions can be added above this line>-------------------------------------------------
 		VersionPlusOne,
 		LatestVersion = VersionPlusOne - 1
