AnimNext: Move workspace systems to a seperate plugin

Workspaces now have schemas that define their restrictions and appearance.
Added a basic universal schema and one for AnimNext.
Added additional factory that tailors workspaces to AnimNext (does not show the class picker for schemas).
Intent is to support more asset types in the future
Workspace state is persisted to Json, with a GUID-based file per workspace, stored locally for each user.

diff --git a/Engine/Plugins/Experimental/AnimNext/AnimNext.uplugin b/Engine/Plugins/Experimental/AnimNext/AnimNext.uplugin
--- a/Engine/Plugins/Experimental/AnimNext/AnimNext.uplugin
+++ b/Engine/Plugins/Experimental/AnimNext/AnimNext.uplugin
@@ -32,6 +32,10 @@
 		{
 			"Name": "PythonScriptPlugin",
 			"Enabled": true
+		},
+		{
+			"Name": "Workspace",
+			"Enabled": true
 		}
 	],
 	"Modules" :
diff --git a/Engine/Plugins/Experimental/AnimNext/Config/BaseAnimNext.ini b/Engine/Plugins/Experimental/AnimNext/Config/BaseAnimNext.ini
--- a/Engine/Plugins/Experimental/AnimNext/Config/BaseAnimNext.ini
+++ b/Engine/Plugins/Experimental/AnimNext/Config/BaseAnimNext.ini
@@ -29,4 +29,7 @@
 +ClassRedirects=(OldName="/Script/AnimNextUncookedOnly.AnimNextGraphEntry", NewName="/Script/AnimNextUncookedOnly.AnimNextGraph_AnimationGraph")
 +ClassRedirects=(OldName="/Script/AnimNextUncookedOnly.AnimNextGraph_TraitSchema", NewName="/Script/AnimNextUncookedOnly.AnimNextGraph_AnimationGraphSchema")
 +PropertyRedirects=(OldName="/Script/AnimNext.AnimNextScheduleEntry_ParamScope.ParameterBlocks", NewName="Parameters")
-+PropertyRedirects=(OldName="/Script/AnimNext.AnimNextScheduleParamScopeEntryTask.ParameterBlocks", NewName="Parameters")
\ No newline at end of file
++PropertyRedirects=(OldName="/Script/AnimNext.AnimNextScheduleParamScopeEntryTask.ParameterBlocks", NewName="Parameters")
+
+; Workspace plugin split
++ClassRedirects=(OldName="/Script/AnimNextEditor.AnimNextWorkspace", NewName="/Script/WorkspaceEditor.Workspace")
\ No newline at end of file
diff --git a/Engine/Plugins/Experimental/AnimNext/Source/AnimNextEditor/AnimNextEditor.Build.cs b/Engine/Plugins/Experimental/AnimNext/Source/AnimNextEditor/AnimNextEditor.Build.cs
--- a/Engine/Plugins/Experimental/AnimNext/Source/AnimNextEditor/AnimNextEditor.Build.cs
+++ b/Engine/Plugins/Experimental/AnimNext/Source/AnimNextEditor/AnimNextEditor.Build.cs
@@ -41,6 +41,7 @@ namespace UnrealBuildTool.Rules
 					"SubobjectEditor",
 					"Settings",
 					"EditorWidgets",
+					"WorkspaceEditor",
 				}
 			);
 		}
diff --git a/Engine/Plugins/Experimental/AnimNext/Source/AnimNextEditor/Private/AnimNextEditorModule.cpp b/Engine/Plugins/Experimental/AnimNext/Source/AnimNextEditor/Private/AnimNextEditorModule.cpp
--- a/Engine/Plugins/Experimental/AnimNext/Source/AnimNextEditor/Private/AnimNextEditorModule.cpp
+++ b/Engine/Plugins/Experimental/AnimNext/Source/AnimNextEditor/Private/AnimNextEditorModule.cpp
@@ -3,7 +3,9 @@
 #include "AnimNextEditorModule.h"
 
 #include "AnimNextConfig.h"
+#include "EdGraphNode_Comment.h"
 #include "ISettingsModule.h"
+#include "IWorkspaceEditor.h"
 #include "ScopedTransaction.h"
 #include "SSimpleButton.h"
 #include "SSimpleComboButton.h"
@@ -22,7 +24,9 @@
 #include "Param/ParamTypePropertyCustomization.h"
 #include "Param/SParameterPicker.h"
 #include "Scheduler/AnimNextSchedule.h"
-#include "Workspace/AnimNextWorkspaceEditor.h"
+#include "IWorkspaceEditorModule.h"
+#include "Common/SActionMenu.h"
+#include "AnimNextRigVMAssetEntry.h"
 
 #define LOCTEXT_NAMESPACE "AnimNextEditorModule"
 
@@ -69,30 +73,7 @@ class FModule : public IModule
 		ParametersGraphPanelPinFactory = MakeShared<FParametersGraphPanelPinFactory>();
 		FEdGraphUtilities::RegisterVisualPinFactory(ParametersGraphPanelPinFactory);
 
-		FWorkspaceEditor::RegisterAssetDocumentWidget(UAnimNextSchedule::StaticClass()->GetFName(), [](TSharedRef<FWorkspaceEditor> InEditor, UObject* InAsset)
-		{
-			UAnimNextSchedule* Schedule = CastChecked<UAnimNextSchedule>(InAsset);
-			FPropertyEditorModule& PropertyEditorModule = FModuleManager::GetModuleChecked<FPropertyEditorModule>( "PropertyEditor" );
-			FDetailsViewArgs DetailsViewArgs;
-			DetailsViewArgs.NameAreaSettings = FDetailsViewArgs::HideNameArea;
-			TSharedRef<IDetailsView> DetailsView = PropertyEditorModule.CreateDetailView(DetailsViewArgs);
-			DetailsView->SetObject(Schedule);
-			return DetailsView;
-		});
-
-		FWorkspaceEditor::RegisterAssetDocumentWidget(UAnimNextGraph::StaticClass()->GetFName(), [](TSharedRef<FWorkspaceEditor> InEditor, UObject* InAsset)
-		{
-			UAnimNextGraph* Graph = CastChecked<UAnimNextGraph>(InAsset);
-			UAnimNextGraph_EditorData* EditorData = UncookedOnly::FUtils::GetEditorData(Graph);
-
-			EditorData->RigVMGraphModifiedEvent.RemoveAll(&InEditor.Get());
-			EditorData->RigVMGraphModifiedEvent.AddSP(InEditor, &FWorkspaceEditor::OnGraphModified);
-
-			return SNew(SRigVMAssetView, EditorData)
-				.OnSelectionChanged(&InEditor.Get(), &FWorkspaceEditor::SetSelectedObjects)
-				.OnOpenGraph(&InEditor.Get(), &FWorkspaceEditor::OnOpenGraph)
-				.OnDeleteEntries(&InEditor.Get(), &FWorkspaceEditor::OnDeleteEntries);
-		});
+		RegisterWorkspaceDocumentTypes();
 
 		SRigVMAssetView::RegisterCategoryFactory("Parameters", [](UAnimNextRigVMAssetEditorData* InEditorData)
 		{
@@ -199,9 +180,8 @@ class FModule : public IModule
 
 		FEdGraphUtilities::UnregisterVisualPinFactory(ParametersGraphPanelPinFactory);
 
-		FWorkspaceEditor::UnregisterAssetDocumentWidget("AnimNextSchedule");
-		FWorkspaceEditor::UnregisterAssetDocumentWidget("AnimNextGraph");
-		
+		UnregisterWorkspaceDocumentTypes();
+
 		SRigVMAssetView::UnregisterCategoryFactory("Parameters");
 		SRigVMAssetView::UnregisterCategoryFactory("Parameter Graphs");
 	}
@@ -212,6 +192,238 @@ class FModule : public IModule
 			.Args(InArgs);
 	}
 
+	void RegisterWorkspaceDocumentTypes()
+	{
+		Workspace::IWorkspaceEditorModule& WorkspaceEditorModule = FModuleManager::Get().LoadModuleChecked<Workspace::IWorkspaceEditorModule>("WorkspaceEditor");
+		WorkspaceEditorModule.RegisterObjectDocumentType(FTopLevelAssetPath(TEXT("/Script/AnimNext.AnimNextSchedule")),
+			Workspace::FObjectDocumentArgs(
+				Workspace::FOnMakeDocumentWidget::CreateLambda([](const Workspace::FWorkspaceEditorContext& InContext)
+				{
+					UAnimNextSchedule* Schedule = CastChecked<UAnimNextSchedule>(InContext.Object);
+					FPropertyEditorModule& PropertyEditorModule = FModuleManager::GetModuleChecked<FPropertyEditorModule>( "PropertyEditor" );
+					FDetailsViewArgs DetailsViewArgs;
+					DetailsViewArgs.NameAreaSettings = FDetailsViewArgs::HideNameArea;
+					TSharedRef<IDetailsView> DetailsView = PropertyEditorModule.CreateDetailView(DetailsViewArgs);
+					DetailsView->SetObject(Schedule);
+					return DetailsView;
+				}),
+				Workspace::WorkspaceTabs::MiddleDocumentArea));
+
+		WorkspaceEditorModule.RegisterObjectDocumentType(FTopLevelAssetPath(TEXT("/Script/AnimNext.AnimNextGraph")),
+			Workspace::FObjectDocumentArgs(
+				Workspace::FOnMakeDocumentWidget::CreateLambda([](const Workspace::FWorkspaceEditorContext& InContext)
+				{
+					UAnimNextGraph* Graph = CastChecked<UAnimNextGraph>(InContext.Object);
+					UAnimNextGraph_EditorData* EditorData = UncookedOnly::FUtils::GetEditorData(Graph);
+
+					TWeakPtr<Workspace::IWorkspaceEditor> WeakWorkspaceEditor = InContext.WorkspaceEditor;
+
+					EditorData->RigVMGraphModifiedEvent.RemoveAll(&InContext.WorkspaceEditor.Get());
+					EditorData->RigVMGraphModifiedEvent.AddSPLambda(&InContext.WorkspaceEditor.Get(), [WeakWorkspaceEditor](ERigVMGraphNotifType InType, URigVMGraph* InGraph, UObject* InSubject)
+					{
+						if(TSharedPtr<Workspace::IWorkspaceEditor> WorkspaceEditor = WeakWorkspaceEditor.Pin())
+						{
+							if (InType == ERigVMGraphNotifType::InteractionBracketClosed)
+							{
+								WorkspaceEditor->RefreshDetails();
+							}
+						}
+					});
+
+					return SNew(SRigVMAssetView, EditorData)
+						.OnSelectionChanged_Lambda([WeakWorkspaceEditor](const TArray<UObject*>& InEntries)
+						{
+							if(TSharedPtr<Workspace::IWorkspaceEditor> WorkspaceEditor = WeakWorkspaceEditor.Pin())
+							{
+								WorkspaceEditor->SetDetailsObjects(InEntries);
+							}
+						})
+						.OnOpenGraph_Lambda([WeakWorkspaceEditor](URigVMGraph* InGraph)
+						{
+							if(TSharedPtr<Workspace::IWorkspaceEditor> WorkspaceEditor = WeakWorkspaceEditor.Pin())
+							{
+								if(IRigVMClientHost* RigVMClientHost = InGraph->GetImplementingOuter<IRigVMClientHost>())
+								{
+									if(UObject* EditorObject = RigVMClientHost->GetEditorObjectForRigVMGraph(InGraph))
+									{
+										WorkspaceEditor->OpenObjects({EditorObject});
+									}
+								}
+							}
+						})
+						.OnDeleteEntries_Lambda([WeakWorkspaceEditor](const TArray<UAnimNextRigVMAssetEntry*>& InEntries)
+						{
+							if(TSharedPtr<Workspace::IWorkspaceEditor> WorkspaceEditor = WeakWorkspaceEditor.Pin())
+							{
+								if(InEntries.Num() > 0)
+								{
+									TArray<UObject*> EdGraphsToClose;
+									EdGraphsToClose.Reserve(InEntries.Num());
+									for(UAnimNextRigVMAssetEntry* Entry : InEntries)
+									{
+										if(IAnimNextRigVMGraphInterface* GraphInterface = Cast<IAnimNextRigVMGraphInterface>(Entry))
+										{
+											if(URigVMEdGraph* EdGraph = GraphInterface->GetEdGraph())
+											{
+												EdGraphsToClose.Add(EdGraph);
+											}
+										}
+									}
+
+									WorkspaceEditor->CloseObjects(EdGraphsToClose);
+								}
+							}
+						});
+				}),
+				Workspace::WorkspaceTabs::LeftDocumentArea));
+
+		Workspace::FGraphDocumentWidgetArgs GraphArgs;
+		GraphArgs.SpawnLocation = Workspace::WorkspaceTabs::MiddleDocumentArea;
+		GraphArgs.OnCreateActionMenu = Workspace::FOnCreateActionMenu::CreateLambda([](const Workspace::FWorkspaceEditorContext& InContext, UEdGraph* InGraph, const FVector2D& InNodePosition, const TArray<UEdGraphPin*>& InDraggedPins, bool bAutoExpand, SGraphEditor::FActionMenuClosed InOnMenuClosed)
+		{
+			TSharedRef<SActionMenu> ActionMenu = SNew(SActionMenu, InGraph)
+				.AutoExpandActionMenu(bAutoExpand)
+				.NewNodePosition(InNodePosition)
+				.DraggedFromPins(InDraggedPins)
+				.OnClosedCallback(InOnMenuClosed);
+
+			TSharedPtr<SWidget> FilterTextBox = StaticCastSharedRef<SWidget>(ActionMenu->GetFilterTextBox());
+			return FActionMenuContent(StaticCastSharedRef<SWidget>(ActionMenu), FilterTextBox);
+		});
+		GraphArgs.OnNodeTextCommitted = Workspace::FOnNodeTextCommitted::CreateLambda([](const Workspace::FWorkspaceEditorContext& InContext, const FText& NewText, ETextCommit::Type CommitInfo, UEdGraphNode* NodeBeingChanged)
+		{
+			URigVMEdGraph* RigVMEdGraph = Cast<URigVMEdGraph>(NodeBeingChanged->GetGraph());
+			if (RigVMEdGraph == nullptr)
+			{
+				return;
+			}
+
+			UEdGraphNode_Comment* CommentBeingChanged = Cast<UEdGraphNode_Comment>(NodeBeingChanged);
+			if (CommentBeingChanged == nullptr)
+			{
+				return;
+			}
+
+			RigVMEdGraph->GetController()->SetCommentTextByName(CommentBeingChanged->GetFName(), NewText.ToString(), CommentBeingChanged->FontSize, CommentBeingChanged->bCommentBubbleVisible, CommentBeingChanged->bColorCommentBubble, true, true);
+		});
+		GraphArgs.OnDeleteSelectedNodes = Workspace::FOnDeleteSelectedNodes::CreateLambda([](const Workspace::FWorkspaceEditorContext& InContext, const FGraphPanelSelectionSet& InSelectedNodes)
+		{
+			if(InSelectedNodes.IsEmpty())
+			{
+				return;
+			}
+
+			URigVMController* Controller = nullptr;
+			
+			bool bRelinkPins = false;
+			TArray<URigVMNode*> NodesToRemove;
+
+			for (FGraphPanelSelectionSet::TConstIterator NodeIt(InSelectedNodes); NodeIt; ++NodeIt)
+			{
+				if (UEdGraphNode* Node = Cast<UEdGraphNode>(*NodeIt))
+				{
+					URigVMEdGraph* RigVMEdGraph = Cast<URigVMEdGraph>(Node->GetGraph());
+					if (RigVMEdGraph == nullptr)
+					{
+						continue;
+					}
+					
+					if (Node->CanUserDeleteNode())
+					{
+						if (const URigVMEdGraphNode* RigVMEdGraphNode = Cast<URigVMEdGraphNode>(Node))
+						{
+							if(Controller == nullptr)
+							{
+								Controller = RigVMEdGraphNode->GetController();
+							}
+
+							bRelinkPins = bRelinkPins || FSlateApplication::Get().GetModifierKeys().IsShiftDown();
+
+							if(URigVMGraph* Model = RigVMEdGraph->GetModel())
+							{
+								if(URigVMNode* ModelNode = Model->FindNodeByName(*RigVMEdGraphNode->GetModelNodePath()))
+								{
+									NodesToRemove.Add(ModelNode);
+								}
+							}
+						}
+						else if (const UEdGraphNode_Comment* CommentNode = Cast<UEdGraphNode_Comment>(Node))
+						{
+							if(URigVMGraph* Model = RigVMEdGraph->GetModel())
+							{
+								if(URigVMNode* ModelNode = Model->FindNodeByName(CommentNode->GetFName()))
+								{
+									NodesToRemove.Add(ModelNode);
+								}
+							}
+						}
+						else
+						{
+							Node->GetGraph()->RemoveNode(Node);
+						}
+					}
+				}
+			}
+
+			if(NodesToRemove.IsEmpty() || Controller == nullptr)
+			{
+				return;
+			}
+
+			Controller->OpenUndoBracket(TEXT("Delete selected nodes"));
+			if(bRelinkPins && NodesToRemove.Num() == 1)
+			{
+				Controller->RelinkSourceAndTargetPins(NodesToRemove[0], true);;
+			}
+			Controller->RemoveNodes(NodesToRemove, true);
+			Controller->CloseUndoBracket();
+		});
+		GraphArgs.OnGraphSelectionChanged = Workspace::FOnGraphSelectionChanged::CreateLambda([](const Workspace::FWorkspaceEditorContext& InContext, const FGraphPanelSelectionSet& NewSelection)
+		{
+			URigVMEdGraph* RigVMEdGraph = Cast<URigVMEdGraph>(InContext.Object);
+			if (RigVMEdGraph == nullptr)
+			{
+				return;
+			}
+
+			if (RigVMEdGraph->bIsSelecting || GIsTransacting)
+			{
+				return;
+			}
+
+			TGuardValue<bool> SelectGuard(RigVMEdGraph->bIsSelecting, true);
+
+			TArray<FName> NodeNamesToSelect;
+			for (UObject* Object : NewSelection)
+			{
+				if (URigVMEdGraphNode* RigVMEdGraphNode = Cast<URigVMEdGraphNode>(Object))
+				{
+					NodeNamesToSelect.Add(RigVMEdGraphNode->GetModelNodeName());
+				}
+				else if(UEdGraphNode* Node = Cast<UEdGraphNode>(Object))
+				{
+					NodeNamesToSelect.Add(Node->GetFName());
+				}
+			}
+			RigVMEdGraph->GetController()->SetNodeSelection(NodeNamesToSelect, true, true);
+
+			InContext.WorkspaceEditor->SetDetailsObjects(NewSelection.Array());
+		});
+
+		WorkspaceEditorModule.RegisterObjectDocumentType(FTopLevelAssetPath(TEXT("/Script/AnimNextUncookedOnly.AnimNextGraph_EdGraph")), WorkspaceEditorModule.CreateGraphDocumentArgs(GraphArgs));
+	}
+
+	void UnregisterWorkspaceDocumentTypes()
+	{
+		if(FModuleManager::Get().IsModuleLoaded("WorkspaceEditor"))
+		{
+			Workspace::IWorkspaceEditorModule& WorkspaceEditorModule = FModuleManager::LoadModuleChecked<Workspace::IWorkspaceEditorModule>("PropertyEditor");
+			WorkspaceEditorModule.UnregisterObjectDocumentType(FTopLevelAssetPath(TEXT("/Script/AnimNext.AnimNextSchedule")));
+			WorkspaceEditorModule.UnregisterObjectDocumentType(FTopLevelAssetPath(TEXT("/Script/AnimNext.AnimNextGraph")));
+			WorkspaceEditorModule.UnregisterObjectDocumentType(FTopLevelAssetPath(TEXT("/Script/AnimNextUncookedOnly.AnimNextGraph_EdGraph")));
+		}
+	}
+	
 	/** Node factory for the AnimNext graph */
 	TSharedPtr<FAnimNextGraphPanelNodeFactory> AnimNextGraphPanelNodeFactory;
 	
diff --git a/Engine/Plugins/Experimental/AnimNext/Source/AnimNextEditor/Private/EditorUtils.cpp b/Engine/Plugins/Experimental/AnimNext/Source/AnimNextEditor/Private/EditorUtils.cpp
--- a/Engine/Plugins/Experimental/AnimNext/Source/AnimNextEditor/Private/EditorUtils.cpp
+++ b/Engine/Plugins/Experimental/AnimNext/Source/AnimNextEditor/Private/EditorUtils.cpp
@@ -8,7 +8,6 @@
 #include "Kismet2/Kismet2NameValidators.h"
 #include "Graph/AnimNextGraph.h"
 #include "PropertyBagDetails.h"
-#include "Workspace/AnimNextWorkspace.h"
 #include "UncookedOnlyUtils.h"
 #include "Param/RigVMDispatch_GetParameter.h"
 
@@ -420,12 +419,6 @@ bool FUtils::DoesParameterNameExistInAsset(const FName InName, const FAssetData&
 	return Exports.Parameters.ContainsByPredicate([InName](const FAnimNextParameterAssetRegistryExportEntry& Entry) { return Entry.Name == InName; });
 }
 
-bool FUtils::GetExportedAssetsForWorkspace(const FAssetData& InWorkspaceAsset, FAnimNextWorkspaceAssetRegistryExports& OutExports)
-{
-	const FString TagValue = InWorkspaceAsset.GetTagValueRef<FString>(UAnimNextWorkspace::ExportsAssetRegistryTag);
-	return FAnimNextWorkspaceAssetRegistryExports::StaticStruct()->ImportText(*TagValue, &OutExports, nullptr, PPF_None, nullptr, FAnimNextWorkspaceAssetRegistryExports::StaticStruct()->GetName()) != nullptr;
-}
-
 }
 
 #undef LOCTEXT_NAMESPACE
\ No newline at end of file
diff --git a/Engine/Plugins/Experimental/AnimNext/Source/AnimNextEditor/Private/EditorUtils.h b/Engine/Plugins/Experimental/AnimNext/Source/AnimNextEditor/Private/EditorUtils.h
--- a/Engine/Plugins/Experimental/AnimNext/Source/AnimNextEditor/Private/EditorUtils.h
+++ b/Engine/Plugins/Experimental/AnimNext/Source/AnimNextEditor/Private/EditorUtils.h
@@ -36,8 +36,6 @@ struct FUtils
 	static bool DoesParameterNameExist(const FName InName);
 	
 	static bool DoesParameterNameExistInAsset(const FName InName, const FAssetData& InAsset);
-
-	static bool GetExportedAssetsForWorkspace(const FAssetData& InWorkspaceAsset, FAnimNextWorkspaceAssetRegistryExports& OutExports);
 };
 
 }
diff --git a/Engine/Plugins/Experimental/AnimNext/Source/AnimNextEditor/Private/Graph/AnimNextGraphAssetDefinition.cpp b/Engine/Plugins/Experimental/AnimNext/Source/AnimNextEditor/Private/Graph/AnimNextGraphAssetDefinition.cpp
--- a/Engine/Plugins/Experimental/AnimNext/Source/AnimNextEditor/Private/Graph/AnimNextGraphAssetDefinition.cpp
+++ b/Engine/Plugins/Experimental/AnimNext/Source/AnimNextEditor/Private/Graph/AnimNextGraphAssetDefinition.cpp
@@ -2,18 +2,20 @@
 
 #include "AnimNextGraphAssetDefinition.h"
 #include "AnimNextGraphEditor.h"
-#include "Workspace/AnimNextWorkspaceEditor.h"
 #include "EditorCVars.h"
+#include "IWorkspaceEditorModule.h"
 
 EAssetCommandResult UAssetDefinition_AnimNextGraph::OpenAssets(const FAssetOpenArgs& OpenArgs) const
 {
 	using namespace UE::AnimNext::Editor;
-	
+	using namespace UE::Workspace;
+
 	for (UAnimNextGraph* Asset : OpenArgs.LoadObjects<UAnimNextGraph>())
 	{
 		if(CVars::GUseWorkspaceEditor.GetValueOnGameThread())
 		{
-			FWorkspaceEditor::OpenWorkspaceForAsset(Asset, FWorkspaceEditor::EOpenWorkspaceMethod::Default);
+			IWorkspaceEditorModule& WorkspaceEditorModule = FModuleManager::Get().LoadModuleChecked<IWorkspaceEditorModule>("WorkspaceEditor");
+			WorkspaceEditorModule.OpenWorkspaceForObject(Asset, EOpenWorkspaceMethod::Default);
 		}
 		else
 		{
diff --git a/Engine/Plugins/Experimental/AnimNext/Source/AnimNextEditor/Private/Graph/AnimationGraphDocumentSummoner.cpp b/Engine/Plugins/Experimental/AnimNext/Source/AnimNextEditor/Private/Graph/AnimationGraphDocumentSummoner.cpp
deleted file mode 100644
diff --git a/Engine/Plugins/Experimental/AnimNext/Source/AnimNextEditor/Private/Graph/AnimationGraphDocumentSummoner.h b/Engine/Plugins/Experimental/AnimNext/Source/AnimNextEditor/Private/Graph/AnimationGraphDocumentSummoner.h
deleted file mode 100644
diff --git a/Engine/Plugins/Experimental/AnimNext/Source/AnimNextEditor/Private/Param/EventGraphDocumentSummoner.cpp b/Engine/Plugins/Experimental/AnimNext/Source/AnimNextEditor/Private/Param/EventGraphDocumentSummoner.cpp
deleted file mode 100644
diff --git a/Engine/Plugins/Experimental/AnimNext/Source/AnimNextEditor/Private/Param/EventGraphDocumentSummoner.h b/Engine/Plugins/Experimental/AnimNext/Source/AnimNextEditor/Private/Param/EventGraphDocumentSummoner.h
deleted file mode 100644
diff --git a/Engine/Plugins/Experimental/AnimNext/Source/AnimNextEditor/Private/Scheduler/ScheduleAssetDefinition.cpp b/Engine/Plugins/Experimental/AnimNext/Source/AnimNextEditor/Private/Scheduler/ScheduleAssetDefinition.cpp
--- a/Engine/Plugins/Experimental/AnimNext/Source/AnimNextEditor/Private/Scheduler/ScheduleAssetDefinition.cpp
+++ b/Engine/Plugins/Experimental/AnimNext/Source/AnimNextEditor/Private/Scheduler/ScheduleAssetDefinition.cpp
@@ -3,19 +3,21 @@
 #include "ScheduleAssetDefinition.h"
 #include "Toolkits/SimpleAssetEditor.h"
 #include "EditorCVars.h"
-#include "Workspace/AnimNextWorkspaceEditor.h"
+#include "IWorkspaceEditorModule.h"
 
 #define LOCTEXT_NAMESPACE "AnimNextAssetDefinitions"
 
 EAssetCommandResult UAssetDefinition_AnimNextSchedule::OpenAssets(const FAssetOpenArgs& OpenArgs) const
 {
 	using namespace UE::AnimNext::Editor;
-	
+	using namespace UE::Workspace;
+
 	for(UObject* Asset : OpenArgs.LoadObjects<UObject>())
 	{
 		if(CVars::GUseWorkspaceEditor.GetValueOnGameThread())
 		{
-			FWorkspaceEditor::OpenWorkspaceForAsset(Asset, FWorkspaceEditor::EOpenWorkspaceMethod::Default);
+			IWorkspaceEditorModule& WorkspaceEditorModule = FModuleManager::Get().LoadModuleChecked<IWorkspaceEditorModule>("WorkspaceEditor");
+			WorkspaceEditorModule.OpenWorkspaceForObject(Asset, EOpenWorkspaceMethod::Default);
 		}
 		else
 		{
diff --git a/Engine/Plugins/Experimental/AnimNext/Source/AnimNextEditor/Private/Workspace/AnimNextWorkspace.cpp b/Engine/Plugins/Experimental/AnimNext/Source/AnimNextEditor/Private/Workspace/AnimNextWorkspace.cpp
deleted file mode 100644
diff --git a/Engine/Plugins/Experimental/AnimNext/Source/AnimNextEditor/Private/Workspace/AnimNextWorkspace.h b/Engine/Plugins/Experimental/AnimNext/Source/AnimNextEditor/Private/Workspace/AnimNextWorkspace.h
deleted file mode 100644
diff --git a/Engine/Plugins/Experimental/AnimNext/Source/AnimNextEditor/Private/Workspace/AnimNextWorkspaceAssetDefinition.cpp b/Engine/Plugins/Experimental/AnimNext/Source/AnimNextEditor/Private/Workspace/AnimNextWorkspaceAssetDefinition.cpp
deleted file mode 100644
diff --git a/Engine/Plugins/Experimental/AnimNext/Source/AnimNextEditor/Private/Workspace/AnimNextWorkspaceAssetDefinition.h b/Engine/Plugins/Experimental/AnimNext/Source/AnimNextEditor/Private/Workspace/AnimNextWorkspaceAssetDefinition.h
deleted file mode 100644
diff --git a/Engine/Plugins/Experimental/AnimNext/Source/AnimNextEditor/Private/Workspace/AnimNextWorkspaceEditor.cpp b/Engine/Plugins/Experimental/AnimNext/Source/AnimNextEditor/Private/Workspace/AnimNextWorkspaceEditor.cpp
deleted file mode 100644
diff --git a/Engine/Plugins/Experimental/AnimNext/Source/AnimNextEditor/Private/Workspace/AnimNextWorkspaceEditor.h b/Engine/Plugins/Experimental/AnimNext/Source/AnimNextEditor/Private/Workspace/AnimNextWorkspaceEditor.h
deleted file mode 100644
diff --git a/Engine/Plugins/Experimental/AnimNext/Source/AnimNextEditor/Private/Workspace/AnimNextWorkspaceEditorMode.cpp b/Engine/Plugins/Experimental/AnimNext/Source/AnimNextEditor/Private/Workspace/AnimNextWorkspaceEditorMode.cpp
deleted file mode 100644
diff --git a/Engine/Plugins/Experimental/AnimNext/Source/AnimNextEditor/Private/Workspace/AnimNextWorkspaceEditorMode.h b/Engine/Plugins/Experimental/AnimNext/Source/AnimNextEditor/Private/Workspace/AnimNextWorkspaceEditorMode.h
deleted file mode 100644
diff --git a/Engine/Plugins/Experimental/AnimNext/Source/AnimNextEditor/Private/Workspace/AnimNextWorkspaceFactory.cpp b/Engine/Plugins/Experimental/AnimNext/Source/AnimNextEditor/Private/Workspace/AnimNextWorkspaceFactory.cpp
--- a/Engine/Plugins/Experimental/AnimNext/Source/AnimNextEditor/Private/Workspace/AnimNextWorkspaceFactory.cpp
+++ b/Engine/Plugins/Experimental/AnimNext/Source/AnimNextEditor/Private/Workspace/AnimNextWorkspaceFactory.cpp
@@ -1,14 +1,14 @@
 ﻿// Copyright Epic Games, Inc. All Rights Reserved.
 
 #include "AnimNextWorkspaceFactory.h"
-#include "AnimNextWorkspace.h"
-#include "UObject/Package.h"
+#include "AnimNextWorkspaceSchema.h"
+#include "AssetTypeCategories.h"
+
+#define LOCTEXT_NAMESPACE "AnimNextWorkspaceFactory"
 
 UAnimNextWorkspaceFactory::UAnimNextWorkspaceFactory()
 {
-	bCreateNew = true;
-	bEditAfterNew = true;
-	SupportedClass = UAnimNextWorkspace::StaticClass();
+	SetSchemaClass(UAnimNextWorkspaceSchema::StaticClass());
 }
 
 bool UAnimNextWorkspaceFactory::ConfigureProperties()
@@ -16,13 +16,20 @@ bool UAnimNextWorkspaceFactory::ConfigureProperties()
 	return true;
 }
 
-UObject* UAnimNextWorkspaceFactory::FactoryCreateNew(UClass* Class, UObject* InParent, FName Name, EObjectFlags Flags, UObject* Context, FFeedbackContext* Warn, FName CallingContext)
+FText UAnimNextWorkspaceFactory::GetDisplayName() const
 {
-	UAnimNextWorkspace* NewParameter = NewObject<UAnimNextWorkspace>(InParent, Class, Name, Flags | RF_Public | RF_Standalone | RF_Transactional | RF_LoadCompleted);
+	return LOCTEXT("DisplayName", "AnimNext Workspace");
+}
 
-	// make sure the package is never cooked.
-	UPackage* Package = NewParameter->GetOutermost();
-	Package->SetPackageFlags(Package->GetPackageFlags() | PKG_EditorOnly);
+uint32 UAnimNextWorkspaceFactory::GetMenuCategories() const
+{
+	return EAssetTypeCategories::Animation;
+}
+
+const TArray<FText>& UAnimNextWorkspaceFactory::GetMenuCategorySubMenus() const
+{
+	static const TArray<FText> Categories = { LOCTEXT("AnimNextSubMenu", "AnimNext") };
+	return Categories;
+}
 
-	return NewParameter;
-}
\ No newline at end of file
+#undef LOCTEXT_NAMESPACE
\ No newline at end of file
diff --git a/Engine/Plugins/Experimental/AnimNext/Source/AnimNextEditor/Private/Workspace/AnimNextWorkspaceFactory.h b/Engine/Plugins/Experimental/AnimNext/Source/AnimNextEditor/Private/Workspace/AnimNextWorkspaceFactory.h
--- a/Engine/Plugins/Experimental/AnimNext/Source/AnimNextEditor/Private/Workspace/AnimNextWorkspaceFactory.h
+++ b/Engine/Plugins/Experimental/AnimNext/Source/AnimNextEditor/Private/Workspace/AnimNextWorkspaceFactory.h
@@ -3,28 +3,19 @@
 #pragma once
 
 #include "CoreMinimal.h"
-#include "Factories/Factory.h"
+#include "WorkspaceFactory.h"
 #include "AnimNextWorkspaceFactory.generated.h"
 
-namespace UE::AnimNext::Editor
-{
-	struct FUtils;
-	class FWorkspaceEditor;
-}
-
 UCLASS(MinimalAPI, BlueprintType)
-class UAnimNextWorkspaceFactory : public UFactory
+class UAnimNextWorkspaceFactory : public UWorkspaceFactory
 {
 	GENERATED_BODY()
 
-	friend class UE::AnimNext::Editor::FWorkspaceEditor;
-
 	UAnimNextWorkspaceFactory();
 
+	// UFactory interface
 	virtual bool ConfigureProperties() override;
-	virtual UObject* FactoryCreateNew(UClass* Class, UObject* InParent, FName Name, EObjectFlags Flags, UObject* Context, FFeedbackContext* Warn, FName CallingContext) override;
-	virtual UObject* FactoryCreateNew(UClass* Class, UObject* InParent, FName Name, EObjectFlags Flags, UObject* Context, FFeedbackContext* Warn) override
-	{
-		return FactoryCreateNew(Class, InParent, Name, Flags, Context, Warn, NAME_None);
-	}
+	virtual FText GetDisplayName() const override;
+	virtual uint32 GetMenuCategories() const override;
+	virtual const TArray<FText>& GetMenuCategorySubMenus() const override;
 };
\ No newline at end of file
diff --git a/Engine/Plugins/Experimental/AnimNext/Source/AnimNextEditor/Private/Workspace/AnimNextWorkspaceSchema.cpp b/Engine/Plugins/Experimental/AnimNext/Source/AnimNextEditor/Private/Workspace/AnimNextWorkspaceSchema.cpp
new file mode 100644
--- /dev/null
+++ b/Engine/Plugins/Experimental/AnimNext/Source/AnimNextEditor/Private/Workspace/AnimNextWorkspaceSchema.cpp
@@ -0,0 +1,25 @@
+﻿// Copyright Epic Games, Inc. All Rights Reserved.
+
+#include "AnimNextWorkspaceSchema.h"
+#include "Graph/AnimNextGraph.h"
+#include "Scheduler/AnimNextSchedule.h"
+
+#define LOCTEXT_NAMESPACE "AnimNextWorkspaceSchema"
+
+FText UAnimNextWorkspaceSchema::GetDisplayName() const
+{
+	return LOCTEXT("DisplayName", "AnimNext Workspace");
+}
+
+TConstArrayView<FTopLevelAssetPath> UAnimNextWorkspaceSchema::GetSupportedAssetClassPaths() const
+{
+	static const FTopLevelAssetPath Assets[] =
+	{
+		UAnimNextSchedule::StaticClass()->GetClassPathName(),
+		UAnimNextGraph::StaticClass()->GetClassPathName(),
+	};
+	
+	return Assets;
+}
+
+#undef LOCTEXT_NAMESPACE
\ No newline at end of file
diff --git a/Engine/Plugins/Experimental/AnimNext/Source/AnimNextEditor/Private/Workspace/AnimNextWorkspaceSchema.h b/Engine/Plugins/Experimental/AnimNext/Source/AnimNextEditor/Private/Workspace/AnimNextWorkspaceSchema.h
new file mode 100644
--- /dev/null
+++ b/Engine/Plugins/Experimental/AnimNext/Source/AnimNextEditor/Private/Workspace/AnimNextWorkspaceSchema.h
@@ -0,0 +1,17 @@
+﻿// Copyright Epic Games, Inc. All Rights Reserved.
+
+#pragma once
+
+#include "WorkspaceSchema.h"
+#include "AnimNextWorkspaceSchema.generated.h"
+
+// Workspace schema allowing all asset types
+UCLASS()
+class UAnimNextWorkspaceSchema : public UWorkspaceSchema
+{
+	GENERATED_BODY()
+
+	// UWorkspaceSchema interface
+	virtual FText GetDisplayName() const override;
+	virtual TConstArrayView<FTopLevelAssetPath> GetSupportedAssetClassPaths() const override;
+};
\ No newline at end of file
diff --git a/Engine/Plugins/Experimental/AnimNext/Source/AnimNextEditor/Private/Workspace/AssetDocumentSummoner.cpp b/Engine/Plugins/Experimental/AnimNext/Source/AnimNextEditor/Private/Workspace/AssetDocumentSummoner.cpp
deleted file mode 100644
diff --git a/Engine/Plugins/Experimental/AnimNext/Source/AnimNextEditor/Private/Workspace/AssetDocumentSummoner.h b/Engine/Plugins/Experimental/AnimNext/Source/AnimNextEditor/Private/Workspace/AssetDocumentSummoner.h
deleted file mode 100644
diff --git a/Engine/Plugins/Experimental/AnimNext/Source/AnimNextEditor/Private/Workspace/DetailsTabSummoner.cpp b/Engine/Plugins/Experimental/AnimNext/Source/AnimNextEditor/Private/Workspace/DetailsTabSummoner.cpp
deleted file mode 100644
diff --git a/Engine/Plugins/Experimental/AnimNext/Source/AnimNextEditor/Private/Workspace/DetailsTabSummoner.h b/Engine/Plugins/Experimental/AnimNext/Source/AnimNextEditor/Private/Workspace/DetailsTabSummoner.h
deleted file mode 100644
diff --git a/Engine/Plugins/Experimental/AnimNext/Source/AnimNextEditor/Private/Workspace/GraphDocumentSummoner.cpp b/Engine/Plugins/Experimental/AnimNext/Source/AnimNextEditor/Private/Workspace/GraphDocumentSummoner.cpp
deleted file mode 100644
diff --git a/Engine/Plugins/Experimental/AnimNext/Source/AnimNextEditor/Private/Workspace/GraphDocumentSummoner.h b/Engine/Plugins/Experimental/AnimNext/Source/AnimNextEditor/Private/Workspace/GraphDocumentSummoner.h
deleted file mode 100644
diff --git a/Engine/Plugins/Experimental/AnimNext/Source/AnimNextEditor/Private/Workspace/SWorkspacePicker.cpp b/Engine/Plugins/Experimental/AnimNext/Source/AnimNextEditor/Private/Workspace/SWorkspacePicker.cpp
deleted file mode 100644
diff --git a/Engine/Plugins/Experimental/AnimNext/Source/AnimNextEditor/Private/Workspace/SWorkspacePicker.h b/Engine/Plugins/Experimental/AnimNext/Source/AnimNextEditor/Private/Workspace/SWorkspacePicker.h
deleted file mode 100644
diff --git a/Engine/Plugins/Experimental/AnimNext/Source/AnimNextEditor/Private/Workspace/SWorkspaceView.cpp b/Engine/Plugins/Experimental/AnimNext/Source/AnimNextEditor/Private/Workspace/SWorkspaceView.cpp
deleted file mode 100644
diff --git a/Engine/Plugins/Experimental/AnimNext/Source/AnimNextEditor/Private/Workspace/SWorkspaceView.h b/Engine/Plugins/Experimental/AnimNext/Source/AnimNextEditor/Private/Workspace/SWorkspaceView.h
deleted file mode 100644
diff --git a/Engine/Plugins/Experimental/AnimNext/Source/AnimNextEditor/Private/Workspace/WorkspaceTabSummoner.cpp b/Engine/Plugins/Experimental/AnimNext/Source/AnimNextEditor/Private/Workspace/WorkspaceTabSummoner.cpp
deleted file mode 100644
diff --git a/Engine/Plugins/Experimental/AnimNext/Source/AnimNextEditor/Private/Workspace/WorkspaceTabSummoner.h b/Engine/Plugins/Experimental/AnimNext/Source/AnimNextEditor/Private/Workspace/WorkspaceTabSummoner.h
deleted file mode 100644
diff --git a/Engine/Plugins/Experimental/Workspace/Source/WorkspaceEditor/Private/AssetDocumentSummoner.cpp b/Engine/Plugins/Experimental/Workspace/Source/WorkspaceEditor/Private/AssetDocumentSummoner.cpp
new file mode 100644
--- /dev/null
+++ b/Engine/Plugins/Experimental/Workspace/Source/WorkspaceEditor/Private/AssetDocumentSummoner.cpp
@@ -0,0 +1,177 @@
+﻿// Copyright Epic Games, Inc. All Rights Reserved.
+
+#include "AssetDocumentSummoner.h"
+
+#include "WorkspaceEditor.h"
+#include "AssetDefinitionRegistry.h"
+#include "ClassIconFinder.h"
+#include "InstancedStruct.h"
+#include "WorkspaceEditorModule.h"
+#include "Widgets/Docking/SDockTab.h"
+#include "Widgets/Layout/SSpacer.h"
+#include "WorkspaceDocumentState.h"
+
+#define LOCTEXT_NAMESPACE "AssetDocumentSummoner"
+
+namespace UE::Workspace
+{
+
+FAssetDocumentSummoner::FAssetDocumentSummoner(FName InIdentifier, TSharedPtr<FWorkspaceEditor> InHostingApp)
+	: FDocumentTabFactoryForObjects<UObject>(InIdentifier, InHostingApp)
+	, HostingAppPtr(InHostingApp)
+{
+}
+
+void FAssetDocumentSummoner::SetAllowedClassPaths(TConstArrayView<FTopLevelAssetPath> InAllowedClassPaths)
+{
+	AllowedClassPaths = InAllowedClassPaths;
+}
+
+void FAssetDocumentSummoner::OnTabActivated(TSharedPtr<SDockTab> Tab) const
+{
+}
+
+void FAssetDocumentSummoner::OnTabBackgrounded(TSharedPtr<SDockTab> Tab) const
+{
+}
+
+void FAssetDocumentSummoner::OnTabRefreshed(TSharedPtr<SDockTab> Tab) const
+{
+}
+
+void FAssetDocumentSummoner::SaveState(TSharedPtr<SDockTab> Tab, TSharedPtr<FTabPayload> Payload) const
+{
+	if(TSharedPtr<FWorkspaceEditor> WorkspaceEditor = HostingAppPtr.Pin())
+	{
+		UObject* Object = Payload->IsValid() ? FTabPayload_UObject::CastChecked<UObject>(Payload) : nullptr;
+		if(Object)
+		{
+			FWorkspaceEditorModule& WorkspaceEditorModule = FModuleManager::LoadModuleChecked<FWorkspaceEditorModule>("WorkspaceEditor");
+			const FObjectDocumentArgs* DocumentArgs = WorkspaceEditorModule.FindObjectDocumentType(Object->GetClass()->GetClassPathName());
+			if(DocumentArgs && DocumentArgs->OnGetDocumentState.IsBound())
+			{
+				WorkspaceEditor->RecordDocumentState(DocumentArgs->OnGetDocumentState.Execute(FWorkspaceEditorContext(WorkspaceEditor.ToSharedRef(), Object), Tab->GetContent()));
+			}
+			else
+			{
+				WorkspaceEditor->RecordDocumentState(TInstancedStruct<FWorkspaceDocumentState>::Make(Object));
+			}
+		}
+	}
+}
+
+TAttribute<FText> FAssetDocumentSummoner::ConstructTabNameForObject(UObject* DocumentID) const
+{
+	TSharedPtr<FWorkspaceEditor> WorkspaceEditor = HostingAppPtr.Pin();
+	if(!WorkspaceEditor.IsValid() || DocumentID == nullptr)
+	{
+		return LOCTEXT("NoneObjectName", "None");
+	}
+
+	FWorkspaceEditorModule& WorkspaceEditorModule = FModuleManager::LoadModuleChecked<FWorkspaceEditorModule>("WorkspaceEditor");
+	if(const FObjectDocumentArgs* DocumentArgs = WorkspaceEditorModule.FindObjectDocumentType(DocumentID->GetClass()->GetClassPathName()))
+	{
+		if(DocumentArgs->OnGetTabName.IsBound())
+		{
+			return DocumentArgs->OnGetTabName.Execute(FWorkspaceEditorContext(WorkspaceEditor.ToSharedRef(), DocumentID));
+		}
+	}
+
+	return MakeAttributeLambda([WeakObject = TWeakObjectPtr<UObject>(DocumentID)]()
+	{
+		if(UObject* Object = WeakObject.Get())
+		{
+			return FText::FromName(Object->GetFName());
+		}
+
+		return LOCTEXT("UnknownObjectName", "Unknown");
+	});
+}
+
+bool FAssetDocumentSummoner::IsPayloadSupported(TSharedRef<FTabPayload> Payload) const
+{
+	UObject* Object = Payload->IsValid() ? FTabPayload_UObject::CastChecked<UObject>(Payload) : nullptr;
+	if(Object)
+	{
+		FWorkspaceEditorModule& WorkspaceEditorModule = FModuleManager::LoadModuleChecked<FWorkspaceEditorModule>("WorkspaceEditor");
+		if(const FObjectDocumentArgs* DocumentArgs = WorkspaceEditorModule.FindObjectDocumentType(Object->GetClass()->GetClassPathName()))
+		{
+			return AllowedClassPaths.IsEmpty() || AllowedClassPaths.Contains(Object->GetClass()->GetClassPathName());
+		}
+	}
+
+	return false;
+}
+
+TAttribute<FText> FAssetDocumentSummoner::ConstructTabLabelSuffix(const FWorkflowTabSpawnInfo& Info) const
+{
+	UObject* Object = Info.Payload->IsValid() ? FTabPayload_UObject::CastChecked<UObject>(Info.Payload) : nullptr;
+	if(Object)
+	{
+		return MakeAttributeLambda([WeakObject = TWeakObjectPtr<UObject>(Object)]()
+		{
+			if(UObject* Object = WeakObject.Get())
+			{
+				if(Object->IsAsset() && Object->GetPackage()->IsDirty())
+				{
+					return LOCTEXT("TabSuffixAsterisk", "*");
+				}
+			}
+
+			return FText::GetEmpty();
+		});
+	}
+
+	return FText::GetEmpty();
+}
+
+TSharedRef<SWidget> FAssetDocumentSummoner::CreateTabBodyForObject(const FWorkflowTabSpawnInfo& Info, UObject* DocumentID) const
+{
+	if(TSharedPtr<FWorkspaceEditor> WorkspaceEditor = HostingAppPtr.Pin())
+	{
+		FWorkspaceEditorModule& WorkspaceEditorModule = FModuleManager::LoadModuleChecked<FWorkspaceEditorModule>("WorkspaceEditor");
+		if(const FObjectDocumentArgs* DocumentArgs = WorkspaceEditorModule.FindObjectDocumentType(DocumentID->GetClass()->GetClassPathName()))
+		{
+			if(DocumentArgs->OnMakeDocumentWidget.IsBound())
+			{
+				return DocumentArgs->OnMakeDocumentWidget.Execute(FWorkspaceEditorContext(WorkspaceEditor.ToSharedRef(), DocumentID));
+			}
+		}
+	}
+
+	return SNew(SSpacer);
+}
+
+const FSlateBrush* FAssetDocumentSummoner::GetTabIconForObject(const FWorkflowTabSpawnInfo& Info, UObject* DocumentID) const
+{
+	if(TSharedPtr<FWorkspaceEditor> WorkspaceEditor = HostingAppPtr.Pin())
+	{
+		FWorkspaceEditorModule& WorkspaceEditorModule = FModuleManager::LoadModuleChecked<FWorkspaceEditorModule>("WorkspaceEditor");
+		if(const FObjectDocumentArgs* DocumentArgs = WorkspaceEditorModule.FindObjectDocumentType(DocumentID->GetClass()->GetClassPathName()))
+		{
+			if(DocumentArgs->OnGetTabIcon.IsBound())
+			{
+				return DocumentArgs->OnGetTabIcon.Execute(FWorkspaceEditorContext(WorkspaceEditor.ToSharedRef(), DocumentID));
+			}
+		}
+
+		if (UAssetDefinitionRegistry* AssetDefinitionRegistry = UAssetDefinitionRegistry::Get())
+		{
+			const FAssetData AssetData(DocumentID);
+			if(const UAssetDefinition* AssetDefinition = AssetDefinitionRegistry->GetAssetDefinitionForAsset(AssetData))
+			{
+				const FSlateBrush* ThumbnailBrush = AssetDefinition->GetThumbnailBrush(AssetData, AssetData.AssetClassPath.GetAssetName());
+				if(ThumbnailBrush == nullptr)
+				{
+					return FClassIconFinder::FindThumbnailForClass(DocumentID->GetClass(), NAME_None);
+				}
+				return ThumbnailBrush;
+			}
+		}
+	}
+	return nullptr;
+}
+
+}
+
+#undef LOCTEXT_NAMESPACE
\ No newline at end of file
diff --git a/Engine/Plugins/Experimental/Workspace/Source/WorkspaceEditor/Private/AssetDocumentSummoner.h b/Engine/Plugins/Experimental/Workspace/Source/WorkspaceEditor/Private/AssetDocumentSummoner.h
new file mode 100644
--- /dev/null
+++ b/Engine/Plugins/Experimental/Workspace/Source/WorkspaceEditor/Private/AssetDocumentSummoner.h
@@ -0,0 +1,53 @@
+﻿// Copyright Epic Games, Inc. All Rights Reserved.
+
+#pragma once
+
+#include "CoreMinimal.h"
+#include "WorkflowOrientedApp/WorkflowUObjectDocuments.h"
+
+class FWorkflowCentricApplication;
+class SGraphEditor;
+
+namespace UE::Workspace
+{
+	class FWorkspaceEditor;
+}
+
+namespace UE::Workspace
+{
+
+struct FAssetDocumentSummoner : public FDocumentTabFactoryForObjects<UObject>
+{
+public:
+	// Delegate called to save the state of a document
+	DECLARE_DELEGATE_OneParam(FOnSaveDocumentState, UObject*);
+
+	FAssetDocumentSummoner(FName InIdentifier, TSharedPtr<FWorkspaceEditor> InHostingApp);
+
+	void SetAllowedClassPaths(TConstArrayView<FTopLevelAssetPath> InAllowedClassPaths);
+
+private:
+	// FWorkflowTabFactory interface
+	virtual void OnTabActivated(TSharedPtr<SDockTab> Tab) const override;
+	virtual void OnTabBackgrounded(TSharedPtr<SDockTab> Tab) const override;
+	virtual void OnTabRefreshed(TSharedPtr<SDockTab> Tab) const override;
+	virtual void SaveState(TSharedPtr<SDockTab> Tab, TSharedPtr<FTabPayload> Payload) const override;
+	virtual bool IsPayloadSupported(TSharedRef<FTabPayload> Payload) const override;
+	virtual TAttribute<FText> ConstructTabLabelSuffix(const FWorkflowTabSpawnInfo& Info) const override;
+
+	// FDocumentTabFactoryForObjects interface
+	virtual TAttribute<FText> ConstructTabNameForObject(UObject* DocumentID) const override;
+	virtual const FSlateBrush* GetTabIconForObject(const FWorkflowTabSpawnInfo& Info, UObject* DocumentID) const override;
+	virtual TSharedRef<SWidget> CreateTabBodyForObject(const FWorkflowTabSpawnInfo& Info, UObject* DocumentID) const override;
+
+	// The hosting app
+	TWeakPtr<FWorkspaceEditor> HostingAppPtr;
+
+	// Command list
+	TSharedPtr<FUICommandList> CommandList;
+
+	// Allowed object types
+	TArray<FTopLevelAssetPath> AllowedClassPaths;
+};
+
+}
diff --git a/Engine/Plugins/Experimental/Workspace/Source/WorkspaceEditor/Private/DefaultWorkspaceSchema.h b/Engine/Plugins/Experimental/Workspace/Source/WorkspaceEditor/Private/DefaultWorkspaceSchema.h
new file mode 100644
--- /dev/null
+++ b/Engine/Plugins/Experimental/Workspace/Source/WorkspaceEditor/Private/DefaultWorkspaceSchema.h
@@ -0,0 +1,17 @@
+﻿// Copyright Epic Games, Inc. All Rights Reserved.
+
+#pragma once
+
+#include "WorkspaceSchema.h"
+#include "DefaultWorkspaceSchema.generated.h"
+
+// Workspace schema allowing all asset types
+UCLASS()
+class UDefaultWorkspaceSchema : public UWorkspaceSchema
+{
+	GENERATED_BODY()
+
+	// UWorkspaceSchema interface
+	virtual FText GetDisplayName() const { return FText::GetEmpty(); }
+	virtual TConstArrayView<FTopLevelAssetPath> GetSupportedAssetClassPaths() const override { return TConstArrayView<FTopLevelAssetPath>(); }
+};
\ No newline at end of file
diff --git a/Engine/Plugins/Experimental/Workspace/Source/WorkspaceEditor/Private/DetailsTabSummoner.cpp b/Engine/Plugins/Experimental/Workspace/Source/WorkspaceEditor/Private/DetailsTabSummoner.cpp
new file mode 100644
--- /dev/null
+++ b/Engine/Plugins/Experimental/Workspace/Source/WorkspaceEditor/Private/DetailsTabSummoner.cpp
@@ -0,0 +1,43 @@
+﻿// Copyright Epic Games, Inc. All Rights Reserved.
+
+#include "DetailsTabSummoner.h"
+#include "WorkspaceEditor.h"
+#include "PropertyEditorModule.h"
+
+#define LOCTEXT_NAMESPACE "DetailsTabSummoner"
+
+namespace UE::Workspace
+{
+
+FDetailsTabSummoner::FDetailsTabSummoner(TSharedPtr<FWorkspaceEditor> InHostingApp, FOnDetailsViewCreated InOnDetailsViewCreated)
+	: FWorkflowTabFactory(WorkspaceTabs::Details, InHostingApp)
+{
+	TabLabel = LOCTEXT("DetailsTabLabel", "Details");
+	TabIcon = FSlateIcon("EditorStyle", "LevelEditor.Tabs.Details");
+	ViewMenuDescription = LOCTEXT("DetailsTabMenuDescription", "Details");
+	ViewMenuTooltip = LOCTEXT("DetailsTabToolTip", "Shows the details tab for selected objects.");
+	bIsSingleton = true;
+
+	FDetailsViewArgs DetailsViewArgs;
+	DetailsViewArgs.NameAreaSettings = FDetailsViewArgs::HideNameArea;
+	DetailsViewArgs.bHideSelectionTip = true;
+
+	FPropertyEditorModule& PropertyEditorModule = FModuleManager::LoadModuleChecked<FPropertyEditorModule>("PropertyEditor");
+	DetailsView = PropertyEditorModule.CreateDetailView(DetailsViewArgs);
+
+	InOnDetailsViewCreated.ExecuteIfBound(DetailsView.ToSharedRef());
+}
+
+TSharedRef<SWidget> FDetailsTabSummoner::CreateTabBody(const FWorkflowTabSpawnInfo& Info) const
+{
+	return DetailsView.ToSharedRef();
+}
+
+FText FDetailsTabSummoner::GetTabToolTipText(const FWorkflowTabSpawnInfo& Info) const
+{
+	return ViewMenuTooltip;
+}
+
+}
+
+#undef LOCTEXT_NAMESPACE
\ No newline at end of file
diff --git a/Engine/Plugins/Experimental/Workspace/Source/WorkspaceEditor/Private/DetailsTabSummoner.h b/Engine/Plugins/Experimental/Workspace/Source/WorkspaceEditor/Private/DetailsTabSummoner.h
new file mode 100644
--- /dev/null
+++ b/Engine/Plugins/Experimental/Workspace/Source/WorkspaceEditor/Private/DetailsTabSummoner.h
@@ -0,0 +1,36 @@
+﻿// Copyright Epic Games, Inc. All Rights Reserved.
+
+#pragma once
+
+#include "WorkflowOrientedApp/WorkflowTabFactory.h"
+
+class IDetailsView;
+
+namespace UE::Workspace
+{
+	class FWorkspaceEditor;
+}
+
+namespace UE::Workspace
+{
+
+DECLARE_DELEGATE_OneParam(FOnDetailsViewCreated, TSharedRef<IDetailsView>);
+
+class FDetailsTabSummoner : public FWorkflowTabFactory
+{
+public:
+	FDetailsTabSummoner(TSharedPtr<FWorkspaceEditor> InHostingApp, FOnDetailsViewCreated InOnDetailsViewCreated);
+
+	TSharedPtr<IDetailsView> GetDetailsView() const
+	{
+		return DetailsView;
+	}
+
+private:
+	virtual TSharedRef<SWidget> CreateTabBody(const FWorkflowTabSpawnInfo& Info) const override;
+	virtual FText GetTabToolTipText(const FWorkflowTabSpawnInfo& Info) const override;
+
+	TSharedPtr<IDetailsView> DetailsView;
+};
+
+}
\ No newline at end of file
diff --git a/Engine/Plugins/Experimental/Workspace/Source/WorkspaceEditor/Private/GraphDocumentState.h b/Engine/Plugins/Experimental/Workspace/Source/WorkspaceEditor/Private/GraphDocumentState.h
new file mode 100644
--- /dev/null
+++ b/Engine/Plugins/Experimental/Workspace/Source/WorkspaceEditor/Private/GraphDocumentState.h
@@ -0,0 +1,27 @@
+﻿// Copyright Epic Games, Inc. All Rights Reserved.
+
+#pragma once
+
+#include "WorkspaceDocumentState.h"
+#include "GraphDocumentState.generated.h"
+
+// Base struct used to persist workspace document state
+USTRUCT()
+struct FGraphDocumentState : public FWorkspaceDocumentState
+{
+	GENERATED_BODY()
+
+	FGraphDocumentState() = default;
+
+	FGraphDocumentState(const UObject* InObject, const FVector2D& InViewLocation, float InZoomAmount)
+		: FWorkspaceDocumentState(InObject)
+		, ViewLocation(InViewLocation)
+		, ZoomAmount(InZoomAmount)
+	{}
+
+	UPROPERTY()
+	FVector2D ViewLocation;
+
+	UPROPERTY()
+	float ZoomAmount;
+};
diff --git a/Engine/Plugins/Experimental/Workspace/Source/WorkspaceEditor/Private/SGraphDocument.cpp b/Engine/Plugins/Experimental/Workspace/Source/WorkspaceEditor/Private/SGraphDocument.cpp
new file mode 100644
--- /dev/null
+++ b/Engine/Plugins/Experimental/Workspace/Source/WorkspaceEditor/Private/SGraphDocument.cpp
@@ -0,0 +1,93 @@
+﻿// Copyright Epic Games, Inc. All Rights Reserved.
+
+#include "SGraphDocument.h"
+
+#include "Framework/Commands/GenericCommands.h"
+#include "WorkspaceEditor.h"
+#include "EdGraph/EdGraph.h"
+
+namespace UE::Workspace
+{
+
+void SGraphDocument::Construct(const FArguments& InArgs, TSharedRef<FWorkspaceEditor> InHostingApp, UEdGraph* InGraph)
+{
+	HostingAppPtr = InHostingApp;
+	EdGraph = InGraph;
+
+	BindCommands();
+
+	OnDeleteSelectedNodes = InArgs._OnDeleteSelectedNodes;
+
+	SGraphEditor::FGraphEditorEvents Events;
+	Events.OnCreateActionMenu = SGraphEditor::FOnCreateActionMenu::CreateLambda([this, OnCreateActionMenu = InArgs._OnCreateActionMenu](UEdGraph* InGraph, const FVector2D& InNodePosition, const TArray<UEdGraphPin*>& InDraggedPins, bool bInAutoExpand, SGraphEditor::FActionMenuClosed InOnMenuClosed)
+	{
+		if(OnCreateActionMenu.IsBound())
+		{
+			return OnCreateActionMenu.Execute(FWorkspaceEditorContext(HostingAppPtr.Pin().ToSharedRef(), EdGraph), InGraph, InNodePosition, InDraggedPins, bInAutoExpand, InOnMenuClosed);
+		}
+		return FActionMenuContent();
+	});
+	Events.OnSelectionChanged = SGraphEditor::FOnSelectionChanged::CreateLambda([this, OnGraphSelectionChanged = InArgs._OnGraphSelectionChanged](const TSet<UObject*>& NewSelection)
+	{
+		OnGraphSelectionChanged.ExecuteIfBound(FWorkspaceEditorContext(HostingAppPtr.Pin().ToSharedRef(), EdGraph), NewSelection);
+	});
+	Events.OnTextCommitted = ::FOnNodeTextCommitted::CreateLambda([this, OnNodeTextCommitted = InArgs._OnNodeTextCommitted](const FText& NewText, ETextCommit::Type CommitInfo, UEdGraphNode* NodeBeingChanged)
+	{
+		OnNodeTextCommitted.ExecuteIfBound(FWorkspaceEditorContext(HostingAppPtr.Pin().ToSharedRef(), EdGraph), NewText, CommitInfo, NodeBeingChanged);
+	});
+
+	ChildSlot
+	[
+		SAssignNew(GraphEditor, SGraphEditor)
+		.AdditionalCommands(CommandList)
+		.IsEditable(this, &SGraphDocument::IsEditable, InGraph)
+		.GraphToEdit(InGraph)
+		.GraphEvents(Events)
+		.AssetEditorToolkit(HostingAppPtr)
+	];
+}
+
+void SGraphDocument::BindCommands()
+{
+	CommandList = MakeShared<FUICommandList>();
+
+	CommandList->MapAction(FGenericCommands::Get().Delete,
+		FExecuteAction::CreateSP(this, &SGraphDocument::DeleteSelectedNodes),
+		FCanExecuteAction::CreateSP(this, &SGraphDocument::CanDeleteSelectedNodes));
+}
+
+void SGraphDocument::DeleteSelectedNodes()
+{
+	const FGraphPanelSelectionSet SelectedNodes = GraphEditor->GetSelectedNodes();
+	OnDeleteSelectedNodes.ExecuteIfBound(FWorkspaceEditorContext(HostingAppPtr.Pin().ToSharedRef(), EdGraph), SelectedNodes);
+}
+
+bool SGraphDocument::CanDeleteSelectedNodes() const
+{
+	const FGraphPanelSelectionSet SelectedNodes = GraphEditor->GetSelectedNodes();
+
+	bool bCanUserDeleteNode = false;
+
+	if(IsEditable(EdGraph) && SelectedNodes.Num() > 0)
+	{
+		for(UObject* NodeObject : SelectedNodes)
+		{
+			// If any nodes allow deleting, then do not disable the delete option
+			UEdGraphNode* Node = Cast<UEdGraphNode>(NodeObject);
+			if(Node && Node->CanUserDeleteNode())
+			{
+				bCanUserDeleteNode = true;
+				break;
+			}
+		}
+	}
+
+	return bCanUserDeleteNode;
+}
+
+bool SGraphDocument::IsEditable(UEdGraph* InGraph) const
+{
+	return InGraph && HostingAppPtr.Pin()->InEditingMode() && InGraph->bEditable;
+}
+
+};
\ No newline at end of file
diff --git a/Engine/Plugins/Experimental/Workspace/Source/WorkspaceEditor/Private/SGraphDocument.h b/Engine/Plugins/Experimental/Workspace/Source/WorkspaceEditor/Private/SGraphDocument.h
new file mode 100644
--- /dev/null
+++ b/Engine/Plugins/Experimental/Workspace/Source/WorkspaceEditor/Private/SGraphDocument.h
@@ -0,0 +1,61 @@
+﻿// Copyright Epic Games, Inc. All Rights Reserved.
+
+#pragma once
+
+#include "GraphEditor.h"
+#include "IWorkspaceEditorModule.h"
+#include "Widgets/SCompoundWidget.h"
+
+namespace UE::Workspace
+{
+	class FWorkspaceEditor;
+	class FWorkspaceEditorModule;
+}
+
+namespace UE::Workspace
+{
+
+// Wrapper widget for a graph editor
+class SGraphDocument : public SCompoundWidget
+{
+	SLATE_BEGIN_ARGS(SGraphDocument) {}
+
+	SLATE_EVENT(FOnCreateActionMenu, OnCreateActionMenu)
+
+	SLATE_EVENT(FOnNodeTextCommitted, OnNodeTextCommitted)
+
+	SLATE_EVENT(FOnDeleteSelectedNodes, OnDeleteSelectedNodes)
+
+	SLATE_EVENT(FOnGraphSelectionChanged, OnGraphSelectionChanged)
+
+	SLATE_END_ARGS()
+
+	void Construct(const FArguments& InArgs, TSharedRef<FWorkspaceEditor> InHostingApp, UEdGraph* InGraph);
+	
+	void BindCommands();
+	
+	void DeleteSelectedNodes();
+
+	bool CanDeleteSelectedNodes() const;
+
+	bool IsEditable(UEdGraph* InGraph) const;
+
+	// The graph we are editing
+	UEdGraph* EdGraph = nullptr;
+
+	// The graph editor we wrap
+	TSharedPtr<SGraphEditor> GraphEditor;
+
+	// Command list for graphs
+	TSharedPtr<FUICommandList> CommandList;
+
+	// The hosting app
+	TWeakPtr<FWorkspaceEditor> HostingAppPtr;
+
+	// Delegate called when we delete nodes
+	FOnDeleteSelectedNodes OnDeleteSelectedNodes;
+
+	friend class FWorkspaceEditorModule;
+};
+
+}
\ No newline at end of file
diff --git a/Engine/Plugins/Experimental/Workspace/Source/WorkspaceEditor/Private/SWorkspacePicker.cpp b/Engine/Plugins/Experimental/Workspace/Source/WorkspaceEditor/Private/SWorkspacePicker.cpp
new file mode 100644
--- /dev/null
+++ b/Engine/Plugins/Experimental/Workspace/Source/WorkspaceEditor/Private/SWorkspacePicker.cpp
@@ -0,0 +1,102 @@
+﻿// Copyright Epic Games, Inc. All Rights Reserved.
+
+#include "SWorkspacePicker.h"
+#include "Workspace.h"
+#include "ContentBrowserModule.h"
+#include "IContentBrowserSingleton.h"
+#include "SPrimaryButton.h"
+#include "Widgets/Layout/SBox.h"
+#include "Widgets/Text/STextBlock.h"
+#include "Framework/Application/SlateApplication.h"
+
+#define LOCTEXT_NAMESPACE "SWorkspacePicker"
+
+namespace UE::Workspace
+{
+
+void SWorkspacePicker::Construct(const FArguments& InArgs)
+{
+	WorkspaceAssets = InArgs._WorkspaceAssets;
+
+	FContentBrowserModule& ContentBrowserModule = FModuleManager::Get().LoadModuleChecked<FContentBrowserModule>(TEXT("ContentBrowser"));
+
+	FAssetPickerConfig AssetPickerConfig;
+	AssetPickerConfig.InitialAssetViewType = EAssetViewType::List;
+	AssetPickerConfig.OnAssetSelected = FOnAssetSelected::CreateLambda([this, OnAssetSelected = InArgs._OnAssetSelected](const FAssetData& InAssetData)
+	{
+		OnAssetSelected.ExecuteIfBound(InAssetData);
+
+		if(TSharedPtr<SWindow> Window = WeakWindow.Pin())
+		{
+			Window->RequestDestroyWindow();
+		}
+	});
+	AssetPickerConfig.Filter.ClassPaths.Add(UWorkspace::StaticClass()->GetClassPathName());
+	AssetPickerConfig.Filter.bRecursiveClasses = true;
+	AssetPickerConfig.OnShouldFilterAsset = FOnShouldFilterAsset::CreateLambda([this](const FAssetData& InAssetData)
+	{
+		return !WorkspaceAssets.IsEmpty() && !WorkspaceAssets.Contains(InAssetData);
+	});
+
+	ChildSlot
+	[
+		SNew(SBox)
+		.WidthOverride(400.0f)
+		.HeightOverride(400.0f)
+		[
+			SNew(SVerticalBox)
+			+SVerticalBox::Slot()
+			.AutoHeight()
+			.Padding(5.0f)
+			.HAlign(HAlign_Left)
+			[
+				SNew(STextBlock)
+				.Text(LOCTEXT("ChooseExistingWorkspaceToDetails", "Selected asset is part of multiple workspaces.\nPlease select the workspace you want to open the asset with"))
+			]
+			+SVerticalBox::Slot()
+			.FillHeight(1.0f)
+			.Padding(5.0f)
+			[
+				ContentBrowserModule.Get().CreateAssetPicker(AssetPickerConfig)
+			]
+			+SVerticalBox::Slot()
+			.AutoHeight()
+			.Padding(5.0f)
+			.HAlign(HAlign_Fill)
+			[
+				SNew(SPrimaryButton)
+				.Text(LOCTEXT("CreateNewWorkspace", "Create New Workspace"))
+				.OnClicked_Lambda([this, OnNewAsset = InArgs._OnNewAsset]()
+				{
+					OnNewAsset.ExecuteIfBound();
+
+					if(TSharedPtr<SWindow> Window = WeakWindow.Pin())
+					{
+						Window->RequestDestroyWindow();
+					}
+					return FReply::Handled();
+				})
+			]
+		]
+	];
+}
+
+void SWorkspacePicker::ShowModal()
+{
+	TSharedRef<SWindow> Window = SNew(SWindow)
+	.Title(LOCTEXT("ChooseExistingWorkspaceToOpen", "Choose Existing Workspace to Open"))
+	.SizingRule(ESizingRule::Autosized)
+	.SupportsMaximize(false)
+	.SupportsMinimize(false)
+	[
+		AsShared()
+	];
+
+	WeakWindow = Window;
+
+	FSlateApplication::Get().AddModalWindow(Window, FGlobalTabmanager::Get()->GetRootWindow());
+}
+
+}
+
+#undef LOCTEXT_NAMESPACE
\ No newline at end of file
diff --git a/Engine/Plugins/Experimental/Workspace/Source/WorkspaceEditor/Private/SWorkspacePicker.h b/Engine/Plugins/Experimental/Workspace/Source/WorkspaceEditor/Private/SWorkspacePicker.h
new file mode 100644
--- /dev/null
+++ b/Engine/Plugins/Experimental/Workspace/Source/WorkspaceEditor/Private/SWorkspacePicker.h
@@ -0,0 +1,38 @@
+﻿// Copyright Epic Games, Inc. All Rights Reserved.
+
+#pragma once
+
+#include "ContentBrowserDelegates.h"
+#include "Widgets/SCompoundWidget.h"
+
+class UWorkspace;
+
+namespace UE::Workspace
+{
+
+class SWorkspacePicker : public SCompoundWidget
+{
+	using FOnNewAsset = TDelegate<void(TConstArrayView<FAssetData>)>; 
+	
+	SLATE_BEGIN_ARGS(SWorkspacePicker) {}
+
+	SLATE_EVENT(FOnAssetSelected, OnAssetSelected)
+
+	SLATE_EVENT(FSimpleDelegate, OnNewAsset)
+
+	// The set of workspace assets the user can choose from
+	SLATE_ARGUMENT(TArray<FAssetData>, WorkspaceAssets)
+
+	SLATE_END_ARGS()
+
+	void Construct(const FArguments& InArgs);
+
+	void ShowModal();
+
+private:
+	TArray<FAssetData> WorkspaceAssets;
+
+	TWeakPtr<SWindow> WeakWindow;
+};
+
+};
\ No newline at end of file
diff --git a/Engine/Plugins/Experimental/Workspace/Source/WorkspaceEditor/Private/SWorkspaceView.cpp b/Engine/Plugins/Experimental/Workspace/Source/WorkspaceEditor/Private/SWorkspaceView.cpp
new file mode 100644
--- /dev/null
+++ b/Engine/Plugins/Experimental/Workspace/Source/WorkspaceEditor/Private/SWorkspaceView.cpp
@@ -0,0 +1,133 @@
+﻿// Copyright Epic Games, Inc. All Rights Reserved.
+
+#include "SWorkspaceView.h"
+#include "Workspace.h"
+#include "ContentBrowserModule.h"
+#include "IContentBrowserSingleton.h"
+#include "SAssetDropTarget.h"
+#include "ScopedTransaction.h"
+#include "WorkspaceSchema.h"
+#include "Framework/Commands/GenericCommands.h"
+
+#define LOCTEXT_NAMESPACE "SWorkspaceView"
+
+namespace UE::Workspace
+{
+
+void SWorkspaceView::Construct(const FArguments& InArgs, UWorkspace* InWorkspace)
+{
+	Workspace = InWorkspace;
+	OnAssetsOpened = InArgs._OnAssetsOpened;
+
+	UICommandList = MakeShared<FUICommandList>();
+
+	UICommandList->MapAction(FGenericCommands::Get().Delete,
+		FExecuteAction::CreateSP(this, &SWorkspaceView::HandleDelete),
+		FCanExecuteAction::CreateSP(this, &SWorkspaceView::HasValidSelection));
+
+	Workspace->ModifiedDelegate.AddSP(this, &SWorkspaceView::HandleWorkspaceModified);
+
+	FContentBrowserModule& ContentBrowserModule = FModuleManager::Get().LoadModuleChecked<FContentBrowserModule>(TEXT("ContentBrowser"));
+
+	FAssetPickerConfig AssetPickerConfig;
+	AssetPickerConfig.InitialAssetViewType = EAssetViewType::List;
+	AssetPickerConfig.OnAssetsActivated = FOnAssetsActivated::CreateLambda([this](const TArray<FAssetData>& InAssets, EAssetTypeActivationMethod::Type InActivationMethod)
+	{
+		if( InActivationMethod == EAssetTypeActivationMethod::DoubleClicked ||
+			InActivationMethod == EAssetTypeActivationMethod::Opened)
+		{
+			OnAssetsOpened.ExecuteIfBound(InAssets);
+		}
+	});
+	AssetPickerConfig.RefreshAssetViewDelegates.Add(&RefreshAssetViewDelegate);
+	AssetPickerConfig.GetCurrentSelectionDelegates.Add(&GetCurrentSelectionDelegate);
+	AssetPickerConfig.Filter = MakeARFilter();
+	AssetPickerConfig.OnShouldFilterAsset = FOnShouldFilterAsset::CreateLambda([this](const FAssetData& InAsset)
+	{
+		return !Workspace->Assets.Contains(TSoftObjectPtr<UObject>(InAsset.GetSoftObjectPath()));
+	});
+	AssetPickerConfig.AssetShowWarningText = LOCTEXT("EmptyWorkspaceText", "Workspace is empty.\nDrag-drop to add assets to this workspace.");
+
+	ChildSlot
+	[
+		SNew(SAssetDropTarget)
+		.bSupportsMultiDrop(true)
+		.OnAssetsDropped_Lambda([this](const FDragDropEvent& InEvent, TArrayView<FAssetData> InAssets)
+		{
+			FScopedTransaction Transaction(LOCTEXT("AddAssets", "Add assets to workspace"));
+
+			Workspace->AddAssets(InAssets);
+		})
+		.OnAreAssetsAcceptableForDropWithReason_Lambda([this](TArrayView<FAssetData> InAssets, FText& OutText)
+		{
+			for(const FAssetData& Asset : InAssets)
+			{
+				if(Workspace->IsAssetSupported(Asset))
+				{
+					return true;
+				}
+			}
+
+			OutText = LOCTEXT("AssetsUnsupportedInWorkspace", "Assets are not supported by this workspace");
+			return false;
+		})
+		.Content()
+		[
+			ContentBrowserModule.Get().CreateAssetPicker(AssetPickerConfig)
+		]
+	];
+}
+
+FARFilter SWorkspaceView::MakeARFilter()
+{
+	FARFilter Filter;
+	Filter.ClassPaths = Workspace->GetSchema()->GetSupportedAssetClassPaths();
+	Filter.bRecursiveClasses = true;
+
+	return Filter;
+}
+
+FReply SWorkspaceView::OnKeyDown(const FGeometry& MyGeometry, const FKeyEvent& InKeyEvent)
+{
+	if (UICommandList.IsValid() && UICommandList->ProcessCommandBindings(InKeyEvent))
+	{
+		return FReply::Handled();
+	}
+	return FReply::Unhandled();
+}
+
+void SWorkspaceView::HandleDelete()
+{
+	if(GetCurrentSelectionDelegate.IsBound())
+	{
+		TArray<FAssetData> Selection = GetCurrentSelectionDelegate.Execute();
+
+		if(Selection.Num() > 0)
+		{
+			FScopedTransaction Transaction(LOCTEXT("RemoveAssets", "Remove assets from workspace"));
+
+			Workspace->RemoveAssets(Selection);
+		}
+	}
+}
+
+bool SWorkspaceView::HasValidSelection() const 
+{
+	if(GetCurrentSelectionDelegate.IsBound())
+	{
+		TArray<FAssetData> Selection = GetCurrentSelectionDelegate.Execute();
+		return Selection.Num() > 0;
+	}
+	return false;
+}
+
+void SWorkspaceView::HandleWorkspaceModified(UWorkspace* InWorkspace)
+{
+	check(InWorkspace == Workspace);
+
+	RefreshAssetViewDelegate.ExecuteIfBound(true);
+}
+
+}
+
+#undef LOCTEXT_NAMESPACE
\ No newline at end of file
diff --git a/Engine/Plugins/Experimental/Workspace/Source/WorkspaceEditor/Private/SWorkspaceView.h b/Engine/Plugins/Experimental/Workspace/Source/WorkspaceEditor/Private/SWorkspaceView.h
new file mode 100644
--- /dev/null
+++ b/Engine/Plugins/Experimental/Workspace/Source/WorkspaceEditor/Private/SWorkspaceView.h
@@ -0,0 +1,48 @@
+﻿// Copyright Epic Games, Inc. All Rights Reserved.
+
+#pragma once
+
+#include "ContentBrowserDelegates.h"
+#include "Widgets/SCompoundWidget.h"
+
+class UWorkspace;
+
+namespace UE::Workspace
+{
+
+class SWorkspaceView : public SCompoundWidget
+{
+	using FOnAssetsOpened = TDelegate<void(TConstArrayView<FAssetData>)>; 
+
+	SLATE_BEGIN_ARGS(SWorkspaceView) {}
+
+	SLATE_EVENT(FOnAssetsOpened, OnAssetsOpened)
+
+	SLATE_END_ARGS()
+
+	void Construct(const FArguments& InArgs, UWorkspace* InWorkspace);
+
+private:
+	// SWidget interface
+	virtual FReply OnKeyDown(const FGeometry& MyGeometry, const FKeyEvent& InKeyEvent) override;
+
+	void HandleDelete();
+
+	bool HasValidSelection() const;
+
+	void HandleWorkspaceModified(UWorkspace* InWorkspace);
+
+	FARFilter MakeARFilter();
+
+	UWorkspace* Workspace = nullptr;
+
+	FOnAssetsOpened OnAssetsOpened;
+
+	FRefreshAssetViewDelegate RefreshAssetViewDelegate;
+
+	FGetCurrentSelectionDelegate GetCurrentSelectionDelegate;
+
+	TSharedPtr<FUICommandList> UICommandList;
+};
+
+};
\ No newline at end of file
diff --git a/Engine/Plugins/Experimental/Workspace/Source/WorkspaceEditor/Private/Workspace.cpp b/Engine/Plugins/Experimental/Workspace/Source/WorkspaceEditor/Private/Workspace.cpp
new file mode 100644
--- /dev/null
+++ b/Engine/Plugins/Experimental/Workspace/Source/WorkspaceEditor/Private/Workspace.cpp
@@ -0,0 +1,267 @@
+// Copyright Epic Games, Inc. All Rights Reserved.
+
+#include "Workspace.h"
+
+#include "WorkspaceSchema.h"
+#include "WorkspaceState.h"
+#include "AssetRegistry/AssetData.h"
+#include "UObject/AssetRegistryTagsContext.h"
+#include "UObject/ConstructorHelpers.h"
+#include "UObject/FortniteMainBranchObjectVersion.h"
+
+const FName UWorkspace::ExportsAssetRegistryTag = TEXT("Exports");
+
+bool UWorkspace::AddAsset(const FAssetData& InAsset, bool bSetupUndoRedo, bool bPrintPythonCommand)
+{
+	if(!IsAssetSupported(InAsset))
+	{
+		ReportError(TEXT("UWorkspace::AddAsset: Unsupported asset supplied."));
+		return false;
+	}
+
+	if(bSetupUndoRedo)
+	{
+		Modify();
+	}
+
+	const int32 NewIndex = Assets.AddUnique(TSoftObjectPtr<UObject>(InAsset.GetSoftObjectPath()));
+	if(NewIndex != INDEX_NONE)
+	{
+		BroadcastModified();
+	}
+
+	return NewIndex != INDEX_NONE;
+}
+
+bool UWorkspace::AddAsset(UObject* InAsset, bool bSetupUndoRedo, bool bPrintPythonCommand)
+{
+	if(InAsset == nullptr)
+	{
+		ReportError(TEXT("UWorkspace::AddAsset: Invalid asset supplied."));
+		return false;
+	}
+
+	return AddAsset(FAssetData(InAsset), bSetupUndoRedo, bPrintPythonCommand);
+}
+
+bool UWorkspace::AddAssets(TConstArrayView<FAssetData> InAssets, bool bSetupUndoRedo, bool bPrintPythonCommand)
+{
+	if(InAssets.Num() == 0)
+	{
+		ReportError(TEXT("UWorkspace::AddAssets: No assets supplied."));
+		return false;
+	}
+
+	bool bAdded = false;
+	{
+		TGuardValue<bool> DisableNotifications(bSuspendNotifications, true);
+		for(const FAssetData& Asset : InAssets)
+		{
+			bAdded |= AddAsset(Asset, bSetupUndoRedo, bPrintPythonCommand);
+		}
+	}
+
+	if(bAdded)
+	{
+		BroadcastModified();
+	}
+
+	return bAdded;
+}
+
+bool UWorkspace::AddAssets(const TArray<UObject*>& InAssets, bool bSetupUndoRedo, bool bPrintPythonCommand)
+{
+	if(InAssets.Num() == 0)
+	{
+		ReportError(TEXT("UWorkspace::AddAssets: No assets supplied."));
+		return false;
+	}
+
+	bool bAdded = false;
+	{
+		TGuardValue<bool> DisableNotifications(bSuspendNotifications, true);
+		for(UObject* Asset : InAssets)
+		{
+			bAdded |= AddAsset(FAssetData(Asset), bSetupUndoRedo, bPrintPythonCommand);
+		}
+	}
+
+	if(bAdded)
+	{
+		BroadcastModified();
+	}
+
+	return bAdded;
+}
+
+bool UWorkspace::RemoveAsset(const FAssetData& InAsset, bool bSetupUndoRedo, bool bPrintPythonCommand)
+{
+	if(!IsAssetSupported(InAsset))
+	{
+		ReportError(TEXT("UWorkspace::RemoveAsset: Unsupported asset supplied."));
+		return false;
+	}
+
+	if(bSetupUndoRedo)
+	{
+		Modify();
+	}
+
+	const int32 NumRemoved = Assets.Remove(TSoftObjectPtr<UObject>(InAsset.GetSoftObjectPath()));
+	if(NumRemoved > 0)
+	{
+		BroadcastModified();
+	}
+
+	return NumRemoved > 0;
+}
+
+bool UWorkspace::RemoveAsset(UObject* InAsset, bool bSetupUndoRedo, bool bPrintPythonCommand)
+{
+	if(InAsset == nullptr)
+	{
+		ReportError(TEXT("UWorkspace::RemoveAsset: Invalid asset supplied."));
+		return false;
+	}
+
+	return RemoveAsset(FAssetData(InAsset), bSetupUndoRedo, bPrintPythonCommand);
+}
+
+bool UWorkspace::RemoveAssets(TArray<UObject*> InAssets, bool bSetupUndoRedo, bool bPrintPythonCommand)
+{
+	if(InAssets.Num() == 0)
+	{
+		ReportError(TEXT("UWorkspace::RemoveAssets: No assets supplied."));
+		return false;
+	}
+
+	bool bRemoved = false;
+	{
+		TGuardValue<bool> DisableNotifications(bSuspendNotifications, true);
+		for(UObject* Asset : InAssets)
+		{
+			bRemoved |= RemoveAsset(FAssetData(Asset), bSetupUndoRedo, bPrintPythonCommand);
+		}
+	}
+
+	if(bRemoved)
+	{
+		BroadcastModified();
+	}
+
+	return bRemoved;
+}
+
+bool UWorkspace::RemoveAssets(TConstArrayView<FAssetData> InAssets, bool bSetupUndoRedo, bool bPrintPythonCommand)
+{
+	if(InAssets.Num() == 0)
+	{
+		ReportError(TEXT("UWorkspace::RemoveAssets: No assets supplied."));
+		return false;
+	}
+
+	bool bRemoved = false;
+	{
+		TGuardValue<bool> DisableNotifications(bSuspendNotifications, true);
+		for(const FAssetData& Asset : InAssets)
+		{
+			bRemoved |= RemoveAsset(Asset, bSetupUndoRedo, bPrintPythonCommand);
+		}
+	}
+
+	if(bRemoved)
+	{
+		BroadcastModified();
+	}
+
+	return bRemoved;
+}
+
+bool UWorkspace::IsAssetSupported(const FAssetData& InAsset)
+{
+	TConstArrayView<FTopLevelAssetPath> SupportedAssets = GetSchema()->GetSupportedAssetClassPaths();
+	return SupportedAssets.IsEmpty() || SupportedAssets.Contains(InAsset.AssetClassPath);
+}
+
+UWorkspaceSchema* UWorkspace::GetSchema() const
+{
+	check(SchemaClass);
+	return SchemaClass->GetDefaultObject<UWorkspaceSchema>();
+}
+
+void UWorkspace::LoadState()
+{
+	GetState()->LoadFromJson(this);
+}
+
+void UWorkspace::SaveState()
+{
+	GetState()->SaveToJson(this);
+}
+
+UWorkspaceState* UWorkspace::GetState() const
+{
+	if(State == nullptr)
+	{
+		State = NewObject<UWorkspaceState>(const_cast<UWorkspace*>(this));
+	}
+
+	return State;
+}
+
+void UWorkspace::BroadcastModified()
+{
+	if(!bSuspendNotifications)
+	{
+		ModifiedDelegate.Broadcast(this);
+	}
+}
+
+void UWorkspace::ReportError(const TCHAR* InMessage) const
+{
+#if WITH_EDITOR
+	FScriptExceptionHandler::Get().HandleException(ELogVerbosity::Error, InMessage, TEXT(""));
+#endif
+}
+
+void UWorkspace::PostTransacted(const FTransactionObjectEvent& TransactionEvent)
+{
+	Super::PostTransacted(TransactionEvent);
+	BroadcastModified();
+}
+
+void UWorkspace::PostLoad()
+{
+	Super::PostLoad();
+
+	if(GetLinkerCustomVersion(FFortniteMainBranchObjectVersion::GUID) < FFortniteMainBranchObjectVersion::AnimNextMoveWorkspaces)
+	{
+		Guid = FGuid::NewGuid();
+		SchemaClass = StaticLoadClass(UWorkspaceSchema::StaticClass(), nullptr, TEXT("/Script/AnimNextEditor.AnimNextWorkspaceSchema"));
+	}
+}
+
+void UWorkspace::PostDuplicate(bool bDuplicateForPIE)
+{
+	Super::PostDuplicate(bDuplicateForPIE);
+
+	Guid = FGuid::NewGuid();
+}
+
+void UWorkspace::GetAssetRegistryTags(FAssetRegistryTagsContext Context) const
+{
+	Super::GetAssetRegistryTags(Context);
+
+	FWorkspaceAssetRegistryExports Exports;
+	Exports.Assets.Reserve(Assets.Num());
+
+	for(const TSoftObjectPtr<UObject>& Asset : Assets)
+	{
+		Exports.Assets.Emplace(Asset.GetUniqueID());
+	}
+
+	FString TagValue;
+	FWorkspaceAssetRegistryExports::StaticStruct()->ExportText(TagValue, &Exports, nullptr, nullptr, PPF_None, nullptr);
+
+	Context.AddTag(FAssetRegistryTag(ExportsAssetRegistryTag, TagValue, FAssetRegistryTag::TT_Hidden));
+}
diff --git a/Engine/Plugins/Experimental/Workspace/Source/WorkspaceEditor/Private/Workspace.h b/Engine/Plugins/Experimental/Workspace/Source/WorkspaceEditor/Private/Workspace.h
new file mode 100644
--- /dev/null
+++ b/Engine/Plugins/Experimental/Workspace/Source/WorkspaceEditor/Private/Workspace.h
@@ -0,0 +1,133 @@
+// Copyright Epic Games, Inc. All Rights Reserved.
+
+#pragma once
+
+#include "CoreMinimal.h"
+#include "Templates/SubclassOf.h"
+#include "Workspace.generated.h"
+
+struct FEditedDocumentInfo;
+class UWorkspace;
+class UWorkspaceSchema;
+class UWorkspaceFactory;
+class UAssetDefinition_Workspace;
+
+namespace UE::Workspace
+{
+	class FWorkspaceEditor;
+	class FWorkspaceEditorModule;
+	class SWorkspaceView;
+}
+
+namespace UE::Workspace
+{
+	// A delegate for subscribing / reacting to workspace modifications.
+	DECLARE_MULTICAST_DELEGATE_OneParam(FOnWorkspaceModified, UWorkspace* /* InWorkspace */);
+}
+
+// Workspace entry used to export to asset registry
+USTRUCT()
+struct FWorkspaceAssetRegistryExportEntry
+{
+	GENERATED_BODY()
+
+	FWorkspaceAssetRegistryExportEntry() = default;
+	
+	FWorkspaceAssetRegistryExportEntry(const FSoftObjectPath& InAsset)
+		: Asset(InAsset)
+	{}
+
+	UPROPERTY()
+	FSoftObjectPath Asset;
+};
+
+// Workspace used to export to asset registry
+USTRUCT()
+struct FWorkspaceAssetRegistryExports
+{
+	GENERATED_BODY()
+
+	UPROPERTY()
+	TArray<FWorkspaceAssetRegistryExportEntry> Assets;
+};
+
+UCLASS()
+class UWorkspace : public UObject
+{
+	GENERATED_BODY()
+
+	friend class UE::Workspace::FWorkspaceEditor;
+	friend class UE::Workspace::FWorkspaceEditorModule;
+	friend class UE::Workspace::SWorkspaceView;
+	friend class UWorkspaceFactory;
+	friend class UAssetDefinition_Workspace;
+	friend class UWorkspaceState;
+
+	WORKSPACEEDITOR_API static const FName ExportsAssetRegistryTag;
+
+	// Adds an asset to the workspace
+	// @return true if the asset was added
+	UFUNCTION(BlueprintCallable, Category = "Workspace")
+	bool AddAsset(UObject* InAsset, bool bSetupUndoRedo = true, bool bPrintPythonCommand = true);
+	bool AddAsset(const FAssetData& InAsset, bool bSetupUndoRedo = true, bool bPrintPythonCommand = true);
+
+	// Adds assets to the workspace
+	// @return true if an asset was added
+	UFUNCTION(BlueprintCallable, Category = "Workspace")
+	bool AddAssets(const TArray<UObject*>& InAssets, bool bSetupUndoRedo = true, bool bPrintPythonCommand = true);
+	bool AddAssets(TConstArrayView<FAssetData> InAssets, bool bSetupUndoRedo = true, bool bPrintPythonCommand = true);
+
+	// Removes an asset from the workspace
+	// @return true if the asset was removed
+	UFUNCTION(BlueprintCallable, Category = "Workspace")
+	bool RemoveAsset(UObject* InAsset, bool bSetupUndoRedo = true, bool bPrintPythonCommand = true);
+	bool RemoveAsset(const FAssetData& InAsset, bool bSetupUndoRedo = true, bool bPrintPythonCommand = true);
+
+	// Removes assets from the workspace
+	// @return true if the asset was removed
+	UFUNCTION(BlueprintCallable, Category = "Workspace")
+	bool RemoveAssets(TArray<UObject*> InAssets, bool bSetupUndoRedo = true, bool bPrintPythonCommand = true);
+	bool RemoveAssets(TConstArrayView<FAssetData> InAssets, bool bSetupUndoRedo = true, bool bPrintPythonCommand = true);
+
+	bool IsAssetSupported(const FAssetData& InAsset);
+
+	UWorkspaceSchema* GetSchema() const;
+
+	void LoadState();
+
+	void SaveState();
+
+	UWorkspaceState* GetState() const;
+
+	void ReportError(const TCHAR* InMessage) const;
+
+	void BroadcastModified();
+
+	// UObject interface
+	virtual void PostLoad() override;
+	virtual void PostDuplicate(bool bDuplicateForPIE) override;
+	virtual void GetAssetRegistryTags(FAssetRegistryTagsContext Context) const override;
+	virtual void PostTransacted(const FTransactionObjectEvent& TransactionEvent) override;
+	virtual bool IsEditorOnly() const override { return true; }
+
+	// All of the assets referenced by this workspace
+	UPROPERTY()
+	TArray<TSoftObjectPtr<UObject>> Assets;
+
+	// Schema for this workspace
+	UPROPERTY(AssetRegistrySearchable)
+	TSubclassOf<UWorkspaceSchema> SchemaClass = nullptr;
+
+	// State of the workspace, persisted to json
+	UPROPERTY(Transient)
+	mutable TObjectPtr<UWorkspaceState> State = nullptr;
+
+	// Guid for persistent identification of this workspace
+	UPROPERTY()
+	FGuid Guid;
+
+	// Delegate to subscribe to modifications
+	UE::Workspace::FOnWorkspaceModified ModifiedDelegate;
+
+	bool bSuspendNotifications = false;
+};
\ No newline at end of file
diff --git a/Engine/Plugins/Experimental/Workspace/Source/WorkspaceEditor/Private/WorkspaceAssetDefinition.cpp b/Engine/Plugins/Experimental/Workspace/Source/WorkspaceEditor/Private/WorkspaceAssetDefinition.cpp
new file mode 100644
--- /dev/null
+++ b/Engine/Plugins/Experimental/Workspace/Source/WorkspaceEditor/Private/WorkspaceAssetDefinition.cpp
@@ -0,0 +1,43 @@
+﻿// Copyright Epic Games, Inc. All Rights Reserved.
+
+#include "WorkspaceAssetDefinition.h"
+#include "WorkspaceEditor.h"
+#include "WorkspaceSchema.h"
+
+#define LOCTEXT_NAMESPACE "AssetDefinition_Workspace"
+
+FText UAssetDefinition_Workspace::GetAssetDisplayName() const
+{
+	return LOCTEXT("Workspace", "Workspace");
+}
+
+FText UAssetDefinition_Workspace::GetAssetDisplayName(const FAssetData& AssetData) const
+{
+	FString TagValue;
+	if(AssetData.GetTagValue(GET_MEMBER_NAME_CHECKED(UWorkspace, SchemaClass), TagValue))
+	{
+		if(UClass* SchemaClass = FindObject<UClass>(nullptr, *TagValue, true))
+		{
+			FText DisplayName = SchemaClass->GetDefaultObject<UWorkspaceSchema>()->GetDisplayName();
+			if(!DisplayName.IsEmpty())
+			{
+				return DisplayName;
+			}
+		}
+	}
+	
+	return GetAssetDisplayName();
+}
+
+EAssetCommandResult UAssetDefinition_Workspace::OpenAssets(const FAssetOpenArgs& OpenArgs) const
+{
+	for (UWorkspace* Asset : OpenArgs.LoadObjects<UWorkspace>())
+	{
+		TSharedRef<UE::Workspace::FWorkspaceEditor> Editor = MakeShared<UE::Workspace::FWorkspaceEditor>();
+		Editor->InitEditor(OpenArgs.GetToolkitMode(), OpenArgs.ToolkitHost, Asset);
+	}
+
+	return EAssetCommandResult::Handled;
+}
+
+#undef LOCTEXT_NAMESPACE
\ No newline at end of file
diff --git a/Engine/Plugins/Experimental/Workspace/Source/WorkspaceEditor/Private/WorkspaceAssetDefinition.h b/Engine/Plugins/Experimental/Workspace/Source/WorkspaceEditor/Private/WorkspaceAssetDefinition.h
new file mode 100644
--- /dev/null
+++ b/Engine/Plugins/Experimental/Workspace/Source/WorkspaceEditor/Private/WorkspaceAssetDefinition.h
@@ -0,0 +1,30 @@
+﻿// Copyright Epic Games, Inc. All Rights Reserved.
+
+#pragma once
+
+#include "Workspace.h"
+#include "AssetDefinitionDefault.h"
+#include "WorkspaceAssetDefinition.generated.h"
+
+#define LOCTEXT_NAMESPACE "AssetDefinitions"
+
+UCLASS()
+class UAssetDefinition_Workspace : public UAssetDefinitionDefault
+{
+	GENERATED_BODY()
+
+public:
+	// UAssetDefinition interface
+	virtual FText GetAssetDisplayName() const override;
+	virtual FText GetAssetDisplayName(const FAssetData& AssetData) const override;
+	virtual FLinearColor GetAssetColor() const override { return FLinearColor(FColor(170,96,48)); }
+	virtual TSoftClassPtr<UObject> GetAssetClass() const override { return UWorkspace::StaticClass(); }
+	virtual EAssetCommandResult OpenAssets(const FAssetOpenArgs& OpenArgs) const override;
+	virtual TConstArrayView<FAssetCategoryPath> GetAssetCategories() const override
+	{
+		static const auto Categories = { FAssetCategoryPath(EAssetCategoryPaths::Misc) };
+		return Categories;
+	}
+};
+
+#undef LOCTEXT_NAMESPACE
\ No newline at end of file
diff --git a/Engine/Plugins/Experimental/Workspace/Source/WorkspaceEditor/Private/WorkspaceEditor.cpp b/Engine/Plugins/Experimental/Workspace/Source/WorkspaceEditor/Private/WorkspaceEditor.cpp
new file mode 100644
--- /dev/null
+++ b/Engine/Plugins/Experimental/Workspace/Source/WorkspaceEditor/Private/WorkspaceEditor.cpp
@@ -0,0 +1,326 @@
+// Copyright Epic Games, Inc. All Rights Reserved.
+
+#include "WorkspaceEditor.h"
+
+#include "WorkspaceEditorMode.h"
+#include "Workspace.h"
+#include "AssetDocumentSummoner.h"
+#include "ExternalPackageHelper.h"
+#include "Widgets/Docking/SDockTab.h"
+#include "WorkflowOrientedApp/WorkflowUObjectDocuments.h"
+#include "SWorkspacePicker.h"
+#include "WorkspaceState.h"
+#include "WorkspaceDocumentState.h"
+#include "WorkspaceEditorModule.h"
+#include "Engine/Blueprint.h"
+
+#define LOCTEXT_NAMESPACE "WorkspaceEditor"
+
+namespace UE::Workspace
+{
+
+namespace WorkspaceModes
+{
+	const FName WorkspaceEditor("WorkspaceEditorMode");
+}
+
+namespace WorkspaceTabs
+{
+	const FName Details("DetailsTab");
+	const FName WorkspaceView("WorkspaceView");
+	const FName LeftDocumentArea("LeftDocumentArea");
+	const FName MiddleDocumentArea("MiddleDocumentArea");
+	const FName RightDocumentArea("RightDocumentArea");
+}
+
+const FName WorkspaceAppIdentifier("WorkspaceEditor");
+
+FWorkspaceEditor::FWorkspaceEditor()
+{
+}
+
+FWorkspaceEditor::~FWorkspaceEditor()
+{
+}
+
+void FWorkspaceEditor::InitEditor(const EToolkitMode::Type InMode, const TSharedPtr<IToolkitHost>& InInitToolkitHost, UWorkspace* InWorkspace)
+{
+	Workspace = InWorkspace;
+
+	Workspace->LoadState();
+
+	DocumentManager = MakeShared<FDocumentTracker>(NAME_None);
+	DocumentManager->Initialize(SharedThis(this));
+
+	FWorkspaceEditorModule& WorkspaceEditorModule = FModuleManager::Get().LoadModuleChecked<FWorkspaceEditorModule>("WorkspaceEditor");
+
+	// Build document summoners for each workspace layout area
+	TSharedRef<FAssetDocumentSummoner> LeftAssetDocumentSummoner = MakeShared<FAssetDocumentSummoner>(WorkspaceTabs::LeftDocumentArea, SharedThis(this));
+	LeftAssetDocumentSummoner->SetAllowedClassPaths(WorkspaceEditorModule.GetAllowedObjectTypesForArea(WorkspaceTabs::LeftDocumentArea));
+	DocumentManager->RegisterDocumentFactory(LeftAssetDocumentSummoner);
+
+	TSharedRef<FAssetDocumentSummoner> MiddleAssetDocumentSummoner = MakeShared<FAssetDocumentSummoner>(WorkspaceTabs::MiddleDocumentArea, SharedThis(this));
+	MiddleAssetDocumentSummoner->SetAllowedClassPaths(WorkspaceEditorModule.GetAllowedObjectTypesForArea(WorkspaceTabs::MiddleDocumentArea));
+	DocumentManager->RegisterDocumentFactory(MiddleAssetDocumentSummoner);
+
+	TSharedRef<FAssetDocumentSummoner> RightAssetDocumentSummoner = MakeShared<FAssetDocumentSummoner>(WorkspaceTabs::RightDocumentArea, SharedThis(this));
+	RightAssetDocumentSummoner->SetAllowedClassPaths(WorkspaceEditorModule.GetAllowedObjectTypesForArea(WorkspaceTabs::RightDocumentArea));
+	DocumentManager->RegisterDocumentFactory(RightAssetDocumentSummoner);
+
+	constexpr bool bCreateDefaultStandaloneMenu = true;
+	constexpr bool bCreateDefaultToolbar = true;
+	InitAssetEditor(InMode, InInitToolkitHost, WorkspaceAppIdentifier, FTabManager::FLayout::NullLayout, bCreateDefaultStandaloneMenu, bCreateDefaultToolbar, InWorkspace);
+
+	BindCommands();
+
+	AddApplicationMode(WorkspaceModes::WorkspaceEditor, MakeShared<FWorkspaceEditorMode>(SharedThis(this)));
+	SetCurrentMode(WorkspaceModes::WorkspaceEditor);
+
+	ExtendMenu();
+	ExtendToolbar();
+	RegenerateMenusAndToolbars();
+}
+
+void FWorkspaceEditor::RestoreEditedObjectState()
+{
+	UWorkspaceState* State = Workspace->GetState();
+	for (const TInstancedStruct<FWorkspaceDocumentState>& DocumentState : State->DocumentStates)
+	{
+		if (UObject* Object = DocumentState.Get().Object.TryLoad())
+		{
+			if(TSharedPtr<SDockTab> DockTab = OpenDocument(Object, FDocumentTracker::RestorePreviousDocument))
+			{
+				FWorkspaceEditorModule& WorkspaceEditorModule = FModuleManager::LoadModuleChecked<FWorkspaceEditorModule>("WorkspaceEditor");
+				const FObjectDocumentArgs* DocumentArgs = WorkspaceEditorModule.FindObjectDocumentType(Object->GetClass()->GetClassPathName());
+				if(DocumentArgs != nullptr && DocumentArgs->OnSetDocumentState.IsBound())
+				{
+					DocumentArgs->OnSetDocumentState.Execute(FWorkspaceEditorContext(SharedThis(this), Object), DockTab->GetContent(), DocumentState);
+				}
+			}
+		}
+	}
+}
+
+void FWorkspaceEditor::SaveEditedObjectState()
+{
+	// Clear edited document state
+	UWorkspaceState* State = Workspace->GetState();
+	State->DocumentStates.Empty();
+
+	// Ask all open documents to save their state, which will update edited documents
+	DocumentManager->SaveAllState();
+
+	// Persist state
+	Workspace->SaveState();
+}
+
+TSharedPtr<SDockTab> FWorkspaceEditor::OpenDocument(const UObject* InForObject, FDocumentTracker::EOpenDocumentCause InCause)
+{
+	TSharedRef<FTabPayload_UObject> Payload = FTabPayload_UObject::Make(InForObject);
+	TSharedPtr<SDockTab> NewTab = DocumentManager->OpenDocument(Payload, InCause);
+
+	if(InCause != FDocumentTracker::RestorePreviousDocument)
+	{
+		FWorkspaceEditorModule& WorkspaceEditorModule = FModuleManager::LoadModuleChecked<FWorkspaceEditorModule>("WorkspaceEditor");
+		const FObjectDocumentArgs* DocumentArgs = WorkspaceEditorModule.FindObjectDocumentType(InForObject->GetClass()->GetClassPathName());
+		if(DocumentArgs != nullptr && DocumentArgs->OnGetDocumentState.IsBound())
+		{
+			RecordDocumentState(DocumentArgs->OnGetDocumentState.Execute(FWorkspaceEditorContext(SharedThis(this), const_cast<UObject*>(InForObject)), NewTab->GetContent()));
+		}
+		else
+		{
+			RecordDocumentState(TInstancedStruct<FWorkspaceDocumentState>::Make(InForObject));
+		}
+	}
+
+	return NewTab;
+}
+
+void FWorkspaceEditor::OpenAssets(TConstArrayView<FAssetData> InAssets)
+{
+	for(const FAssetData& Asset : InAssets)
+	{
+		if(UObject* LoadedAsset = Asset.GetAsset())
+		{
+			OpenDocument(LoadedAsset, FDocumentTracker::EOpenDocumentCause::OpenNewDocument);
+		}
+	}
+}
+
+void FWorkspaceEditor::OpenObjects(TConstArrayView<UObject*> InObjects)
+{
+	for(UObject* Object : InObjects)
+	{
+		OpenDocument(Object, FDocumentTracker::EOpenDocumentCause::OpenNewDocument);
+	}
+}
+
+void FWorkspaceEditor::CloseObjects(TConstArrayView<UObject*> InObjects)
+{
+	if(InObjects.Num() > 0)
+	{
+		for(UObject* Object : InObjects)
+		{
+			CloseDocumentTab(Object);
+		}
+	}
+}
+
+void FWorkspaceEditor::SetDetailsObjects(const TArray<UObject*>& InObjects)
+{
+	if(DetailsView.IsValid())
+	{
+		DetailsView->SetObjects(InObjects);
+	}
+}
+
+void FWorkspaceEditor::RefreshDetails()
+{
+	if (DetailsView.IsValid())
+	{
+		DetailsView->ForceRefresh();
+	}
+}
+
+void FWorkspaceEditor::BindCommands()
+{
+}
+
+void FWorkspaceEditor::ExtendMenu()
+{
+	
+}
+
+void FWorkspaceEditor::ExtendToolbar()
+{
+	
+}
+
+void FWorkspaceEditor::RegisterTabSpawners(const TSharedRef<FTabManager>& InTabManager)
+{
+	DocumentManager->SetTabManager(InTabManager);
+
+	FWorkflowCentricApplication::RegisterTabSpawners(InTabManager);
+}
+
+void FWorkspaceEditor::UnregisterTabSpawners(const TSharedRef<FTabManager>& InTabManager)
+{
+	FWorkflowCentricApplication::UnregisterTabSpawners(InTabManager);
+}
+
+FName FWorkspaceEditor::GetToolkitFName() const
+{
+	return FName("WorkspaceEditor");
+}
+
+FText FWorkspaceEditor::GetBaseToolkitName() const
+{
+	return LOCTEXT("AppLabel", "WorkspaceEditor");
+}
+
+FString FWorkspaceEditor::GetWorldCentricTabPrefix() const
+{
+	return LOCTEXT("WorldCentricTabPrefix", "WorkspaceEditor ").ToString();
+}
+
+FLinearColor FWorkspaceEditor::GetWorldCentricTabColorScale() const
+{
+	return FLinearColor(0.3f, 0.2f, 0.5f, 0.5f);
+}
+
+void FWorkspaceEditor::InitToolMenuContext(FToolMenuContext& InMenuContext)
+{
+}
+
+void FWorkspaceEditor::SaveAsset_Execute()
+{
+	// If asset is a default 'Untitled' workspace, redirect to the 'save as' flow
+	FString AssetPath = Workspace->GetOutermost()->GetPathName();
+	if(AssetPath.StartsWith(TEXT("/Temp/Untitled")))
+	{
+		// Ensure we dont also 'save as' other externally linked assets at this point
+		TGuardValue<bool> SaveWorkspaceOnly(bSavingWorkspaceOnly, true);
+
+		SaveAssetAs_Execute();
+	}
+	else
+	{
+		FWorkflowCentricApplication::SaveAsset_Execute();
+	}
+}
+
+void FWorkspaceEditor::CloseDocumentTab(const UObject* DocumentID)
+{
+	UWorkspaceState* State = Workspace->GetState();
+	State->DocumentStates.Remove(TInstancedStruct<FWorkspaceDocumentState>::Make(DocumentID));
+
+	TSharedRef<FTabPayload_UObject> Payload = FTabPayload_UObject::Make(DocumentID);
+	DocumentManager->CloseTab(Payload);
+}
+
+bool FWorkspaceEditor::InEditingMode() const
+{
+	return true;
+}
+
+void FWorkspaceEditor::GetSaveableObjects(TArray<UObject*>& OutObjects) const
+{
+	// Base class will pick up edited object
+	FWorkflowCentricApplication::GetSaveableObjects(OutObjects);
+
+	if(!bSavingWorkspaceOnly)
+	{
+		for(TSoftObjectPtr<UObject>& SoftAsset : Workspace->Assets)
+		{
+			if(UObject* Asset = SoftAsset.Get())
+			{
+				// Add object referenced by workspace
+				OutObjects.Add(Asset);
+
+				// Get external objects too
+				FExternalPackageHelper::GetExternalSaveableObjects(Asset, OutObjects);
+			}
+		}
+	}
+}
+
+void FWorkspaceEditor::RecordDocumentState(const TInstancedStruct<FWorkspaceDocumentState>& InState)
+{
+	UWorkspaceState* State = Workspace->GetState();
+	State->DocumentStates.AddUnique(InState);
+}
+
+bool FWorkspaceEditor::OnRequestClose(EAssetEditorCloseReason InCloseReason)
+{
+	auto RequiresSave = [this]()
+	{
+		UPackage* Package = Workspace->GetOutermost();
+		return Package->GetPathName().StartsWith(TEXT("/Temp/Untitled"));
+	};
+
+	// Give the user opportunity to save temp workspaces
+	if(RequiresSave())
+	{
+		// Ensure we dont also 'save as' other externally linked assets at this point
+		TGuardValue<bool> SaveWorkspaceOnly(bSavingWorkspaceOnly, true);
+
+		SaveAssetAs_Execute();
+	}
+
+	return true;
+}
+
+void FWorkspaceEditor::OnClose()
+{
+	if (DetailsView.IsValid())
+	{
+		DetailsView->SetObject(nullptr);
+		DetailsView.Reset();
+	}
+
+	FWorkflowCentricApplication::OnClose();
+}
+
+}
+
+#undef LOCTEXT_NAMESPACE
\ No newline at end of file
diff --git a/Engine/Plugins/Experimental/Workspace/Source/WorkspaceEditor/Private/WorkspaceEditor.h b/Engine/Plugins/Experimental/Workspace/Source/WorkspaceEditor/Private/WorkspaceEditor.h
new file mode 100644
--- /dev/null
+++ b/Engine/Plugins/Experimental/Workspace/Source/WorkspaceEditor/Private/WorkspaceEditor.h
@@ -0,0 +1,117 @@
+// Copyright Epic Games, Inc. All Rights Reserved.
+
+#pragma once
+
+#include "CoreMinimal.h"
+#include "GraphEditor.h"
+#include "InstancedStruct.h"
+#include "IWorkspaceEditor.h"
+#include "WorkflowOrientedApp/WorkflowTabManager.h"
+
+struct FWorkspaceDocumentState;
+struct FWorkspaceAssetRegistryExports;
+class UWorkspace;
+class FDocumentTracker;
+class FDocumentTabFactory;
+class FTabInfo;
+class FTabManager;
+class IToolkitHost;
+
+namespace UE::Workspace
+{
+
+class FWorkspaceEditorMode;
+struct FGraphDocumentSummoner;
+struct FWorkspaceTabSummoner;
+struct FAssetDocumentSummoner;
+class FWorkspaceEditorModule;
+
+namespace WorkspaceModes
+{
+	extern const FName WorkspaceEditor;
+}
+
+namespace WorkspaceTabs
+{
+	extern const FName Details;
+	extern const FName WorkspaceView;
+}
+
+class FWorkspaceEditor : public IWorkspaceEditor
+{
+public:
+	FWorkspaceEditor();
+	virtual ~FWorkspaceEditor() override;
+
+	/** Edits the specified asset */
+	void InitEditor(const EToolkitMode::Type InMode, const TSharedPtr<IToolkitHost>& InInitToolkitHost, UWorkspace* InWorkspace);
+
+private:
+	friend class FWorkspaceEditorMode;
+	friend struct FGraphDocumentSummoner;
+	friend class SGraphDocument;
+	friend struct FWorkspaceTabSummoner;
+	friend struct FAssetDocumentSummoner;
+	friend class FWorkspaceEditorModule;
+
+	// IToolkit interface
+	virtual void RegisterTabSpawners(const TSharedRef<FTabManager>& InTabManager) override;
+	virtual void UnregisterTabSpawners(const TSharedRef<FTabManager>& InTabManager) override;
+	virtual FName GetToolkitFName() const override;
+	virtual FText GetBaseToolkitName() const override;
+	virtual FString GetWorldCentricTabPrefix() const override;
+	virtual FLinearColor GetWorldCentricTabColorScale() const override;
+	virtual void InitToolMenuContext(FToolMenuContext& InMenuContext) override;
+	virtual void SaveAsset_Execute() override;
+	virtual bool OnRequestClose(EAssetEditorCloseReason InCloseReason) override;
+	virtual void OnClose() override;
+
+	// FAssetEditorToolkit interface
+	virtual void GetSaveableObjects(TArray<UObject*>& OutObjects) const override;
+
+	// IWorkspaceEditor interface
+	virtual void OpenAssets(TConstArrayView<FAssetData> InAssets) override;
+	virtual void OpenObjects(TConstArrayView<UObject*> InObjects) override;
+	virtual void CloseObjects(TConstArrayView<UObject*> InObjects) override;
+	virtual void SetDetailsObjects(const TArray<UObject*>& InObjects) override;
+	virtual void RefreshDetails() override;
+
+	void BindCommands();
+
+	void ExtendMenu();
+
+	void ExtendToolbar();
+
+	void CloseDocumentTab(const UObject* DocumentID);
+
+	void HandleDetailsViewCreated(TSharedRef<IDetailsView> InDetailsView)
+	{
+		DetailsView = InDetailsView;
+	}
+	
+	bool InEditingMode() const;
+
+	void RestoreEditedObjectState();
+
+	void SaveEditedObjectState();
+
+	TSharedPtr<SDockTab> OpenDocument(const UObject* InForObject, FDocumentTracker::EOpenDocumentCause InCause);
+
+	void RecordDocumentState(const TInstancedStruct<FWorkspaceDocumentState>& InState);
+
+	// The asset we are editing
+	UWorkspace* Workspace = nullptr;
+
+	// Document tracker
+	TSharedPtr<FDocumentTracker> DocumentManager;
+
+	// Command list for this editor
+	TSharedPtr<FUICommandList> CommandList;
+
+	// Our details panel
+	TSharedPtr<IDetailsView> DetailsView;
+
+	bool bSavingWorkspaceOnly = false;
+};
+
+}
\ No newline at end of file
diff --git a/Engine/Plugins/Experimental/Workspace/Source/WorkspaceEditor/Private/WorkspaceEditorMode.cpp b/Engine/Plugins/Experimental/Workspace/Source/WorkspaceEditor/Private/WorkspaceEditorMode.cpp
new file mode 100644
--- /dev/null
+++ b/Engine/Plugins/Experimental/Workspace/Source/WorkspaceEditor/Private/WorkspaceEditorMode.cpp
@@ -0,0 +1,106 @@
+﻿// Copyright Epic Games, Inc. All Rights Reserved.
+
+#include "WorkspaceEditorMode.h"
+
+#include "WorkspaceEditor.h"
+#include "DetailsTabSummoner.h"
+#include "IWorkspaceEditorModule.h"
+#include "WorkspaceTabSummoner.h"
+#include "Widgets/Docking/SDockTab.h"
+
+#define LOCTEXT_NAMESPACE "WorkspaceEditorMode"
+
+namespace UE::Workspace
+{
+
+FWorkspaceEditorMode::FWorkspaceEditorMode(TSharedRef<FWorkspaceEditor> InHostingApp)
+	: FApplicationMode(WorkspaceModes::WorkspaceEditor)
+	, HostingAppPtr(InHostingApp)
+{
+	TSharedRef<FWorkspaceEditor> WorkspaceEditor = StaticCastSharedRef<FWorkspaceEditor>(InHostingApp);
+	
+	TabFactories.RegisterFactory(MakeShared<FDetailsTabSummoner>(WorkspaceEditor, FOnDetailsViewCreated::CreateSP(&WorkspaceEditor.Get(), &FWorkspaceEditor::HandleDetailsViewCreated)));
+	TabFactories.RegisterFactory(MakeShared<FWorkspaceTabSummoner>(WorkspaceEditor));
+
+	TabLayout = FTabManager::NewLayout("Standalone_WorkspaceEditor_Layout_v1.0")
+		->AddArea
+		(
+			FTabManager::NewPrimaryArea()
+			->SetOrientation(Orient_Vertical)
+			->Split
+			(
+				FTabManager::NewSplitter()
+				->SetSizeCoefficient(1.0f)
+				->SetOrientation(Orient_Horizontal)
+				->Split
+				(
+					FTabManager::NewStack()
+					->SetSizeCoefficient(0.25f)
+					->SetHideTabWell(false)
+					->AddTab(WorkspaceTabs::WorkspaceView, ETabState::OpenedTab)
+					->AddTab(WorkspaceTabs::LeftDocumentArea, ETabState::ClosedTab)
+				)
+				->Split
+				(
+					FTabManager::NewStack()
+					->SetSizeCoefficient(0.5f)
+					->SetHideTabWell(false)
+					->AddTab(WorkspaceTabs::MiddleDocumentArea, ETabState::ClosedTab)
+				)
+				->Split
+				(
+					FTabManager::NewStack()
+					->SetSizeCoefficient(0.25f)
+					->SetHideTabWell(false)
+					->AddTab(WorkspaceTabs::RightDocumentArea, ETabState::ClosedTab)
+					->AddTab(WorkspaceTabs::Details, ETabState::OpenedTab)
+				)
+			)
+		);
+
+	WorkspaceEditor->RegisterModeToolbarIfUnregistered(GetModeName());
+}
+
+void FWorkspaceEditorMode::RegisterTabFactories(TSharedPtr<FTabManager> InTabManager)
+{
+	TSharedPtr<FWorkspaceEditor> HostingApp = HostingAppPtr.Pin();
+	HostingApp->PushTabFactories(TabFactories);
+
+	FApplicationMode::RegisterTabFactories(InTabManager);
+}
+
+void FWorkspaceEditorMode::AddTabFactory(FCreateWorkflowTabFactory FactoryCreator)
+{
+	if (FactoryCreator.IsBound())
+	{
+		TabFactories.RegisterFactory(FactoryCreator.Execute(HostingAppPtr.Pin()));
+	}
+}
+
+void FWorkspaceEditorMode::RemoveTabFactory(FName TabFactoryID)
+{
+	TabFactories.UnregisterFactory(TabFactoryID);
+}
+
+void FWorkspaceEditorMode::PreDeactivateMode()
+{
+	FApplicationMode::PreDeactivateMode();
+
+	TSharedPtr<FWorkspaceEditor> WorkspaceEditor = HostingAppPtr.Pin();
+
+	WorkspaceEditor->SaveEditedObjectState();
+}
+
+void FWorkspaceEditorMode::PostActivateMode()
+{
+	// Reopen any documents that were open when the blueprint was last saved
+	TSharedPtr<FWorkspaceEditor> WorkspaceEditor = HostingAppPtr.Pin();
+	WorkspaceEditor->RestoreEditedObjectState();
+
+	FApplicationMode::PostActivateMode();
+}
+
+
+}
+
+#undef LOCTEXT_NAMESPACE
\ No newline at end of file
diff --git a/Engine/Plugins/Experimental/Workspace/Source/WorkspaceEditor/Private/WorkspaceEditorMode.h b/Engine/Plugins/Experimental/Workspace/Source/WorkspaceEditor/Private/WorkspaceEditorMode.h
new file mode 100644
--- /dev/null
+++ b/Engine/Plugins/Experimental/Workspace/Source/WorkspaceEditor/Private/WorkspaceEditorMode.h
@@ -0,0 +1,42 @@
+﻿// Copyright Epic Games, Inc. All Rights Reserved.
+
+#pragma once
+
+#include "CoreMinimal.h"
+#include "EdGraph/EdGraph.h"
+#include "WorkflowOrientedApp/ApplicationMode.h"
+#include "WorkflowOrientedApp/WorkflowTabManager.h"
+
+class FWorkflowCentricApplication;
+class SGraphEditor;
+
+namespace UE::Workspace
+{
+	class FWorkspaceEditor;
+}
+
+namespace UE::Workspace
+{
+
+class FWorkspaceEditorMode : public FApplicationMode
+{
+public:
+	FWorkspaceEditorMode(TSharedRef<FWorkspaceEditor> InHostingApp);
+
+private:
+	// FApplicationMode interface
+	virtual void RegisterTabFactories(TSharedPtr<FTabManager> InTabManager) override;
+	virtual void AddTabFactory(FCreateWorkflowTabFactory FactoryCreator) override;
+	virtual void RemoveTabFactory(FName TabFactoryID) override;
+	virtual void PreDeactivateMode() override;
+	virtual void PostActivateMode() override;
+
+	// The hosting app
+	TWeakPtr<FWorkspaceEditor> HostingAppPtr;
+
+	// The tab factories we support
+	FWorkflowAllowedTabSet TabFactories;
+};
+
+}
+
diff --git a/Engine/Plugins/Experimental/Workspace/Source/WorkspaceEditor/Private/WorkspaceEditorModule.cpp b/Engine/Plugins/Experimental/Workspace/Source/WorkspaceEditor/Private/WorkspaceEditorModule.cpp
new file mode 100644
--- /dev/null
+++ b/Engine/Plugins/Experimental/Workspace/Source/WorkspaceEditor/Private/WorkspaceEditorModule.cpp
@@ -0,0 +1,209 @@
+﻿// Copyright Epic Games, Inc. All Rights Reserved.
+
+#include "WorkspaceEditorModule.h"
+
+#include "EdGraphUtilities.h"
+#include "WorkspaceDocumentState.h"
+#include "GraphDocumentState.h"
+#include "SGraphDocument.h"
+#include "SWorkspacePicker.h"
+#include "Workspace.h"
+#include "WorkspaceEditor.h"
+#include "WorkspaceFactory.h"
+#include "AssetRegistry/AssetRegistryModule.h"
+#include "EdGraph/EdGraph.h"
+#include "EdGraph/EdGraphSchema.h"
+#include "Modules/ModuleManager.h"
+#include "Styling/AppStyle.h"
+
+#define LOCTEXT_NAMESPACE "WorkspaceEditorModule"
+
+namespace UE::Workspace
+{
+
+void FWorkspaceEditorModule::RegisterObjectDocumentType(const FTopLevelAssetPath& InClassPath, const FObjectDocumentArgs& InParams)
+{
+	ensure(InParams.SpawnLocation != NAME_None);
+
+	DocumentAreaMap.FindOrAdd(InParams.SpawnLocation).Add(InClassPath);
+	ObjectDocumentArgs.Add(InClassPath, InParams);
+}
+
+void FWorkspaceEditorModule::UnregisterObjectDocumentType(const FTopLevelAssetPath& InClassPath)
+{
+	if(const FObjectDocumentArgs* ExistingType = FindObjectDocumentType(InClassPath))
+	{
+		DocumentAreaMap.FindChecked(ExistingType->SpawnLocation).Remove(InClassPath);
+	}
+	ObjectDocumentArgs.Remove(InClassPath);
+}
+
+FObjectDocumentArgs FWorkspaceEditorModule::CreateGraphDocumentArgs(const FGraphDocumentWidgetArgs& InArgs)
+{
+	FObjectDocumentArgs Args;
+	Args.OnMakeDocumentWidget = FOnMakeDocumentWidget::CreateLambda([InArgs](const FWorkspaceEditorContext& InContext)
+	{
+		return SNew(SGraphDocument, StaticCastSharedRef<FWorkspaceEditor>(InContext.WorkspaceEditor), CastChecked<UEdGraph>(InContext.Object))
+			.OnCreateActionMenu_Lambda([OnCreateActionMenu = InArgs.OnCreateActionMenu](const FWorkspaceEditorContext& InContext, UEdGraph* InGraph, const FVector2D& InNodePosition, const TArray<UEdGraphPin*>& InDraggedPins, bool bInAutoExpand, SGraphEditor::FActionMenuClosed InOnMenuClosed)
+			{
+				if(OnCreateActionMenu.IsBound())
+				{
+					return OnCreateActionMenu.Execute(InContext, InGraph, InNodePosition, InDraggedPins, bInAutoExpand, InOnMenuClosed);
+				}
+				return FActionMenuContent();
+			})
+			.OnNodeTextCommitted(InArgs.OnNodeTextCommitted)
+			.OnDeleteSelectedNodes(InArgs.OnDeleteSelectedNodes)
+			.OnGraphSelectionChanged(InArgs.OnGraphSelectionChanged);
+	});
+	Args.OnGetTabIcon = FOnGetTabIcon::CreateLambda([](const FWorkspaceEditorContext& InContext)
+	{
+		return FAppStyle::Get().GetBrush(TEXT("GraphEditor.EventGraph_16x"));
+	});
+	Args.OnGetTabName = FOnGetTabName::CreateLambda([](const FWorkspaceEditorContext& InContext)
+	{
+		if(UEdGraph* Graph = Cast<UEdGraph>(InContext.Object))
+		{
+			if (const UEdGraphSchema* Schema = Graph->GetSchema())
+			{
+				FGraphDisplayInfo Info;
+				Schema->GetGraphDisplayInformation(*Graph, /*out*/ Info);
+				return Info.DisplayName;
+			}
+			else
+			{
+				// if we don't have a schema, we're dealing with a malformed (or incomplete graph)...
+				// possibly in the midst of some transaction - here we return the object's outer path 
+				// so we can at least get some context as to which graph we're referring
+				return FText::FromString(Graph->GetPathName());
+			}
+		}
+		return LOCTEXT("UnknownGraphName", "Unknown");
+	});
+	Args.OnGetDocumentState = FOnGetDocumentState::CreateLambda([](const FWorkspaceEditorContext& InContext, TSharedRef<SWidget> InWidget)
+	{
+		TSharedRef<SGraphDocument> GraphDocument = StaticCastSharedRef<SGraphDocument>(InWidget);
+		FVector2D ViewLocation;
+		float ZoomAmount;
+		GraphDocument->GraphEditor->GetViewLocation(ViewLocation, ZoomAmount);
+		return TInstancedStruct<FGraphDocumentState>::Make(InContext.Object, ViewLocation, ZoomAmount);
+	});
+	Args.OnSetDocumentState = FOnSetDocumentState::CreateLambda([](const FWorkspaceEditorContext& InContext, TSharedRef<SWidget> InWidget, const TInstancedStruct<FWorkspaceDocumentState>& InDocumentState)
+	{
+		if(const FGraphDocumentState* GraphDocumentState = InDocumentState.GetPtr<FGraphDocumentState>())
+		{
+			TSharedRef<SGraphDocument> GraphDocument = StaticCastSharedRef<SGraphDocument>(InWidget);
+			GraphDocument->GraphEditor->SetViewLocation(GraphDocumentState->ViewLocation, GraphDocumentState->ZoomAmount);
+		}
+	});
+
+	return Args;
+}
+
+bool FWorkspaceEditorModule::GetExportedAssetsForWorkspace(const FAssetData& InWorkspaceAsset, FWorkspaceAssetRegistryExports& OutExports)
+{
+	const FString TagValue = InWorkspaceAsset.GetTagValueRef<FString>(UWorkspace::ExportsAssetRegistryTag);
+	return FWorkspaceAssetRegistryExports::StaticStruct()->ImportText(*TagValue, &OutExports, nullptr, PPF_None, nullptr, FWorkspaceAssetRegistryExports::StaticStruct()->GetName()) != nullptr;
+}
+
+void FWorkspaceEditorModule::OpenWorkspaceForObject(UObject* InObject, EOpenWorkspaceMethod InOpenMethod)
+{
+	FAssetRegistryModule& AssetRegistryModule = FModuleManager::LoadModuleChecked<FAssetRegistryModule>(TEXT("AssetRegistry"));
+
+	TArray<FAssetData> RelevantWorkspaceAssets;
+
+	if(InOpenMethod != EOpenWorkspaceMethod::AlwaysOpenNewWorkspace)
+	{
+		// Look for existing workspaces that export this asset
+		FARFilter ARFilter;
+		ARFilter.ClassPaths.Add(UWorkspace::StaticClass()->GetClassPathName());
+		ARFilter.bRecursiveClasses = true;
+
+		TArray<FAssetData> AllWorkspaceAssets;
+		AssetRegistryModule.Get().GetAssets(ARFilter, AllWorkspaceAssets);
+
+		for(const FAssetData& WorkspaceAsset : AllWorkspaceAssets)
+		{
+			FWorkspaceAssetRegistryExports Exports;
+			GetExportedAssetsForWorkspace(WorkspaceAsset, Exports);
+
+			FSoftObjectPath ObjectPath(InObject);
+			for(const FWorkspaceAssetRegistryExportEntry& ExportEntry : Exports.Assets)
+			{
+				if(ExportEntry.Asset == ObjectPath)
+				{
+					RelevantWorkspaceAssets.Add(WorkspaceAsset);
+					break;
+				}
+			}
+		}
+	}
+
+	FWorkspaceEditor* WorkspaceEditor = nullptr;
+
+	auto HandleNewWorkspace = [InObject, &WorkspaceEditor]()
+	{
+		UWorkspaceFactory* Factory = NewObject<UWorkspaceFactory>();
+		UPackage* Package = CreatePackage(nullptr);
+		FName PackageName = *FPaths::GetBaseFilename(Package->GetName());
+		UWorkspace* NewWorkspace = CastChecked<UWorkspace>(Factory->FactoryCreateNew(UWorkspace::StaticClass(), Package, PackageName, RF_Public | RF_Standalone, NULL, GWarn));
+		NewWorkspace->AddAsset(InObject, false);
+		NewWorkspace->MarkPackageDirty();
+		TSharedRef<FWorkspaceEditor> Editor = MakeShared<FWorkspaceEditor>();
+		Editor->InitEditor(EToolkitMode::Standalone, nullptr, NewWorkspace);
+
+		WorkspaceEditor = &Editor.Get();
+	};
+
+	auto HandleExistingWorkspace = [InObject, &WorkspaceEditor](const FAssetData& InAssetData)
+	{
+		if(UWorkspace* ExistingWorkspace = Cast<UWorkspace>(InAssetData.GetAsset()))
+		{
+			GEditor->GetEditorSubsystem<UAssetEditorSubsystem>()->OpenEditorForAsset(ExistingWorkspace);
+
+			WorkspaceEditor = static_cast<FWorkspaceEditor*>(GEditor->GetEditorSubsystem<UAssetEditorSubsystem>()->FindEditorForAsset(ExistingWorkspace, true));
+		}
+	};
+
+	if(InOpenMethod == EOpenWorkspaceMethod::AlwaysOpenNewWorkspace || RelevantWorkspaceAssets.Num() == 0)
+	{
+		// No relevant workspaces, so open a new one and add the asset
+		HandleNewWorkspace();
+	}
+	else if(RelevantWorkspaceAssets.Num() == 1)
+	{
+		// One existing workspace, open it
+		HandleExistingWorkspace(RelevantWorkspaceAssets[0]);
+	}
+	else
+	{
+		// Multiple existing workspaces, present a window to let the user choose one to open with
+		TSharedRef<SWorkspacePicker> WorkspacePicker = SNew(SWorkspacePicker)
+			.WorkspaceAssets(RelevantWorkspaceAssets)
+			.OnAssetSelected_Lambda(HandleExistingWorkspace)
+			.OnNewAsset_Lambda(HandleNewWorkspace);
+
+		WorkspacePicker->ShowModal();
+	}
+	
+	if(WorkspaceEditor)
+	{
+		WorkspaceEditor->OpenAssets({InObject});
+	}
+}
+
+const FObjectDocumentArgs* FWorkspaceEditorModule::FindObjectDocumentType(const FTopLevelAssetPath& InClassPath) const
+{
+	return ObjectDocumentArgs.Find(InClassPath);
+}
+
+TArray<FTopLevelAssetPath> FWorkspaceEditorModule::GetAllowedObjectTypesForArea(FName InSpawnLocation) const
+{
+	return DocumentAreaMap.FindRef(InSpawnLocation).Array();
+}
+
+}
+
+IMPLEMENT_MODULE(UE::Workspace::FWorkspaceEditorModule, WorkspaceEditor);
+
+#undef LOCTEXT_NAMESPACE
\ No newline at end of file
diff --git a/Engine/Plugins/Experimental/Workspace/Source/WorkspaceEditor/Private/WorkspaceEditorModule.h b/Engine/Plugins/Experimental/Workspace/Source/WorkspaceEditor/Private/WorkspaceEditorModule.h
new file mode 100644
--- /dev/null
+++ b/Engine/Plugins/Experimental/Workspace/Source/WorkspaceEditor/Private/WorkspaceEditorModule.h
@@ -0,0 +1,45 @@
+﻿// Copyright Epic Games, Inc. All Rights Reserved.
+
+#pragma once
+
+#include "IWorkspaceEditorModule.h"
+#include "UObject/TopLevelAssetPath.h"
+
+struct FWorkspaceAssetRegistryExports;
+
+namespace UE::Workspace
+{
+	struct FAssetDocumentSummoner;
+	class FWorkspaceEditor;
+}
+
+namespace UE::Workspace
+{
+
+class FWorkspaceEditorModule : public IWorkspaceEditorModule
+{
+private:
+	// IWorkspaceModule interface
+	virtual void RegisterObjectDocumentType(const FTopLevelAssetPath& InClassPath, const FObjectDocumentArgs& InParams) override;
+	virtual void UnregisterObjectDocumentType(const FTopLevelAssetPath& InClassPath) override;
+	virtual FObjectDocumentArgs CreateGraphDocumentArgs(const FGraphDocumentWidgetArgs& InArgs) override;
+	virtual void OpenWorkspaceForObject(UObject* InObject, EOpenWorkspaceMethod InOpenMethod) override;
+
+	// Find an existing registered object document type
+	const FObjectDocumentArgs* FindObjectDocumentType(const FTopLevelAssetPath& InClassPath) const;
+
+	// Find the set of allowed object types for the specified spawn location
+	TArray<FTopLevelAssetPath> GetAllowedObjectTypesForArea(FName InSpawnLocation) const;
+
+	// Get the exported set of assets in asset registry tags for a workspace's FAssetData
+	static bool GetExportedAssetsForWorkspace(const FAssetData& InWorkspaceAsset, FWorkspaceAssetRegistryExports& OutExports);
+
+	TMap<FTopLevelAssetPath, FObjectDocumentArgs> ObjectDocumentArgs;
+
+	TMap<FName, TSet<FTopLevelAssetPath>> DocumentAreaMap;
+
+	friend struct FAssetDocumentSummoner;
+	friend class FWorkspaceEditor;
+};
+
+}
diff --git a/Engine/Plugins/Experimental/Workspace/Source/WorkspaceEditor/Private/WorkspaceFactory.cpp b/Engine/Plugins/Experimental/Workspace/Source/WorkspaceEditor/Private/WorkspaceFactory.cpp
new file mode 100644
--- /dev/null
+++ b/Engine/Plugins/Experimental/Workspace/Source/WorkspaceEditor/Private/WorkspaceFactory.cpp
@@ -0,0 +1,69 @@
+﻿// Copyright Epic Games, Inc. All Rights Reserved.
+
+#include "WorkspaceFactory.h"
+
+#include "ClassViewerFilter.h"
+#include "DefaultWorkspaceSchema.h"
+#include "Workspace.h"
+#include "WorkspaceSchema.h"
+#include "Kismet2/SClassPickerDialog.h"
+#include "UObject/Package.h"
+
+#define LOCTEXT_NAMESPACE "WorkspaceFactory"
+
+UWorkspaceFactory::UWorkspaceFactory()
+{
+	bCreateNew = true;
+	bEditAfterNew = true;
+	SupportedClass = UWorkspace::StaticClass();
+}
+
+bool UWorkspaceFactory::ConfigureProperties()
+{
+	if(SchemaClass == nullptr)
+	{
+		class FSchemaFilter : public IClassViewerFilter
+		{
+		public:
+			virtual bool IsClassAllowed(const FClassViewerInitializationOptions& InInitOptions, const UClass* InClass, TSharedRef<FClassViewerFilterFuncs> InFilterFuncs)
+			{
+				return !InClass->HasAnyClassFlags(CLASS_Abstract) && InClass->IsChildOf(UWorkspaceSchema::StaticClass());
+			}
+
+			virtual bool IsUnloadedClassAllowed(const FClassViewerInitializationOptions& InInitOptions, const TSharedRef<const IUnloadedBlueprintData> InUnloadedClassData, TSharedRef<FClassViewerFilterFuncs> InFilterFuncs)
+			{
+				return false;
+			}
+		};
+
+		FClassViewerInitializationOptions Options;
+		Options.NameTypeToDisplay = EClassViewerNameTypeToDisplay::DisplayName;
+		Options.bEditorClassesOnly = true;
+		Options.bExpandAllNodes = true;
+		Options.InitiallySelectedClass = UDefaultWorkspaceSchema::StaticClass();
+		Options.ClassFilters.Add(MakeShared<FSchemaFilter>());
+
+		UClass* ChosenSchema = nullptr;
+		if(SClassPickerDialog::PickClass(LOCTEXT("ChooseSchema", "Choose a Workspace Schema"), Options, ChosenSchema, UWorkspaceSchema::StaticClass()))
+		{
+			SchemaClass = ChosenSchema;
+		}
+	}
+
+	return SchemaClass != nullptr;
+}
+
+UObject* UWorkspaceFactory::FactoryCreateNew(UClass* Class, UObject* InParent, FName Name, EObjectFlags Flags, UObject* Context, FFeedbackContext* Warn, FName CallingContext)
+{
+	UWorkspace* NewWorkspace = NewObject<UWorkspace>(InParent, Class, Name, Flags | RF_Public | RF_Standalone | RF_Transactional | RF_LoadCompleted);
+	NewWorkspace->SchemaClass = SchemaClass;
+	NewWorkspace->Guid = FGuid::NewGuid();
+
+	// make sure the package is never cooked.
+	UPackage* Package = NewWorkspace->GetOutermost();
+	Package->SetPackageFlags(Package->GetPackageFlags() | PKG_EditorOnly);
+
+	return NewWorkspace;
+}
+
+#undef LOCTEXT_NAMESPACE
\ No newline at end of file
diff --git a/Engine/Plugins/Experimental/Workspace/Source/WorkspaceEditor/Private/WorkspaceState.cpp b/Engine/Plugins/Experimental/Workspace/Source/WorkspaceEditor/Private/WorkspaceState.cpp
new file mode 100644
--- /dev/null
+++ b/Engine/Plugins/Experimental/Workspace/Source/WorkspaceEditor/Private/WorkspaceState.cpp
@@ -0,0 +1,76 @@
+﻿// Copyright Epic Games, Inc. All Rights Reserved.
+
+#include "WorkspaceState.h"
+
+#include "JsonObjectConverter.h"
+#include "Workspace.h"
+#include "WorkspaceDocumentState.h"
+#include "AssetRegistry/AssetData.h"
+#include "AssetRegistry/AssetRegistryModule.h"
+#include "Misc/FileHelper.h"
+#include "Misc/Paths.h"
+#include "Modules/ModuleManager.h"
+
+UWorkspaceState::UWorkspaceState()
+{
+	IAssetRegistry& AssetRegistry = FModuleManager::Get().LoadModuleChecked<FAssetRegistryModule>(TEXT("AssetRegistry")).Get();
+	AssetRegistry.OnAssetRenamed().AddUObject(this, &UWorkspaceState::HandleAssetRenamed);
+}
+
+UWorkspaceState::~UWorkspaceState()
+{
+	IAssetRegistry& AssetRegistry = FModuleManager::Get().LoadModuleChecked<FAssetRegistryModule>(TEXT("AssetRegistry")).Get();
+	AssetRegistry.OnAssetRenamed().RemoveAll(this);
+}
+
+FString UWorkspaceState::GetWorkspacePath(const UWorkspace* InWorkspace)
+{
+	// Save the file to a specific directory, keyed by GUID
+	TStringBuilder<512> PathBuilder;
+	PathBuilder.Append(FPaths::ProjectSavedDir());
+	PathBuilder.Append(TEXT("Config/Workspace/"));
+	InWorkspace->Guid.AppendString(PathBuilder);
+	PathBuilder.Append(TEXT(".json"));
+	return PathBuilder.ToString();
+}
+
+void UWorkspaceState::SaveToJson(const UWorkspace* InWorkspace)
+{
+	// Record the path we saved this with
+	WorkspacePath = FSoftObjectPath(InWorkspace);
+
+	FString String;
+	const FString FileName = GetWorkspacePath(InWorkspace);
+	if(FJsonObjectConverter::UStructToJsonObjectString(StaticClass(), this, String))
+	{
+		FFileHelper::SaveStringToFile(String, *FileName);
+	}
+}
+
+void UWorkspaceState::LoadFromJson(const UWorkspace* InWorkspace)
+{
+	FString String;
+	const FString FileName = GetWorkspacePath(InWorkspace);
+	if(FFileHelper::LoadFileToString(String, *FileName))
+	{
+		FJsonObjectConverter::JsonObjectStringToUStruct<UWorkspaceState>(String, this);
+	}
+
+	// Record the path we loaded this with
+	WorkspacePath = FSoftObjectPath(InWorkspace);
+}
+
+void UWorkspaceState::HandleAssetRenamed(const FAssetData& InAssetData, const FString& InOldName)
+{
+	const FSoftObjectPath OldPath(*InOldName);
+
+	// See if any of our assets were renamed and update accordingly
+	for(TInstancedStruct<FWorkspaceDocumentState>& DocumentState : DocumentStates)
+	{
+		FSoftObjectPath& DocumentPath = DocumentState.GetMutable().Object;
+		if(DocumentPath == OldPath)
+		{
+			DocumentPath = InAssetData.GetSoftObjectPath();
+		}
+	}
+}
\ No newline at end of file
diff --git a/Engine/Plugins/Experimental/Workspace/Source/WorkspaceEditor/Private/WorkspaceState.h b/Engine/Plugins/Experimental/Workspace/Source/WorkspaceEditor/Private/WorkspaceState.h
new file mode 100644
--- /dev/null
+++ b/Engine/Plugins/Experimental/Workspace/Source/WorkspaceEditor/Private/WorkspaceState.h
@@ -0,0 +1,48 @@
+﻿// Copyright Epic Games, Inc. All Rights Reserved.
+
+#pragma once
+
+#include "InstancedStruct.h"
+#include "WorkspaceState.generated.h"
+
+class UWorkspace;
+struct FWorkspaceDocumentState;
+
+namespace UE::Workspace
+{
+	class FWorkspaceEditor;
+	class FWorkspaceEditorModule;
+}
+
+UCLASS()
+class UWorkspaceState : public UObject
+{
+	GENERATED_BODY()
+
+	UWorkspaceState();
+	virtual ~UWorkspaceState() override;
+
+	// Generate a file path used to persist workspace state
+	static FString GetWorkspacePath(const UWorkspace* InWorkspace);
+
+	// Save this workspace state to JSON
+	void SaveToJson(const UWorkspace* InWorkspace);
+
+	// Load this workspace state from JSON
+	void LoadFromJson(const UWorkspace* InWorkspace);
+
+	// Update soft object paths to any assets we reference
+	void HandleAssetRenamed(const FAssetData& InAssetData, const FString& InOldName);
+
+	// Path to the workspace we are persisting, to allow files to be more easily parsed out
+	UPROPERTY()
+	FSoftObjectPath WorkspacePath;
+	
+	// All the workspace documents we are persisting
+	UPROPERTY()
+	TArray<TInstancedStruct<FWorkspaceDocumentState>> DocumentStates;
+
+	friend class UE::Workspace::FWorkspaceEditor;
+	friend class UE::Workspace::FWorkspaceEditorModule;
+	friend class UWorkspace;
+};
\ No newline at end of file
diff --git a/Engine/Plugins/Experimental/Workspace/Source/WorkspaceEditor/Private/WorkspaceTabSummoner.cpp b/Engine/Plugins/Experimental/Workspace/Source/WorkspaceEditor/Private/WorkspaceTabSummoner.cpp
new file mode 100644
--- /dev/null
+++ b/Engine/Plugins/Experimental/Workspace/Source/WorkspaceEditor/Private/WorkspaceTabSummoner.cpp
@@ -0,0 +1,45 @@
+﻿// Copyright Epic Games, Inc. All Rights Reserved.
+
+#include "WorkspaceTabSummoner.h"
+#include "WorkspaceEditor.h"
+#include "SWorkspaceView.h"
+
+#define LOCTEXT_NAMESPACE "WorkspaceTabSummoner"
+
+namespace UE::Workspace
+{
+
+FWorkspaceTabSummoner::FWorkspaceTabSummoner(TSharedPtr<FWorkspaceEditor> InHostingApp)
+	: FWorkflowTabFactory(WorkspaceTabs::WorkspaceView, InHostingApp)
+{
+	HostingAppPtr = InHostingApp;
+
+	TabLabel = LOCTEXT("WorkspaceTabLabel", "Workspace");
+	TabIcon = FSlateIcon("EditorStyle", "LevelEditor.Tabs.Outliner");
+	ViewMenuDescription = LOCTEXT("WorkspaceTabMenuDescription", "Workspace");
+	ViewMenuTooltip = LOCTEXT("WorkspaceTabToolTip", "Shows the workspace outliner tab.");
+	bIsSingleton = true;
+
+	WorkspaceView = SNew(SWorkspaceView, InHostingApp->Workspace)
+		.OnAssetsOpened_Lambda([this](TConstArrayView<FAssetData> InAssets)
+		{
+			if(TSharedPtr<FWorkspaceEditor> HostingApp = HostingAppPtr.Pin())
+			{
+				HostingApp->OpenAssets(InAssets);
+			}
+		});
+}
+
+TSharedRef<SWidget> FWorkspaceTabSummoner::CreateTabBody(const FWorkflowTabSpawnInfo& Info) const
+{
+	return WorkspaceView.ToSharedRef();
+}
+
+FText FWorkspaceTabSummoner::GetTabToolTipText(const FWorkflowTabSpawnInfo& Info) const
+{
+	return ViewMenuTooltip;
+}
+
+}
+
+#undef LOCTEXT_NAMESPACE
\ No newline at end of file
diff --git a/Engine/Plugins/Experimental/Workspace/Source/WorkspaceEditor/Private/WorkspaceTabSummoner.h b/Engine/Plugins/Experimental/Workspace/Source/WorkspaceEditor/Private/WorkspaceTabSummoner.h
new file mode 100644
--- /dev/null
+++ b/Engine/Plugins/Experimental/Workspace/Source/WorkspaceEditor/Private/WorkspaceTabSummoner.h
@@ -0,0 +1,33 @@
+﻿// Copyright Epic Games, Inc. All Rights Reserved.
+
+#pragma once
+
+#include "WorkflowOrientedApp/WorkflowTabFactory.h"
+
+namespace UE::Workspace
+{
+	class FWorkspaceEditor;
+	class SWorkspaceView;
+}
+
+namespace UE::Workspace
+{
+
+struct FWorkspaceTabSummoner : public FWorkflowTabFactory
+{
+public:
+	FWorkspaceTabSummoner(TSharedPtr<FWorkspaceEditor> InHostingApp);
+
+private:
+	// FWorkflowTabFactory interface
+	virtual TSharedRef<SWidget> CreateTabBody(const FWorkflowTabSpawnInfo& Info) const override;
+	virtual FText GetTabToolTipText(const FWorkflowTabSpawnInfo& Info) const override;
+
+	// The hosting app
+	TWeakPtr<FWorkspaceEditor> HostingAppPtr;
+
+	// The widget this tab spawner wraps
+	TSharedPtr<SWorkspaceView> WorkspaceView;
+};
+
+}
\ No newline at end of file
diff --git a/Engine/Plugins/Experimental/Workspace/Source/WorkspaceEditor/Public/IWorkspaceEditor.h b/Engine/Plugins/Experimental/Workspace/Source/WorkspaceEditor/Public/IWorkspaceEditor.h
new file mode 100644
--- /dev/null
+++ b/Engine/Plugins/Experimental/Workspace/Source/WorkspaceEditor/Public/IWorkspaceEditor.h
@@ -0,0 +1,29 @@
+﻿// Copyright Epic Games, Inc. All Rights Reserved.
+
+#pragma once
+
+#include "WorkflowOrientedApp/WorkflowCentricApplication.h"
+
+namespace UE::Workspace
+{
+
+class IWorkspaceEditor : public FWorkflowCentricApplication
+{
+public:
+	// Open the supplied assets for editing within the workspace editor
+	virtual void OpenAssets(TConstArrayView<FAssetData> InAssets) = 0;
+
+	// Open the supplied objects for editing within the workspace editor
+	virtual void OpenObjects(TConstArrayView<UObject*> InObjects) = 0;
+
+	// Close the supplied objects if they are open for editing within the workspace editor
+	virtual void CloseObjects(TConstArrayView<UObject*> InObjects) = 0;
+
+	// Show the supplied objects in the workspace editor details panel
+	virtual void SetDetailsObjects(const TArray<UObject*>& InObjects) = 0;
+
+	// Refresh the workspace editor details panel
+	virtual void RefreshDetails() = 0;
+};
+
+}
\ No newline at end of file
diff --git a/Engine/Plugins/Experimental/Workspace/Source/WorkspaceEditor/Public/IWorkspaceEditorModule.h b/Engine/Plugins/Experimental/Workspace/Source/WorkspaceEditor/Public/IWorkspaceEditorModule.h
new file mode 100644
--- /dev/null
+++ b/Engine/Plugins/Experimental/Workspace/Source/WorkspaceEditor/Public/IWorkspaceEditorModule.h
@@ -0,0 +1,137 @@
+﻿// Copyright Epic Games, Inc. All Rights Reserved.
+
+#pragma once
+
+#include "CoreMinimal.h"
+#include "GraphEditor.h"
+#include "InstancedStruct.h"
+#include "Misc/Attribute.h"
+#include "Modules/ModuleInterface.h"
+
+struct FSlateBrush;
+class SWidget;
+struct FTopLevelAssetPath;
+class UEdGraph;
+struct FWorkspaceDocumentState;
+
+namespace UE::Workspace
+{
+	class IWorkspaceEditor;
+}
+
+namespace UE::Workspace
+{
+
+namespace WorkspaceTabs
+{
+	WORKSPACEEDITOR_API extern const FName LeftDocumentArea;
+	WORKSPACEEDITOR_API extern const FName MiddleDocumentArea;
+	WORKSPACEEDITOR_API extern const FName RightDocumentArea;
+}
+
+// Context passed to workspace editor delegates
+struct FWorkspaceEditorContext
+{
+	FWorkspaceEditorContext(const TSharedRef<IWorkspaceEditor>& InWorkspaceEditor, UObject* InObject)
+		: WorkspaceEditor(InWorkspaceEditor)
+		, Object(InObject)
+	{}
+
+	// The current workspace editor
+	TSharedRef<IWorkspaceEditor> WorkspaceEditor;
+
+	// The object being edited
+	UObject* Object;
+};
+
+using FOnMakeDocumentWidget = TDelegate<TSharedRef<SWidget>(const FWorkspaceEditorContext&)>;
+
+using FOnGetTabIcon = TDelegate<const FSlateBrush*(const FWorkspaceEditorContext&)>;
+
+using FOnGetTabName = TDelegate<TAttribute<FText>(const FWorkspaceEditorContext&)>;
+
+using FOnGetDocumentState = TDelegate<TInstancedStruct<FWorkspaceDocumentState>(const FWorkspaceEditorContext&, TSharedRef<SWidget>)>;
+
+using FOnSetDocumentState = TDelegate<void(const FWorkspaceEditorContext&, TSharedRef<SWidget>, const TInstancedStruct<FWorkspaceDocumentState>&)>;
+
+// Arguments used to make document widgets for objects
+struct FObjectDocumentArgs
+{
+	FObjectDocumentArgs() = default;
+
+	FObjectDocumentArgs(FOnMakeDocumentWidget InOnMakeDocumentWidget, FName InSpawnLocation = WorkspaceTabs::MiddleDocumentArea)
+		: OnMakeDocumentWidget(InOnMakeDocumentWidget)
+		, SpawnLocation(InSpawnLocation)
+	{}
+
+	// Delegate called to generate a widget for the supplied object
+	FOnMakeDocumentWidget OnMakeDocumentWidget;
+
+	// Delegate called to build a struct used to store the document's state
+	FOnGetDocumentState OnGetDocumentState;
+
+	// Delegate called to use a struct to restore the document's state
+	FOnSetDocumentState OnSetDocumentState;
+
+	// Delegate called to get the tab icon to display. If this is unset, the icon will default to the asset icon for the class
+	FOnGetTabIcon OnGetTabIcon;
+
+	// Delegate called to get the tab name to display. If this is unset, the object's name will be used
+	FOnGetTabName OnGetTabName;
+
+	// Where to spawn the widget in the workspace layout - e.g. one of WorkspaceTabs
+	FName SpawnLocation = WorkspaceTabs::MiddleDocumentArea;
+};
+
+using FOnDeleteSelectedNodes = TDelegate<void(const FWorkspaceEditorContext&, const FGraphPanelSelectionSet&)>;
+
+using FOnGraphSelectionChanged = TDelegate<void(const FWorkspaceEditorContext&, const FGraphPanelSelectionSet&)>;
+
+using FOnCreateActionMenu = TDelegate<FActionMenuContent(const FWorkspaceEditorContext&, UEdGraph*, const FVector2D&, const TArray<UEdGraphPin*>&, bool, SGraphEditor::FActionMenuClosed)>;
+
+using FOnNodeTextCommitted = TDelegate<void(const FWorkspaceEditorContext&, const FText&, ETextCommit::Type, UEdGraphNode*)>;
+
+// Arguments used to make document widgets for graphs
+struct FGraphDocumentWidgetArgs
+{
+	// Where to spawn the widget in the workspace layout - e.g. one of WorkspaceTabs
+	FName SpawnLocation = WorkspaceTabs::MiddleDocumentArea;
+
+	FOnCreateActionMenu OnCreateActionMenu;
+
+	FOnNodeTextCommitted OnNodeTextCommitted;
+
+	FOnDeleteSelectedNodes OnDeleteSelectedNodes;
+
+	FOnGraphSelectionChanged OnGraphSelectionChanged;
+};
+
+// Enum describing how to open a workspace
+enum class EOpenWorkspaceMethod : int32
+{
+	// If the asset is already used in a workspace, open that (if not already opened)
+	// If the asset is already used in more than one workspace, let the user choose the workspace to open it in
+	// If the asset is not yet in a workspace, create a default workspace, add the asset and open the workspace
+	Default,
+
+	// Always open a new workspace asset and add the asset to it
+	AlwaysOpenNewWorkspace,
+};
+
+class IWorkspaceEditorModule : public IModuleInterface
+{
+public:
+	// Open an object inside a workspace editor.
+	virtual void OpenWorkspaceForObject(UObject* InObject, EOpenWorkspaceMethod InOpenMethod) = 0;
+
+	// Register a widget factory method to spawn for a particular class
+	virtual void RegisterObjectDocumentType(const FTopLevelAssetPath& InClassPath, const FObjectDocumentArgs& InArgs) = 0;
+
+	// Unregister a widget factory method to spawn for a particular class
+	virtual void UnregisterObjectDocumentType(const FTopLevelAssetPath& InClassPath) = 0;
+
+	// Make the required args for a document widget for a UEdGraph
+	virtual FObjectDocumentArgs CreateGraphDocumentArgs(const FGraphDocumentWidgetArgs& InArgs) = 0;
+};
+
+}
diff --git a/Engine/Plugins/Experimental/Workspace/Source/WorkspaceEditor/Public/WorkspaceDocumentState.h b/Engine/Plugins/Experimental/Workspace/Source/WorkspaceEditor/Public/WorkspaceDocumentState.h
new file mode 100644
--- /dev/null
+++ b/Engine/Plugins/Experimental/Workspace/Source/WorkspaceEditor/Public/WorkspaceDocumentState.h
@@ -0,0 +1,26 @@
+﻿// Copyright Epic Games, Inc. All Rights Reserved.
+
+#pragma once
+
+#include "WorkspaceDocumentState.generated.h"
+
+// Base struct used to persist workspace document state
+USTRUCT()
+struct FWorkspaceDocumentState
+{
+	GENERATED_BODY()
+
+	FWorkspaceDocumentState() = default;
+
+	explicit FWorkspaceDocumentState(const UObject* InObject)
+		: Object(InObject)
+	{}
+
+	UPROPERTY()
+	FSoftObjectPath Object;
+
+	friend bool operator==(const FWorkspaceDocumentState& LHS, const FWorkspaceDocumentState& RHS)
+	{
+		return LHS.Object == RHS.Object;
+	}
+};
diff --git a/Engine/Plugins/Experimental/Workspace/Source/WorkspaceEditor/Public/WorkspaceFactory.h b/Engine/Plugins/Experimental/Workspace/Source/WorkspaceEditor/Public/WorkspaceFactory.h
new file mode 100644
--- /dev/null
+++ b/Engine/Plugins/Experimental/Workspace/Source/WorkspaceEditor/Public/WorkspaceFactory.h
@@ -0,0 +1,40 @@
+﻿// Copyright Epic Games, Inc. All Rights Reserved.
+
+#pragma once
+
+#include "CoreMinimal.h"
+#include "Factories/Factory.h"
+#include "WorkspaceFactory.generated.h"
+
+class UWorkspaceSchema;
+
+namespace UE::Workspace
+{
+	struct FUtils;
+	class FWorkspaceEditorModule;
+}
+
+UCLASS(BlueprintType)
+class WORKSPACEEDITOR_API UWorkspaceFactory : public UFactory
+{
+	GENERATED_BODY()
+
+protected:
+	UWorkspaceFactory();
+
+	// Set the schema class for workspaces produced with this factory
+	void SetSchemaClass(TSubclassOf<UWorkspaceSchema> InSchemaClass) { SchemaClass = InSchemaClass; }
+
+private:
+	virtual bool ConfigureProperties() override;
+	virtual UObject* FactoryCreateNew(UClass* Class, UObject* InParent, FName Name, EObjectFlags Flags, UObject* Context, FFeedbackContext* Warn, FName CallingContext) override;
+	virtual UObject* FactoryCreateNew(UClass* Class, UObject* InParent, FName Name, EObjectFlags Flags, UObject* Context, FFeedbackContext* Warn) override
+	{
+		return FactoryCreateNew(Class, InParent, Name, Flags, Context, Warn, NAME_None);
+	}
+
+	UPROPERTY()
+	TSubclassOf<UWorkspaceSchema> SchemaClass = nullptr;
+
+	friend class UE::Workspace::FWorkspaceEditorModule;
+};
\ No newline at end of file
diff --git a/Engine/Plugins/Experimental/Workspace/Source/WorkspaceEditor/Public/WorkspaceSchema.h b/Engine/Plugins/Experimental/Workspace/Source/WorkspaceEditor/Public/WorkspaceSchema.h
new file mode 100644
--- /dev/null
+++ b/Engine/Plugins/Experimental/Workspace/Source/WorkspaceEditor/Public/WorkspaceSchema.h
@@ -0,0 +1,21 @@
+﻿// Copyright Epic Games, Inc. All Rights Reserved.
+
+#pragma once
+
+#include "WorkspaceSchema.generated.h"
+
+class UWorkspace;
+
+// Workspace schema to configure workspace assets to specific use cases
+UCLASS(Abstract)
+class WORKSPACEEDITOR_API UWorkspaceSchema : public UObject
+{
+	GENERATED_BODY()
+
+public:
+	// Get the name to display for workspace assets that use this schema
+	virtual FText GetDisplayName() const PURE_VIRTUAL(UWorkspaceSchema::GetDisplayName, return FText::GetEmpty(); )
+
+	// Get the asset types that are supported by this workspace. If this is empty, all assets are assumed to be supported
+	virtual TConstArrayView<FTopLevelAssetPath> GetSupportedAssetClassPaths() const PURE_VIRTUAL(UWorkspaceSchema::GetSupportedAssetClassPaths, return TConstArrayView<FTopLevelAssetPath>(); )
+};
\ No newline at end of file
diff --git a/Engine/Plugins/Experimental/Workspace/Source/WorkspaceEditor/WorkspaceEditor.Build.cs b/Engine/Plugins/Experimental/Workspace/Source/WorkspaceEditor/WorkspaceEditor.Build.cs
new file mode 100644
--- /dev/null
+++ b/Engine/Plugins/Experimental/Workspace/Source/WorkspaceEditor/WorkspaceEditor.Build.cs
@@ -0,0 +1,42 @@
+// Copyright Epic Games, Inc. All Rights Reserved.
+
+namespace UnrealBuildTool.Rules
+{
+	public class WorkspaceEditor : ModuleRules
+	{
+		public WorkspaceEditor(ReadOnlyTargetRules Target) : base(Target)
+		{
+			PublicDependencyModuleNames.AddRange(
+				new string[]
+				{
+					"Core", 
+					"UnrealEd", 
+					"StructUtils"
+				});
+			
+			PrivateDependencyModuleNames.AddRange(
+				new string[]
+				{
+					"CoreUObject",
+					"Engine",
+					"AssetTools",
+					"SlateCore",
+					"Slate",
+					"InputCore",
+					"PropertyEditor",
+					"GraphEditor",
+					"ToolWidgets",
+					"ToolMenus",
+					"AssetDefinition",
+					"SourceControl", 
+					"KismetWidgets",
+					"DesktopWidgets",
+					"Settings",
+					"EditorWidgets",
+					"Json",
+					"JsonUtilities",
+				}
+			);
+		}
+	}
+}
\ No newline at end of file
diff --git a/Engine/Plugins/Experimental/Workspace/Workspace.uplugin b/Engine/Plugins/Experimental/Workspace/Workspace.uplugin
new file mode 100644
--- /dev/null
+++ b/Engine/Plugins/Experimental/Workspace/Workspace.uplugin
@@ -0,0 +1,33 @@
+{
+	"FileVersion" : 3,
+	"Version" : 1,
+	"VersionName" : "0.1",
+	"FriendlyName" : "Workspace",
+	"Description" : "Editor framework allowing multiple assets to be edited in a unified workspace UI",
+	"Category" : "Editor",
+	"CreatedBy" : "Epic Games, Inc.",
+	"CreatedByURL" : "https://epicgames.com",
+	"DocsURL" : "",
+	"MarketplaceURL" : "",
+	"SupportURL" : "",
+	"EnabledByDefault" : false,
+	"CanContainContent" : false,
+	"IsBetaVersion": true,
+	"IsExperimentalVersion": true,
+	"Installed" : false,
+	"Plugins" :
+	[
+		{
+			"Name" : "StructUtils",
+			"Enabled" : true
+		}
+	],
+	"Modules" :
+	[
+		{
+			"Name" : "WorkspaceEditor",
+			"Type" : "Editor",
+			"LoadingPhase" : "Default"
+		}
+	]
+}
diff --git a/Engine/Source/Runtime/Core/Public/UObject/FortniteMainBranchObjectVersion.h b/Engine/Source/Runtime/Core/Public/UObject/FortniteMainBranchObjectVersion.h
--- a/Engine/Source/Runtime/Core/Public/UObject/FortniteMainBranchObjectVersion.h
+++ b/Engine/Source/Runtime/Core/Public/UObject/FortniteMainBranchObjectVersion.h
@@ -423,6 +423,9 @@ struct FFortniteMainBranchObjectVersion
 		// Combine parameter blocks and graphs
 		AnimNextCombineParameterBlocksAndGraphs,
 
+		// Move workspaces to a seperate plugin
+		AnimNextMoveWorkspaces,
+
 		// -----<new versions can be added above this line>-------------------------------------------------
 		VersionPlusOne,
 		LatestVersion = VersionPlusOne - 1
