Refactor AnimNext trait binding to expose a trait stack binding. This exposes the trait stack as a first class citizen.

diff --git a/Engine/Plugins/Experimental/AnimNext/Source/AnimNext/Internal/TraitCore/TraitTemplate.h b/Engine/Plugins/Experimental/AnimNext/Source/AnimNext/Internal/TraitCore/TraitTemplate.h
--- a/Engine/Plugins/Experimental/AnimNext/Source/AnimNext/Internal/TraitCore/TraitTemplate.h
+++ b/Engine/Plugins/Experimental/AnimNext/Source/AnimNext/Internal/TraitCore/TraitTemplate.h
@@ -33,21 +33,24 @@ namespace UE::AnimNext
 		// Returns the trait registry handle
 		FTraitRegistryHandle GetRegistryHandle() const noexcept { return RegistryHandle; }
 
+		// Returns whether or not this trait template is valid
+		// A trait template is invalid if the trait implementation hasn't been registered
+		bool IsValid() const noexcept { return RegistryHandle.IsValid(); }
+
 		// Returns the trait mode
 		ETraitMode GetMode() const noexcept { return static_cast<ETraitMode>(Mode); }
 
-		// For base traits only, returns the number of additive traits on top in the stack
-		uint32 GetNumAdditiveTraits() const
+		// For base traits only, returns the number of traits on the sub-stack
+		uint32 GetNumStackTraits() const
 		{
 			check(GetMode() == ETraitMode::Base);
-			return AdditiveIndexOrNumAdditive;
+			return TraitIndexOrNumTraits;
 		}
 
-		// For additive traits only, returns the trait index relative to the base trait on the stack
-		uint32 GetAdditiveTraitIndex() const
+		// Returns the trait index relative to the base trait on the sub-stack
+		uint32 GetTraitIndex() const
 		{
-			check(GetMode() == ETraitMode::Additive);
-			return AdditiveIndexOrNumAdditive;
+			return GetMode() == ETraitMode::Base ? 0 : TraitIndexOrNumTraits;
 		}
 
 		// Returns the number of latent properties on this trait
@@ -128,11 +131,11 @@ namespace UE::AnimNext
 		friend struct FNodeTemplate;
 		friend struct FNodeTemplateBuilder;
 
-		FTraitTemplate(FTraitUID InUID, FTraitRegistryHandle InRegistryHandle, ETraitMode InMode, uint32 InAdditiveIndexOrNumAdditive) noexcept
+		FTraitTemplate(FTraitUID InUID, FTraitRegistryHandle InRegistryHandle, ETraitMode InMode, uint32 InTraitIndexOrNumTraits) noexcept
 			: UID(InUID.GetUID())
 			, RegistryHandle(InRegistryHandle)
 			, Mode(static_cast<uint8>(InMode))
-			, AdditiveIndexOrNumAdditive(InAdditiveIndexOrNumAdditive)
+			, TraitIndexOrNumTraits(InTraitIndexOrNumTraits)
 			, NumLatentProperties(0)
 			, NumSubStackLatentProperties(0)
 			// For shared and instance data, 0 is an invalid offset since the data follows their respective header (FNodeDescription or FNodeInstance)
@@ -140,7 +143,8 @@ namespace UE::AnimNext
 			, NodeSharedLatentPropertyHandlesOffset(0)
 			, NodeInstanceOffset(0)
 			, Padding0(0)
-		{}
+		{
+		}
 
 		// Trait globally unique identifier (32 bits)
 		FTraitUIDRaw				UID;
@@ -151,10 +155,10 @@ namespace UE::AnimNext
 		// Trait mode (we only need 1 bit, we could store other flags here)
 		uint8	Mode;
 
-		// For base traits, this contains the number of additive traits on top
-		// For additive traits, this contains its index relative to the base trait
-		// The first additive trait has index 1, there is no index 0 (we re-purpose it for the base trait)
-		uint8	AdditiveIndexOrNumAdditive;
+		// For base traits, this contains the number of traits on our sub-stack
+		// For additive traits, this contains the trait index relative to the base trait
+		// The first additive trait has index 1, there is no index 0 (we re-purpose the value for the base trait)
+		uint8	TraitIndexOrNumTraits;
 
 		// For each latent property defined on a trait, we store a handle in the per node shared data
 		// These handles specify various metadata of the latent property: RigVM memory handle index, whether the property can freeze, instance data offset
diff --git a/Engine/Plugins/Experimental/AnimNext/Source/AnimNext/Private/Graph/AnimNextGraph.cpp b/Engine/Plugins/Experimental/AnimNext/Source/AnimNext/Private/Graph/AnimNextGraph.cpp
--- a/Engine/Plugins/Experimental/AnimNext/Source/AnimNext/Private/Graph/AnimNextGraph.cpp
+++ b/Engine/Plugins/Experimental/AnimNext/Source/AnimNext/Private/Graph/AnimNextGraph.cpp
@@ -78,7 +78,7 @@ void UAnimNextGraph::AllocateInstanceImpl(FAnimNextGraphInstance* ParentGraphIns
 
 	{
 		UE::AnimNext::FExecutionContext Context(Instance);
-		Instance.Impl->GraphInstancePtr = Context.AllocateNodeInstance(UE::AnimNext::FWeakTraitPtr(), ResolvedRootTraitHandle);
+		Instance.Impl->GraphInstancePtr = Context.AllocateNodeInstance(*Instance.Impl, ResolvedRootTraitHandle);
 	}
 
 	if (!Instance.IsValid())
diff --git a/Engine/Plugins/Experimental/AnimNext/Source/AnimNext/Private/Graph/AnimNextGraphInstance.cpp b/Engine/Plugins/Experimental/AnimNext/Source/AnimNext/Private/Graph/AnimNextGraphInstance.cpp
--- a/Engine/Plugins/Experimental/AnimNext/Source/AnimNext/Private/Graph/AnimNextGraphInstance.cpp
+++ b/Engine/Plugins/Experimental/AnimNext/Source/AnimNext/Private/Graph/AnimNextGraphInstance.cpp
@@ -161,7 +161,7 @@ void FAnimNextGraphInstance::Thaw()
 			UE::AnimNext::FExecutionContext Context(*this);
 			if(const FAnimNextTraitHandle* FoundHandle = GraphPtr->ResolvedRootTraitHandles.Find(EntryPoint))
 			{
-				GraphInstancePtr = Context.AllocateNodeInstance(UE::AnimNext::FWeakTraitPtr(), *FoundHandle);
+				GraphInstancePtr = Context.AllocateNodeInstance(*this, *FoundHandle);
 			}
 		}
 
diff --git a/Engine/Plugins/Experimental/AnimNext/Source/AnimNext/Private/Graph/GC_GraphInstanceComponent.cpp b/Engine/Plugins/Experimental/AnimNext/Source/AnimNext/Private/Graph/GC_GraphInstanceComponent.cpp
--- a/Engine/Plugins/Experimental/AnimNext/Source/AnimNext/Private/Graph/GC_GraphInstanceComponent.cpp
+++ b/Engine/Plugins/Experimental/AnimNext/Source/AnimNext/Private/Graph/GC_GraphInstanceComponent.cpp
@@ -7,17 +7,17 @@
 
 namespace UE::AnimNext
 {
-	void FGCGraphInstanceComponent::Register(FAnimNextGraphInstance& GraphInstance, const FWeakTraitPtr& TraitPtr)
+	void FGCGraphInstanceComponent::Register(const FWeakTraitPtr& InTraitPtr)
 	{
-		TraitsWithReferences.Add(FEntry(GraphInstance, TraitPtr));
+		TraitsWithReferences.Add(InTraitPtr);
 	}
 
-	void FGCGraphInstanceComponent::Unregister(const FWeakTraitPtr& TraitPtr)
+	void FGCGraphInstanceComponent::Unregister(const FWeakTraitPtr& InTraitPtr)
 	{
 		const int32 EntryIndex = TraitsWithReferences.IndexOfByPredicate(
-			[&TraitPtr](const FEntry& Entry)
+			[&InTraitPtr](const FWeakTraitPtr& TraitPtr)
 			{
-				return Entry.TraitPtr == TraitPtr;
+				return TraitPtr == InTraitPtr;
 			});
 
 		if (ensure(EntryIndex != INDEX_NONE))
@@ -29,15 +29,20 @@ namespace UE::AnimNext
 	void FGCGraphInstanceComponent::AddReferencedObjects(FReferenceCollector& Collector) const
 	{
 		FExecutionContext Context;
+		FTraitStackBinding TraitStack;
 		TTraitBinding<IGarbageCollection> GCTrait;
 
 		// TODO: If we kept the entries sorted by graph instance, we could re-use the execution context
-		for (const FEntry& Entry : TraitsWithReferences)
+		for (const FWeakTraitPtr& TraitPtr : TraitsWithReferences)
 		{
-			Context.BindTo(Entry.GraphInstance);
-			ensure(Context.GetInterface(Entry.TraitPtr, GCTrait));
+			Context.BindTo(TraitPtr);
 
-			GCTrait.AddReferencedObjects(Context, Collector);
+			if (Context.GetStack(TraitPtr, TraitStack))
+			{
+				ensure(TraitStack.GetInterface(GCTrait));
+
+				GCTrait.AddReferencedObjects(Context, Collector);
+			}
 		}
 	}
 }
diff --git a/Engine/Plugins/Experimental/AnimNext/Source/AnimNext/Private/Graph/GC_GraphInstanceComponent.h b/Engine/Plugins/Experimental/AnimNext/Source/AnimNext/Private/Graph/GC_GraphInstanceComponent.h
--- a/Engine/Plugins/Experimental/AnimNext/Source/AnimNext/Private/Graph/GC_GraphInstanceComponent.h
+++ b/Engine/Plugins/Experimental/AnimNext/Source/AnimNext/Private/Graph/GC_GraphInstanceComponent.h
@@ -6,7 +6,6 @@
 #include "TraitCore/TraitPtr.h"
 #include "Graph/GraphInstanceComponent.h"
 
-struct FAnimNextGraphInstance;
 class FReferenceCollector;
 
 namespace UE::AnimNext
@@ -22,28 +21,16 @@ namespace UE::AnimNext
 
 		// Registers the provided trait with the GC system
 		// Once registered, IGarbageCollection::AddReferencedObjects will be called on it during GC
-		void Register(FAnimNextGraphInstance& GraphInstance, const FWeakTraitPtr& TraitPtr);
+		void Register(const FWeakTraitPtr& InTraitPtr);
 
 		// Unregisters the provided trait from the GC system
-		void Unregister(const FWeakTraitPtr& TraitPtr);
+		void Unregister(const FWeakTraitPtr& InTraitPtr);
 
 		// Called during garbage collection to collect strong object references
 		void AddReferencedObjects(FReferenceCollector& Collector) const;
 
 	private:
-		struct FEntry
-		{
-			FEntry(FAnimNextGraphInstance& InGraphInstance, const FWeakTraitPtr& InTraitPtr)
-				: GraphInstance(InGraphInstance)
-				, TraitPtr(InTraitPtr)
-			{
-			}
-
-			FAnimNextGraphInstance& GraphInstance;
-			FWeakTraitPtr TraitPtr;
-		};
-
 		// List of trait handles that contain UObject references and implement IGarbageCollection
-		TArray<FEntry> TraitsWithReferences;
+		TArray<FWeakTraitPtr> TraitsWithReferences;
 	};
 }
diff --git a/Engine/Plugins/Experimental/AnimNext/Source/AnimNext/Private/Graph/SyncGroup_GraphInstanceComponent.cpp b/Engine/Plugins/Experimental/AnimNext/Source/AnimNext/Private/Graph/SyncGroup_GraphInstanceComponent.cpp
--- a/Engine/Plugins/Experimental/AnimNext/Source/AnimNext/Private/Graph/SyncGroup_GraphInstanceComponent.cpp
+++ b/Engine/Plugins/Experimental/AnimNext/Source/AnimNext/Private/Graph/SyncGroup_GraphInstanceComponent.cpp
@@ -19,6 +19,7 @@ namespace UE::AnimNext
 
 	void FSyncGroupGraphInstanceComponent::PostUpdate(FExecutionContext& Context)
 	{
+		FTraitStackBinding TraitStack;
 		TTraitBinding<IGroupSynchronization> GroupSyncTrait;
 
 		// Now that we have discovered all groups and their memberships, we can perform synchronization
@@ -74,7 +75,8 @@ namespace UE::AnimNext
 				const FSyncGroupMember& GroupLeader = GroupState.Members[LeaderIndex];
 
 				Context.BindTo(GroupLeader.TraitPtr);
-				ensure(Context.GetInterface(GroupLeader.TraitPtr, GroupSyncTrait));
+				ensure(Context.GetStack(GroupLeader.TraitPtr, TraitStack));
+				ensure(TraitStack.GetInterface(GroupSyncTrait));
 
 				LeaderProgressRatio = GroupSyncTrait.AdvanceBy(Context, GroupLeader.TraitState.GetDeltaTime());
 			}
@@ -90,7 +92,8 @@ namespace UE::AnimNext
 				const FSyncGroupMember& GroupMember = GroupState.Members[MemberIndex];
 
 				Context.BindTo(GroupMember.TraitPtr);
-				ensure(Context.GetInterface(GroupMember.TraitPtr, GroupSyncTrait));
+				ensure(Context.GetStack(GroupMember.TraitPtr, TraitStack));
+				ensure(TraitStack.GetInterface(GroupSyncTrait));
 
 				GroupSyncTrait.AdvanceToRatio(Context, LeaderProgressRatio);
 			}
diff --git a/Engine/Plugins/Experimental/AnimNext/Source/AnimNext/Private/TraitCore/ExecutionContext.cpp b/Engine/Plugins/Experimental/AnimNext/Source/AnimNext/Private/TraitCore/ExecutionContext.cpp
--- a/Engine/Plugins/Experimental/AnimNext/Source/AnimNext/Private/TraitCore/ExecutionContext.cpp
+++ b/Engine/Plugins/Experimental/AnimNext/Source/AnimNext/Private/TraitCore/ExecutionContext.cpp
@@ -15,9 +15,9 @@
 namespace UE::AnimNext
 {
 	FExecutionContext::FExecutionContext()
-		: NodeTemplateRegistry(FNodeTemplateRegistry::Get())
+		: MemStack(FMemStack::Get())
+		, NodeTemplateRegistry(FNodeTemplateRegistry::Get())
 		, TraitRegistry(FTraitRegistry::Get())
-		, GraphInstance(nullptr)
 	{
 	}
 
@@ -43,16 +43,13 @@ namespace UE::AnimNext
 
 	void FExecutionContext::BindTo(FAnimNextGraphInstance& InGraphInstance)
 	{
-		if (GraphInstance == &InGraphInstance)
+		FAnimNextGraphInstance* InRootGraphInstance = InGraphInstance.GetRootGraphInstance();
+		if (RootGraphInstance == InRootGraphInstance)
 		{
-			return;	// Already bound to this graph instance, nothing to do
+			return;	// Already bound to this root graph instance, nothing to do
 		}
 
-		if (const UAnimNextGraph* Graph = InGraphInstance.GetGraph())
-		{
-			GraphInstance = &InGraphInstance;
-			GraphSharedData = Graph->SharedDataBuffer;
-		}
+		RootGraphInstance = InRootGraphInstance;
 	}
 
 	void FExecutionContext::BindTo(const FWeakTraitPtr& TraitPtr)
@@ -65,22 +62,27 @@ namespace UE::AnimNext
 
 	bool FExecutionContext::IsBound() const
 	{
-		return GraphInstance != nullptr;
+		return RootGraphInstance != nullptr;
 	}
 
 	bool FExecutionContext::IsBoundTo(const FAnimNextGraphInstancePtr& InGraphInstance) const
 	{
-		return GraphInstance == InGraphInstance.GetImpl();
+		if (FAnimNextGraphInstance* GraphInstance = InGraphInstance.GetImpl())
+		{
+			return RootGraphInstance == GraphInstance->GetRootGraphInstance();
+		}
+
+		return false;
 	}
 
 	bool FExecutionContext::IsBoundTo(const FAnimNextGraphInstance& InGraphInstance) const
 	{
-		return GraphInstance == &InGraphInstance;
+		return RootGraphInstance == InGraphInstance.GetRootGraphInstance();
 	}
 
-	FTraitPtr FExecutionContext::AllocateNodeInstance(const FWeakTraitPtr& ParentBinding, FAnimNextTraitHandle ChildTraitHandle) const
+	FTraitPtr FExecutionContext::AllocateNodeInstance(FAnimNextGraphInstance& GraphInstance, FAnimNextTraitHandle ChildTraitHandle) const
 	{
-		if (!ensure(ChildTraitHandle.IsValid()))
+		if (!ChildTraitHandle.IsValid())
 		{
 			return FTraitPtr();	// Attempting to allocate a node using an invalid trait handle
 		}
@@ -90,42 +92,41 @@ namespace UE::AnimNext
 			return FTraitPtr();	// The execution context must be bound to a valid graph instance
 		}
 
+		if (!ensure(GraphInstance.GetGraph() != nullptr))
+		{
+			return FTraitPtr();	// We need a valid graph instance to allocate into
+		}
+
 		const FNodeHandle ChildNodeHandle = ChildTraitHandle.GetNodeHandle();
-		const FNodeDescription& NodeDesc = GetNodeDescription(ChildNodeHandle);
+		const uint32 ChildTraitIndex = ChildTraitHandle.GetTraitIndex();
 
+		const FNodeDescription& NodeDesc = GetNodeDescription(GraphInstance, ChildNodeHandle);
 		const FNodeTemplate* NodeTemplate = GetNodeTemplate(NodeDesc);
+
 		if (!ensure(NodeTemplate != nullptr))
 		{
 			return FTraitPtr();	// Node template wasn't found, node descriptor is perhaps corrupted
 		}
-
-		const uint32 ChildTraitIndex = ChildTraitHandle.GetTraitIndex();
-
-		if (ChildTraitIndex >= NodeTemplate->GetNumTraits())
+		else if (ChildTraitIndex >= NodeTemplate->GetNumTraits())
 		{
 			return FTraitPtr();	// The requested trait index doesn't exist on that node descriptor
 		}
 
-		// If the trait we wish to allocate lives in the parent node, return a weak handle to it
-		// We use a weak handle to avoid issues when multiple base traits live within the same node
-		// When this happens, a trait can end up pointing to another within the same node causing
-		// the reference count to never reach zero when all other handles are released
-		if (FNodeInstance* ParentNodeInstance = ParentBinding.GetNodeInstance())
-		{
-			if (ParentNodeInstance->GetNodeHandle() == ChildNodeHandle)
-			{
-				return FTraitPtr(ParentNodeInstance, FTraitPtr::IS_WEAK_BIT, ChildTraitIndex);
-			}
-		}
-
 		// We need to allocate a new node instance
-		const FTraitTemplate* TraitDescs = NodeTemplate->GetTraits();
-
 		const uint32 InstanceSize = NodeDesc.GetNodeInstanceDataSize();
 		uint8* NodeInstanceBuffer = reinterpret_cast<uint8*>(FMemory::Malloc(InstanceSize, 16));
-		FNodeInstance* NodeInstance = new(NodeInstanceBuffer) FNodeInstance(*GraphInstance, ChildNodeHandle);
+		FNodeInstance* NodeInstance = new(NodeInstanceBuffer) FNodeInstance(GraphInstance, ChildNodeHandle);
+
+		// Manually bind our stack since we have everything we need
+		FTraitStackBinding StackBinding;
+		StackBinding.Context = this;
+		StackBinding.NodeInstance = NodeInstance;
+		StackBinding.NodeDescription = &NodeDesc;
+		StackBinding.NodeTemplate = NodeTemplate;
 
-		// Start construction with the bottom trait
+		// Start construction with the base trait
+		// We construct the whole node which will include one or more sub-stacks (each with their own base trait)
+		const FTraitTemplate* TraitDescs = NodeTemplate->GetTraits();
 		const FTraitTemplate* StartDesc = TraitDescs;
 		const FTraitTemplate* EndDesc = TraitDescs + NodeTemplate->GetNumTraits();
 
@@ -139,164 +140,83 @@ namespace UE::AnimNext
 
 			const uint32 TraitIndex = TraitDesc - TraitDescs;
 
-			FWeakTraitPtr TraitPtr(NodeInstance, TraitIndex);
-
-			const FAnimNextTraitSharedData* SharedData = TraitDesc->GetTraitDescription(NodeDesc);
-			FTraitInstanceData* InstanceData = TraitDesc->GetTraitInstance(*NodeInstance);
+			if (TraitDesc->GetMode() == ETraitMode::Base)
+			{
+				// A new base trait, update our stack binding
+				StackBinding.BaseTraitIndex = TraitIndex;
+				StackBinding.TopTraitIndex = TraitDesc->GetNumStackTraits() - 1;
+			}
 
-			const FTraitBinding Binding(nullptr, TraitDesc, &NodeDesc, TraitPtr);
+			const FTraitBinding Binding(&StackBinding, Trait, TraitIndex);
 			Trait->ConstructTraitInstance(*this, Binding);
 		}
 
 		return FTraitPtr(NodeInstance, ChildTraitIndex);
 	}
 
-	void FExecutionContext::ReleaseNodeInstance(FTraitPtr& NodePtr) const
+	FTraitPtr FExecutionContext::AllocateNodeInstance(const FWeakTraitPtr& ParentBinding, FAnimNextTraitHandle ChildTraitHandle) const
 	{
-		if (!NodePtr.IsValid())
+		if (!ChildTraitHandle.IsValid())
 		{
-			return;
+			return FTraitPtr();	// Attempting to allocate a node using an invalid trait handle
 		}
 
-		FNodeInstance* NodeInstance = NodePtr.GetNodeInstance();
-
-		if (!ensure(IsBoundTo(NodeInstance->GetOwner())))
+		if (!ensure(IsBound()))
 		{
-			return;	// The execution context isn't bound to the right graph instance
+			return FTraitPtr();	// The execution context must be bound to a valid graph instance
 		}
 
-		// Reset the handle here to simplify the multiple return statements below
-		NodePtr.PackedPointerAndFlags = 0;
-		NodePtr.TraitIndex = 0;
-
-		if (NodeInstance->RemoveReference())
+		if (!ensure(ParentBinding.IsValid()))
 		{
-			return;	// Node instance still has references, we can't release it
+			return FTraitPtr();	// We need a parent binding to know which graph instance to allocate into
 		}
 
-		const FNodeDescription& NodeDesc = GetNodeDescription(NodeInstance->GetNodeHandle());
+		FNodeInstance* ParentNodeInstance = ParentBinding.GetNodeInstance();
+		FAnimNextGraphInstance& GraphInstance = ParentNodeInstance->GetOwner();
 
-		const FNodeTemplate* NodeTemplate = GetNodeTemplate(NodeDesc);
-		if (!ensure(NodeTemplate != nullptr))
-		{
-			return;	// Node template wasn't found, node descriptor is perhaps corrupted (we'll leak the node memory)
-		}
+		const FNodeHandle ChildNodeHandle = ChildTraitHandle.GetNodeHandle();
+		const uint32 ChildTraitIndex = ChildTraitHandle.GetTraitIndex();
 
-		const FTraitTemplate* TraitDescs = NodeTemplate->GetTraits();
+		const FNodeDescription& NodeDesc = GetNodeDescription(GraphInstance, ChildNodeHandle);
+		const FNodeTemplate* NodeTemplate = GetNodeTemplate(NodeDesc);
 
-		// Start destruction with the top trait
-		const FTraitTemplate* StartDesc = TraitDescs + NodeTemplate->GetNumTraits() - 1;
-		const FTraitTemplate* EndDesc = TraitDescs - 1;
-		for (const FTraitTemplate* TraitDesc = StartDesc; TraitDesc != EndDesc; --TraitDesc)
+		if (!ensure(NodeTemplate != nullptr))
 		{
-			if (const FTrait* Trait = GetTrait(*TraitDesc))
-			{
-				const uint32 TraitIndex = TraitDesc - TraitDescs;
-
-				FWeakTraitPtr TraitPtr(NodeInstance, TraitIndex);
-
-				const FTraitBinding Binding(nullptr, TraitDesc, &NodeDesc, TraitPtr);
-				Trait->DestructTraitInstance(*this, Binding);
-			}
+			return FTraitPtr();	// Node template wasn't found, node descriptor is perhaps corrupted
 		}
-
-		NodeInstance->~FNodeInstance();
-		FMemory::Free(NodeInstance);
-	}
-
-	bool FExecutionContext::GetInterfaceImpl(FTraitInterfaceUID InterfaceUID, const FWeakTraitPtr& TraitPtr, FTraitBinding& InterfaceBinding) const
-	{
-		if (!TraitPtr.IsValid())
+		else if (ChildTraitIndex >= NodeTemplate->GetNumTraits())
 		{
-			return false;
+			return FTraitPtr();	// The requested trait index doesn't exist on that node descriptor
 		}
 
-		// TODO: Revisit this to avoid querying each trait over and over
-		//
-		// If we have D traits on a node and I unique interfaces per trait then finding an interface
-		// has O(D*I) complexity: for each trait we have to test every interface
-		// However, if the I interfaces are not unique per trait and instead are unique for the system,
-		// the below code will still have the same complexity. We'll end up testing the same interface multiple
-		// times, once for each trait. Because we will have a small set of interfaces per node that
-		// multiple traits implement (e.g IUpdate, IEvaluate), we can do better.
-		// 
-		// It would be much cheaper if instead we could test if we implement the interface first and then
-		// look for a list of traits (in stack order) that implement it. To do this, when we build the node
-		// template, we could aggregate all interfaces that the traits implement (sorted by their interface id
-		// for determinism or perhaps some other criteria to put popular/hot interfaces first). Then, in the node
-		// template we store a mapping of InterfaceUID to InterfaceIndex (local to that node template). The interface
-		// index can then be used to index into a list of trait indices that implement it.
-		// Code would look like:
-		//    * uint32 InterfaceIndex = NodeTemplate.GetInterfaceUIDs().FindIndex(InterfaceUID);	// INDEX_NONE means that none of the traits implement the interface
-		//    * uint16 TraitIndicesStartOffset = NodeTemplate.InterfaceIndexToTraitIndices(InterfaceIndex);		// An offset relative to the NodeTemplate*
-		//    * const uint8* TraitIndicesForInterface = ((const uint8*NodeTemplate) + TraitIndicesStartOffset;
-		//    * uint8 NumTraitsWithInterface = TraitIndicesForInterface[0];		// Reserve first index for the trait count or if we use a guard value at the end of the array we have to reserve a value
-		//    * uint8 TopTraitIndex = TraitIndicesForInterface[1];				// First index is top of stack
-		// Super trait lookup would be similar:
-		//    * uint8 SuperTraitIndex = TraitIndicesForInterface[TraitIndicesForInterface.FindIndex(CurrentTraitIndex) + 1];	// Need to check if we exceed the max count (bottom of stack)
-		// Once we have the trait index, we use the trait template as we do now
-		//
-		// The above would have a number of advantages:
-		//    * Querying for an interface that the node doesn't implement would be much cheaper (we test each interface once)
-		//    * InterfaceIndex query can easily be implemented in SIMD, we look for a uint32 in a list sequentially
-		//    * Bulk/interleaved querying of interfaces would be much easier (e.g. query for IEvaluate for 4 nodes)
-		//    * A lot less branching is involved, improving throughput
-		//    * We can handle multiple base traits by checking the trait index found which has comparable cost
-		//    * We could store the interface offset from the trait ptr base to avoid the call to GetInterface()
-		//      Instead of testing to find our interface again to return a custom static_cast, we'd store the interface/trait offset
-		//      in the node template alongside the trait index. We can then query for the FTrait* and add the offset
-		//      improving bulk query feasibility.
-		//
-		// However, querying for a Super interface might be slower since we have to start the search from the start.
-		// In contrast, we now start searching at the current trait. In order to keep that cost down and amortize
-		// the search, we would have to cache the InterfaceIndex and the current index in the trait list.
-		// We could use 8-bits for each entry allowing for a max of 256 interfaces per node and 256 traits per node (existing limitation).
-		//
-		// Adding new data to the trait binding isn't ideal. It's current size on 64-bit systems is:
-		//    * 8 bytes for interface pointer (to forward function calls to)
-		//    * 8 bytes for trait template (to get the offsets for shared/instance/latent data)
-		//    * 8 bytes for node description pointer (base of shared data)
-		//    * 8 bytes for node instance pointer (base of instance data, in weak trait ptr)
-		//    * 4 bytes for trait index (in weak trait ptr)
-		//    * 4 bytes of padding (in weak trait ptr)
-		// We current use 4 bytes for the trait index but in practice we only need 1 byte. We use 4 bytes because we pay
-		// for padding anyway. We could use 2 extra bytes for our mapping and still have 5 bytes to space in padding (due to alignment).
-		// In practice, trait bindings always live on the stack and they shouldn't be getting copied/moved around.
-		// We populate them in place when we query, then we use it from its storage on the stack that we first populated.
-		//
-		// To be able to support all of this, instead of having FTrait::GetInterface() implemented by derived types
-		// we would need instead to implement FTrait::GetInterfaceUIDs(). Then on load, when we build a node template and
-		// finalize it, we can build the list of unique interface UIDs and sort it and we can build the mapping structures.
-		// On load, as long as we fit within 64KB we are fine. We have to update FNodeTemplate::Finalize and FNodeTemplateBuilder::BuildNodeTemplate.
-		// Perhaps the builder should use the same code and build in a buffer on the stack and copy the final size out.
-		// We must build the mapping on load because the interfaces we implement are only known at runtime (e.g. can be changed through defines/configuration).
-		// We can still implement GetInterface() in terms of GetInterfaceUIDs if it returns the interface offsets as well.
-
-		FNodeInstance* NodeInstance = TraitPtr.GetNodeInstance();
-
-		if (!ensure(IsBoundTo(NodeInstance->GetOwner())))
+		// If the trait we wish to allocate lives in the parent node, return a weak handle to it
+		// We use a weak handle to avoid issues when multiple base traits live within the same node
+		// When this happens, a trait can end up pointing to another within the same node causing
+		// the reference count to never reach zero when all other handles are released
+		if (ParentNodeInstance->GetNodeHandle() == ChildNodeHandle)
 		{
-			return false;	// The execution context isn't bound to the right graph instance
+			return FTraitPtr(ParentNodeInstance, FTraitPtr::IS_WEAK_BIT, ChildTraitIndex);
 		}
 
-		const FNodeDescription& NodeDesc = GetNodeDescription(NodeInstance->GetNodeHandle());
+		// We need to allocate a new node instance
+		const uint32 InstanceSize = NodeDesc.GetNodeInstanceDataSize();
+		uint8* NodeInstanceBuffer = reinterpret_cast<uint8*>(FMemory::Malloc(InstanceSize, 16));
+		FNodeInstance* NodeInstance = new(NodeInstanceBuffer) FNodeInstance(GraphInstance, ChildNodeHandle);
 
-		const FNodeTemplate* NodeTemplate = GetNodeTemplate(NodeDesc);
-		if (!ensure(NodeTemplate != nullptr))
-		{
-			return false;	// Node template wasn't found, node descriptor is perhaps corrupted
-		}
+		// Manually bind our stack since we have everything we need
+		FTraitStackBinding StackBinding;
+		StackBinding.Context = this;
+		StackBinding.NodeInstance = NodeInstance;
+		StackBinding.NodeDescription = &NodeDesc;
+		StackBinding.NodeTemplate = NodeTemplate;
 
+		// Start construction with the base trait
+		// We construct the whole node which will include one or more sub-stacks (each with their own base trait)
 		const FTraitTemplate* TraitDescs = NodeTemplate->GetTraits();
+		const FTraitTemplate* StartDesc = TraitDescs;
+		const FTraitTemplate* EndDesc = TraitDescs + NodeTemplate->GetNumTraits();
 
-		// We only search within the partial stack of the provided trait
-		const FTraitTemplate* CurrentTraitDesc = TraitDescs + TraitPtr.GetTraitIndex();
-		const FTraitTemplate* BaseTraitDesc = CurrentTraitDesc->GetMode() == ETraitMode::Base ? CurrentTraitDesc : (CurrentTraitDesc - CurrentTraitDesc->GetAdditiveTraitIndex());
-
-		// Start searching with the top additive trait towards our base trait
-		const FTraitTemplate* StartDesc = BaseTraitDesc + BaseTraitDesc->GetNumAdditiveTraits();
-		const FTraitTemplate* EndDesc = BaseTraitDesc - 1;
-		for (const FTraitTemplate* TraitDesc = StartDesc; TraitDesc != EndDesc; --TraitDesc)
+		for (const FTraitTemplate* TraitDesc = StartDesc; TraitDesc != EndDesc; ++TraitDesc)
 		{
 			const FTrait* Trait = GetTrait(*TraitDesc);
 			if (Trait == nullptr)
@@ -304,59 +224,64 @@ namespace UE::AnimNext
 				continue;	// Trait hasn't been loaded or registered, skip it
 			}
 
-			if (const ITraitInterface* Interface = Trait->GetTraitInterface(InterfaceUID))
-			{
-				const FAnimNextTraitSharedData* SharedData = TraitDesc->GetTraitDescription(NodeDesc);
-				FTraitInstanceData* InstanceData = TraitDesc->GetTraitInstance(*NodeInstance);
-
-				const uint32 TraitIndex = TraitDesc - TraitDescs;
-				FWeakTraitPtr InterfaceTraitPtr(NodeInstance, TraitIndex);
+			const uint32 TraitIndex = TraitDesc - TraitDescs;
 
-				InterfaceBinding = FTraitBinding(Interface, TraitDesc, &NodeDesc, InterfaceTraitPtr);
-				return true;
+			if (TraitDesc->GetMode() == ETraitMode::Base)
+			{
+				// A new base trait, update our stack binding
+				StackBinding.BaseTraitIndex = TraitIndex;
+				StackBinding.TopTraitIndex = TraitDesc->GetNumStackTraits() - 1;
 			}
+
+			const FTraitBinding Binding(&StackBinding, Trait, TraitIndex);
+			Trait->ConstructTraitInstance(*this, Binding);
 		}
 
-		// We failed to find a trait that handles this interface
-		return false;
+		return FTraitPtr(NodeInstance, ChildTraitIndex);
 	}
 
-	bool FExecutionContext::GetInterfaceSuperImpl(FTraitInterfaceUID InterfaceUID, const FWeakTraitPtr& TraitPtr, FTraitBinding& SuperBinding) const
+	void FExecutionContext::ReleaseNodeInstance(FTraitPtr& NodePtr) const
 	{
-		if (!TraitPtr.IsValid())
+		if (!NodePtr.IsValid())
 		{
-			return false;	// Trait pointer isn't valid, can't find a 'super'
+			return;
 		}
 
-		FNodeInstance* NodeInstance = TraitPtr.GetNodeInstance();
+		FNodeInstance* NodeInstance = NodePtr.GetNodeInstance();
 
 		if (!ensure(IsBoundTo(NodeInstance->GetOwner())))
 		{
-			return false;	// The execution context isn't bound to the right graph instance
+			return;	// The execution context isn't bound to the right graph instance
 		}
 
-		const FNodeDescription& NodeDesc = GetNodeDescription(NodeInstance->GetNodeHandle());
+		// Reset the handle here to simplify the multiple return statements below
+		NodePtr.PackedPointerAndFlags = 0;
+		NodePtr.TraitIndex = 0;
 
-		const FNodeTemplate* NodeTemplate = GetNodeTemplate(NodeDesc);
-		if (!ensure(NodeTemplate != nullptr))
+		if (NodeInstance->RemoveReference())
 		{
-			return false;	// Node template wasn't found, node descriptor is perhaps corrupted
+			return;	// Node instance still has references, we can't release it
 		}
 
-		const FTraitTemplate* TraitDescs = NodeTemplate->GetTraits();
-
-		// We only search within the partial stack of the provided trait
-		const FTraitTemplate* CurrentTraitDesc = TraitDescs + TraitPtr.GetTraitIndex();
-		if (CurrentTraitDesc->GetMode() == ETraitMode::Base)
+		const FNodeDescription& NodeDesc = GetNodeDescription(*NodeInstance);
+		const FNodeTemplate* NodeTemplate = GetNodeTemplate(NodeDesc);
+		if (!ensure(NodeTemplate != nullptr))
 		{
-			return false;	// We've reached a base trait, we don't forward to any parent trait we might have
+			return;	// Node template wasn't found, node descriptor is perhaps corrupted (we'll leak the node memory)
 		}
 
-		const FTraitTemplate* BaseTraitDesc = CurrentTraitDesc - CurrentTraitDesc->GetAdditiveTraitIndex();
+		// Manually bind our stack since we have everything we need
+		FTraitStackBinding StackBinding;
+		StackBinding.Context = this;
+		StackBinding.NodeInstance = NodeInstance;
+		StackBinding.NodeDescription = &NodeDesc;
+		StackBinding.NodeTemplate = NodeTemplate;
 
-		// Start searching with the next additive trait towards our base trait
-		const FTraitTemplate* StartDesc = CurrentTraitDesc - 1;
-		const FTraitTemplate* EndDesc = BaseTraitDesc - 1;
+		// Start destruction with the top trait
+		// We destruct the whole node which will include one or more sub-stacks (each with their own base trait)
+		const FTraitTemplate* TraitDescs = NodeTemplate->GetTraits();
+		const FTraitTemplate* StartDesc = TraitDescs + NodeTemplate->GetNumTraits() - 1;
+		const FTraitTemplate* EndDesc = TraitDescs - 1;
 		for (const FTraitTemplate* TraitDesc = StartDesc; TraitDesc != EndDesc; --TraitDesc)
 		{
 			const FTrait* Trait = GetTrait(*TraitDesc);
@@ -365,86 +290,67 @@ namespace UE::AnimNext
 				continue;	// Trait hasn't been loaded or registered, skip it
 			}
 
-			if (const ITraitInterface* Interface = Trait->GetTraitInterface(InterfaceUID))
-			{
-				const FAnimNextTraitSharedData* SharedData = TraitDesc->GetTraitDescription(NodeDesc);
-				FTraitInstanceData* InstanceData = TraitDesc->GetTraitInstance(*NodeInstance);
-
-				const uint32 TraitIndex = TraitDesc - TraitDescs;
-				FWeakTraitPtr SuperPtr(NodeInstance, TraitIndex);
+			// Always update our stack binding to make sure it points to the right sub-stack
+			const FTraitTemplate* BaseTraitDesc = TraitDesc - TraitDesc->GetTraitIndex();
+			StackBinding.BaseTraitIndex = BaseTraitDesc - TraitDescs;
+			StackBinding.TopTraitIndex = StackBinding.BaseTraitIndex + BaseTraitDesc->GetNumStackTraits() - 1;
 
-				SuperBinding = FTraitBinding(Interface, TraitDesc, &NodeDesc, SuperPtr);
-				return true;
-			}
+			const uint32 TraitIndex = TraitDesc - TraitDescs;
+			const FTraitBinding Binding(&StackBinding, Trait, TraitIndex);
+			Trait->DestructTraitInstance(*this, Binding);
 		}
 
-		// We failed to find a trait that handles this interface
-		return false;
+		NodeInstance->~FNodeInstance();
+		FMemory::Free(NodeInstance);
 	}
 
-	void FExecutionContext::SnapshotLatentProperties(const FWeakTraitPtr& TraitPtr, bool bIsFrozen) const
+	bool FExecutionContext::GetStack(const FWeakTraitPtr& TraitPtr, FTraitStackBinding& OutStackBinding) const
 	{
 		if (!TraitPtr.IsValid())
 		{
-			return;	// Nothing to do
-		}
-
-		FNodeInstance* NodeInstance = TraitPtr.GetNodeInstance();
-
-		if (!ensure(IsBoundTo(NodeInstance->GetOwner())))
-		{
-			return;	// The execution context isn't bound to the right graph instance
+			OutStackBinding.Reset();
+			return false;
 		}
 
-		const FNodeDescription& NodeDesc = GetNodeDescription(NodeInstance->GetNodeHandle());
-
-		const FNodeTemplate* NodeTemplate = GetNodeTemplate(NodeDesc);
-		if (!ensure(NodeTemplate != nullptr))
+		if (!ensure(IsBoundTo(TraitPtr.GetNodeInstance()->GetOwner())))
 		{
-			return;	// Node template wasn't found, node descriptor is perhaps corrupted
-		}
-
-		const FTraitTemplate* TraitDescs = NodeTemplate->GetTraits();
-
-		// We only snapshot the partial stack the specified trait lives in
-		const FTraitTemplate* CurrentTraitDesc = TraitDescs + TraitPtr.GetTraitIndex();
-		const FTraitTemplate* BaseTraitDesc = CurrentTraitDesc->GetMode() == ETraitMode::Base ? CurrentTraitDesc : (CurrentTraitDesc - CurrentTraitDesc->GetAdditiveTraitIndex());
-
-		const FLatentPropertiesHeader& LatentHeader = BaseTraitDesc->GetTraitLatentPropertiesHeader(NodeDesc);
-		if (!LatentHeader.bHasValidLatentProperties)
-		{
-			return;	// All latent properties are inline, nothing to snapshot
-		}
-		else if (bIsFrozen && LatentHeader.bCanAllPropertiesFreeze)
-		{
-			return;	// We are frozen and all latent properties support freezing, nothing to snapshot
+			OutStackBinding.Reset();
+			return false;	// The execution context isn't bound to the right graph instance
 		}
 
-		const FLatentPropertyHandle* LatentHandles = BaseTraitDesc->GetTraitLatentPropertyHandles(NodeDesc);
-		const uint32 NumLatentHandles = BaseTraitDesc->GetNumSubStackLatentPropreties();
-
-		GraphInstance->ExecuteLatentPins(TConstArrayView<FLatentPropertyHandle>(LatentHandles, NumLatentHandles), NodeInstance, bIsFrozen);
+		OutStackBinding = FTraitStackBinding(*this, TraitPtr);
+		return OutStackBinding.IsValid();	// Construction can fail in rare cases, see constructor
 	}
 
 	FGraphInstanceComponent* FExecutionContext::TryGetComponent(int32 ComponentNameHash, FName ComponentName) const
 	{
-		return GraphInstance->TryGetComponent(ComponentNameHash, ComponentName);
+		check(IsBound());
+		return RootGraphInstance->TryGetComponent(ComponentNameHash, ComponentName);
 	}
 
 	FGraphInstanceComponent& FExecutionContext::AddComponent(int32 ComponentNameHash, FName ComponentName, TSharedPtr<FGraphInstanceComponent>&& Component) const
 	{
-		return GraphInstance->AddComponent(ComponentNameHash, ComponentName, MoveTemp(Component));
+		check(IsBound());
+		return RootGraphInstance->AddComponent(ComponentNameHash, ComponentName, MoveTemp(Component));
 	}
 
 	GraphInstanceComponentMapType::TConstIterator FExecutionContext::GetComponentIterator() const
 	{
-		return GraphInstance->GetComponentIterator();
+		check(IsBound());
+		return RootGraphInstance->GetComponentIterator();
+	}
+
+	const FNodeDescription& FExecutionContext::GetNodeDescription(const FAnimNextGraphInstance& GraphInstance, FNodeHandle NodeHandle) const
+	{
+		// Grab the node description from the specified graph
+		const UAnimNextGraph* Graph = GraphInstance.GetGraph();
+		return *reinterpret_cast<const FNodeDescription*>(&Graph->SharedDataBuffer[NodeHandle.GetSharedOffset()]);
 	}
 
-	const FNodeDescription& FExecutionContext::GetNodeDescription(FNodeHandle NodeHandle) const
+	const FNodeDescription& FExecutionContext::GetNodeDescription(const FNodeInstance& NodeInstance) const
 	{
-		check(NodeHandle.IsValid());
-		return *reinterpret_cast<const FNodeDescription*>(&GraphSharedData[NodeHandle.GetSharedOffset()]);
+		// Grab the node description from the owning graph
+		return GetNodeDescription(NodeInstance.GetOwner(), NodeInstance.GetNodeHandle());
 	}
 
 	const FNodeTemplate* FExecutionContext::GetNodeTemplate(const FNodeDescription& NodeDesc) const
diff --git a/Engine/Plugins/Experimental/AnimNext/Source/AnimNext/Private/TraitCore/NodeTemplate.cpp b/Engine/Plugins/Experimental/AnimNext/Source/AnimNext/Private/TraitCore/NodeTemplate.cpp
--- a/Engine/Plugins/Experimental/AnimNext/Source/AnimNext/Private/TraitCore/NodeTemplate.cpp
+++ b/Engine/Plugins/Experimental/AnimNext/Source/AnimNext/Private/TraitCore/NodeTemplate.cpp
@@ -17,7 +17,7 @@ namespace UE::AnimNext
 			const FTraitTemplate& BaseTraitTemplate = TraitTemplates[BaseTraitIndex];
 			check(BaseTraitTemplate.GetMode() == ETraitMode::Base);
 
-			const uint32 NumSubStackTraits = BaseTraitTemplate.GetNumAdditiveTraits() + 1;
+			const uint32 NumSubStackTraits = BaseTraitTemplate.GetNumStackTraits();
 			uint32 NumSubStackLatentProperties = 0;
 
 			for (uint32 SubStackTraitIndex = 0; SubStackTraitIndex < NumSubStackTraits; ++SubStackTraitIndex)
diff --git a/Engine/Plugins/Experimental/AnimNext/Source/AnimNext/Private/TraitCore/NodeTemplateBuilder.cpp b/Engine/Plugins/Experimental/AnimNext/Source/AnimNext/Private/TraitCore/NodeTemplateBuilder.cpp
--- a/Engine/Plugins/Experimental/AnimNext/Source/AnimNext/Private/TraitCore/NodeTemplateBuilder.cpp
+++ b/Engine/Plugins/Experimental/AnimNext/Source/AnimNext/Private/TraitCore/NodeTemplateBuilder.cpp
@@ -69,11 +69,11 @@ namespace UE::AnimNext
 		const FTrait* Trait = TraitRegistry.Find(TraitHandle);
 		const ETraitMode TraitMode = Trait->GetTraitMode();
 
-		uint32 AdditiveIndexOrNumAdditive;
+		uint32 TraitIndexOrNumTraits;
 		if (TraitMode == ETraitMode::Base)
 		{
-			// Find out how many additive traits we have
-			AdditiveIndexOrNumAdditive = 0;
+			// Find out how many traits we have
+			TraitIndexOrNumTraits = 1;
 			for (int32 Index = TraitIndex + 1; Index < InTraitUIDs.Num(); ++Index)	// Skip ourself
 			{
 				const FTrait* ChildTrait = TraitRegistry.Find(InTraitUIDs[Index]);
@@ -83,13 +83,14 @@ namespace UE::AnimNext
 				}
 
 				// We are additive
-				AdditiveIndexOrNumAdditive++;
+				TraitIndexOrNumTraits++;
 			}
 		}
 		else
 		{
-			// Find out our additive index
-			AdditiveIndexOrNumAdditive = 1;	// Skip ourself
+			// Find out our stack trait index (as opposed to the node trait index)
+			// Iterate from this additive trait until we find our base trait
+			TraitIndexOrNumTraits = 1;	// We are at least the second trait (first additive)
 			for (int32 Index = TraitIndex - 1; Index >= 0; --Index)
 			{
 				const FTrait* ParentTrait = TraitRegistry.Find(InTraitUIDs[Index]);
@@ -99,12 +100,12 @@ namespace UE::AnimNext
 				}
 
 				// We are additive
-				AdditiveIndexOrNumAdditive++;
+				TraitIndexOrNumTraits++;
 			}
 		}
 
 		// Append our trait template
 		const int32 BufferIndex = NodeTemplateBuffer.AddUninitialized(sizeof(FTraitTemplate));
-		new(&NodeTemplateBuffer[BufferIndex]) FTraitTemplate(TraitUID, TraitHandle, TraitMode, AdditiveIndexOrNumAdditive);
+		new(&NodeTemplateBuffer[BufferIndex]) FTraitTemplate(TraitUID, TraitHandle, TraitMode, TraitIndexOrNumTraits);
 	}
 }
diff --git a/Engine/Plugins/Experimental/AnimNext/Source/AnimNext/Private/TraitCore/TraitBinding.cpp b/Engine/Plugins/Experimental/AnimNext/Source/AnimNext/Private/TraitCore/TraitBinding.cpp
--- a/Engine/Plugins/Experimental/AnimNext/Source/AnimNext/Private/TraitCore/TraitBinding.cpp
+++ b/Engine/Plugins/Experimental/AnimNext/Source/AnimNext/Private/TraitCore/TraitBinding.cpp
@@ -8,6 +8,6 @@ namespace UE::AnimNext
 {
 	FTraitInterfaceUID FTraitBinding::GetInterfaceUID() const
 	{
-		return Interface != nullptr ? Interface->GetInterfaceUID() : FTraitInterfaceUID();
+		return InterfaceThisOffset != -1 ? GetInterfaceTyped<ITraitInterface>()->GetInterfaceUID() : FTraitInterfaceUID();
 	}
 }
diff --git a/Engine/Plugins/Experimental/AnimNext/Source/AnimNext/Private/TraitCore/TraitStackBinding.cpp b/Engine/Plugins/Experimental/AnimNext/Source/AnimNext/Private/TraitCore/TraitStackBinding.cpp
new file mode 100644
--- /dev/null
+++ b/Engine/Plugins/Experimental/AnimNext/Source/AnimNext/Private/TraitCore/TraitStackBinding.cpp
@@ -0,0 +1,315 @@
+// Copyright Epic Games, Inc. All Rights Reserved.
+
+#include "TraitCore/TraitStackBinding.h"
+
+#include "Graph/AnimNextGraphInstance.h"
+#include "TraitCore/ExecutionContext.h"
+#include "TraitCore/NodeInstance.h"
+#include "TraitCore/NodeTemplate.h"
+#include "TraitCore/Trait.h"
+#include "TraitCore/TraitBinding.h"
+
+namespace UE::AnimNext
+{
+	FTraitStackBinding::FTraitStackBinding(const FExecutionContext& InContext, const FWeakTraitPtr& TraitPtr)
+		: Context(&InContext)
+	{
+		check(TraitPtr.IsValid());
+
+		NodeInstance = TraitPtr.GetNodeInstance();
+
+		check(InContext.IsBoundTo(NodeInstance->GetOwner()));
+
+		NodeDescription = &InContext.GetNodeDescription(*NodeInstance);
+
+		NodeTemplate = InContext.GetNodeTemplate(*NodeDescription);
+		if (!ensure(NodeTemplate != nullptr))
+		{
+			// Node template wasn't found, node descriptor is perhaps corrupted
+			Reset();
+			return;
+		}
+
+		if (!ensure(TraitPtr.GetTraitIndex() < NodeTemplate->GetNumTraits()))
+		{
+			// The requested trait index doesn't exist on that node descriptor
+			Reset();
+			return;
+		}
+
+		const FTraitTemplate* TraitDescs = NodeTemplate->GetTraits();
+
+		// We only search within the partial stack of the provided trait
+		const FTraitTemplate* CurrentTraitDesc = TraitDescs + TraitPtr.GetTraitIndex();
+		const FTraitTemplate* BaseTraitDesc = CurrentTraitDesc - CurrentTraitDesc->GetTraitIndex();
+
+		BaseTraitIndex = BaseTraitDesc - TraitDescs;
+		TopTraitIndex = BaseTraitIndex + BaseTraitDesc->GetNumStackTraits() - 1;
+	}
+
+	void FTraitStackBinding::Reset()
+	{
+		new (this) FTraitStackBinding();
+	}
+
+	bool FTraitStackBinding::operator==(const FTraitStackBinding& Other) const
+	{
+		return NodeInstance == Other.NodeInstance && BaseTraitIndex == Other.BaseTraitIndex;
+	}
+
+	bool FTraitStackBinding::GetTopTrait(FTraitBinding& OutBinding) const
+	{
+		if (!IsValid())
+		{
+			OutBinding.Reset();
+			return false;
+		}
+
+		const FTraitTemplate* TraitDescs = NodeTemplate->GetTraits();
+		uint32 TraitIndex = TopTraitIndex;
+
+		// Skip invalid trait templates
+		while (!TraitDescs[TraitIndex].IsValid() && TraitIndex != BaseTraitIndex)
+		{
+			TraitIndex--;
+		}
+
+		if (!ensure(TraitDescs[TraitIndex].IsValid()))
+		{
+			OutBinding.Reset();
+			return false;	// No traits were valid on this stack but we should always at least have a valid base trait
+		}
+
+		OutBinding = FTraitBinding(this, Context->GetTrait(TraitDescs[TraitIndex]), TraitIndex);
+		return true;
+	}
+
+	bool FTraitStackBinding::GetParentTrait(const FTraitBinding& ChildBinding, FTraitBinding& OutParentBinding) const
+	{
+		if (!IsValid())
+		{
+			OutParentBinding.Reset();
+			return false;
+		}
+
+		if (ChildBinding.TraitIndex == BaseTraitIndex)
+		{
+			OutParentBinding.Reset();
+			return false;	// No more parents, reached the base of the stack
+		}
+
+		const FTraitTemplate* TraitDescs = NodeTemplate->GetTraits();
+		uint32 ParentTraitIndex = ChildBinding.TraitIndex - 1;
+
+		// Skip invalid trait templates
+		while (!TraitDescs[ParentTraitIndex].IsValid() && ParentTraitIndex != BaseTraitIndex)
+		{
+			ParentTraitIndex--;
+		}
+
+		if (!ensure(TraitDescs[ParentTraitIndex].IsValid()))
+		{
+			OutParentBinding.Reset();
+			return false;	// No parent trait was valid on this stack but we should always at least have a valid base trait
+		}
+
+		OutParentBinding = FTraitBinding(this, Context->GetTrait(TraitDescs[ParentTraitIndex]), ParentTraitIndex);
+		return true;
+	}
+
+	bool FTraitStackBinding::GetBaseTrait(FTraitBinding& OutBinding) const
+	{
+		if (!IsValid())
+		{
+			OutBinding.Reset();
+			return false;
+		}
+
+		const FTraitTemplate* TraitDescs = NodeTemplate->GetTraits();
+
+		if (!ensure(TraitDescs[BaseTraitIndex].IsValid()))
+		{
+			OutBinding.Reset();
+			return false;	// We should always have a valid base trait
+		}
+
+		OutBinding = FTraitBinding(this, Context->GetTrait(TraitDescs[BaseTraitIndex]), BaseTraitIndex);
+		return true;
+	}
+
+	bool FTraitStackBinding::GetChildTrait(const FTraitBinding& ParentBinding, FTraitBinding& OutChildBinding) const
+	{
+		if (!IsValid())
+		{
+			OutChildBinding.Reset();
+			return false;
+		}
+
+		if (ParentBinding.TraitIndex == TopTraitIndex)
+		{
+			OutChildBinding.Reset();
+			return false;	// No more children, reached the top of the stack
+		}
+
+		const FTraitTemplate* TraitDescs = NodeTemplate->GetTraits();
+		uint32 ChildTraitIndex = ParentBinding.TraitIndex + 1;
+
+		// Skip invalid trait templates
+		while (!TraitDescs[ChildTraitIndex].IsValid() && ChildTraitIndex != TopTraitIndex)
+		{
+			ChildTraitIndex++;
+		}
+
+		if (!TraitDescs[ChildTraitIndex].IsValid())
+		{
+			OutChildBinding.Reset();
+			return false;	// We couldn't find a valid child
+		}
+
+		OutChildBinding = FTraitBinding(this, Context->GetTrait(TraitDescs[ChildTraitIndex]), ChildTraitIndex);
+		return true;
+	}
+
+	bool FTraitStackBinding::GetTrait(uint32 TraitIndex, FTraitBinding& OutBinding) const
+	{
+		if (!IsValid())
+		{
+			OutBinding.Reset();
+			return false;
+		}
+
+		const FTraitTemplate* TraitDescs = NodeTemplate->GetTraits();
+		if (TraitIndex >= TraitDescs[BaseTraitIndex].GetNumStackTraits())
+		{
+			OutBinding.Reset();
+			return false;	// Invalid trait index
+		}
+
+		const uint32 StackTraitIndex = BaseTraitIndex + TraitIndex;
+		if (!TraitDescs[StackTraitIndex].IsValid())
+		{
+			OutBinding.Reset();
+			return false;
+		}
+
+		OutBinding = FTraitBinding(this, Context->GetTrait(TraitDescs[StackTraitIndex]), StackTraitIndex);
+		return true;
+	}
+
+	uint32 FTraitStackBinding::GetNumTraits() const
+	{
+		if (!IsValid())
+		{
+			return 0;
+		}
+
+		const FTraitTemplate* TraitDescs = NodeTemplate->GetTraits();
+		return TraitDescs[BaseTraitIndex].GetNumStackTraits();
+	}
+
+	bool FTraitStackBinding::GetInterfaceImpl(FTraitInterfaceUID InterfaceUID, FTraitBinding& OutBinding) const
+	{
+		if (!IsValid())
+		{
+			OutBinding.Reset();
+			return false;
+		}
+
+		// Start searching with the top trait towards our base trait
+		const FTraitTemplate* TraitDescs = NodeTemplate->GetTraits();
+		const FTraitTemplate* StartDesc = TraitDescs + TopTraitIndex;
+		const FTraitTemplate* EndDesc = TraitDescs + BaseTraitIndex - 1;
+
+		for (const FTraitTemplate* TraitDesc = StartDesc; TraitDesc != EndDesc; --TraitDesc)
+		{
+			const FTrait* Trait = Context->GetTrait(*TraitDesc);
+			if (Trait == nullptr)
+			{
+				continue;	// Trait hasn't been loaded or registered, skip it
+			}
+
+			if (const ITraitInterface* Interface = Trait->GetTraitInterface(InterfaceUID))
+			{
+				const uint32 TraitIndex = TraitDesc - TraitDescs;
+				const uint32 InterfaceThisOffset = reinterpret_cast<const uint8*>(Interface) - reinterpret_cast<const uint8*>(Trait);
+
+				OutBinding = FTraitBinding(this, Trait, TraitIndex, InterfaceThisOffset);
+				return true;
+			}
+		}
+
+		// We failed to find the specified interface on the trait stack
+		OutBinding.Reset();
+		return false;
+	}
+
+	bool FTraitStackBinding::GetInterfaceSuperImpl(FTraitInterfaceUID InterfaceUID, const FTraitBinding& Binding, FTraitBinding& OutSuperBinding) const
+	{
+		if (!IsValid())
+		{
+			OutSuperBinding.Reset();
+			return false;
+		}
+
+		const uint32 CurrentTraitIndex = Binding.TraitIndex;
+		if (CurrentTraitIndex == BaseTraitIndex)
+		{
+			OutSuperBinding.Reset();
+			return false;	// We reached the base of the stack, we don't have a super
+		}
+
+		// Start searching with the current trait towards our base trait
+		const FTraitTemplate* TraitDescs = NodeTemplate->GetTraits();
+		const FTraitTemplate* StartDesc = TraitDescs + CurrentTraitIndex - 1;
+		const FTraitTemplate* EndDesc = TraitDescs + BaseTraitIndex - 1;
+
+		for (const FTraitTemplate* TraitDesc = StartDesc; TraitDesc != EndDesc; --TraitDesc)
+		{
+			const FTrait* Trait = Context->GetTrait(*TraitDesc);
+			if (Trait == nullptr)
+			{
+				continue;	// Trait hasn't been loaded or registered, skip it
+			}
+
+			if (const ITraitInterface* Interface = Trait->GetTraitInterface(InterfaceUID))
+			{
+				const uint32 SuperTraitIndex = TraitDesc - TraitDescs;
+				const uint32 InterfaceThisOffset = reinterpret_cast<const uint8*>(Interface) - reinterpret_cast<const uint8*>(Trait);
+
+				OutSuperBinding = FTraitBinding(this, Trait, SuperTraitIndex, InterfaceThisOffset);
+				return true;
+			}
+		}
+
+		// We failed to find the specified interface on the trait stack
+		OutSuperBinding.Reset();
+		return false;
+	}
+
+	void FTraitStackBinding::SnapshotLatentProperties(bool bIsFrozen)
+	{
+		if (!IsValid())
+		{
+			return;	// Nothing to do
+		}
+
+		const FTraitTemplate* TraitDescs = NodeTemplate->GetTraits();
+		const FTraitTemplate* BaseTraitDesc = TraitDescs + BaseTraitIndex;
+
+		const FLatentPropertiesHeader& LatentHeader = BaseTraitDesc->GetTraitLatentPropertiesHeader(*NodeDescription);
+		if (!LatentHeader.bHasValidLatentProperties)
+		{
+			return;	// All latent properties are inline, nothing to snapshot
+		}
+		else if (bIsFrozen && LatentHeader.bCanAllPropertiesFreeze)
+		{
+			return;	// We are frozen and all latent properties support freezing, nothing to snapshot
+		}
+
+		const FLatentPropertyHandle* LatentHandles = BaseTraitDesc->GetTraitLatentPropertyHandles(*NodeDescription);
+		const uint32 NumLatentHandles = BaseTraitDesc->GetNumSubStackLatentPropreties();
+
+		FAnimNextGraphInstance& GraphInstance = NodeInstance->GetOwner();
+		GraphInstance.ExecuteLatentPins(TConstArrayView<FLatentPropertyHandle>(LatentHandles, NumLatentHandles), NodeInstance, bIsFrozen);
+	}
+}
diff --git a/Engine/Plugins/Experimental/AnimNext/Source/AnimNext/Private/TraitCore/TraitTemplate.cpp b/Engine/Plugins/Experimental/AnimNext/Source/AnimNext/Private/TraitCore/TraitTemplate.cpp
--- a/Engine/Plugins/Experimental/AnimNext/Source/AnimNext/Private/TraitCore/TraitTemplate.cpp
+++ b/Engine/Plugins/Experimental/AnimNext/Source/AnimNext/Private/TraitCore/TraitTemplate.cpp
@@ -48,6 +48,6 @@ namespace UE::AnimNext
 		}
 
 		Ar << Mode;
-		Ar << AdditiveIndexOrNumAdditive;
+		Ar << TraitIndexOrNumTraits;
 	}
 }
diff --git a/Engine/Plugins/Experimental/AnimNext/Source/AnimNext/Private/TraitInterfaces/IContinuousBlend.cpp b/Engine/Plugins/Experimental/AnimNext/Source/AnimNext/Private/TraitInterfaces/IContinuousBlend.cpp
--- a/Engine/Plugins/Experimental/AnimNext/Source/AnimNext/Private/TraitInterfaces/IContinuousBlend.cpp
+++ b/Engine/Plugins/Experimental/AnimNext/Source/AnimNext/Private/TraitInterfaces/IContinuousBlend.cpp
@@ -9,7 +9,7 @@ namespace UE::AnimNext
 	float IContinuousBlend::GetBlendWeight(const FExecutionContext& Context, const TTraitBinding<IContinuousBlend>& Binding, int32 ChildIndex) const
 	{
 		TTraitBinding<IContinuousBlend> SuperBinding;
-		if (Context.GetInterfaceSuper(Binding, SuperBinding))
+		if (Binding.GetStackInterfaceSuper(SuperBinding))
 		{
 			return SuperBinding.GetBlendWeight(Context, ChildIndex);
 		}
diff --git a/Engine/Plugins/Experimental/AnimNext/Source/AnimNext/Private/TraitInterfaces/IDiscreteBlend.cpp b/Engine/Plugins/Experimental/AnimNext/Source/AnimNext/Private/TraitInterfaces/IDiscreteBlend.cpp
--- a/Engine/Plugins/Experimental/AnimNext/Source/AnimNext/Private/TraitInterfaces/IDiscreteBlend.cpp
+++ b/Engine/Plugins/Experimental/AnimNext/Source/AnimNext/Private/TraitInterfaces/IDiscreteBlend.cpp
@@ -9,7 +9,7 @@ namespace UE::AnimNext
 	float IDiscreteBlend::GetBlendWeight(const FExecutionContext& Context, const TTraitBinding<IDiscreteBlend>& Binding, int32 ChildIndex) const
 	{
 		TTraitBinding<IDiscreteBlend> SuperBinding;
-		if (Context.GetInterfaceSuper(Binding, SuperBinding))
+		if (Binding.GetStackInterfaceSuper(SuperBinding))
 		{
 			return SuperBinding.GetBlendWeight(Context, ChildIndex);
 		}
@@ -20,7 +20,7 @@ namespace UE::AnimNext
 	const FAlphaBlend* IDiscreteBlend::GetBlendState(const FExecutionContext& Context, const TTraitBinding<IDiscreteBlend>& Binding, int32 ChildIndex) const
 	{
 		TTraitBinding<IDiscreteBlend> SuperBinding;
-		if (Context.GetInterfaceSuper(Binding, SuperBinding))
+		if (Binding.GetStackInterfaceSuper(SuperBinding))
 		{
 			return SuperBinding.GetBlendState(Context, ChildIndex);
 		}
@@ -31,7 +31,7 @@ namespace UE::AnimNext
 	int32 IDiscreteBlend::GetBlendDestinationChildIndex(const FExecutionContext& Context, const TTraitBinding<IDiscreteBlend>& Binding) const
 	{
 		TTraitBinding<IDiscreteBlend> SuperBinding;
-		if (Context.GetInterfaceSuper(Binding, SuperBinding))
+		if (Binding.GetStackInterfaceSuper(SuperBinding))
 		{
 			return SuperBinding.GetBlendDestinationChildIndex(Context);
 		}
@@ -42,7 +42,7 @@ namespace UE::AnimNext
 	void IDiscreteBlend::OnBlendTransition(const FExecutionContext& Context, const TTraitBinding<IDiscreteBlend>& Binding, int32 OldChildIndex, int32 NewChildIndex) const
 	{
 		TTraitBinding<IDiscreteBlend> SuperBinding;
-		if (Context.GetInterfaceSuper(Binding, SuperBinding))
+		if (Binding.GetStackInterfaceSuper(SuperBinding))
 		{
 			return SuperBinding.OnBlendTransition(Context, OldChildIndex, NewChildIndex);
 		}
@@ -51,7 +51,7 @@ namespace UE::AnimNext
 	void IDiscreteBlend::OnBlendInitiated(const FExecutionContext& Context, const TTraitBinding<IDiscreteBlend>& Binding, int32 ChildIndex) const
 	{
 		TTraitBinding<IDiscreteBlend> SuperBinding;
-		if (Context.GetInterfaceSuper(Binding, SuperBinding))
+		if (Binding.GetStackInterfaceSuper(SuperBinding))
 		{
 			return SuperBinding.OnBlendInitiated(Context, ChildIndex);
 		}
@@ -60,7 +60,7 @@ namespace UE::AnimNext
 	void IDiscreteBlend::OnBlendTerminated(const FExecutionContext& Context, const TTraitBinding<IDiscreteBlend>& Binding, int32 ChildIndex) const
 	{
 		TTraitBinding<IDiscreteBlend> SuperBinding;
-		if (Context.GetInterfaceSuper(Binding, SuperBinding))
+		if (Binding.GetStackInterfaceSuper(SuperBinding))
 		{
 			return SuperBinding.OnBlendTerminated(Context, ChildIndex);
 		}
diff --git a/Engine/Plugins/Experimental/AnimNext/Source/AnimNext/Private/TraitInterfaces/IEvaluate.cpp b/Engine/Plugins/Experimental/AnimNext/Source/AnimNext/Private/TraitInterfaces/IEvaluate.cpp
--- a/Engine/Plugins/Experimental/AnimNext/Source/AnimNext/Private/TraitInterfaces/IEvaluate.cpp
+++ b/Engine/Plugins/Experimental/AnimNext/Source/AnimNext/Private/TraitInterfaces/IEvaluate.cpp
@@ -10,16 +10,15 @@ DEFINE_STAT(STAT_AnimNext_EvaluateGraph);
 
 namespace UE::AnimNext
 {
-	FEvaluateTraversalContext::FEvaluateTraversalContext(const FExecutionContext& InExecutionContext, FEvaluationProgram& InEvaluationProgram)
-		: FExecutionContextProxy(InExecutionContext)
-		, EvaluationProgram(InEvaluationProgram)
+	FEvaluateTraversalContext::FEvaluateTraversalContext(FEvaluationProgram& InEvaluationProgram)
+		: EvaluationProgram(InEvaluationProgram)
 	{
 	}
 
 	void IEvaluate::PreEvaluate(FEvaluateTraversalContext& Context, const TTraitBinding<IEvaluate>& Binding) const
 	{
 		TTraitBinding<IEvaluate> SuperBinding;
-		if (Context.GetInterfaceSuper(Binding, SuperBinding))
+		if (Binding.GetStackInterfaceSuper(SuperBinding))
 		{
 			SuperBinding.PreEvaluate(Context);
 		}
@@ -28,7 +27,7 @@ namespace UE::AnimNext
 	void IEvaluate::PostEvaluate(FEvaluateTraversalContext& Context, const TTraitBinding<IEvaluate>& Binding) const
 	{
 		TTraitBinding<IEvaluate> SuperBinding;
-		if (Context.GetInterfaceSuper(Binding, SuperBinding))
+		if (Binding.GetStackInterfaceSuper(SuperBinding))
 		{
 			SuperBinding.PostEvaluate(Context);
 		}
@@ -49,6 +48,9 @@ namespace UE::AnimNext
 		// Which step we wish to perform when we next see this entry
 		EEvaluateStep		DesiredStep = EEvaluateStep::PreEvaluate;
 
+		// The trait stack binding for this evaluate entry
+		FTraitStackBinding	TraitStack;
+
 		// Once we've called PreUpdate, we cache the trait binding to avoid a redundant query to call PostUpdate
 		TTraitBinding<IEvaluate> EvaluateTrait;
 
@@ -93,8 +95,8 @@ namespace UE::AnimNext
 		FChildrenArray Children;
 		Children.Reserve(16);
 
-		FExecutionContext ExecutionContext;
-		FEvaluateTraversalContext TraversalContext(ExecutionContext, EvaluationProgram);
+		FEvaluateTraversalContext TraversalContext(EvaluationProgram);
+		TraversalContext.BindTo(GraphRootPtr);
 
 		// Add the graph root to kick start the evaluation process
 		FEvaluateEntry GraphRootEntry(GraphRootPtr, EEvaluateStep::PreEvaluate);
@@ -114,12 +116,12 @@ namespace UE::AnimNext
 
 			const FWeakTraitPtr& EntryTraitPtr = Entry->TraitPtr;
 
-			// Make sure the execution context is bound to our graph instance
-			ExecutionContext.BindTo(EntryTraitPtr);
-
 			if (Entry->DesiredStep == EEvaluateStep::PreEvaluate)
 			{
-				if (ExecutionContext.GetInterface(EntryTraitPtr, Entry->EvaluateTrait))
+				// Bind and cache our trait stack
+				ensure(TraversalContext.GetStack(EntryTraitPtr, Entry->TraitStack));
+
+				if (Entry->TraitStack.GetInterface(Entry->EvaluateTrait))
 				{
 					// This is the first time we visit this node, time to pre-evaluate
 					Entry->EvaluateTrait.PreEvaluate(TraversalContext);
@@ -135,13 +137,19 @@ namespace UE::AnimNext
 					bIsEntryUsed = false;
 				}
 
-				if (ExecutionContext.GetInterface(EntryTraitPtr, HierarchyTrait))
+				if (Entry->TraitStack.GetInterface(HierarchyTrait))
 				{
-					HierarchyTrait.GetChildren(ExecutionContext, Children);
+					IHierarchy::GetStackChildren(TraversalContext, Entry->TraitStack, Children);
 
 					// Append our children in reserve order so that they are visited in the same order they were added
 					for (int32 ChildIndex = Children.Num() - 1; ChildIndex >= 0; --ChildIndex)
 					{
+						const FWeakTraitPtr& ChildPtr = Children[ChildIndex];
+						if (!ChildPtr.IsValid())
+						{
+							continue;
+						}
+
 						// Insert our new child on top of the stack
 
 						FEvaluateEntry* ChildEntry;
@@ -151,14 +159,14 @@ namespace UE::AnimNext
 							ChildEntry = FreeEntryList;
 							FreeEntryList = ChildEntry->NextFreeEntry;
 
-							ChildEntry->TraitPtr = Children[ChildIndex];
+							ChildEntry->TraitPtr = ChildPtr;
 							ChildEntry->DesiredStep = EEvaluateStep::PreEvaluate;
 							ChildEntry->PrevStackEntry = NodesPendingUpdateStackTop;
 						}
 						else
 						{
 							// Allocate a new entry
-							ChildEntry = new(MemStack) FEvaluateEntry(Children[ChildIndex], EEvaluateStep::PreEvaluate, NodesPendingUpdateStackTop);
+							ChildEntry = new(MemStack) FEvaluateEntry(ChildPtr, EEvaluateStep::PreEvaluate, NodesPendingUpdateStackTop);
 						}
 
 						NodesPendingUpdateStackTop = ChildEntry;
diff --git a/Engine/Plugins/Experimental/AnimNext/Source/AnimNext/Private/TraitInterfaces/IGarbageColllection.cpp b/Engine/Plugins/Experimental/AnimNext/Source/AnimNext/Private/TraitInterfaces/IGarbageColllection.cpp
--- a/Engine/Plugins/Experimental/AnimNext/Source/AnimNext/Private/TraitInterfaces/IGarbageColllection.cpp
+++ b/Engine/Plugins/Experimental/AnimNext/Source/AnimNext/Private/TraitInterfaces/IGarbageColllection.cpp
@@ -10,7 +10,7 @@ namespace UE::AnimNext
 	void IGarbageCollection::RegisterWithGC(const FExecutionContext& Context, const FTraitBinding& Binding)
 	{
 		FGCGraphInstanceComponent& Component = Context.GetComponent<FGCGraphInstanceComponent>();
-		Component.Register(Context.GetGraphInstance(), Binding.GetTraitPtr());
+		Component.Register(Binding.GetTraitPtr());
 	}
 
 	void IGarbageCollection::UnregisterWithGC(const FExecutionContext& Context, const FTraitBinding& Binding)
@@ -22,7 +22,7 @@ namespace UE::AnimNext
 	void IGarbageCollection::AddReferencedObjects(const FExecutionContext& Context, const TTraitBinding<IGarbageCollection>& Binding, FReferenceCollector& Collector) const
 	{
 		TTraitBinding<IGarbageCollection> SuperBinding;
-		if (Context.GetInterfaceSuper(Binding, SuperBinding))
+		if (Binding.GetStackInterfaceSuper(SuperBinding))
 		{
 			SuperBinding.AddReferencedObjects(Context, Collector);
 		}
diff --git a/Engine/Plugins/Experimental/AnimNext/Source/AnimNext/Private/TraitInterfaces/IGroupSynchronization.cpp b/Engine/Plugins/Experimental/AnimNext/Source/AnimNext/Private/TraitInterfaces/IGroupSynchronization.cpp
--- a/Engine/Plugins/Experimental/AnimNext/Source/AnimNext/Private/TraitInterfaces/IGroupSynchronization.cpp
+++ b/Engine/Plugins/Experimental/AnimNext/Source/AnimNext/Private/TraitInterfaces/IGroupSynchronization.cpp
@@ -9,7 +9,7 @@ namespace UE::AnimNext
 	FName IGroupSynchronization::GetGroupName(const FExecutionContext& Context, const TTraitBinding<IGroupSynchronization>& Binding) const
 	{
 		TTraitBinding<IGroupSynchronization> SuperBinding;
-		if (Context.GetInterfaceSuper(Binding, SuperBinding))
+		if (Binding.GetStackInterfaceSuper(SuperBinding))
 		{
 			return SuperBinding.GetGroupName(Context);
 		}
@@ -20,7 +20,7 @@ namespace UE::AnimNext
 	EAnimGroupRole::Type IGroupSynchronization::GetGroupRole(const FExecutionContext& Context, const TTraitBinding<IGroupSynchronization>& Binding) const
 	{
 		TTraitBinding<IGroupSynchronization> SuperBinding;
-		if (Context.GetInterfaceSuper(Binding, SuperBinding))
+		if (Binding.GetStackInterfaceSuper(SuperBinding))
 		{
 			return SuperBinding.GetGroupRole(Context);
 		}
@@ -31,7 +31,7 @@ namespace UE::AnimNext
 	float IGroupSynchronization::AdvanceBy(const FExecutionContext& Context, const TTraitBinding<IGroupSynchronization>& Binding, float DeltaTime) const
 	{
 		TTraitBinding<IGroupSynchronization> SuperBinding;
-		if (Context.GetInterfaceSuper(Binding, SuperBinding))
+		if (Binding.GetStackInterfaceSuper(SuperBinding))
 		{
 			return SuperBinding.AdvanceBy(Context, DeltaTime);
 		}
@@ -42,7 +42,7 @@ namespace UE::AnimNext
 	void IGroupSynchronization::AdvanceToRatio(const FExecutionContext& Context, const TTraitBinding<IGroupSynchronization>& Binding, float ProgressRatio) const
 	{
 		TTraitBinding<IGroupSynchronization> SuperBinding;
-		if (Context.GetInterfaceSuper(Binding, SuperBinding))
+		if (Binding.GetStackInterfaceSuper(SuperBinding))
 		{
 			SuperBinding.AdvanceToRatio(Context, ProgressRatio);
 		}
diff --git a/Engine/Plugins/Experimental/AnimNext/Source/AnimNext/Private/TraitInterfaces/IHierarchy.cpp b/Engine/Plugins/Experimental/AnimNext/Source/AnimNext/Private/TraitInterfaces/IHierarchy.cpp
--- a/Engine/Plugins/Experimental/AnimNext/Source/AnimNext/Private/TraitInterfaces/IHierarchy.cpp
+++ b/Engine/Plugins/Experimental/AnimNext/Source/AnimNext/Private/TraitInterfaces/IHierarchy.cpp
@@ -8,21 +8,59 @@ namespace UE::AnimNext
 {
 	uint32 IHierarchy::GetNumChildren(const FExecutionContext& Context, const TTraitBinding<IHierarchy>& Binding) const
 	{
-		TTraitBinding<IHierarchy> SuperBinding;
-		if (Context.GetInterfaceSuper(Binding, SuperBinding))
+		// We only wish to count children of the queried trait
+		// No need to forward to our super
+		// To count all children of a trait stack, use GetNumStackChildren instead
+		return 0;
+	}
+
+	void IHierarchy::GetChildren(const FExecutionContext& Context, const TTraitBinding<IHierarchy>& Binding, FChildrenArray& Children) const
+	{
+		// We only wish to append children of the queried trait
+		// No need to forward to our super
+		// To get all children of a trait stack, use GetStackChildren instead
+	}
+
+	void IHierarchy::GetStackChildren(const FExecutionContext& Context, const FTraitStackBinding& Binding, FChildrenArray& Children)
+	{
+		Children.Reset();
+
+		if (!Binding.IsValid())
 		{
-			return SuperBinding.GetNumChildren(Context);
+			return;
 		}
 
-		return 0;
+		TTraitBinding<IHierarchy> HierarchyTrait;
+
+		// Visit the trait stack and queue our children
+		Binding.GetInterface(HierarchyTrait);
+		while (HierarchyTrait.IsValid())
+		{
+			HierarchyTrait.GetChildren(Context, Children);
+
+			Binding.GetInterfaceSuper(HierarchyTrait, HierarchyTrait);
+		}
 	}
 
-	void IHierarchy::GetChildren(const FExecutionContext& Context, const TTraitBinding<IHierarchy>& Binding, FChildrenArray& Children) const
+	uint32 IHierarchy::GetNumStackChildren(const FExecutionContext& Context, const FTraitStackBinding& Binding)
 	{
-		TTraitBinding<IHierarchy> SuperBinding;
-		if (Context.GetInterfaceSuper(Binding, SuperBinding))
+		if (!Binding.IsValid())
+		{
+			return 0;
+		}
+
+		TTraitBinding<IHierarchy> HierarchyTrait;
+		uint32 NumChildren = 0;
+
+		// Visit the trait stack and append the children count
+		Binding.GetInterface(HierarchyTrait);
+		while (HierarchyTrait.IsValid())
 		{
-			SuperBinding.GetChildren(Context, Children);
+			NumChildren += HierarchyTrait.GetNumChildren(Context);
+
+			Binding.GetInterfaceSuper(HierarchyTrait, HierarchyTrait);
 		}
+
+		return NumChildren;
 	}
 }
diff --git a/Engine/Plugins/Experimental/AnimNext/Source/AnimNext/Private/TraitInterfaces/ISmoothBlend.cpp b/Engine/Plugins/Experimental/AnimNext/Source/AnimNext/Private/TraitInterfaces/ISmoothBlend.cpp
--- a/Engine/Plugins/Experimental/AnimNext/Source/AnimNext/Private/TraitInterfaces/ISmoothBlend.cpp
+++ b/Engine/Plugins/Experimental/AnimNext/Source/AnimNext/Private/TraitInterfaces/ISmoothBlend.cpp
@@ -9,7 +9,7 @@ namespace UE::AnimNext
 	float ISmoothBlend::GetBlendTime(const FExecutionContext& Context, const TTraitBinding<ISmoothBlend>& Binding, int32 ChildIndex) const
 	{
 		TTraitBinding<ISmoothBlend> SuperBinding;
-		if (Context.GetInterfaceSuper(Binding, SuperBinding))
+		if (Binding.GetStackInterfaceSuper(SuperBinding))
 		{
 			return SuperBinding.GetBlendTime(Context, ChildIndex);
 		}
diff --git a/Engine/Plugins/Experimental/AnimNext/Source/AnimNext/Private/TraitInterfaces/ITimeline.cpp b/Engine/Plugins/Experimental/AnimNext/Source/AnimNext/Private/TraitInterfaces/ITimeline.cpp
--- a/Engine/Plugins/Experimental/AnimNext/Source/AnimNext/Private/TraitInterfaces/ITimeline.cpp
+++ b/Engine/Plugins/Experimental/AnimNext/Source/AnimNext/Private/TraitInterfaces/ITimeline.cpp
@@ -9,7 +9,7 @@ namespace UE::AnimNext
 	float ITimeline::GetPlayRate(const FExecutionContext& Context, const TTraitBinding<ITimeline>& Binding) const
 	{
 		TTraitBinding<ITimeline> SuperBinding;
-		if (Context.GetInterfaceSuper(Binding, SuperBinding))
+		if (Binding.GetStackInterfaceSuper(SuperBinding))
 		{
 			return SuperBinding.GetPlayRate(Context);
 		}
@@ -20,7 +20,7 @@ namespace UE::AnimNext
 	float ITimeline::AdvanceBy(const FExecutionContext& Context, const TTraitBinding<ITimeline>& Binding, float DeltaTime) const
 	{
 		TTraitBinding<ITimeline> SuperBinding;
-		if (Context.GetInterfaceSuper(Binding, SuperBinding))
+		if (Binding.GetStackInterfaceSuper(SuperBinding))
 		{
 			return SuperBinding.AdvanceBy(Context, DeltaTime);
 		}
@@ -31,7 +31,7 @@ namespace UE::AnimNext
 	void ITimeline::AdvanceToRatio(const FExecutionContext& Context, const TTraitBinding<ITimeline>& Binding, float ProgressRatio) const
 	{
 		TTraitBinding<ITimeline> SuperBinding;
-		if (Context.GetInterfaceSuper(Binding, SuperBinding))
+		if (Binding.GetStackInterfaceSuper(SuperBinding))
 		{
 			SuperBinding.AdvanceToRatio(Context, ProgressRatio);
 		}
diff --git a/Engine/Plugins/Experimental/AnimNext/Source/AnimNext/Private/TraitInterfaces/IUpdate.cpp b/Engine/Plugins/Experimental/AnimNext/Source/AnimNext/Private/TraitInterfaces/IUpdate.cpp
--- a/Engine/Plugins/Experimental/AnimNext/Source/AnimNext/Private/TraitInterfaces/IUpdate.cpp
+++ b/Engine/Plugins/Experimental/AnimNext/Source/AnimNext/Private/TraitInterfaces/IUpdate.cpp
@@ -10,83 +10,82 @@ DEFINE_STAT(STAT_AnimNext_UpdateGraph);
 
 namespace UE::AnimNext
 {
-	void IUpdate::PreUpdate(FUpdateTraversalContext& Context, const TTraitBinding<IUpdate>& Binding, const FTraitUpdateState& TraitState) const
+	namespace Private
 	{
-		TTraitBinding<IUpdate> SuperBinding;
-		if (Context.GetInterfaceSuper(Binding, SuperBinding))
+		// This structure is transient and lives either on the stack or the memstack and its destructor may not be called
+		struct FUpdateEntry
 		{
-			SuperBinding.PreUpdate(Context, TraitState);
-		}
+			// The trait state for this entry
+			FTraitUpdateState			TraitState;
+
+			// The trait handle that points to our node to update
+			FWeakTraitPtr				TraitPtr;
+
+			// Whether or not PreUpdate had been called already
+			// TODO: Store bHasPreUpdated in the LSB of the entry pointer to save padding?
+			bool						bHasPreUpdated = false;
+
+			// The trait stack binding for this update entry
+			FTraitStackBinding			TraitStack;
+
+			// Once we've called PreUpdate, we cache the trait binding to avoid a redundant query to call PostUpdate
+			TTraitBinding<IUpdate>		UpdateTrait;
+
+			// These pointers are mutually exclusive
+			// An entry is either part of the queued update stack, the update stack, the free list, or none of the above
+			union
+			{
+				// Next entry in the stack of free entries
+				FUpdateEntry* NextFreeEntry = nullptr;
+
+				// Previous entry on the update stack
+				FUpdateEntry* PrevUpdateStackEntry;
+
+				// Previous entry on the queued update stack
+				FUpdateEntry* PrevQueuedUpdateStackEntry;
+			};
+
+			FUpdateEntry(const FWeakTraitPtr& InTraitPtr, const FTraitUpdateState InTraitState)
+				: TraitState(InTraitState)
+				, TraitPtr(InTraitPtr)
+			{
+			}
+		};
 	}
 
-	void IUpdate::QueueChildrenForTraversal(FUpdateTraversalContext& Context, const TTraitBinding<IUpdate>& Binding, const FTraitUpdateState& TraitState, FUpdateTraversalQueue& TraversalQueue) const
+	void IUpdate::PreUpdate(FUpdateTraversalContext& Context, const TTraitBinding<IUpdate>& Binding, const FTraitUpdateState& TraitState) const
 	{
 		TTraitBinding<IUpdate> SuperBinding;
-		if (Context.GetInterfaceSuper(Binding, SuperBinding))
+		if (Binding.GetStackInterfaceSuper(SuperBinding))
 		{
-			SuperBinding.QueueChildrenForTraversal(Context, TraitState, TraversalQueue);
+			SuperBinding.PreUpdate(Context, TraitState);
 		}
 	}
 
 	void IUpdate::PostUpdate(FUpdateTraversalContext& Context, const TTraitBinding<IUpdate>& Binding, const FTraitUpdateState& TraitState) const
 	{
 		TTraitBinding<IUpdate> SuperBinding;
-		if (Context.GetInterfaceSuper(Binding, SuperBinding))
+		if (Binding.GetStackInterfaceSuper(SuperBinding))
 		{
 			SuperBinding.PostUpdate(Context, TraitState);
 		}
 	}
 
-	//////////////////////////////////////////////////////////////////////////
-	// Traversal implementation
-
-	// This structure is transient and lives either on the stack or the memstack and its destructor may not be called
-	struct FUpdateEntry
+	void IUpdateTraversal::QueueChildrenForTraversal(FUpdateTraversalContext& Context, const TTraitBinding<IUpdateTraversal>& Binding, const FTraitUpdateState& TraitState, FUpdateTraversalQueue& TraversalQueue) const
 	{
-		// The trait state for this entry
-		FTraitUpdateState			TraitState;
-
-		// The trait handle that points to our node to update
-		FWeakTraitPtr				TraitPtr;
-
-		// Whether or not PreUpdate had been called already
-		// TODO: Store bHasPreUpdated in the LSB of the entry pointer to save padding?
-		bool						bHasPreUpdated = false;
-
-		// Once we've called PreUpdate, we cache the trait binding to avoid a redundant query to call PostUpdate
-		TTraitBinding<IUpdate>		UpdateTrait;
-
-		// These pointers are mutually exclusive
-		// An entry is either part of the queued update stack, the update stack, the free list, or none of the above
-		union
-		{
-			// Next entry in the stack of free entries
-			FUpdateEntry*			NextFreeEntry = nullptr;
-
-			// Previous entry on the update stack
-			FUpdateEntry*			PrevUpdateStackEntry;
-
-			// Previous entry on the queued update stack
-			FUpdateEntry*			PrevQueuedUpdateStackEntry;
-		};
-
-		FUpdateEntry(const FWeakTraitPtr& InTraitPtr, const FTraitUpdateState InTraitState)
-			: TraitState(InTraitState)
-			, TraitPtr(InTraitPtr)
-		{
-		}
-	};
+		// Nothing to do
+		// This function is called for each trait on the stack, one by one
+		// No need to forward to our super
+	}
 
-	FUpdateTraversalContext::FUpdateTraversalContext(const FExecutionContext& InExecutionContext, FMemStack& InMemStack)
-		: FExecutionContextProxy(InExecutionContext)
-		, MemStack(InMemStack)
-	{}
+	//////////////////////////////////////////////////////////////////////////
+	// Traversal implementation
 
 	void FUpdateTraversalContext::PushQueuedUpdateEntries(FUpdateTraversalQueue& TraversalQueue)
 	{
 		// Pop every entry from the queued update stack and push them onto the update stack
 		// reversing their order
-		while (FUpdateEntry* Entry = TraversalQueue.QueuedUpdateStackHead)
+		while (Private::FUpdateEntry* Entry = TraversalQueue.QueuedUpdateStackHead)
 		{
 			// Update our queued stack head
 			TraversalQueue.QueuedUpdateStackHead = Entry->PrevQueuedUpdateStackEntry;
@@ -96,15 +95,15 @@ namespace UE::AnimNext
 		}
 	}
 
-	void FUpdateTraversalContext::PushUpdateEntry(FUpdateEntry* Entry)
+	void FUpdateTraversalContext::PushUpdateEntry(Private::FUpdateEntry* Entry)
 	{
 		Entry->PrevUpdateStackEntry = UpdateStackHead;
 		UpdateStackHead = Entry;
 	}
 
-	FUpdateEntry* FUpdateTraversalContext::PopUpdateEntry()
+	Private::FUpdateEntry* FUpdateTraversalContext::PopUpdateEntry()
 	{
-		FUpdateEntry* ChildEntry = UpdateStackHead;
+		Private::FUpdateEntry* ChildEntry = UpdateStackHead;
 		if (ChildEntry != nullptr)
 		{
 			// We have a child, set our new head
@@ -114,15 +113,15 @@ namespace UE::AnimNext
 		return ChildEntry;
 	}
 
-	void FUpdateTraversalContext::PushFreeEntry(FUpdateEntry* Entry)
+	void FUpdateTraversalContext::PushFreeEntry(Private::FUpdateEntry* Entry)
 	{
 		Entry->NextFreeEntry = FreeEntryStackHead;
 		FreeEntryStackHead = Entry;
 	}
 
-	FUpdateEntry* FUpdateTraversalContext::GetNewEntry(const FWeakTraitPtr& TraitPtr, const FTraitUpdateState& TraitState)
+	Private::FUpdateEntry* FUpdateTraversalContext::GetNewEntry(const FWeakTraitPtr& TraitPtr, const FTraitUpdateState& TraitState)
 	{
-		FUpdateEntry* FreeEntry = FreeEntryStackHead;
+		Private::FUpdateEntry* FreeEntry = FreeEntryStackHead;
 		if (FreeEntry != nullptr)
 		{
 			// We have a free entry, set our new head
@@ -137,7 +136,7 @@ namespace UE::AnimNext
 		else
 		{
 			// Allocate a new entry
-			FreeEntry = new(MemStack) FUpdateEntry(TraitPtr, TraitState);
+			FreeEntry = new(MemStack) Private::FUpdateEntry(TraitPtr, TraitState);
 		}
 
 		return FreeEntry;
@@ -155,7 +154,7 @@ namespace UE::AnimNext
 			return;	// Don't queue invalid pointers
 		}
 
-		FUpdateEntry* ChildEntry = TraversalContext.GetNewEntry(ChildPtr, ChildTraitState);
+		Private::FUpdateEntry* ChildEntry = TraversalContext.GetNewEntry(ChildPtr, ChildTraitState);
 
 		// We push children that are queued onto a stack
 		// Once pre-update is done, we'll pop queued entries one by one and push them
@@ -222,44 +221,48 @@ namespace UE::AnimNext
 			return;	// We can only update starting at the root
 		}
 
-		FMemStack& MemStack = FMemStack::Get();
+		FUpdateTraversalContext TraversalContext;
+
+		FMemStack& MemStack = TraversalContext.GetMemStack();
 		FMemMark Mark(MemStack);
 
 		FChildrenArray Children;
 		TTraitBinding<IHierarchy> HierarchyTrait;
+		TTraitBinding<IUpdateTraversal> UpdateTraversalTrait;
+		FTraitBinding TraitBinding;
 
-		FExecutionContext ExecutionContext;
-		FUpdateTraversalContext TraversalContext(ExecutionContext, MemStack);
 		FUpdateTraversalQueue TraversalQueue(TraversalContext);
 
 		// Before we start the traversal, we give the graph instance components the chance to do some work
-		ExecutionContext.BindTo(GraphInstance);
-		for (auto It = ExecutionContext.GetComponentIterator(); It; ++It)
+		TraversalContext.BindTo(GraphInstance);
+		for (auto It = TraversalContext.GetComponentIterator(); It; ++It)
 		{
-			It.Value()->PreUpdate(ExecutionContext);
+			It.Value()->PreUpdate(TraversalContext);
 		}
 
 		// Add the graph root to start the update process
-		FUpdateEntry RootEntry(GraphInstance.GetGraphRootPtr(), FTraitUpdateState(DeltaTime));
+		Private::FUpdateEntry RootEntry(GraphInstance.GetGraphRootPtr(), FTraitUpdateState(DeltaTime));
 		TraversalContext.PushUpdateEntry(&RootEntry);
 
-		while (FUpdateEntry* Entry = TraversalContext.PopUpdateEntry())
+		while (Private::FUpdateEntry* Entry = TraversalContext.PopUpdateEntry())
 		{
 			const FWeakTraitPtr& EntryTraitPtr = Entry->TraitPtr;
 
-			// Make sure the execution context is bound to our graph instance
-			ExecutionContext.BindTo(EntryTraitPtr);
-
 			if (!Entry->bHasPreUpdated)
 			{
 				// This is the first time we visit this node, time to pre-update
-				// But first, if it has latent pins, we must execute and cache their results
+
+				// Bind and cache our trait stack
+				ensure(TraversalContext.GetStack(EntryTraitPtr, Entry->TraitStack));
+
+				// First, if it has latent pins, we must execute and cache their results
 				// This will ensure that other calls into this node will have a consistent view of
 				// what the node saw when it started to update. We thus take a snapshot.
 				const bool bIsFrozen = false;	// Not yet supported
-				ExecutionContext.SnapshotLatentProperties(EntryTraitPtr, bIsFrozen);
+				Entry->TraitStack.SnapshotLatentProperties(bIsFrozen);
 
-				if (ExecutionContext.GetInterface(EntryTraitPtr, Entry->UpdateTrait))
+				// If this trait stack implements IUpdate, call into it
+				if (Entry->TraitStack.GetInterface(Entry->UpdateTrait))
 				{
 					Entry->UpdateTrait.PreUpdate(TraversalContext, Entry->TraitState);
 
@@ -268,42 +271,56 @@ namespace UE::AnimNext
 
 					// Push this entry onto the update stack, we'll call it once all our children have finished executing
 					TraversalContext.PushUpdateEntry(Entry);
-
-					// Request that the trait queues the children it wants to visit
-					// This is a separate function from PreUpdate to simplify traversal management. It is often the case that
-					// the base trait is the one best placed to figure out how to optimally queue children since it
-					// owns the handles to them. However, if an additive trait wishes to override PreUpdate, it might want
-					// to perform logic after the base PreUpdate but before children are queued. Without a separate function,
-					// we would have to rewrite the base PreUpdate entirely and use IHierarchy to query the handles of our children.
-					Entry->UpdateTrait.QueueChildrenForTraversal(TraversalContext, Entry->TraitState, TraversalQueue);
-
-					// Iterate over our queued children and push them onto the update stack
-					// We do this to allow children to be queued in traversal order which is intuitive
-					// but to traverse them in that order, they must be pushed in reverse order onto the update stack
-					TraversalContext.PushQueuedUpdateEntries(TraversalQueue);
 				}
 				else
 				{
-					// This node doesn't implement IUpdate
-					// We'll grab its children and traverse them
-					if (ExecutionContext.GetInterface(EntryTraitPtr, HierarchyTrait))
+					// We don't need this entry anymore
+					TraversalContext.PushFreeEntry(Entry);
+				}
+
+				// Now visit the trait stack and queue our children
+				ensure(Entry->TraitStack.GetTopTrait(TraitBinding));
+				do
+				{
+					if (TraitBinding.AsInterface(UpdateTraversalTrait))
+					{
+						// Request that the trait queues the children it wants to visit
+						// This is a separate function from PreUpdate to simplify traversal management. It is often the case that
+						// the base trait is the one best placed to figure out how to optimally queue children since it
+						// owns the handles to them. However, if an additive trait wishes to override PreUpdate, it might want
+						// to perform logic after the base PreUpdate but before children are queued. Without a separate function,
+						// we would have to rewrite the base PreUpdate entirely and use IHierarchy to query the handles of our children.
+						UpdateTraversalTrait.QueueChildrenForTraversal(TraversalContext, Entry->TraitState, TraversalQueue);
+
+
+						// Iterate over our queued children and push them onto the update stack
+						// We do this to allow children to be queued in traversal order which is intuitive
+						// but to traverse them in that order, they must be pushed in reverse order onto the update stack
+						TraversalContext.PushQueuedUpdateEntries(TraversalQueue);
+					}
+					else if (TraitBinding.AsInterface(HierarchyTrait))
 					{
-						HierarchyTrait.GetChildren(ExecutionContext, Children);
+						HierarchyTrait.GetChildren(TraversalContext, Children);
 
 						// Append our children in reserve order so that they are visited in the same order they were added
 						for (int32 ChildIndex = Children.Num() - 1; ChildIndex >= 0; --ChildIndex)
 						{
-							FUpdateEntry* ChildEntry = TraversalContext.GetNewEntry(Children[ChildIndex], Entry->TraitState);
-							TraversalContext.PushUpdateEntry(ChildEntry);
+							const FWeakTraitPtr& ChildPtr = Children[ChildIndex];
+							if (ChildPtr.IsValid())
+							{
+								Private::FUpdateEntry* ChildEntry = TraversalContext.GetNewEntry(ChildPtr, Entry->TraitState);
+								TraversalContext.PushUpdateEntry(ChildEntry);
+							}
 						}
 
 						// Reset our container for the next time we need it
 						Children.Reset();
 					}
-
-					// We don't need this entry anymore
-					TraversalContext.PushFreeEntry(Entry);
-				}
+					else
+					{
+						// We don't have any children since we don't implement any of the relevant interfaces
+					}
+				} while (Entry->TraitStack.GetParentTrait(TraitBinding, TraitBinding));
 			}
 			else
 			{
@@ -317,10 +334,10 @@ namespace UE::AnimNext
 		}
 
 		// After we finish the traversal, we give the graph instance components the chance to do some work
-		ExecutionContext.BindTo(GraphInstance);
-		for (auto It = ExecutionContext.GetComponentIterator(); It; ++It)
+		TraversalContext.BindTo(GraphInstance);
+		for (auto It = TraversalContext.GetComponentIterator(); It; ++It)
 		{
-			It.Value()->PostUpdate(ExecutionContext);
+			It.Value()->PostUpdate(TraversalContext);
 		}
 	}
 }
diff --git a/Engine/Plugins/Experimental/AnimNext/Source/AnimNext/Private/Traits/BlendByBool.cpp b/Engine/Plugins/Experimental/AnimNext/Source/AnimNext/Private/Traits/BlendByBool.cpp
--- a/Engine/Plugins/Experimental/AnimNext/Source/AnimNext/Private/Traits/BlendByBool.cpp
+++ b/Engine/Plugins/Experimental/AnimNext/Source/AnimNext/Private/Traits/BlendByBool.cpp
@@ -15,6 +15,7 @@ namespace UE::AnimNext
 		GeneratorMacro(IDiscreteBlend) \
 		GeneratorMacro(IHierarchy) \
 		GeneratorMacro(IUpdate) \
+		GeneratorMacro(IUpdateTraversal) \
 
 	GENERATE_ANIM_TRAIT_IMPLEMENTATION(FBlendByBoolTrait, TRAIT_INTERFACE_ENUMERATOR)
 	#undef TRAIT_INTERFACE_ENUMERATOR
@@ -28,7 +29,7 @@ namespace UE::AnimNext
 		FInstanceData* InstanceData = Binding.GetInstanceData<FInstanceData>();
 
 		TTraitBinding<IDiscreteBlend> DiscreteBlendTrait;
-		Context.GetInterface(Binding, DiscreteBlendTrait);
+		Binding.GetStackInterface(DiscreteBlendTrait);
 
 		const int32 DestinationChildIndex = DiscreteBlendTrait.GetBlendDestinationChildIndex(Context);
 		if (InstanceData->PreviousChildIndex != DestinationChildIndex)
@@ -39,7 +40,7 @@ namespace UE::AnimNext
 		}
 	}
 
-	void FBlendByBoolTrait::QueueChildrenForTraversal(FUpdateTraversalContext& Context, const TTraitBinding<IUpdate>& Binding, const FTraitUpdateState& TraitState, FUpdateTraversalQueue& TraversalQueue) const
+	void FBlendByBoolTrait::QueueChildrenForTraversal(FUpdateTraversalContext& Context, const TTraitBinding<IUpdateTraversal>& Binding, const FTraitUpdateState& TraitState, FUpdateTraversalQueue& TraversalQueue) const
 	{
 		const FInstanceData* InstanceData = Binding.GetInstanceData<FInstanceData>();
 
@@ -47,7 +48,7 @@ namespace UE::AnimNext
 		const int32 DestinationChildIndex = InstanceData->PreviousChildIndex;
 
 		TTraitBinding<IDiscreteBlend> DiscreteBlendTrait;
-		Context.GetInterface(Binding, DiscreteBlendTrait);
+		Binding.GetStackInterface(DiscreteBlendTrait);
 
 		const float BlendWeightTrue = DiscreteBlendTrait.GetBlendWeight(Context, TRUE_CHILD_INDEX);
 		if (InstanceData->TrueChild.IsValid() && FAnimWeight::IsRelevant(BlendWeightTrue))
@@ -91,7 +92,7 @@ namespace UE::AnimNext
 	float FBlendByBoolTrait::GetBlendWeight(const FExecutionContext& Context, const TTraitBinding<IDiscreteBlend>& Binding, int32 ChildIndex) const
 	{
 		TTraitBinding<IDiscreteBlend> DiscreteBlendTrait;
-		Context.GetInterface(Binding, DiscreteBlendTrait);
+		Binding.GetStackInterface(DiscreteBlendTrait);
 
 		const float DestinationChildIndex = DiscreteBlendTrait.GetBlendDestinationChildIndex(Context);
 
@@ -114,14 +115,14 @@ namespace UE::AnimNext
 	{
 		const FSharedData* SharedData = Binding.GetSharedData<FSharedData>();
 
-		const bool bCondition = SharedData->GetbCondition(Context, Binding);
+		const bool bCondition = SharedData->GetbCondition(Binding);
 		return bCondition ? TRUE_CHILD_INDEX : FALSE_CHILD_INDEX;
 	}
 
 	void FBlendByBoolTrait::OnBlendTransition(const FExecutionContext& Context, const TTraitBinding<IDiscreteBlend>& Binding, int32 OldChildIndex, int32 NewChildIndex) const
 	{
 		TTraitBinding<IDiscreteBlend> DiscreteBlendTrait;
-		Context.GetInterface(Binding, DiscreteBlendTrait);
+		Binding.GetStackInterface(DiscreteBlendTrait);
 
 		// We initiate immediately when we transition
 		DiscreteBlendTrait.OnBlendInitiated(Context, NewChildIndex);
diff --git a/Engine/Plugins/Experimental/AnimNext/Source/AnimNext/Private/Traits/BlendByBool.h b/Engine/Plugins/Experimental/AnimNext/Source/AnimNext/Private/Traits/BlendByBool.h
--- a/Engine/Plugins/Experimental/AnimNext/Source/AnimNext/Private/Traits/BlendByBool.h
+++ b/Engine/Plugins/Experimental/AnimNext/Source/AnimNext/Private/Traits/BlendByBool.h
@@ -43,7 +43,7 @@ namespace UE::AnimNext
 	 * 
 	 * A trait that can blend two discrete inputs through a boolean.
 	 */
-	struct FBlendByBoolTrait : FBaseTrait, IUpdate, IHierarchy, IDiscreteBlend
+	struct FBlendByBoolTrait : FBaseTrait, IUpdate, IUpdateTraversal, IHierarchy, IDiscreteBlend
 	{
 		DECLARE_ANIM_TRAIT(FBlendByBoolTrait, 0x452572ff, FBaseTrait)
 
@@ -59,7 +59,9 @@ namespace UE::AnimNext
 
 		// IUpdate impl
 		virtual void PreUpdate(FUpdateTraversalContext& Context, const TTraitBinding<IUpdate>& Binding, const FTraitUpdateState& TraitState) const override;
-		virtual void QueueChildrenForTraversal(FUpdateTraversalContext& Context, const TTraitBinding<IUpdate>& Binding, const FTraitUpdateState& TraitState, FUpdateTraversalQueue& TraversalQueue) const override;
+
+		// IUpdateTraversal impl
+		virtual void QueueChildrenForTraversal(FUpdateTraversalContext& Context, const TTraitBinding<IUpdateTraversal>& Binding, const FTraitUpdateState& TraitState, FUpdateTraversalQueue& TraversalQueue) const override;
 
 		// IHierarchy impl
 		virtual uint32 GetNumChildren(const FExecutionContext& Context, const TTraitBinding<IHierarchy>& Binding) const override;
diff --git a/Engine/Plugins/Experimental/AnimNext/Source/AnimNext/Private/Traits/BlendSmoother.cpp b/Engine/Plugins/Experimental/AnimNext/Source/AnimNext/Private/Traits/BlendSmoother.cpp
--- a/Engine/Plugins/Experimental/AnimNext/Source/AnimNext/Private/Traits/BlendSmoother.cpp
+++ b/Engine/Plugins/Experimental/AnimNext/Source/AnimNext/Private/Traits/BlendSmoother.cpp
@@ -140,7 +140,7 @@ namespace UE::AnimNext
 		}
 
 		TTraitBinding<IDiscreteBlend> DiscreteBlendTrait;
-		Context.GetInterface(Binding, DiscreteBlendTrait);
+		Binding.GetStackInterface(DiscreteBlendTrait);
 
 		// Free any newly inactive children
 		const int32 NumChildren = InstanceData->PerChildBlendData.Num();
@@ -195,7 +195,7 @@ namespace UE::AnimNext
 		const float WeightDifference = FMath::Clamp(FMath::Abs(NewChildDesiredWeight - NewChildCurrentWeight), 0.0f, 1.0f);
 
 		TTraitBinding<ISmoothBlend> SmoothBlendTrait;
-		Context.GetInterface(Binding, SmoothBlendTrait);
+		Binding.GetStackInterface(SmoothBlendTrait);
 
 		const float BlendTime = SmoothBlendTrait.GetBlendTime(Context, NewChildIndex);
 		const float RemainingBlendTime = OldChildIndex != INDEX_NONE ? (BlendTime * WeightDifference) : 0.0f;
@@ -226,7 +226,7 @@ namespace UE::AnimNext
 		// We just initiate the new blend manually
 
 		TTraitBinding<IDiscreteBlend> DiscreteBlendTrait;
-		Context.GetInterface(Binding, DiscreteBlendTrait);
+		Binding.GetStackInterface(DiscreteBlendTrait);
 
 		DiscreteBlendTrait.OnBlendInitiated(Context, NewChildIndex);
 	}
@@ -256,10 +256,7 @@ namespace UE::AnimNext
 	{
 		check(InstanceData->PerChildBlendData.IsEmpty());
 
-		TTraitBinding<IHierarchy> HierarchyTrait;
-		Context.GetInterface(Binding, HierarchyTrait);
-
-		const uint32 NumChildren = HierarchyTrait.GetNumChildren(Context);
+		const uint32 NumChildren = IHierarchy::GetNumStackChildren(Context, Binding);
 
 		InstanceData->PerChildBlendData.SetNum(NumChildren);
 
diff --git a/Engine/Plugins/Experimental/AnimNext/Source/AnimNext/Private/Traits/BlendSmootherPerBone.cpp b/Engine/Plugins/Experimental/AnimNext/Source/AnimNext/Private/Traits/BlendSmootherPerBone.cpp
--- a/Engine/Plugins/Experimental/AnimNext/Source/AnimNext/Private/Traits/BlendSmootherPerBone.cpp
+++ b/Engine/Plugins/Experimental/AnimNext/Source/AnimNext/Private/Traits/BlendSmootherPerBone.cpp
@@ -88,7 +88,7 @@ namespace UE::AnimNext
 		}
 
 		TTraitBinding<IDiscreteBlend> DiscreteBlendTrait;
-		Context.GetInterface(Binding, DiscreteBlendTrait);
+		Binding.GetStackInterface(DiscreteBlendTrait);
 
 		const int32 DestinationChildIndex = DiscreteBlendTrait.GetBlendDestinationChildIndex(Context);
 
@@ -138,7 +138,7 @@ namespace UE::AnimNext
 		}
 
 		TTraitBinding<IDiscreteBlend> DiscreteBlendTrait;
-		Context.GetInterface(Binding, DiscreteBlendTrait);
+		Binding.GetStackInterface(DiscreteBlendTrait);
 
 		for (int32 ChildIndex = 0; ChildIndex < NumChildren; ++ChildIndex)
 		{
@@ -158,10 +158,7 @@ namespace UE::AnimNext
 			return;	// No blend profile set, nothing to do
 		}
 
-		TTraitBinding<IHierarchy> HierarchyTrait;
-		Context.GetInterface(Binding, HierarchyTrait);
-
-		const uint32 NumChildren = HierarchyTrait.GetNumChildren(Context);
+		const uint32 NumChildren = IHierarchy::GetNumStackChildren(Context, Binding);
 
 		InstanceData->PerChildBlendData.SetNum(NumChildren);
 		InstanceData->PerBoneSampleData.SetNum(NumChildren);
diff --git a/Engine/Plugins/Experimental/AnimNext/Source/AnimNext/Private/Traits/BlendTwoWay.cpp b/Engine/Plugins/Experimental/AnimNext/Source/AnimNext/Private/Traits/BlendTwoWay.cpp
--- a/Engine/Plugins/Experimental/AnimNext/Source/AnimNext/Private/Traits/BlendTwoWay.cpp
+++ b/Engine/Plugins/Experimental/AnimNext/Source/AnimNext/Private/Traits/BlendTwoWay.cpp
@@ -16,6 +16,7 @@ namespace UE::AnimNext
 		GeneratorMacro(IEvaluate) \
 		GeneratorMacro(IHierarchy) \
 		GeneratorMacro(IUpdate) \
+		GeneratorMacro(IUpdateTraversal) \
 
 	GENERATE_ANIM_TRAIT_IMPLEMENTATION(FBlendTwoWayTrait, TRAIT_INTERFACE_ENUMERATOR)
 	#undef TRAIT_INTERFACE_ENUMERATOR
@@ -29,7 +30,7 @@ namespace UE::AnimNext
 			// We have two children, interpolate them
 
 			TTraitBinding<IContinuousBlend> ContinuousBlendTrait;
-			Context.GetInterface(Binding, ContinuousBlendTrait);
+			Binding.GetStackInterface(ContinuousBlendTrait);
 
 			const float BlendWeight = ContinuousBlendTrait.GetBlendWeight(Context, 1);
 			Context.AppendTask(FAnimNextBlendTwoKeyframesTask::Make(BlendWeight));
@@ -46,7 +47,7 @@ namespace UE::AnimNext
 		FInstanceData* InstanceData = Binding.GetInstanceData<FInstanceData>();
 
 		TTraitBinding<IContinuousBlend> ContinuousBlendTrait;
-		Context.GetInterface(Binding, ContinuousBlendTrait);
+		Binding.GetStackInterface(ContinuousBlendTrait);
 
 		const float BlendWeightB = ContinuousBlendTrait.GetBlendWeight(Context, 1);
 		if (!FAnimWeight::IsFullWeight(BlendWeightB))
@@ -80,12 +81,12 @@ namespace UE::AnimNext
 		}
 	}
 
-	void FBlendTwoWayTrait::QueueChildrenForTraversal(FUpdateTraversalContext& Context, const TTraitBinding<IUpdate>& Binding, const FTraitUpdateState& TraitState, FUpdateTraversalQueue& TraversalQueue) const
+	void FBlendTwoWayTrait::QueueChildrenForTraversal(FUpdateTraversalContext& Context, const TTraitBinding<IUpdateTraversal>& Binding, const FTraitUpdateState& TraitState, FUpdateTraversalQueue& TraversalQueue) const
 	{
 		const FInstanceData* InstanceData = Binding.GetInstanceData<FInstanceData>();
 
 		TTraitBinding<IContinuousBlend> ContinuousBlendTrait;
-		Context.GetInterface(Binding, ContinuousBlendTrait);
+		Binding.GetStackInterface(ContinuousBlendTrait);
 
 		const float BlendWeightB = ContinuousBlendTrait.GetBlendWeight(Context, 1);
 		if (InstanceData->ChildA.IsValid())
@@ -118,7 +119,7 @@ namespace UE::AnimNext
 	{
 		const FSharedData* SharedData = Binding.GetSharedData<FSharedData>();
 
-		const float BlendWeight = SharedData->GetBlendWeight(Context, Binding);
+		const float BlendWeight = SharedData->GetBlendWeight(Binding);
 		const float ClampedWeight = FMath::Clamp(BlendWeight, 0.0f, 1.0f);
 
 		if (ChildIndex == 0)
diff --git a/Engine/Plugins/Experimental/AnimNext/Source/AnimNext/Private/Traits/BlendTwoWay.h b/Engine/Plugins/Experimental/AnimNext/Source/AnimNext/Private/Traits/BlendTwoWay.h
--- a/Engine/Plugins/Experimental/AnimNext/Source/AnimNext/Private/Traits/BlendTwoWay.h
+++ b/Engine/Plugins/Experimental/AnimNext/Source/AnimNext/Private/Traits/BlendTwoWay.h
@@ -44,7 +44,7 @@ namespace UE::AnimNext
 	 * 
 	 * A trait that can blend two inputs.
 	 */
-	struct FBlendTwoWayTrait : FBaseTrait, IEvaluate, IUpdate, IHierarchy, IContinuousBlend
+	struct FBlendTwoWayTrait : FBaseTrait, IEvaluate, IUpdate, IUpdateTraversal, IHierarchy, IContinuousBlend
 	{
 		DECLARE_ANIM_TRAIT(FBlendTwoWayTrait, 0x62f40e2b, FBaseTrait)
 
@@ -61,7 +61,9 @@ namespace UE::AnimNext
 
 		// IUpdate impl
 		virtual void PreUpdate(FUpdateTraversalContext& Context, const TTraitBinding<IUpdate>& Binding, const FTraitUpdateState& TraitState) const override;
-		virtual void QueueChildrenForTraversal(FUpdateTraversalContext& Context, const TTraitBinding<IUpdate>& Binding, const FTraitUpdateState& TraitState, FUpdateTraversalQueue& TraversalQueue) const override;
+
+		// IUpdateTraversal impl
+		virtual void QueueChildrenForTraversal(FUpdateTraversalContext& Context, const TTraitBinding<IUpdateTraversal>& Binding, const FTraitUpdateState& TraitState, FUpdateTraversalQueue& TraversalQueue) const override;
 
 		// IHierarchy impl
 		virtual uint32 GetNumChildren(const FExecutionContext& Context, const TTraitBinding<IHierarchy>& Binding) const override;
diff --git a/Engine/Plugins/Experimental/AnimNext/Source/AnimNext/Private/Traits/SequencePlayer.cpp b/Engine/Plugins/Experimental/AnimNext/Source/AnimNext/Private/Traits/SequencePlayer.cpp
--- a/Engine/Plugins/Experimental/AnimNext/Source/AnimNext/Private/Traits/SequencePlayer.cpp
+++ b/Engine/Plugins/Experimental/AnimNext/Source/AnimNext/Private/Traits/SequencePlayer.cpp
@@ -25,7 +25,7 @@ namespace UE::AnimNext
 		if (SharedData->AnimSequence != nullptr)
 		{
 			const float SequenceLength = SharedData->AnimSequence->GetPlayLength();
-			InternalTimeAccumulator = FMath::Clamp(SharedData->GetStartPosition(Context, Binding), 0.0f, SequenceLength);
+			InternalTimeAccumulator = FMath::Clamp(SharedData->GetStartPosition(Binding), 0.0f, SequenceLength);
 		}
 	}
 
@@ -45,7 +45,7 @@ namespace UE::AnimNext
 	float FSequencePlayerTrait::GetPlayRate(const FExecutionContext& Context, const TTraitBinding<ITimeline>& Binding) const
 	{
 		const FSharedData* SharedData = Binding.GetSharedData<FSharedData>();
-		return SharedData->GetPlayRate(Context, Binding);
+		return SharedData->GetPlayRate(Binding);
 	}
 
 	float FSequencePlayerTrait::AdvanceBy(const FExecutionContext& Context, const TTraitBinding<ITimeline>& Binding, float DeltaTime) const
@@ -56,10 +56,10 @@ namespace UE::AnimNext
 			FInstanceData* InstanceData = Binding.GetInstanceData<FInstanceData>();
 
 			TTraitBinding<ITimeline> TimelineTrait;
-			Context.GetInterface(Binding, TimelineTrait);
+			Binding.GetStackInterface(TimelineTrait);
 
 			const float PlayRate = TimelineTrait.GetPlayRate(Context);
-			const bool bIsLooping = SharedData->GetbLoop(Context, Binding);
+			const bool bIsLooping = SharedData->GetbLoop(Binding);
 			const float SequenceLength = AnimSeq->GetPlayLength();
 
 			FAnimationRuntime::AdvanceTime(bIsLooping, DeltaTime * PlayRate, InstanceData->InternalTimeAccumulator, SequenceLength);
@@ -87,7 +87,7 @@ namespace UE::AnimNext
 	{
 		// We just advance the timeline
 		TTraitBinding<ITimeline> TimelineTrait;
-		Context.GetInterface(Binding, TimelineTrait);
+		Binding.GetStackInterface(TimelineTrait);
 
 		TimelineTrait.AdvanceBy(Context, TraitState.GetDeltaTime());
 	}
diff --git a/Engine/Plugins/Experimental/AnimNext/Source/AnimNext/Private/Traits/SubGraphHost.cpp b/Engine/Plugins/Experimental/AnimNext/Source/AnimNext/Private/Traits/SubGraphHost.cpp
--- a/Engine/Plugins/Experimental/AnimNext/Source/AnimNext/Private/Traits/SubGraphHost.cpp
+++ b/Engine/Plugins/Experimental/AnimNext/Source/AnimNext/Private/Traits/SubGraphHost.cpp
@@ -3,6 +3,7 @@
 #include "Traits/SubGraphHost.h"
 
 #include "TraitCore/ExecutionContext.h"
+#include "TraitCore/NodeInstance.h"
 #include "Graph/AnimNextGraphInstance.h"
 
 namespace UE::AnimNext
@@ -15,6 +16,7 @@ namespace UE::AnimNext
 		GeneratorMacro(IGarbageCollection) \
 		GeneratorMacro(IHierarchy) \
 		GeneratorMacro(IUpdate) \
+		GeneratorMacro(IUpdateTraversal) \
 
 	GENERATE_ANIM_TRAIT_IMPLEMENTATION(FSubGraphHostTrait, TRAIT_INTERFACE_ENUMERATOR)
 	#undef TRAIT_INTERFACE_ENUMERATOR
@@ -76,14 +78,19 @@ namespace UE::AnimNext
 			CurrentActiveEntryPoint = SubGraphSlot.EntryPoint;
 		}
 
-		const TObjectPtr<const UAnimNextGraph> DesiredSubGraph = SharedData->GetSubGraph(Context, Binding);
-		const FName EntryPoint = SharedData->GetEntryPoint(Context, Binding);
+		const TObjectPtr<const UAnimNextGraph> DesiredSubGraph = SharedData->GetSubGraph(Binding);
+		const FName EntryPoint = SharedData->GetEntryPoint(Binding);
 
-		// Check for reentrancy and early-out if we are linking back to the current instance
-		FAnimNextGraphInstance& GraphInstance = Context.GetGraphInstance();
-		if(GraphInstance.UsesGraph(DesiredSubGraph) && GraphInstance.UsesEntryPoint(EntryPoint))
+		// Check for re-entrancy and early-out if we are linking back to the current instance or one of its parents
+		const FAnimNextGraphInstance* OwnerGraphInstance = &Binding.GetTraitPtr().GetNodeInstance()->GetOwner();
+		while (OwnerGraphInstance != nullptr)
 		{
-			return;
+			if (OwnerGraphInstance->UsesGraph(DesiredSubGraph) && OwnerGraphInstance->UsesEntryPoint(EntryPoint))
+			{
+				return;
+			}
+
+			OwnerGraphInstance = OwnerGraphInstance->GetParentGraphInstance();
 		}
 
 		if (!bHasActiveSubGraph || CurrentActiveSubGraph != DesiredSubGraph || CurrentActiveEntryPoint != EntryPoint)
@@ -119,13 +126,13 @@ namespace UE::AnimNext
 			InstanceData->CurrentlyActiveSubGraphIndex = FreeSlotIndex;
 
 			TTraitBinding<IDiscreteBlend> DiscreteBlendTrait;
-			Context.GetInterface(Binding, DiscreteBlendTrait);
+			Binding.GetStackInterface(DiscreteBlendTrait);
 
 			DiscreteBlendTrait.OnBlendTransition(Context, OldChildIndex, NewChildIndex);
 		}
 	}
 
-	void FSubGraphHostTrait::QueueChildrenForTraversal(FUpdateTraversalContext& Context, const TTraitBinding<IUpdate>& Binding, const FTraitUpdateState& TraitState, FUpdateTraversalQueue& TraversalQueue) const
+	void FSubGraphHostTrait::QueueChildrenForTraversal(FUpdateTraversalContext& Context, const TTraitBinding<IUpdateTraversal>& Binding, const FTraitUpdateState& TraitState, FUpdateTraversalQueue& TraversalQueue) const
 	{
 		const FInstanceData* InstanceData = Binding.GetInstanceData<FInstanceData>();
 
@@ -136,7 +143,7 @@ namespace UE::AnimNext
 		}
 
 		TTraitBinding<IDiscreteBlend> DiscreteBlendTrait;
-		Context.GetInterface(Binding, DiscreteBlendTrait);
+		Binding.GetStackInterface(DiscreteBlendTrait);
 
 		for (int32 SubGraphIndex = 0; SubGraphIndex < NumSubGraphs; ++SubGraphIndex)
 		{
@@ -181,7 +188,7 @@ namespace UE::AnimNext
 	void FSubGraphHostTrait::OnBlendTransition(const FExecutionContext& Context, const TTraitBinding<IDiscreteBlend>& Binding, int32 OldChildIndex, int32 NewChildIndex) const
 	{
 		TTraitBinding<IDiscreteBlend> DiscreteBlendTrait;
-		Context.GetInterface(Binding, DiscreteBlendTrait);
+		Binding.GetStackInterface(DiscreteBlendTrait);
 
 		// We initiate immediately when we transition
 		DiscreteBlendTrait.OnBlendInitiated(Context, NewChildIndex);
@@ -201,7 +208,7 @@ namespace UE::AnimNext
 
 			if (SubGraphEntry.State == ESlotState::ActiveWithGraph)
 			{
-				SubGraphEntry.SubGraph->AllocateInstance(Context.GetGraphInstance(), SubGraphEntry.GraphInstance, SubGraphEntry.EntryPoint);
+				SubGraphEntry.SubGraph->AllocateInstance(Binding.GetTraitPtr().GetNodeInstance()->GetOwner(), SubGraphEntry.GraphInstance, SubGraphEntry.EntryPoint);
 			}
 		}
 	}
diff --git a/Engine/Plugins/Experimental/AnimNext/Source/AnimNext/Private/Traits/SubGraphHost.h b/Engine/Plugins/Experimental/AnimNext/Source/AnimNext/Private/Traits/SubGraphHost.h
--- a/Engine/Plugins/Experimental/AnimNext/Source/AnimNext/Private/Traits/SubGraphHost.h
+++ b/Engine/Plugins/Experimental/AnimNext/Source/AnimNext/Private/Traits/SubGraphHost.h
@@ -47,7 +47,7 @@ namespace UE::AnimNext
 	 * 
 	 * A trait that hosts and manages a sub-graph instance.
 	 */
-	struct FSubGraphHostTrait : FBaseTrait, IUpdate, IHierarchy, IDiscreteBlend, IGarbageCollection
+	struct FSubGraphHostTrait : FBaseTrait, IUpdate, IUpdateTraversal, IHierarchy, IDiscreteBlend, IGarbageCollection
 	{
 		DECLARE_ANIM_TRAIT(FSubGraphHostTrait, 0xb1085d7b, FBaseTrait)
 
@@ -93,7 +93,9 @@ namespace UE::AnimNext
 
 		// IUpdate impl
 		virtual void PreUpdate(FUpdateTraversalContext& Context, const TTraitBinding<IUpdate>& Binding, const FTraitUpdateState& TraitState) const override;
-		virtual void QueueChildrenForTraversal(FUpdateTraversalContext& Context, const TTraitBinding<IUpdate>& Binding, const FTraitUpdateState& TraitState, FUpdateTraversalQueue& TraversalQueue) const override;
+
+		// IUpdateTraversal impl
+		virtual void QueueChildrenForTraversal(FUpdateTraversalContext& Context, const TTraitBinding<IUpdateTraversal>& Binding, const FTraitUpdateState& TraitState, FUpdateTraversalQueue& TraversalQueue) const override;
 
 		// IHierarchy impl
 		virtual uint32 GetNumChildren(const FExecutionContext& Context, const TTraitBinding<IHierarchy>& Binding) const override;
diff --git a/Engine/Plugins/Experimental/AnimNext/Source/AnimNext/Private/Traits/SynchronizeUsingGroups.cpp b/Engine/Plugins/Experimental/AnimNext/Source/AnimNext/Private/Traits/SynchronizeUsingGroups.cpp
--- a/Engine/Plugins/Experimental/AnimNext/Source/AnimNext/Private/Traits/SynchronizeUsingGroups.cpp
+++ b/Engine/Plugins/Experimental/AnimNext/Source/AnimNext/Private/Traits/SynchronizeUsingGroups.cpp
@@ -21,7 +21,7 @@ namespace UE::AnimNext
 		FInstanceData* InstanceData = Binding.GetInstanceData<FInstanceData>();
 
 		TTraitBinding<IGroupSynchronization> GroupSyncTrait;
-		Context.GetInterface(Binding, GroupSyncTrait);
+		Binding.GetStackInterface(GroupSyncTrait);
 
 		const FName GroupName = GroupSyncTrait.GetGroupName(Context);
 		const bool bHasGroupName = !GroupName.IsNone();
@@ -66,7 +66,7 @@ namespace UE::AnimNext
 		InstanceData->bFreezeTimeline = false;
 
 		TTraitBinding<ITimeline> TimelineTrait;
-		Context.GetInterface(Binding, TimelineTrait);
+		Binding.GetStackInterface(TimelineTrait);
 
 		const float ProgressRatio = TimelineTrait.AdvanceBy(Context, DeltaTime);
 
@@ -83,7 +83,7 @@ namespace UE::AnimNext
 		InstanceData->bFreezeTimeline = false;
 
 		TTraitBinding<ITimeline> TimelineTrait;
-		Context.GetInterface(Binding, TimelineTrait);
+		Binding.GetStackInterface(TimelineTrait);
 
 		TimelineTrait.AdvanceToRatio(Context, ProgressRatio);
 
diff --git a/Engine/Plugins/Experimental/AnimNext/Source/AnimNext/Public/Graph/AnimNextGraphInstance.h b/Engine/Plugins/Experimental/AnimNext/Source/AnimNext/Public/Graph/AnimNextGraphInstance.h
--- a/Engine/Plugins/Experimental/AnimNext/Source/AnimNext/Public/Graph/AnimNextGraphInstance.h
+++ b/Engine/Plugins/Experimental/AnimNext/Source/AnimNext/Public/Graph/AnimNextGraphInstance.h
@@ -19,6 +19,7 @@ namespace UE::AnimNext
 	struct FExecutionContext;
 	struct FGraphInstanceComponent;
 	struct FLatentPropertyHandle;
+	struct FTraitStackBinding;
 }
 
 using GraphInstanceComponentMapType = TMap<FName, TSharedPtr<UE::AnimNext::FGraphInstanceComponent>>;
@@ -134,6 +135,7 @@ private:
 	friend FRigUnit_AnimNextGraphEvaluator;	// We evaluate the instance
 	friend UE::AnimNext::FExecutionContext;
 	friend FAnimNextGraphInstancePtr;
+	friend UE::AnimNext::FTraitStackBinding;
 };
 
 template<>
diff --git a/Engine/Plugins/Experimental/AnimNext/Source/AnimNext/Public/TraitCore/ExecutionContext.h b/Engine/Plugins/Experimental/AnimNext/Source/AnimNext/Public/TraitCore/ExecutionContext.h
--- a/Engine/Plugins/Experimental/AnimNext/Source/AnimNext/Public/TraitCore/ExecutionContext.h
+++ b/Engine/Plugins/Experimental/AnimNext/Source/AnimNext/Public/TraitCore/ExecutionContext.h
@@ -6,12 +6,18 @@
 #include "TraitCore/TraitBinding.h"
 #include "TraitCore/TraitHandle.h"
 #include "TraitCore/TraitPtr.h"
+#include "TraitCore/TraitStackBinding.h"
+#include "TraitCore/ITraitInterface.h"
 #include "TraitCore/TraitInterfaceUID.h"
 #include "TraitCore/LatentPropertyHandle.h"
 #include "TraitCore/NodeHandle.h"
 #include "Graph/AnimNextGraphInstancePtr.h"
+#include "Graph/GraphInstanceComponent.h"
+
+#include <type_traits>
 
 struct FAnimNextGraphInstance;
+class FMemStack;
 
 namespace UE::AnimNext
 {
@@ -40,6 +46,12 @@ namespace UE::AnimNext
 		// Creates an execution context and binds it to the specified graph instance
 		explicit FExecutionContext(FAnimNextGraphInstance& InGraphInstance);
 
+
+
+		//////////////////////////////////////////////////////////////////////////
+		// The following functions handle the context binding
+		// In order to be used, the execution context must be bound to a valid root graph instance
+
 		// Binds the execution context to the specified graph instance if it differs from the currently bound instance
 		void BindTo(FAnimNextGraphInstancePtr& InGraphInstance);
 
@@ -53,30 +65,28 @@ namespace UE::AnimNext
 		bool IsBound() const;
 
 		// Returns whether or not this execution context is bound to the specified graph instance
+		// Returns true if the root graphs match
 		bool IsBoundTo(const FAnimNextGraphInstancePtr& InGraphInstance) const;
 
 		// Returns whether or not this execution context is bound to the specified graph instance
+		// Returns true if the root graphs match
 		bool IsBoundTo(const FAnimNextGraphInstance& InGraphInstance) const;
 
-		// Queries a node for a trait that implements the specified interface.
-		// If no such trait exists, nullptr is returned.
-		template<class TraitInterface>
-		bool GetInterface(const FWeakTraitPtr& TraitPtr, TTraitBinding<TraitInterface>& InterfaceBinding) const;
+		//////////////////////////////////////////////////////////////////////////
+		// The following functions allow creation of trait stack bindings
 
-		// Queries a node for a trait that implements the specified interface.
-		// If no such trait exists, nullptr is returned.
-		template<class TraitInterface>
-		bool GetInterface(const FTraitBinding& Binding, TTraitBinding<TraitInterface>& InterfaceBinding) const;
+		// Returns a trait stack binding to the stack that contains the specified trait pointer.
+		// Returns false if we failed to do so.
+		bool GetStack(const FWeakTraitPtr& TraitPtr, FTraitStackBinding& OutStackBinding) const;
 
-		// Queries a node for a trait lower on the stack that implements the specified interface.
-		// If no such trait exists, nullptr is returned.
-		template<class TraitInterface>
-		bool GetInterfaceSuper(const FWeakTraitPtr& TraitPtr, TTraitBinding<TraitInterface>& SuperBinding) const;
 
-		// Queries a node for a trait lower on the stack that implements the specified interface.
-		// If no such trait exists, nullptr is returned.
-		template<class TraitInterface>
-		bool GetInterfaceSuper(const FTraitBinding& Binding, TTraitBinding<TraitInterface>& SuperBinding) const;
+
+		//////////////////////////////////////////////////////////////////////////
+		// The following functions handle node lifetime management
+
+		// Allocates a new node instance from a trait handle using the specified graph instance
+		// If the desired trait lives in the current parent, a weak handle to it will be returned
+		FTraitPtr AllocateNodeInstance(FAnimNextGraphInstance& GraphInstance, FAnimNextTraitHandle ChildTraitHandle) const;
 
 		// Allocates a new node instance from a trait handle
 		// If the desired trait lives in the current parent, a weak handle to it will be returned
@@ -90,10 +100,12 @@ namespace UE::AnimNext
 		// there are no more references remaining, reseting the pointer in the process
 		void ReleaseNodeInstance(FTraitPtr& NodePtr) const;
 
-		// Takes a snapshot of all latent properties on the provided node sub-stack (all traits on the sub-stack of the provided one)
-		// Properties can be marked as always updating or as supporting freezing (e.g. when a branch of the graph blends out)
-		// A freezable property does not update when a snapshot is taken of a frozen node
-		void SnapshotLatentProperties(const FWeakTraitPtr& TraitPtr, bool bIsFrozen) const;
+
+
+		//////////////////////////////////////////////////////////////////////////
+		// The following functions handle graph instance components
+		// Graph instance components live on the root graph instance and persist from
+		// frame to frame
 
 		// Returns a typed graph instance component, creating it lazily the first time it is queried
 		template<class ComponentType>
@@ -106,61 +118,48 @@ namespace UE::AnimNext
 		// Returns const iterators to the graph instance component container
 		GraphInstanceComponentMapType::TConstIterator GetComponentIterator() const;
 
-		// Returns the bound graph instance
-		FAnimNextGraphInstance& GetGraphInstance() const { return *GraphInstance; }
+
+
+		//////////////////////////////////////////////////////////////////////////
+		// Misc functions
+
+		// Returns the bound root graph instance
+		FAnimNextGraphInstance& GetRootGraphInstance() const { return *RootGraphInstance; }
+
+		// Returns the local thread memstack
+		FMemStack& GetMemStack() const { return MemStack; }
 
 	private:
 		// No copy or move
 		FExecutionContext(const FExecutionContext&) = delete;
 		FExecutionContext& operator=(const FExecutionContext&) = delete;
 
-		bool GetInterfaceImpl(FTraitInterfaceUID InterfaceUID, const FWeakTraitPtr& TraitPtr, FTraitBinding& InterfaceBinding) const;
-		bool GetInterfaceSuperImpl(FTraitInterfaceUID InterfaceUID, const FWeakTraitPtr& TraitPtr, FTraitBinding& SuperBinding) const;
 		FGraphInstanceComponent* TryGetComponent(int32 ComponentNameHash, FName ComponentName) const;
 		FGraphInstanceComponent& AddComponent(int32 ComponentNameHash, FName ComponentName, TSharedPtr<FGraphInstanceComponent>&& Component) const;
 
-		const FNodeDescription& GetNodeDescription(FNodeHandle NodeHandle) const;
+		const FNodeDescription& GetNodeDescription(const FAnimNextGraphInstance& GraphInstance, FNodeHandle NodeHandle) const;
+		const FNodeDescription& GetNodeDescription(const FNodeInstance& NodeInstance) const;
 		const FNodeTemplate* GetNodeTemplate(const FNodeDescription& NodeDesc) const;
 		const FTrait* GetTrait(const FTraitTemplate& TraitDesc) const;
 
+	protected:
+		// The memstack of the local thread
+		FMemStack& MemStack;
+
+	private:
 		// Cached references to the registries we need
 		const FNodeTemplateRegistry& NodeTemplateRegistry;
 		const FTraitRegistry& TraitRegistry;
 
-		// Cached properties for the currently executing graph
-		FAnimNextGraphInstance* GraphInstance = nullptr;
-		TArrayView<const uint8> GraphSharedData;
+		// Root graph instance we are bound to
+		FAnimNextGraphInstance* RootGraphInstance = nullptr;
+
+		friend struct FTraitStackBinding;
 	};
 
 	//////////////////////////////////////////////////////////////////////////
 	// Inline implementations
 
-	template<class TraitInterface>
-	inline bool FExecutionContext::GetInterface(const FWeakTraitPtr& TraitPtr, TTraitBinding<TraitInterface>& InterfaceBinding) const
-	{
-		constexpr FTraitInterfaceUID InterfaceUID = TraitInterface::InterfaceUID;
-		return GetInterfaceImpl(InterfaceUID, TraitPtr, InterfaceBinding);
-	}
-
-	template<class TraitInterface>
-	inline bool FExecutionContext::GetInterface(const FTraitBinding& Binding, TTraitBinding<TraitInterface>& InterfaceBinding) const
-	{
-		return GetInterface<TraitInterface>(Binding.GetTraitPtr(), InterfaceBinding);
-	}
-
-	template<class TraitInterface>
-	inline bool FExecutionContext::GetInterfaceSuper(const FWeakTraitPtr& TraitPtr, TTraitBinding<TraitInterface>& SuperBinding) const
-	{
-		constexpr FTraitInterfaceUID InterfaceUID = TraitInterface::InterfaceUID;
-		return GetInterfaceSuperImpl(InterfaceUID, TraitPtr, SuperBinding);
-	}
-
-	template<class TraitInterface>
-	inline bool FExecutionContext::GetInterfaceSuper(const FTraitBinding& Binding, TTraitBinding<TraitInterface>& SuperBinding) const
-	{
-		return GetInterfaceSuper<TraitInterface>(Binding.GetTraitPtr(), SuperBinding);
-	}
-
 	inline FTraitPtr FExecutionContext::AllocateNodeInstance(const FTraitBinding& ParentBinding, FAnimNextTraitHandle ChildTraitHandle) const
 	{
 		return AllocateNodeInstance(ParentBinding.GetTraitPtr(), ChildTraitHandle);
@@ -169,6 +168,8 @@ namespace UE::AnimNext
 	template<class ComponentType>
 	ComponentType& FExecutionContext::GetComponent() const
 	{
+		static_assert(std::is_base_of<FGraphInstanceComponent, ComponentType>::value, "ComponentType type must derive from FGraphInstanceComponent");
+
 		const FName ComponentName = ComponentType::StaticComponentName();
 		const int32 ComponentNameHash = GetTypeHash(ComponentName);
 
@@ -183,6 +184,8 @@ namespace UE::AnimNext
 	template<class ComponentType>
 	ComponentType* FExecutionContext::TryGetComponent() const
 	{
+		static_assert(std::is_base_of<FGraphInstanceComponent, ComponentType>::value, "ComponentType type must derive from FGraphInstanceComponent");
+
 		const FName ComponentName = ComponentType::StaticComponentName();
 		const int32 ComponentNameHash = GetTypeHash(ComponentName);
 
diff --git a/Engine/Plugins/Experimental/AnimNext/Source/AnimNext/Public/TraitCore/ExecutionContextProxy.h b/Engine/Plugins/Experimental/AnimNext/Source/AnimNext/Public/TraitCore/ExecutionContextProxy.h
deleted file mode 100644
diff --git a/Engine/Plugins/Experimental/AnimNext/Source/AnimNext/Public/TraitCore/ITraitInterface.h b/Engine/Plugins/Experimental/AnimNext/Source/AnimNext/Public/TraitCore/ITraitInterface.h
--- a/Engine/Plugins/Experimental/AnimNext/Source/AnimNext/Public/TraitCore/ITraitInterface.h
+++ b/Engine/Plugins/Experimental/AnimNext/Source/AnimNext/Public/TraitCore/ITraitInterface.h
@@ -3,7 +3,6 @@
 #pragma once
 
 #include "CoreMinimal.h"
-#include "TraitCore/TraitBinding.h"			// Derived types will need to implement the binding specialization
 #include "TraitCore/TraitInterfaceUID.h"
 
 // Helper macros
diff --git a/Engine/Plugins/Experimental/AnimNext/Source/AnimNext/Public/TraitCore/Trait.h b/Engine/Plugins/Experimental/AnimNext/Source/AnimNext/Public/TraitCore/Trait.h
--- a/Engine/Plugins/Experimental/AnimNext/Source/AnimNext/Public/TraitCore/Trait.h
+++ b/Engine/Plugins/Experimental/AnimNext/Source/AnimNext/Public/TraitCore/Trait.h
@@ -4,7 +4,6 @@
 
 #include "CoreMinimal.h"
 
-#include "TraitCore/TraitBinding.h"
 #include "TraitCore/TraitHandle.h"			// Derived types are likely to refer to other traits as children
 #include "TraitCore/TraitInstanceData.h"
 #include "TraitCore/TraitMode.h"
@@ -130,6 +129,7 @@ class FArchive;
 namespace UE::AnimNext
 {
 	struct FTrait;
+	struct FTraitBinding;
 	struct FTraitMemoryLayout;
 	class FTraitReader;
 	class FTraitWriter;
diff --git a/Engine/Plugins/Experimental/AnimNext/Source/AnimNext/Public/TraitCore/TraitBinding.h b/Engine/Plugins/Experimental/AnimNext/Source/AnimNext/Public/TraitCore/TraitBinding.h
--- a/Engine/Plugins/Experimental/AnimNext/Source/AnimNext/Public/TraitCore/TraitBinding.h
+++ b/Engine/Plugins/Experimental/AnimNext/Source/AnimNext/Public/TraitCore/TraitBinding.h
@@ -3,8 +3,11 @@
 #pragma once
 
 #include "CoreMinimal.h"
+#include "TraitCore/NodeTemplate.h"
+#include "TraitCore/Trait.h"
 #include "TraitCore/TraitPtr.h"
 #include "TraitCore/TraitInterfaceUID.h"
+#include "TraitCore/TraitStackBinding.h"
 #include "TraitCore/TraitTemplate.h"
 
 #include <type_traits>
@@ -13,6 +16,7 @@ struct FAnimNextTraitSharedData;
 
 namespace UE::AnimNext
 {
+	struct FExecutionContext;
 	struct FTraitInstanceData;
 	struct ITraitInterface;
 	struct FNodeDescription;
@@ -25,11 +29,17 @@ namespace UE::AnimNext
 	 */
 	struct ANIMNEXT_API FTraitBinding
 	{
-		// Creates an empty binding.
+		// Creates an empty/invalid binding.
 		FTraitBinding() = default;
 
 		// Returns whether or not this binding is valid.
-		bool IsValid() const { return TraitTemplate != nullptr; }
+		bool IsValid() const { return TraitImpl != nullptr; }
+
+		// Resets the trait binding to an invalid state.
+		void Reset()
+		{
+			new (this) FTraitBinding();
+		}
 
 		// Queries a node for a pointer to its trait shared data.
 		// If the trait handle is invalid, a null pointer is returned.
@@ -43,7 +53,10 @@ namespace UE::AnimNext
 				return nullptr;
 			}
 
-			return static_cast<const SharedDataType*>(TraitTemplate->GetTraitDescription(*NodeDescription));
+			const FTraitTemplate* TraitDescs = Stack->NodeTemplate->GetTraits();
+			const FTraitTemplate* TraitTemplate = TraitDescs + TraitIndex;
+
+			return static_cast<const SharedDataType*>(TraitTemplate->GetTraitDescription(*Stack->NodeDescription));
 		}
 
 		// Queries a node for a pointer to its trait instance data.
@@ -58,19 +71,30 @@ namespace UE::AnimNext
 				return nullptr;
 			}
 
-			return static_cast<InstanceDataType*>(TraitTemplate->GetTraitInstance(*TraitPtr.GetNodeInstance()));
+			const FTraitTemplate* TraitDescs = Stack->NodeTemplate->GetTraits();
+			const FTraitTemplate* TraitTemplate = TraitDescs + TraitIndex;
+
+			return static_cast<InstanceDataType*>(TraitTemplate->GetTraitInstance(*Stack->NodeInstance));
 		}
 
 		// Queries a node for a pointer to its trait latent properties.
 		// If the trait handle is invalid or if we have no latent properties, a null pointer is returned.
 		const FLatentPropertyHandle* GetLatentPropertyHandles() const
 		{
-			if (!IsValid() || !TraitTemplate->HasLatentProperties())
+			if (!IsValid())
+			{
+				return nullptr;
+			}
+
+			const FTraitTemplate* TraitDescs = Stack->NodeTemplate->GetTraits();
+			const FTraitTemplate* TraitTemplate = TraitDescs + TraitIndex;
+
+			if (!TraitTemplate->HasLatentProperties())
 			{
 				return nullptr;
 			}
 
-			return TraitTemplate->GetTraitLatentPropertyHandles(*NodeDescription);
+			return TraitTemplate->GetTraitLatentPropertyHandles(*Stack->NodeDescription);
 		}
 
 		// Returns a pointer to the latent property specified by the provided handle or nullptr if the binding/handle are invalid
@@ -87,50 +111,105 @@ namespace UE::AnimNext
 				return nullptr;
 			}
 
-			const uint8* NodeInstance = (const uint8*)TraitPtr.GetNodeInstance();
+			const uint8* NodeInstance = (const uint8*)Stack->NodeInstance;
 			return (const PropertyType*)(NodeInstance + Handle.GetLatentPropertyOffset());
 		}
 
+		// Queries the trait stack for a trait that implements the specified interface.
+		// If no such trait exists, false is returned.
+		template<class TraitInterface>
+		bool GetStackInterface(TTraitBinding<TraitInterface>& OutBinding) const
+		{
+			return IsValid() && Stack->GetInterface<TraitInterface>(OutBinding);
+		}
+
+		// Queries the trait stack for a trait lower on the stack that implements the specified interface.
+		// If no such trait exists, false is returned.
+		template<class TraitInterface>
+		bool GetStackInterfaceSuper(TTraitBinding<TraitInterface>& OutSuperBinding) const
+		{
+			return IsValid() && Stack->GetInterfaceSuper<TraitInterface>(*this, OutSuperBinding);
+		}
+
+		// Queries the current trait for a new interface.
+		// If no such interface is found, false it returned.
+		template<class TraitInterface>
+		bool AsInterface(TTraitBinding<TraitInterface>& OutBinding) const
+		{
+			static_assert(std::is_base_of<ITraitInterface, TraitInterface>::value, "TraitInterface type must derive from ITraitInterface");
+
+			constexpr FTraitInterfaceUID InterfaceUID = TraitInterface::InterfaceUID;
+			return AsInterfaceImpl(InterfaceUID, OutBinding);
+		}
+
+		// Returns the trait stack binding backing this trait binding.
+		const FTraitStackBinding* GetStack() const { return Stack; }
+
 		// Returns the trait pointer we are bound to.
-		FWeakTraitPtr GetTraitPtr() const { return TraitPtr; }
+		FWeakTraitPtr GetTraitPtr() const { return FWeakTraitPtr(Stack != nullptr ? Stack->NodeInstance : nullptr, TraitIndex); }
+
+		// Returns the trait index on the stack we are bound to.
+		uint32 GetTraitIndex() const { return Stack != nullptr ? (TraitIndex - Stack->BaseTraitIndex) : 0; }
 
 		// Returns the trait interface UID when bound, an invalid UID otherwise.
 		FTraitInterfaceUID GetInterfaceUID() const;
 
 		// Equality and inequality tests
-		bool operator==(const FTraitBinding& RHS) const { return TraitPtr == RHS.TraitPtr && Interface == RHS.Interface; }
-		bool operator!=(const FTraitBinding& RHS) const { return TraitPtr != RHS.TraitPtr || Interface != RHS.Interface; }
+		bool operator==(const FTraitBinding& RHS) const { return Stack == RHS.Stack && TraitIndex == RHS.TraitIndex && InterfaceThisOffset == RHS.InterfaceThisOffset; }
+		bool operator!=(const FTraitBinding& RHS) const { return !operator==(RHS); }
 
 	protected:
-		// Creates a valid binding
-		FTraitBinding(const ITraitInterface* InInterface, const FTraitTemplate* InTraitTemplate, const FNodeDescription* InNodeDescription, FWeakTraitPtr InTraitPtr)
-			: Interface(InInterface)
-			, TraitTemplate(InTraitTemplate)
-			, NodeDescription(InNodeDescription)
-			, TraitPtr(InTraitPtr)
-		{}
+		FTraitBinding(const FTraitStackBinding* InStack, const FTrait* InTraitImpl, uint32 InTraitIndex, int32 InInterfaceThisOffset = -1)
+			: Stack(InStack)
+			, TraitImpl(InTraitImpl)
+			, TraitIndex(InTraitIndex)
+			, InterfaceThisOffset(InInterfaceThisOffset)
+		{
+		}
 
 		// Performs a naked cast to the desired interface type
 		template<class TraitInterfaceType>
 		const TraitInterfaceType* GetInterfaceTyped() const
 		{
 			static_assert(std::is_base_of<ITraitInterface, TraitInterfaceType>::value, "Trait interface data must derive from ITraitInterface");
-			return static_cast<const TraitInterfaceType*>(Interface);
+			check(InterfaceThisOffset != -1);
+			return reinterpret_cast<const TraitInterfaceType*>(reinterpret_cast<const uint8*>(TraitImpl) + InterfaceThisOffset);
+		}
+
+		bool AsInterfaceImpl(FTraitInterfaceUID InterfaceUID, FTraitBinding& OutBinding) const
+		{
+			if (!IsValid())
+			{
+				return false;
+			}
+
+			if (const ITraitInterface* NewInterface = TraitImpl->GetTraitInterface(InterfaceUID))
+			{
+				const int32 NewInterfaceThisOffset = reinterpret_cast<const uint8*>(NewInterface) - reinterpret_cast<const uint8*>(TraitImpl);
+
+				OutBinding = FTraitBinding(Stack, TraitImpl, TraitIndex, NewInterfaceThisOffset);
+				return true;
+			}
+
+			return false;
 		}
 
-		// A pointer to the bound interface or nullptr if we are bound to a trait but none of its interfaces
-		const ITraitInterface*				Interface = nullptr;
+		// A pointer to the bound trait stack or nullptr if we are not bound
+		const FTraitStackBinding* Stack = nullptr;
 
-		// A pointer to the trait template that implements the interface we are bound to or nullptr if we are invalid
-		const FTraitTemplate*				TraitTemplate = nullptr;
+		// A pointer to the trait implementation from which interfaces are found
+		// Can be null if the trait is present at runtime but the trait implementation hasn't been registered/found
+		const FTrait* TraitImpl = nullptr;
 
-		// A pointer to the node shared data we are bound to
-		const FNodeDescription*				NodeDescription = nullptr;
+		// The index of the trait on the bound stack relative to the base of the stack (not relative to the base of the node)
+		uint32 TraitIndex = 0;
 
-		// A weak handle to the trait instance data we are bound to
-		FWeakTraitPtr						TraitPtr;
+		// The offset of the typed interface we are bound to or -1 if we are untyped
+		// The offset is relative to the trait implementation pointer
+		int32 InterfaceThisOffset = -1;
 
-		friend struct FExecutionContext;
+		friend FExecutionContext;
+		friend FTraitStackBinding;
 	};
 
 	/**
diff --git a/Engine/Plugins/Experimental/AnimNext/Source/AnimNext/Public/TraitCore/TraitSharedData.h b/Engine/Plugins/Experimental/AnimNext/Source/AnimNext/Public/TraitCore/TraitSharedData.h
--- a/Engine/Plugins/Experimental/AnimNext/Source/AnimNext/Public/TraitCore/TraitSharedData.h
+++ b/Engine/Plugins/Experimental/AnimNext/Source/AnimNext/Public/TraitCore/TraitSharedData.h
@@ -3,8 +3,6 @@
 #pragma once
 
 #include "CoreMinimal.h"
-#include "TraitCore/TraitBinding.h"
-#include "TraitCore/ExecutionContext.h"
 #include "TraitCore/LatentPropertyHandle.h"
 
 #include <type_traits>
@@ -37,7 +35,7 @@
 	} \
 
 #define ANIM_NEXT_IMPL_DEFINE_LATENT_GETTER(PropertyName) \
-	decltype(PropertyName) Get##PropertyName(const UE::AnimNext::FExecutionContext& Context, const UE::AnimNext::FTraitBinding& Binding) const \
+	decltype(PropertyName) Get##PropertyName(const UE::AnimNext::FTraitBinding& Binding) const \
 	{ \
 		/* We need a mapping of latent property name/offset to latent property index */ \
 		/* This can be built once at runtime using the UE reflection and cached on first call or using a constexpr function, see below */ \
diff --git a/Engine/Plugins/Experimental/AnimNext/Source/AnimNext/Public/TraitCore/TraitStackBinding.h b/Engine/Plugins/Experimental/AnimNext/Source/AnimNext/Public/TraitCore/TraitStackBinding.h
new file mode 100644
--- /dev/null
+++ b/Engine/Plugins/Experimental/AnimNext/Source/AnimNext/Public/TraitCore/TraitStackBinding.h
@@ -0,0 +1,153 @@
+// Copyright Epic Games, Inc. All Rights Reserved.
+
+#pragma once
+
+#include "CoreMinimal.h"
+#include "TraitCore/ITraitInterface.h"
+#include "TraitCore/TraitInterfaceUID.h"
+#include "TraitCore/TraitPtr.h"
+
+#include <type_traits>
+
+namespace UE::AnimNext
+{
+	struct FExecutionContext;
+	struct FNodeDescription;
+	struct FNodeInstance;
+	struct FNodeTemplate;
+	struct FTraitBinding;
+	template<class TraitInterface> struct TTraitBinding;
+	struct FTraitStackBinding;
+
+	/**
+	 * FTraitStackBinding
+	 *
+	 * To keep the node instance cost as low as possible, data such as pointers to the shared data
+	 * are not stored per node and are instead fetched on demand. In order to be able to query for
+	 * interfaces on a trait stack from a trait pointer, a trait stack binding must first be created.
+	 */
+	struct ANIMNEXT_API FTraitStackBinding
+	{
+		// Creates an empty/invalid binding.
+		FTraitStackBinding() = default;
+
+		// Returns whether or not this binding is valid.
+		bool IsValid() const { return Context != nullptr; }
+
+
+
+		//////////////////////////////////////////////////////////////////////////
+		// Misc functions
+
+		// Takes a snapshot of all latent properties on this trait stack.
+		// Properties can be marked as always updating or as supporting freezing (e.g. when a branch of the graph blends out).
+		// A freezable property does not update when a snapshot is taken of a frozen stack.
+		void SnapshotLatentProperties(bool bIsFrozen);
+
+		// Returns a trait pointer to the base of the trait stack.
+		FWeakTraitPtr GetBaseTraitPtr() const { return FWeakTraitPtr(NodeInstance, BaseTraitIndex); }
+
+		// Resets the stack binding to an invalid state.
+		void Reset();
+
+		// Equality testing
+		bool operator==(const FTraitStackBinding& Other) const;
+		bool operator!=(const FTraitStackBinding& Other) const { return !operator==(Other); }
+
+
+
+		//////////////////////////////////////////////////////////////////////////
+		// The following functions allow traversing the trait stack
+
+		// Returns a trait binding to the top of the stack.
+		// Returns true on success, false otherwise.
+		bool GetTopTrait(FTraitBinding& OutBinding) const;
+
+		// Returns a trait binding to the trait below the specified one (its parent).
+		// Returns true on success, false otherwise.
+		bool GetParentTrait(const FTraitBinding& ChildBinding, FTraitBinding& OutParentBinding) const;
+
+		// Returns a trait binding to the base of the stack.
+		// Returns true on success, false otherwise.
+		bool GetBaseTrait(FTraitBinding& OutBinding) const;
+
+		// Returns a trait binding to the trait above the specified one (its child).
+		// Returns true on success, false otherwise.
+		bool GetChildTrait(const FTraitBinding& ParentBinding, FTraitBinding& OutChildBinding) const;
+
+		// Returns a trait binding to the trait at the specified index (relative to the start of the stack).
+		// Returns true on success, false otherwise.
+		bool GetTrait(uint32 TraitIndex, FTraitBinding& OutBinding) const;
+
+		// Returns the number of traits on this stack.
+		uint32 GetNumTraits() const;
+
+
+
+		//////////////////////////////////////////////////////////////////////////
+		// The following functions allow querying interfaces on the trait stack
+
+		// Queries the trait stack for a trait that implements the specified interface.
+		// The search begins at the top of the stack.
+		// Returns true on success, false otherwise.
+		template<class TraitInterface>
+		bool GetInterface(TTraitBinding<TraitInterface>& OutBinding) const;
+
+		// Queries the trait stack for a trait lower on the stack that implements the specified interface.
+		// Returns true on success, false otherwise.
+		template<class TraitInterface>
+		bool GetInterfaceSuper(const FTraitBinding& Binding, TTraitBinding<TraitInterface>& OutSuperBinding) const;
+
+
+
+	private:
+		// Creates a trait stack binding with the stack that owns the specified trait pointer.
+		FTraitStackBinding(const FExecutionContext& InContext, const FWeakTraitPtr& TraitPtr);
+
+		bool GetInterfaceImpl(FTraitInterfaceUID InterfaceUID, FTraitBinding& OutBinding) const;
+		bool GetInterfaceSuperImpl(FTraitInterfaceUID InterfaceUID, const FTraitBinding& Binding, FTraitBinding& OutSuperBinding) const;
+
+		// A pointer to the execution context that created the binding
+		const FExecutionContext* Context = nullptr;
+
+		// A pointer to the node instance data we are bound to
+		FNodeInstance* NodeInstance = nullptr;
+
+		// A pointer to the node shared data we are bound to
+		const FNodeDescription* NodeDescription = nullptr;
+
+		// A pointer to the node template used by the node we are bound to
+		const FNodeTemplate* NodeTemplate = nullptr;
+
+		// The base trait index of the bound trait stack on the node
+		// A node can contain multiple independent trait stacks by having multiple base traits
+		uint32 BaseTraitIndex = 0;	// Only need 8 bits, using 32 since we have padding anyway
+
+		// The top trait index of the bound trait stack on the node
+		uint32 TopTraitIndex = 0;	// Only need 8 bits, using 32 since we have padding anyway
+
+		friend FExecutionContext;
+		friend FTraitBinding;
+	};
+
+	//////////////////////////////////////////////////////////////////////////
+	// Inline implementations
+
+	template<class TraitInterface>
+	inline bool FTraitStackBinding::GetInterface(TTraitBinding<TraitInterface>& OutBinding) const
+	{
+		static_assert(std::is_base_of<ITraitInterface, TraitInterface>::value, "TraitInterface type must derive from ITraitInterface");
+
+		constexpr FTraitInterfaceUID InterfaceUID = TraitInterface::InterfaceUID;
+		return GetInterfaceImpl(InterfaceUID, OutBinding);
+	}
+
+	template<class TraitInterface>
+	inline bool FTraitStackBinding::GetInterfaceSuper(const FTraitBinding& Binding, TTraitBinding<TraitInterface>& OutSuperBinding) const
+	{
+		static_assert(std::is_base_of<ITraitInterface, TraitInterface>::value, "TraitInterface type must derive from ITraitInterface");
+
+		constexpr FTraitInterfaceUID InterfaceUID = TraitInterface::InterfaceUID;
+		return GetInterfaceSuperImpl(InterfaceUID, Binding, OutSuperBinding);
+	}
+}
diff --git a/Engine/Plugins/Experimental/AnimNext/Source/AnimNext/Public/TraitInterfaces/IContinuousBlend.h b/Engine/Plugins/Experimental/AnimNext/Source/AnimNext/Public/TraitInterfaces/IContinuousBlend.h
--- a/Engine/Plugins/Experimental/AnimNext/Source/AnimNext/Public/TraitInterfaces/IContinuousBlend.h
+++ b/Engine/Plugins/Experimental/AnimNext/Source/AnimNext/Public/TraitInterfaces/IContinuousBlend.h
@@ -4,6 +4,7 @@
 
 #include "CoreMinimal.h"
 #include "TraitCore/ITraitInterface.h"
+#include "TraitCore/TraitBinding.h"
 
 namespace UE::AnimNext
 {
diff --git a/Engine/Plugins/Experimental/AnimNext/Source/AnimNext/Public/TraitInterfaces/IDiscreteBlend.h b/Engine/Plugins/Experimental/AnimNext/Source/AnimNext/Public/TraitInterfaces/IDiscreteBlend.h
--- a/Engine/Plugins/Experimental/AnimNext/Source/AnimNext/Public/TraitInterfaces/IDiscreteBlend.h
+++ b/Engine/Plugins/Experimental/AnimNext/Source/AnimNext/Public/TraitInterfaces/IDiscreteBlend.h
@@ -4,6 +4,7 @@
 
 #include "CoreMinimal.h"
 #include "TraitCore/ITraitInterface.h"
+#include "TraitCore/TraitBinding.h"
 
 struct FAlphaBlend;
 
diff --git a/Engine/Plugins/Experimental/AnimNext/Source/AnimNext/Public/TraitInterfaces/IEvaluate.h b/Engine/Plugins/Experimental/AnimNext/Source/AnimNext/Public/TraitInterfaces/IEvaluate.h
--- a/Engine/Plugins/Experimental/AnimNext/Source/AnimNext/Public/TraitInterfaces/IEvaluate.h
+++ b/Engine/Plugins/Experimental/AnimNext/Source/AnimNext/Public/TraitInterfaces/IEvaluate.h
@@ -3,8 +3,9 @@
 #pragma once
 
 #include "CoreMinimal.h"
-#include "TraitCore/ExecutionContextProxy.h"
+#include "TraitCore/ExecutionContext.h"
 #include "TraitCore/ITraitInterface.h"
+#include "TraitCore/TraitBinding.h"
 #include "EvaluationVM/EvaluationProgram.h"
 #include "EvaluationVM/KeyframeState.h"
 
@@ -15,7 +16,7 @@ namespace UE::AnimNext
 	 *
 	 * Contains all relevant transient data for an evaluate traversal and wraps the execution context.
 	 */
-	struct FEvaluateTraversalContext final : FExecutionContextProxy
+	struct FEvaluateTraversalContext final : FExecutionContext
 	{
 		// Appends a new task into the evaluation program, tasks mutate state in the order they have been appended in
 		// This means that child nodes need to evaluate first, tasks will usually be appended in IEvaluate::PostEvaluate
@@ -25,7 +26,7 @@ namespace UE::AnimNext
 		void AppendTask(TaskType&& Task) { EvaluationProgram.AppendTask(MoveTemp(Task)); }
 
 	private:
-		FEvaluateTraversalContext(const FExecutionContext& InExecutionContext, FEvaluationProgram& InEvaluationProgram);
+		explicit FEvaluateTraversalContext(FEvaluationProgram& InEvaluationProgram);
 
 		FEvaluationProgram& EvaluationProgram;
 
diff --git a/Engine/Plugins/Experimental/AnimNext/Source/AnimNext/Public/TraitInterfaces/IGarbageCollection.h b/Engine/Plugins/Experimental/AnimNext/Source/AnimNext/Public/TraitInterfaces/IGarbageCollection.h
--- a/Engine/Plugins/Experimental/AnimNext/Source/AnimNext/Public/TraitInterfaces/IGarbageCollection.h
+++ b/Engine/Plugins/Experimental/AnimNext/Source/AnimNext/Public/TraitInterfaces/IGarbageCollection.h
@@ -4,6 +4,7 @@
 
 #include "CoreMinimal.h"
 #include "TraitCore/ITraitInterface.h"
+#include "TraitCore/TraitBinding.h"
 
 class FReferenceCollector;
 
diff --git a/Engine/Plugins/Experimental/AnimNext/Source/AnimNext/Public/TraitInterfaces/IGroupSynchronization.h b/Engine/Plugins/Experimental/AnimNext/Source/AnimNext/Public/TraitInterfaces/IGroupSynchronization.h
--- a/Engine/Plugins/Experimental/AnimNext/Source/AnimNext/Public/TraitInterfaces/IGroupSynchronization.h
+++ b/Engine/Plugins/Experimental/AnimNext/Source/AnimNext/Public/TraitInterfaces/IGroupSynchronization.h
@@ -5,6 +5,7 @@
 #include "CoreMinimal.h"
 #include "Animation/AnimationAsset.h"
 #include "TraitCore/ITraitInterface.h"
+#include "TraitCore/TraitBinding.h"
 
 namespace UE::AnimNext
 {
diff --git a/Engine/Plugins/Experimental/AnimNext/Source/AnimNext/Public/TraitInterfaces/IHierarchy.h b/Engine/Plugins/Experimental/AnimNext/Source/AnimNext/Public/TraitInterfaces/IHierarchy.h
--- a/Engine/Plugins/Experimental/AnimNext/Source/AnimNext/Public/TraitInterfaces/IHierarchy.h
+++ b/Engine/Plugins/Experimental/AnimNext/Source/AnimNext/Public/TraitInterfaces/IHierarchy.h
@@ -5,6 +5,7 @@
 #include "CoreMinimal.h"
 #include "Misc/MemStack.h"
 #include "TraitCore/ITraitInterface.h"
+#include "TraitCore/TraitBinding.h"
 
 namespace UE::AnimNext
 {
@@ -21,14 +22,26 @@ namespace UE::AnimNext
 	{
 		DECLARE_ANIM_TRAIT_INTERFACE(IHierarchy, 0x846d8a37)
 
-		// Returns the number of children
+		// Returns the number of children of the trait implementation (not the whole stack)
 		// Includes inactive children
 		virtual uint32 GetNumChildren(const FExecutionContext& Context, const TTraitBinding<IHierarchy>& Binding) const;
 
-		// Appends weak handles to any children we wish to traverse.
+		// Appends weak handles to any children we wish to traverse on the trait implementation (not the whole stack).
 		// Traits are responsible for allocating and releasing child instance data.
-		// Empty handles and duplicates can be appended.
+		// Empty handles can be appended.
 		virtual void GetChildren(const FExecutionContext& Context, const TTraitBinding<IHierarchy>& Binding, FChildrenArray& Children) const;
+
+		// Queries the trait stack and calls GetChildren for each trait, appending the result.
+		static void GetStackChildren(const FExecutionContext& Context, const FTraitStackBinding& Binding, FChildrenArray& Children);
+
+		// Queries the trait stack of the specified binding and calls GetChildren for each trait, appending the result.
+		static void GetStackChildren(const FExecutionContext& Context, const FTraitBinding& Binding, FChildrenArray& Children);
+
+		// Queries the trait stack and calls GetNumChildren for each trait, accumulating the result.
+		static uint32 GetNumStackChildren(const FExecutionContext& Context, const FTraitStackBinding& Binding);
+
+		// Queries the trait stack of the specified binding and calls GetNumChildren for each trait, accumulating the result.
+		static uint32 GetNumStackChildren(const FExecutionContext& Context, const FTraitBinding& Binding);
 	};
 
 	/**
@@ -52,4 +65,20 @@ namespace UE::AnimNext
 	protected:
 		const IHierarchy* GetInterface() const { return GetInterfaceTyped<IHierarchy>(); }
 	};
+
+	//////////////////////////////////////////////////////////////////////////
+	// Inline implementations
+
+	inline void IHierarchy::GetStackChildren(const FExecutionContext& Context, const FTraitBinding& Binding, FChildrenArray& Children)
+	{
+		if (Binding.IsValid())
+		{
+			GetStackChildren(Context, *Binding.GetStack(), Children);
+		}
+	}
+
+	inline uint32 IHierarchy::GetNumStackChildren(const FExecutionContext& Context, const FTraitBinding& Binding)
+	{
+		return Binding.IsValid() ? GetNumStackChildren(Context, *Binding.GetStack()) : 0;
+	}
 }
diff --git a/Engine/Plugins/Experimental/AnimNext/Source/AnimNext/Public/TraitInterfaces/ISmoothBlend.h b/Engine/Plugins/Experimental/AnimNext/Source/AnimNext/Public/TraitInterfaces/ISmoothBlend.h
--- a/Engine/Plugins/Experimental/AnimNext/Source/AnimNext/Public/TraitInterfaces/ISmoothBlend.h
+++ b/Engine/Plugins/Experimental/AnimNext/Source/AnimNext/Public/TraitInterfaces/ISmoothBlend.h
@@ -4,6 +4,7 @@
 
 #include "CoreMinimal.h"
 #include "TraitCore/ITraitInterface.h"
+#include "TraitCore/TraitBinding.h"
 
 namespace UE::AnimNext
 {
diff --git a/Engine/Plugins/Experimental/AnimNext/Source/AnimNext/Public/TraitInterfaces/ITimeline.h b/Engine/Plugins/Experimental/AnimNext/Source/AnimNext/Public/TraitInterfaces/ITimeline.h
--- a/Engine/Plugins/Experimental/AnimNext/Source/AnimNext/Public/TraitInterfaces/ITimeline.h
+++ b/Engine/Plugins/Experimental/AnimNext/Source/AnimNext/Public/TraitInterfaces/ITimeline.h
@@ -4,6 +4,7 @@
 
 #include "CoreMinimal.h"
 #include "TraitCore/ITraitInterface.h"
+#include "TraitCore/TraitBinding.h"
 
 namespace UE::AnimNext
 {
diff --git a/Engine/Plugins/Experimental/AnimNext/Source/AnimNext/Public/TraitInterfaces/IUpdate.h b/Engine/Plugins/Experimental/AnimNext/Source/AnimNext/Public/TraitInterfaces/IUpdate.h
--- a/Engine/Plugins/Experimental/AnimNext/Source/AnimNext/Public/TraitInterfaces/IUpdate.h
+++ b/Engine/Plugins/Experimental/AnimNext/Source/AnimNext/Public/TraitInterfaces/IUpdate.h
@@ -3,17 +3,24 @@
 #pragma once
 
 #include "CoreMinimal.h"
-#include "TraitCore/ExecutionContextProxy.h"
+#include "TraitCore/ExecutionContext.h"
 #include "TraitCore/ITraitInterface.h"
+#include "TraitCore/TraitBinding.h"
+
+#include <type_traits>
 
 struct FAnimNextGraphInstancePtr;
 class FMemStack;
 
 namespace UE::AnimNext
 {
-	struct FUpdateEntry;
 	struct FUpdateTraversalQueue;
 
+	namespace Private
+	{
+		struct FUpdateEntry;
+	}
+
 	/**
 	  * FTraitUpdateState
 	  * 
@@ -87,45 +94,44 @@ namespace UE::AnimNext
 		int32 bIsBlendingOut = false;
 	};
 
+	static_assert(sizeof(FTraitUpdateState) <= 16, "Keep the size to 16 bytes for efficient copying");
+
 	/**
 	 * FUpdateTraversalContext
 	 *
 	 * Contains all relevant transient data for an update traversal and wraps the execution context.
 	 */
-	struct FUpdateTraversalContext final : FExecutionContextProxy
+	struct FUpdateTraversalContext final : FExecutionContext
 	{
 	private:
-		// Constructs a new traversal context and wraps the specified execution context
-		explicit FUpdateTraversalContext(const FExecutionContext& InExecutionContext, FMemStack& InMemStack);
+		// Constructs a new traversal context
+		FUpdateTraversalContext() = default;
 
 		// Pops entries from the traversal queue and pushes them onto the update stack
 		void PushQueuedUpdateEntries(FUpdateTraversalQueue& TraversalQueue);
 
 		// Pushes an entry onto the update stack
-		void PushUpdateEntry(FUpdateEntry* Entry);
+		void PushUpdateEntry(Private::FUpdateEntry* Entry);
 
 		// Pops an entry from the top of the update stack
-		FUpdateEntry* PopUpdateEntry();
+		Private::FUpdateEntry* PopUpdateEntry();
 
 		// Pushes an entry onto the free entry stack
-		void PushFreeEntry(FUpdateEntry* Entry);
+		void PushFreeEntry(Private::FUpdateEntry* Entry);
 
 		// Returns a new entry suitable for update queuing
 		// If an entry isn't found in the free stack, a new one is allocated from the memstack
-		FUpdateEntry* GetNewEntry(const FWeakTraitPtr& TraitPtr, const FTraitUpdateState& TraitState);
-
-		// The memstack of the local thread
-		FMemStack& MemStack;
+		Private::FUpdateEntry* GetNewEntry(const FWeakTraitPtr& TraitPtr, const FTraitUpdateState& TraitState);
 
 		// The head pointer of the update stack
 		// This is the traversal execution stack and it contains entries that are
 		// pending their pre-update call and entries waiting for post-update to be called
-		FUpdateEntry* UpdateStackHead = nullptr;
+		Private::FUpdateEntry* UpdateStackHead = nullptr;
 
 		// The head pointer of the free entry stack
 		// Entries are allocated from the memstack and are re-used in LIFO since they'll
 		// be warmer in the CPU cache
-		FUpdateEntry* FreeEntryStackHead = nullptr;
+		Private::FUpdateEntry* FreeEntryStackHead = nullptr;
 
 		friend ANIMNEXT_API void UpdateGraph(FAnimNextGraphInstancePtr& GraphInstance, float DeltaTime);
 		friend FUpdateTraversalQueue;
@@ -135,7 +141,7 @@ namespace UE::AnimNext
 	 * FUpdateTraversalQueue
 	 *
 	 * A queue of children to traverse.
-	 * @see IUpdate::QueueChildrenForTraversal
+	 * @see IUpdateTraversal::QueueChildrenForTraversal
 	 */
 	struct FUpdateTraversalQueue final
 	{
@@ -153,7 +159,7 @@ namespace UE::AnimNext
 		// When a child is queued for traversal, it is first appended to this stack
 		// Once all children have been queued and pre-update terminates, this stack is
 		// emptied and pushed onto the update stack
-		FUpdateEntry* QueuedUpdateStackHead = nullptr;
+		Private::FUpdateEntry* QueuedUpdateStackHead = nullptr;
 
 		friend ANIMNEXT_API void UpdateGraph(FAnimNextGraphInstancePtr& GraphInstance, float DeltaTime);
 		friend FUpdateTraversalContext;
@@ -168,11 +174,6 @@ namespace UE::AnimNext
 	 * the call to the next trait that implements this interface on the trait stack of the node. Once
 	 * all traits have had the chance to PreUpdate, the children will then evaluate and PostUpdate will
 	 * then be called afterwards on the original trait.
-	 * 
-	 * If this interface is implemented on a node, it is responsible for queuing the children for traversal
-	 * using the QueueChildrenForTraversal function. If the interface isn't implemented, children
-	 * are discovered through the IHierarchy interface and traversal propagates to them using the current
-	 * trait state.
 	 */
 	struct ANIMNEXT_API IUpdate : ITraitInterface
 	{
@@ -181,13 +182,30 @@ namespace UE::AnimNext
 		// Called before a traits children are updated
 		virtual void PreUpdate(FUpdateTraversalContext& Context, const TTraitBinding<IUpdate>& Binding, const FTraitUpdateState& TraitState) const;
 
-		// Called after PreUpdate to request that children be queued with the provided context
-		virtual void QueueChildrenForTraversal(FUpdateTraversalContext& Context, const TTraitBinding<IUpdate>& Binding, const FTraitUpdateState& TraitState, FUpdateTraversalQueue& TraversalQueue) const;
-
 		// Called after a traits children have been updated
 		virtual void PostUpdate(FUpdateTraversalContext& Context, const TTraitBinding<IUpdate>& Binding, const FTraitUpdateState& TraitState) const;
 	};
 
+	/**
+	 * IUpdateTraversal
+	 *
+	 * This interface is called during the update traversal.
+	 *
+	 * If a trait needs to modify the update state of its children, it needs to implement this interface.
+	 * There is no need to call the Super to forward the call down the trait stack as the update traversal
+	 * takes care of it.
+	 * 
+	 * If this interface is not implemented, the IHierarchy interface will be used to retrieve and queue
+	 * the children with the same update state as the owning trait.
+	 */
+	struct ANIMNEXT_API IUpdateTraversal : ITraitInterface
+	{
+		DECLARE_ANIM_TRAIT_INTERFACE(IUpdateTraversal, 0x256c21b1)
+
+		// Called after PreUpdate to request that children be queued with the provided context
+		virtual void QueueChildrenForTraversal(FUpdateTraversalContext& Context, const TTraitBinding<IUpdateTraversal>& Binding, const FTraitUpdateState& TraitState, FUpdateTraversalQueue& TraversalQueue) const;
+	};
+
 	/**
 	 * Specialization for trait binding.
 	 */
@@ -200,12 +218,6 @@ namespace UE::AnimNext
 			GetInterface()->PreUpdate(Context, *this, TraitState);
 		}
 
-		// @see IUpdate::QueueChildrenForTraversal
-		void QueueChildrenForTraversal(FUpdateTraversalContext& Context, const FTraitUpdateState& TraitState, FUpdateTraversalQueue& TraversalQueue) const
-		{
-			GetInterface()->QueueChildrenForTraversal(Context, *this, TraitState, TraversalQueue);
-		}
-
 		// @see IUpdate::PostUpdate
 		void PostUpdate(FUpdateTraversalContext& Context, const FTraitUpdateState& TraitState) const
 		{
@@ -216,6 +228,22 @@ namespace UE::AnimNext
 		const IUpdate* GetInterface() const { return GetInterfaceTyped<IUpdate>(); }
 	};
 
+	/**
+	 * Specialization for trait binding.
+	 */
+	template<>
+	struct TTraitBinding<IUpdateTraversal> : FTraitBinding
+	{
+		// @see IUpdateTraversal::QueueChildrenForTraversal
+		void QueueChildrenForTraversal(FUpdateTraversalContext& Context, const FTraitUpdateState& TraitState, FUpdateTraversalQueue& TraversalQueue) const
+		{
+			GetInterface()->QueueChildrenForTraversal(Context, *this, TraitState, TraversalQueue);
+		}
+
+	protected:
+		const IUpdateTraversal* GetInterface() const { return GetInterfaceTyped<IUpdateTraversal>(); }
+	};
+
 	/**
 	 * Updates a sub-graph starting at its root.
 	 *
diff --git a/Engine/Plugins/Experimental/AnimNext/Source/AnimNextTestSuite/Private/AnimNextRuntimeTest.cpp b/Engine/Plugins/Experimental/AnimNext/Source/AnimNextTestSuite/Private/AnimNextRuntimeTest.cpp
--- a/Engine/Plugins/Experimental/AnimNext/Source/AnimNextTestSuite/Private/AnimNextRuntimeTest.cpp
+++ b/Engine/Plugins/Experimental/AnimNext/Source/AnimNextTestSuite/Private/AnimNextRuntimeTest.cpp
@@ -31,6 +31,9 @@ namespace UE::AnimNext
 		ExecuteDefinition.Hash = 0;
 		ExecuteDefinition.MethodName = TEXT("Execute_0");
 
+		// Manually add our entry point since we didn't go through a full RigVM graph
+		Graph.EntryPoints.Reset();
+
 		FAnimNextGraphEntryPoint& EntryPoint = Graph.EntryPoints.AddDefaulted_GetRef();
 		EntryPoint.EntryPointName = FRigUnit_AnimNextGraphRoot::DefaultEntryPoint;
 		EntryPoint.RootTraitHandle = FAnimNextEntryPointHandle(NodeHandles[0]);
diff --git a/Engine/Plugins/Experimental/AnimNext/Source/AnimNextTestSuite/Private/AnimNextTraitBaseTest.cpp b/Engine/Plugins/Experimental/AnimNext/Source/AnimNextTestSuite/Private/AnimNextTraitBaseTest.cpp
--- a/Engine/Plugins/Experimental/AnimNext/Source/AnimNextTestSuite/Private/AnimNextTraitBaseTest.cpp
+++ b/Engine/Plugins/Experimental/AnimNext/Source/AnimNextTestSuite/Private/AnimNextTraitBaseTest.cpp
@@ -56,7 +56,7 @@ namespace UE::AnimNext
 	void IInterfaceA::FuncA(const FExecutionContext& Context, const TTraitBinding<IInterfaceA>& Binding) const
 	{
 		TTraitBinding<IInterfaceA> SuperBinding;
-		if (Context.GetInterfaceSuper(Binding, SuperBinding))
+		if (Binding.GetStackInterfaceSuper(SuperBinding))
 		{
 			SuperBinding.FuncA(Context);
 		}
@@ -86,7 +86,7 @@ namespace UE::AnimNext
 	void IInterfaceB::FuncB(const FExecutionContext& Context, const TTraitBinding<IInterfaceB>& Binding) const
 	{
 		TTraitBinding<IInterfaceB> SuperBinding;
-		if (Context.GetInterfaceSuper(Binding, SuperBinding))
+		if (Binding.GetStackInterfaceSuper(SuperBinding))
 		{
 			SuperBinding.FuncB(Context);
 		}
@@ -116,7 +116,7 @@ namespace UE::AnimNext
 	void IInterfaceC::FuncC(const FExecutionContext& Context, const TTraitBinding<IInterfaceC>& Binding) const
 	{
 		TTraitBinding<IInterfaceC> SuperBinding;
-		if (Context.GetInterfaceSuper(Binding, SuperBinding))
+		if (Binding.GetStackInterfaceSuper(SuperBinding))
 		{
 			SuperBinding.FuncC(Context);
 		}
@@ -591,11 +591,10 @@ bool FAnimationAnimNextRuntimeTest_NodeLifetime::RunTest(const FString& InParame
 
 		// Validate handle bookkeeping
 		{
-			FTraitBinding RootBinding;									// Empty, no parent
 			FAnimNextTraitHandle TraitHandle00(NodeHandles[0], 0);	// Point to first node, first base trait
 
 			// Allocate a node
-			FTraitPtr TraitPtr00 = Context.AllocateNodeInstance(RootBinding, TraitHandle00);
+			FTraitPtr TraitPtr00 = Context.AllocateNodeInstance(*GraphInstance.GetImpl(), TraitHandle00);
 			AddErrorIfFalse(TraitPtr00.IsValid(), "FAnimationAnimNextRuntimeTest_NodeLifetime -> Failed to allocate a node instance");
 			AddErrorIfFalse(TraitPtr00.GetTraitIndex() == 0, "FAnimationAnimNextRuntimeTest_NodeLifetime -> Allocated trait pointer should point to root trait");
 			AddErrorIfFalse(!TraitPtr00.IsWeak(), "FAnimationAnimNextRuntimeTest_NodeLifetime -> Allocated trait pointer should not be weak, we have no parent");
@@ -641,13 +640,12 @@ bool FAnimationAnimNextRuntimeTest_NodeLifetime::RunTest(const FString& InParame
 
 		// Validate parent support
 		{
-			FTraitBinding RootBinding;									// Empty, no parent
 			FAnimNextTraitHandle TraitHandle00(NodeHandles[0], 0);		// Point to first node, first base trait
 			FAnimNextTraitHandle TraitHandle03(NodeHandles[0], 3);		// Point to first node, second base trait
 			FAnimNextTraitHandle TraitHandle10(NodeHandles[1], 0);		// Point to second node, first base trait
 
 			// Allocate our first node
-			FTraitPtr TraitPtr00 = Context.AllocateNodeInstance(RootBinding, TraitHandle00);
+			FTraitPtr TraitPtr00 = Context.AllocateNodeInstance(*GraphInstance.GetImpl(), TraitHandle00);
 			AddErrorIfFalse(TraitPtr00.IsValid(), "FAnimationAnimNextRuntimeTest_NodeLifetime -> Failed to allocate a node instance");
 
 			// Allocate a new node, using the first as a parent
@@ -680,11 +678,10 @@ bool FAnimationAnimNextRuntimeTest_NodeLifetime::RunTest(const FString& InParame
 			Private::DestructedTraits = &DestructedTraits;
 
 			{
-				FTraitBinding RootBinding;									// Empty, no parent
 				FAnimNextTraitHandle TraitHandle00(NodeHandles[0], 0);		// Point to first node, first base trait
 
 				// Allocate our node instance
-				FTraitPtr TraitPtr00 = Context.AllocateNodeInstance(RootBinding, TraitHandle00);
+				FTraitPtr TraitPtr00 = Context.AllocateNodeInstance(*GraphInstance.GetImpl(), TraitHandle00);
 				AddErrorIfFalse(TraitPtr00.IsValid(), "FAnimationAnimNextRuntimeTest_NodeLifetime -> Failed to allocate a node instance");
 
 				// Validate instance constructors
@@ -722,9 +719,9 @@ bool FAnimationAnimNextRuntimeTest_NodeLifetime::RunTest(const FString& InParame
 	return true;
 }
 
-IMPLEMENT_SIMPLE_AUTOMATION_TEST(FAnimationAnimNextRuntimeTest_GetTraitInterface, "Animation.AnimNext.Runtime.GetTraitInterface", EAutomationTestFlags::EditorContext | EAutomationTestFlags::EngineFilter)
+IMPLEMENT_SIMPLE_AUTOMATION_TEST(FAnimationAnimNextRuntimeTest_TraitStackBinding, "Animation.AnimNext.Runtime.TraitStackBinding", EAutomationTestFlags::EditorContext | EAutomationTestFlags::EngineFilter)
 
-bool FAnimationAnimNextRuntimeTest_GetTraitInterface::RunTest(const FString& InParameters)
+bool FAnimationAnimNextRuntimeTest_TraitStackBinding::RunTest(const FString& InParameters)
 {
 	using namespace UE::AnimNext;
 
@@ -735,11 +732,12 @@ bool FAnimationAnimNextRuntimeTest_GetTraitInterface::RunTest(const FString& InP
 
 		UFactory* GraphFactory = NewObject<UAnimNextGraphFactory>();
 		UAnimNextGraph* AnimNextGraph = CastChecked<UAnimNextGraph>(GraphFactory->FactoryCreateNew(UAnimNextGraph::StaticClass(), GetTransientPackage(), TEXT("TestAnimNextGraph"), RF_Transient, nullptr, nullptr, NAME_None));
-		UE_RETURN_ON_ERROR(AnimNextGraph != nullptr, "FAnimationAnimNextRuntimeTest_GetTraitInterface -> Failed to create animation graph");
+		UE_RETURN_ON_ERROR(AnimNextGraph != nullptr, "FAnimationAnimNextRuntimeTest_TraitStackBinding -> Failed to create animation graph");
 
 		FScopedClearNodeTemplateRegistry ScopedClearNodeTemplateRegistry;
 		FNodeTemplateRegistry& Registry = FNodeTemplateRegistry::Get();
 
+		// Build a node with 2 stacks
 		TArray<FTraitUID> NodeTemplateTraitList;
 		NodeTemplateTraitList.Add(FTraitA_Base::TraitUID);
 		NodeTemplateTraitList.Add(FTraitAB_Add::TraitUID);
@@ -752,7 +750,7 @@ bool FAnimationAnimNextRuntimeTest_GetTraitInterface::RunTest(const FString& InP
 		const FNodeTemplate* NodeTemplate0 = FNodeTemplateBuilder::BuildNodeTemplate(NodeTemplateTraitList, NodeTemplateBuffer0);
 
 		FNodeTemplateRegistryHandle TemplateHandle0 = Registry.FindOrAdd(NodeTemplate0);
-		AddErrorIfFalse(TemplateHandle0.IsValid(), "FAnimationAnimNextRuntimeTest_GetTraitInterface -> Registry should contain our template");
+		AddErrorIfFalse(TemplateHandle0.IsValid(), "FAnimationAnimNextRuntimeTest_TraitStackBinding -> Registry should contain our template");
 
 		TArray<FNodeHandle> NodeHandles;
 
@@ -778,7 +776,7 @@ bool FAnimationAnimNextRuntimeTest_GetTraitInterface::RunTest(const FString& InP
 				});
 			TraitWriter.EndNodeWriting();
 
-			AddErrorIfFalse(TraitWriter.GetErrorState() == FTraitWriter::EErrorState::None, "FAnimationAnimNextRuntimeTest_GetTraitInterface -> Failed to write traits");
+			AddErrorIfFalse(TraitWriter.GetErrorState() == FTraitWriter::EErrorState::None, "FAnimationAnimNextRuntimeTest_TraitStackBinding -> Failed to write traits");
 			GraphSharedDataArchiveBuffer = TraitWriter.GetGraphSharedData();
 			GraphReferencedObjects = TraitWriter.GetGraphReferencedObjects();
 		}
@@ -793,407 +791,352 @@ bool FAnimationAnimNextRuntimeTest_GetTraitInterface::RunTest(const FString& InP
 
 		// Validate from the first base trait
 		{
-			FTraitBinding ParentBinding;								// Empty, no parent
-			FAnimNextTraitHandle TraitHandle00(NodeHandles[0], 0);		// Point to first node, first base trait
+			FAnimNextTraitHandle TraitHandle0(NodeHandles[0], 0);		// Point to first node, first base trait
 
-			FTraitPtr TraitPtr00 = Context.AllocateNodeInstance(ParentBinding, TraitHandle00);
-			AddErrorIfFalse(TraitPtr00.IsValid(), "FAnimationAnimNextRuntimeTest_GetTraitInterface -> Failed to allocate a node instance");
-
-			// Validate GetInterface from a trait handle
-			TTraitBinding<IInterfaceC> Binding00C;
-			AddErrorIfFalse(Context.GetInterface(TraitPtr00, Binding00C), "FAnimationAnimNextRuntimeTest_GetTraitInterface -> InterfaceC not found");
-			AddErrorIfFalse(Binding00C.IsValid(), "FAnimationAnimNextRuntimeTest_GetTraitInterface -> InterfaceC binding not valid");
-			AddErrorIfFalse(Binding00C.GetInterfaceUID() == IInterfaceC::InterfaceUID, "FAnimationAnimNextRuntimeTest_GetTraitInterface -> Unexpected interface UID found in trait binding");
-			AddErrorIfFalse(Binding00C.GetTraitPtr().GetTraitIndex() == 2, "FAnimationAnimNextRuntimeTest_GetTraitInterface -> InterfaceC not found on expected trait");
-			AddErrorIfFalse(Binding00C.GetTraitPtr().GetNodeInstance()->GetNodeHandle() == NodeHandles[0], "FAnimationAnimNextRuntimeTest_GetTraitInterface -> InterfaceC not found on expected node");
-			AddErrorIfFalse(Binding00C.GetSharedData<FTraitAC_Add::FSharedData>()->TraitUID == FTraitAC_Add::TraitUID, "FAnimationAnimNextRuntimeTest_GetTraitInterface -> Unexpected shared data in trait binding");
-			AddErrorIfFalse(Binding00C.GetInstanceData<FTraitAC_Add::FInstanceData>()->TraitUID == FTraitAC_Add::TraitUID, "FAnimationAnimNextRuntimeTest_GetTraitInterface -> Unexpected instance data in trait binding");
-
-			TTraitBinding<IInterfaceB> Binding00B;
-			AddErrorIfFalse(Context.GetInterface(TraitPtr00, Binding00B), "FAnimationAnimNextRuntimeTest_GetTraitInterface -> InterfaceB not found");
-			AddErrorIfFalse(Binding00B.IsValid(), "FAnimationAnimNextRuntimeTest_GetTraitInterface -> InterfaceB binding not valid");
-			AddErrorIfFalse(Binding00B.GetInterfaceUID() == IInterfaceB::InterfaceUID, "FAnimationAnimNextRuntimeTest_GetTraitInterface -> Unexpected interface UID found in trait binding");
-			AddErrorIfFalse(Binding00B.GetTraitPtr().GetTraitIndex() == 1, "FAnimationAnimNextRuntimeTest_GetTraitInterface -> InterfaceB not found on expected trait");
-			AddErrorIfFalse(Binding00B.GetTraitPtr().GetNodeInstance()->GetNodeHandle() == NodeHandles[0], "FAnimationAnimNextRuntimeTest_GetTraitInterface -> InterfaceB not found on expected node");
-			AddErrorIfFalse(Binding00B.GetSharedData<FTraitAB_Add::FSharedData>()->TraitUID == FTraitAB_Add::TraitUID, "FAnimationAnimNextRuntimeTest_GetTraitInterface -> Unexpected shared data in trait binding");
-			AddErrorIfFalse(Binding00B.GetInstanceData<FTraitAB_Add::FInstanceData>()->TraitUID == FTraitAB_Add::TraitUID, "FAnimationAnimNextRuntimeTest_GetTraitInterface -> Unexpected instance data in trait binding");
-
-			TTraitBinding<IInterfaceA> Binding00A;
-			AddErrorIfFalse(Context.GetInterface(TraitPtr00, Binding00A), "FAnimationAnimNextRuntimeTest_GetTraitInterface -> InterfaceA not found");
-			AddErrorIfFalse(Binding00A.IsValid(), "FAnimationAnimNextRuntimeTest_GetTraitInterface -> InterfaceA binding not valid");
-			AddErrorIfFalse(Binding00A.GetInterfaceUID() == IInterfaceA::InterfaceUID, "FAnimationAnimNextRuntimeTest_GetTraitInterface -> Unexpected interface UID found in trait binding");
-			AddErrorIfFalse(Binding00A.GetTraitPtr().GetTraitIndex() == 2, "FAnimationAnimNextRuntimeTest_GetTraitInterface -> InterfaceA not found on expected trait");
-			AddErrorIfFalse(Binding00A.GetTraitPtr().GetNodeInstance()->GetNodeHandle() == NodeHandles[0], "FAnimationAnimNextRuntimeTest_GetTraitInterface -> InterfaceA not found on expected node");
-			AddErrorIfFalse(Binding00A.GetSharedData<FTraitAC_Add::FSharedData>()->TraitUID == FTraitAC_Add::TraitUID, "FAnimationAnimNextRuntimeTest_GetTraitInterface -> Unexpected shared data in trait binding");
-			AddErrorIfFalse(Binding00A.GetInstanceData<FTraitAC_Add::FInstanceData>()->TraitUID == FTraitAC_Add::TraitUID, "FAnimationAnimNextRuntimeTest_GetTraitInterface -> Unexpected instance data in trait binding");
-
-			// Validate GetInterface from a trait binding
-			{
-				{
-					{
-						TTraitBinding<IInterfaceC> Binding00C_;
-						AddErrorIfFalse(Context.GetInterface(Binding00C, Binding00C_), "FAnimationAnimNextRuntimeTest_GetTraitInterface -> InterfaceC not found");
-						AddErrorIfFalse(Binding00C == Binding00C_, "FAnimationAnimNextRuntimeTest_GetTraitInterface -> GetInterface methods should return the same result");
-					}
+			FTraitPtr TraitPtr0 = Context.AllocateNodeInstance(*GraphInstance.GetImpl(), TraitHandle0);
+			AddErrorIfFalse(TraitPtr0.IsValid(), "FAnimationAnimNextRuntimeTest_TraitStackBinding -> Failed to allocate a node instance");
 
-					{
-						TTraitBinding<IInterfaceC> Binding00C_;
-						AddErrorIfFalse(Context.GetInterface(Binding00B, Binding00C_), "FAnimationAnimNextRuntimeTest_GetTraitInterface -> InterfaceC not found");
-						AddErrorIfFalse(Binding00C == Binding00C_, "FAnimationAnimNextRuntimeTest_GetTraitInterface -> GetInterface methods should return the same result");
-					}
+			FTraitStackBinding Stack0;
+			AddErrorIfFalse(Context.GetStack(TraitPtr0, Stack0), "FAnimationAnimNextRuntimeTest_TraitStackBinding -> Failed to bind to trait stack");
 
-					{
-						TTraitBinding<IInterfaceC> Binding00C_;
-						AddErrorIfFalse(Context.GetInterface(Binding00A, Binding00C_), "FAnimationAnimNextRuntimeTest_GetTraitInterface -> InterfaceC not found");
-						AddErrorIfFalse(Binding00C == Binding00C_, "FAnimationAnimNextRuntimeTest_GetTraitInterface -> GetInterface methods should return the same result");
-					}
-				}
+			// Test the first stack
+			{
+				AddErrorIfFalse(Stack0.IsValid(), "FAnimationAnimNextRuntimeTest_TraitStackBinding -> Trait stack binding should be valid");
+				AddErrorIfFalse(Stack0.GetBaseTraitPtr() == TraitPtr0, "FAnimationAnimNextRuntimeTest_TraitStackBinding -> Unexpected base trait ptr");
 
+				// Stack binding copying and equality testing
 				{
-					{
-						TTraitBinding<IInterfaceB> Binding00B_;
-						AddErrorIfFalse(Context.GetInterface(Binding00C, Binding00B_), "FAnimationAnimNextRuntimeTest_GetTraitInterface -> InterfaceB not found");
-						AddErrorIfFalse(Binding00B == Binding00B_, "FAnimationAnimNextRuntimeTest_GetTraitInterface -> GetInterface methods should return the same result");
-					}
-
-					{
-						TTraitBinding<IInterfaceB> Binding00B_;
-						AddErrorIfFalse(Context.GetInterface(Binding00B, Binding00B_), "FAnimationAnimNextRuntimeTest_GetTraitInterface -> InterfaceB not found");
-						AddErrorIfFalse(Binding00B == Binding00B_, "FAnimationAnimNextRuntimeTest_GetTraitInterface -> GetInterface methods should return the same result");
-					}
-
-					{
-						TTraitBinding<IInterfaceB> Binding00B_;
-						AddErrorIfFalse(Context.GetInterface(Binding00A, Binding00B_), "FAnimationAnimNextRuntimeTest_GetTraitInterface -> InterfaceB not found");
-						AddErrorIfFalse(Binding00B == Binding00B_, "FAnimationAnimNextRuntimeTest_GetTraitInterface -> GetInterface methods should return the same result");
-					}
+					FTraitStackBinding Stack0Copy = Stack0;
+					AddErrorIfFalse(Stack0Copy.IsValid(), "FAnimationAnimNextRuntimeTest_TraitStackBinding -> Trait stack binding should be valid");
+					AddErrorIfFalse(Stack0Copy.GetBaseTraitPtr() == TraitPtr0, "FAnimationAnimNextRuntimeTest_TraitStackBinding -> Unexpected base trait ptr");
+					AddErrorIfFalse(Stack0 == Stack0Copy, "FAnimationAnimNextRuntimeTest_TraitStackBinding -> Trait stack bindings should be equal");
+					AddErrorIfFalse(!(Stack0 != Stack0Copy), "FAnimationAnimNextRuntimeTest_TraitStackBinding -> Trait stack bindings should be equal");
+
+					Stack0Copy.Reset();
+					AddErrorIfFalse(!Stack0Copy.IsValid(), "FAnimationAnimNextRuntimeTest_TraitStackBinding -> Trait stack binding should be not valid");
+
+					FTraitStackBinding Stack1;
+					AddErrorIfFalse(Context.GetStack(FWeakTraitPtr(TraitPtr0.GetNodeInstance(), 1), Stack1), "FAnimationAnimNextRuntimeTest_TraitStackBinding -> Failed to bind to trait stack");
+					AddErrorIfFalse(Stack1.IsValid(), "FAnimationAnimNextRuntimeTest_TraitStackBinding -> Trait stack binding should be valid");
+					AddErrorIfFalse(Stack1.GetBaseTraitPtr() == TraitPtr0, "FAnimationAnimNextRuntimeTest_TraitStackBinding -> Unexpected base trait ptr");
+					AddErrorIfFalse(Stack0 == Stack1, "FAnimationAnimNextRuntimeTest_TraitStackBinding -> Trait stack bindings should be equal");
 				}
 
+				// Stack binding iteration from top to bottom
 				{
-					{
-						TTraitBinding<IInterfaceA> Binding00A_;
-						AddErrorIfFalse(Context.GetInterface(Binding00C, Binding00A_), "FAnimationAnimNextRuntimeTest_GetTraitInterface -> InterfaceA not found");
-						AddErrorIfFalse(Binding00A == Binding00A_, "FAnimationAnimNextRuntimeTest_GetTraitInterface -> GetInterface methods should return the same result");
-					}
-
-					{
-						TTraitBinding<IInterfaceA> Binding00A_;
-						AddErrorIfFalse(Context.GetInterface(Binding00B, Binding00A_), "FAnimationAnimNextRuntimeTest_GetTraitInterface -> InterfaceA not found");
-						AddErrorIfFalse(Binding00A == Binding00A_, "FAnimationAnimNextRuntimeTest_GetTraitInterface -> GetInterface methods should return the same result");
-					}
-
-					{
-						TTraitBinding<IInterfaceA> Binding00A_;
-						AddErrorIfFalse(Context.GetInterface(Binding00A, Binding00A_), "FAnimationAnimNextRuntimeTest_GetTraitInterface -> InterfaceA not found");
-						AddErrorIfFalse(Binding00A == Binding00A_, "FAnimationAnimNextRuntimeTest_GetTraitInterface -> GetInterface methods should return the same result");
-					}
+					FTraitBinding TraitBinding;
+					AddErrorIfFalse(Stack0.GetTopTrait(TraitBinding), "FAnimationAnimNextRuntimeTest_TraitStackBinding -> Failed to retrieve top trait");
+					AddErrorIfFalse(TraitBinding.IsValid(), "FAnimationAnimNextRuntimeTest_TraitStackBinding -> Trait binding should be valid");
+					AddErrorIfFalse(TraitBinding.GetStack() == &Stack0, "FAnimationAnimNextRuntimeTest_TraitStackBinding -> Unexpected trait stack binding");
+					AddErrorIfFalse(TraitBinding.GetTraitIndex() == 2, "FAnimationAnimNextRuntimeTest_TraitStackBinding -> Unexpected stack local trait index");
+					AddErrorIfFalse(TraitBinding.GetTraitPtr().GetTraitIndex() == 2, "FAnimationAnimNextRuntimeTest_TraitStackBinding -> Unexpected node local trait index");
+
+					AddErrorIfFalse(Stack0.GetParentTrait(TraitBinding, TraitBinding), "FAnimationAnimNextRuntimeTest_TraitStackBinding -> Failed to retrieve parent trait");
+					AddErrorIfFalse(TraitBinding.IsValid(), "FAnimationAnimNextRuntimeTest_TraitStackBinding -> Trait binding should be valid");
+					AddErrorIfFalse(TraitBinding.GetStack() == &Stack0, "FAnimationAnimNextRuntimeTest_TraitStackBinding -> Unexpected trait stack binding");
+					AddErrorIfFalse(TraitBinding.GetTraitIndex() == 1, "FAnimationAnimNextRuntimeTest_TraitStackBinding -> Unexpected stack local trait index");
+					AddErrorIfFalse(TraitBinding.GetTraitPtr().GetTraitIndex() == 1, "FAnimationAnimNextRuntimeTest_TraitStackBinding -> Unexpected node local trait index");
+
+					AddErrorIfFalse(Stack0.GetParentTrait(TraitBinding, TraitBinding), "FAnimationAnimNextRuntimeTest_TraitStackBinding -> Failed to retrieve parent trait");
+					AddErrorIfFalse(TraitBinding.IsValid(), "FAnimationAnimNextRuntimeTest_TraitStackBinding -> Trait binding should be valid");
+					AddErrorIfFalse(TraitBinding.GetStack() == &Stack0, "FAnimationAnimNextRuntimeTest_TraitStackBinding -> Unexpected trait stack binding");
+					AddErrorIfFalse(TraitBinding.GetTraitIndex() == 0, "FAnimationAnimNextRuntimeTest_TraitStackBinding -> Unexpected stack local trait index");
+					AddErrorIfFalse(TraitBinding.GetTraitPtr().GetTraitIndex() == 0, "FAnimationAnimNextRuntimeTest_TraitStackBinding -> Unexpected node local trait index");
+					AddErrorIfFalse(TraitBinding.GetTraitPtr() == Stack0.GetBaseTraitPtr(), "FAnimationAnimNextRuntimeTest_TraitStackBinding -> Unexpected stack base trait ptr");
+
+					AddErrorIfFalse(!Stack0.GetParentTrait(TraitBinding, TraitBinding), "FAnimationAnimNextRuntimeTest_TraitStackBinding -> Unexpected parent trait");
 				}
-			}
-		}
-
-		// Validate from the second base trait
-		{
-			FTraitBinding ParentBinding;								// Empty, no parent
-			FAnimNextTraitHandle TraitHandle03(NodeHandles[0], 3);		// Point to first node, second base trait
 
-			FTraitPtr TraitPtr03 = Context.AllocateNodeInstance(ParentBinding, TraitHandle03);
-			AddErrorIfFalse(TraitPtr03.IsValid(), "FAnimationAnimNextRuntimeTest_GetTraitInterface -> Failed to allocate a node instance");
-
-			// Validate GetInterface from a trait handle
-			TTraitBinding<IInterfaceC> Binding03C;
-			AddErrorIfFalse(Context.GetInterface(TraitPtr03, Binding03C), "FAnimationAnimNextRuntimeTest_GetTraitInterface -> InterfaceC not found");
-			AddErrorIfFalse(Binding03C.IsValid(), "FAnimationAnimNextRuntimeTest_GetTraitInterface -> InterfaceC binding not valid");
-			AddErrorIfFalse(Binding03C.GetInterfaceUID() == IInterfaceC::InterfaceUID, "FAnimationAnimNextRuntimeTest_GetTraitInterface -> Unexpected interface UID found in trait binding");
-			AddErrorIfFalse(Binding03C.GetTraitPtr().GetTraitIndex() == 4, "FAnimationAnimNextRuntimeTest_GetTraitInterface -> InterfaceC not found on expected trait");
-			AddErrorIfFalse(Binding03C.GetTraitPtr().GetNodeInstance()->GetNodeHandle() == NodeHandles[0], "FAnimationAnimNextRuntimeTest_GetTraitInterface -> InterfaceC not found on expected node");
-			AddErrorIfFalse(Binding03C.GetSharedData<FTraitAC_Add::FSharedData>()->TraitUID == FTraitAC_Add::TraitUID, "FAnimationAnimNextRuntimeTest_GetTraitInterface -> Unexpected shared data in trait binding");
-			AddErrorIfFalse(Binding03C.GetInstanceData<FTraitAC_Add::FInstanceData>()->TraitUID == FTraitAC_Add::TraitUID, "FAnimationAnimNextRuntimeTest_GetTraitInterface -> Unexpected instance data in trait binding");
-
-			TTraitBinding<IInterfaceB> Binding03B;
-			AddErrorIfFalse(!Context.GetInterface(TraitPtr03, Binding03B), "FAnimationAnimNextRuntimeTest_GetTraitInterface -> InterfaceB found");
-			AddErrorIfFalse(!Binding03B.IsValid(), "FAnimationAnimNextRuntimeTest_GetTraitInterface -> InterfaceB binding should not be valid");
-
-			TTraitBinding<IInterfaceA> Binding03A;
-			AddErrorIfFalse(Context.GetInterface(TraitPtr03, Binding03A), "FAnimationAnimNextRuntimeTest_GetTraitInterface -> InterfaceA not found");
-			AddErrorIfFalse(Binding03A.IsValid(), "FAnimationAnimNextRuntimeTest_GetTraitInterface -> InterfaceA binding not valid");
-			AddErrorIfFalse(Binding03A.GetInterfaceUID() == IInterfaceA::InterfaceUID, "FAnimationAnimNextRuntimeTest_GetTraitInterface -> Unexpected interface UID found in trait binding");
-			AddErrorIfFalse(Binding03A.GetTraitPtr().GetTraitIndex() == 4, "FAnimationAnimNextRuntimeTest_GetTraitInterface -> InterfaceA not found on expected trait");
-			AddErrorIfFalse(Binding03A.GetTraitPtr().GetNodeInstance()->GetNodeHandle() == NodeHandles[0], "FAnimationAnimNextRuntimeTest_GetTraitInterface -> InterfaceA not found on expected node");
-			AddErrorIfFalse(Binding03A.GetSharedData<FTraitAC_Add::FSharedData>()->TraitUID == FTraitAC_Add::TraitUID, "FAnimationAnimNextRuntimeTest_GetTraitInterface -> Unexpected shared data in trait binding");
-			AddErrorIfFalse(Binding03A.GetInstanceData<FTraitAC_Add::FInstanceData>()->TraitUID == FTraitAC_Add::TraitUID, "FAnimationAnimNextRuntimeTest_GetTraitInterface -> Unexpected instance data in trait binding");
-
-			// Validate GetInterface from a trait binding
-			{
+				// Stack binding iteration from bottom to top
 				{
-					{
-						TTraitBinding<IInterfaceC> Binding03C_;
-						AddErrorIfFalse(Context.GetInterface(Binding03C, Binding03C_), "FAnimationAnimNextRuntimeTest_GetTraitInterface -> InterfaceC not found");
-						AddErrorIfFalse(Binding03C == Binding03C_, "FAnimationAnimNextRuntimeTest_GetTraitInterface -> GetInterface methods should return the same result");
-					}
-
-					{
-						TTraitBinding<IInterfaceC> Binding03C_;
-						AddErrorIfFalse(!Context.GetInterface(Binding03B, Binding03C_), "FAnimationAnimNextRuntimeTest_GetTraitInterface -> InterfaceC found");
-					}
-
-					{
-						TTraitBinding<IInterfaceC> Binding03C_;
-						AddErrorIfFalse(Context.GetInterface(Binding03A, Binding03C_), "FAnimationAnimNextRuntimeTest_GetTraitInterface -> InterfaceC not found");
-						AddErrorIfFalse(Binding03C == Binding03C_, "FAnimationAnimNextRuntimeTest_GetTraitInterface -> GetInterface methods should return the same result");
-					}
+					FTraitBinding TraitBinding;
+					AddErrorIfFalse(Stack0.GetBaseTrait(TraitBinding), "FAnimationAnimNextRuntimeTest_TraitStackBinding -> Failed to retrieve base trait");
+					AddErrorIfFalse(TraitBinding.IsValid(), "FAnimationAnimNextRuntimeTest_TraitStackBinding -> Trait binding should be valid");
+					AddErrorIfFalse(TraitBinding.GetStack() == &Stack0, "FAnimationAnimNextRuntimeTest_TraitStackBinding -> Unexpected trait stack binding");
+					AddErrorIfFalse(TraitBinding.GetTraitIndex() == 0, "FAnimationAnimNextRuntimeTest_TraitStackBinding -> Unexpected stack local trait index");
+					AddErrorIfFalse(TraitBinding.GetTraitPtr().GetTraitIndex() == 0, "FAnimationAnimNextRuntimeTest_TraitStackBinding -> Unexpected node local trait index");
+					AddErrorIfFalse(TraitBinding.GetTraitPtr() == Stack0.GetBaseTraitPtr(), "FAnimationAnimNextRuntimeTest_TraitStackBinding -> Unexpected stack base trait ptr");
+
+					AddErrorIfFalse(Stack0.GetChildTrait(TraitBinding, TraitBinding), "FAnimationAnimNextRuntimeTest_TraitStackBinding -> Failed to retrieve child trait");
+					AddErrorIfFalse(TraitBinding.IsValid(), "FAnimationAnimNextRuntimeTest_TraitStackBinding -> Trait binding should be valid");
+					AddErrorIfFalse(TraitBinding.GetStack() == &Stack0, "FAnimationAnimNextRuntimeTest_TraitStackBinding -> Unexpected trait stack binding");
+					AddErrorIfFalse(TraitBinding.GetTraitIndex() == 1, "FAnimationAnimNextRuntimeTest_TraitStackBinding -> Unexpected stack local trait index");
+					AddErrorIfFalse(TraitBinding.GetTraitPtr().GetTraitIndex() == 1, "FAnimationAnimNextRuntimeTest_TraitStackBinding -> Unexpected node local trait index");
+
+					AddErrorIfFalse(Stack0.GetChildTrait(TraitBinding, TraitBinding), "FAnimationAnimNextRuntimeTest_TraitStackBinding -> Failed to retrieve child trait");
+					AddErrorIfFalse(TraitBinding.IsValid(), "FAnimationAnimNextRuntimeTest_TraitStackBinding -> Trait binding should be valid");
+					AddErrorIfFalse(TraitBinding.GetStack() == &Stack0, "FAnimationAnimNextRuntimeTest_TraitStackBinding -> Unexpected trait stack binding");
+					AddErrorIfFalse(TraitBinding.GetTraitIndex() == 2, "FAnimationAnimNextRuntimeTest_TraitStackBinding -> Unexpected stack local trait index");
+					AddErrorIfFalse(TraitBinding.GetTraitPtr().GetTraitIndex() == 2, "FAnimationAnimNextRuntimeTest_TraitStackBinding -> Unexpected node local trait index");
+
+					AddErrorIfFalse(!Stack0.GetChildTrait(TraitBinding, TraitBinding), "FAnimationAnimNextRuntimeTest_TraitStackBinding -> Unexpected child trait");
 				}
 
+				// Stack binding arbitrary iteration
 				{
-					{
-						TTraitBinding<IInterfaceB> Binding03B_;
-						AddErrorIfFalse(!Context.GetInterface(Binding03C, Binding03B_), "FAnimationAnimNextRuntimeTest_GetTraitInterface -> InterfaceB found");
-						AddErrorIfFalse(Binding03B == Binding03B_, "FAnimationAnimNextRuntimeTest_GetTraitInterface -> GetInterface methods should return the same result");
-					}
-
-					{
-						TTraitBinding<IInterfaceB> Binding03B_;
-						AddErrorIfFalse(!Context.GetInterface(Binding03B, Binding03B_), "FAnimationAnimNextRuntimeTest_GetTraitInterface -> InterfaceB found");
-						AddErrorIfFalse(Binding03B == Binding03B_, "FAnimationAnimNextRuntimeTest_GetTraitInterface -> GetInterface methods should return the same result");
-					}
-
-					{
-						TTraitBinding<IInterfaceB> Binding03B_;
-						AddErrorIfFalse(!Context.GetInterface(Binding03A, Binding03B_), "FAnimationAnimNextRuntimeTest_GetTraitInterface -> InterfaceB found");
-						AddErrorIfFalse(Binding03B == Binding03B_, "FAnimationAnimNextRuntimeTest_GetTraitInterface -> GetInterface methods should return the same result");
-					}
+					AddErrorIfFalse(Stack0.GetNumTraits() == 3, "FAnimationAnimNextRuntimeTest_TraitStackBinding -> Unexpected trait stack size");
+
+					FTraitBinding TraitBinding;
+					AddErrorIfFalse(Stack0.GetTrait(0, TraitBinding), "FAnimationAnimNextRuntimeTest_TraitStackBinding -> Failed to retrieve trait");
+					AddErrorIfFalse(TraitBinding.IsValid(), "FAnimationAnimNextRuntimeTest_TraitStackBinding -> Trait binding should be valid");
+					AddErrorIfFalse(TraitBinding.GetStack() == &Stack0, "FAnimationAnimNextRuntimeTest_TraitStackBinding -> Unexpected trait stack binding");
+					AddErrorIfFalse(TraitBinding.GetTraitIndex() == 0, "FAnimationAnimNextRuntimeTest_TraitStackBinding -> Unexpected stack local trait index");
+					AddErrorIfFalse(TraitBinding.GetTraitPtr().GetTraitIndex() == 0, "FAnimationAnimNextRuntimeTest_TraitStackBinding -> Unexpected node local trait index");
+					AddErrorIfFalse(TraitBinding.GetTraitPtr() == Stack0.GetBaseTraitPtr(), "FAnimationAnimNextRuntimeTest_TraitStackBinding -> Unexpected stack base trait ptr");
+
+					AddErrorIfFalse(Stack0.GetTrait(1, TraitBinding), "FAnimationAnimNextRuntimeTest_TraitStackBinding -> Failed to retrieve trait");
+					AddErrorIfFalse(TraitBinding.IsValid(), "FAnimationAnimNextRuntimeTest_TraitStackBinding -> Trait binding should be valid");
+					AddErrorIfFalse(TraitBinding.GetStack() == &Stack0, "FAnimationAnimNextRuntimeTest_TraitStackBinding -> Unexpected trait stack binding");
+					AddErrorIfFalse(TraitBinding.GetTraitIndex() == 1, "FAnimationAnimNextRuntimeTest_TraitStackBinding -> Unexpected stack local trait index");
+					AddErrorIfFalse(TraitBinding.GetTraitPtr().GetTraitIndex() == 1, "FAnimationAnimNextRuntimeTest_TraitStackBinding -> Unexpected node local trait index");
+
+					AddErrorIfFalse(Stack0.GetTrait(2, TraitBinding), "FAnimationAnimNextRuntimeTest_TraitStackBinding -> Failed to retrieve trait");
+					AddErrorIfFalse(TraitBinding.IsValid(), "FAnimationAnimNextRuntimeTest_TraitStackBinding -> Trait binding should be valid");
+					AddErrorIfFalse(TraitBinding.GetStack() == &Stack0, "FAnimationAnimNextRuntimeTest_TraitStackBinding -> Unexpected trait stack binding");
+					AddErrorIfFalse(TraitBinding.GetTraitIndex() == 2, "FAnimationAnimNextRuntimeTest_TraitStackBinding -> Unexpected stack local trait index");
+					AddErrorIfFalse(TraitBinding.GetTraitPtr().GetTraitIndex() == 2, "FAnimationAnimNextRuntimeTest_TraitStackBinding -> Unexpected node local trait index");
 				}
 
+				// Interface query validation
 				{
-					{
-						TTraitBinding<IInterfaceA> Binding03A_;
-						AddErrorIfFalse(Context.GetInterface(Binding03C, Binding03A_), "FAnimationAnimNextRuntimeTest_GetTraitInterface -> InterfaceA not found");
-						AddErrorIfFalse(Binding03A == Binding03A_, "FAnimationAnimNextRuntimeTest_GetTraitInterface -> GetInterface methods should return the same result");
-					}
-
-					{
-						TTraitBinding<IInterfaceA> Binding03A_;
-						AddErrorIfFalse(!Context.GetInterface(Binding03B, Binding03A_), "FAnimationAnimNextRuntimeTest_GetTraitInterface -> InterfaceA found");
-					}
-
-					{
-						TTraitBinding<IInterfaceA> Binding03A_;
-						AddErrorIfFalse(Context.GetInterface(Binding03A, Binding03A_), "FAnimationAnimNextRuntimeTest_GetTraitInterface -> InterfaceA not found");
-						AddErrorIfFalse(Binding03A == Binding03A_, "FAnimationAnimNextRuntimeTest_GetTraitInterface -> GetInterface methods should return the same result");
-					}
+					TTraitBinding<IInterfaceC> Binding0C;
+					AddErrorIfFalse(Stack0.GetInterface(Binding0C), "FAnimationAnimNextRuntimeTest_TraitStackBinding -> InterfaceC not found");
+					AddErrorIfFalse(Binding0C.IsValid(), "FAnimationAnimNextRuntimeTest_TraitStackBinding -> InterfaceC binding not valid");
+					AddErrorIfFalse(Binding0C.GetInterfaceUID() == IInterfaceC::InterfaceUID, "FAnimationAnimNextRuntimeTest_TraitStackBinding -> Unexpected interface UID found in trait binding");
+					AddErrorIfFalse(Binding0C.GetTraitPtr().GetTraitIndex() == 2, "FAnimationAnimNextRuntimeTest_TraitStackBinding -> InterfaceC not found on expected trait");
+					AddErrorIfFalse(Binding0C.GetTraitPtr().GetNodeInstance()->GetNodeHandle() == NodeHandles[0], "FAnimationAnimNextRuntimeTest_TraitStackBinding -> InterfaceC not found on expected node");
+					AddErrorIfFalse(Binding0C.GetSharedData<FTraitAC_Add::FSharedData>()->TraitUID == FTraitAC_Add::TraitUID, "FAnimationAnimNextRuntimeTest_TraitStackBinding -> Unexpected shared data in trait binding");
+					AddErrorIfFalse(Binding0C.GetInstanceData<FTraitAC_Add::FInstanceData>()->TraitUID == FTraitAC_Add::TraitUID, "FAnimationAnimNextRuntimeTest_TraitStackBinding -> Unexpected instance data in trait binding");
+
+					TTraitBinding<IInterfaceC> Binding0CCachedValid = Binding0C;
+					TTraitBinding<IInterfaceC> Binding0CCachedInvalid = Binding0C;	// Invalidated when we query for the super just below
+
+					AddErrorIfFalse(!Stack0.GetInterfaceSuper(Binding0C, Binding0C), "FAnimationAnimNextRuntimeTest_TraitStackBinding -> InterfaceC should not be found");
+					AddErrorIfFalse(!Binding0CCachedInvalid.GetStackInterfaceSuper(Binding0CCachedInvalid), "FAnimationAnimNextRuntimeTest_TraitStackBinding -> InterfaceC should not be found");
+
+					TTraitBinding<IInterfaceA> Binding0ACasted;
+					AddErrorIfFalse(Binding0CCachedValid.AsInterface(Binding0ACasted), "FAnimationAnimNextRuntimeTest_TraitStackBinding -> InterfaceA should be found");
+					AddErrorIfFalse(Binding0ACasted.IsValid(), "FAnimationAnimNextRuntimeTest_TraitStackBinding -> InterfaceA binding not valid");
+					AddErrorIfFalse(Binding0ACasted.GetInterfaceUID() == IInterfaceA::InterfaceUID, "FAnimationAnimNextRuntimeTest_TraitStackBinding -> Unexpected interface UID found in trait binding");
+					AddErrorIfFalse(Binding0ACasted.GetTraitPtr().GetTraitIndex() == 2, "FAnimationAnimNextRuntimeTest_TraitStackBinding -> InterfaceA not found on expected trait");
+					AddErrorIfFalse(Binding0ACasted.GetTraitPtr().GetNodeInstance()->GetNodeHandle() == NodeHandles[0], "FAnimationAnimNextRuntimeTest_TraitStackBinding -> InterfaceA not found on expected node");
+					AddErrorIfFalse(Binding0ACasted.GetSharedData<FTraitAC_Add::FSharedData>()->TraitUID == FTraitAC_Add::TraitUID, "FAnimationAnimNextRuntimeTest_TraitStackBinding -> Unexpected shared data in trait binding");
+					AddErrorIfFalse(Binding0ACasted.GetInstanceData<FTraitAC_Add::FInstanceData>()->TraitUID == FTraitAC_Add::TraitUID, "FAnimationAnimNextRuntimeTest_TraitStackBinding -> Unexpected instance data in trait binding");
+					AddErrorIfFalse(Binding0CCachedValid != Binding0ACasted, "FAnimationAnimNextRuntimeTest_TraitStackBinding -> Trait bindings should not be equal");
+
+					TTraitBinding<IInterfaceB> Binding0BCasted;
+					AddErrorIfFalse(!Binding0CCachedValid.AsInterface(Binding0BCasted), "FAnimationAnimNextRuntimeTest_TraitStackBinding -> InterfaceB should not be found");
+
+					// Test the same interface query but from an existing trait binding
+					TTraitBinding<IInterfaceC> Binding0CFromTrait;
+					AddErrorIfFalse(Binding0CCachedValid.GetStackInterface(Binding0CFromTrait), "FAnimationAnimNextRuntimeTest_TraitStackBinding -> InterfaceC not found");
+					AddErrorIfFalse(Binding0CFromTrait.IsValid(), "FAnimationAnimNextRuntimeTest_TraitStackBinding -> InterfaceC binding not valid");
+					AddErrorIfFalse(Binding0CFromTrait.GetInterfaceUID() == IInterfaceC::InterfaceUID, "FAnimationAnimNextRuntimeTest_TraitStackBinding -> Unexpected interface UID found in trait binding");
+					AddErrorIfFalse(Binding0CFromTrait.GetTraitPtr().GetTraitIndex() == 2, "FAnimationAnimNextRuntimeTest_TraitStackBinding -> InterfaceC not found on expected trait");
+					AddErrorIfFalse(Binding0CFromTrait.GetTraitPtr().GetNodeInstance()->GetNodeHandle() == NodeHandles[0], "FAnimationAnimNextRuntimeTest_TraitStackBinding -> InterfaceC not found on expected node");
+					AddErrorIfFalse(Binding0CFromTrait.GetSharedData<FTraitAC_Add::FSharedData>()->TraitUID == FTraitAC_Add::TraitUID, "FAnimationAnimNextRuntimeTest_TraitStackBinding -> Unexpected shared data in trait binding");
+					AddErrorIfFalse(Binding0CFromTrait.GetInstanceData<FTraitAC_Add::FInstanceData>()->TraitUID == FTraitAC_Add::TraitUID, "FAnimationAnimNextRuntimeTest_TraitStackBinding -> Unexpected instance data in trait binding");
+					AddErrorIfFalse(Binding0CCachedValid == Binding0CFromTrait, "FAnimationAnimNextRuntimeTest_TraitStackBinding -> Trait bindings should be equal");
+
+					TTraitBinding<IInterfaceB> Binding0B;
+					AddErrorIfFalse(Stack0.GetInterface(Binding0B), "FAnimationAnimNextRuntimeTest_TraitStackBinding -> InterfaceB not found");
+					AddErrorIfFalse(Binding0B.IsValid(), "FAnimationAnimNextRuntimeTest_TraitStackBinding -> InterfaceB binding not valid");
+					AddErrorIfFalse(Binding0B.GetInterfaceUID() == IInterfaceB::InterfaceUID, "FAnimationAnimNextRuntimeTest_TraitStackBinding -> Unexpected interface UID found in trait binding");
+					AddErrorIfFalse(Binding0B.GetTraitPtr().GetTraitIndex() == 1, "FAnimationAnimNextRuntimeTest_TraitStackBinding -> InterfaceB not found on expected trait");
+					AddErrorIfFalse(Binding0B.GetTraitPtr().GetNodeInstance()->GetNodeHandle() == NodeHandles[0], "FAnimationAnimNextRuntimeTest_TraitStackBinding -> InterfaceB not found on expected node");
+					AddErrorIfFalse(Binding0B.GetSharedData<FTraitAB_Add::FSharedData>()->TraitUID == FTraitAB_Add::TraitUID, "FAnimationAnimNextRuntimeTest_TraitStackBinding -> Unexpected shared data in trait binding");
+					AddErrorIfFalse(Binding0B.GetInstanceData<FTraitAB_Add::FInstanceData>()->TraitUID == FTraitAB_Add::TraitUID, "FAnimationAnimNextRuntimeTest_TraitStackBinding -> Unexpected instance data in trait binding");
+					AddErrorIfFalse(!Stack0.GetInterfaceSuper(Binding0B, Binding0B), "FAnimationAnimNextRuntimeTest_TraitStackBinding -> InterfaceB should not be found");
+					AddErrorIfFalse(!Binding0B.GetStackInterfaceSuper(Binding0B), "FAnimationAnimNextRuntimeTest_TraitStackBinding -> InterfaceB should not be found");
+
+					TTraitBinding<IInterfaceA> Binding0A;
+					AddErrorIfFalse(Stack0.GetInterface(Binding0A), "FAnimationAnimNextRuntimeTest_TraitStackBinding -> InterfaceA not found");
+					AddErrorIfFalse(Binding0A.IsValid(), "FAnimationAnimNextRuntimeTest_TraitStackBinding -> InterfaceA binding not valid");
+					AddErrorIfFalse(Binding0A.GetInterfaceUID() == IInterfaceA::InterfaceUID, "FAnimationAnimNextRuntimeTest_TraitStackBinding -> Unexpected interface UID found in trait binding");
+					AddErrorIfFalse(Binding0A.GetTraitPtr().GetTraitIndex() == 2, "FAnimationAnimNextRuntimeTest_TraitStackBinding -> InterfaceA not found on expected trait");
+					AddErrorIfFalse(Binding0A.GetTraitPtr().GetNodeInstance()->GetNodeHandle() == NodeHandles[0], "FAnimationAnimNextRuntimeTest_TraitStackBinding -> InterfaceA not found on expected node");
+					AddErrorIfFalse(Binding0A.GetSharedData<FTraitAC_Add::FSharedData>()->TraitUID == FTraitAC_Add::TraitUID, "FAnimationAnimNextRuntimeTest_TraitStackBinding -> Unexpected shared data in trait binding");
+					AddErrorIfFalse(Binding0A.GetInstanceData<FTraitAC_Add::FInstanceData>()->TraitUID == FTraitAC_Add::TraitUID, "FAnimationAnimNextRuntimeTest_TraitStackBinding -> Unexpected instance data in trait binding");
+
+					AddErrorIfFalse(Stack0.GetInterfaceSuper(Binding0A, Binding0A), "FAnimationAnimNextRuntimeTest_TraitStackBinding -> InterfaceA not found");
+					AddErrorIfFalse(Binding0A.IsValid(), "FAnimationAnimNextRuntimeTest_TraitStackBinding -> InterfaceA binding not valid");
+					AddErrorIfFalse(Binding0A.GetInterfaceUID() == IInterfaceA::InterfaceUID, "FAnimationAnimNextRuntimeTest_TraitStackBinding -> Unexpected interface UID found in trait binding");
+					AddErrorIfFalse(Binding0A.GetTraitPtr().GetTraitIndex() == 1, "FAnimationAnimNextRuntimeTest_TraitStackBinding -> InterfaceA not found on expected trait");
+					AddErrorIfFalse(Binding0A.GetTraitPtr().GetNodeInstance()->GetNodeHandle() == NodeHandles[0], "FAnimationAnimNextRuntimeTest_TraitStackBinding -> InterfaceA not found on expected node");
+					AddErrorIfFalse(Binding0A.GetSharedData<FTraitAB_Add::FSharedData>()->TraitUID == FTraitAB_Add::TraitUID, "FAnimationAnimNextRuntimeTest_TraitStackBinding -> Unexpected shared data in trait binding");
+					AddErrorIfFalse(Binding0A.GetInstanceData<FTraitAB_Add::FInstanceData>()->TraitUID == FTraitAB_Add::TraitUID, "FAnimationAnimNextRuntimeTest_TraitStackBinding -> Unexpected instance data in trait binding");
+
+					TTraitBinding<IInterfaceA> Binding0ACached = Binding0A;
+
+					AddErrorIfFalse(Stack0.GetInterfaceSuper(Binding0A, Binding0A), "FAnimationAnimNextRuntimeTest_TraitStackBinding -> InterfaceA not found");
+					AddErrorIfFalse(Binding0A.IsValid(), "FAnimationAnimNextRuntimeTest_TraitStackBinding -> InterfaceA binding not valid");
+					AddErrorIfFalse(Binding0A.GetInterfaceUID() == IInterfaceA::InterfaceUID, "FAnimationAnimNextRuntimeTest_TraitStackBinding -> Unexpected interface UID found in trait binding");
+					AddErrorIfFalse(Binding0A.GetTraitPtr().GetTraitIndex() == 0, "FAnimationAnimNextRuntimeTest_TraitStackBinding -> InterfaceA not found on expected trait");
+					AddErrorIfFalse(Binding0A.GetTraitPtr().GetNodeInstance()->GetNodeHandle() == NodeHandles[0], "FAnimationAnimNextRuntimeTest_TraitStackBinding -> InterfaceA not found on expected node");
+					AddErrorIfFalse(Binding0A.GetSharedData<FTraitA_Base::FSharedData>()->TraitUID == FTraitA_Base::TraitUID, "FAnimationAnimNextRuntimeTest_TraitStackBinding -> Unexpected shared data in trait binding");
+					AddErrorIfFalse(Binding0A.GetInstanceData<FTraitA_Base::FInstanceData>()->TraitUID == FTraitA_Base::TraitUID, "FAnimationAnimNextRuntimeTest_TraitStackBinding -> Unexpected instance data in trait binding");
+					AddErrorIfFalse(!Stack0.GetInterfaceSuper(Binding0A, Binding0A), "FAnimationAnimNextRuntimeTest_TraitStackBinding -> InterfaceA should not be found");
+
+					// Test the same interface query but from an existing trait binding
+					TTraitBinding<IInterfaceA> Binding0AFromTrait;
+					AddErrorIfFalse(Binding0ACached.GetStackInterfaceSuper(Binding0AFromTrait), "FAnimationAnimNextRuntimeTest_TraitStackBinding -> InterfaceA not found");
+					AddErrorIfFalse(Binding0AFromTrait.IsValid(), "FAnimationAnimNextRuntimeTest_TraitStackBinding -> InterfaceA binding not valid");
+					AddErrorIfFalse(Binding0AFromTrait.GetInterfaceUID() == IInterfaceA::InterfaceUID, "FAnimationAnimNextRuntimeTest_TraitStackBinding -> Unexpected interface UID found in trait binding");
+					AddErrorIfFalse(Binding0AFromTrait.GetTraitPtr().GetTraitIndex() == 0, "FAnimationAnimNextRuntimeTest_TraitStackBinding -> InterfaceA not found on expected trait");
+					AddErrorIfFalse(Binding0AFromTrait.GetTraitPtr().GetNodeInstance()->GetNodeHandle() == NodeHandles[0], "FAnimationAnimNextRuntimeTest_TraitStackBinding -> InterfaceA not found on expected node");
+					AddErrorIfFalse(Binding0AFromTrait.GetSharedData<FTraitA_Base::FSharedData>()->TraitUID == FTraitA_Base::TraitUID, "FAnimationAnimNextRuntimeTest_TraitStackBinding -> Unexpected shared data in trait binding");
+					AddErrorIfFalse(Binding0AFromTrait.GetInstanceData<FTraitA_Base::FInstanceData>()->TraitUID == FTraitA_Base::TraitUID, "FAnimationAnimNextRuntimeTest_TraitStackBinding -> Unexpected instance data in trait binding");
+					AddErrorIfFalse(!Binding0AFromTrait.GetStackInterfaceSuper(Binding0AFromTrait), "FAnimationAnimNextRuntimeTest_TraitStackBinding -> InterfaceA should not be found");
 				}
-			}
-		}
 
-		Registry.Unregister(NodeTemplate0);
-
-		AddErrorIfFalse(Registry.GetNum() == 0, "FAnimationAnimNextRuntimeTest_GetTraitInterface -> Registry should contain 0 templates");
-	}
-
-	Tests::FUtils::CleanupAfterTests();
-	
-	return true;
-}
-
-IMPLEMENT_SIMPLE_AUTOMATION_TEST(FAnimationAnimNextRuntimeTest_GetTraitInterfaceSuper, "Animation.AnimNext.Runtime.GetTraitInterfaceSuper", EAutomationTestFlags::EditorContext | EAutomationTestFlags::EngineFilter)
-
-bool FAnimationAnimNextRuntimeTest_GetTraitInterfaceSuper::RunTest(const FString& InParameters)
-{
-	using namespace UE::AnimNext;
-
-	{
-		AUTO_REGISTER_ANIM_TRAIT(FTraitA_Base)
-		AUTO_REGISTER_ANIM_TRAIT(FTraitAB_Add)
-		AUTO_REGISTER_ANIM_TRAIT(FTraitAC_Add)
-
-		UFactory* GraphFactory = NewObject<UAnimNextGraphFactory>();
-		UAnimNextGraph* AnimNextGraph = CastChecked<UAnimNextGraph>(GraphFactory->FactoryCreateNew(UAnimNextGraph::StaticClass(), GetTransientPackage(), TEXT("TestAnimNextGraph"), RF_Transient, nullptr, nullptr, NAME_None));
-		UE_RETURN_ON_ERROR(AnimNextGraph != nullptr, "FAnimationAnimNextRuntimeTest_GetTraitInterfaceSuper -> Failed to create animation graph");
-
-		FScopedClearNodeTemplateRegistry ScopedClearNodeTemplateRegistry;
-		FNodeTemplateRegistry& Registry = FNodeTemplateRegistry::Get();
-
-		TArray<FTraitUID> NodeTemplateTraitList;
-		NodeTemplateTraitList.Add(FTraitA_Base::TraitUID);
-		NodeTemplateTraitList.Add(FTraitAB_Add::TraitUID);
-		NodeTemplateTraitList.Add(FTraitAC_Add::TraitUID);
-		NodeTemplateTraitList.Add(FTraitA_Base::TraitUID);
-		NodeTemplateTraitList.Add(FTraitAC_Add::TraitUID);
-
-		// Populate our node template registry
-		TArray<uint8> NodeTemplateBuffer0;
-		const FNodeTemplate* NodeTemplate0 = FNodeTemplateBuilder::BuildNodeTemplate(NodeTemplateTraitList, NodeTemplateBuffer0);
-
-		FNodeTemplateRegistryHandle TemplateHandle0 = Registry.FindOrAdd(NodeTemplate0);
-		AddErrorIfFalse(Registry.GetNum() == 1, "FAnimationAnimNextRuntimeTest_GetTraitInterfaceSuper -> Registry should contain 1 template");
-		AddErrorIfFalse(TemplateHandle0.IsValid(), "FAnimationAnimNextRuntimeTest_GetTraitInterfaceSuper -> Registry should contain our template");
-
-		TArray<FNodeHandle> NodeHandles;
-
-		// Write our graph
-		TArray<uint8> GraphSharedDataArchiveBuffer;
-		TArray<TObjectPtr<UObject>> GraphReferencedObjects;
-		{
-			FTraitWriter TraitWriter;
-
-			NodeHandles.Add(TraitWriter.RegisterNode(*NodeTemplate0));
-
-			// We don't have trait properties
-
-			TraitWriter.BeginNodeWriting();
-			TraitWriter.WriteNode(NodeHandles[0],
-				[](uint32 TraitIndex, FName PropertyName)
+				// Trait binding misc
 				{
-					return FString();
-				},
-				[](uint32 TraitIndex, FName PropertyName)
-				{
-					return MAX_uint16;
-				});
-			TraitWriter.EndNodeWriting();
-
-			AddErrorIfFalse(TraitWriter.GetErrorState() == FTraitWriter::EErrorState::None, "FAnimationAnimNextRuntimeTest_GetTraitInterfaceSuper -> Failed to write traits");
-			GraphSharedDataArchiveBuffer = TraitWriter.GetGraphSharedData();
-			GraphReferencedObjects = TraitWriter.GetGraphReferencedObjects();
-		}
-
-		// Read our graph
-		FTestUtils::LoadFromArchiveBuffer(*AnimNextGraph, NodeHandles, GraphSharedDataArchiveBuffer);
-
-		FAnimNextGraphInstancePtr GraphInstance;
-		AnimNextGraph->AllocateInstance(GraphInstance);
+					FTraitBinding TraitBinding;
+					AddErrorIfFalse(Stack0.GetTrait(0, TraitBinding), "FAnimationAnimNextRuntimeTest_TraitStackBinding -> Failed to retrieve trait");
+					AddErrorIfFalse(TraitBinding.IsValid(), "FAnimationAnimNextRuntimeTest_TraitStackBinding -> Trait binding should be valid");
+
+					FTraitBinding TraitBinding1 = TraitBinding;
+					AddErrorIfFalse(TraitBinding1.IsValid(), "FAnimationAnimNextRuntimeTest_TraitStackBinding -> Trait binding should be valid");
+					AddErrorIfFalse(TraitBinding == TraitBinding1, "FAnimationAnimNextRuntimeTest_TraitStackBinding -> Trait bindings should be equal");
+					AddErrorIfFalse(!(TraitBinding != TraitBinding1), "FAnimationAnimNextRuntimeTest_TraitStackBinding -> Trait bindings should be equal");
+				}
+			}
 
-		FExecutionContext Context(GraphInstance);
+			FAnimNextTraitHandle TraitHandle3(NodeHandles[0], 3);		// Point to first node, second base trait
 
-		// Validate from the first base trait
-		{
-			FTraitBinding ParentBinding;								// Empty, no parent
-			FAnimNextTraitHandle TraitHandle00(NodeHandles[0], 0);		// Point to first node, first base trait
+			FTraitPtr TraitPtr3 = Context.AllocateNodeInstance(*GraphInstance.GetImpl(), TraitHandle3);
+			AddErrorIfFalse(TraitPtr3.IsValid(), "FAnimationAnimNextRuntimeTest_TraitStackBinding -> Failed to allocate a node instance");
 
-			FTraitPtr TraitPtr00 = Context.AllocateNodeInstance(ParentBinding, TraitHandle00);
-			AddErrorIfFalse(TraitPtr00.IsValid(), "FAnimationAnimNextRuntimeTest_GetTraitInterfaceSuper -> Failed to allocate a node instance");
+			FTraitStackBinding Stack3;
+			AddErrorIfFalse(Context.GetStack(TraitPtr3, Stack3), "FAnimationAnimNextRuntimeTest_TraitStackBinding -> Failed to bind to trait stack");
 
+			// Test the second stack
 			{
-				// Get a valid trait binding: FTraitAC_Add
-				TTraitBinding<IInterfaceC> Binding02C;
-				AddErrorIfFalse(Context.GetInterface(TraitPtr00, Binding02C), "FAnimationAnimNextRuntimeTest_GetTraitInterfaceSuper -> InterfaceC not found");
-
-				// Validate GetInterfaceSuper from a trait handle
-				TTraitBinding<IInterfaceC> SuperBinding02C;
-				AddErrorIfFalse(!Context.GetInterfaceSuper(Binding02C.GetTraitPtr(), SuperBinding02C), "FAnimationAnimNextRuntimeTest_GetTraitInterfaceSuper -> InterfaceC found");
-
-				// Validate GetInterfaceSuper from a trait binding
-				TTraitBinding<IInterfaceC> SuperBinding02C_;
-				AddErrorIfFalse(!Context.GetInterfaceSuper(Binding02C, SuperBinding02C_), "FAnimationAnimNextRuntimeTest_GetTraitInterfaceSuper -> InterfaceC found");
-				AddErrorIfFalse(SuperBinding02C == SuperBinding02C_, "FAnimationAnimNextRuntimeTest_GetTraitInterfaceSuper -> GetInterfaceSuper methods should return the same result");
-			}
+				AddErrorIfFalse(Stack3.IsValid(), "FAnimationAnimNextRuntimeTest_TraitStackBinding -> Trait stack binding should be valid");
+				AddErrorIfFalse(Stack3.GetBaseTraitPtr() == TraitPtr3, "FAnimationAnimNextRuntimeTest_TraitStackBinding -> Unexpected base trait ptr");
 
-			{
-				// Get a valid trait binding: FTraitAC_Add
-				TTraitBinding<IInterfaceA> Binding02A;
-				AddErrorIfFalse(Context.GetInterface(TraitPtr00, Binding02A), "FAnimationAnimNextRuntimeTest_GetTraitInterfaceSuper -> InterfaceA not found");
-
-				// Validate GetInterfaceSuper from a trait handle, FTraitAB_Add
-				TTraitBinding<IInterfaceA> SuperBinding02A;
-				AddErrorIfFalse(Context.GetInterfaceSuper(Binding02A.GetTraitPtr(), SuperBinding02A), "FAnimationAnimNextRuntimeTest_GetTraitInterfaceSuper -> InterfaceA not found");
-				AddErrorIfFalse(SuperBinding02A.IsValid(), "FAnimationAnimNextRuntimeTest_GetTraitInterfaceSuper -> InterfaceA binding not valid");
-				AddErrorIfFalse(SuperBinding02A.GetInterfaceUID() == IInterfaceA::InterfaceUID, "FAnimationAnimNextRuntimeTest_GetTraitInterfaceSuper -> Unexpected interface UID found in trait binding");
-				AddErrorIfFalse(SuperBinding02A.GetTraitPtr().GetTraitIndex() == 1, "FAnimationAnimNextRuntimeTest_GetTraitInterfaceSuper -> InterfaceA not found on expected trait");
-				AddErrorIfFalse(SuperBinding02A.GetTraitPtr().GetNodeInstance()->GetNodeHandle() == NodeHandles[0], "FAnimationAnimNextRuntimeTest_GetTraitInterfaceSuper -> InterfaceA not found on expected node");
-				AddErrorIfFalse(SuperBinding02A.GetSharedData<FTraitAB_Add::FSharedData>()->TraitUID == FTraitAB_Add::TraitUID, "FAnimationAnimNextRuntimeTest_GetTraitInterfaceSuper -> Unexpected shared data in trait binding");
-				AddErrorIfFalse(SuperBinding02A.GetInstanceData<FTraitAB_Add::FInstanceData>()->TraitUID == FTraitAB_Add::TraitUID, "FAnimationAnimNextRuntimeTest_GetTraitInterfaceSuper -> Unexpected instance data in trait binding");
-
-				// Validate GetInterfaceSuper from a trait binding, FTraitAB_Add
-				TTraitBinding<IInterfaceA> SuperBinding02A_;
-				AddErrorIfFalse(Context.GetInterfaceSuper(Binding02A, SuperBinding02A_), "FAnimationAnimNextRuntimeTest_GetTraitInterfaceSuper -> InterfaceA not found");
-				AddErrorIfFalse(SuperBinding02A == SuperBinding02A_, "FAnimationAnimNextRuntimeTest_GetTraitInterfaceSuper -> GetInterfaceSuper methods should return the same result");
-
-				// Validate GetInterfaceSuper from a trait handle, FTraitA_Base
-				TTraitBinding<IInterfaceA> SuperBinding01A;
-				AddErrorIfFalse(Context.GetInterfaceSuper(SuperBinding02A.GetTraitPtr(), SuperBinding01A), "FAnimationAnimNextRuntimeTest_GetTraitInterfaceSuper -> InterfaceA not found");
-				AddErrorIfFalse(SuperBinding01A.IsValid(), "FAnimationAnimNextRuntimeTest_GetTraitInterfaceSuper -> InterfaceA binding not valid");
-				AddErrorIfFalse(SuperBinding01A.GetInterfaceUID() == IInterfaceA::InterfaceUID, "FAnimationAnimNextRuntimeTest_GetTraitInterfaceSuper -> Unexpected interface UID found in trait binding");
-				AddErrorIfFalse(SuperBinding01A.GetTraitPtr().GetTraitIndex() == 0, "FAnimationAnimNextRuntimeTest_GetTraitInterfaceSuper -> InterfaceA not found on expected trait");
-				AddErrorIfFalse(SuperBinding01A.GetTraitPtr().GetNodeInstance()->GetNodeHandle() == NodeHandles[0], "FAnimationAnimNextRuntimeTest_GetTraitInterfaceSuper -> InterfaceA not found on expected node");
-				AddErrorIfFalse(SuperBinding01A.GetSharedData<FTraitA_Base::FSharedData>()->TraitUID == FTraitA_Base::TraitUID, "FAnimationAnimNextRuntimeTest_GetTraitInterfaceSuper -> Unexpected shared data in trait binding");
-				AddErrorIfFalse(SuperBinding01A.GetInstanceData<FTraitA_Base::FInstanceData>()->TraitUID == FTraitA_Base::TraitUID, "FAnimationAnimNextRuntimeTest_GetTraitInterfaceSuper -> Unexpected instance data in trait binding");
-
-				// Validate GetInterfaceSuper from a trait binding, FTraitA_Base
-				TTraitBinding<IInterfaceA> SuperBinding01A_;
-				AddErrorIfFalse(Context.GetInterfaceSuper(SuperBinding02A, SuperBinding01A_), "FAnimationAnimNextRuntimeTest_GetTraitInterfaceSuper -> InterfaceA not found");
-				AddErrorIfFalse(SuperBinding01A == SuperBinding01A_, "FAnimationAnimNextRuntimeTest_GetTraitInterfaceSuper -> GetInterfaceSuper methods should return the same result");
-
-				// Validate GetInterfaceSuper from a trait handle
-				TTraitBinding<IInterfaceA> SuperBinding00A;
-				AddErrorIfFalse(!Context.GetInterfaceSuper(SuperBinding01A.GetTraitPtr(), SuperBinding00A), "FAnimationAnimNextRuntimeTest_GetTraitInterfaceSuper -> InterfaceA found");
-
-				// Validate GetInterfaceSuper from a trait binding
-				TTraitBinding<IInterfaceA> SuperBinding00A_;
-				AddErrorIfFalse(!Context.GetInterfaceSuper(SuperBinding01A, SuperBinding00A_), "FAnimationAnimNextRuntimeTest_GetTraitInterfaceSuper -> InterfaceA found");
-				AddErrorIfFalse(SuperBinding00A == SuperBinding00A_, "FAnimationAnimNextRuntimeTest_GetTraitInterfaceSuper -> GetInterfaceSuper methods should return the same result");
-			}
-		}
+				// Stack binding copying and equality testing
+				{
+					AddErrorIfFalse(Stack0 != Stack3, "FAnimationAnimNextRuntimeTest_TraitStackBinding -> Trait stack bindings should not be equal");
+
+					FTraitStackBinding Stack3Copy = Stack3;
+					AddErrorIfFalse(Stack3Copy.IsValid(), "FAnimationAnimNextRuntimeTest_TraitStackBinding -> Trait stack binding should be valid");
+					AddErrorIfFalse(Stack3Copy.GetBaseTraitPtr() == TraitPtr3, "FAnimationAnimNextRuntimeTest_TraitStackBinding -> Unexpected base trait ptr");
+					AddErrorIfFalse(Stack3 == Stack3Copy, "FAnimationAnimNextRuntimeTest_TraitStackBinding -> Trait stack bindings should be equal");
+					AddErrorIfFalse(!(Stack3 != Stack3Copy), "FAnimationAnimNextRuntimeTest_TraitStackBinding -> Trait stack bindings should be equal");
+
+					Stack3Copy.Reset();
+					AddErrorIfFalse(!Stack3Copy.IsValid(), "FAnimationAnimNextRuntimeTest_TraitStackBinding -> Trait stack binding should be not valid");
+
+					FTraitStackBinding Stack4;
+					AddErrorIfFalse(Context.GetStack(FWeakTraitPtr(TraitPtr3.GetNodeInstance(), 4), Stack4), "FAnimationAnimNextRuntimeTest_TraitStackBinding -> Failed to bind to trait stack");
+					AddErrorIfFalse(Stack4.IsValid(), "FAnimationAnimNextRuntimeTest_TraitStackBinding -> Trait stack binding should be valid");
+					AddErrorIfFalse(Stack4.GetBaseTraitPtr() == TraitPtr3, "FAnimationAnimNextRuntimeTest_TraitStackBinding -> Unexpected base trait ptr");
+					AddErrorIfFalse(Stack3 == Stack4, "FAnimationAnimNextRuntimeTest_TraitStackBinding -> Trait stack bindings should be equal");
+				}
 
-		// Validate from the second base trait
-		{
-			FTraitBinding ParentBinding;								// Empty, no parent
-			FAnimNextTraitHandle TraitHandle03(NodeHandles[0], 3);		// Point to first node, second base trait
+				// Stack binding iteration from top to bottom
+				{
+					FTraitBinding TraitBinding;
+					AddErrorIfFalse(Stack3.GetTopTrait(TraitBinding), "FAnimationAnimNextRuntimeTest_TraitStackBinding -> Failed to retrieve top trait");
+					AddErrorIfFalse(TraitBinding.IsValid(), "FAnimationAnimNextRuntimeTest_TraitStackBinding -> Trait binding should be valid");
+					AddErrorIfFalse(TraitBinding.GetStack() == &Stack3, "FAnimationAnimNextRuntimeTest_TraitStackBinding -> Unexpected trait stack binding");
+					AddErrorIfFalse(TraitBinding.GetTraitIndex() == 1, "FAnimationAnimNextRuntimeTest_TraitStackBinding -> Unexpected stack local trait index");
+					AddErrorIfFalse(TraitBinding.GetTraitPtr().GetTraitIndex() == 4, "FAnimationAnimNextRuntimeTest_TraitStackBinding -> Unexpected node local trait index");
+
+					AddErrorIfFalse(Stack3.GetParentTrait(TraitBinding, TraitBinding), "FAnimationAnimNextRuntimeTest_TraitStackBinding -> Failed to retrieve parent trait");
+					AddErrorIfFalse(TraitBinding.IsValid(), "FAnimationAnimNextRuntimeTest_TraitStackBinding -> Trait binding should be valid");
+					AddErrorIfFalse(TraitBinding.GetStack() == &Stack3, "FAnimationAnimNextRuntimeTest_TraitStackBinding -> Unexpected trait stack binding");
+					AddErrorIfFalse(TraitBinding.GetTraitIndex() == 0, "FAnimationAnimNextRuntimeTest_TraitStackBinding -> Unexpected stack local trait index");
+					AddErrorIfFalse(TraitBinding.GetTraitPtr().GetTraitIndex() == 3, "FAnimationAnimNextRuntimeTest_TraitStackBinding -> Unexpected node local trait index");
+					AddErrorIfFalse(TraitBinding.GetTraitPtr() == Stack3.GetBaseTraitPtr(), "FAnimationAnimNextRuntimeTest_TraitStackBinding -> Unexpected stack base trait ptr");
+
+					AddErrorIfFalse(!Stack3.GetParentTrait(TraitBinding, TraitBinding), "FAnimationAnimNextRuntimeTest_TraitStackBinding -> Unexpected parent trait");
+				}
 
-			FTraitPtr TraitPtr03 = Context.AllocateNodeInstance(ParentBinding, TraitHandle03);
-			AddErrorIfFalse(TraitPtr03.IsValid(), "FAnimationAnimNextRuntimeTest_GetTraitInterfaceSuper -> Failed to allocate a node instance");
+				// Stack binding iteration from bottom to top
+				{
+					FTraitBinding TraitBinding;
+					AddErrorIfFalse(Stack3.GetBaseTrait(TraitBinding), "FAnimationAnimNextRuntimeTest_TraitStackBinding -> Failed to retrieve base trait");
+					AddErrorIfFalse(TraitBinding.IsValid(), "FAnimationAnimNextRuntimeTest_TraitStackBinding -> Trait binding should be valid");
+					AddErrorIfFalse(TraitBinding.GetStack() == &Stack3, "FAnimationAnimNextRuntimeTest_TraitStackBinding -> Unexpected trait stack binding");
+					AddErrorIfFalse(TraitBinding.GetTraitIndex() == 0, "FAnimationAnimNextRuntimeTest_TraitStackBinding -> Unexpected stack local trait index");
+					AddErrorIfFalse(TraitBinding.GetTraitPtr().GetTraitIndex() == 3, "FAnimationAnimNextRuntimeTest_TraitStackBinding -> Unexpected node local trait index");
+					AddErrorIfFalse(TraitBinding.GetTraitPtr() == Stack3.GetBaseTraitPtr(), "FAnimationAnimNextRuntimeTest_TraitStackBinding -> Unexpected stack base trait ptr");
+
+					AddErrorIfFalse(Stack3.GetChildTrait(TraitBinding, TraitBinding), "FAnimationAnimNextRuntimeTest_TraitStackBinding -> Failed to retrieve child trait");
+					AddErrorIfFalse(TraitBinding.IsValid(), "FAnimationAnimNextRuntimeTest_TraitStackBinding -> Trait binding should be valid");
+					AddErrorIfFalse(TraitBinding.GetStack() == &Stack3, "FAnimationAnimNextRuntimeTest_TraitStackBinding -> Unexpected trait stack binding");
+					AddErrorIfFalse(TraitBinding.GetTraitIndex() == 1, "FAnimationAnimNextRuntimeTest_TraitStackBinding -> Unexpected stack local trait index");
+					AddErrorIfFalse(TraitBinding.GetTraitPtr().GetTraitIndex() == 4, "FAnimationAnimNextRuntimeTest_TraitStackBinding -> Unexpected node local trait index");
+
+					AddErrorIfFalse(!Stack3.GetChildTrait(TraitBinding, TraitBinding), "FAnimationAnimNextRuntimeTest_TraitStackBinding -> Unexpected child trait");
+				}
 
-			{
-				// Get a valid trait binding: FTraitAC_Add
-				TTraitBinding<IInterfaceC> Binding04C;
-				AddErrorIfFalse(Context.GetInterface(TraitPtr03, Binding04C), "FAnimationAnimNextRuntimeTest_GetTraitInterfaceSuper -> InterfaceC not found");
-
-				// Validate GetInterfaceSuper from a trait handle
-				TTraitBinding<IInterfaceC> SuperBinding04C;
-				AddErrorIfFalse(!Context.GetInterfaceSuper(Binding04C.GetTraitPtr(), SuperBinding04C), "FAnimationAnimNextRuntimeTest_GetTraitInterfaceSuper -> InterfaceC found");
-
-				// Validate GetInterfaceSuper from a trait binding
-				TTraitBinding<IInterfaceC> SuperBinding04C_;
-				AddErrorIfFalse(!Context.GetInterfaceSuper(Binding04C, SuperBinding04C_), "FAnimationAnimNextRuntimeTest_GetTraitInterfaceSuper -> InterfaceC found");
-				AddErrorIfFalse(SuperBinding04C == SuperBinding04C_, "FAnimationAnimNextRuntimeTest_GetTraitInterfaceSuper -> GetInterfaceSuper methods should return the same result");
-			}
+				// Stack binding arbitrary iteration
+				{
+					AddErrorIfFalse(Stack3.GetNumTraits() == 2, "FAnimationAnimNextRuntimeTest_TraitStackBinding -> Unexpected trait stack size");
+
+					FTraitBinding TraitBinding;
+					AddErrorIfFalse(Stack3.GetTrait(0, TraitBinding), "FAnimationAnimNextRuntimeTest_TraitStackBinding -> Failed to retrieve trait");
+					AddErrorIfFalse(TraitBinding.IsValid(), "FAnimationAnimNextRuntimeTest_TraitStackBinding -> Trait binding should be valid");
+					AddErrorIfFalse(TraitBinding.GetStack() == &Stack3, "FAnimationAnimNextRuntimeTest_TraitStackBinding -> Unexpected trait stack binding");
+					AddErrorIfFalse(TraitBinding.GetTraitIndex() == 0, "FAnimationAnimNextRuntimeTest_TraitStackBinding -> Unexpected stack local trait index");
+					AddErrorIfFalse(TraitBinding.GetTraitPtr().GetTraitIndex() == 3, "FAnimationAnimNextRuntimeTest_TraitStackBinding -> Unexpected node local trait index");
+					AddErrorIfFalse(TraitBinding.GetTraitPtr() == Stack3.GetBaseTraitPtr(), "FAnimationAnimNextRuntimeTest_TraitStackBinding -> Unexpected stack base trait ptr");
+
+					AddErrorIfFalse(Stack3.GetTrait(1, TraitBinding), "FAnimationAnimNextRuntimeTest_TraitStackBinding -> Failed to retrieve trait");
+					AddErrorIfFalse(TraitBinding.IsValid(), "FAnimationAnimNextRuntimeTest_TraitStackBinding -> Trait binding should be valid");
+					AddErrorIfFalse(TraitBinding.GetStack() == &Stack3, "FAnimationAnimNextRuntimeTest_TraitStackBinding -> Unexpected trait stack binding");
+					AddErrorIfFalse(TraitBinding.GetTraitIndex() == 1, "FAnimationAnimNextRuntimeTest_TraitStackBinding -> Unexpected stack local trait index");
+					AddErrorIfFalse(TraitBinding.GetTraitPtr().GetTraitIndex() == 4, "FAnimationAnimNextRuntimeTest_TraitStackBinding -> Unexpected node local trait index");
+				}
 
-			{
-				// Get a valid trait binding: FTraitAC_Add
-				TTraitBinding<IInterfaceA> Binding04A;
-				AddErrorIfFalse(Context.GetInterface(TraitPtr03, Binding04A), "FAnimationAnimNextRuntimeTest_GetTraitInterfaceSuper -> InterfaceA not found");
-
-				// Validate GetInterfaceSuper from a trait handle, FTraitA_Base
-				TTraitBinding<IInterfaceA> SuperBinding04A;
-				AddErrorIfFalse(Context.GetInterfaceSuper(Binding04A.GetTraitPtr(), SuperBinding04A), "FAnimationAnimNextRuntimeTest_GetTraitInterfaceSuper -> InterfaceA not found");
-				AddErrorIfFalse(SuperBinding04A.IsValid(), "FAnimationAnimNextRuntimeTest_GetTraitInterfaceSuper -> InterfaceA binding not valid");
-				AddErrorIfFalse(SuperBinding04A.GetInterfaceUID() == IInterfaceA::InterfaceUID, "FAnimationAnimNextRuntimeTest_GetTraitInterfaceSuper -> Unexpected interface UID found in trait binding");
-				AddErrorIfFalse(SuperBinding04A.GetTraitPtr().GetTraitIndex() == 3, "FAnimationAnimNextRuntimeTest_GetTraitInterfaceSuper -> InterfaceA not found on expected trait");
-				AddErrorIfFalse(SuperBinding04A.GetTraitPtr().GetNodeInstance()->GetNodeHandle() == NodeHandles[0], "FAnimationAnimNextRuntimeTest_GetTraitInterfaceSuper -> InterfaceA not found on expected node");
-				AddErrorIfFalse(SuperBinding04A.GetSharedData<FTraitA_Base::FSharedData>()->TraitUID == FTraitA_Base::TraitUID, "FAnimationAnimNextRuntimeTest_GetTraitInterfaceSuper -> Unexpected shared data in trait binding");
-				AddErrorIfFalse(SuperBinding04A.GetInstanceData<FTraitA_Base::FInstanceData>()->TraitUID == FTraitA_Base::TraitUID, "FAnimationAnimNextRuntimeTest_GetTraitInterfaceSuper -> Unexpected instance data in trait binding");
-
-				// Validate GetInterfaceSuper from a trait binding, FTraitA_Base
-				TTraitBinding<IInterfaceA> SuperBinding04A_;
-				AddErrorIfFalse(Context.GetInterfaceSuper(Binding04A, SuperBinding04A_), "FAnimationAnimNextRuntimeTest_GetTraitInterfaceSuper -> InterfaceA not found");
-				AddErrorIfFalse(SuperBinding04A == SuperBinding04A_, "FAnimationAnimNextRuntimeTest_GetTraitInterfaceSuper -> GetInterfaceSuper methods should return the same result");
-
-				// Validate GetInterfaceSuper from a trait handle
-				TTraitBinding<IInterfaceA> SuperBinding03A;
-				AddErrorIfFalse(!Context.GetInterfaceSuper(SuperBinding04A.GetTraitPtr(), SuperBinding03A), "FAnimationAnimNextRuntimeTest_GetTraitInterfaceSuper -> InterfaceA found");
-
-				// Validate GetInterfaceSuper from a trait binding
-				TTraitBinding<IInterfaceA> SuperBinding03A_;
-				AddErrorIfFalse(!Context.GetInterfaceSuper(SuperBinding04A, SuperBinding03A_), "FAnimationAnimNextRuntimeTest_GetTraitInterfaceSuper -> InterfaceA found");
-				AddErrorIfFalse(SuperBinding03A == SuperBinding03A_, "FAnimationAnimNextRuntimeTest_GetTraitInterfaceSuper -> GetInterfaceSuper methods should return the same result");
+				// AC -> A
+				// Interface query validation
+				{
+					TTraitBinding<IInterfaceC> Binding3C;
+					AddErrorIfFalse(Stack3.GetInterface(Binding3C), "FAnimationAnimNextRuntimeTest_TraitStackBinding -> InterfaceC not found");
+					AddErrorIfFalse(Binding3C.IsValid(), "FAnimationAnimNextRuntimeTest_TraitStackBinding -> InterfaceC binding not valid");
+					AddErrorIfFalse(Binding3C.GetInterfaceUID() == IInterfaceC::InterfaceUID, "FAnimationAnimNextRuntimeTest_TraitStackBinding -> Unexpected interface UID found in trait binding");
+					AddErrorIfFalse(Binding3C.GetTraitPtr().GetTraitIndex() == 4, "FAnimationAnimNextRuntimeTest_TraitStackBinding -> InterfaceC not found on expected trait");
+					AddErrorIfFalse(Binding3C.GetTraitPtr().GetNodeInstance()->GetNodeHandle() == NodeHandles[0], "FAnimationAnimNextRuntimeTest_TraitStackBinding -> InterfaceC not found on expected node");
+					AddErrorIfFalse(Binding3C.GetSharedData<FTraitAC_Add::FSharedData>()->TraitUID == FTraitAC_Add::TraitUID, "FAnimationAnimNextRuntimeTest_TraitStackBinding -> Unexpected shared data in trait binding");
+					AddErrorIfFalse(Binding3C.GetInstanceData<FTraitAC_Add::FInstanceData>()->TraitUID == FTraitAC_Add::TraitUID, "FAnimationAnimNextRuntimeTest_TraitStackBinding -> Unexpected instance data in trait binding");
+					AddErrorIfFalse(!Stack3.GetInterfaceSuper(Binding3C, Binding3C), "FAnimationAnimNextRuntimeTest_TraitStackBinding -> InterfaceC should not be found");
+
+					TTraitBinding<IInterfaceB> Binding3B;
+					AddErrorIfFalse(!Stack3.GetInterface(Binding3B), "FAnimationAnimNextRuntimeTest_TraitStackBinding -> InterfaceB should not be found");
+					AddErrorIfFalse(!Binding3B.IsValid(), "FAnimationAnimNextRuntimeTest_TraitStackBinding -> InterfaceB binding not valid");
+
+					TTraitBinding<IInterfaceA> Binding3A;
+					AddErrorIfFalse(Stack3.GetInterface(Binding3A), "FAnimationAnimNextRuntimeTest_TraitStackBinding -> InterfaceA not found");
+					AddErrorIfFalse(Binding3A.IsValid(), "FAnimationAnimNextRuntimeTest_TraitStackBinding -> InterfaceA binding not valid");
+					AddErrorIfFalse(Binding3A.GetInterfaceUID() == IInterfaceA::InterfaceUID, "FAnimationAnimNextRuntimeTest_TraitStackBinding -> Unexpected interface UID found in trait binding");
+					AddErrorIfFalse(Binding3A.GetTraitPtr().GetTraitIndex() == 4, "FAnimationAnimNextRuntimeTest_TraitStackBinding -> InterfaceA not found on expected trait");
+					AddErrorIfFalse(Binding3A.GetTraitPtr().GetNodeInstance()->GetNodeHandle() == NodeHandles[0], "FAnimationAnimNextRuntimeTest_TraitStackBinding -> InterfaceA not found on expected node");
+					AddErrorIfFalse(Binding3A.GetSharedData<FTraitAC_Add::FSharedData>()->TraitUID == FTraitAC_Add::TraitUID, "FAnimationAnimNextRuntimeTest_TraitStackBinding -> Unexpected shared data in trait binding");
+					AddErrorIfFalse(Binding3A.GetInstanceData<FTraitAC_Add::FInstanceData>()->TraitUID == FTraitAC_Add::TraitUID, "FAnimationAnimNextRuntimeTest_TraitStackBinding -> Unexpected instance data in trait binding");
+
+					AddErrorIfFalse(Stack3.GetInterfaceSuper(Binding3A, Binding3A), "FAnimationAnimNextRuntimeTest_TraitStackBinding -> InterfaceA not found");
+					AddErrorIfFalse(Binding3A.IsValid(), "FAnimationAnimNextRuntimeTest_TraitStackBinding -> InterfaceA binding not valid");
+					AddErrorIfFalse(Binding3A.GetInterfaceUID() == IInterfaceA::InterfaceUID, "FAnimationAnimNextRuntimeTest_TraitStackBinding -> Unexpected interface UID found in trait binding");
+					AddErrorIfFalse(Binding3A.GetTraitPtr().GetTraitIndex() == 3, "FAnimationAnimNextRuntimeTest_TraitStackBinding -> InterfaceA not found on expected trait");
+					AddErrorIfFalse(Binding3A.GetTraitPtr().GetNodeInstance()->GetNodeHandle() == NodeHandles[0], "FAnimationAnimNextRuntimeTest_TraitStackBinding -> InterfaceA not found on expected node");
+					AddErrorIfFalse(Binding3A.GetSharedData<FTraitA_Base::FSharedData>()->TraitUID == FTraitA_Base::TraitUID, "FAnimationAnimNextRuntimeTest_TraitStackBinding -> Unexpected shared data in trait binding");
+					AddErrorIfFalse(Binding3A.GetInstanceData<FTraitA_Base::FInstanceData>()->TraitUID == FTraitA_Base::TraitUID, "FAnimationAnimNextRuntimeTest_TraitStackBinding -> Unexpected instance data in trait binding");
+					AddErrorIfFalse(!Stack3.GetInterfaceSuper(Binding3A, Binding3A), "FAnimationAnimNextRuntimeTest_TraitStackBinding -> InterfaceA should not be found");
+				}
 			}
 		}
 
 		Registry.Unregister(NodeTemplate0);
 
-		AddErrorIfFalse(Registry.GetNum() == 0, "FAnimationAnimNextRuntimeTest_GetTraitInterface -> Registry should contain 0 templates");
+		AddErrorIfFalse(Registry.GetNum() == 0, "FAnimationAnimNextRuntimeTest_TraitStackBinding -> Registry should contain 0 templates");
 	}
 
 	Tests::FUtils::CleanupAfterTests();
-	
+
 	return true;
 }
 
@@ -1425,20 +1368,25 @@ bool FAnimationAnimNextRuntimeTest_TraitSerialization::RunTest(const FString& In
 
 			// Validate trait serialization
 			{
-				FTraitBinding ParentBinding;								// Empty, no parent
 				FAnimNextTraitHandle TraitHandle0(NodeHandles[0], 0);		// Point to first node, first base trait
 				FAnimNextTraitHandle TraitHandle1(NodeHandles[1], 0);		// Point to second node, first base trait
 
-				FTraitPtr TraitPtr0 = Context.AllocateNodeInstance(ParentBinding, TraitHandle0);
+				FTraitPtr TraitPtr0 = Context.AllocateNodeInstance(*GraphInstance.GetImpl(), TraitHandle0);
 				AddErrorIfFalse(TraitPtr0.IsValid(), "FAnimationAnimNextRuntimeTest_TraitSerialization -> Failed to allocate a node instance");
 
-				FTraitPtr TraitPtr1 = Context.AllocateNodeInstance(ParentBinding, TraitHandle1);
+				FTraitPtr TraitPtr1 = Context.AllocateNodeInstance(*GraphInstance.GetImpl(), TraitHandle1);
 				AddErrorIfFalse(TraitPtr1.IsValid(), "FAnimationAnimNextRuntimeTest_TraitSerialization -> Failed to allocate a node instance");
 
+				FTraitStackBinding StackNode0;
+				AddErrorIfFalse(Context.GetStack(TraitPtr0, StackNode0), "FAnimationAnimNextRuntimeTest_TraitSerialization -> Failed to bind to trait stack");
+
+				FTraitStackBinding StackNode1;
+				AddErrorIfFalse(Context.GetStack(TraitPtr1, StackNode1), "FAnimationAnimNextRuntimeTest_TraitSerialization -> Failed to bind to trait stack");
+
 				// Validate shared data for base trait on node 0
 				{
 					TTraitBinding<IInterfaceA> BindingA0;
-					AddErrorIfFalse(Context.GetInterface(TraitPtr0, BindingA0), "FAnimationAnimNextRuntimeTest_TraitSerialization -> InterfaceA not found");
+					AddErrorIfFalse(StackNode0.GetInterface(BindingA0), "FAnimationAnimNextRuntimeTest_TraitSerialization -> InterfaceA not found");
 
 					const auto* SharedDataA0 = BindingA0.GetSharedData<FTraitSerialization_Base::FSharedData>();
 
@@ -1455,7 +1403,7 @@ bool FAnimationAnimNextRuntimeTest_TraitSerialization::RunTest(const FString& In
 				// Validate shared data for additive trait on node 0
 				{
 					TTraitBinding<IInterfaceB> BindingB0;
-					AddErrorIfFalse(Context.GetInterface(TraitPtr0, BindingB0), "FAnimationAnimNextRuntimeTest_TraitSerialization -> InterfaceB not found");
+					AddErrorIfFalse(StackNode0.GetInterface(BindingB0), "FAnimationAnimNextRuntimeTest_TraitSerialization -> InterfaceB not found");
 
 					const auto* SharedDataB0 = BindingB0.GetSharedData<FTraitSerialization_Add::FSharedData>();
 
@@ -1472,7 +1420,7 @@ bool FAnimationAnimNextRuntimeTest_TraitSerialization::RunTest(const FString& In
 				// Validate shared data for native trait on node 0
 				{
 					TTraitBinding<IInterfaceC> BindingC0;
-					AddErrorIfFalse(Context.GetInterface(TraitPtr0, BindingC0), "FAnimationAnimNextRuntimeTest_TraitSerialization -> InterfaceC not found");
+					AddErrorIfFalse(StackNode0.GetInterface(BindingC0), "FAnimationAnimNextRuntimeTest_TraitSerialization -> InterfaceC not found");
 
 					const auto* SharedDataC0 = BindingC0.GetSharedData<FTraitNativeSerialization_Add::FSharedData>();
 
@@ -1490,7 +1438,7 @@ bool FAnimationAnimNextRuntimeTest_TraitSerialization::RunTest(const FString& In
 				// Validate shared data for base trait on node 1
 				{
 					TTraitBinding<IInterfaceA> BindingA1;
-					AddErrorIfFalse(Context.GetInterface(TraitPtr1, BindingA1), "FAnimationAnimNextRuntimeTest_TraitSerialization -> InterfaceA not found");
+					AddErrorIfFalse(StackNode1.GetInterface(BindingA1), "FAnimationAnimNextRuntimeTest_TraitSerialization -> InterfaceA not found");
 
 					const auto* SharedDataA1 = BindingA1.GetSharedData<FTraitSerialization_Base::FSharedData>();
 
@@ -1507,7 +1455,7 @@ bool FAnimationAnimNextRuntimeTest_TraitSerialization::RunTest(const FString& In
 				// Validate shared data for additive trait on node 1
 				{
 					TTraitBinding<IInterfaceB> BindingB1;
-					AddErrorIfFalse(Context.GetInterface(TraitPtr1, BindingB1), "FAnimationAnimNextRuntimeTest_TraitSerialization -> InterfaceB not found");
+					AddErrorIfFalse(StackNode1.GetInterface(BindingB1), "FAnimationAnimNextRuntimeTest_TraitSerialization -> InterfaceB not found");
 
 					const auto* SharedDataB1 = BindingB1.GetSharedData<FTraitSerialization_Add::FSharedData>();
 
@@ -1524,7 +1472,7 @@ bool FAnimationAnimNextRuntimeTest_TraitSerialization::RunTest(const FString& In
 				// Validate shared data for native trait on node 1
 				{
 					TTraitBinding<IInterfaceC> BindingC1;
-					AddErrorIfFalse(Context.GetInterface(TraitPtr1, BindingC1), "FAnimationAnimNextRuntimeTest_TraitSerialization -> InterfaceC not found");
+					AddErrorIfFalse(StackNode1.GetInterface(BindingC1), "FAnimationAnimNextRuntimeTest_TraitSerialization -> InterfaceC not found");
 
 					const auto* SharedDataC1 = BindingC1.GetSharedData<FTraitNativeSerialization_Add::FSharedData>();
 
diff --git a/Engine/Plugins/Experimental/AnimNext/Source/AnimNextTestSuite/Private/AnimNextTraitGraphTest.cpp b/Engine/Plugins/Experimental/AnimNext/Source/AnimNextTestSuite/Private/AnimNextTraitGraphTest.cpp
--- a/Engine/Plugins/Experimental/AnimNext/Source/AnimNextTestSuite/Private/AnimNextTraitGraphTest.cpp
+++ b/Engine/Plugins/Experimental/AnimNext/Source/AnimNextTestSuite/Private/AnimNextTraitGraphTest.cpp
@@ -52,9 +52,9 @@ namespace UE::AnimNext
 			ParamStack.GetMutableParam<int32>("SomeInt32") = SharedData->SomeInt32;
 			ParamStack.GetMutableParam<float>("SomeFloat") = SharedData->SomeFloat;
 
-			ParamStack.GetMutableParam<int32>("SomeLatentInt32") = SharedData->GetSomeLatentInt32(Context, Binding);				// MathAdd with constants, latent
-			ParamStack.GetMutableParam<int32>("SomeOtherLatentInt32") = SharedData->GetSomeOtherLatentInt32(Context, Binding);		// GetParameter, latent
-			ParamStack.GetMutableParam<float>("SomeLatentFloat") = SharedData->GetSomeLatentFloat(Context, Binding);				// Inline value, not latent
+			ParamStack.GetMutableParam<int32>("SomeLatentInt32") = SharedData->GetSomeLatentInt32(Binding);				// MathAdd with constants, latent
+			ParamStack.GetMutableParam<int32>("SomeOtherLatentInt32") = SharedData->GetSomeOtherLatentInt32(Binding);	// GetParameter, latent
+			ParamStack.GetMutableParam<float>("SomeLatentFloat") = SharedData->GetSomeLatentFloat(Binding);				// Inline value, not latent
 		}
 
 		// IEvaluate impl
diff --git a/Engine/Plugins/Experimental/AnimNext/Source/AnimNextTestSuite/Private/AnimNextTraitGraphTest.h b/Engine/Plugins/Experimental/AnimNext/Source/AnimNextTestSuite/Private/AnimNextTraitGraphTest.h
--- a/Engine/Plugins/Experimental/AnimNext/Source/AnimNextTestSuite/Private/AnimNextTraitGraphTest.h
+++ b/Engine/Plugins/Experimental/AnimNext/Source/AnimNextTestSuite/Private/AnimNextTraitGraphTest.h
@@ -4,6 +4,7 @@
 
 #include "CoreMinimal.h"
 
+#include "TraitCore/TraitBinding.h"
 #include "TraitCore/TraitSharedData.h"
 #include "Graph/AnimNextGraph.h"
 
diff --git a/Engine/Plugins/Experimental/AnimNext/Source/AnimNextTestSuite/Private/AnimNextTraitInterfacesTest.cpp b/Engine/Plugins/Experimental/AnimNext/Source/AnimNextTestSuite/Private/AnimNextTraitInterfacesTest.cpp
--- a/Engine/Plugins/Experimental/AnimNext/Source/AnimNextTestSuite/Private/AnimNextTraitInterfacesTest.cpp
+++ b/Engine/Plugins/Experimental/AnimNext/Source/AnimNextTestSuite/Private/AnimNextTraitInterfacesTest.cpp
@@ -120,8 +120,6 @@ namespace UE::AnimNext
 			const FInstanceData* InstanceData = Binding.GetInstanceData<FInstanceData>();
 
 			Children.Add(InstanceData->Child);
-
-			IHierarchy::GetChildren(Context, Binding, Children);
 		}
 	};
 
@@ -132,7 +130,7 @@ namespace UE::AnimNext
 	GENERATE_ANIM_TRAIT_IMPLEMENTATION(FTraitWithOneChild, TRAIT_INTERFACE_ENUMERATOR)
 	#undef TRAIT_INTERFACE_ENUMERATOR
 
-	struct FTraitWithChildren : FBaseTrait, IHierarchy, IUpdate, IEvaluate
+	struct FTraitWithChildren : FBaseTrait, IHierarchy, IUpdate, IUpdateTraversal, IEvaluate
 	{
 		DECLARE_ANIM_TRAIT(FTraitWithChildren, 0x4b296948, FBaseTrait)
 
@@ -161,8 +159,6 @@ namespace UE::AnimNext
 
 			Children.Add(InstanceData->Children[0]);
 			Children.Add(InstanceData->Children[1]);
-
-			IHierarchy::GetChildren(Context, Binding, Children);
 		}
 
 		// IUpdate impl
@@ -176,14 +172,6 @@ namespace UE::AnimNext
 			IUpdate::PreUpdate(Context, Binding, TraitState);
 		}
 
-		virtual void QueueChildrenForTraversal(FUpdateTraversalContext& Context, const TTraitBinding<IUpdate>& Binding, const FTraitUpdateState& TraitState, FUpdateTraversalQueue& TraversalQueue) const override
-		{
-			const FInstanceData* InstanceData = Binding.GetInstanceData<FInstanceData>();
-
-			TraversalQueue.Push(InstanceData->Children[0], TraitState);
-			TraversalQueue.Push(InstanceData->Children[1], TraitState);
-		}
-
 		virtual void PostUpdate(FUpdateTraversalContext& Context, const TTraitBinding<IUpdate>& Binding, const FTraitUpdateState& TraitState) const override
 		{
 			if (Private::UpdatedTraits != nullptr)
@@ -194,6 +182,15 @@ namespace UE::AnimNext
 			IUpdate::PostUpdate(Context, Binding, TraitState);
 		}
 
+		// IUpdateTraversal impl
+		virtual void QueueChildrenForTraversal(FUpdateTraversalContext& Context, const TTraitBinding<IUpdateTraversal>& Binding, const FTraitUpdateState& TraitState, FUpdateTraversalQueue& TraversalQueue) const override
+		{
+			const FInstanceData* InstanceData = Binding.GetInstanceData<FInstanceData>();
+
+			TraversalQueue.Push(InstanceData->Children[0], TraitState);
+			TraversalQueue.Push(InstanceData->Children[1], TraitState);
+		}
+
 		// IEvaluate impl
 		virtual void PreEvaluate(FEvaluateTraversalContext& Context, const TTraitBinding<IEvaluate>& Binding) const override
 		{
@@ -221,6 +218,7 @@ namespace UE::AnimNext
 		GeneratorMacro(IEvaluate) \
 		GeneratorMacro(IHierarchy) \
 		GeneratorMacro(IUpdate) \
+		GeneratorMacro(IUpdateTraversal) \
 
 	GENERATE_ANIM_TRAIT_IMPLEMENTATION(FTraitWithChildren, TRAIT_INTERFACE_ENUMERATOR)
 	#undef TRAIT_INTERFACE_ENUMERATOR
@@ -361,40 +359,73 @@ bool FAnimationAnimNextRuntimeTest_IHierarchy::RunTest(const FString& InParamete
 		{
 			FMemMark Mark(FMemStack::Get());
 
-			FWeakTraitPtr NullPtr;								// Empty, no parent
 			FAnimNextTraitHandle RootHandle(NodeHandles[3]);	// Point to NodeD, first base trait
 
-			FTraitPtr NodeDPtr = Context.AllocateNodeInstance(NullPtr, RootHandle);
+			FTraitPtr NodeDPtr = Context.AllocateNodeInstance(*GraphInstance.GetImpl(), RootHandle);
 			AddErrorIfFalse(NodeDPtr.IsValid(), "FAnimationAnimNextRuntimeTest_IHierarchy -> Failed to allocate root node instance");
 
+			FTraitStackBinding StackNodeD;
+			AddErrorIfFalse(Context.GetStack(NodeDPtr, StackNodeD), "FAnimationAnimNextRuntimeTest_IHierarchy -> Failed to bind to trait stack");
+
 			TTraitBinding<IHierarchy> HierarchyBindingNodeD;
-			AddErrorIfFalse(Context.GetInterface(NodeDPtr, HierarchyBindingNodeD), "FAnimationAnimNextRuntimeTest_IHierarchy -> IHierarchy not found");
+			AddErrorIfFalse(StackNodeD.GetInterface(HierarchyBindingNodeD), "FAnimationAnimNextRuntimeTest_IHierarchy -> IHierarchy not found");
 
 			FChildrenArray ChildrenNodeD;
 			HierarchyBindingNodeD.GetChildren(Context, ChildrenNodeD);
 
 			AddErrorIfFalse(ChildrenNodeD.Num() == 2, "FAnimationAnimNextRuntimeTest_IHierarchy -> Expected 2 children");
+			AddErrorIfFalse(HierarchyBindingNodeD.GetNumChildren(Context) == 2, "FAnimationAnimNextRuntimeTest_IHierarchy -> Expected 2 children");
+			AddErrorIfFalse(ChildrenNodeD[0].IsValid() && ChildrenNodeD[0].GetNodeInstance()->GetNodeHandle() == NodeHandles[0], "FAnimationAnimNextRuntimeTest_IHierarchy -> Expected child: NodeA");
+			AddErrorIfFalse(ChildrenNodeD[1].IsValid() && ChildrenNodeD[1].GetNodeInstance()->GetNodeHandle() == NodeHandles[2], "FAnimationAnimNextRuntimeTest_IHierarchy -> Expected child: NodeC");
+
+			ChildrenNodeD.Reset();
+			IHierarchy::GetStackChildren(Context, StackNodeD, ChildrenNodeD);
+
+			AddErrorIfFalse(ChildrenNodeD.Num() == 2, "FAnimationAnimNextRuntimeTest_IHierarchy -> Expected 2 children");
+			AddErrorIfFalse(IHierarchy::GetNumStackChildren(Context, StackNodeD) == 2, "FAnimationAnimNextRuntimeTest_IHierarchy -> Expected 2 children");
 			AddErrorIfFalse(ChildrenNodeD[0].IsValid() && ChildrenNodeD[0].GetNodeInstance()->GetNodeHandle() == NodeHandles[0], "FAnimationAnimNextRuntimeTest_IHierarchy -> Expected child: NodeA");
 			AddErrorIfFalse(ChildrenNodeD[1].IsValid() && ChildrenNodeD[1].GetNodeInstance()->GetNodeHandle() == NodeHandles[2], "FAnimationAnimNextRuntimeTest_IHierarchy -> Expected child: NodeC");
 
 			{
+				FTraitStackBinding StackNodeC;
+				AddErrorIfFalse(Context.GetStack(ChildrenNodeD[1], StackNodeC), "FAnimationAnimNextRuntimeTest_IHierarchy -> Failed to bind to trait stack");
+
 				TTraitBinding<IHierarchy> HierarchyBindingNodeC;
-				AddErrorIfFalse(Context.GetInterface(ChildrenNodeD[1], HierarchyBindingNodeC), "FAnimationAnimNextRuntimeTest_IHierarchy -> IHierarchy not found");
+				AddErrorIfFalse(StackNodeC.GetInterface(HierarchyBindingNodeC), "FAnimationAnimNextRuntimeTest_IHierarchy -> IHierarchy not found");
 
 				FChildrenArray ChildrenNodeC;
 				HierarchyBindingNodeC.GetChildren(Context, ChildrenNodeC);
 
-				AddErrorIfFalse(ChildrenNodeC.Num() == 1, "FAnimationAnimNextRuntimeTest_IHierarchy -> Expected 2 children");
+				AddErrorIfFalse(ChildrenNodeC.Num() == 1, "FAnimationAnimNextRuntimeTest_IHierarchy -> Expected 1 child");
+				AddErrorIfFalse(HierarchyBindingNodeC.GetNumChildren(Context) == 1, "FAnimationAnimNextRuntimeTest_IHierarchy -> Expected 1 child");
+				AddErrorIfFalse(ChildrenNodeC[0].IsValid() && ChildrenNodeC[0].GetNodeInstance()->GetNodeHandle() == NodeHandles[1], "FAnimationAnimNextRuntimeTest_IHierarchy -> Expected child: NodeB");
+
+				ChildrenNodeC.Reset();
+				IHierarchy::GetStackChildren(Context, StackNodeC, ChildrenNodeC);
+
+				AddErrorIfFalse(ChildrenNodeC.Num() == 1, "FAnimationAnimNextRuntimeTest_IHierarchy -> Expected 1 child");
+				AddErrorIfFalse(IHierarchy::GetNumStackChildren(Context, StackNodeC) == 1, "FAnimationAnimNextRuntimeTest_IHierarchy -> Expected 1 child");
 				AddErrorIfFalse(ChildrenNodeC[0].IsValid() && ChildrenNodeC[0].GetNodeInstance()->GetNodeHandle() == NodeHandles[1], "FAnimationAnimNextRuntimeTest_IHierarchy -> Expected child: NodeB");
 
 				{
+					FTraitStackBinding StackNodeB;
+					AddErrorIfFalse(Context.GetStack(ChildrenNodeC[0], StackNodeB), "FAnimationAnimNextRuntimeTest_IHierarchy -> Failed to bind to trait stack");
+
 					TTraitBinding<IHierarchy> HierarchyBindingNodeB;
-					AddErrorIfFalse(Context.GetInterface(ChildrenNodeC[0], HierarchyBindingNodeB), "FAnimationAnimNextRuntimeTest_IHierarchy -> IHierarchy not found");
+					AddErrorIfFalse(StackNodeB.GetInterface(HierarchyBindingNodeB), "FAnimationAnimNextRuntimeTest_IHierarchy -> IHierarchy not found");
 
 					FChildrenArray ChildrenNodeB;
 					HierarchyBindingNodeB.GetChildren(Context, ChildrenNodeB);
 
 					AddErrorIfFalse(ChildrenNodeB.Num() == 1, "FAnimationAnimNextRuntimeTest_IHierarchy -> Expected 1 child");
+					AddErrorIfFalse(HierarchyBindingNodeB.GetNumChildren(Context) == 1, "FAnimationAnimNextRuntimeTest_IHierarchy -> Expected 1 child");
+					AddErrorIfFalse(ChildrenNodeB[0].IsValid() && ChildrenNodeB[0].GetNodeInstance()->GetNodeHandle() == NodeHandles[0], "FAnimationAnimNextRuntimeTest_IHierarchy -> Expected child: NodeA");
+
+					ChildrenNodeB.Reset();
+					IHierarchy::GetStackChildren(Context, StackNodeB, ChildrenNodeB);
+
+					AddErrorIfFalse(ChildrenNodeB.Num() == 1, "FAnimationAnimNextRuntimeTest_IHierarchy -> Expected 1 child");
+					AddErrorIfFalse(IHierarchy::GetNumStackChildren(Context, StackNodeB) == 1, "FAnimationAnimNextRuntimeTest_IHierarchy -> Expected 1 child");
 					AddErrorIfFalse(ChildrenNodeB[0].IsValid() && ChildrenNodeB[0].GetNodeInstance()->GetNodeHandle() == NodeHandles[0], "FAnimationAnimNextRuntimeTest_IHierarchy -> Expected child: NodeA");
 				}
 			}
@@ -687,5 +718,4 @@ bool FAnimationAnimNextRuntimeTest_IEvaluate::RunTest(const FString& InParameter
 
 	return true;
 }
-
 #endif
