Implemented AnimNext graph copy / paste / cut / duplicate functions on graph nodes
Refractored RigVMBlueprint editor side copy paste to use a helper function implemented in RigVMEditorTools, so it can be reused by AnimNext

#tests 65de1fe6db43688bc0ecf232

diff --git a/Engine/Plugins/Experimental/AnimNext/Source/AnimNextEditor/AnimNextEditor.Build.cs b/Engine/Plugins/Experimental/AnimNext/Source/AnimNextEditor/AnimNextEditor.Build.cs
--- a/Engine/Plugins/Experimental/AnimNext/Source/AnimNextEditor/AnimNextEditor.Build.cs
+++ b/Engine/Plugins/Experimental/AnimNext/Source/AnimNextEditor/AnimNextEditor.Build.cs
@@ -11,6 +11,7 @@ namespace UnrealBuildTool.Rules
 				{
 					"Core",
 					"CoreUObject",
+					"ApplicationCore",
 					"Engine",
 					"AssetTools",
 					"UnrealEd",
diff --git a/Engine/Plugins/Experimental/AnimNext/Source/AnimNextEditor/Private/AnimNextEditorModule.cpp b/Engine/Plugins/Experimental/AnimNext/Source/AnimNextEditor/Private/AnimNextEditorModule.cpp
--- a/Engine/Plugins/Experimental/AnimNext/Source/AnimNextEditor/Private/AnimNextEditorModule.cpp
+++ b/Engine/Plugins/Experimental/AnimNext/Source/AnimNextEditor/Private/AnimNextEditorModule.cpp
@@ -11,11 +11,13 @@
 #include "SSimpleComboButton.h"
 #include "UncookedOnlyUtils.h"
 #include "Common/SRigVMAssetView.h"
+#include "Editor/RigVMEditorTools.h"
 #include "Framework/Application/SlateApplication.h"
 #include "Graph/AnimNextGraph.h"
 #include "Graph/AnimNextGraphPanelNodeFactory.h"
 #include "Graph/AnimNextGraph_EdGraphNodeCustomization.h"
 #include "Graph/AnimNextGraph_EditorData.h"
+#include "HAL/PlatformApplicationMisc.h"
 #include "Param/ParameterCustomization.h"
 #include "Param/ParameterPickerArgs.h"
 #include "Param/ParametersGraphPanelPinFactory.h"
@@ -306,7 +308,27 @@ class FModule : public IModule
 
 			RigVMEdGraph->GetController()->SetCommentTextByName(CommentBeingChanged->GetFName(), NewText.ToString(), CommentBeingChanged->FontSize, CommentBeingChanged->bCommentBubbleVisible, CommentBeingChanged->bColorCommentBubble, true, true);
 		});
-		GraphArgs.OnDeleteSelectedNodes = Workspace::FOnDeleteSelectedNodes::CreateLambda([](const Workspace::FWorkspaceEditorContext& InContext, const FGraphPanelSelectionSet& InSelectedNodes)
+		GraphArgs.OnCanDeleteSelectedNodes = Workspace::FOnCanPerformActionOnSelectedNodes::CreateLambda([](const Workspace::FWorkspaceEditorContext& InContext, const FGraphPanelSelectionSet& InSelectedNodes)
+		{
+			bool bCanUserDeleteNode = false;
+
+			if(InSelectedNodes.Num() > 0)
+			{
+				for(UObject* NodeObject : InSelectedNodes)
+				{
+					// If any nodes allow deleting, then do not disable the delete option
+					const UEdGraphNode* Node = Cast<UEdGraphNode>(NodeObject);
+					if(Node && Node->CanUserDeleteNode())
+					{
+						bCanUserDeleteNode = true;
+						break;
+					}
+				}
+			}
+
+			return bCanUserDeleteNode;
+		});
+		GraphArgs.OnDeleteSelectedNodes = Workspace::FOnPerformActionOnSelectedNodes::CreateLambda([](const Workspace::FWorkspaceEditorContext& InContext, const FGraphPanelSelectionSet& InSelectedNodes)
 		{
 			if(InSelectedNodes.IsEmpty())
 			{
@@ -378,6 +400,160 @@ class FModule : public IModule
 			Controller->RemoveNodes(NodesToRemove, true);
 			Controller->CloseUndoBracket();
 		});
+		GraphArgs.OnCanCopySelectedNodes = Workspace::FOnCanPerformActionOnSelectedNodes::CreateLambda([](const Workspace::FWorkspaceEditorContext& InContext, const FGraphPanelSelectionSet& InSelectedNodes)
+		{
+			bool bCanUserCopyNode = false;
+
+			if(InSelectedNodes.Num() > 0)
+			{
+				bCanUserCopyNode = true;
+			}
+
+			return bCanUserCopyNode;
+		});
+		GraphArgs.OnCopySelectedNodes = Workspace::FOnPerformActionOnSelectedNodes::CreateLambda([](const Workspace::FWorkspaceEditorContext& InContext, const FGraphPanelSelectionSet& InSelectedNodes)
+		{
+			if(InSelectedNodes.IsEmpty())
+			{
+				return;
+			}
+
+			URigVMEdGraph* RigVMEdGraph = Cast<URigVMEdGraph>(InContext.Object);
+			if (RigVMEdGraph == nullptr)
+			{
+				return;
+			}
+
+			URigVMController* Controller = RigVMEdGraph->GetController();
+
+			FString ExportedText = Controller->ExportSelectedNodesToText();
+			FPlatformApplicationMisc::ClipboardCopy(*ExportedText);
+		});
+		GraphArgs.OnCanPasteNodes = Workspace::FOnCanPasteNodes::CreateLambda([](const Workspace::FWorkspaceEditorContext& InContext, const FString& InImportData)
+		{
+			bool bCanUserImportNodes = false;
+
+			if (!InImportData.IsEmpty())
+			{
+				bCanUserImportNodes = true;
+			}
+
+			return bCanUserImportNodes;
+		});
+		GraphArgs.OnPasteNodes = Workspace::FOnPasteNodes::CreateLambda([](const Workspace::FWorkspaceEditorContext& InContext, const FVector2D& InPasteLocation, const FString& InImportData)
+		{
+			if(InImportData.IsEmpty())
+			{
+				return;
+			}
+
+			URigVMEdGraph* RigVMEdGraph = Cast<URigVMEdGraph>(InContext.Object);
+			if (RigVMEdGraph == nullptr)
+			{
+				return;
+			}
+
+			if (IRigVMClientHost* RigVMClientHost = RigVMEdGraph->GetImplementingOuter<IRigVMClientHost>())
+			{
+				FString TextToImport;
+				FPlatformApplicationMisc::ClipboardPaste(TextToImport);
+				URigVMController* Controller = RigVMEdGraph->GetController();
+
+				Controller->OpenUndoBracket(TEXT("Pasted Nodes."));
+
+				if (UE::RigVM::Editor::Tools::PasteNodes(InPasteLocation, TextToImport, Controller, RigVMEdGraph->GetModel(), RigVMClientHost->GetLocalFunctionLibrary(), RigVMClientHost->GetRigVMGraphFunctionHost()))
+				{
+					Controller->CloseUndoBracket();
+				}
+				else
+				{
+					Controller->CancelUndoBracket();
+				}
+			}
+		});
+		UE::Workspace::FOnCanPerformActionOnSelectedNodes& OnCanCopySelectedNodes = GraphArgs.OnCanCopySelectedNodes;
+		UE::Workspace::FOnCanPerformActionOnSelectedNodes& OnCanDeleteSelectedNodes = GraphArgs.OnCanDeleteSelectedNodes;
+		GraphArgs.OnCanCutSelectedNodes = Workspace::FOnCanPerformActionOnSelectedNodes::CreateLambda([OnCanCopySelectedNodes, OnCanDeleteSelectedNodes](const Workspace::FWorkspaceEditorContext& InContext, const FGraphPanelSelectionSet& InSelectedNodes)
+			{
+				bool bCanUserCopyNode = false;
+
+				if (OnCanCopySelectedNodes.IsBound() && OnCanDeleteSelectedNodes.IsBound())
+				{
+					bCanUserCopyNode = OnCanCopySelectedNodes.Execute(InContext, InSelectedNodes) && OnCanDeleteSelectedNodes.Execute(InContext, InSelectedNodes);
+				}
+
+				return bCanUserCopyNode;
+			});
+		UE::Workspace::FOnPerformActionOnSelectedNodes& OnCopySelectedNodes = GraphArgs.OnCopySelectedNodes;
+		UE::Workspace::FOnPerformActionOnSelectedNodes& OnDeleteSelectedNodes = GraphArgs.OnDeleteSelectedNodes;
+		GraphArgs.OnCutSelectedNodes = Workspace::FOnPerformActionOnSelectedNodes::CreateLambda([OnCopySelectedNodes, OnDeleteSelectedNodes](const Workspace::FWorkspaceEditorContext& InContext, const FGraphPanelSelectionSet& InSelectedNodes)
+			{
+				if (InSelectedNodes.IsEmpty())
+				{
+					return;
+				}
+
+				URigVMEdGraph* RigVMEdGraph = Cast<URigVMEdGraph>(InContext.Object);
+				if (RigVMEdGraph == nullptr)
+				{
+					return;
+				}
+
+
+				if (OnCopySelectedNodes.IsBound() && OnDeleteSelectedNodes.IsBound())
+				{
+					URigVMController* Controller = RigVMEdGraph->GetController();
+
+					OnCopySelectedNodes.Execute(InContext, InSelectedNodes);
+
+					Controller->OpenUndoBracket(TEXT("Cut Nodes."));
+					OnDeleteSelectedNodes.Execute(InContext, InSelectedNodes);
+					Controller->CloseUndoBracket();
+				}
+			});
+		UE::Workspace::FOnCanPasteNodes& OnCanPasteNodes = GraphArgs.OnCanPasteNodes;
+		GraphArgs.OnCanDuplicateSelectedNodes = Workspace::FOnCanPerformActionOnSelectedNodes::CreateLambda([OnCanCopySelectedNodes, OnCanPasteNodes](const Workspace::FWorkspaceEditorContext& InContext, const FGraphPanelSelectionSet& InSelectedNodes)
+			{
+				bool bCanUserCopyNode = false;
+
+				if (OnCanCopySelectedNodes.IsBound() && OnCanPasteNodes.IsBound())
+				{
+					FString TextToImport;
+					FPlatformApplicationMisc::ClipboardPaste(TextToImport);
+
+					bCanUserCopyNode = OnCanCopySelectedNodes.Execute(InContext, InSelectedNodes) && OnCanPasteNodes.Execute(InContext, TextToImport);
+				}
+
+				return bCanUserCopyNode;
+			});
+		UE::Workspace::FOnPasteNodes& OnPasteNodes = GraphArgs.OnPasteNodes;
+		GraphArgs.OnDuplicateSelectedNodes = Workspace::FOnDuplicateSelectedNodes::CreateLambda([OnCopySelectedNodes, OnPasteNodes](const Workspace::FWorkspaceEditorContext& InContext, const FVector2D& InPasteLocation, const FGraphPanelSelectionSet& InSelectedNodes)
+			{
+				if (InSelectedNodes.IsEmpty())
+				{
+					return;
+				}
+
+				URigVMEdGraph* RigVMEdGraph = Cast<URigVMEdGraph>(InContext.Object);
+				if (RigVMEdGraph == nullptr)
+				{
+					return;
+				}
+
+				if (OnCopySelectedNodes.IsBound() && OnPasteNodes.IsBound())
+				{
+					OnCopySelectedNodes.Execute(InContext, InSelectedNodes);
+
+					FString TextToImport;
+					FPlatformApplicationMisc::ClipboardPaste(TextToImport);
+
+					URigVMController* Controller = RigVMEdGraph->GetController();
+
+					Controller->OpenUndoBracket(TEXT("Duplicate Nodes."));
+					OnPasteNodes.Execute(InContext, InPasteLocation, TextToImport);
+					Controller->CloseUndoBracket();
+				}
+			});
 		GraphArgs.OnGraphSelectionChanged = Workspace::FOnGraphSelectionChanged::CreateLambda([](const Workspace::FWorkspaceEditorContext& InContext, const FGraphPanelSelectionSet& NewSelection)
 		{
 			URigVMEdGraph* RigVMEdGraph = Cast<URigVMEdGraph>(InContext.Object);
diff --git a/Engine/Plugins/Experimental/Workspace/Source/WorkspaceEditor/Private/SGraphDocument.cpp b/Engine/Plugins/Experimental/Workspace/Source/WorkspaceEditor/Private/SGraphDocument.cpp
--- a/Engine/Plugins/Experimental/Workspace/Source/WorkspaceEditor/Private/SGraphDocument.cpp
+++ b/Engine/Plugins/Experimental/Workspace/Source/WorkspaceEditor/Private/SGraphDocument.cpp
@@ -1,10 +1,11 @@
-﻿// Copyright Epic Games, Inc. All Rights Reserved.
+// Copyright Epic Games, Inc. All Rights Reserved.
 
 #include "SGraphDocument.h"
 
 #include "Framework/Commands/GenericCommands.h"
 #include "WorkspaceEditor.h"
 #include "EdGraph/EdGraph.h"
+#include "HAL/PlatformApplicationMisc.h"
 
 namespace UE::Workspace
 {
@@ -16,7 +17,16 @@ void SGraphDocument::Construct(const FArguments& InArgs, TSharedRef<FWorkspaceEd
 
 	BindCommands();
 
+	OnCanDeleteSelectedNodes = InArgs._OnCanDeleteSelectedNodes;
 	OnDeleteSelectedNodes = InArgs._OnDeleteSelectedNodes;
+	OnCanCutSelectedNodes = InArgs._OnCanCutSelectedNodes;
+	OnCutSelectedNodes = InArgs._OnCutSelectedNodes;
+	OnCanCopySelectedNodes = InArgs._OnCanCopySelectedNodes;
+	OnCopySelectedNodes = InArgs._OnCopySelectedNodes;
+	OnCanPasteNodes = InArgs._OnCanPasteNodes;
+	OnPasteNodes = InArgs._OnPasteNodes;
+	OnCanDuplicateSelectedNodes = InArgs._OnCanDuplicateSelectedNodes;
+	OnDuplicateSelectedNodes = InArgs._OnDuplicateSelectedNodes;
 
 	SGraphEditor::FGraphEditorEvents Events;
 	Events.OnCreateActionMenu = SGraphEditor::FOnCreateActionMenu::CreateLambda([this, OnCreateActionMenu = InArgs._OnCreateActionMenu](UEdGraph* InGraph, const FVector2D& InNodePosition, const TArray<UEdGraphPin*>& InDraggedPins, bool bInAutoExpand, SGraphEditor::FActionMenuClosed InOnMenuClosed)
@@ -54,6 +64,34 @@ void SGraphDocument::BindCommands()
 	CommandList->MapAction(FGenericCommands::Get().Delete,
 		FExecuteAction::CreateSP(this, &SGraphDocument::DeleteSelectedNodes),
 		FCanExecuteAction::CreateSP(this, &SGraphDocument::CanDeleteSelectedNodes));
+
+	CommandList->MapAction(FGenericCommands::Get().Cut,
+		FExecuteAction::CreateSP(this, &SGraphDocument::CutSelectedNodes),
+		FCanExecuteAction::CreateSP(this, &SGraphDocument::CanCutSelectedNodes));
+
+	CommandList->MapAction(FGenericCommands::Get().Copy,
+		FExecuteAction::CreateSP(this, &SGraphDocument::CopySelectedNodes),
+		FCanExecuteAction::CreateSP(this, &SGraphDocument::CanCopySelectedNodes));
+
+	CommandList->MapAction(FGenericCommands::Get().Paste,
+		FExecuteAction::CreateSP(this, &SGraphDocument::PasteNodes),
+		FCanExecuteAction::CreateSP(this, &SGraphDocument::CanPasteNodes));
+
+	CommandList->MapAction(FGenericCommands::Get().Duplicate,
+		FExecuteAction::CreateSP(this, &SGraphDocument::DuplicateSelectedNodes),
+		FCanExecuteAction::CreateSP(this, &SGraphDocument::CanDuplicateSelectedNodes));
+
+	CommandList->MapAction(FGenericCommands::Get().SelectAll,
+		FExecuteAction::CreateSP(this, &SGraphDocument::SelectAllNodes),
+		FCanExecuteAction::CreateSP(this, &SGraphDocument::CanSelectAllNodes));
+}
+
+bool SGraphDocument::CanDeleteSelectedNodes() const
+{
+	const bool bCan = OnCanCutSelectedNodes.IsBound()
+		? IsEditable(EdGraph) && OnCanCutSelectedNodes.Execute(FWorkspaceEditorContext(HostingAppPtr.Pin().ToSharedRef(), EdGraph), GraphEditor->GetSelectedNodes())
+		: false;
+	return bCan;
 }
 
 void SGraphDocument::DeleteSelectedNodes()
@@ -62,29 +100,79 @@ void SGraphDocument::DeleteSelectedNodes()
 	OnDeleteSelectedNodes.ExecuteIfBound(FWorkspaceEditorContext(HostingAppPtr.Pin().ToSharedRef(), EdGraph), SelectedNodes);
 }
 
-bool SGraphDocument::CanDeleteSelectedNodes() const
+bool SGraphDocument::CanCutSelectedNodes() const
+{
+	const bool bCan = OnCanCutSelectedNodes.IsBound() 
+		? IsEditable(EdGraph) && OnCanCutSelectedNodes.Execute(FWorkspaceEditorContext(HostingAppPtr.Pin().ToSharedRef(), EdGraph), GraphEditor->GetSelectedNodes())
+		: false;
+	return bCan;
+}
+
+void SGraphDocument::CutSelectedNodes()
 {
 	const FGraphPanelSelectionSet SelectedNodes = GraphEditor->GetSelectedNodes();
+	OnCutSelectedNodes.ExecuteIfBound(FWorkspaceEditorContext(HostingAppPtr.Pin().ToSharedRef(), EdGraph), SelectedNodes);
+}
 
-	bool bCanUserDeleteNode = false;
+bool SGraphDocument::CanCopySelectedNodes() const
+{
+	const bool bCan = OnCanCopySelectedNodes.IsBound()
+		? OnCanCopySelectedNodes.Execute(FWorkspaceEditorContext(HostingAppPtr.Pin().ToSharedRef(), EdGraph), GraphEditor->GetSelectedNodes())
+		: false;
+	return bCan;
+}
 
-	if(IsEditable(EdGraph) && SelectedNodes.Num() > 0)
-	{
-		for(UObject* NodeObject : SelectedNodes)
-		{
-			// If any nodes allow deleting, then do not disable the delete option
-			UEdGraphNode* Node = Cast<UEdGraphNode>(NodeObject);
-			if(Node && Node->CanUserDeleteNode())
-			{
-				bCanUserDeleteNode = true;
-				break;
-			}
-		}
-	}
+void SGraphDocument::CopySelectedNodes()
+{
+	const FGraphPanelSelectionSet SelectedNodes = GraphEditor->GetSelectedNodes();
+	OnCopySelectedNodes.ExecuteIfBound(FWorkspaceEditorContext(HostingAppPtr.Pin().ToSharedRef(), EdGraph), SelectedNodes);
+}
+
+bool SGraphDocument::CanPasteNodes() const
+{
+	FString TextToImport;
+	FPlatformApplicationMisc::ClipboardPaste(TextToImport);
 
-	return bCanUserDeleteNode;
+	const bool bCan = OnCanPasteNodes.IsBound()
+		? IsEditable(EdGraph) && !TextToImport.IsEmpty() && OnCanPasteNodes.Execute(FWorkspaceEditorContext(HostingAppPtr.Pin().ToSharedRef(), EdGraph), TextToImport)
+		: false;
+	return bCan;
 }
 
+void SGraphDocument::PasteNodes()
+{
+	FString TextToImport;
+	FPlatformApplicationMisc::ClipboardPaste(TextToImport);
+
+	OnPasteNodes.ExecuteIfBound(FWorkspaceEditorContext(HostingAppPtr.Pin().ToSharedRef(), EdGraph), GraphEditor->GetPasteLocation(), TextToImport);
+}
+
+bool SGraphDocument::CanDuplicateSelectedNodes() const
+{
+	const bool bCan = OnCanDuplicateSelectedNodes.IsBound()
+		? IsEditable(EdGraph) && OnCanDuplicateSelectedNodes.Execute(FWorkspaceEditorContext(HostingAppPtr.Pin().ToSharedRef(), EdGraph), GraphEditor->GetSelectedNodes())
+		: false;
+	return bCan;
+}
+
+void SGraphDocument::DuplicateSelectedNodes()
+{
+	const FGraphPanelSelectionSet SelectedNodes = GraphEditor->GetSelectedNodes();
+	OnDuplicateSelectedNodes.ExecuteIfBound(FWorkspaceEditorContext(HostingAppPtr.Pin().ToSharedRef(), EdGraph), GraphEditor->GetPasteLocation(), SelectedNodes);
+}
+
+bool SGraphDocument::CanSelectAllNodes() const
+{
+	const bool bCan = GraphEditor.IsValid();
+	return bCan;
+}
+
+void SGraphDocument::SelectAllNodes()
+{
+	GraphEditor->SelectAllNodes();
+}
+
+
 bool SGraphDocument::IsEditable(UEdGraph* InGraph) const
 {
 	return InGraph && HostingAppPtr.Pin()->InEditingMode() && InGraph->bEditable;
diff --git a/Engine/Plugins/Experimental/Workspace/Source/WorkspaceEditor/Private/SGraphDocument.h b/Engine/Plugins/Experimental/Workspace/Source/WorkspaceEditor/Private/SGraphDocument.h
--- a/Engine/Plugins/Experimental/Workspace/Source/WorkspaceEditor/Private/SGraphDocument.h
+++ b/Engine/Plugins/Experimental/Workspace/Source/WorkspaceEditor/Private/SGraphDocument.h
@@ -1,4 +1,4 @@
-﻿// Copyright Epic Games, Inc. All Rights Reserved.
+// Copyright Epic Games, Inc. All Rights Reserved.
 
 #pragma once
 
@@ -24,19 +24,48 @@ class SGraphDocument : public SCompoundWidget
 
 	SLATE_EVENT(FOnNodeTextCommitted, OnNodeTextCommitted)
 
-	SLATE_EVENT(FOnDeleteSelectedNodes, OnDeleteSelectedNodes)
-
 	SLATE_EVENT(FOnGraphSelectionChanged, OnGraphSelectionChanged)
 
+	SLATE_EVENT(FOnCanPerformActionOnSelectedNodes, OnCanDeleteSelectedNodes)
+	SLATE_EVENT(FOnPerformActionOnSelectedNodes, OnDeleteSelectedNodes)
+
+	SLATE_EVENT(FOnCanPerformActionOnSelectedNodes, OnCanCutSelectedNodes)
+	SLATE_EVENT(FOnPerformActionOnSelectedNodes, OnCutSelectedNodes)
+
+	SLATE_EVENT(FOnCanPerformActionOnSelectedNodes, OnCanCopySelectedNodes)
+	SLATE_EVENT(FOnPerformActionOnSelectedNodes, OnCopySelectedNodes)
+
+	SLATE_EVENT(FOnCanPasteNodes, OnCanPasteNodes)
+	SLATE_EVENT(FOnPasteNodes, OnPasteNodes)
+
+	SLATE_EVENT(FOnCanPerformActionOnSelectedNodes, OnCanDuplicateSelectedNodes)
+	SLATE_EVENT(FOnDuplicateSelectedNodes, OnDuplicateSelectedNodes)
+
 	SLATE_END_ARGS()
 
 	void Construct(const FArguments& InArgs, TSharedRef<FWorkspaceEditor> InHostingApp, UEdGraph* InGraph);
 	
 	void BindCommands();
 	
-	void DeleteSelectedNodes();
 
 	bool CanDeleteSelectedNodes() const;
+	void DeleteSelectedNodes();
+
+	bool CanCutSelectedNodes() const;
+	void CutSelectedNodes();
+
+	bool CanCopySelectedNodes() const;
+	void CopySelectedNodes();
+
+	bool CanPasteNodes() const;
+	void PasteNodes();
+
+	bool CanDuplicateSelectedNodes() const;
+	void DuplicateSelectedNodes();
+
+	bool CanSelectAllNodes() const;
+	void SelectAllNodes();
+
 
 	bool IsEditable(UEdGraph* InGraph) const;
 
@@ -53,9 +82,26 @@ class SGraphDocument : public SCompoundWidget
 	TWeakPtr<FWorkspaceEditor> HostingAppPtr;
 
 	// Delegate called when we delete nodes
-	FOnDeleteSelectedNodes OnDeleteSelectedNodes;
+	FOnCanPerformActionOnSelectedNodes OnCanDeleteSelectedNodes;
+	FOnPerformActionOnSelectedNodes OnDeleteSelectedNodes;
+
+	// Cut delegates
+	FOnCanPerformActionOnSelectedNodes OnCanCutSelectedNodes;
+	FOnPerformActionOnSelectedNodes OnCutSelectedNodes;
+
+	// Copy delegates
+	FOnCanPerformActionOnSelectedNodes OnCanCopySelectedNodes;
+	FOnPerformActionOnSelectedNodes OnCopySelectedNodes;
+
+	// Paste delegates
+	FOnCanPasteNodes OnCanPasteNodes;
+	FOnPasteNodes OnPasteNodes;
+
+	// Duplicate delegates
+	FOnCanPerformActionOnSelectedNodes OnCanDuplicateSelectedNodes;
+	FOnDuplicateSelectedNodes OnDuplicateSelectedNodes;
 
 	friend class FWorkspaceEditorModule;
 };
 
-}
\ No newline at end of file
+}
diff --git a/Engine/Plugins/Experimental/Workspace/Source/WorkspaceEditor/Private/WorkspaceEditorModule.cpp b/Engine/Plugins/Experimental/Workspace/Source/WorkspaceEditor/Private/WorkspaceEditorModule.cpp
--- a/Engine/Plugins/Experimental/Workspace/Source/WorkspaceEditor/Private/WorkspaceEditorModule.cpp
+++ b/Engine/Plugins/Experimental/Workspace/Source/WorkspaceEditor/Private/WorkspaceEditorModule.cpp
@@ -1,4 +1,4 @@
-﻿// Copyright Epic Games, Inc. All Rights Reserved.
+// Copyright Epic Games, Inc. All Rights Reserved.
 
 #include "WorkspaceEditorModule.h"
 
@@ -53,8 +53,17 @@ FObjectDocumentArgs FWorkspaceEditorModule::CreateGraphDocumentArgs(const FGraph
 				return FActionMenuContent();
 			})
 			.OnNodeTextCommitted(InArgs.OnNodeTextCommitted)
+			.OnGraphSelectionChanged(InArgs.OnGraphSelectionChanged)
+			.OnCanDeleteSelectedNodes(InArgs.OnCanDeleteSelectedNodes)
 			.OnDeleteSelectedNodes(InArgs.OnDeleteSelectedNodes)
-			.OnGraphSelectionChanged(InArgs.OnGraphSelectionChanged);
+			.OnCanCutSelectedNodes(InArgs.OnCanCutSelectedNodes)
+			.OnCutSelectedNodes(InArgs.OnCutSelectedNodes)
+			.OnCanCopySelectedNodes(InArgs.OnCanCopySelectedNodes)
+			.OnCopySelectedNodes(InArgs.OnCopySelectedNodes)
+			.OnCanPasteNodes(InArgs.OnCanPasteNodes)
+			.OnPasteNodes(InArgs.OnPasteNodes)
+			.OnCanDuplicateSelectedNodes(InArgs.OnCanDuplicateSelectedNodes)
+			.OnDuplicateSelectedNodes(InArgs.OnDuplicateSelectedNodes);
 	});
 	Args.OnGetTabIcon = FOnGetTabIcon::CreateLambda([](const FWorkspaceEditorContext& InContext)
 	{
diff --git a/Engine/Plugins/Experimental/Workspace/Source/WorkspaceEditor/Public/IWorkspaceEditorModule.h b/Engine/Plugins/Experimental/Workspace/Source/WorkspaceEditor/Public/IWorkspaceEditorModule.h
--- a/Engine/Plugins/Experimental/Workspace/Source/WorkspaceEditor/Public/IWorkspaceEditorModule.h
+++ b/Engine/Plugins/Experimental/Workspace/Source/WorkspaceEditor/Public/IWorkspaceEditorModule.h
@@ -1,4 +1,4 @@
-﻿// Copyright Epic Games, Inc. All Rights Reserved.
+// Copyright Epic Games, Inc. All Rights Reserved.
 
 #pragma once
 
@@ -83,7 +83,6 @@ struct FObjectDocumentArgs
 	FName SpawnLocation = WorkspaceTabs::MiddleDocumentArea;
 };
 
-using FOnDeleteSelectedNodes = TDelegate<void(const FWorkspaceEditorContext&, const FGraphPanelSelectionSet&)>;
 
 using FOnGraphSelectionChanged = TDelegate<void(const FWorkspaceEditorContext&, const FGraphPanelSelectionSet&)>;
 
@@ -91,6 +90,14 @@ using FOnCreateActionMenu = TDelegate<FActionMenuContent(const FWorkspaceEditorC
 
 using FOnNodeTextCommitted = TDelegate<void(const FWorkspaceEditorContext&, const FText&, ETextCommit::Type, UEdGraphNode*)>;
 
+using FOnCanPerformActionOnSelectedNodes = TDelegate<bool(const FWorkspaceEditorContext&, const FGraphPanelSelectionSet&)>;
+using FOnPerformActionOnSelectedNodes = TDelegate<void(const FWorkspaceEditorContext&, const FGraphPanelSelectionSet&)>;
+
+using FOnCanPasteNodes = TDelegate<bool(const FWorkspaceEditorContext&, const FString&)>;
+using FOnPasteNodes = TDelegate<void(const FWorkspaceEditorContext&, const FVector2D&, const FString&)>;
+
+using FOnDuplicateSelectedNodes = TDelegate<void(const FWorkspaceEditorContext&, const FVector2D&, const FGraphPanelSelectionSet&)>;
+
 // Arguments used to make document widgets for graphs
 struct FGraphDocumentWidgetArgs
 {
@@ -101,7 +108,23 @@ struct FGraphDocumentWidgetArgs
 
 	FOnNodeTextCommitted OnNodeTextCommitted;
 
-	FOnDeleteSelectedNodes OnDeleteSelectedNodes;
+	FOnCanPerformActionOnSelectedNodes OnCanCutSelectedNodes;
+	FOnPerformActionOnSelectedNodes OnCutSelectedNodes;
+
+	FOnCanPerformActionOnSelectedNodes OnCanCopySelectedNodes;
+	FOnPerformActionOnSelectedNodes OnCopySelectedNodes;
+
+	FOnCanPasteNodes OnCanPasteNodes;
+	FOnPasteNodes OnPasteNodes;
+
+	FOnCanPerformActionOnSelectedNodes OnCanDeleteSelectedNodes;
+	FOnPerformActionOnSelectedNodes OnDeleteSelectedNodes;
+
+	FOnCanPerformActionOnSelectedNodes OnCanDuplicateSelectedNodes;
+	FOnDuplicateSelectedNodes OnDuplicateSelectedNodes;
+
+	FOnCanPerformActionOnSelectedNodes OnCanSelectAllNodes;
+	FOnPerformActionOnSelectedNodes OnSelectAllNodes;
 
 	FOnGraphSelectionChanged OnGraphSelectionChanged;
 };
diff --git a/Engine/Plugins/Experimental/Workspace/Source/WorkspaceEditor/WorkspaceEditor.Build.cs b/Engine/Plugins/Experimental/Workspace/Source/WorkspaceEditor/WorkspaceEditor.Build.cs
--- a/Engine/Plugins/Experimental/Workspace/Source/WorkspaceEditor/WorkspaceEditor.Build.cs
+++ b/Engine/Plugins/Experimental/Workspace/Source/WorkspaceEditor/WorkspaceEditor.Build.cs
@@ -20,6 +20,7 @@ namespace UnrealBuildTool.Rules
 					"CoreUObject",
 					"Engine",
 					"AssetTools",
+					"ApplicationCore",
 					"SlateCore",
 					"Slate",
 					"InputCore",
diff --git a/Engine/Plugins/Runtime/RigVM/Source/RigVMDeveloper/Private/RigVMBlueprint.cpp b/Engine/Plugins/Runtime/RigVM/Source/RigVMDeveloper/Private/RigVMBlueprint.cpp
--- a/Engine/Plugins/Runtime/RigVM/Source/RigVMDeveloper/Private/RigVMBlueprint.cpp
+++ b/Engine/Plugins/Runtime/RigVM/Source/RigVMDeveloper/Private/RigVMBlueprint.cpp
@@ -4000,7 +4000,7 @@ void URigVMBlueprint::BroadcastPostEditChangeChainProperty(FPropertyChangedChain
 
 void URigVMBlueprint::BroadcastRequestLocalizeFunctionDialog(FRigVMGraphFunctionIdentifier InFunction, bool bForce)
 {
-	RequestLocalizeFunctionDialog.Broadcast(InFunction, this, bForce);
+	RequestLocalizeFunctionDialog.Broadcast(InFunction, GetController(GetDefaultModel()), GetRigVMGraphFunctionHost(), bForce);
 }
 
 void URigVMBlueprint::BroadCastReportCompilerMessage(EMessageSeverity::Type InSeverity, UObject* InSubject, const FString& InMessage)
diff --git a/Engine/Plugins/Runtime/RigVM/Source/RigVMDeveloper/Public/RigVMBlueprint.h b/Engine/Plugins/Runtime/RigVM/Source/RigVMDeveloper/Public/RigVMBlueprint.h
--- a/Engine/Plugins/Runtime/RigVM/Source/RigVMDeveloper/Public/RigVMBlueprint.h
+++ b/Engine/Plugins/Runtime/RigVM/Source/RigVMDeveloper/Public/RigVMBlueprint.h
@@ -34,7 +34,7 @@ DECLARE_EVENT_OneParam(URigVMBlueprint, FOnRigVMExternalVariablesChanged, const
 DECLARE_EVENT_TwoParams(URigVMBlueprint, FOnRigVMNodeDoubleClicked, URigVMBlueprint*, URigVMNode*);
 DECLARE_EVENT_OneParam(URigVMBlueprint, FOnRigVMGraphImported, UEdGraph*);
 DECLARE_EVENT_OneParam(URigVMBlueprint, FOnRigVMPostEditChangeChainProperty, FPropertyChangedChainEvent&);
-DECLARE_EVENT_ThreeParams(URigVMBlueprint, FOnRigVMLocalizeFunctionDialogRequested, FRigVMGraphFunctionIdentifier&, URigVMBlueprint*, bool);
+DECLARE_EVENT_FourParams(URigVMBlueprint, FOnRigVMLocalizeFunctionDialogRequested, FRigVMGraphFunctionIdentifier&, URigVMController*, IRigVMGraphFunctionHost*, bool);
 DECLARE_EVENT_ThreeParams(URigVMBlueprint, FOnRigVMReportCompilerMessage, EMessageSeverity::Type, UObject*, const FString&);
 DECLARE_DELEGATE_RetVal_FourParams(FRigVMController_BulkEditResult, FRigVMOnBulkEditDialogRequestedDelegate, URigVMBlueprint*, URigVMController*, URigVMLibraryNode*, ERigVMControllerBulkEditType);
 DECLARE_DELEGATE_RetVal_OneParam(bool, FRigVMOnBreakLinksDialogRequestedDelegate, TArray<URigVMLink*>);
@@ -42,7 +42,6 @@ DECLARE_DELEGATE_RetVal_OneParam(TRigVMTypeIndex, FRigVMOnPinTypeSelectionReques
 DECLARE_EVENT(URigVMBlueprint, FOnRigVMBreakpointAdded);
 DECLARE_EVENT_OneParam(URigVMBlueprint, FOnRigVMRequestInspectObject, const TArray<UObject*>& );
 DECLARE_EVENT_OneParam(URigVMBlueprint, FOnRigVMRequestInspectMemoryStorage, const TArray<FRigVMMemoryStorageStruct*>&);
-//DECLARE_DELEGATE_RetVal(URigVMGraph*, FRigVMBlueprintGetFocusedGraph);
 
 USTRUCT()
 struct RIGVMDEVELOPER_API FRigVMPythonSettings
diff --git a/Engine/Plugins/Runtime/RigVM/Source/RigVMEditor/Private/Editor/RigVMEditor.cpp b/Engine/Plugins/Runtime/RigVM/Source/RigVMEditor/Private/Editor/RigVMEditor.cpp
--- a/Engine/Plugins/Runtime/RigVM/Source/RigVMEditor/Private/Editor/RigVMEditor.cpp
+++ b/Engine/Plugins/Runtime/RigVM/Source/RigVMEditor/Private/Editor/RigVMEditor.cpp
@@ -1,6 +1,7 @@
 // Copyright Epic Games, Inc. All Rights Reserved.
 
 #include "Editor/RigVMEditor.h"
+#include "Editor/RigVMEditorTools.h"
 #include "Editor/RigVMEditorMenuContext.h"
 #include "HAL/PlatformApplicationMisc.h"
 #include "SMyBlueprint.h"
@@ -779,65 +780,18 @@ void FRigVMEditor::PasteNodes()
 	FString TextToImport;
 	FPlatformApplicationMisc::ClipboardPaste(TextToImport);
 
-	TGuardValue<FRigVMController_RequestLocalizeFunctionDelegate> RequestLocalizeDelegateGuard(
-		GetFocusedController()->RequestLocalizeFunctionDelegate,
-		FRigVMController_RequestLocalizeFunctionDelegate::CreateLambda([this](FRigVMGraphFunctionIdentifier& InFunctionToLocalize)
-		{
-			OnRequestLocalizeFunctionDialog(InFunctionToLocalize, GetRigVMBlueprint(), true);
-
-		   const URigVMLibraryNode* LocalizedFunctionNode = GetRigVMBlueprint()->GetLocalFunctionLibrary()->FindPreviouslyLocalizedFunction(InFunctionToLocalize);
-		   return LocalizedFunctionNode != nullptr;
-		
-		})
-	);
-	
-	TArray<FName> NodeNames = GetFocusedController()->ImportNodesFromText(TextToImport, true, true);
+	URigVMController* FocusedController = GetFocusedController();
 
-	if (NodeNames.Num() > 0)
+	const bool bPastePerformed = UE::RigVM::Editor::Tools::PasteNodes(PasteLocation, TextToImport, FocusedController, GetFocusedModel(), GetRigVMBlueprint()->GetLocalFunctionLibrary(), GetRigVMBlueprint()->GetRigVMGraphFunctionHost());
+	if (bPastePerformed)
 	{
-		FBox2D Bounds;
-		Bounds.bIsValid = false;
-
-		TArray<FName> NodesToSelect;
-		for (const FName& NodeName : NodeNames)
-		{
-			const URigVMNode* Node = GetFocusedModel()->FindNodeByName(NodeName);
-			check(Node);
-
-			if (Node->IsInjected())
-			{
-				continue;
-			}
-			NodesToSelect.Add(NodeName);
-
-			FVector2D Position = Node->GetPosition();
-			FVector2D Size = Node->GetSize();
-
-			if (!Bounds.bIsValid)
-			{
-				Bounds.Min = Bounds.Max = Position;
-				Bounds.bIsValid = true;
-			}
-			Bounds += Position;
-			Bounds += Position + Size;
-		}
-
-		for (const FName& NodeName : NodesToSelect)
-		{
-			const URigVMNode* Node = GetFocusedModel()->FindNodeByName(NodeName);
-			check(Node);
-
-			FVector2D Position = Node->GetPosition();
-			GetFocusedController()->SetNodePositionByName(NodeName, PasteLocation + Position - Bounds.GetCenter(), true, false, true);
-		}
-
-		GetFocusedController()->SetNodeSelection(NodesToSelect);
-		GetFocusedController()->CloseUndoBracket();
+		FocusedController->CloseUndoBracket();
 	}
 	else
 	{
-		GetFocusedController()->CancelUndoBracket();
+		FocusedController->CancelUndoBracket();
 	}
+
 }
 
 URigVMBlueprint* FRigVMEditor::GetRigVMBlueprint() const
@@ -2147,33 +2101,11 @@ void FRigVMEditor::OnWrappedPropertyChangedChainEvent(URigVMDetailsViewWrapperOb
 }
 
 void FRigVMEditor::OnRequestLocalizeFunctionDialog(FRigVMGraphFunctionIdentifier& InFunction,
-	URigVMBlueprint* InTargetBlueprint, bool bForce)
+	URigVMController* InTargetController,
+	IRigVMGraphFunctionHost* InTargetFunctionHost,
+	bool bForce)
 {
-	check(InTargetBlueprint);
-
-	if(InTargetBlueprint != GetRigVMBlueprint())
-	{
-		return;
-	}
-	
-	if(URigVMController* TargetController = InTargetBlueprint->GetController(InTargetBlueprint->GetDefaultModel()))
-	{
-		bool bIsPublic;
-		if (FRigVMGraphFunctionData::FindFunctionData(InFunction, &bIsPublic))
-		{
-			if (bForce || bIsPublic)
-			{
-				TSharedRef<SRigVMGraphFunctionLocalizationDialog> LocalizationDialog = SNew(SRigVMGraphFunctionLocalizationDialog)
-							.Function(InFunction)
-							.TargetBlueprint(InTargetBlueprint);
-
-				if (LocalizationDialog->ShowModal() != EAppReturnType::Cancel)
-				{
-					TargetController->LocalizeFunctions(LocalizationDialog->GetFunctionsToLocalize(), true, true, true);
-				}
-			}
-		}
-	}
+	UE::RigVM::Editor::Tools::OnRequestLocalizeFunctionDialog(InFunction, InTargetController, InTargetFunctionHost, bForce);
 }
 
 FRigVMController_BulkEditResult FRigVMEditor::OnRequestBulkEditDialog(URigVMBlueprint* InBlueprint,
diff --git a/Engine/Plugins/Runtime/RigVM/Source/RigVMEditor/Private/Editor/RigVMEditorTools.cpp b/Engine/Plugins/Runtime/RigVM/Source/RigVMEditor/Private/Editor/RigVMEditorTools.cpp
new file mode 100644
--- /dev/null
+++ b/Engine/Plugins/Runtime/RigVM/Source/RigVMEditor/Private/Editor/RigVMEditorTools.cpp
@@ -0,0 +1,111 @@
+// Copyright Epic Games, Inc. All Rights Reserved.
+
+#include "Editor/RigVMEditorTools.h"
+
+#include "RigVMModel/RigVMController.h"
+#include "RigVMModel/RigVMGraph.h"
+#include "RigVMModel/RigVMFunctionLibrary.h"
+#include "RigVMCore/RigVMGraphFunctionHost.h"
+#include "RigVMCore/RigVMGraphFunctionDefinition.h"
+#include "Widgets/SRigVMGraphFunctionLocalizationWidget.h"
+
+namespace UE::RigVM::Editor::Tools
+{
+
+RIGVMEDITOR_API bool PasteNodes(const FVector2D& PasteLocation
+	, const FString& TextToImport
+	, URigVMController* InFocusedController
+	, URigVMGraph* InFocusedModel
+	, URigVMFunctionLibrary* InLocalFunctionLibrary
+	, IRigVMGraphFunctionHost* InGraphFunctionHost)
+{
+	bool bPastePerformed = false;
+
+	InFocusedController->OpenUndoBracket(TEXT("Pasted Nodes."));
+
+	TGuardValue<FRigVMController_RequestLocalizeFunctionDelegate> RequestLocalizeDelegateGuard(
+		InFocusedController->RequestLocalizeFunctionDelegate,
+		FRigVMController_RequestLocalizeFunctionDelegate::CreateLambda([InFocusedController, InLocalFunctionLibrary, InGraphFunctionHost](FRigVMGraphFunctionIdentifier& InFunctionToLocalize)
+			{
+				OnRequestLocalizeFunctionDialog(InFunctionToLocalize, InFocusedController, InGraphFunctionHost, true);
+
+				const URigVMLibraryNode* LocalizedFunctionNode = InLocalFunctionLibrary->FindPreviouslyLocalizedFunction(InFunctionToLocalize);
+				return LocalizedFunctionNode != nullptr;
+			})
+	);
+
+	TArray<FName> NodeNames = InFocusedController->ImportNodesFromText(TextToImport, true, true);
+
+	if (NodeNames.Num() > 0)
+	{
+		FBox2D Bounds;
+		Bounds.bIsValid = false;
+
+		TArray<FName> NodesToSelect;
+		for (const FName& NodeName : NodeNames)
+		{
+			const URigVMNode* Node = InFocusedModel->FindNodeByName(NodeName);
+			check(Node);
+
+			if (Node->IsInjected())
+			{
+				continue;
+			}
+			NodesToSelect.Add(NodeName);
+
+			FVector2D Position = Node->GetPosition();
+			FVector2D Size = Node->GetSize();
+
+			if (!Bounds.bIsValid)
+			{
+				Bounds.Min = Bounds.Max = Position;
+				Bounds.bIsValid = true;
+			}
+			Bounds += Position;
+			Bounds += Position + Size;
+		}
+
+		for (const FName& NodeName : NodesToSelect)
+		{
+			const URigVMNode* Node = InFocusedModel->FindNodeByName(NodeName);
+			check(Node);
+
+			FVector2D Position = Node->GetPosition();
+			InFocusedController->SetNodePositionByName(NodeName, PasteLocation + Position - Bounds.GetCenter(), true, false, true);
+		}
+
+		InFocusedController->SetNodeSelection(NodesToSelect);
+		InFocusedController->CloseUndoBracket();
+
+		bPastePerformed = true;
+	}
+
+	return bPastePerformed;
+}
+
+RIGVMEDITOR_API void OnRequestLocalizeFunctionDialog(FRigVMGraphFunctionIdentifier& InFunction
+	, URigVMController* InTargetController
+	, IRigVMGraphFunctionHost* InTargetFunctionHost
+	, bool bForce)
+{
+	if (InTargetController != nullptr)
+	{
+		bool bIsPublic;
+		if (FRigVMGraphFunctionData::FindFunctionData(InFunction, &bIsPublic))
+		{
+			if (bForce || bIsPublic)
+			{
+				TSharedRef<SRigVMGraphFunctionLocalizationDialog> LocalizationDialog = SNew(SRigVMGraphFunctionLocalizationDialog)
+					.Function(InFunction)
+					.GraphFunctionHost(InTargetFunctionHost);
+
+				if (LocalizationDialog->ShowModal() != EAppReturnType::Cancel)
+				{
+					InTargetController->LocalizeFunctions(LocalizationDialog->GetFunctionsToLocalize(), true, true, true);
+				}
+			}
+		}
+	}
+}
+
+} // end namespace UE::RigVM::Editor::Tools
diff --git a/Engine/Plugins/Runtime/RigVM/Source/RigVMEditor/Private/Widgets/SRigVMGraphFunctionLocalizationWidget.cpp b/Engine/Plugins/Runtime/RigVM/Source/RigVMEditor/Private/Widgets/SRigVMGraphFunctionLocalizationWidget.cpp
--- a/Engine/Plugins/Runtime/RigVM/Source/RigVMEditor/Private/Widgets/SRigVMGraphFunctionLocalizationWidget.cpp
+++ b/Engine/Plugins/Runtime/RigVM/Source/RigVMEditor/Private/Widgets/SRigVMGraphFunctionLocalizationWidget.cpp
@@ -86,7 +86,7 @@ void SRigVMGraphFunctionLocalizationTableRow::Construct(const FArguments& InArgs
 /// SRigVMGraphFunctionLocalizationWidget
 ///////////////////////////////////////////////////////////
 
-void SRigVMGraphFunctionLocalizationWidget::Construct(const FArguments& InArgs, const FRigVMGraphFunctionIdentifier& InFunctionToLocalize, URigVMBlueprint* InTargetBlueprint)
+void SRigVMGraphFunctionLocalizationWidget::Construct(const FArguments& InArgs, const FRigVMGraphFunctionIdentifier& InFunctionToLocalize, IRigVMGraphFunctionHost* InTargetFunctionHost)
 {
 	FunctionsToLocalize.Add(InFunctionToLocalize);
 	FunctionItems.Reset();
@@ -100,7 +100,6 @@ void SRigVMGraphFunctionLocalizationWidget::Construct(const FArguments& InArgs,
 		NodesToVisit.Add(InFunctionToLocalize);
 	}
 
-	IRigVMGraphFunctionHost* TargetFunctionHost = InTargetBlueprint->GetRigVMGraphFunctionHost();
 	for(int32 NodeToVisitIndex=0; NodeToVisitIndex<NodesToVisit.Num(); NodeToVisitIndex++)
 	{
 		const FRigVMGraphFunctionIdentifier& NodeToVisit = NodesToVisit[NodeToVisitIndex];
@@ -117,7 +116,7 @@ void SRigVMGraphFunctionLocalizationWidget::Construct(const FArguments& InArgs,
 			NodesToVisit.AddUnique(Pair.Key);
 		}
 		
-		if (NodeToVisit.HostObject != Cast<UObject>(TargetFunctionHost))
+		if (NodeToVisit.HostObject != Cast<UObject>(InTargetFunctionHost))
 		{
 			FunctionsForTable.AddUnique(FunctionData);
 			if (bIsPublic)
@@ -245,7 +244,7 @@ void SRigVMGraphFunctionLocalizationDialog::Construct(const FArguments& InArgs)
                 .HAlign(HAlign_Left)
                 .Padding(2)
                 [
-                    SAssignNew(FunctionsWidget, SRigVMGraphFunctionLocalizationWidget, InArgs._Function, InArgs._TargetBlueprint)
+                    SAssignNew(FunctionsWidget, SRigVMGraphFunctionLocalizationWidget, InArgs._Function, InArgs._GraphFunctionHost)
                 ]
 
                 +SVerticalBox::Slot()
diff --git a/Engine/Plugins/Runtime/RigVM/Source/RigVMEditor/Public/Editor/RigVMEditor.h b/Engine/Plugins/Runtime/RigVM/Source/RigVMEditor/Public/Editor/RigVMEditor.h
--- a/Engine/Plugins/Runtime/RigVM/Source/RigVMEditor/Public/Editor/RigVMEditor.h
+++ b/Engine/Plugins/Runtime/RigVM/Source/RigVMEditor/Public/Editor/RigVMEditor.h
@@ -181,7 +181,7 @@ protected:
 	virtual void OnFinishedChangingProperties(const FPropertyChangedEvent& PropertyChangedEvent) override;
 	void OnPropertyChanged(UObject* InObject, FPropertyChangedEvent& InEvent);
 	virtual void OnWrappedPropertyChangedChainEvent(URigVMDetailsViewWrapperObject* InWrapperObject, const FString& InPropertyPath, FPropertyChangedChainEvent& InPropertyChangedChainEvent);
-	void OnRequestLocalizeFunctionDialog(FRigVMGraphFunctionIdentifier& InFunction, URigVMBlueprint* InTargetBlueprint, bool bForce);
+	void OnRequestLocalizeFunctionDialog(FRigVMGraphFunctionIdentifier& InFunction, URigVMController* InTargetController, IRigVMGraphFunctionHost* InTargetFunctionHost, bool bForce);
 	FRigVMController_BulkEditResult OnRequestBulkEditDialog(URigVMBlueprint* InBlueprint, URigVMController* InController, URigVMLibraryNode* InFunction, ERigVMControllerBulkEditType InEditType);
 	bool OnRequestBreakLinksDialog(TArray<URigVMLink*> InLinks);
 	TRigVMTypeIndex OnRequestPinTypeSelectionDialog(const TArray<TRigVMTypeIndex>& InTypes);
diff --git a/Engine/Plugins/Runtime/RigVM/Source/RigVMEditor/Public/Editor/RigVMEditorTools.h b/Engine/Plugins/Runtime/RigVM/Source/RigVMEditor/Public/Editor/RigVMEditorTools.h
new file mode 100644
--- /dev/null
+++ b/Engine/Plugins/Runtime/RigVM/Source/RigVMEditor/Public/Editor/RigVMEditorTools.h
@@ -0,0 +1,28 @@
+// Copyright Epic Games, Inc. All Rights Reserved.
+
+#pragma once
+
+#include "CoreMinimal.h"
+
+class URigVMController;
+class URigVMGraph;
+class URigVMFunctionLibrary;
+class IRigVMGraphFunctionHost;
+struct FRigVMGraphFunctionIdentifier;
+
+namespace UE::RigVM::Editor::Tools
+{
+
+RIGVMEDITOR_API bool PasteNodes(const FVector2D& PasteLocation
+	, const FString& TextToImport
+	, URigVMController* InFocusedController
+	, URigVMGraph* InFocusedModel
+	, URigVMFunctionLibrary* InLocalFunctionLibrary
+	, IRigVMGraphFunctionHost* InGraphFunctionHost);
+
+RIGVMEDITOR_API void OnRequestLocalizeFunctionDialog(FRigVMGraphFunctionIdentifier& InFunction
+	, URigVMController* InTargetController
+	, IRigVMGraphFunctionHost* InTargetFunctionHost
+	, bool bForce);
+
+} // end namespace UE::RigVM::Editor::Tools
diff --git a/Engine/Plugins/Runtime/RigVM/Source/RigVMEditor/Public/Widgets/SRigVMGraphFunctionLocalizationWidget.h b/Engine/Plugins/Runtime/RigVM/Source/RigVMEditor/Public/Widgets/SRigVMGraphFunctionLocalizationWidget.h
--- a/Engine/Plugins/Runtime/RigVM/Source/RigVMEditor/Public/Widgets/SRigVMGraphFunctionLocalizationWidget.h
+++ b/Engine/Plugins/Runtime/RigVM/Source/RigVMEditor/Public/Widgets/SRigVMGraphFunctionLocalizationWidget.h
@@ -34,7 +34,7 @@ public:
 	SLATE_BEGIN_ARGS(SRigVMGraphFunctionLocalizationWidget) {}
 	SLATE_END_ARGS()
 
-	void Construct(const FArguments& InArgs, const FRigVMGraphFunctionIdentifier& InFunctionToLocalize, URigVMBlueprint* InTargetBlueprint);
+	void Construct(const FArguments& InArgs, const FRigVMGraphFunctionIdentifier& InFunctionToLocalize, IRigVMGraphFunctionHost* InTargetFunctionHost);
 
 	TSharedRef<ITableRow> GenerateFunctionListRow(TSharedPtr<SRigVMGraphFunctionLocalizationItem> InItem, const TSharedRef<STableViewBase>& InOwningTable);
 	ECheckBoxState IsFunctionEnabled(const FRigVMGraphFunctionIdentifier InFunction) const;
@@ -60,7 +60,7 @@ public:
 	}
 
 	SLATE_ARGUMENT(FRigVMGraphFunctionIdentifier, Function)
-	SLATE_ARGUMENT(URigVMBlueprint*, TargetBlueprint)
+	SLATE_ARGUMENT(IRigVMGraphFunctionHost*, GraphFunctionHost)
 	SLATE_END_ARGS()
 
 	void Construct(const FArguments& InArgs);
