Implement scoped trait interface support in AnimNext to allow parent traits to publish information that children can query for

diff --git a/Engine/Plugins/Experimental/AnimNext/Source/AnimNext/Private/TraitCore/ExecutionContext.cpp b/Engine/Plugins/Experimental/AnimNext/Source/AnimNext/Private/TraitCore/ExecutionContext.cpp
--- a/Engine/Plugins/Experimental/AnimNext/Source/AnimNext/Private/TraitCore/ExecutionContext.cpp
+++ b/Engine/Plugins/Experimental/AnimNext/Source/AnimNext/Private/TraitCore/ExecutionContext.cpp
@@ -14,6 +14,67 @@
 
 namespace UE::AnimNext
 {
+	namespace Private
+	{
+		// Represents an entry for a scoped interface
+		struct FScopedInterfaceEntry
+		{
+			// The trait stack, copied from the source when we push a scoped interface
+			FTraitStackBinding Stack;
+
+			// The trait that implements the interface
+			// We lazily cache the binding of the scoped interface
+			FTraitBinding Trait;
+
+			// The scoped interface
+			FTraitInterfaceUID InterfaceUID;
+
+			// The trait index on the stack that implements our scoped interface
+			uint8 TraitIndex = 0;
+
+			// Whether or not the scoped interface trait binding has been cached
+			bool bIsTraitCached = false;
+
+			union
+			{
+				// Next entry in the stack of free entries
+				FScopedInterfaceEntry* NextFreeEntry = nullptr;
+
+				// The previous entry on the scoped interface stack
+				FScopedInterfaceEntry* PrevScopedInterfaceStackEntry;
+			};
+
+			FScopedInterfaceEntry(const FTraitBinding& InTrait, FTraitInterfaceUID InInterfaceUID)
+				: Stack(*InTrait.GetStack())
+				, InterfaceUID(InInterfaceUID)
+				, TraitIndex(InTrait.GetTraitIndex())
+			{
+			}
+
+			// Lazily construct the trait binding to our scoped interface
+			bool LazilyCacheTrait()
+			{
+				if (bIsTraitCached)
+				{
+					return true;	// Already cached
+				}
+
+				if (!Stack.GetTrait(TraitIndex, Trait))
+				{
+					return false;
+				}
+
+				if (!Trait.AsInterfaceImpl(InterfaceUID, Trait))
+				{
+					return false;
+				}
+
+				bIsTraitCached = true;
+				return true;
+			}
+		};
+	}
+
 	FExecutionContext::FExecutionContext()
 		: MemStack(FMemStack::Get())
 		, NodeTemplateRegistry(FNodeTemplateRegistry::Get())
@@ -322,6 +383,157 @@ namespace UE::AnimNext
 		return OutStackBinding.IsValid();	// Construction can fail in rare cases, see constructor
 	}
 
+	void FExecutionContext::PushScopedInterfaceImpl(FTraitInterfaceUID InterfaceUID, const FTraitBinding& Binding)
+	{
+		if (!Binding.IsValid())
+		{
+			return;	// Don't queue invalid pointers
+		}
+
+		// We don't have any specific handling for duplicate entries, if a scoped interface is pushed twice,
+		// it must also be popped twice (if popped manually)
+
+		Private::FScopedInterfaceEntry* ScopedEntry = FreeScopedInterfaceEntryStackHead;
+		if (ScopedEntry != nullptr)
+		{
+			// We have a free entry, set our new head
+			FreeScopedInterfaceEntryStackHead = ScopedEntry->NextFreeEntry;
+
+			// Update our entry
+			ScopedEntry->Stack = *Binding.GetStack();
+			ScopedEntry->TraitIndex = Binding.GetTraitIndex();
+			ScopedEntry->InterfaceUID = InterfaceUID;
+			ScopedEntry->bIsTraitCached = false;
+			ScopedEntry->NextFreeEntry = nullptr;		// Mark it as not being a member of any list
+		}
+		else
+		{
+			// Allocate a new entry
+			ScopedEntry = new(MemStack) Private::FScopedInterfaceEntry(Binding, InterfaceUID);
+		}
+
+		ScopedEntry->PrevScopedInterfaceStackEntry = ScopedInterfaceStackHead;
+		ScopedInterfaceStackHead = ScopedEntry;
+
+#if !UE_BUILD_TEST && !UE_BUILD_SHIPPING
+		// In development builds, we lazily query right away to ensure the interface we push is present
+		// In Test and Shipping, we'll do so only when/if the interface is actually queried
+		ensure(ScopedEntry->LazilyCacheTrait());
+#endif
+	}
+
+	bool FExecutionContext::PopScopedInterfaceImpl(FTraitInterfaceUID InterfaceUID, const FTraitBinding& Binding)
+	{
+		if (!Binding.IsValid())
+		{
+			return false;
+		}
+
+		// We don't have any specific handling for duplicate entries, if a scoped interface is pushed twice,
+		// it must also be popped twice (if popped manually)
+
+		// Start searching at the top of the stack
+		Private::FScopedInterfaceEntry* Entry = ScopedInterfaceStackHead;
+		if (Entry->InterfaceUID == InterfaceUID &&				// Same interface
+			Entry->Stack == *Binding.GetStack() &&				// Same stack
+			Entry->TraitIndex == Binding.GetTraitIndex())		// Same trait
+		{
+			// We found the interface we were looking for, pop it
+			// Add our entry to the free list
+			Private::FScopedInterfaceEntry* PrevEntry = Entry->PrevScopedInterfaceStackEntry;
+			Entry->NextFreeEntry = FreeScopedInterfaceEntryStackHead;
+			FreeScopedInterfaceEntryStackHead = Entry;
+			ScopedInterfaceStackHead = PrevEntry;
+
+			return true;
+		}
+
+		return false;
+	}
+
+	bool FExecutionContext::PopStackScopedInterfaces(const FTraitStackBinding& StackBinding)
+	{
+		if (!StackBinding.IsValid())
+		{
+			return false;
+		}
+
+		bool bAnyPopped = false;
+
+		// Start searching at the top of the stack
+		Private::FScopedInterfaceEntry* Entry = ScopedInterfaceStackHead;
+		while (Entry != nullptr)
+		{
+			if (Entry->Stack != StackBinding)
+			{
+				// This entry doesn't match our trait stack, stop searching
+				break;
+			}
+
+			// We found a scoped interface owned by the trait stack, pop it
+			// Add our entry to the free list
+			Private::FScopedInterfaceEntry* PrevEntry = Entry->PrevScopedInterfaceStackEntry;
+			Entry->NextFreeEntry = FreeScopedInterfaceEntryStackHead;
+			FreeScopedInterfaceEntryStackHead = Entry;
+			ScopedInterfaceStackHead = PrevEntry;
+			bAnyPopped = true;
+
+			// Continue execution in case this trait stack pushed multiple scoped interfaces
+
+			// Move to the next entry on the stack
+			Entry = PrevEntry;
+		}
+
+		return bAnyPopped;
+	}
+
+	bool FExecutionContext::GetScopedInterfaceImpl(FTraitInterfaceUID InterfaceUID, FTraitBinding& OutBinding) const
+	{
+		// Start searching at the top of the stack
+		Private::FScopedInterfaceEntry* Entry = ScopedInterfaceStackHead;
+		while (Entry != nullptr)
+		{
+			if (Entry->InterfaceUID == InterfaceUID)
+			{
+				// We found the interface we were looking for, return it
+				Entry->LazilyCacheTrait();
+
+				OutBinding = Entry->Trait;
+				return true;
+			}
+
+			// Move to the next entry on the stack
+			Entry = Entry->PrevScopedInterfaceStackEntry;
+		}
+
+		// We didn't find the interface we were looking for
+		OutBinding.Reset();
+		return false;
+	}
+
+	void FExecutionContext::ForEachScopedInterfaceImpl(FTraitInterfaceUID InterfaceUID, TFunctionRef<bool(FTraitBinding& Binding)> InFunction) const
+	{
+		// Start searching at the top of the stack
+		Private::FScopedInterfaceEntry* Entry = ScopedInterfaceStackHead;
+		while (Entry != nullptr)
+		{
+			if (Entry->InterfaceUID == InterfaceUID)
+			{
+				// We found the interface we were looking for, forward it to our callback
+				Entry->LazilyCacheTrait();
+
+				const bool bContinueSearching = InFunction(Entry->Trait);
+				if (!bContinueSearching)
+				{
+					break;	// The callback returned false, we are done searching
+				}
+			}
+
+			// Move to the next entry on the stack
+			Entry = Entry->PrevScopedInterfaceStackEntry;
+		}
+	}
+
 	FGraphInstanceComponent* FExecutionContext::TryGetComponent(int32 ComponentNameHash, FName ComponentName) const
 	{
 		check(IsBound());
diff --git a/Engine/Plugins/Experimental/AnimNext/Source/AnimNext/Private/TraitInterfaces/IUpdate.cpp b/Engine/Plugins/Experimental/AnimNext/Source/AnimNext/Private/TraitInterfaces/IUpdate.cpp
--- a/Engine/Plugins/Experimental/AnimNext/Source/AnimNext/Private/TraitInterfaces/IUpdate.cpp
+++ b/Engine/Plugins/Experimental/AnimNext/Source/AnimNext/Private/TraitInterfaces/IUpdate.cpp
@@ -328,6 +328,9 @@ namespace UE::AnimNext
 				check(Entry->UpdateTrait.IsValid());
 				Entry->UpdateTrait.PostUpdate(TraversalContext, Entry->TraitState);
 
+				// Now that it finished updating, we can pop any scoped interfaces this node might have pushed
+				TraversalContext.PopStackScopedInterfaces(Entry->TraitStack);
+
 				// We don't need this entry anymore
 				TraversalContext.PushFreeEntry(Entry);
 			}
@@ -339,5 +342,9 @@ namespace UE::AnimNext
 		{
 			It.Value()->PostUpdate(TraversalContext);
 		}
+
+		// At this point, we shouldn't have any remaining scoped interfaces
+		// If this fails, it means we failed to pop them due to a push/pop mismatch
+		ensure(!TraversalContext.HasScopedInterfaces());
 	}
 }
diff --git a/Engine/Plugins/Experimental/AnimNext/Source/AnimNext/Public/TraitCore/ExecutionContext.h b/Engine/Plugins/Experimental/AnimNext/Source/AnimNext/Public/TraitCore/ExecutionContext.h
--- a/Engine/Plugins/Experimental/AnimNext/Source/AnimNext/Public/TraitCore/ExecutionContext.h
+++ b/Engine/Plugins/Experimental/AnimNext/Source/AnimNext/Public/TraitCore/ExecutionContext.h
@@ -8,6 +8,7 @@
 #include "TraitCore/TraitPtr.h"
 #include "TraitCore/TraitStackBinding.h"
 #include "TraitCore/ITraitInterface.h"
+#include "TraitCore/IScopedTraitInterface.h"
 #include "TraitCore/TraitInterfaceUID.h"
 #include "TraitCore/LatentPropertyHandle.h"
 #include "TraitCore/NodeHandle.h"
@@ -21,6 +22,11 @@ class FMemStack;
 
 namespace UE::AnimNext
 {
+	namespace Private
+	{
+		struct FScopedInterfaceEntry;
+	}
+
 	struct FNodeDescription;
 	struct FNodeInstance;
 	struct FNodeTemplateRegistry;
@@ -81,6 +87,55 @@ namespace UE::AnimNext
 
 
 
+		//////////////////////////////////////////////////////////////////////////
+		// The following functions allow for scope interface management and queries
+		// Scoped interfaces live on a stack within the execution context
+		// A parent node can publish a scoped interface that its children can query for
+
+		// Pushed a scoped interface with the specified interface owned by the specified trait
+		// Scoped interfaces automatically pop when the owned trait finishes its update
+		template<class ScopedTraitInterface>
+		void PushScopedInterface(const FTraitBinding& Binding);
+
+		// Pushed an optional scoped interface with the specified interface owned by the specified trait
+		// The scoped interface is only pushed if the condition is true
+		// Scoped interfaces automatically pop when the owned trait finishes its update
+		template<class ScopedTraitInterface>
+		void PushOptionalScopedInterface(bool bCondition, const FTraitBinding& Binding);
+
+		// Pops a scoped interface owned by the specified trait
+		// An interface can only be popped if it lives at the top of the stack
+		// Returns false if we failed to pop the specified interface
+		// Failure can occur if the scoped interface is missing or not present on top of the stack
+		template<class ScopedTraitInterface>
+		bool PopScopedInterface(const FTraitBinding& Binding);
+
+		// Pops all scoped interfaces owned by the trait stack that contains the specified trait
+		// An interface can only be popped if it lives at the top of the stack
+		// Returns true if any scoped interfaces were popped, false otherwise
+		bool PopStackScopedInterfaces(const FTraitBinding& Binding);
+
+		// Pops all scoped interfaces owned by the specified trait stack
+		// An interface can only be popped if it lives at the top of the stack
+		// Returns true if any scoped interfaces were popped, false otherwise
+		bool PopStackScopedInterfaces(const FTraitStackBinding& StackBinding);
+
+		// Queries the scoped interface stack starting at the top for a trait that implements the specified interface.
+		// If no such trait exists, false is returned.
+		template<class ScopedTraitInterface>
+		bool GetScopedInterface(TTraitBinding<ScopedTraitInterface>& OutBinding) const;
+
+		// Queries the scoped interface stack starting at the top for a trait that implements the specified interface.
+		// When found, the provided callback will be called with the scoped interface binding.
+		// Iteration will continue as long as the callback returns true.
+		template<class ScopedTraitInterface>
+		void ForEachScopedInterface(TFunctionRef<bool(TTraitBinding<ScopedTraitInterface>& Binding)> InFunction) const;
+
+		// Returns whether or not any scoped interfaces have been pushed onto our stack
+		bool HasScopedInterfaces() const { return ScopedInterfaceStackHead != nullptr; }
+
+
+
 		//////////////////////////////////////////////////////////////////////////
 		// The following functions handle node lifetime management
 
@@ -134,6 +189,10 @@ namespace UE::AnimNext
 		FExecutionContext(const FExecutionContext&) = delete;
 		FExecutionContext& operator=(const FExecutionContext&) = delete;
 
+		void PushScopedInterfaceImpl(FTraitInterfaceUID InterfaceUID, const FTraitBinding& Binding);
+		bool PopScopedInterfaceImpl(FTraitInterfaceUID InterfaceUID, const FTraitBinding& Binding);
+		bool GetScopedInterfaceImpl(FTraitInterfaceUID InterfaceUID, FTraitBinding& OutBinding) const;
+		void ForEachScopedInterfaceImpl(FTraitInterfaceUID InterfaceUID, TFunctionRef<bool(FTraitBinding& Binding)> InFunction) const;
 		FGraphInstanceComponent* TryGetComponent(int32 ComponentNameHash, FName ComponentName) const;
 		FGraphInstanceComponent& AddComponent(int32 ComponentNameHash, FName ComponentName, TSharedPtr<FGraphInstanceComponent>&& Component) const;
 
@@ -154,12 +213,85 @@ namespace UE::AnimNext
 		// Root graph instance we are bound to
 		FAnimNextGraphInstance* RootGraphInstance = nullptr;
 
+		// The head pointer of the scoped interface stack
+		Private::FScopedInterfaceEntry* ScopedInterfaceStackHead = nullptr;
+
+		// The head pointer of the free scoped interface entry stack
+		// Entries are allocated from the memstack and are re-used in LIFO since they'll
+		// be warmer in the CPU cache
+		Private::FScopedInterfaceEntry* FreeScopedInterfaceEntryStackHead = nullptr;
+
 		friend struct FTraitStackBinding;
 	};
 
 	//////////////////////////////////////////////////////////////////////////
 	// Inline implementations
 
+	template<class ScopedTraitInterface>
+	inline void FExecutionContext::PushScopedInterface(const FTraitBinding& Binding)
+	{
+		static_assert(std::is_base_of<IScopedTraitInterface, ScopedTraitInterface>::value, "ScopedTraitInterface type must derive from ITraitInterface");
+
+		constexpr FTraitInterfaceUID InterfaceUID = ScopedTraitInterface::InterfaceUID;
+		PushScopedInterfaceImpl(InterfaceUID, Binding);
+	}
+
+	template<class ScopedTraitInterface>
+	inline void FExecutionContext::PushOptionalScopedInterface(bool bCondition, const FTraitBinding& Binding)
+	{
+		static_assert(std::is_base_of<IScopedTraitInterface, ScopedTraitInterface>::value, "ScopedTraitInterface type must derive from ITraitInterface");
+
+		if (!bCondition)
+		{
+			return;	// Scoped interface not required
+		}
+
+		PushScopedInterface<ScopedTraitInterface>(Binding);
+	}
+
+	template<class ScopedTraitInterface>
+	inline bool FExecutionContext::PopScopedInterface(const FTraitBinding& Binding)
+	{
+		static_assert(std::is_base_of<IScopedTraitInterface, ScopedTraitInterface>::value, "ScopedTraitInterface type must derive from ITraitInterface");
+
+		constexpr FTraitInterfaceUID InterfaceUID = ScopedTraitInterface::InterfaceUID;
+		return PopScopedInterfaceImpl(InterfaceUID, Binding);
+	}
+
+	inline bool FExecutionContext::PopStackScopedInterfaces(const FTraitBinding& Binding)
+	{
+		if (!Binding.IsValid())
+		{
+			return false;
+		}
+
+		return PopStackScopedInterfaces(*Binding.GetStack());
+	}
+
+	template<class ScopedTraitInterface>
+	inline bool FExecutionContext::GetScopedInterface(TTraitBinding<ScopedTraitInterface>& OutBinding) const
+	{
+		static_assert(std::is_base_of<IScopedTraitInterface, ScopedTraitInterface>::value, "ScopedTraitInterface type must derive from IScopedTraitInterface");
+
+		constexpr FTraitInterfaceUID InterfaceUID = ScopedTraitInterface::InterfaceUID;
+		return GetScopedInterfaceImpl(InterfaceUID, OutBinding);
+	}
+
+	template<class ScopedTraitInterface>
+	inline void FExecutionContext::ForEachScopedInterface(TFunctionRef<bool(TTraitBinding<ScopedTraitInterface>& Binding)> InFunction) const
+	{
+		static_assert(std::is_base_of<IScopedTraitInterface, ScopedTraitInterface>::value, "ScopedTraitInterface type must derive from IScopedTraitInterface");
+
+		constexpr FTraitInterfaceUID InterfaceUID = ScopedTraitInterface::InterfaceUID;
+
+		ForEachScopedInterfaceImpl(InterfaceUID,
+			[&InFunction](FTraitBinding& Binding)
+			{
+				// This cast is safe because we already queried the interface into the binding
+				return InFunction(static_cast<TTraitBinding<ScopedTraitInterface>&>(Binding));
+			});
+	}
+
 	inline FTraitPtr FExecutionContext::AllocateNodeInstance(const FTraitBinding& ParentBinding, FAnimNextTraitHandle ChildTraitHandle) const
 	{
 		return AllocateNodeInstance(ParentBinding.GetTraitPtr(), ChildTraitHandle);
diff --git a/Engine/Plugins/Experimental/AnimNext/Source/AnimNext/Public/TraitCore/IScopedTraitInterface.h b/Engine/Plugins/Experimental/AnimNext/Source/AnimNext/Public/TraitCore/IScopedTraitInterface.h
new file mode 100644
--- /dev/null
+++ b/Engine/Plugins/Experimental/AnimNext/Source/AnimNext/Public/TraitCore/IScopedTraitInterface.h
@@ -0,0 +1,19 @@
+// Copyright Epic Games, Inc. All Rights Reserved.
+
+#pragma once
+
+#include "CoreMinimal.h"
+#include "TraitCore/ITraitInterface.h"
+
+namespace UE::AnimNext
+{
+	/**
+	 * IScopedTraitInterface
+	 * 
+	 * Base type for all scoped trait interfaces. Used for type safety.
+	 */
+	struct ANIMNEXT_API IScopedTraitInterface : ITraitInterface
+	{
+		DECLARE_ANIM_TRAIT_INTERFACE(IScopedTraitInterface, 0x621b01eb)
+	};
+}
diff --git a/Engine/Plugins/Experimental/AnimNext/Source/AnimNext/Public/TraitCore/TraitBinding.h b/Engine/Plugins/Experimental/AnimNext/Source/AnimNext/Public/TraitCore/TraitBinding.h
--- a/Engine/Plugins/Experimental/AnimNext/Source/AnimNext/Public/TraitCore/TraitBinding.h
+++ b/Engine/Plugins/Experimental/AnimNext/Source/AnimNext/Public/TraitCore/TraitBinding.h
@@ -21,6 +21,11 @@ namespace UE::AnimNext
 	struct ITraitInterface;
 	struct FNodeDescription;
 
+	namespace Private
+	{
+		struct FScopedInterfaceEntry;
+	}
+
 	/**
 	 * FTraitBinding
 	 * 
@@ -210,6 +215,7 @@ namespace UE::AnimNext
 
 		friend FExecutionContext;
 		friend FTraitStackBinding;
+		friend Private::FScopedInterfaceEntry;
 	};
 
 	/**
diff --git a/Engine/Plugins/Experimental/AnimNext/Source/AnimNextTestSuite/Private/AnimNextTraitInterfacesTest.cpp b/Engine/Plugins/Experimental/AnimNext/Source/AnimNextTestSuite/Private/AnimNextTraitInterfacesTest.cpp
--- a/Engine/Plugins/Experimental/AnimNext/Source/AnimNextTestSuite/Private/AnimNextTraitInterfacesTest.cpp
+++ b/Engine/Plugins/Experimental/AnimNext/Source/AnimNextTestSuite/Private/AnimNextTraitInterfacesTest.cpp
@@ -16,6 +16,7 @@
 #include "TraitCore/TraitWriter.h"
 #include "TraitCore/ExecutionContext.h"
 #include "TraitCore/ITraitInterface.h"
+#include "TraitCore/IScopedTraitInterface.h"
 #include "TraitCore/NodeInstance.h"
 #include "TraitCore/NodeTemplateBuilder.h"
 #include "TraitCore/NodeTemplateRegistry.h"
@@ -35,6 +36,10 @@ namespace UE::AnimNext
 	{
 		static TArray<FTraitUID>* UpdatedTraits = nullptr;
 		static TArray<FTraitUID>* EvaluatedTraits = nullptr;
+
+		static FName TestTag(TEXT("MyTag"));
+		static TArray<bool>* IsTagInScope = nullptr;
+		static bool AutoPopTag = false;
 	}
 
 	struct FTraitWithNoChildren : FBaseTrait, IUpdate, IEvaluate
@@ -222,6 +227,130 @@ namespace UE::AnimNext
 
 	GENERATE_ANIM_TRAIT_IMPLEMENTATION(FTraitWithChildren, TRAIT_INTERFACE_ENUMERATOR)
 	#undef TRAIT_INTERFACE_ENUMERATOR
+
+	struct IScopedTagInterface : IScopedTraitInterface
+	{
+		DECLARE_ANIM_TRAIT_INTERFACE(IScopedTagInterface, 0xb8f4b397)
+
+		virtual FName GetTag(const FExecutionContext& Context, const TTraitBinding<IScopedTagInterface>& Binding) const;
+
+		static bool IsTagInScope(const FExecutionContext& Context, FName Tag);
+	};
+
+	template<>
+	struct TTraitBinding<IScopedTagInterface> : FTraitBinding
+	{
+		FName GetTag(const FExecutionContext& Context) const
+		{
+			return GetInterface()->GetTag(Context, *this);
+		}
+
+	protected:
+		const IScopedTagInterface* GetInterface() const { return GetInterfaceTyped<IScopedTagInterface>(); }
+	};
+
+	FName IScopedTagInterface::GetTag(const FExecutionContext& Context, const TTraitBinding<IScopedTagInterface>& Binding) const
+	{
+		TTraitBinding<IScopedTagInterface> SuperBinding;
+		if (Binding.GetStackInterfaceSuper(SuperBinding))
+		{
+			return SuperBinding.GetTag(Context);
+		}
+
+		return NAME_None;
+	}
+
+	bool IScopedTagInterface::IsTagInScope(const FExecutionContext& Context, FName Tag)
+	{
+		bool bResult = false;
+
+		Context.ForEachScopedInterface<IScopedTagInterface>([&Context, Tag, &bResult](TTraitBinding<IScopedTagInterface>& InterfaceBinding)
+			{
+				if (Tag == InterfaceBinding.GetTag(Context))
+				{
+					// We found our tag, stop iterating
+					bResult = true;
+					return false;
+				}
+
+				// Keep searching
+				return true;
+			});
+
+		return bResult;
+	}
+
+	// Adds a scoped tag
+	struct FScopedTagTrait : FAdditiveTrait, IScopedTagInterface, IUpdate
+	{
+		DECLARE_ANIM_TRAIT(FScopedTagTrait, 0x4b296948, FAdditiveTrait)
+
+		// IScopedTagInterface impl
+		virtual FName GetTag(const FExecutionContext& Context, const TTraitBinding<IScopedTagInterface>& Binding) const override
+		{
+			return Private::TestTag;
+		}
+
+		// IUpdate impl
+		virtual void PreUpdate(FUpdateTraversalContext& Context, const TTraitBinding<IUpdate>& Binding, const FTraitUpdateState& TraitState) const override
+		{
+			Context.PushScopedInterface<IScopedTagInterface>(Binding);
+
+			IUpdate::PreUpdate(Context, Binding, TraitState);
+		}
+
+		virtual void PostUpdate(FUpdateTraversalContext& Context, const TTraitBinding<IUpdate>& Binding, const FTraitUpdateState& TraitState) const override
+		{
+			if (!Private::AutoPopTag)
+			{
+				ensure(Context.PopScopedInterface<IScopedTagInterface>(Binding));
+			}
+
+			IUpdate::PostUpdate(Context, Binding, TraitState);
+		}
+	};
+
+	// Trait implementation boilerplate
+	#define TRAIT_INTERFACE_ENUMERATOR(GeneratorMacro) \
+		GeneratorMacro(IScopedTagInterface) \
+		GeneratorMacro(IUpdate) \
+
+	GENERATE_ANIM_TRAIT_IMPLEMENTATION(FScopedTagTrait, TRAIT_INTERFACE_ENUMERATOR)
+	#undef TRAIT_INTERFACE_ENUMERATOR
+
+	// Tests if we have a scoped tag
+	struct FTestScopedTagTrait : FAdditiveTrait, IUpdate
+	{
+		DECLARE_ANIM_TRAIT(FTestScopedTagTrait, 0x922bd23d, FAdditiveTrait)
+
+		// IUpdate impl
+		virtual void PreUpdate(FUpdateTraversalContext& Context, const TTraitBinding<IUpdate>& Binding, const FTraitUpdateState& TraitState) const override
+		{
+			if (Private::IsTagInScope != nullptr)
+			{
+				Private::IsTagInScope->Add(IScopedTagInterface::IsTagInScope(Context, Private::TestTag));
+			}
+
+			IUpdate::PreUpdate(Context, Binding, TraitState);
+		}
+
+		virtual void PostUpdate(FUpdateTraversalContext& Context, const TTraitBinding<IUpdate>& Binding, const FTraitUpdateState& TraitState) const override
+		{
+			if (Private::IsTagInScope != nullptr)
+			{
+				Private::IsTagInScope->Add(IScopedTagInterface::IsTagInScope(Context, Private::TestTag));
+			}
+
+			IUpdate::PostUpdate(Context, Binding, TraitState);
+		}
+	};
+
+	// Trait implementation boilerplate
+	#define TRAIT_INTERFACE_ENUMERATOR(GeneratorMacro) \
+		GeneratorMacro(IUpdate) \
+
+	GENERATE_ANIM_TRAIT_IMPLEMENTATION(FTestScopedTagTrait, TRAIT_INTERFACE_ENUMERATOR)
+	#undef TRAIT_INTERFACE_ENUMERATOR
 }
 
 IMPLEMENT_SIMPLE_AUTOMATION_TEST(FAnimationAnimNextRuntimeTest_IHierarchy, "Animation.AnimNext.Runtime.IHierarchy", EAutomationTestFlags::EditorContext | EAutomationTestFlags::EngineFilter)
@@ -718,4 +847,159 @@ bool FAnimationAnimNextRuntimeTest_IEvaluate::RunTest(const FString& InParameter
 
 	return true;
 }
+IMPLEMENT_SIMPLE_AUTOMATION_TEST(FAnimationAnimNextRuntimeTest_IScopedInterface, "Animation.AnimNext.Runtime.IScopedInterface", EAutomationTestFlags::EditorContext | EAutomationTestFlags::EngineFilter)
+
+bool FAnimationAnimNextRuntimeTest_IScopedInterface::RunTest(const FString& InParameters)
+{
+	using namespace UE::AnimNext;
+	{
+		AUTO_REGISTER_ANIM_TRAIT(FTraitWithOneChild)
+		AUTO_REGISTER_ANIM_TRAIT(FScopedTagTrait)
+		AUTO_REGISTER_ANIM_TRAIT(FTestScopedTagTrait)
+
+		UFactory* GraphFactory = NewObject<UAnimNextGraphFactory>();
+		UAnimNextGraph* AnimNextGraph = CastChecked<UAnimNextGraph>(GraphFactory->FactoryCreateNew(UAnimNextGraph::StaticClass(), GetTransientPackage(), TEXT("TestAnimNextGraph"), RF_Transient, nullptr, nullptr, NAME_None));
+		UE_RETURN_ON_ERROR(AnimNextGraph != nullptr, "FAnimationAnimNextRuntimeTest_IScopedInterface -> Failed to create animation graph");
+
+		FScopedClearNodeTemplateRegistry ScopedClearNodeTemplateRegistry;
+		FNodeTemplateRegistry& Registry = FNodeTemplateRegistry::Get();
+
+		// We create a few node templates
+		// Template A has a single child and tests for our tag
+		TArray<FTraitUID> NodeTemplateTraitList0;
+		NodeTemplateTraitList0.Add(FTraitWithOneChild::TraitUID);
+		NodeTemplateTraitList0.Add(FTestScopedTagTrait::TraitUID);
+
+		// Template B has a single child, it tests and pushes our tag
+		TArray<FTraitUID> NodeTemplateTraitList1;
+		NodeTemplateTraitList1.Add(FTraitWithOneChild::TraitUID);
+		NodeTemplateTraitList1.Add(FTestScopedTagTrait::TraitUID);	// Test after push/pop
+		NodeTemplateTraitList1.Add(FScopedTagTrait::TraitUID);
+		NodeTemplateTraitList1.Add(FTestScopedTagTrait::TraitUID);	// Test before push/pop
+
+		// Populate our node template registry
+		TArray<uint8> NodeTemplateBuffer0, NodeTemplateBuffer1;
+		const FNodeTemplate* NodeTemplate0 = FNodeTemplateBuilder::BuildNodeTemplate(NodeTemplateTraitList0, NodeTemplateBuffer0);
+		const FNodeTemplate* NodeTemplate1 = FNodeTemplateBuilder::BuildNodeTemplate(NodeTemplateTraitList1, NodeTemplateBuffer1);
+
+		// Build our graph, it as follow:
+		// NodeA has no child (tag is scoped)
+		// NodeB has one child: NodeA (NodeB adds the scoped tag)
+		// NodeC (root) has one child: NodeB (no tag scoped)
+
+		TArray<FNodeHandle> NodeHandles;
+
+		// Write our graph
+		TArray<uint8> GraphSharedDataArchiveBuffer;
+		TArray<TObjectPtr<UObject>> GraphReferencedObjects;
+		{
+			FTraitWriter TraitWriter;
+
+			NodeHandles.Add(TraitWriter.RegisterNode(*NodeTemplate0));	// NodeC (root node)
+			NodeHandles.Add(TraitWriter.RegisterNode(*NodeTemplate1));	// NodeB
+			NodeHandles.Add(TraitWriter.RegisterNode(*NodeTemplate0));	// NodeA
+
+			// We don't have trait properties
+			TArray<TMap<FName, FString>> TraitPropertiesA;
+			TraitPropertiesA.AddDefaulted(NodeTemplateTraitList0.Num());
+			TraitPropertiesA[0].Add(TEXT("Child"), ToString<FTraitWithOneChild::FSharedData>(TEXT("Child"), FAnimNextTraitHandle()));
+
+			TArray<TMap<FName, FString>> TraitPropertiesB;
+			TraitPropertiesB.AddDefaulted(NodeTemplateTraitList1.Num());
+			TraitPropertiesB[0].Add(TEXT("Child"), ToString<FTraitWithOneChild::FSharedData>(TEXT("Child"), FAnimNextTraitHandle(NodeHandles[2])));
+
+			TArray<TMap<FName, FString>> TraitPropertiesC;
+			TraitPropertiesC.AddDefaulted(NodeTemplateTraitList0.Num());
+			TraitPropertiesC[0].Add(TEXT("Child"), ToString<FTraitWithOneChild::FSharedData>(TEXT("Child"), FAnimNextTraitHandle(NodeHandles[1])));
+
+			TraitWriter.BeginNodeWriting();
+			TraitWriter.WriteNode(NodeHandles[0],
+				[&TraitPropertiesC](uint32 TraitIndex, FName PropertyName)
+				{
+					return TraitPropertiesC[TraitIndex][PropertyName];
+				},
+				[](uint32 TraitIndex, FName PropertyName)
+				{
+					return MAX_uint16;
+				});
+			TraitWriter.WriteNode(NodeHandles[1],
+				[&TraitPropertiesB](uint32 TraitIndex, FName PropertyName)
+				{
+					return TraitPropertiesB[TraitIndex][PropertyName];
+				},
+				[](uint32 TraitIndex, FName PropertyName)
+				{
+					return MAX_uint16;
+				});
+			TraitWriter.WriteNode(NodeHandles[2],
+				[&TraitPropertiesA](uint32 TraitIndex, FName PropertyName)
+				{
+					return TraitPropertiesA[TraitIndex][PropertyName];
+				},
+				[](uint32 TraitIndex, FName PropertyName)
+				{
+					return MAX_uint16;
+				});
+			TraitWriter.EndNodeWriting();
+
+			AddErrorIfFalse(TraitWriter.GetErrorState() == FTraitWriter::EErrorState::None, "FAnimationAnimNextRuntimeTest_IScopedInterface -> Failed to write traits");
+			GraphSharedDataArchiveBuffer = TraitWriter.GetGraphSharedData();
+			GraphReferencedObjects = TraitWriter.GetGraphReferencedObjects();
+		}
+
+		// Read our graph
+		FTestUtils::LoadFromArchiveBuffer(*AnimNextGraph, NodeHandles, GraphSharedDataArchiveBuffer);
+
+		FAnimNextGraphInstancePtr GraphInstance;
+		AnimNextGraph->AllocateInstance(GraphInstance);
+
+		FExecutionContext Context(GraphInstance);
+
+		{
+			TArray<bool> IsTagInScope;
+			Private::IsTagInScope = &IsTagInScope;
+
+			// Call pre/post update on our graph
+			Private::AutoPopTag = true;
+			UpdateGraph(GraphInstance, 0.0333f);
+
+			AddErrorIfFalse(IsTagInScope.Num() == 8, "FAnimationAnimNextRuntimeTest_IScopedInterface -> Unexpected number of entries");
+			AddErrorIfFalse(IsTagInScope[0] == false, "FAnimationAnimNextRuntimeTest_IScopedInterface -> Unexpected scoped tag state");		// NodeC::PreUpdate (template 0)
+			AddErrorIfFalse(IsTagInScope[1] == false, "FAnimationAnimNextRuntimeTest_IScopedInterface -> Unexpected scoped tag state");		// NodeB::Before::PreUpdate (template 1)
+			AddErrorIfFalse(IsTagInScope[2] == true, "FAnimationAnimNextRuntimeTest_IScopedInterface -> Unexpected scoped tag state");		// NodeB::After::PreUpdate (template 1)
+			AddErrorIfFalse(IsTagInScope[3] == true, "FAnimationAnimNextRuntimeTest_IScopedInterface -> Unexpected scoped tag state");		// NodeA::PreUpdate (template 0)
+			AddErrorIfFalse(IsTagInScope[4] == true, "FAnimationAnimNextRuntimeTest_IScopedInterface -> Unexpected scoped tag state");		// NodeA::PostUpdate (template 0)
+			AddErrorIfFalse(IsTagInScope[5] == true, "FAnimationAnimNextRuntimeTest_IScopedInterface -> Unexpected scoped tag state");		// NodeB::Before::PostUpdate (template 1)
+			AddErrorIfFalse(IsTagInScope[6] == true, "FAnimationAnimNextRuntimeTest_IScopedInterface -> Unexpected scoped tag state");		// NodeB::After::PostUpdate (template 1)
+			AddErrorIfFalse(IsTagInScope[7] == false, "FAnimationAnimNextRuntimeTest_IScopedInterface -> Unexpected scoped tag state");		// NodeC::PostUpdate (template 0)
+
+			IsTagInScope.Reset();
+			Private::AutoPopTag = false;
+
+			// Call pre/post update on our graph
+			UpdateGraph(GraphInstance, 0.0333f);
+
+			AddErrorIfFalse(IsTagInScope.Num() == 8, "FAnimationAnimNextRuntimeTest_IScopedInterface -> Unexpected number of entries");
+			AddErrorIfFalse(IsTagInScope[0] == false, "FAnimationAnimNextRuntimeTest_IScopedInterface -> Unexpected scoped tag state");		// NodeC::PreUpdate (template 0)
+			AddErrorIfFalse(IsTagInScope[1] == false, "FAnimationAnimNextRuntimeTest_IScopedInterface -> Unexpected scoped tag state");		// NodeB::Before::PreUpdate (template 1)
+			AddErrorIfFalse(IsTagInScope[2] == true, "FAnimationAnimNextRuntimeTest_IScopedInterface -> Unexpected scoped tag state");		// NodeB::After::PreUpdate (template 1)
+			AddErrorIfFalse(IsTagInScope[3] == true, "FAnimationAnimNextRuntimeTest_IScopedInterface -> Unexpected scoped tag state");		// NodeA::PreUpdate (template 0)
+			AddErrorIfFalse(IsTagInScope[4] == true, "FAnimationAnimNextRuntimeTest_IScopedInterface -> Unexpected scoped tag state");		// NodeA::PostUpdate (template 0)
+			AddErrorIfFalse(IsTagInScope[5] == true, "FAnimationAnimNextRuntimeTest_IScopedInterface -> Unexpected scoped tag state");		// NodeB::Before::PostUpdate (template 1)
+			AddErrorIfFalse(IsTagInScope[6] == false, "FAnimationAnimNextRuntimeTest_IScopedInterface -> Unexpected scoped tag state");		// NodeB::After::PostUpdate (template 1)
+			AddErrorIfFalse(IsTagInScope[7] == false, "FAnimationAnimNextRuntimeTest_IScopedInterface -> Unexpected scoped tag state");		// NodeC::PostUpdate (template 0)
+
+			Private::UpdatedTraits = nullptr;
+		}
+
+		Registry.Unregister(NodeTemplate0);
+		Registry.Unregister(NodeTemplate1);
+
+		AddErrorIfFalse(Registry.GetNum() == 0, "FAnimationAnimNextRuntimeTest_IScopedInterface -> Registry should contain 0 templates");
+	}
+	Tests::FUtils::CleanupAfterTests();
+
+	return true;
+}
+
 #endif
