Addd multiselection support to AnimNext EdGraphNode details panel
AnimNext Traits and RigVMNodes are displayed using categories
Updated FStructOnScopeStructureDataProvider and FInstancePropertyBagStructureDataProvider to allow multiple instances of the passed struct / PropertyBag

diff --git a/Engine/Plugins/Experimental/AnimNext/Source/AnimNextEditor/Private/Graph/AnimNextGraph_EdGraphNodeCustomization.cpp b/Engine/Plugins/Experimental/AnimNext/Source/AnimNextEditor/Private/Graph/AnimNextGraph_EdGraphNodeCustomization.cpp
--- a/Engine/Plugins/Experimental/AnimNext/Source/AnimNextEditor/Private/Graph/AnimNextGraph_EdGraphNodeCustomization.cpp
+++ b/Engine/Plugins/Experimental/AnimNext/Source/AnimNextEditor/Private/Graph/AnimNextGraph_EdGraphNodeCustomization.cpp
@@ -20,200 +20,352 @@ void FAnimNextGraph_EdGraphNodeCustomization::CustomizeDetails(IDetailLayoutBuil
 	TArray<TWeakObjectPtr<UObject>> Objects;
 	DetailBuilder.GetObjectsBeingCustomized(Objects);
 
-	if (Objects.IsEmpty() || Objects.Num() > 1)
+	if (Objects.IsEmpty())
 	{
 		return;
 	}
 
-	if (UAnimNextGraph_EdGraphNode* EdGraphNode = Cast<UAnimNextGraph_EdGraphNode>(Objects[0].Get()))
+	CustomizeObjects(DetailBuilder, Objects);
+}
+
+void FAnimNextGraph_EdGraphNodeCustomization::CustomizeObjects(IDetailLayoutBuilder& DetailBuilder, const TArray<TWeakObjectPtr<UObject>>& Objects)
+{
+	const int32 NumNodes = Objects.Num();
+	for (int i = 0; i < NumNodes; i++)
 	{
-		if (URigVMNode* ModelNode = EdGraphNode->GetModelNode())
+		if (UAnimNextGraph_EdGraphNode* EdGraphNode = Cast<UAnimNextGraph_EdGraphNode>(Objects[i].Get()))
 		{
-			// If there are any decorators in the node, treat it as a trait stack
-			const TArray<URigVMPin*> DecoratorPins = ModelNode->GetDecoratorPins();
-			if (DecoratorPins.Num() > 0)
+			if (URigVMNode* ModelNode = EdGraphNode->GetModelNode())
 			{
-				// For each Decorator (represented as a pin in the node)
-				for (const URigVMPin* DecoratorPin : DecoratorPins)
+				// If there are any traits in the node, treat it as a trait stack
+				const TArray<URigVMPin*> TraitPins = ModelNode->GetDecoratorPins();
+				if (TraitPins.Num() > 0)
 				{
-					if (DecoratorPin->IsExecuteContext())
-					{
-						continue;
-					}
-
-					// Create a temporary decorator instance, in order to get the correct DecoratorSharedDataStruct
-					if (TSharedPtr<FStructOnScope> ScopedDecorator = ModelNode->GetDecoratorInstance(DecoratorPin->GetFName()))
-					{
-						// Create a scoped struct with the Decorator Shared Instance Data
-						const FRigVMDecorator* Decorator = (FRigVMDecorator*)ScopedDecorator->GetStructMemory();
+					GenerateTraitData(EdGraphNode, CategoryDetailsData);
+				}
+				else
+				{
+					GenerateRigVMData(EdGraphNode, CategoryDetailsData);
+				}
+			}
+		}
+	}
 
-						// Import the default pin values into the Decorator Shared Instance Data
-						TSharedPtr<FStructOnScope> ScopedSharedData;
-						if (UScriptStruct* DecoratorSharedInstanceData = Decorator->GetDecoratorSharedDataStruct())
-						{
-							ScopedSharedData = MakeShared<FStructOnScope>(DecoratorSharedInstanceData);
+	for (TSharedPtr<FCategoryDetailsData>& DetailsData : CategoryDetailsData)
+	{
+		PopulateCategory(DetailBuilder, DetailsData);
+	}
+}
 
-							FRigVMPinDefaultValueImportErrorContext ErrorPipe;
-							LOG_SCOPE_VERBOSITY_OVERRIDE(LogExec, ELogVerbosity::Verbose);
-							const FString DefaultValue = DecoratorPin->GetDefaultValue();
-							DecoratorSharedInstanceData->ImportText(*DefaultValue, ScopedSharedData->GetStructMemory(), nullptr, PPF_SerializedAsImportText, &ErrorPipe, DecoratorSharedInstanceData->GetName());
-						}
-						else // or use the decorator itself for simple decorators
-						{
-							ScopedSharedData = ScopedDecorator;
-						}
+void FAnimNextGraph_EdGraphNodeCustomization::GenerateTraitData(UAnimNextGraph_EdGraphNode* EdGraphNode, TArray<TSharedPtr<FCategoryDetailsData>>& CategoryDetailsData)
+{
+	if (URigVMNode* ModelNode = EdGraphNode->GetModelNode())
+	{
+		// If there are any traits in the node, treat it as a trait stack
+		const TArray<URigVMPin*> TraitPins = ModelNode->GetDecoratorPins();
+		if (TraitPins.Num() > 0)
+		{
+			// For each Trait (represented as a pin in the node)
+			for (const URigVMPin* TraitPin : TraitPins)
+			{
+				if (TraitPin->IsExecuteContext())
+				{
+					continue;
+				}
 
-						// Create a category with the DisplayName of the decorator Shared Data
-						IDetailCategoryBuilder& ParameterCategory = DetailBuilder.EditCategory(*ScopedSharedData->GetStruct()->GetDisplayNameText().ToString(), FText::GetEmpty(), ECategoryPriority::Important);
+				// Create a temporary trait instance, in order to get the correct TraitSharedDataStruct
+				if (TSharedPtr<FStructOnScope> ScopedTrait = ModelNode->GetDecoratorInstance(TraitPin->GetFName()))
+				{
+					// Create a scoped struct with the Trait Shared Instance Data and store it for later use
+					// The idea is creating one template per Trait type, but pass all the selected instances
+					// So the multiselection works
+					const FRigVMDecorator* Trait = (FRigVMDecorator*)ScopedTrait->GetStructMemory();
 
-						// Now add all the structure properties and handle modifications
-						TArray<IDetailPropertyRow*> OutPropertiesRow;
-						ParameterCategory.AddAllExternalStructureProperties(ScopedSharedData.ToSharedRef(), EPropertyLocation::Default, &OutPropertiesRow);
-						for (IDetailPropertyRow* DetailPropertyRow : OutPropertiesRow)
-						{
-							if (TSharedPtr<IPropertyHandle> PropertyHandle = DetailPropertyRow->GetPropertyHandle(); PropertyHandle.IsValid())
+					// Import the default pin values into the Trait Shared Instance Data
+					FTraitStackDetailsData* TraitData = nullptr;
+					if (UScriptStruct* TraitSharedInstanceData = Trait->GetDecoratorSharedDataStruct())
+					{
+						const TSharedPtr<FCategoryDetailsData>* TraitDataPtr = CategoryDetailsData.FindByPredicate([TraitSharedInstanceData](const TSharedPtr<FCategoryDetailsData>& InItem)
 							{
-								const FProperty* Property = PropertyHandle->GetProperty();
+								return InItem->Type == FCategoryDetailsData::EType::TraitStack && InItem->Name == TraitSharedInstanceData->GetFName();
+							});
 
-								const auto ModifyAsset = [EdGraphNode]()
-								{
-									EdGraphNode->Modify();
-								};
+						if (TraitDataPtr == nullptr)
+						{
+							TraitDataPtr = &CategoryDetailsData.Add_GetRef(MakeShared<FTraitStackDetailsData>(TraitSharedInstanceData->GetFName()));
+						}
 
-								const auto UpdatePinDefaultValue = [EdGraphNode, Property, ScopedSharedData]()
-								{
-									// Extract the value from the property and assign it to the Pin as a default value (via Schema)
-									const uint8* StructMemberMemoryPtr = Property->ContainerPtrToValuePtr<uint8>(ScopedSharedData->GetStructMemory());
-									const FString ValueStr = FRigVMStruct::ExportToFullyQualifiedText(Property, StructMemberMemoryPtr, true);
+						TraitData = static_cast<FTraitStackDetailsData*>(TraitDataPtr->Get());
 
-									for (UEdGraphPin* EdGraphPin : EdGraphNode->Pins)
-									{
-										// Find the EdGraphPin that corresponds to the Property
-										if (EdGraphPin->GetFName().ToString().EndsWith(Property->GetFName().ToString(), ESearchCase::CaseSensitive))
-										{
-											if (URigVMPin* ModelPin = EdGraphNode->FindModelPinFromGraphPin(EdGraphPin))
-											{
-												// Set the default value using the Controller
-												EdGraphNode->GetController()->SetPinDefaultValue(ModelPin->GetPinPath(), ValueStr);
-											}
-											break;
-										}
-									}
-								};
+						// Store EdGraphNode and ScopedTraitData, as we will need that later to transfer data if the user makes modifications to the Traits in the details panel
+						TraitData->EdGraphNodes.Add(EdGraphNode);
+						const int32 ScopedDataIndex = TraitData->ScopedSharedDataInstances.Add(MakeShared<FStructOnScope>(TraitSharedInstanceData));
 
-								if (Decorator->GetDecoratorSharedDataStruct() != nullptr) // TODO zzz : Investigate why shared works with value and decorator struct with child value
-								{
-									PropertyHandle->SetOnPropertyValuePreChange(FSimpleDelegate::CreateLambda(ModifyAsset));
-									PropertyHandle->SetOnPropertyValueChanged(FSimpleDelegate::CreateLambda(UpdatePinDefaultValue));
-								}
-								else
-								{
-									PropertyHandle->SetOnChildPropertyValuePreChange(FSimpleDelegate::CreateLambda(ModifyAsset));
-									PropertyHandle->SetOnChildPropertyValueChanged(FSimpleDelegate::CreateLambda(UpdatePinDefaultValue));
-								}
-							}
-						}
+						// Fill the scoped data with the pin data
+						FRigVMPinDefaultValueImportErrorContext ErrorPipe;
+						LOG_SCOPE_VERBOSITY_OVERRIDE(LogExec, ELogVerbosity::Verbose);
+						const FString DefaultValue = TraitPin->GetDefaultValue();
+						TraitSharedInstanceData->ImportText(*DefaultValue, TraitData->ScopedSharedDataInstances[ScopedDataIndex]->GetStructMemory(), nullptr, PPF_SerializedAsImportText, &ErrorPipe, TraitSharedInstanceData->GetName());
 					}
 				}
 			}
-			else
+		}
+	}
+}
+
+void FAnimNextGraph_EdGraphNodeCustomization::GenerateRigVMData(UAnimNextGraph_EdGraphNode* EdGraphNode, TArray<TSharedPtr<FCategoryDetailsData>>& CategoryDetailsData)
+{
+	if (URigVMNode* ModelNode = EdGraphNode->GetModelNode())
+	{
+		// For nodes that aren't Traits, we display the pins as properties
+		const TArray<URigVMPin*> ModelPins = ModelNode->GetPins();
+		TArray<TWeakObjectPtr<URigVMPin>> PinsToDisplay;
+		PinsToDisplay.Reserve(ModelPins.Num());
+
+		// First, obtain the pins to display
+		for (URigVMPin* Pin : ModelPins)
+		{
+			if (Pin->IsExecuteContext())
 			{
-				// For nodes that aren't decorators, we display the pins as properties
-				const TArray<URigVMPin*> ModelPins = ModelNode->GetPins();
-				TArray<URigVMPin*> PinsToDisplay;
-				PinsToDisplay.Reserve(ModelPins.Num());
+				continue;
+			}
 
-				// First, obtain the pins to display
-				for (URigVMPin* Pin : ModelPins)
-				{
-					if (Pin->IsExecuteContext())
-					{
-						continue;
-					}
-					const ERigVMPinDirection PinDirection = Pin->GetDirection();
-					if (PinDirection != ERigVMPinDirection::Hidden && (PinDirection == ERigVMPinDirection::IO || PinDirection == ERigVMPinDirection::Input))
-					{
-						PinsToDisplay.Add(Pin);
-					}
-				}
+			const ERigVMPinDirection PinDirection = Pin->GetDirection();
+			if (PinDirection != ERigVMPinDirection::Hidden && (PinDirection == ERigVMPinDirection::IO || PinDirection == ERigVMPinDirection::Input))
+			{
+				PinsToDisplay.Add(Pin);
+			}
+		}
+
+		if (PinsToDisplay.Num() > 0)
+		{
+			const FName NodeName = *EdGraphNode->GetNodeTitle(ENodeTitleType::ListView).ToString(); // TODO : check if this is a good name to use, can't sue GetFName as it comes with instance postfix
 
-				if (PinsToDisplay.Num() > 0)
+			FRigVMNodeDetailsData* RigVMData = nullptr;
+
+			TSharedPtr<FCategoryDetailsData>* RigVMDataPtr = CategoryDetailsData.FindByPredicate([&NodeName](const TSharedPtr<FCategoryDetailsData>& InItem)
 				{
-					// Create the category where the properties will be added, using the Title of the EdGraphNode
-					IDetailCategoryBuilder& ParameterCategory = DetailBuilder.EditCategory(*EdGraphNode->GetNodeTitle(ENodeTitleType::ListView).ToString(), FText::GetEmpty(), ECategoryPriority::Default);
-
-					// Then, create a custom property bag to store the data, initializating the properties in the struct with the pin default values
-					GenerateMemoryStorage(PinsToDisplay, MemoryStorage);
-					
-					// And finally, add the struct properties to the customization
-					PopulateCategory(ParameterCategory, PinsToDisplay, MemoryStorage, EdGraphNode);
-				}
+					return InItem->Type == FCategoryDetailsData::EType::RigVMNode && InItem->Name == NodeName;
+				});
+
+			if (RigVMDataPtr == nullptr)
+			{
+				RigVMDataPtr = &CategoryDetailsData.Add_GetRef(MakeShared<FRigVMNodeDetailsData>(NodeName));
 			}
+
+			RigVMData = static_cast<FRigVMNodeDetailsData*>(RigVMDataPtr->Get());
+			RigVMData->ModelPinsToDisplay = PinsToDisplay;
+
+			// Store EdGraphNode and generated Memory, as we will need that later to transfer data if the user makes modifications in the details panel
+			RigVMData->EdGraphNodes.Add(EdGraphNode);
+			const int32 MemoryStorageIndex = RigVMData->MemoryStorages.Add(MakeShared<FRigVMMemoryStorageStruct>());
+
+			// Then, create a custom property bag to store the data, initializating the properties in the struct with the pin default values
+			GenerateMemoryStorage(PinsToDisplay, *RigVMData->MemoryStorages[MemoryStorageIndex].Get());
 		}
 	}
 }
 
-void FAnimNextGraph_EdGraphNodeCustomization::GenerateMemoryStorage(const TArray<URigVMPin*>& ModelPinsToDisplay, FRigVMMemoryStorageStruct& MemoryStorage)
+void FAnimNextGraph_EdGraphNodeCustomization::PopulateCategory(IDetailLayoutBuilder& DetailBuilder, const TSharedPtr<FCategoryDetailsData> &CategoryDetailsData)
 {
-	TArray<FRigVMPropertyDescription> PropertyDescriptions;
-	PropertyDescriptions.Reserve(ModelPinsToDisplay.Num());
-
-	for (const URigVMPin* ModelPin : ModelPinsToDisplay)
+	switch (CategoryDetailsData->Type)
 	{
-		FRigVMPropertyDescription& PropertyDesc = PropertyDescriptions.AddDefaulted_GetRef();
-
-		PropertyDesc.Name = ModelPin->GetFName();
-		PropertyDesc.Property = nullptr;
-		PropertyDesc.CPPType = ModelPin->GetCPPType();
-		PropertyDesc.CPPTypeObject = ModelPin->GetCPPTypeObject();
-		if (ModelPin->IsArray())
+		case FRigVMNodeDetailsData::EType::TraitStack:
 		{
-			PropertyDesc.Containers.Add(EPinContainerType::Array);
+			PopulateCategory(DetailBuilder, StaticCastSharedPtr<FTraitStackDetailsData>(CategoryDetailsData));
+			break;
 		}
-		PropertyDesc.DefaultValue = ModelPin->GetDefaultValue();
+		case FRigVMNodeDetailsData::EType::RigVMNode:
+		{
+			PopulateCategory(DetailBuilder, StaticCastSharedPtr<FRigVMNodeDetailsData>(CategoryDetailsData));
+			break;
+		}
+		default:
+			break;
 	}
-
-	MemoryStorage.AddProperties(PropertyDescriptions);
 }
 
-void FAnimNextGraph_EdGraphNodeCustomization::PopulateCategory(IDetailCategoryBuilder& Category, const TArray<URigVMPin*>& ModelPinsToDisplay, FRigVMMemoryStorageStruct& MemoryStorage, UAnimNextGraph_EdGraphNode* EdGraphNode)
+void FAnimNextGraph_EdGraphNodeCustomization::PopulateCategory(IDetailLayoutBuilder& DetailBuilder, const TSharedPtr<FTraitStackDetailsData>& TraitData)
 {
-	for (URigVMPin* ModelPin : ModelPinsToDisplay)
+	if (TraitData == nullptr)
 	{
-		FAddPropertyParams AddPropertyParams;
-		IDetailPropertyRow* DetailPropertyRow = Category.AddExternalStructureProperty(MakeShared<FInstancePropertyBagStructureDataProvider>(MemoryStorage), ModelPin->GetFName(), EPropertyLocation::Default, AddPropertyParams);
+		return;
+	}
 
-		if (TSharedPtr<IPropertyHandle> Handle = DetailPropertyRow->GetPropertyHandle(); Handle.IsValid())
-		{
-			Handle->SetOnChildPropertyValuePreChange(FSimpleDelegate::CreateLambda([EdGraphNode]()
+	check(TraitData->ScopedSharedDataInstances.Num() > 0);
+	check(TraitData->ScopedSharedDataInstances.Num() == TraitData->EdGraphNodes.Num());
+
+	// Create a category with the DisplayName of the Trait Shared Data
+	const FString TraitDisplayName = *TraitData->ScopedSharedDataInstances[0]->GetStruct()->GetDisplayNameText().ToString();
+	const FName CategoryName = (TraitData->EdGraphNodes.Num() == 1)
+		? *TraitDisplayName
+		: *FString::Printf(TEXT("%s (% d)"), *TraitDisplayName, TraitData->EdGraphNodes.Num());
+
+	IDetailCategoryBuilder& ParameterCategory = DetailBuilder.EditCategory(CategoryName, FText::GetEmpty(), ECategoryPriority::Important);
+	
+	FAddPropertyParams AddPropertyParams;
+	AddPropertyParams.CreateCategoryNodes(true);
+	AddPropertyParams.HideRootObjectNode(true);
+
+	IDetailPropertyRow* DetailPropertyRow = ParameterCategory.AddExternalStructureProperty(MakeShared<FStructOnScopeStructureDataProvider>(TraitData->ScopedSharedDataInstances), NAME_None, EPropertyLocation::Default, AddPropertyParams);
+	if (TSharedPtr<IPropertyHandle> PropertyHandle = DetailPropertyRow->GetPropertyHandle(); PropertyHandle.IsValid())
+	{
+		const TWeakPtr<FTraitStackDetailsData> TraitDataWeak = TraitData.ToWeakPtr();
+
+		const auto ModifyAsset = [TraitDataWeak]()
 			{
-				EdGraphNode->Modify(); // needed to enable the transaction when we modify the PropertyBag
-			}));
-			Handle->SetOnPropertyValueChanged(FSimpleDelegate::CreateLambda([&MemoryStorage, EdGraphNode, ModelPin]()
+				if (const TSharedPtr<FTraitStackDetailsData> TraitData = TraitDataWeak.Pin())
+				{
+					for (TWeakObjectPtr<UAnimNextGraph_EdGraphNode>& EdGraphNode : TraitData->EdGraphNodes)
+					{
+						EdGraphNode->Modify();
+					}
+				}
+			};
+
+		const auto UpdatePinDefaultValue = [TraitDataWeak](const FPropertyChangedEvent& InEvent)
 			{
-				const FString ValueStr = MemoryStorage.GetDataAsStringByName(ModelPin->GetFName());
-				// Set the default value using the Controller
-				EdGraphNode->GetController()->SetPinDefaultValue(ModelPin->GetPinPath(), ValueStr);
-			}));
-		}
+				if (const TSharedPtr<FTraitStackDetailsData> TraitData = TraitDataWeak.Pin())
+				{
+					// Avoid VM recompilation for each Set Default Value
+					FRigVMControllerCompileBracketScope CompileScope(TraitData->EdGraphNodes[0]->GetController());
+
+					const int32 NumTraitInstances = TraitData->ScopedSharedDataInstances.Num();
+					for (int32 InstanceIndex = 0; InstanceIndex < NumTraitInstances; InstanceIndex++)
+					{
+						TWeakObjectPtr<UAnimNextGraph_EdGraphNode>& EdGraphNode = TraitData->EdGraphNodes[InstanceIndex];
+						if (EdGraphNode.IsValid())
+						{
+							TSharedPtr<FStructOnScope>& ScopedSharedData = TraitData->ScopedSharedDataInstances[InstanceIndex];
+
+							// Extract the value from the property and assign it to the Pin as a default value (via Schema)
+							const uint8* StructMemberMemoryPtr = InEvent.Property->ContainerPtrToValuePtr<uint8>(ScopedSharedData->GetStructMemory());
+							const FString ValueStr = FRigVMStruct::ExportToFullyQualifiedText(InEvent.Property, StructMemberMemoryPtr, true);
+						
+							for (UEdGraphPin* EdGraphPin : EdGraphNode->Pins)
+							{
+								// Find the EdGraphPin that corresponds to the Property
+								const FString ModelPinName = FString::Printf(TEXT(".%s"), *InEvent.Property->GetFName().ToString());
+								if (EdGraphPin->GetFName().ToString().EndsWith(ModelPinName, ESearchCase::CaseSensitive))
+								{
+									if (URigVMPin* ModelPin = EdGraphNode->FindModelPinFromGraphPin(EdGraphPin))
+									{
+										// Set the default value using the Controller
+										EdGraphNode->GetController()->SetPinDefaultValue(ModelPin->GetPinPath(), ValueStr);
+									}
+									break;
+								}
+							}
+						}
+					}
+				}
+			};
+
+		PropertyHandle->SetOnChildPropertyValuePreChange(FSimpleDelegate::CreateLambda(ModifyAsset));
+		PropertyHandle->SetOnChildPropertyValueChangedWithData(TDelegate<void(const FPropertyChangedEvent&)>::CreateLambda(UpdatePinDefaultValue));
 	}
 }
 
-FText FAnimNextGraph_EdGraphNodeCustomization::GetName() const
+void FAnimNextGraph_EdGraphNodeCustomization::PopulateCategory(IDetailLayoutBuilder& DetailBuilder, const TSharedPtr<FRigVMNodeDetailsData>& RigVMTypeData)
 {
-	return FText();
-}
+	if (RigVMTypeData == nullptr)
+	{
+		return;
+	}
 
-void FAnimNextGraph_EdGraphNodeCustomization::SetName(const FText& InNewText, ETextCommit::Type InCommitType)
-{
+	check(RigVMTypeData->MemoryStorages.Num() > 0);
+	check(RigVMTypeData->MemoryStorages.Num() == RigVMTypeData->EdGraphNodes.Num());
+
+	// Create a category with the NodeTitle
+	const FName CategoryName = (RigVMTypeData->EdGraphNodes.Num() == 1)
+		? RigVMTypeData->Name
+		: *FString::Printf(TEXT("%s (% d)"), *RigVMTypeData->Name.ToString(), RigVMTypeData->EdGraphNodes.Num());
+
+	IDetailCategoryBuilder& ParameterCategory = DetailBuilder.EditCategory(CategoryName, FText::GetEmpty(), ECategoryPriority::Default);
+
+	for (TWeakObjectPtr<URigVMPin>& TemplateModelPin : RigVMTypeData->ModelPinsToDisplay)
+	{
+		const FName TemplateModelPinName = TemplateModelPin->GetFName();
+		FAddPropertyParams AddPropertyParams;
+		IDetailPropertyRow* DetailPropertyRow = ParameterCategory.AddExternalStructureProperty(MakeShared<TInstancedPropertyBagStructureDataProvider<FRigVMMemoryStorageStruct>>(RigVMTypeData->MemoryStorages), TemplateModelPinName, EPropertyLocation::Default, AddPropertyParams);
+
+		if (TSharedPtr<IPropertyHandle> Handle = DetailPropertyRow->GetPropertyHandle(); Handle.IsValid())
+		{
+			const TWeakPtr<FRigVMNodeDetailsData> RigVMTypeDataWeak = RigVMTypeData.ToWeakPtr();
+
+			Handle->SetOnChildPropertyValuePreChange(FSimpleDelegate::CreateLambda([RigVMTypeDataWeak]()
+				{
+					if (const TSharedPtr<FRigVMNodeDetailsData> RigVMTypeData = RigVMTypeDataWeak.Pin())
+					{
+						for (auto& EdGraphNode : RigVMTypeData->EdGraphNodes)
+						{
+							EdGraphNode->Modify(); // needed to enable the transaction when we modify the PropertyBag
+						}
+					}
+				}));
+
+			Handle->SetOnPropertyValueChanged(FSimpleDelegate::CreateLambda([RigVMTypeDataWeak, TemplateModelPinName]()
+				{
+					if (const TSharedPtr<FRigVMNodeDetailsData> RigVMTypeData = RigVMTypeDataWeak.Pin())
+					{
+						// Avoid VM recompilation for each Set Default Value
+						FRigVMControllerCompileBracketScope CompileScope(RigVMTypeData->EdGraphNodes[0]->GetController());
+
+						const int32 NumInstances = RigVMTypeData->MemoryStorages.Num();
+						for (int32 InstanceIndex = 0; InstanceIndex < NumInstances; InstanceIndex++)
+						{
+							TWeakObjectPtr<UAnimNextGraph_EdGraphNode> EdGraphNode = RigVMTypeData->EdGraphNodes[InstanceIndex];
+							if (EdGraphNode.IsValid())
+							{
+								const FString ValueStr = RigVMTypeData->MemoryStorages[InstanceIndex].Get()->GetDataAsStringByName(TemplateModelPinName);
+
+								// Set the default value using the Controller
+								for (UEdGraphPin* EdGraphPin : EdGraphNode->Pins)
+								{
+									const FString ModelPinName = FString::Printf(TEXT(".%s"), *TemplateModelPinName.ToString());
+									// Find the EdGraphPin that corresponds to the model pin
+									if (EdGraphPin->GetFName().ToString().EndsWith(ModelPinName, ESearchCase::CaseSensitive))
+									{
+										if (URigVMPin* ModelPin = EdGraphNode->FindModelPinFromGraphPin(EdGraphPin))
+										{
+											EdGraphNode->GetController()->SetPinDefaultValue(ModelPin->GetPinPath(), ValueStr);
+										}
+									}
+								}
+							}
+						}
+					}
+				}));
+		}
+	}
 }
 
-bool FAnimNextGraph_EdGraphNodeCustomization::OnVerifyNameChanged(const FText& InText, FText& OutErrorMessage)
+void FAnimNextGraph_EdGraphNodeCustomization::GenerateMemoryStorage(const TArray<TWeakObjectPtr<URigVMPin>> & ModelPinsToDisplay, FRigVMMemoryStorageStruct& MemoryStorage)
 {
-	return true;
-}
+	TArray<FRigVMPropertyDescription> PropertyDescriptions;
+	PropertyDescriptions.Reserve(ModelPinsToDisplay.Num());
 
+	for (const TWeakObjectPtr<URigVMPin>& ModelPin : ModelPinsToDisplay)
+	{
+		if (ModelPin.IsValid())
+		{
+			FRigVMPropertyDescription& PropertyDesc = PropertyDescriptions.AddDefaulted_GetRef();
 
+			PropertyDesc.Name = ModelPin->GetFName();
+			PropertyDesc.Property = nullptr;
+			PropertyDesc.CPPType = ModelPin->GetCPPType();
+			PropertyDesc.CPPTypeObject = ModelPin->GetCPPTypeObject();
+			if (ModelPin->IsArray())
+			{
+				PropertyDesc.Containers.Add(EPinContainerType::Array);
+			}
+			PropertyDesc.DefaultValue = ModelPin->GetDefaultValue();
+		}
+	}
+
+	MemoryStorage.AddProperties(PropertyDescriptions);
 }
 
+} // end namespace UE::AnimNext::Editor
+
 #undef LOCTEXT_NAMESPACE
diff --git a/Engine/Plugins/Experimental/AnimNext/Source/AnimNextEditor/Private/Graph/AnimNextGraph_EdGraphNodeCustomization.h b/Engine/Plugins/Experimental/AnimNext/Source/AnimNextEditor/Private/Graph/AnimNextGraph_EdGraphNodeCustomization.h
--- a/Engine/Plugins/Experimental/AnimNext/Source/AnimNextEditor/Private/Graph/AnimNextGraph_EdGraphNodeCustomization.h
+++ b/Engine/Plugins/Experimental/AnimNext/Source/AnimNextEditor/Private/Graph/AnimNextGraph_EdGraphNodeCustomization.h
@@ -16,20 +16,81 @@ namespace UE::AnimNext::Editor
 
 class FAnimNextGraph_EdGraphNodeCustomization : public IDetailCustomization
 {
-private:
+protected:
 
+	// --- IDetailCustomization Begin ---
 	/** Called when details should be customized */
 	virtual void CustomizeDetails(IDetailLayoutBuilder& DetailBuilder) override;
+	// --- IDetailCustomization End ---
 
-	FText GetName() const;
-	void SetName(const FText& InNewText, ETextCommit::Type InCommitType);
-	bool OnVerifyNameChanged(const FText& InText, FText& OutErrorMessage);
+	struct FCategoryDetailsData
+	{
+		enum class EType : uint8
+		{
+			TraitStack = 0,
+			RigVMNode,
+			// --- ---
+			Invalid,
+			Num = Invalid
+		};
 
-	static void GenerateMemoryStorage(const TArray<URigVMPin*>& ModelPinsToDisplay, FRigVMMemoryStorageStruct& MemoryStorage);
-	static void PopulateCategory(IDetailCategoryBuilder& Category, const TArray<URigVMPin*>& ModelPinsToDisplay, FRigVMMemoryStorageStruct& MemoryStorage, UAnimNextGraph_EdGraphNode* EdGraphNode);
+		FCategoryDetailsData() = default;
+		explicit FCategoryDetailsData(EType InType)
+			: Type(InType)
+		{
+		}
+		FCategoryDetailsData(EType InType, const FName& InName)
+			: Type(InType)
+			, Name(InName)
+		{
+		}
 
+		EType Type = EType::Invalid;
+		FName Name;
+		TArray<TWeakObjectPtr<UAnimNextGraph_EdGraphNode>> EdGraphNodes;
+	};
 
-	FRigVMMemoryStorageStruct MemoryStorage;
+	struct FTraitStackDetailsData : FCategoryDetailsData
+	{
+		FTraitStackDetailsData()
+			: FCategoryDetailsData(FCategoryDetailsData::EType::TraitStack)
+		{
+		}
+		explicit FTraitStackDetailsData(const FName InName)
+			: FCategoryDetailsData(FCategoryDetailsData::EType::TraitStack, InName)
+		{
+		}
+
+		TArray<TSharedPtr<FStructOnScope>> ScopedSharedDataInstances;
+	};
+
+	struct FRigVMNodeDetailsData : FCategoryDetailsData
+	{
+		FRigVMNodeDetailsData()
+			: FCategoryDetailsData(FCategoryDetailsData::EType::RigVMNode)
+		{
+		}
+		explicit FRigVMNodeDetailsData(const FName InName)
+			: FCategoryDetailsData(FCategoryDetailsData::EType::RigVMNode, InName)
+		{
+		}
+
+		TArray<TWeakObjectPtr<URigVMPin>> ModelPinsToDisplay;
+		TArray <TSharedPtr<FRigVMMemoryStorageStruct>> MemoryStorages;
+	};
+
+	void CustomizeObjects(IDetailLayoutBuilder& DetailBuilder, const TArray<TWeakObjectPtr<UObject>>& InObjects);
+
+	static void GenerateTraitData(UAnimNextGraph_EdGraphNode* EdGraphNode, TArray<TSharedPtr<FCategoryDetailsData>>& CategoryDetailsData);
+	static void GenerateRigVMData(UAnimNextGraph_EdGraphNode* EdGraphNode, TArray<TSharedPtr<FCategoryDetailsData>>& CategoryDetailsData);
+
+	static void PopulateCategory(IDetailLayoutBuilder& DetailBuilder, const TSharedPtr<FCategoryDetailsData>& CategoryDetailsData);
+	static void PopulateCategory(IDetailLayoutBuilder& DetailBuilder, const TSharedPtr<FTraitStackDetailsData>& TraitData);
+	static void PopulateCategory(IDetailLayoutBuilder& DetailBuilder, const TSharedPtr<FRigVMNodeDetailsData>& RigVMTypeData);
+
+	static void GenerateMemoryStorage(const TArray<TWeakObjectPtr<URigVMPin>> & ModelPinsToDisplay, FRigVMMemoryStorageStruct& MemoryStorage);
+
+	TArray<TSharedPtr<FCategoryDetailsData>> CategoryDetailsData;
 };
 
 }
diff --git a/Engine/Plugins/Experimental/StructUtils/Source/StructUtilsEditor/Public/InstancedPropertyBagStructureDataProvider.h b/Engine/Plugins/Experimental/StructUtils/Source/StructUtilsEditor/Public/InstancedPropertyBagStructureDataProvider.h
--- a/Engine/Plugins/Experimental/StructUtils/Source/StructUtilsEditor/Public/InstancedPropertyBagStructureDataProvider.h
+++ b/Engine/Plugins/Experimental/StructUtils/Source/StructUtilsEditor/Public/InstancedPropertyBagStructureDataProvider.h
@@ -1,4 +1,4 @@
-﻿// Copyright Epic Games, Inc. All Rights Reserved.
+// Copyright Epic Games, Inc. All Rights Reserved.
 
 #pragma once
 
@@ -40,4 +40,74 @@ public:
 
 protected:
 	FInstancedPropertyBag& PropertyBag;
-};
\ No newline at end of file
+};
+
+//	TInstancePropertyBagStructureDataProvider
+//	Allows structure views to use FInstancedPropertyBag even if the bag layout changes.
+//	The caller needs to make sure that property bag outlives the property view widget.
+//	This version enables a single structure with multiple instances, and the use of
+//  FInstancedPropertyBag derived types
+
+template<typename BagInstanceType>
+class TInstancedPropertyBagStructureDataProvider : public IStructureDataProvider
+{
+public:
+	explicit TInstancedPropertyBagStructureDataProvider(const TSharedPtr<BagInstanceType>& InPropertyBag)
+		: PropertyBagInstances( {InPropertyBag} )
+	{
+	}
+
+	explicit TInstancedPropertyBagStructureDataProvider(const TArray<TSharedPtr<BagInstanceType>>& InPropertyBagInstances)
+		: PropertyBagInstances(InPropertyBagInstances)
+	{
+	}
+
+	virtual bool IsValid() const override
+	{
+		const UStruct* BaseStructure = GetBaseStructure();
+		if (BaseStructure != nullptr)
+		{
+			for (const TSharedPtr<BagInstanceType>& BagInstance : PropertyBagInstances)
+			{
+				if (BagInstance.IsValid() && BagInstance->IsValid())
+				{
+					const UStruct* Struct = BagInstance->GetPropertyBagStruct();
+					if (Struct->IsChildOf(BaseStructure))
+					{
+						return true;
+					}
+				}
+			}
+		}
+
+		return false;
+	}
+
+	virtual const UStruct* GetBaseStructure() const override
+	{
+		return (PropertyBagInstances.Num() > 0) 
+			? FStructOnScopeStructureDataProvider::FindBaseStructure(PropertyBagInstances)
+			: nullptr;
+	}
+
+	virtual void GetInstances(TArray<TSharedPtr<FStructOnScope>>& OutInstances, const UStruct* ExpectedBaseStructure) const override
+	{
+		if (PropertyBagInstances.Num() > 0 && ExpectedBaseStructure != nullptr)
+		{
+			for (auto& BagInstance : PropertyBagInstances)
+			{
+				if (BagInstance.IsValid() && BagInstance->IsValid())
+				{
+					const UStruct* Struct = BagInstance->GetPropertyBagStruct();
+					if (Struct && Struct->IsChildOf(ExpectedBaseStructure))
+					{
+						OutInstances.Add(MakeShared<FStructOnScope>(Struct, BagInstance->GetMutableValue().GetMemory()));
+					}
+				}
+			}
+		}
+	}
+
+protected:
+	TArray<TSharedPtr<BagInstanceType>> PropertyBagInstances;
+};
diff --git a/Engine/Plugins/Experimental/Workspace/Source/WorkspaceEditor/Private/DetailsTabSummoner.cpp b/Engine/Plugins/Experimental/Workspace/Source/WorkspaceEditor/Private/DetailsTabSummoner.cpp
--- a/Engine/Plugins/Experimental/Workspace/Source/WorkspaceEditor/Private/DetailsTabSummoner.cpp
+++ b/Engine/Plugins/Experimental/Workspace/Source/WorkspaceEditor/Private/DetailsTabSummoner.cpp
@@ -1,4 +1,4 @@
-﻿// Copyright Epic Games, Inc. All Rights Reserved.
+// Copyright Epic Games, Inc. All Rights Reserved.
 
 #include "DetailsTabSummoner.h"
 #include "WorkspaceEditor.h"
@@ -21,6 +21,7 @@ FDetailsTabSummoner::FDetailsTabSummoner(TSharedPtr<FWorkspaceEditor> InHostingA
 	FDetailsViewArgs DetailsViewArgs;
 	DetailsViewArgs.NameAreaSettings = FDetailsViewArgs::HideNameArea;
 	DetailsViewArgs.bHideSelectionTip = true;
+	DetailsViewArgs.StyleKey = MakeShared<FDetailsViewStyleKey>(FDetailsViewStyleKeys::Card());
 
 	FPropertyEditorModule& PropertyEditorModule = FModuleManager::LoadModuleChecked<FPropertyEditorModule>("PropertyEditor");
 	DetailsView = PropertyEditorModule.CreateDetailView(DetailsViewArgs);
diff --git a/Engine/Plugins/Runtime/RigVM/Source/RigVM/Public/RigVMCore/RigVMMemoryStorageStruct.h b/Engine/Plugins/Runtime/RigVM/Source/RigVM/Public/RigVMCore/RigVMMemoryStorageStruct.h
--- a/Engine/Plugins/Runtime/RigVM/Source/RigVM/Public/RigVMCore/RigVMMemoryStorageStruct.h
+++ b/Engine/Plugins/Runtime/RigVM/Source/RigVM/Public/RigVMCore/RigVMMemoryStorageStruct.h
@@ -364,6 +364,14 @@ struct RIGVM_API FRigVMMemoryStorageStruct : public FInstancedPropertyBag
 		FRigVMMemoryHandle& InTargetHandle,
 		FRigVMMemoryHandle& InSourceHandle);
 
+	/**
+	* Returns the PropertyBag struct using a standard function signature, so it can be used at IStructureDataProvider
+	*/
+	inline const UStruct * GetStruct() const
+	{
+		return GetPropertyBagStruct();
+	}
+
 protected:
 	ERigVMMemoryType MemoryType = ERigVMMemoryType::Invalid;
 
diff --git a/Engine/Source/Editor/PropertyEditor/Public/IStructureDataProvider.h b/Engine/Source/Editor/PropertyEditor/Public/IStructureDataProvider.h
--- a/Engine/Source/Editor/PropertyEditor/Public/IStructureDataProvider.h
+++ b/Engine/Source/Editor/PropertyEditor/Public/IStructureDataProvider.h
@@ -93,46 +93,119 @@ public:
 
 
 //-----------------------------------------------------------------------------
-//	FStructOnScopeStructureDataProvider - Implementation of standalone struct that holds one value.
+//	FStructOnScopeStructureDataProvider - Implementation of standalone struct that provides 
+//  one base structure and one or more instances (allowing multi selection edit)
 //-----------------------------------------------------------------------------
 class FStructOnScopeStructureDataProvider : public IStructureDataProvider
 {
 public:
 	FStructOnScopeStructureDataProvider() = default;
-	FStructOnScopeStructureDataProvider(const TSharedPtr<FStructOnScope>& InStructData)
-		: StructData(InStructData)
+	explicit FStructOnScopeStructureDataProvider(const TSharedPtr<FStructOnScope>& InStructData)
+		: StructDataInstances( {InStructData} )
+	{
+	}
+	
+	explicit FStructOnScopeStructureDataProvider(const TArray<TSharedPtr<FStructOnScope>>& InStructDataInstances)
+		: StructDataInstances(InStructDataInstances)
 	{
 	}
 	
 	void SetStructData(const TSharedPtr<FStructOnScope>& InStructData)
 	{
-		StructData = InStructData;
+		StructDataInstances.Reset(1);
+		StructDataInstances.Add(InStructData);
 	}
+	void SetStructData(const TArray<TSharedPtr<FStructOnScope>>& InStructData)
+	{
+		StructDataInstances.Reset(InStructData.Num());
+		StructDataInstances = InStructData;
+	}
+
 	
 	virtual bool IsValid() const override
 	{
-		return StructData.IsValid() && StructData->IsValid();
+		const UStruct* BaseStructure = GetBaseStructure();
+		for (const TSharedPtr<FStructOnScope>& StructInstance : StructDataInstances)
+		{
+			if (StructInstance.IsValid() && StructInstance->IsValid())
+			{
+				const UStruct* Struct = StructInstance->GetStruct();
+				if (Struct->IsChildOf(BaseStructure))
+				{
+					return true;
+				}
+			}
+		}
+
+		return false;
 	};
-	
+
 	virtual const UStruct* GetBaseStructure() const override
 	{
-		return StructData.IsValid() ? StructData->GetStruct() : nullptr;
+		return FindBaseStructure(StructDataInstances);
 	}
 	
 	virtual void GetInstances(TArray<TSharedPtr<FStructOnScope>>& OutInstances, const UStruct* ExpectedBaseStructure) const override
 	{
-		if (StructData.IsValid())
+		if (ExpectedBaseStructure != nullptr)
+		{
+			for (const TSharedPtr<FStructOnScope> &StructData : StructDataInstances)
+			{
+				if (StructData.IsValid()) 
+				{
+					const UStruct* Struct = StructData->GetStruct();
+					if (Struct && Struct->IsChildOf(ExpectedBaseStructure))
+					{
+						OutInstances.Add(StructData);				
+					}
+				}
+			}
+		}
+	}
+
+	template<typename ContainerType>
+	static const UStruct* FindBaseStructure(const TArray<TSharedPtr<ContainerType>>& StructDataInstances)
+	{
+		const UStruct* BaseStructure = nullptr;
+
+		const int32 NumInstances = StructDataInstances.Num();
+		if (NumInstances > 0)
 		{
-			const UStruct* Struct = StructData->GetStruct();
-			if (ExpectedBaseStructure
-				&& Struct
-				&& Struct->IsChildOf(ExpectedBaseStructure))
+			int32 InstanceIndex = 0;
+
+			// find first valid instance
+			while ( !(StructDataInstances[InstanceIndex].IsValid() && StructDataInstances[InstanceIndex]->IsValid()) )
 			{
-				OutInstances.Add(StructData);				
+				++InstanceIndex;
+				if (InstanceIndex == NumInstances)
+				{
+					// no valid instances found
+					return nullptr;
+				}
+			}
+
+			BaseStructure = StructDataInstances[InstanceIndex]->GetStruct();
+
+			// Iterate from there, if any instance left
+			for (int32 i = InstanceIndex + 1; i < NumInstances; i++)
+			{
+				const TSharedPtr<ContainerType>& StructData = StructDataInstances[i];
+				if (StructData.IsValid() && StructData->IsValid())
+				{
+					const UStruct* ScriptStruct = StructData->GetStruct();
+					// check if the base structure is an ancestor of the script struct
+					while (BaseStructure && ScriptStruct && !ScriptStruct->IsChildOf(BaseStructure))
+					{
+						// if not, go a level up in the base structure hierarchy
+						BaseStructure = Cast<UStruct>(BaseStructure->GetSuperStruct());
+					}
+				}
 			}
 		}
+
+		return BaseStructure;
 	}
 
 protected:
-	TSharedPtr<FStructOnScope> StructData;
-};
\ No newline at end of file
+	TArray<TSharedPtr<FStructOnScope>> StructDataInstances;
+};
