Fixed AnimNext node details customization applying values to other pins that match the pin name. Now uses the pin path to avoid selecting incorrect pins when updating the value from the details panel.
Also fixed parameters details customization not handling child property changes.

diff --git a/Engine/Plugins/Experimental/AnimNext/Source/AnimNextEditor/Private/Common/SRigVMAssetView.cpp b/Engine/Plugins/Experimental/AnimNext/Source/AnimNextEditor/Private/Common/SRigVMAssetView.cpp
--- a/Engine/Plugins/Experimental/AnimNext/Source/AnimNextEditor/Private/Common/SRigVMAssetView.cpp
+++ b/Engine/Plugins/Experimental/AnimNext/Source/AnimNextEditor/Private/Common/SRigVMAssetView.cpp
@@ -856,17 +856,21 @@ class SRigVMAssetViewRow : public SMultiColumnTableRow<TSharedRef<FRigVMAssetVie
 						FInstancedPropertyBag& PropertyBag = ParameterInterface->GetPropertyBag();
 						const FName ParameterName = AssetEntry->GetEntryName();
 
-						if (PropertyBag.FindPropertyDescByName(ParameterName))
-						{
-							FSinglePropertyParams SinglePropertyArgs;
-							SinglePropertyArgs.NamePlacement = EPropertyNamePlacement::Hidden;
-							SinglePropertyArgs.NotifyHook = this;
 
-							FPropertyEditorModule& PropertyEditorModule = FModuleManager::Get().LoadModuleChecked<FPropertyEditorModule>("PropertyEditor");
-							const TSharedPtr<ISinglePropertyView> SingleStructPropertyView = PropertyEditorModule.CreateSingleProperty(MakeShared<FInstancePropertyBagStructureDataProvider>(PropertyBag), ParameterName, SinglePropertyArgs);
-							if (SingleStructPropertyView.IsValid())
+						if (const FPropertyBagPropertyDesc* PropertyDesc = PropertyBag.FindPropertyDescByName(ParameterName))
+						{
+							if (PropertyDesc->ContainerTypes.IsEmpty()) // avoid trying to inline containers
 							{
-								ColumnWidget = SingleStructPropertyView.ToSharedRef();
+								FSinglePropertyParams SinglePropertyArgs;
+								SinglePropertyArgs.NamePlacement = EPropertyNamePlacement::Hidden;
+								SinglePropertyArgs.NotifyHook = this;
+
+								FPropertyEditorModule& PropertyEditorModule = FModuleManager::Get().LoadModuleChecked<FPropertyEditorModule>("PropertyEditor");
+								const TSharedPtr<ISinglePropertyView> SingleStructPropertyView = PropertyEditorModule.CreateSingleProperty(MakeShared<FInstancePropertyBagStructureDataProvider>(PropertyBag), ParameterName, SinglePropertyArgs);
+								if (SingleStructPropertyView.IsValid())
+								{
+									ColumnWidget = SingleStructPropertyView.ToSharedRef();
+								}
 							}
 						}
 					}
diff --git a/Engine/Plugins/Experimental/AnimNext/Source/AnimNextEditor/Private/Graph/AnimNextGraph_EdGraphNodeCustomization.cpp b/Engine/Plugins/Experimental/AnimNext/Source/AnimNextEditor/Private/Graph/AnimNextGraph_EdGraphNodeCustomization.cpp
--- a/Engine/Plugins/Experimental/AnimNext/Source/AnimNextEditor/Private/Graph/AnimNextGraph_EdGraphNodeCustomization.cpp
+++ b/Engine/Plugins/Experimental/AnimNext/Source/AnimNextEditor/Private/Graph/AnimNextGraph_EdGraphNodeCustomization.cpp
@@ -122,6 +122,8 @@ void FAnimNextGraph_EdGraphNodeCustomization::GenerateRigVMData(UAnimNextGraph_E
 		TArray<TWeakObjectPtr<URigVMPin>> PinsToDisplay;
 		PinsToDisplay.Reserve(ModelPins.Num());
 
+		TArray<FString> ModelPinPaths;
+
 		// First, obtain the pins to display
 		for (URigVMPin* Pin : ModelPins)
 		{
@@ -134,6 +136,7 @@ void FAnimNextGraph_EdGraphNodeCustomization::GenerateRigVMData(UAnimNextGraph_E
 			if (PinDirection != ERigVMPinDirection::Hidden && (PinDirection == ERigVMPinDirection::IO || PinDirection == ERigVMPinDirection::Input))
 			{
 				PinsToDisplay.Add(Pin);
+				ModelPinPaths.Add(Pin->GetPinPath());
 			}
 		}
 
@@ -141,26 +144,33 @@ void FAnimNextGraph_EdGraphNodeCustomization::GenerateRigVMData(UAnimNextGraph_E
 		{
 			const FName NodeName = *EdGraphNode->GetNodeTitle(ENodeTitleType::ListView).ToString(); // TODO : check if this is a good name to use, can't sue GetFName as it comes with instance postfix
 
-			FRigVMNodeDetailsData* RigVMData = nullptr;
-
 			TSharedPtr<FCategoryDetailsData>* RigVMDataPtr = CategoryDetailsData.FindByPredicate([&NodeName](const TSharedPtr<FCategoryDetailsData>& InItem)
 				{
 					return InItem->Type == FCategoryDetailsData::EType::RigVMNode && InItem->Name == NodeName;
 				});
 
-			if (RigVMDataPtr == nullptr)
+			FRigVMNodeDetailsData* RigVMData = RigVMDataPtr  ? static_cast<FRigVMNodeDetailsData*>(RigVMDataPtr->Get()) : nullptr;
+			if (RigVMData == nullptr)
 			{
 				RigVMDataPtr = &CategoryDetailsData.Add_GetRef(MakeShared<FRigVMNodeDetailsData>(NodeName));
-			}
+				RigVMData = static_cast<FRigVMNodeDetailsData*>(RigVMDataPtr->Get());
 
-			RigVMData = static_cast<FRigVMNodeDetailsData*>(RigVMDataPtr->Get());
-			RigVMData->ModelPinsToDisplay = PinsToDisplay;
+				// Store the model Pin Names that will be shown (only when we create the Data, in case of multiselection)
+				for (const TWeakObjectPtr<URigVMPin>& Pin : PinsToDisplay)
+				{
+					RigVMData->ModelPinsNamesToDisplay.Add(Pin->GetFName());
+				}
+			}
 
 			// Store EdGraphNode and generated Memory, as we will need that later to transfer data if the user makes modifications in the details panel
 			RigVMData->EdGraphNodes.Add(EdGraphNode);
 			const int32 MemoryStorageIndex = RigVMData->MemoryStorages.Add(MakeShared<FRigVMMemoryStorageStruct>());
 
-			// Then, create a custom property bag to store the data, initializating the properties in the struct with the pin default values
+			// Store the Model Pin Paths, needed later to update the value of the correct model pin
+			TArray<FString>& PinPaths = RigVMData->ModelPinPaths.Emplace_GetRef();
+			Swap(PinPaths, ModelPinPaths);
+
+			// Then, create a custom property bag to store the data, initializing the properties in the struct with the pin default values
 			GenerateMemoryStorage(PinsToDisplay, *RigVMData->MemoryStorages[MemoryStorageIndex].Get());
 		}
 	}
@@ -212,17 +222,6 @@ void FAnimNextGraph_EdGraphNodeCustomization::PopulateCategory(IDetailLayoutBuil
 	{
 		const TWeakPtr<FTraitStackDetailsData> TraitDataWeak = TraitData.ToWeakPtr();
 
-		const auto ModifyAsset = [TraitDataWeak]()
-			{
-				if (const TSharedPtr<FTraitStackDetailsData> TraitData = TraitDataWeak.Pin())
-				{
-					for (TWeakObjectPtr<UAnimNextGraph_EdGraphNode>& EdGraphNode : TraitData->EdGraphNodes)
-					{
-						EdGraphNode->Modify();
-					}
-				}
-			};
-
 		const auto UpdatePinDefaultValue = [TraitDataWeak](const FPropertyChangedEvent& InEvent)
 			{
 				if (const TSharedPtr<FTraitStackDetailsData> TraitData = TraitDataWeak.Pin())
@@ -238,14 +237,21 @@ void FAnimNextGraph_EdGraphNodeCustomization::PopulateCategory(IDetailLayoutBuil
 						{
 							TSharedPtr<FStructOnScope>& ScopedSharedData = TraitData->ScopedSharedDataInstances[InstanceIndex];
 
+							const bool bIsContainer = InEvent.Property->GetOwnerProperty()->IsA<FArrayProperty>()
+								|| InEvent.Property->GetOwnerProperty()->IsA<FMapProperty>()
+								|| InEvent.Property->GetOwnerProperty()->IsA<FSetProperty>();
+							
+							// For some reason, sub properties of a container does not come with the correct struct offsets, so getting the container property in that case
+							FProperty* Property = bIsContainer ? InEvent.Property->GetOwnerProperty() : InEvent.Property;
+							
 							// Extract the value from the property and assign it to the Pin as a default value (via Schema)
-							const uint8* StructMemberMemoryPtr = InEvent.Property->ContainerPtrToValuePtr<uint8>(ScopedSharedData->GetStructMemory());
-							const FString ValueStr = FRigVMStruct::ExportToFullyQualifiedText(InEvent.Property, StructMemberMemoryPtr, true);
+							const uint8* StructMemberMemoryPtr = Property->ContainerPtrToValuePtr<uint8>(ScopedSharedData->GetStructMemory());
+							const FString ValueStr = FRigVMStruct::ExportToFullyQualifiedText(Property, StructMemberMemoryPtr, true);
 						
 							for (UEdGraphPin* EdGraphPin : EdGraphNode->Pins)
 							{
 								// Find the EdGraphPin that corresponds to the Property
-								const FString ModelPinName = FString::Printf(TEXT(".%s"), *InEvent.Property->GetFName().ToString());
+								const FString ModelPinName = FString::Printf(TEXT(".%s"), *Property->GetFName().ToString());
 								if (EdGraphPin->GetFName().ToString().EndsWith(ModelPinName, ESearchCase::CaseSensitive))
 								{
 									if (URigVMPin* ModelPin = EdGraphNode->FindModelPinFromGraphPin(EdGraphPin))
@@ -261,7 +267,6 @@ void FAnimNextGraph_EdGraphNodeCustomization::PopulateCategory(IDetailLayoutBuil
 				}
 			};
 
-		PropertyHandle->SetOnChildPropertyValuePreChange(FSimpleDelegate::CreateLambda(ModifyAsset));
 		PropertyHandle->SetOnChildPropertyValueChangedWithData(TDelegate<void(const FPropertyChangedEvent&)>::CreateLambda(UpdatePinDefaultValue));
 	}
 }
@@ -283,9 +288,8 @@ void FAnimNextGraph_EdGraphNodeCustomization::PopulateCategory(IDetailLayoutBuil
 
 	IDetailCategoryBuilder& ParameterCategory = DetailBuilder.EditCategory(CategoryName, FText::GetEmpty(), ECategoryPriority::Default);
 
-	for (TWeakObjectPtr<URigVMPin>& TemplateModelPin : RigVMTypeData->ModelPinsToDisplay)
+	for (const FName& TemplateModelPinName : RigVMTypeData->ModelPinsNamesToDisplay)
 	{
-		const FName TemplateModelPinName = TemplateModelPin->GetFName();
 		FAddPropertyParams AddPropertyParams;
 		IDetailPropertyRow* DetailPropertyRow = ParameterCategory.AddExternalStructureProperty(MakeShared<TInstancedPropertyBagStructureDataProvider<FRigVMMemoryStorageStruct>>(RigVMTypeData->MemoryStorages), TemplateModelPinName, EPropertyLocation::Default, AddPropertyParams);
 
@@ -293,18 +297,7 @@ void FAnimNextGraph_EdGraphNodeCustomization::PopulateCategory(IDetailLayoutBuil
 		{
 			const TWeakPtr<FRigVMNodeDetailsData> RigVMTypeDataWeak = RigVMTypeData.ToWeakPtr();
 
-			Handle->SetOnChildPropertyValuePreChange(FSimpleDelegate::CreateLambda([RigVMTypeDataWeak]()
-				{
-					if (const TSharedPtr<FRigVMNodeDetailsData> RigVMTypeData = RigVMTypeDataWeak.Pin())
-					{
-						for (auto& EdGraphNode : RigVMTypeData->EdGraphNodes)
-						{
-							EdGraphNode->Modify(); // needed to enable the transaction when we modify the PropertyBag
-						}
-					}
-				}));
-
-			Handle->SetOnPropertyValueChanged(FSimpleDelegate::CreateLambda([RigVMTypeDataWeak, TemplateModelPinName]()
+			const auto UpdatePinDefaultValue = [RigVMTypeDataWeak, TemplateModelPinName](const FPropertyChangedEvent& InEvent)
 				{
 					if (const TSharedPtr<FRigVMNodeDetailsData> RigVMTypeData = RigVMTypeDataWeak.Pin())
 					{
@@ -319,23 +312,23 @@ void FAnimNextGraph_EdGraphNodeCustomization::PopulateCategory(IDetailLayoutBuil
 							{
 								const FString ValueStr = RigVMTypeData->MemoryStorages[InstanceIndex].Get()->GetDataAsStringByName(TemplateModelPinName);
 
-								// Set the default value using the Controller
-								for (UEdGraphPin* EdGraphPin : EdGraphNode->Pins)
+								const TArray<FString>& ModelPinPaths = RigVMTypeData->ModelPinPaths[InstanceIndex];
+								for (const FString& PinPath : ModelPinPaths)
 								{
 									const FString ModelPinName = FString::Printf(TEXT(".%s"), *TemplateModelPinName.ToString());
-									// Find the EdGraphPin that corresponds to the model pin
-									if (EdGraphPin->GetFName().ToString().EndsWith(ModelPinName, ESearchCase::CaseSensitive))
+									if (PinPath.EndsWith(ModelPinName, ESearchCase::CaseSensitive))
 									{
-										if (URigVMPin* ModelPin = EdGraphNode->FindModelPinFromGraphPin(EdGraphPin))
-										{
-											EdGraphNode->GetController()->SetPinDefaultValue(ModelPin->GetPinPath(), ValueStr);
-										}
+										EdGraphNode->GetController()->SetPinDefaultValue(PinPath, ValueStr);
+										break;
 									}
 								}
 							}
 						}
 					}
-				}));
+				};
+
+			Handle->SetOnPropertyValueChangedWithData(TDelegate<void(const FPropertyChangedEvent&)>::CreateLambda(UpdatePinDefaultValue));
+			Handle->SetOnChildPropertyValueChangedWithData(TDelegate<void(const FPropertyChangedEvent&)>::CreateLambda(UpdatePinDefaultValue));
 		}
 	}
 }
diff --git a/Engine/Plugins/Experimental/AnimNext/Source/AnimNextEditor/Private/Graph/AnimNextGraph_EdGraphNodeCustomization.h b/Engine/Plugins/Experimental/AnimNext/Source/AnimNextEditor/Private/Graph/AnimNextGraph_EdGraphNodeCustomization.h
--- a/Engine/Plugins/Experimental/AnimNext/Source/AnimNextEditor/Private/Graph/AnimNextGraph_EdGraphNodeCustomization.h
+++ b/Engine/Plugins/Experimental/AnimNext/Source/AnimNextEditor/Private/Graph/AnimNextGraph_EdGraphNodeCustomization.h
@@ -75,7 +75,8 @@ protected:
 		{
 		}
 
-		TArray<TWeakObjectPtr<URigVMPin>> ModelPinsToDisplay;
+		TArray<FName> ModelPinsNamesToDisplay;
+		TArray <TArray<FString>> ModelPinPaths;
 		TArray <TSharedPtr<FRigVMMemoryStorageStruct>> MemoryStorages;
 	};
 
diff --git a/Engine/Plugins/Experimental/AnimNext/Source/AnimNextEditor/Private/Param/ParameterCustomization.cpp b/Engine/Plugins/Experimental/AnimNext/Source/AnimNextEditor/Private/Param/ParameterCustomization.cpp
--- a/Engine/Plugins/Experimental/AnimNext/Source/AnimNextEditor/Private/Param/ParameterCustomization.cpp
+++ b/Engine/Plugins/Experimental/AnimNext/Source/AnimNextEditor/Private/Param/ParameterCustomization.cpp
@@ -45,25 +45,39 @@ void FParameterCustomization::CustomizeDetails(IDetailLayoutBuilder& DetailBuild
 					FAddPropertyParams AddPropertyParams;
 					TArray<IDetailPropertyRow*> DetailPropertyRows;
 
-					if (ReferencedGraph->PropertyBag.FindPropertyDescByName(ParameterName))
+					if (const FPropertyBagPropertyDesc* PropertyDesc = ReferencedGraph->PropertyBag.FindPropertyDescByName(ParameterName))
 					{
 						IDetailPropertyRow* DetailPropertyRow = DefaultValueCategory.AddExternalStructureProperty(MakeShared<FInstancePropertyBagStructureDataProvider>(ReferencedGraph->PropertyBag), ParameterName, EPropertyLocation::Default, AddPropertyParams);
 						if (TSharedPtr<IPropertyHandle> Handle = DetailPropertyRow->GetPropertyHandle(); Handle.IsValid())
 						{
-							Handle->SetOnChildPropertyValuePreChange(FSimpleDelegate::CreateLambda([this, ReferencedGraph]()
-							{
-								ReferencedGraph->Modify(); // needed to enable the transaction when we modify the PropertyBag
-							}));
-							Handle->SetOnPropertyValueChanged(FSimpleDelegate::CreateLambda([this, ReferencedGraph, ParameterName]()
-							{
-								if (UAnimNextGraph_EditorData* EditorData = Cast<UAnimNextGraph_EditorData>(ReferencedGraph->EditorData))
+							const TWeakObjectPtr<UAnimNextGraph> ReferencedGraphWeak = ReferencedGraph;
+
+							const auto OnPropertyValuePreChange = [ReferencedGraphWeak]()
+								{
+									if (ReferencedGraphWeak.IsValid())
+									{
+										ReferencedGraphWeak->Modify(); // needed to enable the transaction when we modify the PropertyBag
+									}
+								};
+							const auto OnPropertyValueChange = [ReferencedGraphWeak](const FPropertyChangedEvent& InEvent)
 								{
-									if (UAnimNextRigVMAssetEntry* AssetEntry = EditorData->FindEntry(ParameterName))
+									if (ReferencedGraphWeak.IsValid())
 									{
-										AssetEntry->MarkPackageDirty();
+										if (UAnimNextGraph_EditorData* EditorData = Cast<UAnimNextGraph_EditorData>(ReferencedGraphWeak->EditorData))
+										{
+											if (UAnimNextRigVMAssetEntry* AssetEntry = EditorData->FindEntry(InEvent.GetPropertyName()))
+											{
+												AssetEntry->MarkPackageDirty();
+											}
+										}
 									}
-								}
-							}));
+								};
+
+							Handle->SetOnPropertyValuePreChange(FSimpleDelegate::CreateLambda(OnPropertyValuePreChange));
+							Handle->SetOnPropertyValueChangedWithData(TDelegate<void(const FPropertyChangedEvent&)>::CreateLambda(OnPropertyValueChange));
+
+							Handle->SetOnChildPropertyValuePreChange(FSimpleDelegate::CreateLambda(OnPropertyValuePreChange));
+							Handle->SetOnChildPropertyValueChangedWithData(TDelegate<void(const FPropertyChangedEvent&)>::CreateLambda(OnPropertyValueChange));
 						}
 					}
 				}
diff --git a/Engine/Plugins/Experimental/AnimNext/Source/AnimNextUncookedOnly/Private/UncookedOnlyUtils.cpp b/Engine/Plugins/Experimental/AnimNext/Source/AnimNextUncookedOnly/Private/UncookedOnlyUtils.cpp
--- a/Engine/Plugins/Experimental/AnimNext/Source/AnimNextUncookedOnly/Private/UncookedOnlyUtils.cpp
+++ b/Engine/Plugins/Experimental/AnimNext/Source/AnimNextUncookedOnly/Private/UncookedOnlyUtils.cpp
@@ -1303,6 +1303,11 @@ void FUtils::GetAssetParameters(const UAnimNextRigVMAssetEditorData* EditorData,
 
 void FUtils::GetGraphParameters(const URigVMGraph* Graph, FAnimNextParameterProviderAssetRegistryExports& OutExports)
 {
+	if (Graph == nullptr)
+	{
+		return;
+	}
+
 	const TArray<URigVMNode*>& Nodes = Graph->GetNodes();
 	for (URigVMNode* Node : Nodes)
 	{
