Follow-up to switching to using FLazyName for control rig built in types (32291909)

#rnx

diff --git a/Engine/Plugins/Experimental/AnimNext/Source/AnimNext/Private/Param/ParamType.cpp b/Engine/Plugins/Experimental/AnimNext/Source/AnimNext/Private/Param/ParamType.cpp
--- a/Engine/Plugins/Experimental/AnimNext/Source/AnimNext/Private/Param/ParamType.cpp
+++ b/Engine/Plugins/Experimental/AnimNext/Source/AnimNext/Private/Param/ParamType.cpp
@@ -23,9 +23,9 @@ FAnimNextParamType::FAnimNextParamType(EValueType InValueType, EContainerType In
 FAnimNextParamType FAnimNextParamType::FromRigVMTemplateArgument(const FRigVMTemplateArgumentType& RigVMType)
 {
 	FAnimNextParamType Type;	
-	const FString CPPType = RigVMType.CPPType.ToString();
+	const FString CPPTypeString = RigVMType.CPPType.ToString();
 	
-	if (RigVMTypeUtils::IsArrayType(CPPType))
+	if (RigVMTypeUtils::IsArrayType(CPPTypeString))
 	{
 		Type.ContainerType = EPropertyBagContainerType::Array;
 	}
@@ -34,6 +34,8 @@ FAnimNextParamType FAnimNextParamType::FromRigVMTemplateArgument(const FRigVMTem
 	static const FName Int64TypeName(TEXT("Int64"));
 	static const FName UInt64TypeName(TEXT("UInt64"));
 
+	const FName CPPType = *CPPTypeString;
+	
 	if (CPPType == RigVMTypeUtils::BoolTypeName)
 	{
 		Type.ValueType = EPropertyBagPropertyType::Bool;
@@ -87,7 +89,7 @@ FAnimNextParamType FAnimNextParamType::FromRigVMTemplateArgument(const FRigVMTem
 	}
 	else
 	{
-		ensureMsgf(false, TEXT("Unsupported type : %s"), *CPPType);
+		ensureMsgf(false, TEXT("Unsupported type : %s"), *CPPTypeString);
 		Type.ValueType = EPropertyBagPropertyType::None;
 	}
 
