AnimNext: Loading and saving fixes

Basic fixes to unblock people for now.
This adds a 'force save' option to AnimNext graph's CB context menu that will resave the asset and all its external subobjects (more generic solution pending discussion).

Also this fixes a load-order issue where schedules could be compiled before the exposed parameters of graphs could be cached and output to the asset registry. Caching of asset registry exports is now done lazily as well as at compilation time.

diff --git a/Engine/Plugins/Experimental/AnimNext/Source/AnimNextEditor/AnimNextEditor.Build.cs b/Engine/Plugins/Experimental/AnimNext/Source/AnimNextEditor/AnimNextEditor.Build.cs
--- a/Engine/Plugins/Experimental/AnimNext/Source/AnimNextEditor/AnimNextEditor.Build.cs
+++ b/Engine/Plugins/Experimental/AnimNext/Source/AnimNextEditor/AnimNextEditor.Build.cs
@@ -43,6 +43,7 @@ namespace UnrealBuildTool.Rules
 					"Settings",
 					"EditorWidgets",
 					"WorkspaceEditor",
+					"ContentBrowser",
 				}
 			);
 		}
diff --git a/Engine/Plugins/Experimental/AnimNext/Source/AnimNextEditor/Private/Graph/AnimNextGraphAssetDefinition.cpp b/Engine/Plugins/Experimental/AnimNext/Source/AnimNextEditor/Private/Graph/AnimNextGraphAssetDefinition.cpp
--- a/Engine/Plugins/Experimental/AnimNext/Source/AnimNextEditor/Private/Graph/AnimNextGraphAssetDefinition.cpp
+++ b/Engine/Plugins/Experimental/AnimNext/Source/AnimNextEditor/Private/Graph/AnimNextGraphAssetDefinition.cpp
@@ -2,10 +2,14 @@
 
 #include "AnimNextGraphAssetDefinition.h"
 #include "AnimNextGraphEditor.h"
+#include "ContentBrowserMenuContexts.h"
 #include "EditorCVars.h"
+#include "FileHelpers.h"
 #include "IWorkspaceEditorModule.h"
 #include "Workspace/AnimNextWorkspaceFactory.h"
 
+#define LOCTEXT_NAMESPACE "AssetDefinition_AnimNextGraph"
+
 EAssetCommandResult UAssetDefinition_AnimNextGraph::OpenAssets(const FAssetOpenArgs& OpenArgs) const
 {
 	using namespace UE::AnimNext::Editor;
@@ -45,4 +49,55 @@ FText UAssetDefinition_AnimNextEventGraph::GetObjectDisplayNameText(UObject* Obj
 {
 	UAnimNextRigVMAssetEntry* Parameter = CastChecked<UAnimNextRigVMAssetEntry>(Object);
 	return Parameter->GetDisplayName();
-}
\ No newline at end of file
+}
+
+namespace UE::AnimNext::Editor
+{
+
+static FDelayedAutoRegisterHelper AutoRegisterAssetMenuItems(EDelayedRegisterRunPhase::EndOfEngineInit, []
+{
+	UToolMenus::RegisterStartupCallback(FSimpleMulticastDelegate::FDelegate::CreateLambda([]()
+	{
+		FToolMenuOwnerScoped OwnerScoped(UE_MODULE_NAME);
+		UToolMenu* Menu = UE::ContentBrowser::ExtendToolMenu_AssetContextMenu(UAnimNextGraph::StaticClass());
+		 
+		FToolMenuSection& Section = Menu->FindOrAddSection("GetAssetActions");
+		Section.AddDynamicEntry(NAME_None, FNewToolMenuSectionDelegate::CreateLambda([](FToolMenuSection& InSection)
+		{
+			InSection.AddMenuEntry(
+				"ForceSave",
+				LOCTEXT("ForceSaveLabel", "Force Save"),
+				LOCTEXT("ForceSaveTooltip", "Force the save of this item and all its subobjects."),
+				FSlateIcon(FAppStyle::Get().GetStyleSetName(), "Icons.Save"),
+				FToolUIAction(FToolMenuExecuteAction::CreateLambda([](const FToolMenuContext& InContext)
+				{
+					if(UContentBrowserAssetContextMenuContext* AssetContext = InContext.FindContext<UContentBrowserAssetContextMenuContext>())
+					{
+						TArray<UPackage*> PackagesToSave;
+						for (const FAssetData& AssetData : AssetContext->SelectedAssets)
+						{
+							if (UPackage* Package = AssetData.GetPackage())
+							{
+								PackagesToSave.Add(Package);
+								PackagesToSave.Append(Package->GetExternalPackages());
+							}
+						}
+
+						if(PackagesToSave.Num() > 0)
+						{
+							FEditorFileUtils::FPromptForCheckoutAndSaveParams SaveParams;
+							SaveParams.bCheckDirty = false;
+							SaveParams.bPromptToSave = false;
+							SaveParams.bIsExplicitSave = true;
+
+							FEditorFileUtils::PromptForCheckoutAndSave(PackagesToSave, SaveParams);
+						}
+					}
+				})));
+		}));
+	}));
+});
+
+}
+
+#undef LOCTEXT_NAMESPACE
\ No newline at end of file
diff --git a/Engine/Plugins/Experimental/AnimNext/Source/AnimNextUncookedOnly/Internal/AnimNextRigVMAssetEditorData.h b/Engine/Plugins/Experimental/AnimNext/Source/AnimNextUncookedOnly/Internal/AnimNextRigVMAssetEditorData.h
--- a/Engine/Plugins/Experimental/AnimNext/Source/AnimNextUncookedOnly/Internal/AnimNextRigVMAssetEditorData.h
+++ b/Engine/Plugins/Experimental/AnimNext/Source/AnimNextUncookedOnly/Internal/AnimNextRigVMAssetEditorData.h
@@ -225,7 +225,7 @@ protected:
 
 	// Find an entry that corresponds to the specified RigVMGraph. This uses the name of the graph to match the entry 
 	UAnimNextRigVMAssetEntry* FindEntryForRigVMGraph(URigVMGraph* InRigVMGraph) const;
-
+	
 	/** All entries in this asset - not saved, either serialized or discovered at load time */
 	UPROPERTY(transient)
 	TArray<TObjectPtr<UAnimNextRigVMAssetEntry>> Entries;
@@ -269,9 +269,9 @@ protected:
 	// Delegate to subscribe to modifications to this editor data
 	UE::AnimNext::UncookedOnly::FOnEditorDataModified ModifiedDelegate;
 
-	// Cached exports, generated on compilation
-	FAnimNextParameterProviderAssetRegistryExports CachedExports;
-	
+	// Cached exports, generated lazily or on compilation
+	mutable TOptional<FAnimNextParameterProviderAssetRegistryExports> CachedExports;
+
 	bool bAutoRecompileVM = true;
 	bool bErrorsDuringCompilation = false;
 	bool bSuspendModelNotificationsForSelf = false;
diff --git a/Engine/Plugins/Experimental/AnimNext/Source/AnimNextUncookedOnly/Private/AnimNextRigVMAssetEditorData.cpp b/Engine/Plugins/Experimental/AnimNext/Source/AnimNextUncookedOnly/Private/AnimNextRigVMAssetEditorData.cpp
--- a/Engine/Plugins/Experimental/AnimNext/Source/AnimNextUncookedOnly/Private/AnimNextRigVMAssetEditorData.cpp
+++ b/Engine/Plugins/Experimental/AnimNext/Source/AnimNextUncookedOnly/Private/AnimNextRigVMAssetEditorData.cpp
@@ -154,6 +154,13 @@ void UAnimNextRigVMAssetEditorData::GetAssetRegistryTags(FAssetRegistryTagsConte
 {
 	Super::GetAssetRegistryTags(Context);
 
+	// We may not have compiled yet, so cache exports if we havent already
+	if(!CachedExports.IsSet())
+	{
+		CachedExports = FAnimNextParameterProviderAssetRegistryExports();
+		UE::AnimNext::UncookedOnly::FUtils::GetAssetParameters(this, CachedExports.GetValue());
+	}
+
 	FString TagValue;
 	FAnimNextParameterProviderAssetRegistryExports::StaticStruct()->ExportText(TagValue, &CachedExports, nullptr, nullptr, PPF_None, nullptr);
 	Context.AddTag(FAssetRegistryTag(UE::AnimNext::ExportsAnimNextAssetRegistryTag, TagValue, FAssetRegistryTag::TT_Hidden));
@@ -779,4 +786,4 @@ UAnimNextRigVMAssetEntry* UAnimNextRigVMAssetEditorData::FindEntryForRigVMGraph(
 	}
 
 	return nullptr;
-}
+}
\ No newline at end of file
diff --git a/Engine/Plugins/Experimental/AnimNext/Source/AnimNextUncookedOnly/Private/Graph/AnimNextGraph_EditorData.cpp b/Engine/Plugins/Experimental/AnimNext/Source/AnimNextUncookedOnly/Private/Graph/AnimNextGraph_EditorData.cpp
--- a/Engine/Plugins/Experimental/AnimNext/Source/AnimNextUncookedOnly/Private/Graph/AnimNextGraph_EditorData.cpp
+++ b/Engine/Plugins/Experimental/AnimNext/Source/AnimNextUncookedOnly/Private/Graph/AnimNextGraph_EditorData.cpp
@@ -129,7 +129,8 @@ void UAnimNextGraph_EditorData::PostLoad()
 
 void UAnimNextGraph_EditorData::RecompileVM()
 {
-	UE::AnimNext::UncookedOnly::FUtils::GetAssetParameters(this, CachedExports);
+	CachedExports = FAnimNextParameterProviderAssetRegistryExports();
+	UE::AnimNext::UncookedOnly::FUtils::GetAssetParameters(this, CachedExports.GetValue());
 	UE::AnimNext::UncookedOnly::FUtils::Compile(GetTypedOuter<UAnimNextGraph>());
 }
 
