Added an event to Workspace Module to be able to register details customizations from clients
Added RigVM Math Customizations at AnimNext Module

diff --git a/Engine/Plugins/Experimental/AnimNext/Source/AnimNextEditor/AnimNextEditor.Build.cs b/Engine/Plugins/Experimental/AnimNext/Source/AnimNextEditor/AnimNextEditor.Build.cs
--- a/Engine/Plugins/Experimental/AnimNext/Source/AnimNextEditor/AnimNextEditor.Build.cs
+++ b/Engine/Plugins/Experimental/AnimNext/Source/AnimNextEditor/AnimNextEditor.Build.cs
@@ -19,6 +19,7 @@ namespace UnrealBuildTool.Rules
 					"AnimNextUncookedOnly",
 					"UnrealEd",
 					"SlateCore",
+					"AnimationCore",
 					"Slate",
 					"InputCore",
 					"PropertyEditor",
diff --git a/Engine/Plugins/Experimental/AnimNext/Source/AnimNextEditor/Private/AnimNextEditorModule.cpp b/Engine/Plugins/Experimental/AnimNext/Source/AnimNextEditor/Private/AnimNextEditorModule.cpp
--- a/Engine/Plugins/Experimental/AnimNext/Source/AnimNextEditor/Private/AnimNextEditorModule.cpp
+++ b/Engine/Plugins/Experimental/AnimNext/Source/AnimNextEditor/Private/AnimNextEditorModule.cpp
@@ -29,6 +29,7 @@
 #include "IWorkspaceEditorModule.h"
 #include "Common/SActionMenu.h"
 #include "AnimNextRigVMAssetEntry.h"
+#include "Editor/RigVMGraphDetailCustomization.h"
 
 #define LOCTEXT_NAMESPACE "AnimNextEditorModule"
 
@@ -75,7 +76,30 @@ class FModule : public IModule
 		ParametersGraphPanelPinFactory = MakeShared<FParametersGraphPanelPinFactory>();
 		FEdGraphUtilities::RegisterVisualPinFactory(ParametersGraphPanelPinFactory);
 
-		RegisterWorkspaceDocumentTypes();
+		Workspace::IWorkspaceEditorModule& WorkspaceEditorModule = FModuleManager::Get().LoadModuleChecked<Workspace::IWorkspaceEditorModule>("WorkspaceEditor");
+
+		RegisterWorkspaceDocumentTypes(WorkspaceEditorModule);
+
+		WorkspaceEditorModule.OnRegisterWorkspaceDetailsCustomization().AddLambda([](TSharedPtr<IDetailsView>& InDetailsView)
+			{
+				TArray<UScriptStruct*> StructsToCustomize = {
+					TBaseStructure<FVector>::Get(),
+					TBaseStructure<FVector2D>::Get(),
+					TBaseStructure<FVector4>::Get(),
+					TBaseStructure<FRotator>::Get(),
+					TBaseStructure<FQuat>::Get(),
+					TBaseStructure<FTransform>::Get(),
+					TBaseStructure<FEulerTransform>::Get(),
+				};
+
+				for (UScriptStruct* StructToCustomize : StructsToCustomize)
+				{
+					InDetailsView->RegisterInstancedCustomPropertyTypeLayout(StructToCustomize->GetFName(), FOnGetPropertyTypeCustomizationInstance::CreateLambda([]()
+						{
+							return FRigVMGraphMathTypeDetailCustomization::MakeInstance();
+						}));
+				}
+			});
 
 		SRigVMAssetView::RegisterCategoryFactory("Parameters", [](UAnimNextRigVMAssetEditorData* InEditorData)
 		{
@@ -194,9 +218,9 @@ class FModule : public IModule
 			.Args(InArgs);
 	}
 
-	void RegisterWorkspaceDocumentTypes()
+	void RegisterWorkspaceDocumentTypes(Workspace::IWorkspaceEditorModule& WorkspaceEditorModule)
 	{
-		Workspace::IWorkspaceEditorModule& WorkspaceEditorModule = FModuleManager::Get().LoadModuleChecked<Workspace::IWorkspaceEditorModule>("WorkspaceEditor");
+		// --- AnimNextSchedule ---
 		WorkspaceEditorModule.RegisterObjectDocumentType(FTopLevelAssetPath(TEXT("/Script/AnimNext.AnimNextSchedule")),
 			Workspace::FObjectDocumentArgs(
 				Workspace::FOnMakeDocumentWidget::CreateLambda([](const Workspace::FWorkspaceEditorContext& InContext)
@@ -211,6 +235,7 @@ class FModule : public IModule
 				}),
 				Workspace::WorkspaceTabs::MiddleDocumentArea));
 
+		// --- AnimNextGraph ---
 		WorkspaceEditorModule.RegisterObjectDocumentType(FTopLevelAssetPath(TEXT("/Script/AnimNext.AnimNextGraph")),
 			Workspace::FObjectDocumentArgs(
 				Workspace::FOnMakeDocumentWidget::CreateLambda([](const Workspace::FWorkspaceEditorContext& InContext)
@@ -279,6 +304,7 @@ class FModule : public IModule
 				}),
 				Workspace::WorkspaceTabs::LeftDocumentArea));
 
+		// --- AnimNextGraph_EdGraph ---
 		Workspace::FGraphDocumentWidgetArgs GraphArgs;
 		GraphArgs.SpawnLocation = Workspace::WorkspaceTabs::MiddleDocumentArea;
 		GraphArgs.OnCreateActionMenu = Workspace::FOnCreateActionMenu::CreateLambda([](const Workspace::FWorkspaceEditorContext& InContext, UEdGraph* InGraph, const FVector2D& InNodePosition, const TArray<UEdGraphPin*>& InDraggedPins, bool bAutoExpand, SGraphEditor::FActionMenuClosed InOnMenuClosed)
