Introduce event handling in AnimNext

diff --git a/Engine/Plugins/Experimental/AnimNext/Source/AnimNext/Private/Graph/AnimNextGraphInstance.cpp b/Engine/Plugins/Experimental/AnimNext/Source/AnimNext/Private/Graph/AnimNextGraphInstance.cpp
--- a/Engine/Plugins/Experimental/AnimNext/Source/AnimNext/Private/Graph/AnimNextGraphInstance.cpp
+++ b/Engine/Plugins/Experimental/AnimNext/Source/AnimNext/Private/Graph/AnimNextGraphInstance.cpp
@@ -7,6 +7,7 @@
 #include "Graph/AnimNextGraph.h"
 #include "Graph/GC_GraphInstanceComponent.h"
 #include "Graph/RigUnit_AnimNextShimRoot.h"
+#include "Misc/ScopeRWLock.h"
 
 DEFINE_STAT(STAT_AnimNext_Graph_RigVM);
 
@@ -117,6 +118,38 @@ GraphInstanceComponentMapType::TConstIterator FAnimNextGraphInstance::GetCompone
 	return RootGraphInstance->Components.CreateConstIterator();
 }
 
+void FAnimNextGraphInstance::QueueInputTraitEvent(FAnimNextTraitEventPtr Event)
+{
+	if (!Event || !Event->IsValid())
+	{
+		return;
+	}
+
+	FWriteScopeLock WriteLock(RootGraphInstance->InputEventListLock);
+	RootGraphInstance->InputEventList.Push(Event);
+}
+
+void FAnimNextGraphInstance::QueueInputTraitEvents(const UE::AnimNext::FTraitEventList& Events)
+{
+	FWriteScopeLock WriteLock(RootGraphInstance->InputEventListLock);
+
+	for (const FAnimNextTraitEventPtr& Event : Events)
+	{
+		if (Event->IsValid())
+		{
+			RootGraphInstance->InputEventList.Push(Event);
+		}
+	}
+}
+
+void FAnimNextGraphInstance::CollectInputTraitEvents(UE::AnimNext::FTraitEventList& OutInputEvents)
+{
+	OutInputEvents.Reset();
+
+	FWriteScopeLock WriteLock(RootGraphInstance->InputEventListLock);
+	Swap(OutInputEvents, RootGraphInstance->InputEventList);
+}
+
 void FAnimNextGraphInstance::ExecuteLatentPins(const TConstArrayView<UE::AnimNext::FLatentPropertyHandle>& LatentHandles, void* DestinationBasePtr, bool bIsFrozen)
 {
 	SCOPE_CYCLE_COUNTER(STAT_AnimNext_Graph_RigVM);
diff --git a/Engine/Plugins/Experimental/AnimNext/Source/AnimNext/Private/Graph/AnimNextGraphInstancePtr.cpp b/Engine/Plugins/Experimental/AnimNext/Source/AnimNext/Private/Graph/AnimNextGraphInstancePtr.cpp
--- a/Engine/Plugins/Experimental/AnimNext/Source/AnimNext/Private/Graph/AnimNextGraphInstancePtr.cpp
+++ b/Engine/Plugins/Experimental/AnimNext/Source/AnimNext/Private/Graph/AnimNextGraphInstancePtr.cpp
@@ -88,3 +88,21 @@ GraphInstanceComponentMapType::TConstIterator FAnimNextGraphInstancePtr::GetComp
 	check(Impl);
 	return Impl->GetComponentIterator();
 }
+
+void FAnimNextGraphInstancePtr::QueueInputTraitEvent(FAnimNextTraitEventPtr Event)
+{
+	check(Impl);
+	Impl->QueueInputTraitEvent(Event);
+}
+
+void FAnimNextGraphInstancePtr::QueueInputTraitEvents(const UE::AnimNext::FTraitEventList& Events)
+{
+	check(Impl);
+	Impl->QueueInputTraitEvents(Events);
+}
+
+void FAnimNextGraphInstancePtr::CollectInputTraitEvents(UE::AnimNext::FTraitEventList& OutInputEvents)
+{
+	check(Impl);
+	Impl->CollectInputTraitEvents(OutInputEvents);
+}
diff --git a/Engine/Plugins/Experimental/AnimNext/Source/AnimNext/Private/TraitCore/ExecutionContext.cpp b/Engine/Plugins/Experimental/AnimNext/Source/AnimNext/Private/TraitCore/ExecutionContext.cpp
--- a/Engine/Plugins/Experimental/AnimNext/Source/AnimNext/Private/TraitCore/ExecutionContext.cpp
+++ b/Engine/Plugins/Experimental/AnimNext/Source/AnimNext/Private/TraitCore/ExecutionContext.cpp
@@ -552,6 +552,16 @@ namespace UE::AnimNext
 		return RootGraphInstance->GetComponentIterator();
 	}
 
+	void FExecutionContext::RaiseInputTraitEvent(FAnimNextTraitEventPtr Event)
+	{
+		ensureMsgf(false, TEXT("Raising input trait events is not supported in this context"));
+	}
+
+	void FExecutionContext::RaiseOutputTraitEvent(FAnimNextTraitEventPtr Event)
+	{
+		ensureMsgf(false, TEXT("Raising output trait events is not supported in this context"));
+	}
+
 	const FNodeDescription& FExecutionContext::GetNodeDescription(const FAnimNextGraphInstance& GraphInstance, FNodeHandle NodeHandle) const
 	{
 		// Grab the node description from the specified graph
diff --git a/Engine/Plugins/Experimental/AnimNext/Source/AnimNext/Private/TraitCore/Trait.cpp b/Engine/Plugins/Experimental/AnimNext/Source/AnimNext/Private/TraitCore/Trait.cpp
--- a/Engine/Plugins/Experimental/AnimNext/Source/AnimNext/Private/TraitCore/Trait.cpp
+++ b/Engine/Plugins/Experimental/AnimNext/Source/AnimNext/Private/TraitCore/Trait.cpp
@@ -56,6 +56,25 @@ namespace UE::AnimNext
 		return Result;
 	}
 
+	TArray<FTraitEventUID> FTrait::BuildTraitEventList(
+		const TConstArrayView<FTraitEventUID>& SuperEvents,
+		std::initializer_list<FTraitEventUID> EventList)
+	{
+		TArray<FTraitEventUID> Result;
+		Result.Reserve(SuperEvents.Num() + EventList.size());
+
+		Result.Append(SuperEvents);
+
+		for (FTraitEventUID InterfaceID : EventList)
+		{
+			Result.AddUnique(InterfaceID);
+		}
+
+		Result.Shrink();
+		Result.Sort();
+		return Result;
+	}
+
 #if WITH_EDITOR
 	void FTrait::SaveTraitSharedData(const TFunction<FString(FName PropertyName)>& GetTraitProperty, FAnimNextTraitSharedData& OutSharedData) const
 	{
diff --git a/Engine/Plugins/Experimental/AnimNext/Source/AnimNext/Private/TraitCore/TraitEvent.cpp b/Engine/Plugins/Experimental/AnimNext/Source/AnimNext/Private/TraitCore/TraitEvent.cpp
new file mode 100644
--- /dev/null
+++ b/Engine/Plugins/Experimental/AnimNext/Source/AnimNext/Private/TraitCore/TraitEvent.cpp
@@ -0,0 +1,62 @@
+// Copyright Epic Games, Inc. All Rights Reserved.
+
+#include "TraitCore/TraitEvent.h"
+#include "TraitCore/TraitEventList.h"
+#include "TraitCore/TraitBinding.h"
+#include "TraitCore/TraitStackBinding.h"
+#include "Graph/GraphInstanceComponent.h"
+
+namespace UE::AnimNext
+{
+	void RaiseTraitEvent(FExecutionContext& Context, const FTraitStackBinding& StackBinding, FAnimNextTraitEvent& Event)
+	{
+		if (!StackBinding.IsValid() || !Event.IsValid())
+		{
+			return;
+		}
+
+		FTraitBinding Binding;
+		ensure(StackBinding.GetTopTrait(Binding));
+		do
+		{
+			const ETraitStackPropagation Propagation = Binding.GetTrait()->OnTraitEvent(Context, Binding, Event);
+			if (Propagation == ETraitStackPropagation::Stop || Event.IsConsumed())
+			{
+				// We don't wish to propagate the call further, stop here
+				break;
+			}
+		} while (StackBinding.GetParentTrait(Binding, Binding));
+	}
+
+	void RaiseTraitEvents(FExecutionContext& Context, const FTraitStackBinding& StackBinding, const FTraitEventList& EventList)
+	{
+		if (!StackBinding.IsValid())
+		{
+			return;
+		}
+
+		// Event handlers can raise events and as such the list may change while we iterate
+		// However, if an event is added while we iterate, we will not visit it
+		const int32 NumEvents = EventList.Num();
+		for (int32 EventIndex = 0; EventIndex < NumEvents; ++EventIndex)
+		{
+			const FAnimNextTraitEventPtr Event = EventList[EventIndex];
+			RaiseTraitEvent(Context, StackBinding, *Event);
+		}
+	}
+
+	void RaiseTraitEvents(FExecutionContext& Context, FGraphInstanceComponent& Component, const FTraitEventList& EventList)
+	{
+		// Event handlers can raise events and as such the list may change while we iterate
+		// However, if an event is added while we iterate, we will not visit it
+		const int32 NumEvents = EventList.Num();
+		for (int32 EventIndex = 0; EventIndex < NumEvents; ++EventIndex)
+		{
+			const FAnimNextTraitEventPtr Event = EventList[EventIndex];
+			if (Event->IsValid())
+			{
+				Component.OnTraitEvent(Context, *Event);
+			}
+		}
+	}
+}
diff --git a/Engine/Plugins/Experimental/AnimNext/Source/AnimNext/Private/TraitCore/TraitEventLifetime.cpp b/Engine/Plugins/Experimental/AnimNext/Source/AnimNext/Private/TraitCore/TraitEventLifetime.cpp
new file mode 100644
--- /dev/null
+++ b/Engine/Plugins/Experimental/AnimNext/Source/AnimNext/Private/TraitCore/TraitEventLifetime.cpp
@@ -0,0 +1,23 @@
+// Copyright Epic Games, Inc. All Rights Reserved.
+
+#include "TraitCore/TraitEventLifetime.h"
+
+namespace UE::AnimNext
+{
+	bool FTraitEventLifetime::Decrement()
+	{
+		if (IsInfinite())
+		{
+			// Event has infinite duration, we never expire
+			return false;
+		}
+		else if (IsExpired())
+		{
+			// Event has already expired
+			return true;
+		}
+
+		LifetimeCount--;
+		return LifetimeCount == 0;
+	}
+}
diff --git a/Engine/Plugins/Experimental/AnimNext/Source/AnimNext/Public/ConstExprUID.h b/Engine/Plugins/Experimental/AnimNext/Source/AnimNext/Public/ConstExprUID.h
new file mode 100644
--- /dev/null
+++ b/Engine/Plugins/Experimental/AnimNext/Source/AnimNext/Public/ConstExprUID.h
@@ -0,0 +1,121 @@
+// Copyright Epic Games, Inc. All Rights Reserved.
+
+#pragma once
+
+#include "CoreMinimal.h"
+
+namespace UE::AnimNext
+{
+	// Type alias for a raw UID, not typesafe
+	using FFConstExprUIDRaw = uint32;
+
+	// Implements a constexpr usable version of FNVa32 for ANSICHAR string literals
+	constexpr uint32 ConstexprStringFnv32(const ANSICHAR* StringLiteral)
+	{
+		constexpr uint32 Offset = 0x811c9dc5;
+		constexpr uint32 Prime = 0x01000193;
+
+		const char* CharPtr = StringLiteral;
+
+		uint32 Fnv = Offset;
+		while (*CharPtr != 0)
+		{
+			Fnv ^= *CharPtr++;
+			Fnv *= Prime;
+		}
+
+		return Fnv;
+	}
+
+	// Implements a constexpr usable version of FNVa32 for WIDECHAR string literals
+	constexpr uint32 ConstexprStringFnv32(const WIDECHAR* StringLiteral)
+	{
+		constexpr uint32 Offset = 0x811c9dc5;
+		constexpr uint32 Prime = 0x01000193;
+
+		const WIDECHAR* CharPtr = StringLiteral;
+
+		uint32 Fnv = Offset;
+		while (*CharPtr != 0)
+		{
+			const uint32 HighByte = (*CharPtr >> 8) & 0xff;
+			Fnv ^= HighByte;
+			Fnv *= Prime;
+
+			const uint32 LowByte = *CharPtr & 0xff;
+			Fnv ^= LowByte;
+			Fnv *= Prime;
+
+			CharPtr++;
+		}
+
+		return Fnv;
+	}
+
+	/**
+	 * A constexpr safe UID
+	 *
+	 * Encapsulates an constexpr UID.
+	 * The string is exposed in non-shipping builds for logging and debugging purposes.
+	 * The UID should be generated from the provided string using FNV1a with 32 bits.
+	 *
+	 * The whole struct is meant to be 'constexpr' to allow inlining.
+	 */
+	struct FConstExprUID final
+	{
+		// Constructs an invalid UID
+		constexpr FConstExprUID() noexcept
+			: UID(INVALID_UID)
+#if !UE_BUILD_SHIPPING
+			, DebugName(TEXT("<Invalid ConstExprUID UID>"))
+#endif
+		{
+		}
+
+		// Constructs a UID
+		explicit constexpr FConstExprUID(FFConstExprUIDRaw InUID, const TCHAR* InDebugName = TEXT("<Unknown ConstExprUID Name>")) noexcept
+			: UID(InUID)
+#if !UE_BUILD_SHIPPING
+			, DebugName(InDebugName)
+#endif
+		{
+		}
+
+		// Constructs a UID
+		static constexpr FConstExprUID MakeFromString(const TCHAR* InName) noexcept
+		{
+			return FConstExprUID(ConstexprStringFnv32(InName), InName);
+		}
+
+#if !UE_BUILD_SHIPPING
+		// Returns a literal string to the debug name
+		constexpr const TCHAR* GetDebugName() const noexcept { return DebugName; }
+#endif
+
+		// Returns the raw UID
+		constexpr FFConstExprUIDRaw GetUID() const noexcept { return UID; }
+
+		// Returns whether this UID is valid or not
+		constexpr bool IsValid() const noexcept { return UID != INVALID_UID; }
+
+	private:
+		static constexpr FFConstExprUIDRaw INVALID_UID = 0;
+
+		FFConstExprUIDRaw	UID;
+
+#if !UE_BUILD_SHIPPING
+		const TCHAR*		DebugName;
+#endif
+	};
+
+	// Compares for equality and inequality
+	constexpr bool operator==(FConstExprUID LHS, FConstExprUID RHS) noexcept { return LHS.GetUID() == RHS.GetUID(); }
+	constexpr bool operator!=(FConstExprUID LHS, FConstExprUID RHS) noexcept { return LHS.GetUID() != RHS.GetUID(); }
+	constexpr bool operator==(FConstExprUID LHS, FFConstExprUIDRaw RHS) noexcept { return LHS.GetUID() == RHS; }
+	constexpr bool operator!=(FConstExprUID LHS, FFConstExprUIDRaw RHS) noexcept { return LHS.GetUID() != RHS; }
+	constexpr bool operator==(FFConstExprUIDRaw LHS, FConstExprUID RHS) noexcept { return LHS == RHS.GetUID(); }
+	constexpr bool operator!=(FFConstExprUIDRaw LHS, FConstExprUID RHS) noexcept { return LHS != RHS.GetUID(); }
+
+	// For sorting
+	constexpr bool operator<(FConstExprUID LHS, FConstExprUID RHS) { return LHS.GetUID() < RHS.GetUID(); }
+}
diff --git a/Engine/Plugins/Experimental/AnimNext/Source/AnimNext/Public/EvaluationVM/EvaluationVM.h b/Engine/Plugins/Experimental/AnimNext/Source/AnimNext/Public/EvaluationVM/EvaluationVM.h
--- a/Engine/Plugins/Experimental/AnimNext/Source/AnimNext/Public/EvaluationVM/EvaluationVM.h
+++ b/Engine/Plugins/Experimental/AnimNext/Source/AnimNext/Public/EvaluationVM/EvaluationVM.h
@@ -4,6 +4,7 @@
 
 #include "CoreMinimal.h"
 
+#include "ConstExprUID.h"
 #include "Animation/AnimCurveFilter.h"
 #include "EvaluationVM/EvaluationFlags.h"
 #include "EvaluationVM/KeyframeState.h"
@@ -22,23 +23,6 @@
 
 namespace UE::AnimNext
 {
-	constexpr uint32 ConstexprStringFnv32(const char* StringLiteral)
-	{
-		constexpr uint32 Offset = 0x811c9dc5;
-		constexpr uint32 Prime = 0x01000193;
-
-		const char* CharPtr = StringLiteral;
-
-		uint32 Fnv = Offset;
-		while (*CharPtr != 0)
-		{
-			Fnv ^= *CharPtr++;
-			Fnv *= Prime;
-		}
-
-		return Fnv;
-	}
-
 	// Helper function that returns a UID for the specified type
 	// @see ANIM_NEXT_ENABLE_EVALUATION_STACK_USAGE
 	template<typename Type>
diff --git a/Engine/Plugins/Experimental/AnimNext/Source/AnimNext/Public/Graph/AnimNextGraphInstance.h b/Engine/Plugins/Experimental/AnimNext/Source/AnimNext/Public/Graph/AnimNextGraphInstance.h
--- a/Engine/Plugins/Experimental/AnimNext/Source/AnimNext/Public/Graph/AnimNextGraphInstance.h
+++ b/Engine/Plugins/Experimental/AnimNext/Source/AnimNext/Public/Graph/AnimNextGraphInstance.h
@@ -3,6 +3,9 @@
 #pragma once
 
 #include "CoreMinimal.h"
+#include "HAL/CriticalSection.h"
+#include "TraitCore/TraitEvent.h"
+#include "TraitCore/TraitEventList.h"
 #include "TraitCore/TraitPtr.h"
 #include "RigVMCore/RigVMExecuteContext.h"
 
@@ -90,6 +93,18 @@ struct ANIMNEXT_API FAnimNextGraphInstance
 	// Returns const iterators to the graph instance component container
 	GraphInstanceComponentMapType::TConstIterator GetComponentIterator() const;
 
+	// Queues an input trait event (thread safe)
+	// Input events will be processed in the next graph update after they are queued
+	void QueueInputTraitEvent(FAnimNextTraitEventPtr Event);
+
+	// Queues a list of input trait events (thread safe)
+	// Input events will be processed in the next graph update after they are queued
+	// Expired or invalid events will be skipped
+	void QueueInputTraitEvents(const UE::AnimNext::FTraitEventList& Events);
+
+	// Collects all currently queued input trait events (thread safe)
+	void CollectInputTraitEvents(UE::AnimNext::FTraitEventList& OutInputEvents);
+
 private:
 	// Returns a pointer to the specified component, or nullptr if not found
 	UE::AnimNext::FGraphInstanceComponent* TryGetComponent(int32 ComponentNameHash, FName ComponentName) const;
@@ -131,6 +146,12 @@ private:
 	// Graph instance components that persist from update to update
 	GraphInstanceComponentMapType Components;
 
+	// Input event list to be processed on the next update
+	UE::AnimNext::FTraitEventList InputEventList;
+
+	// Lock to ensure input event list actions are thread safe
+	FRWLock InputEventListLock;
+
 	friend UAnimNextGraph;					// The graph is the one that allocates instances
 	friend FRigUnit_AnimNextGraphEvaluator;	// We evaluate the instance
 	friend UE::AnimNext::FExecutionContext;
diff --git a/Engine/Plugins/Experimental/AnimNext/Source/AnimNext/Public/Graph/AnimNextGraphInstancePtr.h b/Engine/Plugins/Experimental/AnimNext/Source/AnimNext/Public/Graph/AnimNextGraphInstancePtr.h
--- a/Engine/Plugins/Experimental/AnimNext/Source/AnimNext/Public/Graph/AnimNextGraphInstancePtr.h
+++ b/Engine/Plugins/Experimental/AnimNext/Source/AnimNext/Public/Graph/AnimNextGraphInstancePtr.h
@@ -3,6 +3,7 @@
 #pragma once
 
 #include "CoreMinimal.h"
+#include "TraitCore/TraitEvent.h"
 #include "TraitCore/TraitPtr.h"
 
 #include "AnimNextGraphInstancePtr.generated.h"
@@ -80,6 +81,18 @@ struct ANIMNEXT_API FAnimNextGraphInstancePtr
 	// Returns const iterators to the graph instance component container
 	GraphInstanceComponentMapType::TConstIterator GetComponentIterator() const;
 
+	// Queues an input trait event (thread safe)
+	// Input events will be processed in the next graph update after they are queued
+	void QueueInputTraitEvent(FAnimNextTraitEventPtr Event);
+
+	// Queues a list of input trait events (thread safe)
+	// Input events will be processed in the next graph update after they are queued
+	// Expired or invalid events will be skipped
+	void QueueInputTraitEvents(const UE::AnimNext::FTraitEventList& Events);
+
+	// Collects all currently queued input trait events (thread safe)
+	void CollectInputTraitEvents(UE::AnimNext::FTraitEventList& OutInputEvents);
+
 private:
 	// Returns a pointer to the specified component, or nullptr if not found
 	UE::AnimNext::FGraphInstanceComponent* TryGetComponent(int32 ComponentNameHash, FName ComponentName) const;
diff --git a/Engine/Plugins/Experimental/AnimNext/Source/AnimNext/Public/Graph/GraphInstanceComponent.h b/Engine/Plugins/Experimental/AnimNext/Source/AnimNext/Public/Graph/GraphInstanceComponent.h
--- a/Engine/Plugins/Experimental/AnimNext/Source/AnimNext/Public/Graph/GraphInstanceComponent.h
+++ b/Engine/Plugins/Experimental/AnimNext/Source/AnimNext/Public/Graph/GraphInstanceComponent.h
@@ -4,6 +4,8 @@
 
 #include "CoreMinimal.h"
 
+struct FAnimNextTraitEvent;
+
 // Declares a graph instance component and implements the necessary boilerplate
 #define DECLARE_ANIM_GRAPH_INSTANCE_COMPONENT(ComponentType) \
 	static FName StaticComponentName() { return FName(TEXT(#ComponentType)); } \
@@ -34,5 +36,8 @@ namespace UE::AnimNext
 		// Called after the update traversal completes, after every node has been visited
 		// The execution context provided is bound to the graph root and can be bound to anything the component wishes
 		virtual void PostUpdate(FExecutionContext& Context) {}
+
+		// Called before PreUpdate with input events and before PostUpdate with output events
+		virtual void OnTraitEvent(FExecutionContext& Context, FAnimNextTraitEvent& Event) {}
 	};
 }
diff --git a/Engine/Plugins/Experimental/AnimNext/Source/AnimNext/Public/TraitCore/ExecutionContext.h b/Engine/Plugins/Experimental/AnimNext/Source/AnimNext/Public/TraitCore/ExecutionContext.h
--- a/Engine/Plugins/Experimental/AnimNext/Source/AnimNext/Public/TraitCore/ExecutionContext.h
+++ b/Engine/Plugins/Experimental/AnimNext/Source/AnimNext/Public/TraitCore/ExecutionContext.h
@@ -4,6 +4,7 @@
 
 #include "CoreMinimal.h"
 #include "TraitCore/TraitBinding.h"
+#include "TraitCore/TraitEvent.h"
 #include "TraitCore/TraitHandle.h"
 #include "TraitCore/TraitPtr.h"
 #include "TraitCore/TraitStackBinding.h"
@@ -52,6 +53,8 @@ namespace UE::AnimNext
 		// Creates an execution context and binds it to the specified graph instance
 		explicit FExecutionContext(FAnimNextGraphInstance& InGraphInstance);
 
+		// Derived types can control various features
+		virtual ~FExecutionContext() = default;
 
 
 		//////////////////////////////////////////////////////////////////////////
@@ -175,6 +178,21 @@ namespace UE::AnimNext
 
 
 
+		//////////////////////////////////////////////////////////////////////////
+		// The following functions handle node lifetime management
+
+		// Raises an input trait event
+		// Behavior is traversal dependent, see the derived type for your current traversal (e.g. IUpdate)
+		// Default behavior asserts
+		virtual void RaiseInputTraitEvent(FAnimNextTraitEventPtr Event);
+
+		// Raises an output trait event
+		// Behavior is traversal dependent, see the derived type for your current traversal (e.g. IUpdate)
+		// Default behavior asserts
+		virtual void RaiseOutputTraitEvent(FAnimNextTraitEventPtr Event);
+
+
+
 		//////////////////////////////////////////////////////////////////////////
 		// Misc functions
 
diff --git a/Engine/Plugins/Experimental/AnimNext/Source/AnimNext/Public/TraitCore/ITraitInterface.h b/Engine/Plugins/Experimental/AnimNext/Source/AnimNext/Public/TraitCore/ITraitInterface.h
--- a/Engine/Plugins/Experimental/AnimNext/Source/AnimNext/Public/TraitCore/ITraitInterface.h
+++ b/Engine/Plugins/Experimental/AnimNext/Source/AnimNext/Public/TraitCore/ITraitInterface.h
@@ -32,4 +32,19 @@ namespace UE::AnimNext
 		// Returns the globally unique UID for this interface
 		virtual FTraitInterfaceUID GetInterfaceUID() const { return InterfaceUID; };
 	};
+
+	/**
+	 * Trait Stack Propagation
+	 * 
+	 * This enum can be used to signal whether an interface call should be forwarded to its parent
+	 * on the trait stack or not.
+	 */
+	enum class ETraitStackPropagation
+	{
+		// Forward the call to our parent on the trait stack (if we have one)
+		Continue,
+
+		// Do not forward the call to our parent, execution stops
+		Stop,
+	};
 }
diff --git a/Engine/Plugins/Experimental/AnimNext/Source/AnimNext/Public/TraitCore/Trait.h b/Engine/Plugins/Experimental/AnimNext/Source/AnimNext/Public/TraitCore/Trait.h
--- a/Engine/Plugins/Experimental/AnimNext/Source/AnimNext/Public/TraitCore/Trait.h
+++ b/Engine/Plugins/Experimental/AnimNext/Source/AnimNext/Public/TraitCore/Trait.h
@@ -4,6 +4,7 @@
 
 #include "CoreMinimal.h"
 
+#include "TraitCore/TraitEvent.h"
 #include "TraitCore/TraitHandle.h"			// Derived types are likely to refer to other traits as children
 #include "TraitCore/TraitInstanceData.h"
 #include "TraitCore/TraitMode.h"
@@ -32,6 +33,8 @@ class FArchive;
 	virtual void DestructTraitInstance(const UE::AnimNext::FExecutionContext& Context, const UE::AnimNext::FTraitBinding& Binding) const override; \
 	virtual const UE::AnimNext::ITraitInterface* GetTraitInterface(UE::AnimNext::FTraitInterfaceUID InterfaceUID) const override; \
 	virtual TConstArrayView<FTraitInterfaceUID> GetTraitInterfaces() const override; \
+	virtual UE::AnimNext::ETraitStackPropagation OnTraitEvent(UE::AnimNext::FExecutionContext& Context, UE::AnimNext::FTraitBinding& Binding, FAnimNextTraitEvent& Event) const override; \
+	virtual TConstArrayView<FTraitEventUID> GetTraitEvents() const override; \
 	virtual uint32 GetNumLatentTraitProperties() const override { return -FSharedData::GetLatentPropertyIndex(~(size_t)0); } \
 	virtual FTraitLatentPropertyMemoryLayout GetLatentPropertyMemoryLayout(FName PropertyName, uint32 PropertyIndex) const override; \
 	static_assert(std::is_base_of<FAnimNextTraitSharedData, FSharedData>::value, "Trait shared data must derive from FAnimNextTraitSharedData"); \
@@ -101,20 +104,93 @@ class FArchive;
 	{ \
 		/* Thread safe cache initialization */ \
 		static TArray<UE::AnimNext::FTraitInterfaceUID> CachedInterfaceList = FTrait::BuildTraitInterfaceList( \
-		TraitSuper::GetTraitInterfaces(), \
-		{ \
-			InterfaceEnumeratorMacro(ANIM_NEXT_IMPL_GET_INTERFACES_IMPL_FOR_INTERFACE) \
-		}); \
+			TraitSuper::GetTraitInterfaces(), \
+			{ \
+				InterfaceEnumeratorMacro(ANIM_NEXT_IMPL_GET_INTERFACES_IMPL_FOR_INTERFACE) \
+			}); \
 		return CachedInterfaceList; \
 	}
 
+namespace UE::AnimNext::Private
+{
+	// Helper to grab the event type from an event handler function signature
+	template<typename HandlerEventType>
+	struct EventHandlerTypeTrait;
+
+	template<typename HandlerEventType>
+	struct EventHandlerTypeTrait<ETraitStackPropagation(*)(FExecutionContext&, FTraitBinding&, HandlerEventType&)>
+	{
+		using EventType = std::remove_const_t<HandlerEventType>;
+	};
+
+	template<typename HandlerEventType>
+	struct EventHandlerTypeTrait<ETraitStackPropagation(*)(const FExecutionContext&, FTraitBinding&, HandlerEventType&)>
+	{
+		using EventType = std::remove_const_t<HandlerEventType>;
+	};
+
+	template<typename BaseType, typename HandlerEventType>
+	struct EventHandlerTypeTrait<ETraitStackPropagation(BaseType::*)(FExecutionContext&, FTraitBinding&, HandlerEventType&) const>
+	{
+		using EventType = std::remove_const_t<HandlerEventType>;
+	};
+
+	template<typename BaseType, typename HandlerEventType>
+	struct EventHandlerTypeTrait<ETraitStackPropagation(BaseType::*)(const FExecutionContext&, FTraitBinding&, HandlerEventType&) const>
+	{
+		using EventType = std::remove_const_t<HandlerEventType>;
+	};
+}
+
+// Helper that handles the OnTraitEvent() details for each event specified by the generator macro
+#define ANIM_NEXT_IMPL_ON_TRAIT_EVENT_IMPL_FOR_EVENT(EventHandler) \
+	if (EventUID == UE::AnimNext::Private::EventHandlerTypeTrait<decltype(&EventHandler)>::EventType::TypeUID) \
+	{ \
+		return EventHandler(Context, Binding, static_cast<UE::AnimNext::Private::EventHandlerTypeTrait<decltype(&EventHandler)>::EventType&>(Event)); \
+	}
+
+// Implements OnTraitEvent()
+#define ANIM_NEXT_IMPL_DEFINE_ANIM_TRAIT_ON_TRAIT_EVENT(TraitName, EventEnumeratorMacro) \
+	UE::AnimNext::ETraitStackPropagation TraitName::OnTraitEvent(UE::AnimNext::FExecutionContext& Context, UE::AnimNext::FTraitBinding& Binding, FAnimNextTraitEvent& Event) const \
+	{ \
+		const UE::AnimNext::FTraitEventUID EventUID = Event.GetTypeUID(); \
+		EventEnumeratorMacro(ANIM_NEXT_IMPL_ON_TRAIT_EVENT_IMPL_FOR_EVENT) \
+		/* Forward to base implementation */ \
+		return TraitSuper::OnTraitEvent(Context, Binding, Event); \
+	}
+
+// Helper that handles the GetTraitEvents() details for each event specified by the generator macro
+#define ANIM_NEXT_IMPL_GET_TRAIT_EVENTS_IMPL_FOR_EVENT(EventHandler) \
+	UE::AnimNext::Private::EventHandlerTypeTrait<decltype(&EventHandler)>::EventType::TypeUID,
+
+// Implements GetTraitEvents()
+#define ANIM_NEXT_IMPL_DEFINE_ANIM_TRAIT_GET_TRAIT_EVENTS(TraitName, EventEnumeratorMacro) \
+	TConstArrayView<UE::AnimNext::FTraitEventUID> TraitName::GetTraitEvents() const \
+	{ \
+		/* Thread safe cache initialization */ \
+		static TArray<UE::AnimNext::FTraitEventUID> CachedEventList = FTrait::BuildTraitEventList( \
+			TraitSuper::GetTraitEvents(), \
+			{ \
+				EventEnumeratorMacro(ANIM_NEXT_IMPL_GET_TRAIT_EVENTS_IMPL_FOR_EVENT) \
+			}); \
+		return CachedEventList; \
+	}
+
+// A dummy trait interface generator for traits that do not implement any interfaces
+#define NULL_ANIM_TRAIT_INTERFACE_ENUMERATOR(GeneratorMacro)
+
+// A dummy trait event generator for traits that do not handle any events
+#define NULL_ANIM_TRAIT_EVENT_ENUMERATOR(GeneratorMacro)
+
 /**
   * This macro defines the necessary boilerplate for implementing FTrait. See above for usage example.
   */
-#define GENERATE_ANIM_TRAIT_IMPLEMENTATION(TraitName, InterfaceEnumeratorMacro) \
+#define GENERATE_ANIM_TRAIT_IMPLEMENTATION(TraitName, InterfaceEnumeratorMacro, EventEnumeratorMacro) \
 	ANIM_NEXT_IMPL_DEFINE_ANIM_TRAIT(TraitName) \
 	ANIM_NEXT_IMPL_DEFINE_ANIM_TRAIT_GET_INTERFACE(TraitName, InterfaceEnumeratorMacro) \
-	ANIM_NEXT_IMPL_DEFINE_ANIM_TRAIT_GET_INTERFACES(TraitName, InterfaceEnumeratorMacro)
+	ANIM_NEXT_IMPL_DEFINE_ANIM_TRAIT_GET_INTERFACES(TraitName, InterfaceEnumeratorMacro) \
+	ANIM_NEXT_IMPL_DEFINE_ANIM_TRAIT_ON_TRAIT_EVENT(TraitName, EventEnumeratorMacro) \
+	ANIM_NEXT_IMPL_DEFINE_ANIM_TRAIT_GET_TRAIT_EVENTS(TraitName, EventEnumeratorMacro)
 
 // Allows a trait to auto-register and unregister within the current execution scope
 // The trait must be found in the current scope without a namespace qualification
@@ -262,6 +338,12 @@ namespace UE::AnimNext
 		// Returns a list of interfaces that this trait supports
 		virtual TConstArrayView<FTraitInterfaceUID> GetTraitInterfaces() const { return TConstArrayView<FTraitInterfaceUID>(); }
 
+		// Called when an event reaches an instance of this trait
+		virtual ETraitStackPropagation OnTraitEvent(FExecutionContext& Context, FTraitBinding& Binding, FAnimNextTraitEvent& Event) const { return ETraitStackPropagation::Continue; }
+
+		// Returns a list of events that this trait handles
+		virtual TConstArrayView<FTraitEventUID> GetTraitEvents() const { return TConstArrayView<FTraitEventUID>(); }
+
 		// The number of latent property properties in the shared data of this trait
 		virtual uint32 GetNumLatentTraitProperties() const { return 0; }
 
@@ -295,6 +377,11 @@ namespace UE::AnimNext
 		static TArray<FTraitInterfaceUID> BuildTraitInterfaceList(
 			const TConstArrayView<FTraitInterfaceUID>& SuperInterfaces,
 			std::initializer_list<FTraitInterfaceUID> InterfaceList);
+
+		// Builds a list of events with the provided super events and current events as an initializer list
+		static TArray<FTraitEventUID> BuildTraitEventList(
+			const TConstArrayView<FTraitEventUID>& SuperEvents,
+			std::initializer_list<FTraitEventUID> EventList);
 	};
 
 	// Base class for base traits that are standalone
diff --git a/Engine/Plugins/Experimental/AnimNext/Source/AnimNext/Public/TraitCore/TraitBinding.h b/Engine/Plugins/Experimental/AnimNext/Source/AnimNext/Public/TraitCore/TraitBinding.h
--- a/Engine/Plugins/Experimental/AnimNext/Source/AnimNext/Public/TraitCore/TraitBinding.h
+++ b/Engine/Plugins/Experimental/AnimNext/Source/AnimNext/Public/TraitCore/TraitBinding.h
@@ -156,6 +156,9 @@ namespace UE::AnimNext
 		// Returns the trait index on the stack we are bound to.
 		uint32 GetTraitIndex() const { return Stack != nullptr ? (TraitIndex - Stack->BaseTraitIndex) : 0; }
 
+		// Returns a pointer to the trait implementation if valid and if the trait is present at runtime
+		const FTrait* GetTrait() const { return TraitImpl; }
+
 		// Returns the trait interface UID when bound, an invalid UID otherwise.
 		FTraitInterfaceUID GetInterfaceUID() const;
 
diff --git a/Engine/Plugins/Experimental/AnimNext/Source/AnimNext/Public/TraitCore/TraitEvent.h b/Engine/Plugins/Experimental/AnimNext/Source/AnimNext/Public/TraitCore/TraitEvent.h
new file mode 100644
--- /dev/null
+++ b/Engine/Plugins/Experimental/AnimNext/Source/AnimNext/Public/TraitCore/TraitEvent.h
@@ -0,0 +1,135 @@
+// Copyright Epic Games, Inc. All Rights Reserved.
+
+#pragma once
+
+#include "CoreMinimal.h"
+
+#include "ConstExprUID.h"
+#include "TraitCore/TraitEventLifetime.h"
+
+#include "TraitEvent.generated.h"
+
+namespace UE::AnimNext
+{
+	// Create an alias to improve readability
+	using FTraitEventUID = UE::AnimNext::FConstExprUID;
+}
+
+// Helper macros
+// In a trait event struct declaration, this macro declares the necessary boilerplate we require
+#define DECLARE_ANIM_TRAIT_EVENT(EventName) \
+	/* FAnimNextTraitEvent impl */ \
+	static constexpr UE::AnimNext::FTraitEventUID TypeUID = UE::AnimNext::FTraitEventUID::MakeFromString(TEXT(#EventName)); \
+	virtual UE::AnimNext::FTraitEventUID GetTypeUID() const override { return TypeUID; }
+
+
+/**
+ * Trait Event
+ * 
+ * Encapsulates an event in the trait system.
+ * Events can be marked as handled to signal to future handlers that no further action needs to be taken.
+ * Consuming an event prevents it from propagating to other handlers.
+ */
+USTRUCT(BlueprintType)
+struct FAnimNextTraitEvent
+{
+	GENERATED_BODY()
+
+	// Creates a transient event
+	FAnimNextTraitEvent() = default;
+
+	// Allow typesafe destruction of derived types
+	virtual ~FAnimNextTraitEvent() = default;
+
+	// Returns the event type UID
+	virtual UE::AnimNext::FTraitEventUID GetTypeUID() const { return UE::AnimNext::FTraitEventUID(); }
+
+	// Returns whether or not this event is valid
+	// An event is valid if it hasn't been consumed and if it isn't expired
+	bool IsValid() const { return !IsConsumed() && !IsExpired(); }
+
+	// Returns the lifetime of this event
+	UE::AnimNext::FTraitEventLifetime GetLifetime() const { return Lifetime; }
+
+	// Sets the desired lifetime for this event
+	void SetLifetime(UE::AnimNext::FTraitEventLifetime InLifetime) { Lifetime = InLifetime; }
+
+	// Marks this event as being handled.
+	// Subsequent handlers can use this to determine if they should take action or not.
+	void MarkHandled() { bIsHandled = true; }
+
+	// Returns whether or not this event has been handled by a previous handler
+	bool IsHandled() const { return bIsHandled; }
+
+	// Marks this event as being consumed.
+	// This event won't be forward to any other handlers.
+	void MarkConsumed() { bIsConsumed = true; }
+
+	// Returns whether or not this event has been consumed by a previous handler
+	bool IsConsumed() const { return bIsConsumed; }
+
+	// Returns whether or not this event has expired
+	bool IsExpired() const { return Lifetime.IsExpired(); }
+
+	// Returns whether or not this event has infinite duration
+	bool IsInfinite() const { return Lifetime.IsInfinite(); }
+
+	// Returns whether or not this event has transient duration
+	bool IsTransient() const { return Lifetime.IsTransient(); }
+
+	// Decrements the lifetime count and returns whether or not this event has expired in the process
+	bool DecrementLifetime() { return Lifetime.Decrement(); }
+
+	// Returns a pointer to this event cast to the specified type if the TypeUIDs match, nullptr otherwise
+	template<class EventType>
+	EventType* AsType() { return EventType::TypeUID == GetTypeUID() ? static_cast<EventType*>(this) : nullptr; }
+
+	// Returns a pointer to this event cast to the specified type if the TypeUIDs match, nullptr otherwise
+	template<class EventType>
+	const EventType* AsType() const { return EventType::TypeUID == GetTypeUID() ? static_cast<const EventType*>(this) : nullptr; }
+
+protected:
+	// Creates an event of the specified type and associated lifetime
+	explicit FAnimNextTraitEvent(UE::AnimNext::FTraitEventLifetime InLifetime)
+		: Lifetime(InLifetime)
+	{
+	}
+
+private:
+	// The specified lifetime of this event
+	UE::AnimNext::FTraitEventLifetime Lifetime = UE::AnimNext::FTraitEventLifetime::MakeTransient();
+
+	// Signals whether someone has opted to handle this event
+	bool bIsHandled = false;
+
+	// Signals whether someone consumed this event and it thus shouldn't be forwarded to other handlers
+	bool bIsConsumed = false;
+};
+
+// Create a shared pointer alias for trait events
+using FAnimNextTraitEventPtr = TSharedPtr<FAnimNextTraitEvent, ESPMode::ThreadSafe>;
+
+template<class EventType, typename... TArgs>
+TSharedPtr<EventType, ESPMode::ThreadSafe> MakeTraitEvent(TArgs&&... Args)
+{
+	return MakeShared<EventType>(Forward<TArgs>(Args)...);
+}
+
+namespace UE::AnimNext
+{
+	struct FExecutionContext;
+	struct FTraitStackBinding;
+	struct FTraitEventList;
+	struct FGraphInstanceComponent;
+
+	// Raises a trait event on the specified trait stack
+	// This calls OnTraitEvent on each trait of the stack following propagation rules
+	ANIMNEXT_API void RaiseTraitEvent(FExecutionContext& Context, const FTraitStackBinding& StackBinding, FAnimNextTraitEvent& Event);
+
+	// Raises every trait event from the input list on the specified trait stack
+	// This calls OnTraitEvent on each trait of the stack following propagation rules
+	ANIMNEXT_API void RaiseTraitEvents(FExecutionContext& Context, const FTraitStackBinding& StackBinding, const FTraitEventList& EventList);
+
+	// Raises every trait event from the input list on the specified component
+	ANIMNEXT_API void RaiseTraitEvents(FExecutionContext& Context, FGraphInstanceComponent& Component, const FTraitEventList& EventList);
+}
diff --git a/Engine/Plugins/Experimental/AnimNext/Source/AnimNext/Public/TraitCore/TraitEventLifetime.h b/Engine/Plugins/Experimental/AnimNext/Source/AnimNext/Public/TraitCore/TraitEventLifetime.h
new file mode 100644
--- /dev/null
+++ b/Engine/Plugins/Experimental/AnimNext/Source/AnimNext/Public/TraitCore/TraitEventLifetime.h
@@ -0,0 +1,69 @@
+// Copyright Epic Games, Inc. All Rights Reserved.
+
+#pragma once
+
+#include "CoreMinimal.h"
+
+namespace UE::AnimNext
+{
+	/**
+	 * Trait Event Lifetime
+	 * 
+	 * This encapsulates a specific lifetime duration.
+	 */
+	struct FTraitEventLifetime final
+	{
+		// Creates a lifetime that extends until the next graph update terminates
+		static constexpr FTraitEventLifetime MakeTransient() noexcept
+		{
+			return FTraitEventLifetime(TRANSIENT_LIFETIME);
+		}
+
+		// Creates a lifetime that extends for the specified number of graph updates
+		static constexpr FTraitEventLifetime MakeUntil(uint32 LifetimeCount) noexcept
+		{
+			return FTraitEventLifetime(LifetimeCount);
+		}
+
+		// Creates a lifetime that extends forever
+		static constexpr FTraitEventLifetime MakeInfinite() noexcept
+		{
+			return FTraitEventLifetime(INFINITE_LIFETIME);
+		}
+
+		// Creates an expired lifetime
+		constexpr FTraitEventLifetime() noexcept
+			: LifetimeCount(EXPIRED_LIFETIME)
+		{
+		}
+
+		// Returns whether or not this lifetime has expired
+		constexpr bool IsExpired() const { return LifetimeCount == EXPIRED_LIFETIME; }
+
+		// Returns whether or not this lifetime is infinite
+		constexpr bool IsInfinite() const { return LifetimeCount == INFINITE_LIFETIME; }
+
+		// Returns whether or not this lifetime is transient
+		constexpr bool IsTransient() const { return LifetimeCount == TRANSIENT_LIFETIME; }
+
+		// Decrements the lifetime count and returns whether or not it has expired in the process
+		ANIMNEXT_API bool Decrement();
+
+	private:
+		// Infinite lifetime duration
+		static constexpr uint32 INFINITE_LIFETIME = ~0;
+
+		// Transient lifetime duration
+		static constexpr uint32 TRANSIENT_LIFETIME = 1;
+
+		// Expired lifetime duration
+		static constexpr uint32 EXPIRED_LIFETIME = 0;
+
+		constexpr explicit FTraitEventLifetime(uint32 InLifetimeCount) noexcept
+			: LifetimeCount(InLifetimeCount)
+		{
+		}
+
+		uint32 LifetimeCount;
+	};
+}
diff --git a/Engine/Plugins/Experimental/AnimNext/Source/AnimNext/Public/TraitCore/TraitEventList.h b/Engine/Plugins/Experimental/AnimNext/Source/AnimNext/Public/TraitCore/TraitEventList.h
new file mode 100644
--- /dev/null
+++ b/Engine/Plugins/Experimental/AnimNext/Source/AnimNext/Public/TraitCore/TraitEventList.h
@@ -0,0 +1,48 @@
+// Copyright Epic Games, Inc. All Rights Reserved.
+
+#pragma once
+
+#include "CoreMinimal.h"
+
+#include "TraitCore/TraitEvent.h"
+
+namespace UE::AnimNext
+{
+	/**
+	 * Trait Event List
+	 * 
+	 * Encapsulates a list of trait events.
+	 */
+	struct FTraitEventList
+	{
+		using IteratorType = TArray<FAnimNextTraitEventPtr>::RangedForIteratorType;
+		using ConstIteratorType = TArray<FAnimNextTraitEventPtr>::RangedForConstIteratorType;
+
+		FTraitEventList() = default;
+
+		void Push(FAnimNextTraitEventPtr Event) { Events.Add(MoveTemp(Event)); }
+		void Reset() { Events.Reset(); }
+		int32 Num() const { return Events.Num(); }
+
+		FAnimNextTraitEventPtr& operator[](int32 EventIndex) { return Events[EventIndex]; }
+		const FAnimNextTraitEventPtr& operator[](int32 EventIndex) const { return Events[EventIndex]; }
+
+		IteratorType begin() { return Events.begin(); }
+		ConstIteratorType begin() const { return Events.begin(); }
+		IteratorType end() { return Events.end(); }
+		ConstIteratorType end() const { return Events.end(); }
+
+		// Decrements the lifetime of every event within the list
+		void DecrementLifetime()
+		{
+			for (FAnimNextTraitEventPtr& Event : Events)
+			{
+				Event->DecrementLifetime();
+			}
+		}
+
+	private:
+		// A list of events
+		TArray<FAnimNextTraitEventPtr> Events;
+	};
+}
diff --git a/Engine/Plugins/Experimental/AnimNext/Source/AnimNextTestSuite/Private/AnimNextTraitBaseTest.cpp b/Engine/Plugins/Experimental/AnimNext/Source/AnimNextTestSuite/Private/AnimNextTraitBaseTest.cpp
--- a/Engine/Plugins/Experimental/AnimNext/Source/AnimNextTestSuite/Private/AnimNextTraitBaseTest.cpp
+++ b/Engine/Plugins/Experimental/AnimNext/Source/AnimNextTestSuite/Private/AnimNextTraitBaseTest.cpp
@@ -161,7 +161,7 @@ namespace UE::AnimNext
 	#define TRAIT_INTERFACE_ENUMERATOR(GeneratorMacro) \
 		GeneratorMacro(IInterfaceA) \
 
-	GENERATE_ANIM_TRAIT_IMPLEMENTATION(FTraitA_Base, TRAIT_INTERFACE_ENUMERATOR)
+	GENERATE_ANIM_TRAIT_IMPLEMENTATION(FTraitA_Base, TRAIT_INTERFACE_ENUMERATOR, NULL_ANIM_TRAIT_EVENT_ENUMERATOR)
 	#undef TRAIT_INTERFACE_ENUMERATOR
 
 	//////////////////////////////////////////////////////////////////////////
@@ -209,7 +209,7 @@ namespace UE::AnimNext
 		GeneratorMacro(IInterfaceA) \
 		GeneratorMacro(IInterfaceB) \
 
-	GENERATE_ANIM_TRAIT_IMPLEMENTATION(FTraitAB_Add, TRAIT_INTERFACE_ENUMERATOR)
+	GENERATE_ANIM_TRAIT_IMPLEMENTATION(FTraitAB_Add, TRAIT_INTERFACE_ENUMERATOR, NULL_ANIM_TRAIT_EVENT_ENUMERATOR)
 	#undef TRAIT_INTERFACE_ENUMERATOR
 
 	//////////////////////////////////////////////////////////////////////////
@@ -257,7 +257,7 @@ namespace UE::AnimNext
 		GeneratorMacro(IInterfaceA) \
 		GeneratorMacro(IInterfaceC) \
 
-	GENERATE_ANIM_TRAIT_IMPLEMENTATION(FTraitAC_Add, TRAIT_INTERFACE_ENUMERATOR)
+	GENERATE_ANIM_TRAIT_IMPLEMENTATION(FTraitAC_Add, TRAIT_INTERFACE_ENUMERATOR, NULL_ANIM_TRAIT_EVENT_ENUMERATOR)
 	#undef TRAIT_INTERFACE_ENUMERATOR
 
 
@@ -274,7 +274,7 @@ namespace UE::AnimNext
 	#define TRAIT_INTERFACE_ENUMERATOR(GeneratorMacro) \
 		GeneratorMacro(IInterfaceA) \
 
-	GENERATE_ANIM_TRAIT_IMPLEMENTATION(FTraitSerialization_Base, TRAIT_INTERFACE_ENUMERATOR)
+	GENERATE_ANIM_TRAIT_IMPLEMENTATION(FTraitSerialization_Base, TRAIT_INTERFACE_ENUMERATOR, NULL_ANIM_TRAIT_EVENT_ENUMERATOR)
 	#undef TRAIT_INTERFACE_ENUMERATOR
 
 	//////////////////////////////////////////////////////////////////////////
@@ -290,7 +290,7 @@ namespace UE::AnimNext
 	#define TRAIT_INTERFACE_ENUMERATOR(GeneratorMacro) \
 		GeneratorMacro(IInterfaceB) \
 
-	GENERATE_ANIM_TRAIT_IMPLEMENTATION(FTraitSerialization_Add, TRAIT_INTERFACE_ENUMERATOR)
+	GENERATE_ANIM_TRAIT_IMPLEMENTATION(FTraitSerialization_Add, TRAIT_INTERFACE_ENUMERATOR, NULL_ANIM_TRAIT_EVENT_ENUMERATOR)
 	#undef TRAIT_INTERFACE_ENUMERATOR
 
 	//////////////////////////////////////////////////////////////////////////
@@ -306,7 +306,7 @@ namespace UE::AnimNext
 	#define TRAIT_INTERFACE_ENUMERATOR(GeneratorMacro) \
 		GeneratorMacro(IInterfaceC) \
 
-	GENERATE_ANIM_TRAIT_IMPLEMENTATION(FTraitNativeSerialization_Add, TRAIT_INTERFACE_ENUMERATOR)
+	GENERATE_ANIM_TRAIT_IMPLEMENTATION(FTraitNativeSerialization_Add, TRAIT_INTERFACE_ENUMERATOR, NULL_ANIM_TRAIT_EVENT_ENUMERATOR)
 	#undef TRAIT_INTERFACE_ENUMERATOR
 }
 
diff --git a/Engine/Plugins/Experimental/AnimNext/Source/AnimNextTestSuite/Private/AnimNextTraitEventTest.cpp b/Engine/Plugins/Experimental/AnimNext/Source/AnimNextTestSuite/Private/AnimNextTraitEventTest.cpp
new file mode 100644
--- /dev/null
+++ b/Engine/Plugins/Experimental/AnimNext/Source/AnimNextTestSuite/Private/AnimNextTraitEventTest.cpp
@@ -0,0 +1,428 @@
+// Copyright Epic Games, Inc. All Rights Reserved.
+
+#include "AnimNextTraitEventTest.h"
+#include "AnimNextRuntimeTest.h"
+#include "AnimNextTest.h"
+
+#include "Misc/AutomationTest.h"
+
+#if WITH_DEV_AUTOMATION_TESTS
+
+#include "TraitCore/Trait.h"
+#include "TraitCore/TraitBinding.h"
+#include "TraitCore/TraitEventList.h"
+#include "TraitCore/TraitReader.h"
+#include "TraitCore/TraitWriter.h"
+#include "TraitCore/ExecutionContext.h"
+#include "TraitCore/NodeTemplateBuilder.h"
+#include "TraitCore/NodeTemplateRegistry.h"
+#include "Graph/AnimNextGraph.h"
+#include "Graph/GraphFactory.h"
+
+//****************************************************************************
+// AnimNext Runtime TraitEvent Tests
+//****************************************************************************
+
+namespace UE::AnimNext
+{
+	struct FTraitCoreTest_EventAB_Base : FBaseTrait
+	{
+		DECLARE_ANIM_TRAIT(FTraitCoreTest_EventAB_Base, 0x6bab1b28, FBaseTrait)
+
+		// We support static functions
+		static ETraitStackPropagation OnEventA(const FExecutionContext& Context, FTraitBinding& Binding, FTraitCoreTest_EventA& Event)
+		{
+			Event.VisitedTraits.Add(FTraitCoreTest_EventAB_Base::TraitUID);
+			return ETraitStackPropagation::Continue;
+		}
+
+		// We support const member functions
+		ETraitStackPropagation OnEventB(const FExecutionContext& Context, FTraitBinding& Binding, FTraitCoreTest_EventB& Event) const
+		{
+			Event.VisitedTraits.Add(FTraitCoreTest_EventAB_Base::TraitUID);
+			return ETraitStackPropagation::Continue;
+		}
+	};
+
+	// Trait implementation boilerplate
+	#define TRAIT_EVENT_ENUMERATOR(GeneratorMacro) \
+		GeneratorMacro(FTraitCoreTest_EventAB_Base::OnEventA) \
+		GeneratorMacro(FTraitCoreTest_EventAB_Base::OnEventB) \
+
+	GENERATE_ANIM_TRAIT_IMPLEMENTATION(FTraitCoreTest_EventAB_Base, NULL_ANIM_TRAIT_INTERFACE_ENUMERATOR, TRAIT_EVENT_ENUMERATOR)
+	#undef TRAIT_EVENT_ENUMERATOR
+
+	struct FTraitCoreTest_EventAB_Add : FAdditiveTrait
+	{
+		DECLARE_ANIM_TRAIT(FTraitCoreTest_EventAB_Add, 0x193709c4, FAdditiveTrait)
+
+		ETraitStackPropagation OnEventA(const FExecutionContext& Context, FTraitBinding& Binding, FTraitCoreTest_EventA& Event) const
+		{
+			Event.VisitedTraits.Add(FTraitCoreTest_EventAB_Add::TraitUID);
+			return Event.bAlwaysForwardToBase ? ETraitStackPropagation::Continue : ETraitStackPropagation::Stop;
+		}
+
+		ETraitStackPropagation OnEventB(const FExecutionContext& Context, FTraitBinding& Binding, FTraitCoreTest_EventB& Event) const
+		{
+			Event.VisitedTraits.Add(FTraitCoreTest_EventAB_Add::TraitUID);
+			return ETraitStackPropagation::Continue;
+		}
+	};
+
+	// Trait implementation boilerplate
+	#define TRAIT_EVENT_ENUMERATOR(GeneratorMacro) \
+		GeneratorMacro(FTraitCoreTest_EventAB_Add::OnEventA) \
+		GeneratorMacro(FTraitCoreTest_EventAB_Add::OnEventB) \
+
+	GENERATE_ANIM_TRAIT_IMPLEMENTATION(FTraitCoreTest_EventAB_Add, NULL_ANIM_TRAIT_INTERFACE_ENUMERATOR, TRAIT_EVENT_ENUMERATOR)
+	#undef TRAIT_EVENT_ENUMERATOR
+
+	struct FTraitCoreTest_EventA_Add : FAdditiveTrait
+	{
+		DECLARE_ANIM_TRAIT(FTraitCoreTest_EventA_Add, 0xcf0402da, FAdditiveTrait)
+
+		ETraitStackPropagation OnEventA(const FExecutionContext& Context, FTraitBinding& Binding, FTraitCoreTest_EventA& Event) const
+		{
+			Event.VisitedTraits.Add(FTraitCoreTest_EventA_Add::TraitUID);
+			return ETraitStackPropagation::Continue;
+		}
+	};
+
+	// Trait implementation boilerplate
+	#define TRAIT_EVENT_ENUMERATOR(GeneratorMacro) \
+		GeneratorMacro(FTraitCoreTest_EventA_Add::OnEventA) \
+
+	GENERATE_ANIM_TRAIT_IMPLEMENTATION(FTraitCoreTest_EventA_Add, NULL_ANIM_TRAIT_INTERFACE_ENUMERATOR, TRAIT_EVENT_ENUMERATOR)
+	#undef TRAIT_EVENT_ENUMERATOR
+
+	struct FTraitCoreTest_EventB_Add : FAdditiveTrait
+	{
+		DECLARE_ANIM_TRAIT(FTraitCoreTest_EventB_Add, 0x24ba435b, FAdditiveTrait)
+
+		ETraitStackPropagation OnEventB(const FExecutionContext& Context, FTraitBinding& Binding, FTraitCoreTest_EventB& Event) const
+		{
+			Event.VisitedTraits.Add(FTraitCoreTest_EventB_Add::TraitUID);
+			return ETraitStackPropagation::Continue;
+		}
+	};
+
+	// Trait implementation boilerplate
+	#define TRAIT_EVENT_ENUMERATOR(GeneratorMacro) \
+		GeneratorMacro(FTraitCoreTest_EventB_Add::OnEventB) \
+
+	GENERATE_ANIM_TRAIT_IMPLEMENTATION(FTraitCoreTest_EventB_Add, NULL_ANIM_TRAIT_INTERFACE_ENUMERATOR, TRAIT_EVENT_ENUMERATOR)
+	#undef TRAIT_EVENT_ENUMERATOR
+}
+
+IMPLEMENT_SIMPLE_AUTOMATION_TEST(FAnimationAnimNextRuntimeTest_TraitEventLifetime, "Animation.AnimNext.Runtime.TraitEventLifetime", EAutomationTestFlags::EditorContext | EAutomationTestFlags::EngineFilter)
+
+bool FAnimationAnimNextRuntimeTest_TraitEventLifetime::RunTest(const FString& InParameters)
+{
+	using namespace UE::AnimNext;
+
+	{
+		FTraitEventLifetime Lifetime;
+		AddErrorIfFalse(!Lifetime.IsTransient(), "FAnimationAnimNextRuntimeTest_TraitEventLifetime -> Default constructed lifetime should be expired");
+		AddErrorIfFalse(!Lifetime.IsInfinite(), "FAnimationAnimNextRuntimeTest_TraitEventLifetime -> Default constructed lifetime should be expired");
+		AddErrorIfFalse(Lifetime.IsExpired(), "FAnimationAnimNextRuntimeTest_TraitEventLifetime -> Default constructed lifetime should be expired");
+
+		bool bIsExpired = Lifetime.Decrement();
+		AddErrorIfFalse(bIsExpired, "FAnimationAnimNextRuntimeTest_TraitEventLifetime -> Decremented lifetime should be expired");
+		AddErrorIfFalse(!Lifetime.IsTransient(), "FAnimationAnimNextRuntimeTest_TraitEventLifetime -> Decremented lifetime should be expired");
+		AddErrorIfFalse(!Lifetime.IsInfinite(), "FAnimationAnimNextRuntimeTest_TraitEventLifetime -> Decremented lifetime should be expired");
+		AddErrorIfFalse(Lifetime.IsExpired(), "FAnimationAnimNextRuntimeTest_TraitEventLifetime -> Decremented lifetime should be expired");
+	}
+
+	{
+		FTraitEventLifetime Lifetime = FTraitEventLifetime::MakeTransient();
+		AddErrorIfFalse(Lifetime.IsTransient(), "FAnimationAnimNextRuntimeTest_TraitEventLifetime -> Constructed lifetime should be transient");
+		AddErrorIfFalse(!Lifetime.IsInfinite(), "FAnimationAnimNextRuntimeTest_TraitEventLifetime -> Constructed lifetime should be transient");
+		AddErrorIfFalse(!Lifetime.IsExpired(), "FAnimationAnimNextRuntimeTest_TraitEventLifetime -> Constructed lifetime should be transient");
+
+		bool bIsExpired = Lifetime.Decrement();
+		AddErrorIfFalse(bIsExpired, "FAnimationAnimNextRuntimeTest_TraitEventLifetime -> Decremented lifetime should be expired");
+		AddErrorIfFalse(!Lifetime.IsTransient(), "FAnimationAnimNextRuntimeTest_TraitEventLifetime -> Decremented lifetime should be expired");
+		AddErrorIfFalse(!Lifetime.IsInfinite(), "FAnimationAnimNextRuntimeTest_TraitEventLifetime -> Decremented lifetime should be expired");
+		AddErrorIfFalse(Lifetime.IsExpired(), "FAnimationAnimNextRuntimeTest_TraitEventLifetime -> Decremented lifetime should be expired");
+	}
+
+	{
+		FTraitEventLifetime Lifetime = FTraitEventLifetime::MakeInfinite();
+		AddErrorIfFalse(!Lifetime.IsTransient(), "FAnimationAnimNextRuntimeTest_TraitEventLifetime -> Constructed lifetime should be infinite");
+		AddErrorIfFalse(Lifetime.IsInfinite(), "FAnimationAnimNextRuntimeTest_TraitEventLifetime -> Constructed lifetime should be infinite");
+		AddErrorIfFalse(!Lifetime.IsExpired(), "FAnimationAnimNextRuntimeTest_TraitEventLifetime -> Constructed lifetime should be infinite");
+
+		bool bIsExpired = Lifetime.Decrement();
+		AddErrorIfFalse(!bIsExpired, "FAnimationAnimNextRuntimeTest_TraitEventLifetime -> Decremented lifetime should be infinite");
+		AddErrorIfFalse(!Lifetime.IsTransient(), "FAnimationAnimNextRuntimeTest_TraitEventLifetime -> Decremented lifetime should be infinite");
+		AddErrorIfFalse(Lifetime.IsInfinite(), "FAnimationAnimNextRuntimeTest_TraitEventLifetime -> Decremented lifetime should be infinite");
+		AddErrorIfFalse(!Lifetime.IsExpired(), "FAnimationAnimNextRuntimeTest_TraitEventLifetime -> Decremented lifetime should be infinite");
+	}
+
+	{
+		FTraitEventLifetime Lifetime = FTraitEventLifetime::MakeUntil(2);
+		AddErrorIfFalse(!Lifetime.IsTransient(), "FAnimationAnimNextRuntimeTest_TraitEventLifetime -> Constructed lifetime should be finite");
+		AddErrorIfFalse(!Lifetime.IsInfinite(), "FAnimationAnimNextRuntimeTest_TraitEventLifetime -> Constructed lifetime should be finite");
+		AddErrorIfFalse(!Lifetime.IsExpired(), "FAnimationAnimNextRuntimeTest_TraitEventLifetime -> Constructed lifetime should be finite");
+
+		bool bIsExpired = Lifetime.Decrement();
+		AddErrorIfFalse(!bIsExpired, "FAnimationAnimNextRuntimeTest_TraitEventLifetime -> Decremented lifetime should be transient");
+		AddErrorIfFalse(Lifetime.IsTransient(), "FAnimationAnimNextRuntimeTest_TraitEventLifetime -> Decremented lifetime should be transient");
+		AddErrorIfFalse(!Lifetime.IsInfinite(), "FAnimationAnimNextRuntimeTest_TraitEventLifetime -> Decremented lifetime should be transient");
+		AddErrorIfFalse(!Lifetime.IsExpired(), "FAnimationAnimNextRuntimeTest_TraitEventLifetime -> Decremented lifetime should be transient");
+
+		bIsExpired = Lifetime.Decrement();
+		AddErrorIfFalse(bIsExpired, "FAnimationAnimNextRuntimeTest_TraitEventLifetime -> Decremented lifetime should be expired");
+		AddErrorIfFalse(!Lifetime.IsTransient(), "FAnimationAnimNextRuntimeTest_TraitEventLifetime -> Decremented lifetime should be expired");
+		AddErrorIfFalse(!Lifetime.IsInfinite(), "FAnimationAnimNextRuntimeTest_TraitEventLifetime -> Decremented lifetime should be expired");
+		AddErrorIfFalse(Lifetime.IsExpired(), "FAnimationAnimNextRuntimeTest_TraitEventLifetime -> Decremented lifetime should be expired");
+	}
+	
+	return true;
+}
+
+IMPLEMENT_SIMPLE_AUTOMATION_TEST(FAnimationAnimNextRuntimeTest_TraitEvent, "Animation.AnimNext.Runtime.TraitEvent", EAutomationTestFlags::EditorContext | EAutomationTestFlags::EngineFilter)
+
+bool FAnimationAnimNextRuntimeTest_TraitEvent::RunTest(const FString& InParameters)
+{
+	using namespace UE::AnimNext;
+
+	{
+		FAnimNextTraitEvent Event;
+		AddErrorIfFalse(Event.IsValid(), "FAnimationAnimNextRuntimeTest_TraitEvent -> Default constructed event should be valid");
+		AddErrorIfFalse(!Event.IsHandled(), "FAnimationAnimNextRuntimeTest_TraitEvent -> Default constructed event should not be handled");
+		AddErrorIfFalse(!Event.IsConsumed(), "FAnimationAnimNextRuntimeTest_TraitEvent -> Default constructed event should not be consumed");
+		AddErrorIfFalse(!Event.IsExpired(), "FAnimationAnimNextRuntimeTest_TraitEvent -> Default constructed event should not be expired");
+		AddErrorIfFalse(!Event.IsInfinite(), "FAnimationAnimNextRuntimeTest_TraitEvent -> Default constructed event should be transient");
+		AddErrorIfFalse(Event.IsTransient(), "FAnimationAnimNextRuntimeTest_TraitEvent -> Default constructed event should be transient");
+
+		bool bIsExpired = Event.DecrementLifetime();
+		AddErrorIfFalse(bIsExpired, "FAnimationAnimNextRuntimeTest_TraitEvent -> Event should be expired");
+		AddErrorIfFalse(!Event.IsValid(), "FAnimationAnimNextRuntimeTest_TraitEvent -> Event should be expired");
+		AddErrorIfFalse(!Event.IsHandled(), "FAnimationAnimNextRuntimeTest_TraitEvent -> Event should not be handled");
+		AddErrorIfFalse(!Event.IsConsumed(), "FAnimationAnimNextRuntimeTest_TraitEvent -> Event should not be consumed");
+		AddErrorIfFalse(Event.IsExpired(), "FAnimationAnimNextRuntimeTest_TraitEvent -> Event should be expired");
+		AddErrorIfFalse(!Event.IsInfinite(), "FAnimationAnimNextRuntimeTest_TraitEvent -> Event should be expired");
+		AddErrorIfFalse(!Event.IsTransient(), "FAnimationAnimNextRuntimeTest_TraitEvent -> Event should be expired");
+	}
+
+	{
+		FAnimNextTraitEvent Event;
+		AddErrorIfFalse(Event.IsValid(), "FAnimationAnimNextRuntimeTest_TraitEvent -> Default constructed event should be valid");
+		AddErrorIfFalse(!Event.IsHandled(), "FAnimationAnimNextRuntimeTest_TraitEvent -> Default constructed event should not be handled");
+		AddErrorIfFalse(!Event.IsConsumed(), "FAnimationAnimNextRuntimeTest_TraitEvent -> Default constructed event should not be consumed");
+
+		Event.MarkHandled();
+		AddErrorIfFalse(Event.IsValid(), "FAnimationAnimNextRuntimeTest_TraitEvent -> Event should be valid and handled");
+		AddErrorIfFalse(Event.IsHandled(), "FAnimationAnimNextRuntimeTest_TraitEvent -> Event should be valid and handled");
+		AddErrorIfFalse(!Event.IsConsumed(), "FAnimationAnimNextRuntimeTest_TraitEvent -> Event should not be consumed");
+
+		Event.MarkConsumed();
+		AddErrorIfFalse(!Event.IsValid(), "FAnimationAnimNextRuntimeTest_TraitEvent -> Event should be consumed");
+		AddErrorIfFalse(Event.IsHandled(), "FAnimationAnimNextRuntimeTest_TraitEvent -> Event should be handled and consumed");
+		AddErrorIfFalse(Event.IsConsumed(), "FAnimationAnimNextRuntimeTest_TraitEvent -> Event should be consumed");
+	}
+
+	return true;
+}
+
+IMPLEMENT_SIMPLE_AUTOMATION_TEST(FAnimationAnimNextRuntimeTest_TraitEventRaising, "Animation.AnimNext.Runtime.TraitEventRaising", EAutomationTestFlags::EditorContext | EAutomationTestFlags::EngineFilter)
+
+bool FAnimationAnimNextRuntimeTest_TraitEventRaising::RunTest(const FString& InParameters)
+{
+	using namespace UE::AnimNext;
+
+	// Test event raising on a node trait stack
+	{
+		AUTO_REGISTER_ANIM_TRAIT(FTraitCoreTest_EventAB_Base)
+		AUTO_REGISTER_ANIM_TRAIT(FTraitCoreTest_EventAB_Add)
+		AUTO_REGISTER_ANIM_TRAIT(FTraitCoreTest_EventA_Add)
+		AUTO_REGISTER_ANIM_TRAIT(FTraitCoreTest_EventB_Add)
+
+		UFactory* GraphFactory = NewObject<UAnimNextGraphFactory>();
+		UAnimNextGraph* AnimNextGraph = CastChecked<UAnimNextGraph>(GraphFactory->FactoryCreateNew(UAnimNextGraph::StaticClass(), GetTransientPackage(), TEXT("TestAnimNextGraph"), RF_Transient, nullptr, nullptr, NAME_None));
+		UE_RETURN_ON_ERROR(AnimNextGraph != nullptr, "FAnimationAnimNextRuntimeTest_TraitEventRaising -> Failed to create animation graph");
+
+		FScopedClearNodeTemplateRegistry ScopedClearNodeTemplateRegistry;
+		FNodeTemplateRegistry& Registry = FNodeTemplateRegistry::Get();
+
+		// Build a node
+		TArray<FTraitUID> NodeTemplateTraitList;
+		NodeTemplateTraitList.Add(FTraitCoreTest_EventAB_Base::TraitUID);
+		NodeTemplateTraitList.Add(FTraitCoreTest_EventAB_Add::TraitUID);
+		NodeTemplateTraitList.Add(FTraitCoreTest_EventA_Add::TraitUID);
+		NodeTemplateTraitList.Add(FTraitCoreTest_EventB_Add::TraitUID);
+
+		// Populate our node template registry
+		TArray<uint8> NodeTemplateBuffer0;
+		const FNodeTemplate* NodeTemplate0 = FNodeTemplateBuilder::BuildNodeTemplate(NodeTemplateTraitList, NodeTemplateBuffer0);
+
+		FNodeTemplateRegistryHandle TemplateHandle0 = Registry.FindOrAdd(NodeTemplate0);
+		AddErrorIfFalse(TemplateHandle0.IsValid(), "FAnimationAnimNextRuntimeTest_TraitEventRaising -> Registry should contain our template");
+
+		TArray<FNodeHandle> NodeHandles;
+
+		// Write our graph
+		TArray<uint8> GraphSharedDataArchiveBuffer;
+		TArray<TObjectPtr<UObject>> GraphReferencedObjects;
+		{
+			FTraitWriter TraitWriter;
+
+			NodeHandles.Add(TraitWriter.RegisterNode(*NodeTemplate0));
+
+			// We don't have trait properties
+
+			TraitWriter.BeginNodeWriting();
+			TraitWriter.WriteNode(NodeHandles[0],
+				[](uint32 TraitIndex, FName PropertyName)
+				{
+					return FString();
+				},
+				[](uint32 TraitIndex, FName PropertyName)
+				{
+					return MAX_uint16;
+				});
+			TraitWriter.EndNodeWriting();
+
+			AddErrorIfFalse(TraitWriter.GetErrorState() == FTraitWriter::EErrorState::None, "FAnimationAnimNextRuntimeTest_TraitEventRaising -> Failed to write traits");
+			GraphSharedDataArchiveBuffer = TraitWriter.GetGraphSharedData();
+			GraphReferencedObjects = TraitWriter.GetGraphReferencedObjects();
+		}
+
+		// Read our graph
+		FTestUtils::LoadFromArchiveBuffer(*AnimNextGraph, NodeHandles, GraphSharedDataArchiveBuffer);
+
+		FAnimNextGraphInstancePtr GraphInstance;
+		AnimNextGraph->AllocateInstance(GraphInstance);
+
+		FExecutionContext Context(GraphInstance);
+
+		{
+			FAnimNextTraitHandle TraitHandle0(NodeHandles[0], 0);	// Point to node base trait
+
+			FTraitPtr TraitPtr0 = Context.AllocateNodeInstance(*GraphInstance.GetImpl(), TraitHandle0);
+			AddErrorIfFalse(TraitPtr0.IsValid(), "FAnimationAnimNextRuntimeTest_TraitEventRaising -> Failed to allocate a node instance");
+
+			// Node has 4 traits: AB base, AB add, A add, B add
+			FTraitStackBinding Stack0;
+			AddErrorIfFalse(Context.GetStack(TraitPtr0, Stack0), "FAnimationAnimNextRuntimeTest_TraitEventRaising -> Failed to bind to trait stack");
+
+			// Send events A and B, make sure every trait is visited
+			{
+				auto EventA = MakeTraitEvent<FTraitCoreTest_EventA>();
+				RaiseTraitEvent(Context, Stack0, *EventA);
+
+				UE_RETURN_ON_ERROR(EventA->VisitedTraits.Num() == 3, "FAnimationAnimNextRuntimeTest_TraitEventRaising -> Unexpected number of traits visited");
+				AddErrorIfFalse(EventA->VisitedTraits[0] == NodeTemplateTraitList[2], "FAnimationAnimNextRuntimeTest_TraitEventRaising -> Unexpected trait visited");
+				AddErrorIfFalse(EventA->VisitedTraits[1] == NodeTemplateTraitList[1], "FAnimationAnimNextRuntimeTest_TraitEventRaising -> Unexpected trait visited");
+				AddErrorIfFalse(EventA->VisitedTraits[2] == NodeTemplateTraitList[0], "FAnimationAnimNextRuntimeTest_TraitEventRaising -> Unexpected trait visited");
+
+				auto EventB = MakeTraitEvent<FTraitCoreTest_EventB>();
+				RaiseTraitEvent(Context, Stack0, *EventB);
+
+				UE_RETURN_ON_ERROR(EventB->VisitedTraits.Num() == 3, "FAnimationAnimNextRuntimeTest_TraitEventRaising -> Unexpected number of traits visited");
+				AddErrorIfFalse(EventB->VisitedTraits[0] == NodeTemplateTraitList[3], "FAnimationAnimNextRuntimeTest_TraitEventRaising -> Unexpected trait visited");
+				AddErrorIfFalse(EventB->VisitedTraits[1] == NodeTemplateTraitList[1], "FAnimationAnimNextRuntimeTest_TraitEventRaising -> Unexpected trait visited");
+				AddErrorIfFalse(EventB->VisitedTraits[2] == NodeTemplateTraitList[0], "FAnimationAnimNextRuntimeTest_TraitEventRaising -> Unexpected trait visited");
+			}
+
+			// Send events A and B in a list, make sure results are the same as test above
+			{
+				auto EventA = MakeTraitEvent<FTraitCoreTest_EventA>();
+				auto EventB = MakeTraitEvent<FTraitCoreTest_EventB>();
+
+				FTraitEventList EventList;
+				EventList.Push(EventA);
+				EventList.Push(EventB);
+
+				RaiseTraitEvents(Context, Stack0, EventList);
+
+				UE_RETURN_ON_ERROR(EventA->VisitedTraits.Num() == 3, "FAnimationAnimNextRuntimeTest_TraitEventRaising -> Unexpected number of traits visited");
+				AddErrorIfFalse(EventA->VisitedTraits[0] == NodeTemplateTraitList[2], "FAnimationAnimNextRuntimeTest_TraitEventRaising -> Unexpected trait visited");
+				AddErrorIfFalse(EventA->VisitedTraits[1] == NodeTemplateTraitList[1], "FAnimationAnimNextRuntimeTest_TraitEventRaising -> Unexpected trait visited");
+				AddErrorIfFalse(EventA->VisitedTraits[2] == NodeTemplateTraitList[0], "FAnimationAnimNextRuntimeTest_TraitEventRaising -> Unexpected trait visited");
+
+				UE_RETURN_ON_ERROR(EventB->VisitedTraits.Num() == 3, "FAnimationAnimNextRuntimeTest_TraitEventRaising -> Unexpected number of traits visited");
+				AddErrorIfFalse(EventB->VisitedTraits[0] == NodeTemplateTraitList[3], "FAnimationAnimNextRuntimeTest_TraitEventRaising -> Unexpected trait visited");
+				AddErrorIfFalse(EventB->VisitedTraits[1] == NodeTemplateTraitList[1], "FAnimationAnimNextRuntimeTest_TraitEventRaising -> Unexpected trait visited");
+				AddErrorIfFalse(EventB->VisitedTraits[2] == NodeTemplateTraitList[0], "FAnimationAnimNextRuntimeTest_TraitEventRaising -> Unexpected trait visited");
+			}
+
+			// Toggle AB add to block A, make sure B visits every trait but A doesn't visit the base
+			{
+				auto EventA = MakeTraitEvent<FTraitCoreTest_EventA>();
+				EventA->bAlwaysForwardToBase = false;
+				RaiseTraitEvent(Context, Stack0, *EventA);
+
+				UE_RETURN_ON_ERROR(EventA->VisitedTraits.Num() == 2, "FAnimationAnimNextRuntimeTest_TraitEventRaising -> Unexpected number of traits visited");
+				AddErrorIfFalse(EventA->VisitedTraits[0] == NodeTemplateTraitList[2], "FAnimationAnimNextRuntimeTest_TraitEventRaising -> Unexpected trait visited");
+				AddErrorIfFalse(EventA->VisitedTraits[1] == NodeTemplateTraitList[1], "FAnimationAnimNextRuntimeTest_TraitEventRaising -> Unexpected trait visited");
+
+				auto EventB = MakeTraitEvent<FTraitCoreTest_EventB>();
+				RaiseTraitEvent(Context, Stack0, *EventB);
+
+				UE_RETURN_ON_ERROR(EventB->VisitedTraits.Num() == 3, "FAnimationAnimNextRuntimeTest_TraitEventRaising -> Unexpected number of traits visited");
+				AddErrorIfFalse(EventB->VisitedTraits[0] == NodeTemplateTraitList[3], "FAnimationAnimNextRuntimeTest_TraitEventRaising -> Unexpected trait visited");
+				AddErrorIfFalse(EventB->VisitedTraits[1] == NodeTemplateTraitList[1], "FAnimationAnimNextRuntimeTest_TraitEventRaising -> Unexpected trait visited");
+				AddErrorIfFalse(EventB->VisitedTraits[2] == NodeTemplateTraitList[0], "FAnimationAnimNextRuntimeTest_TraitEventRaising -> Unexpected trait visited");
+			}
+
+			// Send events A and B in a list, make sure results are the same as test above
+			{
+				auto EventA = MakeTraitEvent<FTraitCoreTest_EventA>();
+				EventA->bAlwaysForwardToBase = false;
+				auto EventB = MakeTraitEvent<FTraitCoreTest_EventB>();
+
+				FTraitEventList EventList;
+				EventList.Push(EventA);
+				EventList.Push(EventB);
+
+				RaiseTraitEvents(Context, Stack0, EventList);
+
+				UE_RETURN_ON_ERROR(EventA->VisitedTraits.Num() == 2, "FAnimationAnimNextRuntimeTest_TraitEventRaising -> Unexpected number of traits visited");
+				AddErrorIfFalse(EventA->VisitedTraits[0] == NodeTemplateTraitList[2], "FAnimationAnimNextRuntimeTest_TraitEventRaising -> Unexpected trait visited");
+				AddErrorIfFalse(EventA->VisitedTraits[1] == NodeTemplateTraitList[1], "FAnimationAnimNextRuntimeTest_TraitEventRaising -> Unexpected trait visited");
+
+				UE_RETURN_ON_ERROR(EventB->VisitedTraits.Num() == 3, "FAnimationAnimNextRuntimeTest_TraitEventRaising -> Unexpected number of traits visited");
+				AddErrorIfFalse(EventB->VisitedTraits[0] == NodeTemplateTraitList[3], "FAnimationAnimNextRuntimeTest_TraitEventRaising -> Unexpected trait visited");
+				AddErrorIfFalse(EventB->VisitedTraits[1] == NodeTemplateTraitList[1], "FAnimationAnimNextRuntimeTest_TraitEventRaising -> Unexpected trait visited");
+				AddErrorIfFalse(EventB->VisitedTraits[2] == NodeTemplateTraitList[0], "FAnimationAnimNextRuntimeTest_TraitEventRaising -> Unexpected trait visited");
+			}
+
+			// Validate that invalid/consumed events are skipped
+			{
+				auto EventA = MakeTraitEvent<FTraitCoreTest_EventA>();
+				EventA->MarkConsumed();
+
+				RaiseTraitEvent(Context, Stack0, *EventA);
+
+				AddErrorIfFalse(EventA->VisitedTraits.Num() == 0, "FAnimationAnimNextRuntimeTest_TraitEventRaising -> Unexpected number of traits visited");
+
+				auto EventB = MakeTraitEvent<FTraitCoreTest_EventB>();
+
+				FTraitEventList EventList;
+				EventList.Push(EventA);
+				EventList.Push(EventB);
+
+				RaiseTraitEvents(Context, Stack0, EventList);
+
+				AddErrorIfFalse(EventA->VisitedTraits.Num() == 0, "FAnimationAnimNextRuntimeTest_TraitEventRaising -> Unexpected number of traits visited");
+
+				UE_RETURN_ON_ERROR(EventB->VisitedTraits.Num() == 3, "FAnimationAnimNextRuntimeTest_TraitEventRaising -> Unexpected number of traits visited");
+				AddErrorIfFalse(EventB->VisitedTraits[0] == NodeTemplateTraitList[3], "FAnimationAnimNextRuntimeTest_TraitEventRaising -> Unexpected trait visited");
+				AddErrorIfFalse(EventB->VisitedTraits[1] == NodeTemplateTraitList[1], "FAnimationAnimNextRuntimeTest_TraitEventRaising -> Unexpected trait visited");
+				AddErrorIfFalse(EventB->VisitedTraits[2] == NodeTemplateTraitList[0], "FAnimationAnimNextRuntimeTest_TraitEventRaising -> Unexpected trait visited");
+			}
+		}
+
+		Registry.Unregister(NodeTemplate0);
+
+		AddErrorIfFalse(Registry.GetNum() == 0, "FAnimationAnimNextRuntimeTest_TraitEventRaising -> Registry should contain 0 templates");
+	}
+
+	Tests::FUtils::CleanupAfterTests();
+
+	return true;
+}
+
+#endif
diff --git a/Engine/Plugins/Experimental/AnimNext/Source/AnimNextTestSuite/Private/AnimNextTraitEventTest.h b/Engine/Plugins/Experimental/AnimNext/Source/AnimNextTestSuite/Private/AnimNextTraitEventTest.h
new file mode 100644
--- /dev/null
+++ b/Engine/Plugins/Experimental/AnimNext/Source/AnimNextTestSuite/Private/AnimNextTraitEventTest.h
@@ -0,0 +1,30 @@
+// Copyright Epic Games, Inc. All Rights Reserved.
+
+#pragma once
+
+#include "CoreMinimal.h"
+
+#include "TraitCore/TraitUID.h"
+#include "TraitCore/TraitEvent.h"
+
+#include "AnimNextTraitEventTest.generated.h"
+
+USTRUCT()
+struct FTraitCoreTest_EventA : public FAnimNextTraitEvent
+{
+	GENERATED_BODY()
+	DECLARE_ANIM_TRAIT_EVENT(FTraitCoreTest_EventA)
+
+	bool bAlwaysForwardToBase = true;
+
+	TArray<UE::AnimNext::FTraitUID> VisitedTraits;
+};
+
+USTRUCT()
+struct FTraitCoreTest_EventB : public FAnimNextTraitEvent
+{
+	GENERATED_BODY()
+	DECLARE_ANIM_TRAIT_EVENT(FTraitCoreTest_EventB)
+
+	TArray<UE::AnimNext::FTraitUID> VisitedTraits;
+};
diff --git a/Engine/Plugins/Experimental/AnimNextAnimGraph/Source/AnimNextAnimGraph/Private/TraitInterfaces/IUpdate.cpp b/Engine/Plugins/Experimental/AnimNextAnimGraph/Source/AnimNextAnimGraph/Private/TraitInterfaces/IUpdate.cpp
--- a/Engine/Plugins/Experimental/AnimNextAnimGraph/Source/AnimNextAnimGraph/Private/TraitInterfaces/IUpdate.cpp
+++ b/Engine/Plugins/Experimental/AnimNextAnimGraph/Source/AnimNextAnimGraph/Private/TraitInterfaces/IUpdate.cpp
@@ -12,6 +12,15 @@ namespace UE::AnimNext
 {
 	namespace Private
 	{
+		struct FUpdateEventBookkeepingEntry;
+
+		struct FUpdateEventBookkeepingList
+		{
+			// We maintain a double linked list of event bookkeeping entries we execute before post-update
+			FUpdateEventBookkeepingEntry* EventBookkeepingHead = nullptr;
+			FUpdateEventBookkeepingEntry* EventBookkeepingTail = nullptr;
+		};
+
 		// This structure is transient and lives either on the stack or the memstack and its destructor may not be called
 		struct FUpdateEntry
 		{
@@ -31,6 +40,12 @@ namespace UE::AnimNext
 			// Once we've called PreUpdate, we cache the trait binding to avoid a redundant query to call PostUpdate
 			TTraitBinding<IUpdate>		UpdateTrait;
 
+			// A pointer to our parent entry or nullptr if we are the root
+			FUpdateEntry*				ParentEntry = nullptr;
+
+			// We maintain a double linked list of event bookkeeping entries we execute before post-update
+			FUpdateEventBookkeepingList EventBookkeepingList;
+
 			// These pointers are mutually exclusive
 			// An entry is either part of the queued update stack, the update stack, the free list, or none of the above
 			union
@@ -51,6 +66,94 @@ namespace UE::AnimNext
 			{
 			}
 		};
+
+		// The update traversal performs various event bookkeeping actions
+		enum class FUpdateEventBookkeepingAction : uint8
+		{
+			// Pushes an output trait event
+			PushOutput,
+
+			// Consumes a trait event
+			Consume,
+		};
+
+		// Encapsulates a bookkeeping entry that we'll execute before post-update
+		// These may be allocated on the memstack and their destructor might not run
+		struct FUpdateEventBookkeepingEntry
+		{
+			// The event the action manipulates
+			FAnimNextTraitEventPtr Event;
+
+			// The action to perform
+			// TODO: Could back a single bit in the LSB of one of the linked list pointers
+			FUpdateEventBookkeepingAction Action = FUpdateEventBookkeepingAction::Consume;
+
+			// Once allocated and bound to an update entry, bookkeeping entries form a double linked list
+			// where we append at the tail (next) to maintain queue ordering when we execute them
+			// When the bookkeeping entry isn't bound to an update entry, it lives in a free list where
+			// the next entry is the top of the free list
+			FUpdateEventBookkeepingEntry* NextEntry = nullptr;
+			FUpdateEventBookkeepingEntry* PrevEntry = nullptr;
+
+			// Creates a fresh entry
+			FUpdateEventBookkeepingEntry(FUpdateEventBookkeepingAction InAction, FAnimNextTraitEventPtr InEvent)
+				: Event(InEvent)
+				, Action(InAction)
+			{
+			}
+		};
+
+		// Queues the specified bookkeeping entry in the provided update entry
+		static void QueueBookkeepingEntry(FUpdateEventBookkeepingList& BookkeepingList, FUpdateEventBookkeepingEntry* BookkeepingEntry)
+		{
+			// Previous entry is the current tail (if any)
+			BookkeepingEntry->PrevEntry = BookkeepingList.EventBookkeepingTail;
+
+			if (BookkeepingList.EventBookkeepingHead == nullptr)
+			{
+				// This is the first bookkeeping entry, start our list
+				BookkeepingList.EventBookkeepingHead = BookkeepingEntry;
+			}
+			else
+			{
+				// Stitch the current tail with our new entry before we update it
+				BookkeepingList.EventBookkeepingTail->NextEntry = BookkeepingEntry;
+			}
+
+			// Append our entry at the tail
+			BookkeepingList.EventBookkeepingTail = BookkeepingEntry;
+		}
+
+		// Raises every queued event from the list on the specified entry
+		static void RaiseTraitEvents(FUpdateTraversalContext& Context, Private::FUpdateEntry* UpdateEntry, const UE::AnimNext::FTraitEventList& EventList)
+		{
+			// TODO: Performance note
+			// 
+			// Event lists are typically very small or empty and similarly most nodes handle few or no events
+			// They are thus a great fit to leverage bloom filters
+			// A node can pre-compute and cache one in its node template. This bloom filter contains all the event types it handles
+			// An event list can build a bloom filter of the event types it contains
+			// 
+			// Here (in this function), we could test if the event list bloom filter overlaps the node bloom filter: (node filter AND list filter) != 0
+			// If any bits intersect, then perhaps the node handles a type contained in the list (if the node handles nothing or if the list is empty, the result is always 0)
+			// With most nodes handling few events and the event list containing few events, we are likely to be able to skip many nodes with a very cheap test
+			// 
+			// Next, when we iterate over the event list, we can perform a similar test again for every event: (event filter AND node filter) == event filter
+			// If the event filter is contained in the node filter, the node might be handling the event type
+			// If not, then for sure it doesn't handle that event and we can avoid the virtual call and the event branching
+			// With most nodes handling few events and with an inclusion test, the rate of false positives is likely very low and we can skip most events
+			// avoiding the dispatch cost
+			// 
+			// To efficiently support this, we need to be able to build and cache bloom filters for each event in our list
+			// Caching the filter in the event itself is tricky as there is no good place to perform initialization work
+			// if we allow inheritance of events. Instead, the event list could store a struct with the event ptr and the bloom filter
+			// that we build when the event is inserted into the list. This would avoid the need to call a virtual function on the event
+			// to return a static constexpr filter. Similarly, we could store the event type UID alongside. We could pass this to the FTrait::OnTraitEvent
+			// function to avoid the virtual call we have for GetTypeUID.
+			// 
+
+			UE::AnimNext::RaiseTraitEvents(Context, UpdateEntry->TraitStack, EventList);
+		}
 	}
 
 	void IUpdate::PreUpdate(FUpdateTraversalContext& Context, const TTraitBinding<IUpdate>& Binding, const FTraitUpdateState& TraitState) const
@@ -81,7 +184,96 @@ namespace UE::AnimNext
 	//////////////////////////////////////////////////////////////////////////
 	// Traversal implementation
 
-	void FUpdateTraversalContext::PushQueuedUpdateEntries(FUpdateTraversalQueue& TraversalQueue)
+	void FUpdateTraversalContext::RaiseInputTraitEvent(FAnimNextTraitEventPtr Event)
+	{
+		if (!Event || !Event->IsValid())
+		{
+			return;
+		}
+
+		if (!ensureMsgf(ExecutingEntry == nullptr || !ExecutingEntry->bHasPreUpdated, TEXT("Input trait events can only be raised before a trait stack pre-updates")))
+		{
+			return;
+		}
+
+		if (ExecutingEntry != nullptr)
+		{
+			// If we are currently executing a node, we don't want the input event to be seen by our parent/siblings
+			// Add a bookkeeping entry to consume the event when we post-update
+			Private::FUpdateEventBookkeepingEntry* BookkeepingEntry = GetNewBookkeepingEntry(Private::FUpdateEventBookkeepingAction::Consume, Event);
+
+			Private::QueueBookkeepingEntry(ExecutingEntry->EventBookkeepingList, BookkeepingEntry);
+		}
+
+		InputEventList.Push(Event);
+	}
+
+	void FUpdateTraversalContext::RaiseOutputTraitEvent(FAnimNextTraitEventPtr Event)
+	{
+		if (!Event || !Event->IsValid())
+		{
+			return;
+		}
+
+		ensureMsgf(Event->IsTransient(), TEXT("Output trait events must have transient duration"));
+
+		if (ExecutingEntry != nullptr && ExecutingEntry->ParentEntry != nullptr)
+		{
+			// If we are currently executing a node, we don't want the output event to be visible on this node or its siblings
+			// Only its parent should see it
+			// Add a bookkeeping entry to push the event when our parent post-updates
+			// If we don't have a parent (e.g. root node), then we append to a fake parent list
+			Private::FUpdateEventBookkeepingEntry* BookkeepingEntry = GetNewBookkeepingEntry(Private::FUpdateEventBookkeepingAction::PushOutput, Event);
+
+			Private::FUpdateEventBookkeepingList* BookkeepingList = ExecutingEntry->ParentEntry != nullptr ? &ExecutingEntry->ParentEntry->EventBookkeepingList : RootParentBookkeepingEntryList;
+			Private::QueueBookkeepingEntry(*BookkeepingList, BookkeepingEntry);
+		}
+		else
+		{
+			// We aren't executing a trait stack or we are the root stack, just queue the output
+			// We might be in a component pre/post-update
+			OutputEventList.Push(Event);
+		}
+	}
+
+	void FUpdateTraversalContext::ExecuteBookkeepingActions(Private::FUpdateEventBookkeepingList& BookkeepingList)
+	{
+		if (BookkeepingList.EventBookkeepingHead == nullptr)
+		{
+			return;	// Nothing to do
+		}
+
+		// Iterate over our action list
+		Private::FUpdateEventBookkeepingEntry* BookkeepingEntry = BookkeepingList.EventBookkeepingHead;
+		while (BookkeepingEntry != nullptr)
+		{
+			switch (BookkeepingEntry->Action)
+			{
+			case Private::FUpdateEventBookkeepingAction::PushOutput:
+				OutputEventList.Push(BookkeepingEntry->Event);
+				break;
+			case Private::FUpdateEventBookkeepingAction::Consume:
+				BookkeepingEntry->Event->MarkConsumed();
+				break;
+			}
+
+			// Reset our pointer manually since the destructor won't run
+			BookkeepingEntry->Event.Reset();
+
+			Private::FUpdateEventBookkeepingEntry* NextEntry = BookkeepingEntry->NextEntry;
+
+			// Return our entry to the free list
+			PushFreeBookkeepingEntry(BookkeepingEntry);
+
+			// Continue iterating
+			BookkeepingEntry = NextEntry;
+		}
+
+		// Clear the list
+		BookkeepingList.EventBookkeepingHead = BookkeepingList.EventBookkeepingTail = nullptr;
+	}
+
+	void FUpdateTraversalContext::PushQueuedUpdateEntries(FUpdateTraversalQueue& TraversalQueue, Private::FUpdateEntry* ParentEntry)
 	{
 		// Pop every entry from the queued update stack and push them onto the update stack
 		// reversing their order
@@ -90,6 +282,8 @@ namespace UE::AnimNext
 			// Update our queued stack head
 			TraversalQueue.QueuedUpdateStackHead = Entry->PrevQueuedUpdateStackEntry;
 
+			Entry->ParentEntry = ParentEntry;
+
 			// Push our new entry onto the update stack
 			PushUpdateEntry(Entry);
 		}
@@ -131,6 +325,8 @@ namespace UE::AnimNext
 			FreeEntry->TraitState = TraitState;
 			FreeEntry->TraitPtr = TraitPtr;
 			FreeEntry->bHasPreUpdated = false;
+			FreeEntry->ParentEntry = nullptr;
+			FreeEntry->EventBookkeepingList = Private::FUpdateEventBookkeepingList();
 			FreeEntry->NextFreeEntry = nullptr;		// Mark it as not being a member of any list
 		}
 		else
@@ -142,6 +338,36 @@ namespace UE::AnimNext
 		return FreeEntry;
 	}
 
+	Private::FUpdateEventBookkeepingEntry* FUpdateTraversalContext::GetNewBookkeepingEntry(Private::FUpdateEventBookkeepingAction Action, FAnimNextTraitEventPtr Event)
+	{
+		Private::FUpdateEventBookkeepingEntry* FreeEntry = FreeBookkeepingEntryStackHead;
+		if (FreeEntry != nullptr)
+		{
+			// We have a free entry, set our new head
+			FreeBookkeepingEntryStackHead = FreeEntry->NextEntry;
+
+			// Update our entry
+			FreeEntry->Event = Event;
+			FreeEntry->Action = Action;
+			FreeEntry->NextEntry = nullptr;		// Mark it as not being a member of any list
+		}
+		else
+		{
+			// Allocate a new entry
+			FreeEntry = new(MemStack) Private::FUpdateEventBookkeepingEntry(Action, Event);
+		}
+
+		return FreeEntry;
+	}
+
+	void FUpdateTraversalContext::PushFreeBookkeepingEntry(Private::FUpdateEventBookkeepingEntry* Entry)
+	{
+		Entry->NextEntry = FreeBookkeepingEntryStackHead;
+		Entry->PrevEntry = nullptr;
+
+		FreeBookkeepingEntryStackHead = Entry;
+	}
+
 	FUpdateTraversalQueue::FUpdateTraversalQueue(FUpdateTraversalContext& InTraversalContext)
 		: TraversalContext(InTraversalContext)
 	{
@@ -233,10 +459,18 @@ namespace UE::AnimNext
 
 		FUpdateTraversalQueue TraversalQueue(TraversalContext);
 
+		Private::FUpdateEventBookkeepingList RootParentBookkeepingEntryList;
+		TraversalContext.RootParentBookkeepingEntryList = &RootParentBookkeepingEntryList;
+
+		// Grab our input events, we'll propagate them down the graph as we traverse
+		GraphInstance.CollectInputTraitEvents(TraversalContext.InputEventList);
+
 		// Before we start the traversal, we give the graph instance components the chance to do some work
 		TraversalContext.BindTo(GraphInstance);
 		for (auto It = TraversalContext.GetComponentIterator(); It; ++It)
 		{
+			RaiseTraitEvents(TraversalContext, *It.Value(), TraversalContext.InputEventList);
+
 			It.Value()->PreUpdate(TraversalContext);
 		}
 
@@ -246,6 +480,8 @@ namespace UE::AnimNext
 
 		while (Private::FUpdateEntry* Entry = TraversalContext.PopUpdateEntry())
 		{
+			TraversalContext.ExecutingEntry = Entry;
+
 			const FWeakTraitPtr& EntryTraitPtr = Entry->TraitPtr;
 
 			if (!Entry->bHasPreUpdated)
@@ -261,22 +497,20 @@ namespace UE::AnimNext
 				const bool bIsFrozen = false;	// Not yet supported
 				Entry->TraitStack.SnapshotLatentProperties(bIsFrozen);
 
+				// Raise our input events
+				Private::RaiseTraitEvents(TraversalContext, Entry, TraversalContext.InputEventList);
+
 				// If this trait stack implements IUpdate, call into it
 				if (Entry->TraitStack.GetInterface(Entry->UpdateTrait))
 				{
 					Entry->UpdateTrait.PreUpdate(TraversalContext, Entry->TraitState);
+				}
 
-					// Make sure that next time we visit this entry, we'll post-update
-					Entry->bHasPreUpdated = true;
+				// Make sure that next time we visit this entry, we'll post-update
+				Entry->bHasPreUpdated = true;
 
-					// Push this entry onto the update stack, we'll call it once all our children have finished executing
-					TraversalContext.PushUpdateEntry(Entry);
-				}
-				else
-				{
-					// We don't need this entry anymore
-					TraversalContext.PushFreeEntry(Entry);
-				}
+				// Push this entry onto the update stack, we'll call it once all our children have finished executing
+				TraversalContext.PushUpdateEntry(Entry);
 
 				// Now visit the trait stack and queue our children
 				ensure(Entry->TraitStack.GetTopTrait(TraitBinding));
@@ -296,7 +530,7 @@ namespace UE::AnimNext
 						// Iterate over our queued children and push them onto the update stack
 						// We do this to allow children to be queued in traversal order which is intuitive
 						// but to traverse them in that order, they must be pushed in reverse order onto the update stack
-						TraversalContext.PushQueuedUpdateEntries(TraversalQueue);
+						TraversalContext.PushQueuedUpdateEntries(TraversalQueue, Entry);
 					}
 					else if (TraitBinding.AsInterface(HierarchyTrait))
 					{
@@ -309,6 +543,7 @@ namespace UE::AnimNext
 							if (ChildPtr.IsValid())
 							{
 								Private::FUpdateEntry* ChildEntry = TraversalContext.GetNewEntry(ChildPtr, Entry->TraitState);
+								ChildEntry->ParentEntry = Entry;
 								TraversalContext.PushUpdateEntry(ChildEntry);
 							}
 						}
@@ -324,9 +559,17 @@ namespace UE::AnimNext
 			}
 			else
 			{
+				// Execute event bookkeeping actions
+				TraversalContext.ExecuteBookkeepingActions(Entry->EventBookkeepingList);
+
+				// Raise our output events
+				Private::RaiseTraitEvents(TraversalContext, Entry, TraversalContext.OutputEventList);
+
 				// We've already visited this node once, time to PostUpdate
-				check(Entry->UpdateTrait.IsValid());
-				Entry->UpdateTrait.PostUpdate(TraversalContext, Entry->TraitState);
+				if (Entry->UpdateTrait.IsValid())
+				{
+					Entry->UpdateTrait.PostUpdate(TraversalContext, Entry->TraitState);
+				}
 
 				// Now that it finished updating, we can pop any scoped interfaces this node might have pushed
 				TraversalContext.PopStackScopedInterfaces(Entry->TraitStack);
@@ -336,15 +579,30 @@ namespace UE::AnimNext
 			}
 		}
 
+		// Clear our executing entry
+		TraversalContext.ExecutingEntry = nullptr;
+
+		// Executing any bookkeeping our root node might need
+		TraversalContext.ExecuteBookkeepingActions(*TraversalContext.RootParentBookkeepingEntryList);
+
 		// After we finish the traversal, we give the graph instance components the chance to do some work
-		TraversalContext.BindTo(GraphInstance);
 		for (auto It = TraversalContext.GetComponentIterator(); It; ++It)
 		{
+			RaiseTraitEvents(TraversalContext, *It.Value(), TraversalContext.OutputEventList);
+
 			It.Value()->PostUpdate(TraversalContext);
 		}
 
+		// Decrement the remaining lifetime of the input events we processed and queue up any remaining events
+		TraversalContext.InputEventList.DecrementLifetime();
+		GraphInstance.QueueInputTraitEvents(TraversalContext.InputEventList);
+
 		// At this point, we shouldn't have any remaining scoped interfaces
 		// If this fails, it means we failed to pop them due to a push/pop mismatch
 		ensure(!TraversalContext.HasScopedInterfaces());
+
+		// TODO: Figure out what to do with the output events, we probably want to return them here and let the caller
+		// determine whether they should be ignored, filtered, or handled (e.g. post-physics vs pre-physics)
+		// The caller can determine what events can fire because it knows where it is running (e.g. worker vs main thread)
 	}
 }
diff --git a/Engine/Plugins/Experimental/AnimNextAnimGraph/Source/AnimNextAnimGraph/Private/Traits/BlendByBool.cpp b/Engine/Plugins/Experimental/AnimNextAnimGraph/Source/AnimNextAnimGraph/Private/Traits/BlendByBool.cpp
--- a/Engine/Plugins/Experimental/AnimNextAnimGraph/Source/AnimNextAnimGraph/Private/Traits/BlendByBool.cpp
+++ b/Engine/Plugins/Experimental/AnimNextAnimGraph/Source/AnimNextAnimGraph/Private/Traits/BlendByBool.cpp
@@ -17,7 +17,7 @@ namespace UE::AnimNext
 		GeneratorMacro(IUpdate) \
 		GeneratorMacro(IUpdateTraversal) \
 
-	GENERATE_ANIM_TRAIT_IMPLEMENTATION(FBlendByBoolTrait, TRAIT_INTERFACE_ENUMERATOR)
+	GENERATE_ANIM_TRAIT_IMPLEMENTATION(FBlendByBoolTrait, TRAIT_INTERFACE_ENUMERATOR, NULL_ANIM_TRAIT_EVENT_ENUMERATOR)
 	#undef TRAIT_INTERFACE_ENUMERATOR
 
 	static constexpr int32 TRUE_CHILD_INDEX = 0;
diff --git a/Engine/Plugins/Experimental/AnimNextAnimGraph/Source/AnimNextAnimGraph/Private/Traits/BlendInertializer.cpp b/Engine/Plugins/Experimental/AnimNextAnimGraph/Source/AnimNextAnimGraph/Private/Traits/BlendInertializer.cpp
--- a/Engine/Plugins/Experimental/AnimNextAnimGraph/Source/AnimNextAnimGraph/Private/Traits/BlendInertializer.cpp
+++ b/Engine/Plugins/Experimental/AnimNextAnimGraph/Source/AnimNextAnimGraph/Private/Traits/BlendInertializer.cpp
@@ -13,7 +13,7 @@ namespace UE::AnimNext
 		GeneratorMacro(IDiscreteBlend) \
 		GeneratorMacro(ISmoothBlend) \
 
-	GENERATE_ANIM_TRAIT_IMPLEMENTATION(FBlendInertializerTrait, TRAIT_INTERFACE_ENUMERATOR)
+	GENERATE_ANIM_TRAIT_IMPLEMENTATION(FBlendInertializerTrait, TRAIT_INTERFACE_ENUMERATOR, NULL_ANIM_TRAIT_EVENT_ENUMERATOR)
 	#undef TRAIT_INTERFACE_ENUMERATOR
 
 	void FBlendInertializerTrait::OnBlendTransition(const FExecutionContext& Context, const TTraitBinding<IDiscreteBlend>& Binding, int32 OldChildIndex, int32 NewChildIndex) const
diff --git a/Engine/Plugins/Experimental/AnimNextAnimGraph/Source/AnimNextAnimGraph/Private/Traits/BlendSmoother.cpp b/Engine/Plugins/Experimental/AnimNextAnimGraph/Source/AnimNextAnimGraph/Private/Traits/BlendSmoother.cpp
--- a/Engine/Plugins/Experimental/AnimNextAnimGraph/Source/AnimNextAnimGraph/Private/Traits/BlendSmoother.cpp
+++ b/Engine/Plugins/Experimental/AnimNextAnimGraph/Source/AnimNextAnimGraph/Private/Traits/BlendSmoother.cpp
@@ -19,7 +19,7 @@ namespace UE::AnimNext
 		GeneratorMacro(ISmoothBlend) \
 		GeneratorMacro(IUpdate) \
 
-	GENERATE_ANIM_TRAIT_IMPLEMENTATION(FBlendSmootherTrait, TRAIT_INTERFACE_ENUMERATOR)
+	GENERATE_ANIM_TRAIT_IMPLEMENTATION(FBlendSmootherTrait, TRAIT_INTERFACE_ENUMERATOR, NULL_ANIM_TRAIT_EVENT_ENUMERATOR)
 	#undef TRAIT_INTERFACE_ENUMERATOR
 
 	void FBlendSmootherTrait::PostEvaluate(FEvaluateTraversalContext& Context, const TTraitBinding<IEvaluate>& Binding) const
diff --git a/Engine/Plugins/Experimental/AnimNextAnimGraph/Source/AnimNextAnimGraph/Private/Traits/BlendSmootherPerBone.cpp b/Engine/Plugins/Experimental/AnimNextAnimGraph/Source/AnimNextAnimGraph/Private/Traits/BlendSmootherPerBone.cpp
--- a/Engine/Plugins/Experimental/AnimNextAnimGraph/Source/AnimNextAnimGraph/Private/Traits/BlendSmootherPerBone.cpp
+++ b/Engine/Plugins/Experimental/AnimNextAnimGraph/Source/AnimNextAnimGraph/Private/Traits/BlendSmootherPerBone.cpp
@@ -19,7 +19,7 @@ namespace UE::AnimNext
 		GeneratorMacro(IEvaluate) \
 		GeneratorMacro(IUpdate) \
 
-	GENERATE_ANIM_TRAIT_IMPLEMENTATION(FBlendSmootherPerBoneTrait, TRAIT_INTERFACE_ENUMERATOR)
+	GENERATE_ANIM_TRAIT_IMPLEMENTATION(FBlendSmootherPerBoneTrait, TRAIT_INTERFACE_ENUMERATOR, NULL_ANIM_TRAIT_EVENT_ENUMERATOR)
 	#undef TRAIT_INTERFACE_ENUMERATOR
 
 	void FBlendSmootherPerBoneTrait::PostEvaluate(FEvaluateTraversalContext& Context, const TTraitBinding<IEvaluate>& Binding) const
diff --git a/Engine/Plugins/Experimental/AnimNextAnimGraph/Source/AnimNextAnimGraph/Private/Traits/BlendTwoWay.cpp b/Engine/Plugins/Experimental/AnimNextAnimGraph/Source/AnimNextAnimGraph/Private/Traits/BlendTwoWay.cpp
--- a/Engine/Plugins/Experimental/AnimNextAnimGraph/Source/AnimNextAnimGraph/Private/Traits/BlendTwoWay.cpp
+++ b/Engine/Plugins/Experimental/AnimNextAnimGraph/Source/AnimNextAnimGraph/Private/Traits/BlendTwoWay.cpp
@@ -18,7 +18,7 @@ namespace UE::AnimNext
 		GeneratorMacro(IUpdate) \
 		GeneratorMacro(IUpdateTraversal) \
 
-	GENERATE_ANIM_TRAIT_IMPLEMENTATION(FBlendTwoWayTrait, TRAIT_INTERFACE_ENUMERATOR)
+	GENERATE_ANIM_TRAIT_IMPLEMENTATION(FBlendTwoWayTrait, TRAIT_INTERFACE_ENUMERATOR, NULL_ANIM_TRAIT_EVENT_ENUMERATOR)
 	#undef TRAIT_INTERFACE_ENUMERATOR
 
 	void FBlendTwoWayTrait::PostEvaluate(FEvaluateTraversalContext& Context, const TTraitBinding<IEvaluate>& Binding) const
diff --git a/Engine/Plugins/Experimental/AnimNextAnimGraph/Source/AnimNextAnimGraph/Private/Traits/ReferencePoseTrait.cpp b/Engine/Plugins/Experimental/AnimNextAnimGraph/Source/AnimNextAnimGraph/Private/Traits/ReferencePoseTrait.cpp
--- a/Engine/Plugins/Experimental/AnimNextAnimGraph/Source/AnimNextAnimGraph/Private/Traits/ReferencePoseTrait.cpp
+++ b/Engine/Plugins/Experimental/AnimNextAnimGraph/Source/AnimNextAnimGraph/Private/Traits/ReferencePoseTrait.cpp
@@ -13,7 +13,7 @@ namespace UE::AnimNext
 	#define TRAIT_INTERFACE_ENUMERATOR(GeneratorMacro) \
 		GeneratorMacro(IEvaluate) \
 
-	GENERATE_ANIM_TRAIT_IMPLEMENTATION(FReferencePoseTrait, TRAIT_INTERFACE_ENUMERATOR)
+	GENERATE_ANIM_TRAIT_IMPLEMENTATION(FReferencePoseTrait, TRAIT_INTERFACE_ENUMERATOR, NULL_ANIM_TRAIT_EVENT_ENUMERATOR)
 	#undef TRAIT_INTERFACE_ENUMERATOR
 
 	void FReferencePoseTrait::PreEvaluate(FEvaluateTraversalContext& Context, const TTraitBinding<IEvaluate>& Binding) const
diff --git a/Engine/Plugins/Experimental/AnimNextAnimGraph/Source/AnimNextAnimGraph/Private/Traits/SequencePlayer.cpp b/Engine/Plugins/Experimental/AnimNextAnimGraph/Source/AnimNextAnimGraph/Private/Traits/SequencePlayer.cpp
--- a/Engine/Plugins/Experimental/AnimNextAnimGraph/Source/AnimNextAnimGraph/Private/Traits/SequencePlayer.cpp
+++ b/Engine/Plugins/Experimental/AnimNextAnimGraph/Source/AnimNextAnimGraph/Private/Traits/SequencePlayer.cpp
@@ -16,7 +16,7 @@ namespace UE::AnimNext
 		GeneratorMacro(ITimeline) \
 		GeneratorMacro(IUpdate) \
 
-	GENERATE_ANIM_TRAIT_IMPLEMENTATION(FSequencePlayerTrait, TRAIT_INTERFACE_ENUMERATOR)
+	GENERATE_ANIM_TRAIT_IMPLEMENTATION(FSequencePlayerTrait, TRAIT_INTERFACE_ENUMERATOR, NULL_ANIM_TRAIT_EVENT_ENUMERATOR)
 	#undef TRAIT_INTERFACE_ENUMERATOR
 
 	void FSequencePlayerTrait::FInstanceData::Construct(const FExecutionContext& Context, const FTraitBinding& Binding)
diff --git a/Engine/Plugins/Experimental/AnimNextAnimGraph/Source/AnimNextAnimGraph/Private/Traits/SubGraphHost.cpp b/Engine/Plugins/Experimental/AnimNextAnimGraph/Source/AnimNextAnimGraph/Private/Traits/SubGraphHost.cpp
--- a/Engine/Plugins/Experimental/AnimNextAnimGraph/Source/AnimNextAnimGraph/Private/Traits/SubGraphHost.cpp
+++ b/Engine/Plugins/Experimental/AnimNextAnimGraph/Source/AnimNextAnimGraph/Private/Traits/SubGraphHost.cpp
@@ -18,7 +18,7 @@ namespace UE::AnimNext
 		GeneratorMacro(IUpdate) \
 		GeneratorMacro(IUpdateTraversal) \
 
-	GENERATE_ANIM_TRAIT_IMPLEMENTATION(FSubGraphHostTrait, TRAIT_INTERFACE_ENUMERATOR)
+	GENERATE_ANIM_TRAIT_IMPLEMENTATION(FSubGraphHostTrait, TRAIT_INTERFACE_ENUMERATOR, NULL_ANIM_TRAIT_EVENT_ENUMERATOR)
 	#undef TRAIT_INTERFACE_ENUMERATOR
 
 	void FSubGraphHostTrait::FInstanceData::Construct(const FExecutionContext& Context, const FTraitBinding& Binding)
diff --git a/Engine/Plugins/Experimental/AnimNextAnimGraph/Source/AnimNextAnimGraph/Private/Traits/SynchronizeUsingGroups.cpp b/Engine/Plugins/Experimental/AnimNextAnimGraph/Source/AnimNextAnimGraph/Private/Traits/SynchronizeUsingGroups.cpp
--- a/Engine/Plugins/Experimental/AnimNextAnimGraph/Source/AnimNextAnimGraph/Private/Traits/SynchronizeUsingGroups.cpp
+++ b/Engine/Plugins/Experimental/AnimNextAnimGraph/Source/AnimNextAnimGraph/Private/Traits/SynchronizeUsingGroups.cpp
@@ -13,7 +13,7 @@ namespace UE::AnimNext
 		GeneratorMacro(ITimeline) \
 		GeneratorMacro(IUpdate) \
 
-	GENERATE_ANIM_TRAIT_IMPLEMENTATION(FSynchronizeUsingGroupsTrait, TRAIT_INTERFACE_ENUMERATOR)
+	GENERATE_ANIM_TRAIT_IMPLEMENTATION(FSynchronizeUsingGroupsTrait, TRAIT_INTERFACE_ENUMERATOR, NULL_ANIM_TRAIT_EVENT_ENUMERATOR)
 	#undef TRAIT_INTERFACE_ENUMERATOR
 
 	void FSynchronizeUsingGroupsTrait::PreUpdate(FUpdateTraversalContext& Context, const TTraitBinding<IUpdate>& Binding, const FTraitUpdateState& TraitState) const
diff --git a/Engine/Plugins/Experimental/AnimNextAnimGraph/Source/AnimNextAnimGraph/Public/TraitInterfaces/IUpdate.h b/Engine/Plugins/Experimental/AnimNextAnimGraph/Source/AnimNextAnimGraph/Public/TraitInterfaces/IUpdate.h
--- a/Engine/Plugins/Experimental/AnimNextAnimGraph/Source/AnimNextAnimGraph/Public/TraitInterfaces/IUpdate.h
+++ b/Engine/Plugins/Experimental/AnimNextAnimGraph/Source/AnimNextAnimGraph/Public/TraitInterfaces/IUpdate.h
@@ -6,6 +6,7 @@
 #include "TraitCore/ExecutionContext.h"
 #include "TraitCore/ITraitInterface.h"
 #include "TraitCore/TraitBinding.h"
+#include "TraitCore/TraitEventList.h"
 
 #include <type_traits>
 
@@ -19,6 +20,9 @@ namespace UE::AnimNext
 	namespace Private
 	{
 		struct FUpdateEntry;
+		struct FUpdateEventBookkeepingList;
+		struct FUpdateEventBookkeepingEntry;
+		enum class FUpdateEventBookkeepingAction : uint8;
 	}
 
 	/**
@@ -103,12 +107,21 @@ namespace UE::AnimNext
 	 */
 	struct FUpdateTraversalContext final : FExecutionContext
 	{
+		// Raises an input trait event
+		// The raised event will be seen by every child of the currently updating trait stack
+		// Input trait events can only be raised during PreUpdate
+		virtual void RaiseInputTraitEvent(FAnimNextTraitEventPtr Event) override;
+
+		// Raises an output trait event
+		// The raised event will be seen by every parent of the currently updating trait stack
+		virtual void RaiseOutputTraitEvent(FAnimNextTraitEventPtr Event) override;
+
 	private:
 		// Constructs a new traversal context
 		FUpdateTraversalContext() = default;
 
 		// Pops entries from the traversal queue and pushes them onto the update stack
-		void PushQueuedUpdateEntries(FUpdateTraversalQueue& TraversalQueue);
+		void PushQueuedUpdateEntries(FUpdateTraversalQueue& TraversalQueue, Private::FUpdateEntry* ParentEntry);
 
 		// Pushes an entry onto the update stack
 		void PushUpdateEntry(Private::FUpdateEntry* Entry);
@@ -123,6 +136,23 @@ namespace UE::AnimNext
 		// If an entry isn't found in the free stack, a new one is allocated from the memstack
 		Private::FUpdateEntry* GetNewEntry(const FWeakTraitPtr& TraitPtr, const FTraitUpdateState& TraitState);
 
+		// Returns a new bookkeeping entry suitable for queuing
+		// If an entry isn't found in the free stack, a new one is allocated from the memstack
+		Private::FUpdateEventBookkeepingEntry* GetNewBookkeepingEntry(Private::FUpdateEventBookkeepingAction Action, FAnimNextTraitEventPtr Event);
+
+		// Pushes a bookkeeping entry onto the free bookkeeping entry stack
+		void PushFreeBookkeepingEntry(Private::FUpdateEventBookkeepingEntry* Entry);
+
+		// Executes the event bookkeeping actions and clears the bookkeeping list
+		void ExecuteBookkeepingActions(Private::FUpdateEventBookkeepingList& BookkeepingList);
+
+		// The input and output event lists
+		UE::AnimNext::FTraitEventList InputEventList;
+		UE::AnimNext::FTraitEventList OutputEventList;
+
+		// The currently executing entry
+		Private::FUpdateEntry* ExecutingEntry = nullptr;
+
 		// The head pointer of the update stack
 		// This is the traversal execution stack and it contains entries that are
 		// pending their pre-update call and entries waiting for post-update to be called
@@ -133,6 +163,14 @@ namespace UE::AnimNext
 		// be warmer in the CPU cache
 		Private::FUpdateEntry* FreeEntryStackHead = nullptr;
 
+		// The head pointer of the free bookkeeping entry stack
+		// Entries are allocated from the memstack and are re-used in LIFO since they'll
+		// be warmer in the CPU cache
+		Private::FUpdateEventBookkeepingEntry* FreeBookkeepingEntryStackHead = nullptr;
+
+		// The root node doesn't have a parent but we need a bookkeeping list regardless
+		Private::FUpdateEventBookkeepingList* RootParentBookkeepingEntryList = nullptr;
+
 		friend ANIMNEXTANIMGRAPH_API void UpdateGraph(FAnimNextGraphInstancePtr& GraphInstance, float DeltaTime);
 		friend FUpdateTraversalQueue;
 	};
diff --git a/Engine/Plugins/Experimental/AnimNextAnimGraph/Source/AnimNextAnimGraphTestSuite/Private/AnimNextAnimGraphTraitEventTest.cpp b/Engine/Plugins/Experimental/AnimNextAnimGraph/Source/AnimNextAnimGraphTestSuite/Private/AnimNextAnimGraphTraitEventTest.cpp
new file mode 100644
--- /dev/null
+++ b/Engine/Plugins/Experimental/AnimNextAnimGraph/Source/AnimNextAnimGraphTestSuite/Private/AnimNextAnimGraphTraitEventTest.cpp
@@ -0,0 +1,423 @@
+// Copyright Epic Games, Inc. All Rights Reserved.
+
+#include "AnimNextAnimGraphTraitEventTest.h"
+
+#include "AnimNextRuntimeTest.h"
+#include "AnimNextTest.h"
+#include "AnimNextAnimGraphTraitInterfacesTest.h"
+#include "TraitCore/Trait.h"
+#include "TraitCore/TraitBinding.h"
+#include "TraitCore/TraitReader.h"
+#include "TraitCore/TraitWriter.h"
+#include "TraitCore/ExecutionContext.h"
+#include "TraitCore/NodeTemplateBuilder.h"
+#include "TraitCore/NodeTemplateRegistry.h"
+#include "TraitInterfaces/IHierarchy.h"
+#include "TraitInterfaces/IUpdate.h"
+#include "Graph/AnimNextGraph.h"
+#include "Graph/GraphFactory.h"
+#include "Misc/AutomationTest.h"
+
+#if WITH_DEV_AUTOMATION_TESTS
+
+//****************************************************************************
+// AnimNext Runtime Trait Graph Event Tests
+//****************************************************************************
+
+namespace UE::AnimNext
+{
+	struct FTraitGraphTest_EventAB_NoChildren : FBaseTrait
+	{
+		DECLARE_ANIM_TRAIT(FTraitGraphTest_EventAB_NoChildren, 0xfe325e9e, FBaseTrait)
+
+		ETraitStackPropagation OnEventA(const FExecutionContext& Context, FTraitBinding& Binding, FTraitAnimGraphTest_EventA& Event) const
+		{
+			Event.VisitedTraits.Add(FTraitGraphTest_EventAB_NoChildren::TraitUID);
+			return ETraitStackPropagation::Continue;
+		}
+
+		ETraitStackPropagation OnEventB(const FExecutionContext& Context, FTraitBinding& Binding, FTraitAnimGraphTest_EventB& Event) const
+		{
+			Event.VisitedTraits.Add(FTraitGraphTest_EventAB_NoChildren::TraitUID);
+			return ETraitStackPropagation::Continue;
+		}
+	};
+
+	// Trait implementation boilerplate
+	#define TRAIT_EVENT_ENUMERATOR(GeneratorMacro) \
+		GeneratorMacro(FTraitGraphTest_EventAB_NoChildren::OnEventA) \
+		GeneratorMacro(FTraitGraphTest_EventAB_NoChildren::OnEventB) \
+
+	GENERATE_ANIM_TRAIT_IMPLEMENTATION(FTraitGraphTest_EventAB_NoChildren, NULL_ANIM_TRAIT_INTERFACE_ENUMERATOR, TRAIT_EVENT_ENUMERATOR)
+	#undef TRAIT_EVENT_ENUMERATOR
+
+	struct FTraitGraphTest_EventAB_OneChild : FBaseTrait, IHierarchy
+	{
+		DECLARE_ANIM_TRAIT(FTraitGraphTest_EventAB_OneChild, 0xfc122a80, FBaseTrait)
+
+		using FSharedData = FTraitWithOneChildSharedData;
+
+		struct FInstanceData : FTrait::FInstanceData
+		{
+			FTraitPtr Child;
+
+			void Construct(const FExecutionContext& Context, const FTraitBinding& Binding)
+			{
+				Child = Context.AllocateNodeInstance(Binding.GetTraitPtr(), Binding.GetSharedData<FSharedData>()->Child);
+			}
+		};
+
+		ETraitStackPropagation OnEventA(const FExecutionContext& Context, FTraitBinding& Binding, FTraitAnimGraphTest_EventA& Event) const
+		{
+			if (Event.bTestFlag)
+			{
+				Event.MarkConsumed();
+			}
+
+			Event.VisitedTraits.Add(FTraitGraphTest_EventAB_OneChild::TraitUID);
+			return ETraitStackPropagation::Continue;
+		}
+
+		ETraitStackPropagation OnEventB(FExecutionContext& Context, FTraitBinding& Binding, FTraitAnimGraphTest_EventB& Event) const
+		{
+			if (Event.bTestFlag0)
+			{
+				auto EventA = MakeTraitEvent<FTraitAnimGraphTest_EventA>();
+				Event.ChildEvent = EventA;
+
+				Context.RaiseInputTraitEvent(EventA);
+			}
+
+			if (Event.bTestFlag1)
+			{
+				auto EventA = MakeTraitEvent<FTraitAnimGraphTest_EventA>();
+				Event.ChildEvent = EventA;
+
+				Context.RaiseOutputTraitEvent(EventA);
+			}
+
+			Event.VisitedTraits.Add(FTraitGraphTest_EventAB_OneChild::TraitUID);
+			return ETraitStackPropagation::Continue;
+		}
+
+		// IHierarchy impl
+		virtual uint32 GetNumChildren(const FExecutionContext& Context, const TTraitBinding<IHierarchy>& Binding) const override
+		{
+			return 1;
+		}
+
+		virtual void GetChildren(const FExecutionContext& Context, const TTraitBinding<IHierarchy>& Binding, FChildrenArray& Children) const override
+		{
+			const FInstanceData* InstanceData = Binding.GetInstanceData<FInstanceData>();
+
+			Children.Add(InstanceData->Child);
+		}
+	};
+
+	// Trait implementation boilerplate
+	#define TRAIT_INTERFACE_ENUMERATOR(GeneratorMacro) \
+		GeneratorMacro(IHierarchy) \
+
+	#define TRAIT_EVENT_ENUMERATOR(GeneratorMacro) \
+		GeneratorMacro(FTraitGraphTest_EventAB_OneChild::OnEventA) \
+		GeneratorMacro(FTraitGraphTest_EventAB_OneChild::OnEventB) \
+
+	GENERATE_ANIM_TRAIT_IMPLEMENTATION(FTraitGraphTest_EventAB_OneChild, TRAIT_INTERFACE_ENUMERATOR, TRAIT_EVENT_ENUMERATOR)
+	#undef TRAIT_INTERFACE_ENUMERATOR
+	#undef TRAIT_EVENT_ENUMERATOR
+
+	struct FTraitGraphTest_EventAB_TwoChildren : FBaseTrait, IHierarchy
+	{
+		DECLARE_ANIM_TRAIT(FTraitGraphTest_EventAB_TwoChildren, 0xa440e33f, FBaseTrait)
+
+		using FSharedData = FTraitWithChildrenSharedData;
+
+		struct FInstanceData : FTrait::FInstanceData
+		{
+			FTraitPtr Children[2];
+
+			void Construct(const FExecutionContext& Context, const FTraitBinding& Binding)
+			{
+				Children[0] = Context.AllocateNodeInstance(Binding.GetTraitPtr(), Binding.GetSharedData<FSharedData>()->Children[0]);
+				Children[1] = Context.AllocateNodeInstance(Binding.GetTraitPtr(), Binding.GetSharedData<FSharedData>()->Children[1]);
+			}
+		};
+
+		ETraitStackPropagation OnEventA(const FExecutionContext& Context, FTraitBinding& Binding, FTraitAnimGraphTest_EventA& Event) const
+		{
+			Event.VisitedTraits.Add(FTraitGraphTest_EventAB_TwoChildren::TraitUID);
+			return ETraitStackPropagation::Continue;
+		}
+
+		ETraitStackPropagation OnEventB(const FExecutionContext& Context, FTraitBinding& Binding, FTraitAnimGraphTest_EventB& Event) const
+		{
+			Event.VisitedTraits.Add(FTraitGraphTest_EventAB_TwoChildren::TraitUID);
+			return ETraitStackPropagation::Continue;
+		}
+
+		// IHierarchy impl
+		virtual uint32 GetNumChildren(const FExecutionContext& Context, const TTraitBinding<IHierarchy>& Binding) const override
+		{
+			return 2;
+		}
+
+		virtual void GetChildren(const FExecutionContext& Context, const TTraitBinding<IHierarchy>& Binding, FChildrenArray& Children) const override
+		{
+			const FInstanceData* InstanceData = Binding.GetInstanceData<FInstanceData>();
+
+			Children.Add(InstanceData->Children[0]);
+			Children.Add(InstanceData->Children[1]);
+		}
+	};
+
+	// Trait implementation boilerplate
+	#define TRAIT_INTERFACE_ENUMERATOR(GeneratorMacro) \
+		GeneratorMacro(IHierarchy) \
+
+	#define TRAIT_EVENT_ENUMERATOR(GeneratorMacro) \
+		GeneratorMacro(FTraitGraphTest_EventAB_TwoChildren::OnEventA) \
+		GeneratorMacro(FTraitGraphTest_EventAB_TwoChildren::OnEventB) \
+
+	GENERATE_ANIM_TRAIT_IMPLEMENTATION(FTraitGraphTest_EventAB_TwoChildren, TRAIT_INTERFACE_ENUMERATOR, TRAIT_EVENT_ENUMERATOR)
+	#undef TRAIT_INTERFACE_ENUMERATOR
+	#undef TRAIT_EVENT_ENUMERATOR
+}
+
+IMPLEMENT_SIMPLE_AUTOMATION_TEST(FAnimationAnimNextRuntimeTest_GraphTraitEvent, "Animation.AnimNext.Runtime.Graph.TraitEvent", EAutomationTestFlags::EditorContext | EAutomationTestFlags::EngineFilter)
+
+bool FAnimationAnimNextRuntimeTest_GraphTraitEvent::RunTest(const FString& InParameters)
+{
+	using namespace UE::AnimNext;
+
+	{
+		AUTO_REGISTER_ANIM_TRAIT(FTraitGraphTest_EventAB_NoChildren)
+		AUTO_REGISTER_ANIM_TRAIT(FTraitGraphTest_EventAB_OneChild)
+		AUTO_REGISTER_ANIM_TRAIT(FTraitGraphTest_EventAB_TwoChildren)
+
+		UFactory* GraphFactory = NewObject<UAnimNextGraphFactory>();
+		UAnimNextGraph* AnimNextGraph = CastChecked<UAnimNextGraph>(GraphFactory->FactoryCreateNew(UAnimNextGraph::StaticClass(), GetTransientPackage(), TEXT("TestAnimNextGraph"), RF_Transient, nullptr, nullptr, NAME_None));
+		UE_RETURN_ON_ERROR(AnimNextGraph != nullptr, "FAnimationAnimNextRuntimeTest_GraphTraitEvent -> Failed to create animation graph");
+
+		FScopedClearNodeTemplateRegistry ScopedClearNodeTemplateRegistry;
+		FNodeTemplateRegistry& Registry = FNodeTemplateRegistry::Get();
+
+		// Graph: Root -> NodeA -> NodeB
+		//                NodeC
+
+		// We create a few node templates
+		// Template A has a single child node (NodeA)
+		TArray<FTraitUID> NodeTemplateTraitListA;
+		NodeTemplateTraitListA.Add(FTraitGraphTest_EventAB_OneChild::TraitUID);
+
+		// Template B has two children (Root)
+		TArray<FTraitUID> NodeTemplateTraitListB;
+		NodeTemplateTraitListB.Add(FTraitGraphTest_EventAB_TwoChildren::TraitUID);
+
+		// Template C has no children (NodeB and NodeC)
+		TArray<FTraitUID> NodeTemplateTraitListC;
+		NodeTemplateTraitListC.Add(FTraitGraphTest_EventAB_NoChildren::TraitUID);
+
+		// Populate our node template registry
+		TArray<uint8> NodeTemplateBufferA, NodeTemplateBufferB, NodeTemplateBufferC;
+		const FNodeTemplate* NodeTemplateA = FNodeTemplateBuilder::BuildNodeTemplate(NodeTemplateTraitListA, NodeTemplateBufferA);
+		const FNodeTemplate* NodeTemplateB = FNodeTemplateBuilder::BuildNodeTemplate(NodeTemplateTraitListB, NodeTemplateBufferB);
+		const FNodeTemplate* NodeTemplateC = FNodeTemplateBuilder::BuildNodeTemplate(NodeTemplateTraitListC, NodeTemplateBufferC);
+
+		TArray<FNodeHandle> NodeHandles;
+
+		// Write our graph
+		TArray<uint8> GraphSharedDataArchiveBuffer;
+		TArray<TObjectPtr<UObject>> GraphReferencedObjects;
+		{
+			FTraitWriter TraitWriter;
+
+			NodeHandles.Add(TraitWriter.RegisterNode(*NodeTemplateB));	// Root
+			NodeHandles.Add(TraitWriter.RegisterNode(*NodeTemplateA));	// NodeA
+			NodeHandles.Add(TraitWriter.RegisterNode(*NodeTemplateC));	// NodeB
+			NodeHandles.Add(TraitWriter.RegisterNode(*NodeTemplateC));	// NodeC
+
+			TArray<TMap<FName, FString>> TraitPropertiesRoot;
+			TraitPropertiesRoot.AddDefaulted(NodeTemplateTraitListB.Num());
+			FAnimNextTraitHandle ChildrenHandlesRoot[2] = { FAnimNextTraitHandle(NodeHandles[1]), FAnimNextTraitHandle(NodeHandles[3]) };
+			TraitPropertiesRoot[0].Add(TEXT("Children"), ToString<FTraitGraphTest_EventAB_TwoChildren::FSharedData>(TEXT("Children"), ChildrenHandlesRoot));
+
+			TArray<TMap<FName, FString>> TraitPropertiesNodeA;
+			TraitPropertiesNodeA.AddDefaulted(NodeTemplateTraitListA.Num());
+			TraitPropertiesNodeA[0].Add(TEXT("Child"), ToString<FTraitGraphTest_EventAB_OneChild::FSharedData>(TEXT("Child"), FAnimNextTraitHandle(NodeHandles[2])));
+
+			TArray<TMap<FName, FString>> TraitPropertiesNodeB;
+			TraitPropertiesNodeB.AddDefaulted(NodeTemplateTraitListC.Num());
+
+			TArray<TMap<FName, FString>> TraitPropertiesNodeC;
+			TraitPropertiesNodeC.AddDefaulted(NodeTemplateTraitListC.Num());
+
+			TraitWriter.BeginNodeWriting();
+			TraitWriter.WriteNode(NodeHandles[0],
+				[&TraitPropertiesRoot](uint32 TraitIndex, FName PropertyName)
+				{
+					return TraitPropertiesRoot[TraitIndex][PropertyName];
+				},
+				[](uint32 TraitIndex, FName PropertyName)
+				{
+					return MAX_uint16;
+				});
+			TraitWriter.WriteNode(NodeHandles[1],
+				[&TraitPropertiesNodeA](uint32 TraitIndex, FName PropertyName)
+				{
+					return TraitPropertiesNodeA[TraitIndex][PropertyName];
+				},
+				[](uint32 TraitIndex, FName PropertyName)
+				{
+					return MAX_uint16;
+				});
+			TraitWriter.WriteNode(NodeHandles[2],
+				[&TraitPropertiesNodeB](uint32 TraitIndex, FName PropertyName)
+				{
+					return TraitPropertiesNodeB[TraitIndex][PropertyName];
+				},
+				[](uint32 TraitIndex, FName PropertyName)
+				{
+					return MAX_uint16;
+				});
+			TraitWriter.WriteNode(NodeHandles[3],
+				[&TraitPropertiesNodeC](uint32 TraitIndex, FName PropertyName)
+				{
+					return TraitPropertiesNodeC[TraitIndex][PropertyName];
+				},
+				[](uint32 TraitIndex, FName PropertyName)
+				{
+					return MAX_uint16;
+				});
+			TraitWriter.EndNodeWriting();
+
+			AddErrorIfFalse(TraitWriter.GetErrorState() == FTraitWriter::EErrorState::None, "FAnimationAnimNextRuntimeTest_GraphTraitEvent -> Failed to write traits");
+			GraphSharedDataArchiveBuffer = TraitWriter.GetGraphSharedData();
+			GraphReferencedObjects = TraitWriter.GetGraphReferencedObjects();
+		}
+
+		// Read our graph
+		FTestUtils::LoadFromArchiveBuffer(*AnimNextGraph, NodeHandles, GraphSharedDataArchiveBuffer);
+
+		FAnimNextGraphInstancePtr GraphInstance;
+		AnimNextGraph->AllocateInstance(GraphInstance);
+
+		{
+			// Raise EventA and EventB on graph, every node sees them
+			{
+				auto EventA = MakeTraitEvent<FTraitAnimGraphTest_EventA>();
+				auto EventB = MakeTraitEvent<FTraitAnimGraphTest_EventB>();
+
+				GraphInstance.QueueInputTraitEvent(EventA);
+				GraphInstance.QueueInputTraitEvent(EventB);
+
+				UpdateGraph(GraphInstance, 0.0333f);
+
+				UE_RETURN_ON_ERROR(EventA->VisitedTraits.Num() == 4, "FAnimationAnimNextRuntimeTest_GraphTraitEvent -> Unexpected number of traits visited");
+				AddErrorIfFalse(EventA->VisitedTraits[0] == NodeTemplateTraitListB[0], "FAnimationAnimNextRuntimeTest_GraphTraitEvent -> Unexpected trait visited");
+				AddErrorIfFalse(EventA->VisitedTraits[1] == NodeTemplateTraitListA[0], "FAnimationAnimNextRuntimeTest_GraphTraitEvent -> Unexpected trait visited");
+				AddErrorIfFalse(EventA->VisitedTraits[2] == NodeTemplateTraitListC[0], "FAnimationAnimNextRuntimeTest_GraphTraitEvent -> Unexpected trait visited");
+				AddErrorIfFalse(EventA->VisitedTraits[3] == NodeTemplateTraitListC[0], "FAnimationAnimNextRuntimeTest_GraphTraitEvent -> Unexpected trait visited");
+
+				UE_RETURN_ON_ERROR(EventB->VisitedTraits.Num() == 4, "FAnimationAnimNextRuntimeTest_GraphTraitEvent -> Unexpected number of traits visited");
+				AddErrorIfFalse(EventB->VisitedTraits[0] == NodeTemplateTraitListB[0], "FAnimationAnimNextRuntimeTest_GraphTraitEvent -> Unexpected trait visited");
+				AddErrorIfFalse(EventB->VisitedTraits[1] == NodeTemplateTraitListA[0], "FAnimationAnimNextRuntimeTest_GraphTraitEvent -> Unexpected trait visited");
+				AddErrorIfFalse(EventB->VisitedTraits[2] == NodeTemplateTraitListC[0], "FAnimationAnimNextRuntimeTest_GraphTraitEvent -> Unexpected trait visited");
+				AddErrorIfFalse(EventB->VisitedTraits[3] == NodeTemplateTraitListC[0], "FAnimationAnimNextRuntimeTest_GraphTraitEvent -> Unexpected trait visited");
+			}
+
+			// Raise EventA on graph, NodeA consumes it, only Root and NodeA see it
+			{
+				auto EventA = MakeTraitEvent<FTraitAnimGraphTest_EventA>();
+				EventA->bTestFlag = true;
+
+				GraphInstance.QueueInputTraitEvent(EventA);
+
+				UpdateGraph(GraphInstance, 0.0333f);
+
+				UE_RETURN_ON_ERROR(EventA->VisitedTraits.Num() == 2, "FAnimationAnimNextRuntimeTest_GraphTraitEvent -> Unexpected number of traits visited");
+				AddErrorIfFalse(EventA->VisitedTraits[0] == NodeTemplateTraitListB[0], "FAnimationAnimNextRuntimeTest_GraphTraitEvent -> Unexpected trait visited");
+				AddErrorIfFalse(EventA->VisitedTraits[1] == NodeTemplateTraitListA[0], "FAnimationAnimNextRuntimeTest_GraphTraitEvent -> Unexpected trait visited");
+			}
+
+			// Raise EventA and EventB on graph (in list), NodeA consumes EventA, only Root and NodeA see it, every node sees EventB
+			{
+				auto EventA = MakeTraitEvent<FTraitAnimGraphTest_EventA>();
+				EventA->bTestFlag = true;
+				auto EventB = MakeTraitEvent<FTraitAnimGraphTest_EventB>();
+
+				GraphInstance.QueueInputTraitEvent(EventA);
+				GraphInstance.QueueInputTraitEvent(EventB);
+
+				UpdateGraph(GraphInstance, 0.0333f);
+
+				UE_RETURN_ON_ERROR(EventA->VisitedTraits.Num() == 2, "FAnimationAnimNextRuntimeTest_GraphTraitEvent -> Unexpected number of traits visited");
+				AddErrorIfFalse(EventA->VisitedTraits[0] == NodeTemplateTraitListB[0], "FAnimationAnimNextRuntimeTest_GraphTraitEvent -> Unexpected trait visited");
+				AddErrorIfFalse(EventA->VisitedTraits[1] == NodeTemplateTraitListA[0], "FAnimationAnimNextRuntimeTest_GraphTraitEvent -> Unexpected trait visited");
+
+				UE_RETURN_ON_ERROR(EventB->VisitedTraits.Num() == 4, "FAnimationAnimNextRuntimeTest_GraphTraitEvent -> Unexpected number of traits visited");
+				AddErrorIfFalse(EventB->VisitedTraits[0] == NodeTemplateTraitListB[0], "FAnimationAnimNextRuntimeTest_GraphTraitEvent -> Unexpected trait visited");
+				AddErrorIfFalse(EventB->VisitedTraits[1] == NodeTemplateTraitListA[0], "FAnimationAnimNextRuntimeTest_GraphTraitEvent -> Unexpected trait visited");
+				AddErrorIfFalse(EventB->VisitedTraits[2] == NodeTemplateTraitListC[0], "FAnimationAnimNextRuntimeTest_GraphTraitEvent -> Unexpected trait visited");
+				AddErrorIfFalse(EventB->VisitedTraits[3] == NodeTemplateTraitListC[0], "FAnimationAnimNextRuntimeTest_GraphTraitEvent -> Unexpected trait visited");
+			}
+
+			// NodeA raises input EventA, only visible to NodeB
+			{
+				auto EventB = MakeTraitEvent<FTraitAnimGraphTest_EventB>();
+				EventB->bTestFlag0 = true;
+
+				GraphInstance.QueueInputTraitEvent(EventB);
+
+				UpdateGraph(GraphInstance, 0.0333f);
+
+				UE_RETURN_ON_ERROR(EventB->VisitedTraits.Num() == 4, "FAnimationAnimNextRuntimeTest_GraphTraitEvent -> Unexpected number of traits visited");
+				AddErrorIfFalse(EventB->VisitedTraits[0] == NodeTemplateTraitListB[0], "FAnimationAnimNextRuntimeTest_GraphTraitEvent -> Unexpected trait visited");
+				AddErrorIfFalse(EventB->VisitedTraits[1] == NodeTemplateTraitListA[0], "FAnimationAnimNextRuntimeTest_GraphTraitEvent -> Unexpected trait visited");
+				AddErrorIfFalse(EventB->VisitedTraits[2] == NodeTemplateTraitListC[0], "FAnimationAnimNextRuntimeTest_GraphTraitEvent -> Unexpected trait visited");
+				AddErrorIfFalse(EventB->VisitedTraits[3] == NodeTemplateTraitListC[0], "FAnimationAnimNextRuntimeTest_GraphTraitEvent -> Unexpected trait visited");
+
+				UE_RETURN_ON_ERROR(EventB->ChildEvent.IsValid(), "FAnimationAnimNextRuntimeTest_GraphTraitEvent -> Expected child event");
+				auto EventA = EventB->ChildEvent->AsType<FTraitAnimGraphTest_EventA>();
+				UE_RETURN_ON_ERROR(EventA != nullptr, "FAnimationAnimNextRuntimeTest_GraphTraitEvent -> Expected child event of correct type");
+
+				UE_RETURN_ON_ERROR(EventA->VisitedTraits.Num() == 1, "FAnimationAnimNextRuntimeTest_GraphTraitEvent -> Unexpected number of traits visited");
+				AddErrorIfFalse(EventA->VisitedTraits[0] == NodeTemplateTraitListC[0], "FAnimationAnimNextRuntimeTest_GraphTraitEvent -> Unexpected trait visited");
+			}
+
+			// NodeA raises output event, only visible to Root
+			{
+				auto EventB = MakeTraitEvent<FTraitAnimGraphTest_EventB>();
+				EventB->bTestFlag1 = true;
+
+				GraphInstance.QueueInputTraitEvent(EventB);
+
+				UpdateGraph(GraphInstance, 0.0333f);
+
+				UE_RETURN_ON_ERROR(EventB->VisitedTraits.Num() == 4, "FAnimationAnimNextRuntimeTest_GraphTraitEvent -> Unexpected number of traits visited");
+				AddErrorIfFalse(EventB->VisitedTraits[0] == NodeTemplateTraitListB[0], "FAnimationAnimNextRuntimeTest_GraphTraitEvent -> Unexpected trait visited");
+				AddErrorIfFalse(EventB->VisitedTraits[1] == NodeTemplateTraitListA[0], "FAnimationAnimNextRuntimeTest_GraphTraitEvent -> Unexpected trait visited");
+				AddErrorIfFalse(EventB->VisitedTraits[2] == NodeTemplateTraitListC[0], "FAnimationAnimNextRuntimeTest_GraphTraitEvent -> Unexpected trait visited");
+				AddErrorIfFalse(EventB->VisitedTraits[3] == NodeTemplateTraitListC[0], "FAnimationAnimNextRuntimeTest_GraphTraitEvent -> Unexpected trait visited");
+
+				UE_RETURN_ON_ERROR(EventB->ChildEvent.IsValid(), "FAnimationAnimNextRuntimeTest_GraphTraitEvent -> Expected child event");
+				auto EventA = EventB->ChildEvent->AsType<FTraitAnimGraphTest_EventA>();
+				UE_RETURN_ON_ERROR(EventA != nullptr, "FAnimationAnimNextRuntimeTest_GraphTraitEvent -> Expected child event of correct type");
+
+				UE_RETURN_ON_ERROR(EventA->VisitedTraits.Num() == 1, "FAnimationAnimNextRuntimeTest_GraphTraitEvent -> Unexpected number of traits visited");
+				AddErrorIfFalse(EventA->VisitedTraits[0] == NodeTemplateTraitListB[0], "FAnimationAnimNextRuntimeTest_GraphTraitEvent -> Unexpected trait visited");
+			}
+		}
+
+		Registry.Unregister(NodeTemplateA);
+		Registry.Unregister(NodeTemplateB);
+		Registry.Unregister(NodeTemplateC);
+
+		AddErrorIfFalse(Registry.GetNum() == 0, "FAnimationAnimNextRuntimeTest_GraphTraitEvent -> Registry should contain 0 templates");
+	}
+
+	Tests::FUtils::CleanupAfterTests();
+
+	return true;
+}
+
+#endif
diff --git a/Engine/Plugins/Experimental/AnimNextAnimGraph/Source/AnimNextAnimGraphTestSuite/Private/AnimNextAnimGraphTraitEventTest.h b/Engine/Plugins/Experimental/AnimNextAnimGraph/Source/AnimNextAnimGraphTestSuite/Private/AnimNextAnimGraphTraitEventTest.h
new file mode 100644
--- /dev/null
+++ b/Engine/Plugins/Experimental/AnimNextAnimGraph/Source/AnimNextAnimGraphTestSuite/Private/AnimNextAnimGraphTraitEventTest.h
@@ -0,0 +1,35 @@
+// Copyright Epic Games, Inc. All Rights Reserved.
+
+#pragma once
+
+#include "CoreMinimal.h"
+
+#include "TraitCore/TraitUID.h"
+#include "TraitCore/TraitEvent.h"
+
+#include "AnimNextAnimGraphTraitEventTest.generated.h"
+
+USTRUCT()
+struct FTraitAnimGraphTest_EventA : public FAnimNextTraitEvent
+{
+	GENERATED_BODY()
+	DECLARE_ANIM_TRAIT_EVENT(FTraitAnimGraphTest_EventA)
+
+	bool bTestFlag = false;
+
+	TArray<UE::AnimNext::FTraitUID> VisitedTraits;
+};
+
+USTRUCT()
+struct FTraitAnimGraphTest_EventB : public FAnimNextTraitEvent
+{
+	GENERATED_BODY()
+	DECLARE_ANIM_TRAIT_EVENT(FTraitAnimGraphTest_EventB)
+
+	bool bTestFlag0 = false;
+	bool bTestFlag1 = false;
+
+	TArray<UE::AnimNext::FTraitUID> VisitedTraits;
+
+	FAnimNextTraitEventPtr ChildEvent;
+};
diff --git a/Engine/Plugins/Experimental/AnimNextAnimGraph/Source/AnimNextAnimGraphTestSuite/Private/AnimNextAnimGraphTraitGraphTest.cpp b/Engine/Plugins/Experimental/AnimNextAnimGraph/Source/AnimNextAnimGraphTestSuite/Private/AnimNextAnimGraphTraitGraphTest.cpp
--- a/Engine/Plugins/Experimental/AnimNextAnimGraph/Source/AnimNextAnimGraphTestSuite/Private/AnimNextAnimGraphTraitGraphTest.cpp
+++ b/Engine/Plugins/Experimental/AnimNextAnimGraph/Source/AnimNextAnimGraphTestSuite/Private/AnimNextAnimGraphTraitGraphTest.cpp
@@ -74,7 +74,7 @@ namespace UE::AnimNext
 		GeneratorMacro(IEvaluate) \
 		GeneratorMacro(IUpdate) \
 
-	GENERATE_ANIM_TRAIT_IMPLEMENTATION(FTestTrait, TRAIT_INTERFACE_ENUMERATOR)
+	GENERATE_ANIM_TRAIT_IMPLEMENTATION(FTestTrait, TRAIT_INTERFACE_ENUMERATOR, NULL_ANIM_TRAIT_EVENT_ENUMERATOR)
 	#undef TRAIT_INTERFACE_ENUMERATOR
 }
 
diff --git a/Engine/Plugins/Experimental/AnimNextAnimGraph/Source/AnimNextAnimGraphTestSuite/Private/AnimNextAnimGraphTraitInterfacesTest.cpp b/Engine/Plugins/Experimental/AnimNextAnimGraph/Source/AnimNextAnimGraphTestSuite/Private/AnimNextAnimGraphTraitInterfacesTest.cpp
--- a/Engine/Plugins/Experimental/AnimNextAnimGraph/Source/AnimNextAnimGraphTestSuite/Private/AnimNextAnimGraphTraitInterfacesTest.cpp
+++ b/Engine/Plugins/Experimental/AnimNextAnimGraph/Source/AnimNextAnimGraphTestSuite/Private/AnimNextAnimGraphTraitInterfacesTest.cpp
@@ -94,7 +94,7 @@ namespace UE::AnimNext
 		GeneratorMacro(IEvaluate) \
 		GeneratorMacro(IUpdate) \
 
-	GENERATE_ANIM_TRAIT_IMPLEMENTATION(FTraitWithNoChildren, TRAIT_INTERFACE_ENUMERATOR)
+	GENERATE_ANIM_TRAIT_IMPLEMENTATION(FTraitWithNoChildren, TRAIT_INTERFACE_ENUMERATOR, NULL_ANIM_TRAIT_EVENT_ENUMERATOR)
 	#undef TRAIT_INTERFACE_ENUMERATOR
 
 	// This trait does not update or evaluate
@@ -132,7 +132,7 @@ namespace UE::AnimNext
 	#define TRAIT_INTERFACE_ENUMERATOR(GeneratorMacro) \
 		GeneratorMacro(IHierarchy) \
 
-	GENERATE_ANIM_TRAIT_IMPLEMENTATION(FTraitWithOneChild, TRAIT_INTERFACE_ENUMERATOR)
+	GENERATE_ANIM_TRAIT_IMPLEMENTATION(FTraitWithOneChild, TRAIT_INTERFACE_ENUMERATOR, NULL_ANIM_TRAIT_EVENT_ENUMERATOR)
 	#undef TRAIT_INTERFACE_ENUMERATOR
 
 	struct FTraitWithChildren : FBaseTrait, IHierarchy, IUpdate, IUpdateTraversal, IEvaluate
@@ -225,7 +225,7 @@ namespace UE::AnimNext
 		GeneratorMacro(IUpdate) \
 		GeneratorMacro(IUpdateTraversal) \
 
-	GENERATE_ANIM_TRAIT_IMPLEMENTATION(FTraitWithChildren, TRAIT_INTERFACE_ENUMERATOR)
+	GENERATE_ANIM_TRAIT_IMPLEMENTATION(FTraitWithChildren, TRAIT_INTERFACE_ENUMERATOR, NULL_ANIM_TRAIT_EVENT_ENUMERATOR)
 	#undef TRAIT_INTERFACE_ENUMERATOR
 
 	struct IScopedTagInterface : IScopedTraitInterface
@@ -315,7 +315,7 @@ namespace UE::AnimNext
 		GeneratorMacro(IScopedTagInterface) \
 		GeneratorMacro(IUpdate) \
 
-	GENERATE_ANIM_TRAIT_IMPLEMENTATION(FScopedTagTrait, TRAIT_INTERFACE_ENUMERATOR)
+	GENERATE_ANIM_TRAIT_IMPLEMENTATION(FScopedTagTrait, TRAIT_INTERFACE_ENUMERATOR, NULL_ANIM_TRAIT_EVENT_ENUMERATOR)
 	#undef TRAIT_INTERFACE_ENUMERATOR
 
 	// Tests if we have a scoped tag
@@ -349,7 +349,7 @@ namespace UE::AnimNext
 	#define TRAIT_INTERFACE_ENUMERATOR(GeneratorMacro) \
 		GeneratorMacro(IUpdate) \
 
-	GENERATE_ANIM_TRAIT_IMPLEMENTATION(FTestScopedTagTrait, TRAIT_INTERFACE_ENUMERATOR)
+	GENERATE_ANIM_TRAIT_IMPLEMENTATION(FTestScopedTagTrait, TRAIT_INTERFACE_ENUMERATOR, NULL_ANIM_TRAIT_EVENT_ENUMERATOR)
 	#undef TRAIT_INTERFACE_ENUMERATOR
 }
 
