AnimNext: Parameters rework - UOL integration

This reworks the parameters system quite extensively. It replaces the fixed-function object accessor system with one based on the concept of external instance identifiers that correspond to parameter sources.
UOLs are one such instance identifier that correspond to a way to find regular UObjects.
In the future we can expand this to allow nearly any way of locating data outside the system, as long as it's identifier is exposed as a UScriptStruct.

diff --git a/Engine/Plugins/Experimental/AnimNext/Source/AnimNext/AnimNext.Build.cs b/Engine/Plugins/Experimental/AnimNext/Source/AnimNext/AnimNext.Build.cs
--- a/Engine/Plugins/Experimental/AnimNext/Source/AnimNext/AnimNext.Build.cs
+++ b/Engine/Plugins/Experimental/AnimNext/Source/AnimNext/AnimNext.Build.cs
@@ -21,7 +21,8 @@ namespace UnrealBuildTool.Rules
 			PrivateDependencyModuleNames.AddRange(
 				new string[]
 				{
-					"AssetRegistry"
+					"AssetRegistry",
+					"UniversalObjectLocator",
 				}
 			);
 
diff --git a/Engine/Plugins/Experimental/AnimNext/Source/AnimNext/Internal/AnimNextConfig.h b/Engine/Plugins/Experimental/AnimNext/Source/AnimNext/Internal/AnimNextConfig.h
--- a/Engine/Plugins/Experimental/AnimNext/Source/AnimNext/Internal/AnimNextConfig.h
+++ b/Engine/Plugins/Experimental/AnimNext/Source/AnimNext/Internal/AnimNextConfig.h
@@ -2,14 +2,8 @@
 
 #pragma once
 
-#include "Param/AnimNextObjectAccessorConfig.h"
 #include "AnimNextConfig.generated.h"
 
-namespace UE::AnimNext
-{
-	class FObjectProxyFactory;
-}
-
 UCLASS(Config=AnimNext)
 class ANIMNEXT_API UAnimNextConfig : public UObject
 {
@@ -20,11 +14,4 @@ private:
 #if WITH_EDITOR
 	virtual void PostEditChangeProperty(FPropertyChangedEvent& PropertyChangedEvent) override;
 #endif
-
-private:
-	friend class UE::AnimNext::FObjectProxyFactory;
-
-	/** The classes that are exposed to AnimNext systems */
-	UPROPERTY(Config, EditAnywhere, Category = "Exposed Classes")
-	TArray<FAnimNextObjectAccessorConfig> ExposedClasses;
 };
\ No newline at end of file
diff --git a/Engine/Plugins/Experimental/AnimNext/Source/AnimNext/Internal/Graph/AnimNextGraph.h b/Engine/Plugins/Experimental/AnimNext/Source/AnimNext/Internal/Graph/AnimNextGraph.h
--- a/Engine/Plugins/Experimental/AnimNext/Source/AnimNext/Internal/Graph/AnimNextGraph.h
+++ b/Engine/Plugins/Experimental/AnimNext/Source/AnimNext/Internal/Graph/AnimNextGraph.h
@@ -27,6 +27,7 @@ struct FAnimNextGraphInstancePtr;
 struct FAnimNextGraphInstance;
 class UAnimNextSchedule;
 struct FAnimNextScheduleGraphTask;
+struct FAnimNextEditorParam;
 struct FAnimNextParam;
 
 namespace UE::AnimNext
@@ -62,14 +63,13 @@ class ANIMNEXT_API UAnimNextGraph : public UAnimNextRigVMAsset, public IAnimNext
 	GENERATED_BODY()
 
 public:
-	static const UE::AnimNext::FParamId DefaultReferencePoseId;
-	static const UE::AnimNext::FParamId DefaultCurrentLODId;
-
 	UAnimNextGraph(const FObjectInitializer& ObjectInitializer);
 
 	// UObject interface
-	virtual void PostLoad() override;
 	virtual void Serialize(FArchive& Ar) override;
+#if WITH_EDITOR
+	virtual void PostEditChangeProperty(FPropertyChangedEvent& PropertyChangedEvent) override;
+#endif
 
 	// IAnimNextScheduleTermInterface interface
 	virtual TConstArrayView<UE::AnimNext::FScheduleTerm> GetTerms() const override;
@@ -85,12 +85,6 @@ public:
 	// @param	InEntryPoint				The entry point to use. If this is NAME_None then the default entry point for this graph is used
 	void AllocateInstance(FAnimNextGraphInstance& InOutParentGraphInstance, FAnimNextGraphInstancePtr& OutInstance, FName InEntryPoint = NAME_None) const;
 
-	// Get the parameter to use to access the reference pose
-	UE::AnimNext::FParamId GetReferencePoseParam() const { return ReferencePoseId; }
-
-	// Get the parameter to use to access the current LOD
-	UE::AnimNext::FParamId GetCurrentLODParam() const { return CurrentLODId; }
-
 	// Update the parameter layer, if any
 	void UpdateLayer(UE::AnimNext::FParamStackLayerHandle& InHandle, float InDeltaTime) const;
 
@@ -111,6 +105,12 @@ protected:
 	void ThawGraphInstances();
 #endif
 
+	// Get the fully-qualified default entry point name (/Path/To/Asset.Asset:EntryPoint)
+	FName GetDefaultEntryPoint() const;
+
+	// Set the default entry point name (unqualified)
+	void SetDefaultEntryPoint(FName InEntryPoint);
+
 	friend class UAnimNextGraphFactory;
 	friend class UAnimNextGraph_EditorData;
 	friend class UAnimNextGraph_Parameter;
@@ -156,20 +156,13 @@ protected:
 	UPROPERTY()
 	TArray<TObjectPtr<UObject>> GraphReferencedObjects;
 
-	// The entry point that this graph defaults to using
+	// The entry point that this graph defaults to using (unqualified by asset). Use GetDefaultEntryPoint to get the fully-qualified name.
 	UPROPERTY(EditAnywhere, Category = "Graph")
 	FName DefaultEntryPoint = FRigUnit_AnimNextGraphRoot::DefaultEntryPoint;
 
-	// The parameter to use to access the reference pose
-	UPROPERTY(EditAnywhere, Category = "Graph", meta=(CustomWidget = "ParamName", AllowedParamType = "FAnimNextGraphReferencePose"))
-	FName ReferencePose = DefaultReferencePoseId.GetName();
-
-	// The parameter to use to access the current LOD
-	UPROPERTY(EditAnywhere, Category = "Graph", meta=(CustomWidget = "ParamName", AllowedParamType = "int32"))
-	FName CurrentLOD = DefaultCurrentLODId.GetName();
-
-	UE::AnimNext::FParamId ReferencePoseId = UE::AnimNext::FParamId(ReferencePose);
-	UE::AnimNext::FParamId CurrentLODId = UE::AnimNext::FParamId(CurrentLOD);
+	// Cached fully-qualified entry point, set via by SetDefaultEntryPoint.
+	UPROPERTY(Transient)
+	FName CachedDefaultEntryPoint = NAME_None;
 
 	// Hash of required parameters
 	UPROPERTY()
diff --git a/Engine/Plugins/Experimental/AnimNext/Source/AnimNext/Internal/Graph/RigUnit_AnimNextGraphRoot.h b/Engine/Plugins/Experimental/AnimNext/Source/AnimNext/Internal/Graph/RigUnit_AnimNextGraphRoot.h
--- a/Engine/Plugins/Experimental/AnimNext/Source/AnimNext/Internal/Graph/RigUnit_AnimNextGraphRoot.h
+++ b/Engine/Plugins/Experimental/AnimNext/Source/AnimNext/Internal/Graph/RigUnit_AnimNextGraphRoot.h
@@ -27,7 +27,7 @@ struct ANIMNEXT_API FRigUnit_AnimNextGraphRoot : public FRigUnit_AnimNextBase
 	virtual bool CanOnlyExistOnce() const override { return true; }
 
 	// The execution result
-	UPROPERTY(EditAnywhere, Category = Result, meta = (Input))
+	UPROPERTY(EditAnywhere, Category = Result, meta = (Input, HideSubPins))
 	FAnimNextTraitHandle Result;
 
 	// In order for this node to be considered an executable RigUnit, it needs a pin to derive from FRigVMExecuteContext
diff --git a/Engine/Plugins/Experimental/AnimNext/Source/AnimNext/Internal/Param/AnimNextEditorParam.h b/Engine/Plugins/Experimental/AnimNext/Source/AnimNext/Internal/Param/AnimNextEditorParam.h
new file mode 100644
--- /dev/null
+++ b/Engine/Plugins/Experimental/AnimNext/Source/AnimNext/Internal/Param/AnimNextEditorParam.h
@@ -0,0 +1,41 @@
+// Copyright Epic Games, Inc. All Rights Reserved.
+
+#pragma once
+
+#include "CoreMinimal.h"
+#include "Param/ParamType.h"
+#include "UniversalObjectLocator.h"
+#include "InstancedStruct.h"
+#include "AnimNextEditorParam.generated.h"
+
+struct FAnimNextParam;
+struct FAnimNextParamInstanceIdentifier;
+
+// TODO: Move this to UncookedOnly module once schedules are refactored - this should be an editor-only type
+
+USTRUCT()
+struct FAnimNextEditorParam
+{
+	GENERATED_BODY()
+
+	FAnimNextEditorParam() = default;
+
+	ANIMNEXT_API FAnimNextEditorParam(FName InName, const FAnimNextParamType& InType, const TInstancedStruct<FAnimNextParamInstanceIdentifier>& InInstanceId);
+#if WITH_EDITORONLY_DATA
+	ANIMNEXT_API explicit FAnimNextEditorParam(const FAnimNextParam& InAnimNextParam);
+#endif
+	bool IsValid() const
+	{
+		return !Name.IsNone() && Type.IsValid();
+	}
+
+	UPROPERTY(EditAnywhere, Category = "Parameter")
+	FName Name;
+
+	UPROPERTY(EditAnywhere, Category = "Parameter")
+	FAnimNextParamType Type;
+
+	UPROPERTY(EditAnywhere, Category = "Parameter", NoClear, meta=(ExcludeBaseStruct))
+	TInstancedStruct<FAnimNextParamInstanceIdentifier> InstanceId;
+};
+
diff --git a/Engine/Plugins/Experimental/AnimNext/Source/AnimNext/Internal/Param/AnimNextParam.h b/Engine/Plugins/Experimental/AnimNext/Source/AnimNext/Internal/Param/AnimNextParam.h
--- a/Engine/Plugins/Experimental/AnimNext/Source/AnimNext/Internal/Param/AnimNextParam.h
+++ b/Engine/Plugins/Experimental/AnimNext/Source/AnimNext/Internal/Param/AnimNextParam.h
@@ -3,9 +3,17 @@
 #pragma once
 
 #include "CoreMinimal.h"
+#include "Param/ParamId.h"
 #include "Param/ParamType.h"
 #include "AnimNextParam.generated.h"
 
+struct FAnimNextParamInstanceIdentifier;
+template<typename T> struct TInstancedStruct;
+
+#if WITH_EDITORONLY_DATA
+struct FAnimNextEditorParam;
+#endif
+
 USTRUCT()
 struct FAnimNextParam
 {
@@ -13,14 +21,42 @@ struct FAnimNextParam
 
 	FAnimNextParam() = default;
 
-	FAnimNextParam(FName InName, const FAnimNextParamType& InType)
-		: Name(InName)
-		, Type(InType)
-	{}
+	ANIMNEXT_API FAnimNextParam(FName InName, const FAnimNextParamType& InType, const TInstancedStruct<FAnimNextParamInstanceIdentifier>& InInstanceId);
+
+#if WITH_EDITORONLY_DATA
+	ANIMNEXT_API explicit FAnimNextParam(const FAnimNextEditorParam& InEditorParam);
+#endif
+
+	bool IsValid() const
+	{
+		return !Name.IsNone() && Type.IsValid();
+	}
+
+	// Get the param ID for this parameter, lazily calculating the parameter hash
+	UE::AnimNext::FParamId GetParamId() const
+	{
+		if(Hash == 0)
+		{
+			Hash = UE::AnimNext::FParamId::CalculateHash(Name, InstanceId);
+		}
 
-	UPROPERTY(EditAnywhere, Category = "Parameter")
+		return UE::AnimNext::FParamId(Name, InstanceId, Hash);
+	}
+
+	UPROPERTY()
 	FName Name;
 
-	UPROPERTY(EditAnywhere, Category = "Parameter")
+	UPROPERTY()
+	FName InstanceId;
+
+#if WITH_EDITORONLY_DATA
+	UPROPERTY()
+	TObjectPtr<const UScriptStruct> InstanceIdType;
+#endif
+
+	UPROPERTY()
 	FAnimNextParamType Type;
+
+	// Cached hash
+	mutable uint32 Hash = 0;
 };
\ No newline at end of file
diff --git a/Engine/Plugins/Experimental/AnimNext/Source/AnimNext/Internal/Param/AnimNextParamInstanceIdentifier.h b/Engine/Plugins/Experimental/AnimNext/Source/AnimNext/Internal/Param/AnimNextParamInstanceIdentifier.h
new file mode 100644
--- /dev/null
+++ b/Engine/Plugins/Experimental/AnimNext/Source/AnimNext/Internal/Param/AnimNextParamInstanceIdentifier.h
@@ -0,0 +1,25 @@
+﻿// Copyright Epic Games, Inc. All Rights Reserved.
+
+#pragma once
+
+#include "CoreMinimal.h"
+#include "AnimNextParamInstanceIdentifier.generated.h"
+
+// TODO: Move this to UncookedOnly module once schedules are refactored - this should be an editor-only type
+// Base struct that abstracts an identifier for parameter source instances
+USTRUCT(meta=(Hidden))
+struct FAnimNextParamInstanceIdentifier
+{
+	GENERATED_BODY()
+
+	virtual ~FAnimNextParamInstanceIdentifier() = default;
+
+	// Check the identifier for validity
+	virtual bool IsValid() const { return false; }
+
+	// Get an FName corresponding to this instance identifier 
+	virtual FName ToName() const { return NAME_None; }
+
+	// Set up this identifier from the supplied name
+	virtual void FromName(FName InName) {}
+};
diff --git a/Engine/Plugins/Experimental/AnimNext/Source/AnimNext/Internal/Param/AnimNextParamUniversalObjectLocator.h b/Engine/Plugins/Experimental/AnimNext/Source/AnimNext/Internal/Param/AnimNextParamUniversalObjectLocator.h
new file mode 100644
--- /dev/null
+++ b/Engine/Plugins/Experimental/AnimNext/Source/AnimNext/Internal/Param/AnimNextParamUniversalObjectLocator.h
@@ -0,0 +1,37 @@
+﻿// Copyright Epic Games, Inc. All Rights Reserved.
+
+#pragma once
+
+#include "CoreMinimal.h"
+#include "AnimNextParamInstanceIdentifier.h"
+#include "ParamUtils.h"
+#include "UniversalObjectLocator.h"
+#include "UniversalObjectLocatorStringParams.h"
+#include "AnimNextParamUniversalObjectLocator.generated.h"
+
+// TODO: Move this to UncookedOnly module once schedules are refactored - this should be an editor-only type
+// Wraps a universal object locator to identify external object instances to the parameter system
+USTRUCT()
+struct FAnimNextParamUniversalObjectLocator : public FAnimNextParamInstanceIdentifier
+{
+	GENERATED_BODY()
+
+	// FAnimNextParamInstanceIdentifier interface
+	virtual bool IsValid() const override
+	{
+		return !Locator.IsEmpty();
+	}
+
+	virtual FName ToName() const override
+	{
+		return UE::AnimNext::FParamUtils::LocatorToName(Locator);
+	}
+	
+	virtual void FromName(FName InName) override
+	{
+		Locator = FUniversalObjectLocator::FromString(InName.ToString(), UE::UniversalObjectLocator::FParseStringParams());
+	}
+
+	UPROPERTY(EditAnywhere, Category = "Parameter", meta = (LocatorContext="AnimNextContext"))
+	FUniversalObjectLocator Locator;
+};
diff --git a/Engine/Plugins/Experimental/AnimNext/Source/AnimNext/Internal/Param/ObjectProxyFactory.h b/Engine/Plugins/Experimental/AnimNext/Source/AnimNext/Internal/Param/ObjectProxyFactory.h
new file mode 100644
--- /dev/null
+++ b/Engine/Plugins/Experimental/AnimNext/Source/AnimNext/Internal/Param/ObjectProxyFactory.h
@@ -0,0 +1,61 @@
+// Copyright Epic Games, Inc. All Rights Reserved.
+
+#pragma once
+
+#include "CoreMinimal.h"
+#include "Param/IParameterSourceFactory.h"
+#include "UObject/ObjectKey.h"
+
+class AActor;
+class UActorComponent;
+class UAnimNextConfig;
+class UAnimNextParameterSchema;
+struct FAnimNextParamInstanceIdentifier;
+
+namespace UE::AnimNext
+{
+	struct FClassProxy;
+	class FModule;
+}
+
+namespace UE::AnimNext::UncookedOnly
+{
+	class FObjectProxyType;
+}
+
+namespace UE::AnimNext
+{
+
+// Factory for object proxies that supply 'external' parameters
+class FObjectProxyFactory : public IParameterSourceFactory
+{
+	friend class ::UAnimNextConfig;
+	friend class FModule;
+	friend class UE::AnimNext::UncookedOnly::FObjectProxyType;
+
+public:
+	FObjectProxyFactory();
+	virtual ~FObjectProxyFactory() override;
+
+private:
+	// Refresh built-in accessors
+	void Refresh();
+
+	// IParameterSourceFactory interface
+	virtual TUniquePtr<IParameterSource> CreateParameterSource(const FParameterSourceContext& InContext, const TInstancedStruct<FAnimNextParamInstanceIdentifier>& InInstanceId, TConstArrayView<FName> InRequiredParameters) const override;
+
+	// Finds a class proxy for the supplied class
+	ANIMNEXT_API TSharedRef<FClassProxy> FindOrCreateClassProxy(const UClass* InClass) const;
+
+	// Map of classes -> proxy
+	mutable TMap<TObjectKey<UClass>, TSharedPtr<FClassProxy>> ClassMap;
+
+	// Detect concurrent access for ObjectAccessors
+	UE_MT_DECLARE_RW_RECURSIVE_ACCESS_DETECTOR(ObjectAccessorsAccessDetector);
+
+#if WITH_EDITOR
+	FDelegateHandle OnObjectsReinstancedHandle;
+#endif
+};
+
+}
\ No newline at end of file
diff --git a/Engine/Plugins/Experimental/AnimNext/Source/AnimNext/Internal/Param/ParamUtils.h b/Engine/Plugins/Experimental/AnimNext/Source/AnimNext/Internal/Param/ParamUtils.h
--- a/Engine/Plugins/Experimental/AnimNext/Source/AnimNext/Internal/Param/ParamUtils.h
+++ b/Engine/Plugins/Experimental/AnimNext/Source/AnimNext/Internal/Param/ParamUtils.h
@@ -5,6 +5,7 @@
 #include "CoreMinimal.h"
 
 struct FAnimNextParamType;
+struct FUniversalObjectLocator;
 
 namespace UE::AnimNext
 {
@@ -23,6 +24,15 @@ struct FParamUtils
 	// imply directionality, e.g. IsCompatible((int64)InLHS, (int32)InRHS) is allowed as no data loss 
 	// occurs, but IsCompatible((int32)InLHS, (int64)InRHS) is not as B could be truncated.
 	static ANIMNEXT_API FParamCompatibility GetCompatibility(const FAnimNextParamType& InLHS, const FAnimNextParamType& InRHS);
+
+	// Check whether the supplied function can be used to access/map parameters
+	static ANIMNEXT_API bool CanUseFunction(const UFunction* InFunction);
+
+	// Check whether the supplied property can be used to access/map parameters
+	static ANIMNEXT_API bool CanUseProperty(const FProperty* InProperty);
+
+	// Convert a UOL to an FName
+	static ANIMNEXT_API FName LocatorToName(const FUniversalObjectLocator& InLocator);
 };
 
 }
\ No newline at end of file
diff --git a/Engine/Plugins/Experimental/AnimNext/Source/AnimNext/Internal/Param/RigVMDispatch_GetLayerParameter.h b/Engine/Plugins/Experimental/AnimNext/Source/AnimNext/Internal/Param/RigVMDispatch_GetLayerParameter.h
--- a/Engine/Plugins/Experimental/AnimNext/Source/AnimNext/Internal/Param/RigVMDispatch_GetLayerParameter.h
+++ b/Engine/Plugins/Experimental/AnimNext/Source/AnimNext/Internal/Param/RigVMDispatch_GetLayerParameter.h
@@ -14,7 +14,7 @@ namespace UE::AnimNext::UncookedOnly
 /*
  * Gets a parameter's current value from the graph
  */
-USTRUCT(meta = (DisplayName = "Get Parameter from Graph", Category="Parameters", NodeColor = "0.8, 0, 0.2, 1"))
+USTRUCT(meta = (DisplayName = "Get Graph Parameter", Category="Parameters", NodeColor = "0.8, 0, 0.2, 1"))
 struct ANIMNEXT_API FRigVMDispatch_GetLayerParameter : public FRigVMDispatchFactory
 {
 	GENERATED_BODY()
diff --git a/Engine/Plugins/Experimental/AnimNext/Source/AnimNext/Internal/Param/RigVMDispatch_GetParameter.h b/Engine/Plugins/Experimental/AnimNext/Source/AnimNext/Internal/Param/RigVMDispatch_GetParameter.h
--- a/Engine/Plugins/Experimental/AnimNext/Source/AnimNext/Internal/Param/RigVMDispatch_GetParameter.h
+++ b/Engine/Plugins/Experimental/AnimNext/Source/AnimNext/Internal/Param/RigVMDispatch_GetParameter.h
@@ -14,7 +14,7 @@ namespace UE::AnimNext::UncookedOnly
 /*
  * Gets a parameter's value
  */
-USTRUCT(meta = (DisplayName = "Get Parameter", Category="Parameters", NodeColor = "0.8, 0, 0.2, 1"))
+USTRUCT(meta = (Deprecated, DisplayName = "Get Parameter", Category="Parameters", NodeColor = "0.8, 0, 0.2, 1"))
 struct ANIMNEXT_API FRigVMDispatch_GetParameter : public FRigVMDispatchFactory
 {
 	GENERATED_BODY()
diff --git a/Engine/Plugins/Experimental/AnimNext/Source/AnimNext/Internal/Param/RigVMDispatch_GetScopedParameter.h b/Engine/Plugins/Experimental/AnimNext/Source/AnimNext/Internal/Param/RigVMDispatch_GetScopedParameter.h
new file mode 100644
--- /dev/null
+++ b/Engine/Plugins/Experimental/AnimNext/Source/AnimNext/Internal/Param/RigVMDispatch_GetScopedParameter.h
@@ -0,0 +1,52 @@
+// Copyright Epic Games, Inc. All Rights Reserved.
+
+#pragma once
+
+#include "RigVMCore/RigVMDispatchFactory.h"
+#include "AnimNextExecuteContext.h"
+#include "RigVMDispatch_GetScopedParameter.generated.h"
+
+namespace UE::AnimNext::Editor
+{
+	class SGraphPinParam;
+}
+
+namespace UE::AnimNext::UncookedOnly
+{
+	struct FUtils;
+}
+
+/*
+ * Gets a parameter's value
+ */
+USTRUCT(meta = (DisplayName = "Get Parameter", Category="Parameters", NodeColor = "0.8, 0, 0.2, 1"))
+struct ANIMNEXT_API FRigVMDispatch_GetScopedParameter : public FRigVMDispatchFactory
+{
+	GENERATED_BODY()
+
+	FRigVMDispatch_GetScopedParameter();
+
+private:
+	friend struct UE::AnimNext::UncookedOnly::FUtils;
+	friend class UE::AnimNext::Editor::SGraphPinParam;
+
+	virtual UScriptStruct* GetExecuteContextStruct() const { return FAnimNextExecuteContext::StaticStruct(); }
+	virtual FName GetArgumentNameForOperandIndex(int32 InOperandIndex, int32 InTotalOperands) const override;
+	virtual const TArray<FRigVMTemplateArgumentInfo>& GetArgumentInfos() const override;
+#if WITH_EDITOR
+	virtual FString GetArgumentMetaData(const FName& InArgumentName, const FName& InMetaDataKey) const override;
+#endif
+	virtual FRigVMTemplateTypeMap OnNewArgumentType(const FName& InArgumentName, TRigVMTypeIndex InTypeIndex) const override;
+	virtual bool IsSingleton() const override { return true; }
+
+	virtual FRigVMFunctionPtr GetDispatchFunctionImpl(const FRigVMTemplateTypeMap& InTypes) const override
+	{
+		return &FRigVMDispatch_GetScopedParameter::Execute;
+	}
+	static void Execute(FRigVMExtendedExecuteContext& InContext, FRigVMMemoryHandleArray Handles, FRigVMPredicateBranchArray RigVMBranches);
+
+	static const FName ParameterName;
+	static const FName ValueName;
+	static const FName ParameterIdName;
+	static const FName TypeHandleName;
+};
diff --git a/Engine/Plugins/Experimental/AnimNext/Source/AnimNext/Internal/Scheduler/AnimNextSchedule.h b/Engine/Plugins/Experimental/AnimNext/Source/AnimNext/Internal/Scheduler/AnimNextSchedule.h
--- a/Engine/Plugins/Experimental/AnimNext/Source/AnimNext/Internal/Scheduler/AnimNextSchedule.h
+++ b/Engine/Plugins/Experimental/AnimNext/Source/AnimNext/Internal/Scheduler/AnimNextSchedule.h
@@ -4,6 +4,8 @@
 
 #include "CoreMinimal.h"
 #include "AnimNextScheduleGraphTask.h"
+#include "Param/AnimNextEditorParam.h"
+#include "Param/AnimNextParam.h"
 #include "Scheduler/AnimNextSchedulePortTask.h"
 #include "Scheduler/AnimNextScheduleExternalTask.h"
 #include "Scheduler/AnimNextScheduleParamScopeTask.h"
@@ -17,7 +19,7 @@ class UAnimNextSchedule;
 class UAnimNextSchedulerWorldSubsystem;
 class UAnimNextComponent;
 struct FAnimNextSchedulerEntry;
-struct FAnimNextParam;
+struct FAnimNextEditorParam;
 
 namespace UE::AnimNext
 {
@@ -80,16 +82,16 @@ private:
 	TObjectPtr<UAnimNextGraph> Graph = nullptr;
 
 	// Parameter to get the graph from dynamically
-	UPROPERTY(EditAnywhere, Category = "Graph", meta = (CustomWidget = "ParamName", AllowedParamType = "TObjectPtr<UAnimNextGraph>", AllowNone))
-	FName DynamicGraph;
+	UPROPERTY(EditAnywhere, Category = "Graph", meta = (AllowedParamType = "TObjectPtr<UAnimNextGraph>"))
+	FAnimNextEditorParam DynamicGraph;
 
 	// An optional entry point to use when running the supplied graph
-	UPROPERTY(EditAnywhere, Category = "Graph", meta = (CustomWidget = "ParamName", AllowedParamType = "FName", AllowNone))
-	FName EntryPoint;
+	UPROPERTY(EditAnywhere, Category = "Graph", meta = (AllowedParamType = "FAnimNextEntryPoint"))
+	FAnimNextEditorParam EntryPoint;
 
 	// All parameters that are required by this graph to run (only required if dynamic as static graph params can be discovered by the compiler)
-	UPROPERTY(EditAnywhere, Category = "Graph", meta = (CustomWidget = "ParamName"))
-	TArray<FAnimNextParam> RequiredParameters;
+	UPROPERTY(EditAnywhere, Category = "Graph")
+	TArray<FAnimNextEditorParam> RequiredParameters;
 
 	// The intermediate terms used by the graph
 	UPROPERTY(EditAnywhere, Category = "Graph")
@@ -122,8 +124,8 @@ private:
 	friend struct UE::AnimNext::UncookedOnly::FUtils;
 
 	// The external task we wrap
-	UPROPERTY(EditAnywhere, Category = "External Task", meta = (CustomWidget = "ParamName", AllowedParamType = "FAnimNextExternalTaskBinding"))
-	FName ExternalTask;
+	UPROPERTY(EditAnywhere, Category = "External Task", meta = (AllowedParamType = "FAnimNextExternalTaskBinding"))
+	FAnimNextEditorParam ExternalTask;
 };
 
 UCLASS(MinimalAPI, DisplayName="Scope")
@@ -138,8 +140,8 @@ private:
 	virtual void GetPreloadDependencies(TArray<UObject*>& OutDeps) override;
 
 	// The scope to use
-	UPROPERTY(EditAnywhere, Category = "Parameters", meta = (CustomWidget = "ParamName", AllowedParamType = "FAnimNextScope"))
-	FName Scope;
+	UPROPERTY(EditAnywhere, Category = "Parameters", meta = (AllowedParamType = "FAnimNextScope"))
+	FAnimNextEditorParam Scope;
 
 	// Parameters to apply in this scope
 	UPROPERTY(EditAnywhere, Category = "Parameters")
diff --git a/Engine/Plugins/Experimental/AnimNext/Source/AnimNext/Internal/Scheduler/AnimNextScheduleExternalParamTask.h b/Engine/Plugins/Experimental/AnimNext/Source/AnimNext/Internal/Scheduler/AnimNextScheduleExternalParamTask.h
--- a/Engine/Plugins/Experimental/AnimNext/Source/AnimNext/Internal/Scheduler/AnimNextScheduleExternalParamTask.h
+++ b/Engine/Plugins/Experimental/AnimNext/Source/AnimNext/Internal/Scheduler/AnimNextScheduleExternalParamTask.h
@@ -3,10 +3,12 @@
 #pragma once
 
 #include "CoreMinimal.h"
+#include "InstancedStruct.h"
 #include "AnimNextScheduleExternalParamTask.generated.h"
 
 class UAnimNextSchedule;
 struct FAnimNextSchedulerEntry;
+struct FAnimNextParamInstanceIdentifier;
 
 namespace UE::AnimNext
 {
@@ -25,9 +27,9 @@ struct FAnimNextScheduleExternalParameterSource
 {
 	GENERATED_BODY()
 
-	// Parameter source name to use
+	// Locator for this parameter source
 	UPROPERTY()
-	FName ParameterSource;
+	TInstancedStruct<FAnimNextParamInstanceIdentifier> InstanceId;
 
 	// Names of all parameters that are required in this task
 	UPROPERTY()
diff --git a/Engine/Plugins/Experimental/AnimNext/Source/AnimNext/Internal/Scheduler/AnimNextScheduleExternalTask.h b/Engine/Plugins/Experimental/AnimNext/Source/AnimNext/Internal/Scheduler/AnimNextScheduleExternalTask.h
--- a/Engine/Plugins/Experimental/AnimNext/Source/AnimNext/Internal/Scheduler/AnimNextScheduleExternalTask.h
+++ b/Engine/Plugins/Experimental/AnimNext/Source/AnimNext/Internal/Scheduler/AnimNextScheduleExternalTask.h
@@ -3,6 +3,7 @@
 #pragma once
 
 #include "CoreMinimal.h"
+#include "Param/AnimNextParam.h"
 #include "AnimNextScheduleExternalTask.generated.h"
 
 class UAnimNextSchedule;
@@ -44,5 +45,5 @@ private:
 	uint32 ParamParentScopeIndex = MAX_uint32;
 
 	UPROPERTY()
-	FName ExternalTask;
+	FAnimNextParam ExternalTask;
 };
diff --git a/Engine/Plugins/Experimental/AnimNext/Source/AnimNext/Internal/Scheduler/AnimNextScheduleGraphTask.h b/Engine/Plugins/Experimental/AnimNext/Source/AnimNext/Internal/Scheduler/AnimNextScheduleGraphTask.h
--- a/Engine/Plugins/Experimental/AnimNext/Source/AnimNext/Internal/Scheduler/AnimNextScheduleGraphTask.h
+++ b/Engine/Plugins/Experimental/AnimNext/Source/AnimNext/Internal/Scheduler/AnimNextScheduleGraphTask.h
@@ -4,10 +4,14 @@
 
 #include "CoreMinimal.h"
 #include "Graph/AnimNext_LODPose.h"
+#include "Param/AnimNextParam.h"
+#include "Param/ParamId.h"
 #include "AnimNextScheduleGraphTask.generated.h"
 
 class UAnimNextGraph;
-struct FAnimNextParam;
+struct FAnimNextEditorParam;
+struct FAnimNextParamUniversalObjectLocator;
+template<typename T> struct TInstancedStruct;
 
 namespace UE::AnimNext
 {
@@ -22,6 +26,13 @@ namespace UE::AnimNext::UncookedOnly
 	struct FUtils;
 }
 
+namespace UE::AnimNext
+{
+	extern TInstancedStruct<FAnimNextParamUniversalObjectLocator> GetCharacterInstanceId();
+	extern FName GetCharacterInstanceIdName();
+	extern FParamId GetMeshComponentParamId();
+}
+
 USTRUCT()
 struct FAnimNextScheduleGraphTask
 {
@@ -42,6 +53,10 @@ private:
 	// Verify graph's required parameters are satisfied by this task's supplied parameters
 	void VerifyRequiredParameters(UAnimNextGraph* InGraphToRun) const;
 
+#if WITH_EDITORONLY_DATA 
+	// Get required parameters needed for internal operation
+	ANIMNEXT_API static TArray<FAnimNextEditorParam> GetRequiredParametersInternal();
+#endif
 private:
 	UPROPERTY()
 	uint32 TaskIndex = MAX_uint32;
@@ -53,13 +68,13 @@ private:
 	uint32 ParamParentScopeIndex = MAX_uint32;
 
 	UPROPERTY()
-	FName EntryPoint;
+	FAnimNextParam EntryPoint;
 
 	UPROPERTY()
 	TObjectPtr<UAnimNextGraph> Graph = nullptr;
 
 	UPROPERTY()
-	FName DynamicGraph;
+	FAnimNextParam DynamicGraph;
 
 	// Index of each term in the schedule intermediates
 	UPROPERTY()
diff --git a/Engine/Plugins/Experimental/AnimNext/Source/AnimNext/Internal/Scheduler/AnimNextScheduleParamScopeTask.h b/Engine/Plugins/Experimental/AnimNext/Source/AnimNext/Internal/Scheduler/AnimNextScheduleParamScopeTask.h
--- a/Engine/Plugins/Experimental/AnimNext/Source/AnimNext/Internal/Scheduler/AnimNextScheduleParamScopeTask.h
+++ b/Engine/Plugins/Experimental/AnimNext/Source/AnimNext/Internal/Scheduler/AnimNextScheduleParamScopeTask.h
@@ -3,6 +3,7 @@
 #pragma once
 
 #include "CoreMinimal.h"
+#include "Param/AnimNextParam.h"
 #include "AnimNextScheduleParamScopeTask.generated.h"
 
 class UAnimNextSchedule;
@@ -51,9 +52,9 @@ private:
 	UPROPERTY()
 	uint32 TickFunctionIndex = MAX_uint32;
 
-	// The name of the scope
+	// The scope to use
 	UPROPERTY()
-	FName Scope;
+	FAnimNextParam Scope;
 
 	// Parameters to apply in this scopes
 	UPROPERTY()
@@ -80,7 +81,7 @@ private:
 	UPROPERTY()
 	uint32 ParamScopeIndex = MAX_uint32;
 
-	// The name of the scope
+	// The scope to use
 	UPROPERTY()
-	FName Scope;
+	FAnimNextParam Scope;
 };
diff --git a/Engine/Plugins/Experimental/AnimNext/Source/AnimNext/Internal/Scheduler/ScheduleInitializationContext.h b/Engine/Plugins/Experimental/AnimNext/Source/AnimNext/Internal/Scheduler/ScheduleInitializationContext.h
new file mode 100644
--- /dev/null
+++ b/Engine/Plugins/Experimental/AnimNext/Source/AnimNext/Internal/Scheduler/ScheduleInitializationContext.h
@@ -0,0 +1,43 @@
+// Copyright Epic Games, Inc. All Rights Reserved.
+
+#pragma once
+
+#include "CoreMinimal.h"
+
+struct FAnimNextSchedulerEntry;
+
+namespace UE::AnimNext
+{	
+	class IParameterSource;
+	struct FScheduleContext;
+}
+
+namespace UE::AnimNext
+{
+
+enum class EParameterScopeOrdering : int32
+{
+	// Value will be pushed before the scope, allowing the static scope to potentially override the value
+	Before,
+
+	// Value will be pushed after the scope, potentially overriding the static scope
+	After,
+};
+
+// Context passed to schedule initialization callbacks
+struct FScheduleInitializationContext
+{
+public:
+	// Apply the supplied parameter source to the specified scope, evicting any source that was there previously
+	ANIMNEXT_API void ApplyParametersToScope(FName InScope, EParameterScopeOrdering InOrdering, TUniquePtr<IParameterSource>&& InParameters) const;
+
+private:
+	FScheduleInitializationContext(const FScheduleContext& InContext);
+
+	// The context we wrap
+	const FScheduleContext& Context;
+
+	friend struct ::FAnimNextSchedulerEntry;
+};
+
+}
\ No newline at end of file
diff --git a/Engine/Plugins/Experimental/AnimNext/Source/AnimNext/Internal/Scheduler/ScheduleTaskContext.h b/Engine/Plugins/Experimental/AnimNext/Source/AnimNext/Internal/Scheduler/ScheduleTaskContext.h
new file mode 100644
--- /dev/null
+++ b/Engine/Plugins/Experimental/AnimNext/Source/AnimNext/Internal/Scheduler/ScheduleTaskContext.h
@@ -0,0 +1,36 @@
+// Copyright Epic Games, Inc. All Rights Reserved.
+
+#pragma once
+
+#include "CoreMinimal.h"
+
+namespace UE::AnimNext
+{	
+	class IParameterSource;
+	struct FScheduleContext;
+	enum class EParameterScopeOrdering : int32;
+	struct FScheduleBeginTickFunction;
+	struct FScheduleTickFunction;
+}
+
+namespace UE::AnimNext
+{
+
+// Context passed to schedule task callbacks
+struct FScheduleTaskContext
+{
+public:
+	// Apply the supplied parameter source to the specified scope, evicting any source that was there previously
+	ANIMNEXT_API void ApplyParametersToScope(FName InScope, EParameterScopeOrdering InOrdering, TUniquePtr<IParameterSource>&& InParameters) const;
+
+private:
+	FScheduleTaskContext(const FScheduleContext& InContext);
+
+	// The context we wrap
+	const FScheduleContext& Context;
+
+	friend struct FScheduleBeginTickFunction;
+	friend struct FScheduleTickFunction;
+};
+
+}
\ No newline at end of file
diff --git a/Engine/Plugins/Experimental/AnimNext/Source/AnimNext/Internal/Scheduler/Scheduler.h b/Engine/Plugins/Experimental/AnimNext/Source/AnimNext/Internal/Scheduler/Scheduler.h
--- a/Engine/Plugins/Experimental/AnimNext/Source/AnimNext/Internal/Scheduler/Scheduler.h
+++ b/Engine/Plugins/Experimental/AnimNext/Source/AnimNext/Internal/Scheduler/Scheduler.h
@@ -16,6 +16,8 @@ namespace UE::AnimNext
 	struct FSchedulePortDefinition;
 	struct FParamStackLayerHandle;
 	struct FScheduleContext;
+	struct FScheduleInitializationContext;
+	struct FScheduleTaskContext;
 }
 
 namespace UE::AnimNext::UncookedOnly
@@ -38,15 +40,15 @@ struct FScheduler
 	// Acquire a handle that binds a schedule
 	// Initial parameter binding can be achieved via InitializeCallback, which is called once the schedule's data
 	// structures have been set up, but before it is first run.
-	static FScheduleHandle AcquireHandle(UObject* InObject, UAnimNextSchedule* InSchedule, EAnimNextScheduleInitMethod InInitMethod, TUniqueFunction<void(const FScheduleContext&)>&& InitializeCallback = nullptr);
+	static ANIMNEXT_API FScheduleHandle AcquireHandle(UObject* InObject, UAnimNextSchedule* InSchedule, EAnimNextScheduleInitMethod InInitMethod, TUniqueFunction<void(const FScheduleInitializationContext&)>&& InitializeCallback = nullptr);
 
 	// Release an already acquired handle
 	// The full release of the binding referenced by the handle map be deferred after this call is made
-	static void ReleaseHandle(UObject* InObject, FScheduleHandle& InHandle);
+	static ANIMNEXT_API void ReleaseHandle(UObject* InObject, FScheduleHandle& InHandle);
 
 	// Enables or disables the schedule parameterization represented by the supplied handle
 	// This operation is deferred until the next time the schedule ticks
-	static void EnableHandle(UObject* InObject, FScheduleHandle InHandle, bool bInEnabled);
+	static ANIMNEXT_API void EnableHandle(UObject* InObject, FScheduleHandle InHandle, bool bInEnabled);
 
 	enum class ETaskRunLocation : int32
 	{
@@ -62,7 +64,7 @@ struct FScheduler
 	// @param	InTaskName		The name of the task in the schedule to run the supplied task relative to
 	// @param	InTaskFunction	The function to run
 	// @param	InLocation		Where to run the task, before or after
-	static void QueueTask(UObject* InObject, FScheduleHandle InHandle, FName InScheduleTaskName, TUniqueFunction<void(const FScheduleContext&)>&& InTaskFunction, ETaskRunLocation InLocation = ETaskRunLocation::Before);
+	static ANIMNEXT_API void QueueTask(UObject* InObject, FScheduleHandle InHandle, FName InScheduleTaskName, TUniqueFunction<void(const FScheduleTaskContext&)>&& InTaskFunction, ETaskRunLocation InLocation = ETaskRunLocation::Before);
 
 private:
 	friend struct UE::AnimNext::UncookedOnly::FUtils;
diff --git a/Engine/Plugins/Experimental/AnimNext/Source/AnimNext/Private/AnimNextConfig.cpp b/Engine/Plugins/Experimental/AnimNext/Source/AnimNext/Private/AnimNextConfig.cpp
--- a/Engine/Plugins/Experimental/AnimNext/Source/AnimNext/Private/AnimNextConfig.cpp
+++ b/Engine/Plugins/Experimental/AnimNext/Source/AnimNext/Private/AnimNextConfig.cpp
@@ -1,19 +1,11 @@
 ﻿// Copyright Epic Games, Inc. All Rights Reserved.
 
 #include "AnimNextConfig.h"
-#include "Param/ObjectProxyFactory.h"
 
 #if WITH_EDITOR
 
 void UAnimNextConfig::PostEditChangeProperty(FPropertyChangedEvent& PropertyChangedEvent)
 {
-	using namespace UE::AnimNext;
-	
-	if(PropertyChangedEvent.Property->GetFName() == GET_MEMBER_NAME_CHECKED(UAnimNextConfig, ExposedClasses))
-	{
-		FObjectProxyFactory::Refresh();
-	}
-
 	SaveConfig();
 }
 
diff --git a/Engine/Plugins/Experimental/AnimNext/Source/AnimNext/Private/AnimNextModule.cpp b/Engine/Plugins/Experimental/AnimNext/Source/AnimNext/Private/AnimNextModule.cpp
new file mode 100644
--- /dev/null
+++ b/Engine/Plugins/Experimental/AnimNext/Source/AnimNext/Private/AnimNextModule.cpp
@@ -0,0 +1,376 @@
+// Copyright Epic Games, Inc. All Rights Reserved.
+
+#include "CoreMinimal.h"
+#include "AnimNextConfig.h"
+#include "DataRegistry.h"
+#include "IAnimNextModule.h"
+#include "IUniversalObjectLocatorModule.h"
+#include "RigVMRuntimeDataRegistry.h"
+#include "Animation/AnimSequence.h"
+#include "Animation/BlendProfile.h"
+#include "Curves/CurveFloat.h"
+#include "Graph/AnimNextGraph.h"
+#include "Modules/ModuleManager.h"
+#include "Param/AnimNextEditorParam.h"
+#include "Param/AnimNextObjectCastLocatorFragment.h"
+#include "Param/AnimNextObjectFunctionLocatorFragment.h"
+#include "Param/AnimNextObjectPropertyLocatorFragment.h"
+#include "Param/AnimNextParam.h"
+#include "Param/AnimNextTag.h"
+#include "Param/ObjectProxyFactory.h"
+#include "RigVMCore/RigVMRegistry.h"
+#include "Scheduler/AnimNextTickFunctionBinding.h"
+#include "Scheduler/Scheduler.h"
+#include "TraitCore/NodeTemplateRegistry.h"
+#include "TraitCore/TraitRegistry.h"
+
+// Enable console commands only in development builds when logging is enabled
+#define WITH_ANIMNEXT_CONSOLE_COMMANDS (!UE_BUILD_SHIPPING && !NO_LOGGING)
+
+#if WITH_ANIMNEXT_CONSOLE_COMMANDS
+#include "HAL/IConsoleManager.h"
+#include "UObject/UObjectIterator.h"
+
+#include "TraitCore/NodeDescription.h"
+#include "TraitCore/NodeTemplate.h"
+#include "TraitCore/TraitTemplate.h"
+#endif
+
+namespace UE::AnimNext
+{
+
+class FModule : public IAnimNextModule
+{
+public:
+	virtual void StartupModule() override
+	{
+		GetMutableDefault<UAnimNextConfig>()->LoadConfig();
+
+		FRigVMRegistry& RigVMRegistry = FRigVMRegistry::Get();
+		static TPair<UClass*, FRigVMRegistry::ERegisterObjectOperation> const AllowedObjectTypes[] =
+		{
+			{ UAnimSequence::StaticClass(), FRigVMRegistry::ERegisterObjectOperation::Class },
+			{ UScriptStruct::StaticClass(), FRigVMRegistry::ERegisterObjectOperation::Class },
+			{ UBlendProfile::StaticClass(), FRigVMRegistry::ERegisterObjectOperation::Class },
+			{ UCurveFloat::StaticClass(), FRigVMRegistry::ERegisterObjectOperation::Class },
+			{ UAnimNextGraph::StaticClass(), FRigVMRegistry::ERegisterObjectOperation::Class },
+		};
+
+		RigVMRegistry.RegisterObjectTypes(AllowedObjectTypes);
+
+		static UScriptStruct* const AllowedStructTypes[] =
+		{
+			FAnimNextEditorParam::StaticStruct(),
+			FAnimNextParam::StaticStruct(),
+			FAnimNextScope::StaticStruct(),
+			FAnimNextEntryPoint::StaticStruct(),
+			FUniversalObjectLocator::StaticStruct(),
+			FAnimNextTickFunctionBinding::StaticStruct(),
+		};
+
+		RigVMRegistry.RegisterStructTypes(AllowedStructTypes);
+
+		RegisterParameterSourceFactory("ObjectProxy", MakeShared<FObjectProxyFactory>());
+		FDataRegistry::Init();
+		FTraitRegistry::Init();
+		FNodeTemplateRegistry::Init();
+		FScheduler::Init();
+		FRigVMRuntimeDataRegistry::Init();
+
+		UE::UniversalObjectLocator::IUniversalObjectLocatorModule& UolModule = FModuleManager::Get().LoadModuleChecked<UE::UniversalObjectLocator::IUniversalObjectLocatorModule>("UniversalObjectLocator");
+		FDelayedAutoRegisterHelper(EDelayedRegisterRunPhase::ObjectSystemReady,
+			[&UolModule]
+			{
+				{
+					UE::UniversalObjectLocator::FFragmentTypeParameters FragmentTypeParams("animobjfunc", NSLOCTEXT("Engine", "AnimNextObjectFunctionFragment", "Function"));
+					FragmentTypeParams.PrimaryEditorType = "AnimNextObjectFunction";
+					FAnimNextObjectFunctionLocatorFragment::FragmentType = UolModule.RegisterFragmentType<FAnimNextObjectFunctionLocatorFragment>(FragmentTypeParams);
+				}
+				{
+					UE::UniversalObjectLocator::FFragmentTypeParameters FragmentTypeParams("animobjprop", NSLOCTEXT("Engine", "AnimNextObjectPropertyFragment", "Property"));
+					FragmentTypeParams.PrimaryEditorType = "AnimNextObjectProperty";
+					FAnimNextObjectPropertyLocatorFragment::FragmentType = UolModule.RegisterFragmentType<FAnimNextObjectPropertyLocatorFragment>(FragmentTypeParams);
+				}
+				{
+					UE::UniversalObjectLocator::FFragmentTypeParameters FragmentTypeParams("animobjcast", NSLOCTEXT("Engine", "AnimNextCastFragment", "Cast"));
+					FragmentTypeParams.PrimaryEditorType = "AnimNextObjectCast";
+					FAnimNextObjectCastLocatorFragment::FragmentType = UolModule.RegisterFragmentType<FAnimNextObjectCastLocatorFragment>(FragmentTypeParams);
+				}
+			});
+#if WITH_ANIMNEXT_CONSOLE_COMMANDS
+		if (!IsRunningCommandlet())
+		{
+			ConsoleCommands.Add(IConsoleManager::Get().RegisterConsoleCommand(
+				TEXT("AnimNext.ListNodeTemplates"),
+				TEXT("Dumps statistics about node templates to the log."),
+				FConsoleCommandWithArgsDelegate::CreateRaw(this, &FModule::ListNodeTemplates),
+				ECVF_Default
+			));
+			ConsoleCommands.Add(IConsoleManager::Get().RegisterConsoleCommand(
+				TEXT("AnimNext.ListAnimGraphs"),
+				TEXT("Dumps statistics about animation graphs to the log."),
+				FConsoleCommandWithArgsDelegate::CreateRaw(this, &FModule::ListAnimGraphs),
+				ECVF_Default
+			));
+		}
+#endif
+	}
+
+	virtual void ShutdownModule() override
+	{
+		FRigVMRuntimeDataRegistry::Destroy();
+		FScheduler::Destroy();
+		FNodeTemplateRegistry::Destroy();
+		FTraitRegistry::Destroy();
+		FDataRegistry::Destroy();
+		UnregisterParameterSourceFactory("ObjectProxy");
+
+#if WITH_ANIMNEXT_CONSOLE_COMMANDS
+		for (IConsoleObject* Cmd : ConsoleCommands)
+		{
+			IConsoleManager::Get().UnregisterConsoleObject(Cmd);
+		}
+		ConsoleCommands.Empty();
+#endif
+	}
+
+	const IAnimNextAnimGraph* AnimGraphImpl = nullptr;
+
+	virtual void RegisterAnimNextAnimGraph(const IAnimNextAnimGraph& InAnimGraphImpl) override
+	{
+		AnimGraphImpl = &InAnimGraphImpl;
+	}
+
+	virtual void UnregisterAnimNextAnimGraph() override
+	{
+		AnimGraphImpl = nullptr;
+	}
+
+	virtual void UpdateGraph(FAnimNextGraphInstancePtr& GraphInstance, float DeltaTime) override
+	{
+		if (AnimGraphImpl != nullptr)
+		{
+			AnimGraphImpl->UpdateGraph(GraphInstance, DeltaTime);
+		}
+	}
+
+	virtual void EvaluateGraph(FAnimNextGraphInstancePtr& GraphInstance, const UE::AnimNext::FReferencePose& RefPose, int32 GraphLODLevel, FLODPoseHeap& OutputPose) const override
+	{
+		if (AnimGraphImpl != nullptr)
+		{
+			AnimGraphImpl->EvaluateGraph(GraphInstance, RefPose, GraphLODLevel, OutputPose);
+		}
+	}
+
+	virtual TUniquePtr<IParameterSource> CreateParameterSource(const FParameterSourceContext& InContext, const TInstancedStruct<FAnimNextParamInstanceIdentifier>& InInstanceId, TConstArrayView<FName> InRequiredParameters) const override
+	{
+		for(const TPair<FName, TSharedRef<IParameterSourceFactory>>& FactoryPair : ParameterSourceFactories)
+		{
+			// TODO: in future we could introduce a priority system or more complex logic here. For now it is earliest wins as are sources are
+			// all mutually exclusive at the moment.
+			TUniquePtr<IParameterSource> ParameterSource = FactoryPair.Value->CreateParameterSource(InContext, InInstanceId, InRequiredParameters);
+			if(ParameterSource.IsValid())
+			{
+				return ParameterSource;
+			}
+		}
+
+		return nullptr;
+	}
+
+	virtual void RegisterParameterSourceFactory(FName InName, TSharedRef<IParameterSourceFactory> InFactory) override
+	{
+		ParameterSourceFactories.Add(InName, InFactory);
+	}
+
+	virtual void UnregisterParameterSourceFactory(FName InName) override
+	{
+		ParameterSourceFactories.Remove(InName);
+	}
+
+	virtual TSharedPtr<IParameterSourceFactory> FindParameterSourceFactory(FName InName) override
+	{
+		if(TSharedRef<IParameterSourceFactory>* FoundFactory = ParameterSourceFactories.Find(InName))
+		{
+			return *FoundFactory;
+		}
+
+		return nullptr;
+	}
+
+#if WITH_ANIMNEXT_CONSOLE_COMMANDS
+	TArray<IConsoleObject*> ConsoleCommands;
+
+	void ListNodeTemplates(const TArray<FString>& Args)
+	{
+		// Turn off log times to make diff-ing easier
+		TGuardValue<ELogTimes::Type> DisableLogTimes(GPrintLogTimes, ELogTimes::None);
+
+		// Make sure to log everything
+		const ELogVerbosity::Type OldVerbosity = LogAnimation.GetVerbosity();
+		LogAnimation.SetVerbosity(ELogVerbosity::All);
+
+		const FNodeTemplateRegistry& NodeTemplateRegistry = FNodeTemplateRegistry::Get();
+		const FTraitRegistry& TraitRegistry = FTraitRegistry::Get();
+
+		UE_LOG(LogAnimation, Log, TEXT("===== AnimNext Node Templates ====="));
+		UE_LOG(LogAnimation, Log, TEXT("Template Buffer Size: %u bytes"), NodeTemplateRegistry.TemplateBuffer.GetAllocatedSize());
+
+		for (auto It = NodeTemplateRegistry.TemplateUIDToHandleMap.CreateConstIterator(); It; ++It)
+		{
+			const FNodeTemplateRegistryHandle Handle = It.Value();
+			const FNodeTemplate* NodeTemplate = NodeTemplateRegistry.Find(Handle);
+
+			const uint32 NumTraits = NodeTemplate->GetNumTraits();
+
+			UE_LOG(LogAnimation, Log, TEXT("[%x] has %u traits ..."), NodeTemplate->GetUID(), NumTraits);
+			UE_LOG(LogAnimation, Log, TEXT("    Template Size: %u bytes"), NodeTemplate->GetNodeTemplateSize());
+			UE_LOG(LogAnimation, Log, TEXT("    Shared Data Size: %u bytes"), NodeTemplate->GetNodeSharedDataSize());
+			UE_LOG(LogAnimation, Log, TEXT("    Instance Data Size: %u bytes"), NodeTemplate->GetNodeInstanceDataSize());
+			UE_LOG(LogAnimation, Log, TEXT("    Traits ..."));
+
+			const FTraitTemplate* TraitTemplates = NodeTemplate->GetTraits();
+			for (uint32 TraitIndex = 0; TraitIndex < NumTraits; ++TraitIndex)
+			{
+				const FTraitTemplate* TraitTemplate = TraitTemplates + TraitIndex;
+				const FTrait* Trait = TraitRegistry.Find(TraitTemplate->GetRegistryHandle());
+				const FString TraitName = Trait != nullptr ? Trait->GetTraitName() : TEXT("<Unknown>");
+
+				const uint32 NextTraitIndex = TraitIndex + 1;
+				const uint32 EndOfNextTraitSharedData = NextTraitIndex < NumTraits ? TraitTemplates[NextTraitIndex].GetNodeSharedOffset() : NodeTemplate->GetNodeSharedDataSize();
+				const uint32 TraitSharedDataSize = EndOfNextTraitSharedData - TraitTemplate->GetNodeSharedOffset();
+
+				const uint32 EndOfNextTraitInstanceData = NextTraitIndex < NumTraits ? TraitTemplates[NextTraitIndex].GetNodeInstanceOffset() : NodeTemplate->GetNodeInstanceDataSize();
+				const uint32 TraitInstanceDataSize = EndOfNextTraitInstanceData - TraitTemplate->GetNodeInstanceOffset();
+
+				UE_LOG(LogAnimation, Log, TEXT("            %u: [%x] %s (%s)"), TraitIndex, TraitTemplate->GetUID().GetUID(), *TraitName, TraitTemplate->GetMode() == ETraitMode::Base ? TEXT("Base") : TEXT("Additive"));
+				UE_LOG(LogAnimation, Log, TEXT("                Shared Data: [Offset: %u bytes, Size: %u bytes]"), TraitTemplate->GetNodeSharedOffset(), TraitSharedDataSize);
+				if (TraitTemplate->HasLatentProperties() && Trait != nullptr)
+				{
+					UE_LOG(LogAnimation, Log, TEXT("                Shared Data Latent Property Handles: [Offset: %u bytes, Count: %u]"), TraitTemplate->GetNodeSharedLatentPropertyHandlesOffset(), Trait->GetNumLatentTraitProperties());
+				}
+				UE_LOG(LogAnimation, Log, TEXT("                Instance Data: [Offset: %u bytes, Size: %u bytes]"), TraitTemplate->GetNodeInstanceOffset(), TraitInstanceDataSize);
+			}
+		}
+
+		LogAnimation.SetVerbosity(OldVerbosity);
+	}
+
+	void ListAnimGraphs(const TArray<FString>& Args)
+	{
+		// Turn off log times to make diff-ing easier
+		TGuardValue<ELogTimes::Type> DisableLogTimes(GPrintLogTimes, ELogTimes::None);
+
+		// Make sure to log everything
+		const ELogVerbosity::Type OldVerbosity = LogAnimation.GetVerbosity();
+		LogAnimation.SetVerbosity(ELogVerbosity::All);
+
+		TArray<const UAnimNextGraph*> AnimGraphs;
+
+		for (TObjectIterator<UAnimNextGraph> It; It; ++It)
+		{
+			AnimGraphs.Add(*It);
+		}
+
+		struct FCompareObjectNames
+		{
+			FORCEINLINE bool operator()(const UAnimNextGraph& Lhs, const UAnimNextGraph& Rhs) const
+			{
+				return Lhs.GetPathName().Compare(Rhs.GetPathName()) < 0;
+			}
+		};
+		AnimGraphs.Sort(FCompareObjectNames());
+
+		const FNodeTemplateRegistry& NodeTemplateRegistry = FNodeTemplateRegistry::Get();
+		const FTraitRegistry& TraitRegistry = FTraitRegistry::Get();
+		const bool bDetailedOutput = true;
+
+		UE_LOG(LogAnimation, Log, TEXT("===== AnimNext Animation Graphs ====="));
+		UE_LOG(LogAnimation, Log, TEXT("Num Graphs: %u"), AnimGraphs.Num());
+
+		for (const UAnimNextGraph* AnimGraph : AnimGraphs)
+		{
+			uint32 TotalInstanceSize = 0;
+			uint32 NumNodes = 0;
+			{
+				// We always have a node at offset 0
+				int32 NodeOffset = 0;
+
+				while (NodeOffset < AnimGraph->SharedDataBuffer.Num())
+				{
+					const FNodeDescription* NodeDesc = reinterpret_cast<const FNodeDescription*>(&AnimGraph->SharedDataBuffer[NodeOffset]);
+
+					TotalInstanceSize += NodeDesc->GetNodeInstanceDataSize();
+					NumNodes++;
+
+					const FNodeTemplate* NodeTemplate = NodeTemplateRegistry.Find(NodeDesc->GetTemplateHandle());
+					NodeOffset += NodeTemplate->GetNodeSharedDataSize();
+				}
+			}
+
+			UE_LOG(LogAnimation, Log, TEXT("    %s ..."), *AnimGraph->GetPathName());
+			UE_LOG(LogAnimation, Log, TEXT("        Shared Data Size: %.2f KB"), double(AnimGraph->SharedDataBuffer.Num()) / 1024.0);
+			UE_LOG(LogAnimation, Log, TEXT("        Max Instance Data Size: %.2f KB"), double(TotalInstanceSize) / 1024.0);
+			UE_LOG(LogAnimation, Log, TEXT("        Num Nodes: %u"), NumNodes);
+
+			if (bDetailedOutput)
+			{
+				// We always have a node at offset 0
+				int32 NodeOffset = 0;
+
+				while (NodeOffset < AnimGraph->SharedDataBuffer.Num())
+				{
+					const FNodeDescription* NodeDesc = reinterpret_cast<const FNodeDescription*>(&AnimGraph->SharedDataBuffer[NodeOffset]);
+					const FNodeTemplate* NodeTemplate = NodeTemplateRegistry.Find(NodeDesc->GetTemplateHandle());
+
+					const uint32 NumTraits = NodeTemplate->GetNumTraits();
+
+					UE_LOG(LogAnimation, Log, TEXT("        Node %u: [Template %x with %u traits]"), NodeDesc->GetUID().GetNodeIndex(), NodeTemplate->GetUID(), NumTraits);
+					UE_LOG(LogAnimation, Log, TEXT("            Shared Data: [Offset: %u bytes, Size: %u bytes]"), NodeOffset, NodeTemplate->GetNodeSharedDataSize());
+					UE_LOG(LogAnimation, Log, TEXT("            Instance Data Size: %u bytes"), NodeDesc->GetNodeInstanceDataSize());
+					UE_LOG(LogAnimation, Log, TEXT("            Traits ..."));
+
+					const FTraitTemplate* TraitTemplates = NodeTemplate->GetTraits();
+					for (uint32 TraitIndex = 0; TraitIndex < NumTraits; ++TraitIndex)
+					{
+						const FTraitTemplate* TraitTemplate = TraitTemplates + TraitIndex;
+						const FTrait* Trait = TraitRegistry.Find(TraitTemplate->GetRegistryHandle());
+						const FString TraitName = Trait != nullptr ? Trait->GetTraitName() : TEXT("<Unknown>");
+
+						const uint32 NextTraitIndex = TraitIndex + 1;
+						const uint32 EndOfNextTraitSharedData = NextTraitIndex < NumTraits ? TraitTemplates[NextTraitIndex].GetNodeSharedOffset() : NodeTemplate->GetNodeSharedDataSize();
+						const uint32 TraitSharedDataSize = EndOfNextTraitSharedData - TraitTemplate->GetNodeSharedOffset();
+
+						const uint32 EndOfNextTraitInstanceData = NextTraitIndex < NumTraits ? TraitTemplates[NextTraitIndex].GetNodeInstanceOffset() : NodeTemplate->GetNodeInstanceDataSize();
+						const uint32 TraitInstanceDataSize = EndOfNextTraitInstanceData - TraitTemplate->GetNodeInstanceOffset();
+
+						UE_LOG(LogAnimation, Log, TEXT("                    %u: [%x] %s (%s)"), TraitIndex, TraitTemplate->GetUID().GetUID(), *TraitName, TraitTemplate->GetMode() == ETraitMode::Base ? TEXT("Base") : TEXT("Additive"));
+						UE_LOG(LogAnimation, Log, TEXT("                        Shared Data: [Offset: %u bytes, Size: %u bytes]"), TraitTemplate->GetNodeSharedOffset(), TraitSharedDataSize);
+						if (TraitTemplate->HasLatentProperties() && Trait != nullptr)
+						{
+							UE_LOG(LogAnimation, Log, TEXT("                        Shared Data Latent Property Handles: [Offset: %u bytes, Count: %u]"), TraitTemplate->GetNodeSharedLatentPropertyHandlesOffset(), Trait->GetNumLatentTraitProperties());
+						}
+						UE_LOG(LogAnimation, Log, TEXT("                        Instance Data: [Offset: %u bytes, Size: %u bytes]"), TraitTemplate->GetNodeInstanceOffset(), TraitInstanceDataSize);
+					}
+
+					NodeOffset += NodeTemplate->GetNodeSharedDataSize();
+				}
+			}
+		}
+
+		LogAnimation.SetVerbosity(OldVerbosity);
+	}
+#endif
+	
+	// All known factories
+	TMap<FName, TSharedRef<IParameterSourceFactory>> ParameterSourceFactories;
+};
+
+IAnimNextModule& IAnimNextModule::Get()
+{
+	return FModuleManager::LoadModuleChecked<IAnimNextModule>(TEXT("AnimNext"));
+}
+
+}
+
+IMPLEMENT_MODULE(UE::AnimNext::FModule, AnimNext)
diff --git a/Engine/Plugins/Experimental/AnimNext/Source/AnimNext/Private/Component/AnimNextComponent.cpp b/Engine/Plugins/Experimental/AnimNext/Source/AnimNext/Private/Component/AnimNextComponent.cpp
--- a/Engine/Plugins/Experimental/AnimNext/Source/AnimNext/Private/Component/AnimNextComponent.cpp
+++ b/Engine/Plugins/Experimental/AnimNext/Source/AnimNext/Private/Component/AnimNextComponent.cpp
@@ -3,53 +3,21 @@
 #include "Component/AnimNextComponent.h"
 
 #include "Blueprint/BlueprintExceptionInfo.h"
-#include "Param/ParamUtils.h"
 #include "Scheduler/Scheduler.h"
 #include "Scheduler/ScheduleContext.h"
 #include "Param/PropertyBagProxy.h"
+#include "Scheduler/ScheduleTaskContext.h"
+#include "Scheduler/ScheduleInitializationContext.h"
 
 namespace UE::AnimNext::Private
 {
 
-static void SetValuesInScopeHelper(FScheduleInstanceData& InInstanceData, FName InScope, EAnimNextParameterScopeOrdering InOrdering, TConstArrayView<FPropertyBagProxy::FPropertyAndValue> InPropertiesAndValues)
+template<typename ContextType>
+static void SetValuesInScopeHelper(const ContextType& InContext, FName InScope, EAnimNextParameterScopeOrdering InOrdering, TConstArrayView<FPropertyBagProxy::FPropertyAndValue> InPropertiesAndValues)
 {
-	// apply to root
-	if(InScope == NAME_None)
-	{
-		if(!InInstanceData.RootUserScope.IsValid())
-		{
-			InInstanceData.RootUserScope = MakeUnique<FPropertyBagProxy>();
-		}
-
-		InInstanceData.RootUserScope->AddPropertiesAndValues(InPropertiesAndValues);
-	}
-	else // apply to specified scope
-	{
-		FScheduleInstanceData::FUserScope& ScopeSource = InInstanceData.UserScopes.FindOrAdd(InScope);
-		TUniquePtr<FPropertyBagProxy>* ProxyToUse = nullptr;
-		switch(InOrdering)
-		{
-		default:
-		case EAnimNextParameterScopeOrdering::Before:
-			if(!ScopeSource.BeforeSource.IsValid())
-			{
-				ScopeSource.BeforeSource = MakeUnique<FPropertyBagProxy>();
-			}
-			ProxyToUse = &ScopeSource.BeforeSource;
-			break;
-		case EAnimNextParameterScopeOrdering::After:
-			if(!ScopeSource.AfterSource.IsValid())
-			{
-				ScopeSource.AfterSource = MakeUnique<FPropertyBagProxy>();
-			}
-			ProxyToUse = &ScopeSource.AfterSource;
-			break;
-		}
-
-		check(ProxyToUse && ProxyToUse->IsValid());
-
-		(*ProxyToUse)->AddPropertiesAndValues(InPropertiesAndValues);
-	}
+	TUniquePtr<FPropertyBagProxy> PropertyBagProxy = MakeUnique<FPropertyBagProxy>();
+	PropertyBagProxy->AddPropertiesAndValues(InPropertiesAndValues);
+	InContext.ApplyParametersToScope(InScope, (EParameterScopeOrdering)InOrdering, MoveTemp(PropertyBagProxy));
 }
 
 }
@@ -63,10 +31,8 @@ void UAnimNextComponent::OnRegister()
 	if (Schedule)
 	{
 		// Initialization callback to set up any persistent external parameters
-		auto Initialize = [this](const FScheduleContext& InContext)
+		auto Initialize = [this](const FScheduleInitializationContext& InContext)
 		{
-			FScheduleInstanceData& InstanceData = InContext.GetInstanceData();
-
 			// First group params into scopes
 			TMap<FName, TArray<UAnimNextComponentParameter*, TInlineAllocator<4>>, TInlineSetAllocator<4>> ParamsByScope;
 			for(UAnimNextComponentParameter* Parameter : Parameters)
@@ -90,7 +56,7 @@ void UAnimNextComponent::OnRegister()
 				}
 
 				// NOTE: Layer is always applied 'before' currently. If we have a use case for 'After' we can add it to UAnimNextComponentParameter
-				Private::SetValuesInScopeHelper(InContext.GetInstanceData(), ParamPair.Key, EAnimNextParameterScopeOrdering::Before, PropertiesAndValues);
+				Private::SetValuesInScopeHelper(InContext, ParamPair.Key, EAnimNextParameterScopeOrdering::Before, PropertiesAndValues);
 			}
 		};
 
@@ -106,7 +72,6 @@ void UAnimNextComponent::OnUnregister()
 	Super::OnUnregister();
 
 	FScheduler::ReleaseHandle(this, SchedulerHandle);
-	SchedulerHandle.Invalidate();
 }
 
 void UAnimNextComponent::SetParameterInScope(FName Scope, EAnimNextParameterScopeOrdering Ordering, FName Name, int32 Value)
@@ -162,13 +127,13 @@ DEFINE_FUNCTION(UAnimNextComponent::execSetParameterInScope)
 	const void* ValuePtr = ValueProp->ContainerPtrToValuePtr<void>(ContainerPtr);
 	NewProperty->SetValue_InContainer(PropertyBag->GetMutableValue().GetMemory(), ValuePtr);
 
-	FScheduler::QueueTask(P_THIS, P_THIS->SchedulerHandle, Scope, [Scope, Name, NewProperty, PropertyBag = MoveTemp(PropertyBag), Ordering](const FScheduleContext& InContext) mutable
+	FScheduler::QueueTask(P_THIS, P_THIS->SchedulerHandle, Scope, [Scope, Name, NewProperty, PropertyBag = MoveTemp(PropertyBag), Ordering](const FScheduleTaskContext& InContext) mutable
 	{
 		FPropertyBagProxy::FPropertyAndValue PropertyAndValue;
 		PropertyAndValue.Name = Name;
 		PropertyAndValue.Property = NewProperty;
 		PropertyAndValue.ContainerPtr = PropertyBag->GetValue().GetMemory();
-		Private::SetValuesInScopeHelper(InContext.GetInstanceData(), Scope, Ordering, { PropertyAndValue });
+		Private::SetValuesInScopeHelper(InContext, Scope, Ordering, { PropertyAndValue });
 	});
 
 	P_NATIVE_END;
diff --git a/Engine/Plugins/Experimental/AnimNext/Source/AnimNext/Private/Component/AnimNextMeshComponent.h b/Engine/Plugins/Experimental/AnimNext/Source/AnimNext/Private/Component/AnimNextMeshComponent.h
--- a/Engine/Plugins/Experimental/AnimNext/Source/AnimNext/Private/Component/AnimNextMeshComponent.h
+++ b/Engine/Plugins/Experimental/AnimNext/Source/AnimNext/Private/Component/AnimNextMeshComponent.h
@@ -7,6 +7,7 @@
 
 struct FAnimNextGraphReferencePose;
 class UAnimNextSchedulePort_AnimNextMeshComponentPose;
+struct FAnimNextScheduleGraphTask;
 
 // Mesh component for use with AnimNext
 UCLASS(MinimalAPI, Blueprintable, meta = (BlueprintSpawnableComponent))
@@ -15,6 +16,7 @@ class UAnimNextMeshComponent : public USkeletalMeshComponent
 	GENERATED_BODY()
 
 	friend class UAnimNextSchedulePort_AnimNextMeshComponentPose;
+	friend struct FAnimNextScheduleGraphTask;
 
 	UAnimNextMeshComponent();
 
diff --git a/Engine/Plugins/Experimental/AnimNext/Source/AnimNext/Private/Graph/AnimNextGraph.cpp b/Engine/Plugins/Experimental/AnimNext/Source/AnimNext/Private/Graph/AnimNextGraph.cpp
--- a/Engine/Plugins/Experimental/AnimNext/Source/AnimNext/Private/Graph/AnimNextGraph.cpp
+++ b/Engine/Plugins/Experimental/AnimNext/Source/AnimNext/Private/Graph/AnimNextGraph.cpp
@@ -27,9 +27,6 @@ DEFINE_STAT(STAT_AnimNext_Graph_UpdateParamLayer);
 
 #include UE_INLINE_GENERATED_CPP_BY_NAME(AnimNextGraph)
 
-const UE::AnimNext::FParamId UAnimNextGraph::DefaultReferencePoseId = UE::AnimNext::FParamId("UE_AnimNextMeshComponent_ReferencePose");
-const UE::AnimNext::FParamId UAnimNextGraph::DefaultCurrentLODId = UE::AnimNext::FParamId("UE_AnimNextMeshComponent_PredictedLODLevel");
-
 UAnimNextGraph::UAnimNextGraph(const FObjectInitializer& ObjectInitializer)
 	: Super(ObjectInitializer)
 {
@@ -52,7 +49,7 @@ void UAnimNextGraph::AllocateInstanceImpl(FAnimNextGraphInstance* ParentGraphIns
 
 	Instance.Release();
 
-	const FName EntryPoint = (InEntryPoint == NAME_None) ? DefaultEntryPoint : InEntryPoint;
+	const FName EntryPoint = (InEntryPoint == NAME_None) ? GetDefaultEntryPoint() : InEntryPoint;
 	const FAnimNextTraitHandle ResolvedRootTraitHandle = ResolvedRootTraitHandles.FindRef(EntryPoint);
 	if (!ResolvedRootTraitHandle.IsValid())
 	{
@@ -125,16 +122,6 @@ void UAnimNextGraph::UpdateLayer(UE::AnimNext::FParamStackLayerHandle& InHandle,
 	}
 }
 
-void UAnimNextGraph::PostLoad()
-{
-	using namespace UE::AnimNext;
-
-	Super::PostLoad();
-
-	ReferencePoseId = FParamId(ReferencePose);
-	CurrentLODId = FParamId(CurrentLOD);
-}
-
 void UAnimNextGraph::Serialize(FArchive& Ar)
 {
 	Super::Serialize(Ar);
@@ -172,6 +159,8 @@ void UAnimNextGraph::Serialize(FArchive& Ar)
 				LoadFromArchiveBuffer(SharedDataArchiveBuffer);
 			}
 		}
+
+		SetDefaultEntryPoint(DefaultEntryPoint);
 	}
 	else if (Ar.IsSaving())
 	{
@@ -195,6 +184,16 @@ void UAnimNextGraph::Serialize(FArchive& Ar)
 	}
 }
 
+#if WITH_EDITOR
+void UAnimNextGraph::PostEditChangeProperty(FPropertyChangedEvent& PropertyChangedEvent)
+{
+	if(PropertyChangedEvent.Property && PropertyChangedEvent.Property->GetFName() == GET_MEMBER_NAME_CHECKED(UAnimNextGraph, DefaultEntryPoint))
+	{
+		SetDefaultEntryPoint(DefaultEntryPoint); 
+	}
+}
+#endif
+
 TConstArrayView<UE::AnimNext::FScheduleTerm> UAnimNextGraph::GetTerms() const
 {
 	using namespace UE::AnimNext;
@@ -232,7 +231,7 @@ bool UAnimNextGraph::LoadFromArchiveBuffer(const TArray<uint8>& InSharedDataArch
 	else
 	{
 		SharedDataBuffer.Empty(0);
-		ResolvedRootTraitHandles.Add(FRigUnit_AnimNextGraphRoot::DefaultEntryPoint, FAnimNextTraitHandle());
+		ResolvedRootTraitHandles.Add(GetDefaultEntryPoint(), FAnimNextTraitHandle());
 		return false;
 	}
 }
@@ -260,3 +259,19 @@ void UAnimNextGraph::ThawGraphInstances()
 	}
 }
 #endif
+
+void UAnimNextGraph::SetDefaultEntryPoint(FName InEntryPoint)
+{
+	DefaultEntryPoint = InEntryPoint;
+
+	TStringBuilder<256> StringBuilder;
+	StringBuilder.Append(GetPathName());
+	StringBuilder.Append(TEXT(":"));
+	DefaultEntryPoint.AppendString(StringBuilder);
+	CachedDefaultEntryPoint = FName(StringBuilder.ToView());
+}
+
+FName UAnimNextGraph::GetDefaultEntryPoint() const
+{
+	return CachedDefaultEntryPoint;
+}
\ No newline at end of file
diff --git a/Engine/Plugins/Experimental/AnimNext/Source/AnimNext/Private/Param/AnimNextEditorParam.cpp b/Engine/Plugins/Experimental/AnimNext/Source/AnimNext/Private/Param/AnimNextEditorParam.cpp
new file mode 100644
--- /dev/null
+++ b/Engine/Plugins/Experimental/AnimNext/Source/AnimNext/Private/Param/AnimNextEditorParam.cpp
@@ -0,0 +1,28 @@
+﻿// Copyright Epic Games, Inc. All Rights Reserved.
+
+#include "Param/AnimNextEditorParam.h"
+
+#include "Param/AnimNextParam.h"
+#include "Param/AnimNextParamInstanceIdentifier.h"
+
+FAnimNextEditorParam::FAnimNextEditorParam(FName InName, const FAnimNextParamType& InType, const TInstancedStruct<FAnimNextParamInstanceIdentifier>& InInstanceId)
+	: Name(InName)
+	, Type(InType)
+	, InstanceId(InInstanceId)
+{
+}
+
+#if WITH_EDITORONLY_DATA
+
+FAnimNextEditorParam::FAnimNextEditorParam(const FAnimNextParam& InAnimNextParam)
+	: Name(InAnimNextParam.Name)
+	, Type(InAnimNextParam.Type)
+{
+	if(InAnimNextParam.InstanceIdType != nullptr)
+	{
+		InstanceId.InitializeAsScriptStruct(InAnimNextParam.InstanceIdType);
+		InstanceId.GetMutable().FromName(InAnimNextParam.InstanceId);
+	}
+}
+
+#endif
\ No newline at end of file
diff --git a/Engine/Plugins/Experimental/AnimNext/Source/AnimNext/Private/Param/AnimNextParam.cpp b/Engine/Plugins/Experimental/AnimNext/Source/AnimNext/Private/Param/AnimNextParam.cpp
new file mode 100644
--- /dev/null
+++ b/Engine/Plugins/Experimental/AnimNext/Source/AnimNext/Private/Param/AnimNextParam.cpp
@@ -0,0 +1,26 @@
+﻿// Copyright Epic Games, Inc. All Rights Reserved.
+
+#include "Param/AnimNextParam.h"
+
+#include "InstancedStruct.h"
+#include "Param/AnimNextEditorParam.h"
+#include "Param/AnimNextParamInstanceIdentifier.h"
+
+FAnimNextParam::FAnimNextParam(FName InName, const FAnimNextParamType& InType, const TInstancedStruct<FAnimNextParamInstanceIdentifier>& InInstanceId)
+	: Name(InName)
+	, InstanceId(InInstanceId.IsValid() ? InInstanceId.Get().ToName() : NAME_None)
+#if WITH_EDITORONLY_DATA
+	, InstanceIdType(InInstanceId.GetScriptStruct())
+#endif
+	, Type(InType)
+{
+}
+
+#if WITH_EDITORONLY_DATA
+
+FAnimNextParam::FAnimNextParam(const FAnimNextEditorParam& InEditorParam)
+	: FAnimNextParam(InEditorParam.Name, InEditorParam.Type, InEditorParam.InstanceId)
+{
+}
+
+#endif
\ No newline at end of file
diff --git a/Engine/Plugins/Experimental/AnimNext/Source/AnimNext/Private/Param/ClassProxy.cpp b/Engine/Plugins/Experimental/AnimNext/Source/AnimNext/Private/Param/ClassProxy.cpp
--- a/Engine/Plugins/Experimental/AnimNext/Source/AnimNext/Private/Param/ClassProxy.cpp
+++ b/Engine/Plugins/Experimental/AnimNext/Source/AnimNext/Private/Param/ClassProxy.cpp
@@ -1,62 +1,25 @@
 ﻿// Copyright Epic Games, Inc. All Rights Reserved.
 
-#include "ClassProxy.h"
+#include "Param/ClassProxy.h"
 #include "Logging/StructuredLog.h"
 #include "Param/AnimNextClassExtensionLibrary.h"
 #include "Param/ParamTypeHandle.h"
-#include "EngineLogs.h"
+#include "Param/ParamUtils.h"
 
 namespace UE::AnimNext
 {
 
 FClassProxy::FClassProxy(const UClass* InClass)
-	: Class(InClass)
 {
-	// Add any additional extension libraries that extend this class first
-	{
-		// We allow only native static 'accessor wrapper' functions (one object param and the return value)
-		auto CheckSignature = [](const UFunction* InFunction, const UClass* InProxiedClass)
-		{
-			if(!InFunction->HasAllFunctionFlags(FUNC_BlueprintCallable | FUNC_Static | FUNC_Native))
-			{
-				return false;
-			}
-
-			if(InFunction->NumParms != 2)
-			{
-				return false;
-			}
-		
-			int32 ParamIndex = 0;
-			for(TFieldIterator<FProperty> It(InFunction); It && (It->PropertyFlags & CPF_Parm); ++It, ++ParamIndex)
-			{
-				// Check parameter is an object of the expected class
-				if(ParamIndex == 0)
-				{
-					FObjectProperty* ObjectProperty = CastField<FObjectProperty>(*It);
-					if(ObjectProperty == nullptr)
-					{ 
-						return false;
-					}
-				
-					if(!ObjectProperty->PropertyClass->IsChildOf(InProxiedClass))
-					{
-						return false;
-					}
-				}
+	Refresh(InClass);
+}
 
-				// Check return value
-				if(ParamIndex == 1)
-				{
-					if((It->PropertyFlags & CPF_ReturnParm) == 0)
-					{
-						return false;
-					}
-				}
-			}
-			return true;
-		};
+void FClassProxy::Refresh(const UClass* InClass)
+{
+	Class = InClass;
 
+	// Add any additional extension libraries that extend this class first
+	{
 		TArray<UClass*> Classes;
 		GetDerivedClasses(UAnimNextClassExtensionLibrary::StaticClass(), Classes);
 		for(UClass* ProxyClass : Classes)
@@ -68,29 +31,23 @@ FClassProxy::FClassProxy(const UClass* InClass)
 				for(TFieldIterator<UFunction> It(ProxyClass); It; ++It)
 				{
 					UFunction* Function = *It;
-					if(CheckSignature(Function, ExtendedClass))
+					if(FParamUtils::CanUseFunction(Function))
 					{
 						const FProperty* ReturnProperty = Function->GetReturnProperty();
 						FParamTypeHandle TypeHandle = FParamTypeHandle::FromProperty(ReturnProperty);
 						if(TypeHandle.IsValid())
 						{
-							FString FunctionName = Function->GetName();
-
-							// Remove common prefixes.
-							// TODO: explicit mapping or regex transformations?
-							FunctionName.RemoveFromStart(TEXT("K2_"));
-							FunctionName.RemoveFromStart(TEXT("Get"));
-
 							FClassProxyParameter Parameter;
 							Parameter.AccessType = EClassProxyParameterAccessType::HoistedFunction;
-							Parameter.ClassParameterName = *FunctionName;
+							Parameter.ParameterName = *Function->GetPathName();
 							Parameter.Function = Function;
 							Parameter.Type = TypeHandle.GetType();
 #if WITH_EDITOR
+							Parameter.DisplayName = Function->GetDisplayNameText();
 							Parameter.Tooltip = Function->GetToolTipText();
 							Parameter.bThreadSafe = Function->HasMetaData("BlueprintThreadSafe");
 #endif
-							ParameterNameMap.Add(Parameter.ClassParameterName, Parameters.Add(Parameter));
+							ParameterNameMap.Add(Parameter.ParameterName, Parameters.Add(Parameter));
 						}
 					}
 				}
@@ -102,34 +59,26 @@ FClassProxy::FClassProxy(const UClass* InClass)
 	for(TFieldIterator<UFunction> It(InClass, EFieldIterationFlags::IncludeSuper | EFieldIterationFlags::IncludeInterfaces); It; ++It)
 	{
 		UFunction* Function = *It;
-
-		// We add only 'accessor' functions (no params apart from the return value) that have valid return types
-		const FProperty* ReturnProperty = Function->GetReturnProperty();
-		if(ReturnProperty != nullptr && Function->NumParms == 1 && Function->HasAnyFunctionFlags(FUNC_BlueprintCallable))
+		if(FParamUtils::CanUseFunction(Function))
 		{
+			const FProperty* ReturnProperty = Function->GetReturnProperty();
 			FParamTypeHandle TypeHandle = FParamTypeHandle::FromProperty(ReturnProperty);
 			if(TypeHandle.IsValid())
 			{
-				FString FunctionName = Function->GetName();
-
-				// Remove common prefixes.
-				// TODO: explicit mapping or regex transformations?
-				FunctionName.RemoveFromStart(TEXT("K2_"));
-				FunctionName.RemoveFromStart(TEXT("Get"));
-
-				FName ParameterName = *FunctionName;
+				FName ParameterName = *Function->GetPathName();
 				if(!ParameterNameMap.Contains(ParameterName))
 				{
 					FClassProxyParameter Parameter;
 					Parameter.AccessType = EClassProxyParameterAccessType::AccessorFunction;
-					Parameter.ClassParameterName = *FunctionName;
+					Parameter.ParameterName = ParameterName;
 					Parameter.Function = Function;
 					Parameter.Type = TypeHandle.GetType();
 #if WITH_EDITOR
+					Parameter.DisplayName = Function->GetDisplayNameText();
 					Parameter.Tooltip = Function->GetToolTipText();
 					Parameter.bThreadSafe = Function->HasMetaData("BlueprintThreadSafe");
 #endif
-					ParameterNameMap.Add(Parameter.ClassParameterName, Parameters.Add(Parameter));
+					ParameterNameMap.Add(Parameter.ParameterName, Parameters.Add(Parameter));
 				}
 			}
 		}
@@ -139,30 +88,25 @@ FClassProxy::FClassProxy(const UClass* InClass)
 	for(TFieldIterator<FProperty> It(InClass, EFieldIterationFlags::IncludeSuper | EFieldIterationFlags::IncludeInterfaces); It; ++It)
 	{
 		FProperty* Property = *It;
-		if(Property->HasAnyPropertyFlags(CPF_Edit | CPF_BlueprintVisible))
+		if(FParamUtils::CanUseProperty(Property))
 		{
 			FParamTypeHandle TypeHandle = FParamTypeHandle::FromProperty(Property);
 			if(TypeHandle.IsValid())
 			{
-				FString PropertyName = Property->GetName();
-				if(TypeHandle.GetParameterType() == FParamTypeHandle::EParamType::Bool)
-				{
-					PropertyName.RemoveFromStart(TEXT("b"));
-				}
-
-				FName ParameterName = *PropertyName;
+				FName ParameterName = *Property->GetPathName();
 				if(!ParameterNameMap.Contains(ParameterName))
 				{
 					FClassProxyParameter Parameter;
 					Parameter.AccessType = EClassProxyParameterAccessType::Property;
-					Parameter.ClassParameterName = ParameterName;
+					Parameter.ParameterName = ParameterName;
 					Parameter.Property = Property;
 					Parameter.Type = TypeHandle.GetType();
 #if WITH_EDITOR
+					Parameter.DisplayName = Property->GetDisplayNameText();
 					Parameter.Tooltip = Property->GetToolTipText();
 					Parameter.bThreadSafe = false;
 #endif
-					ParameterNameMap.Add(Parameter.ClassParameterName, Parameters.Add(Parameter));
+					ParameterNameMap.Add(Parameter.ParameterName, Parameters.Add(Parameter));
 				}
 			}
 		}
diff --git a/Engine/Plugins/Experimental/AnimNext/Source/AnimNext/Private/Param/ObjectProxy.cpp b/Engine/Plugins/Experimental/AnimNext/Source/AnimNext/Private/Param/ObjectProxy.cpp
--- a/Engine/Plugins/Experimental/AnimNext/Source/AnimNext/Private/Param/ObjectProxy.cpp
+++ b/Engine/Plugins/Experimental/AnimNext/Source/AnimNext/Private/Param/ObjectProxy.cpp
@@ -1,18 +1,16 @@
 // Copyright Epic Games, Inc. All Rights Reserved.
 
 #include "ObjectProxy.h"
-
-#include "ClassProxy.h"
-#include "ObjectAccessor.h"
+#include "Param/ClassProxy.h"
 #include "Param/ParamStack.h"
 
 namespace UE::AnimNext
 {
 
-FObjectProxy::FObjectProxy(UObject* InObject, const TSharedRef<FObjectAccessor>& InObjectAccessor)
+FObjectProxy::FObjectProxy(const UObject* InObject, FStringView InObjectLocatorPath, const TSharedRef<FClassProxy>& InClassProxy)
 	: Object(InObject)
-	, ObjectAccessor(InObjectAccessor)
-	, RootParameterName(ObjectAccessor->AccessorName)
+	, ClassProxy(InClassProxy)
+	, RootParameterName(InObjectLocatorPath)
 {
 	// Always supply the root parameter in index 0
 	ParameterCache.AddProperty(RootParameterName, EPropertyBagPropertyType::Object, InObject->GetClass());
@@ -24,9 +22,9 @@ void FObjectProxy::Update(float DeltaTime)
 	const UPropertyBag* PropertyBag = ParameterCache.GetPropertyBagStruct();
 	TConstArrayView<FPropertyBagPropertyDesc> PropertyDescs = PropertyBag->GetPropertyDescs();
 	uint8* StructData = ParameterCache.GetMutableValue().GetMemory();
-	UObject* ResolvedObject = Object.Get();
-	PropertyDescs[0].CachedProperty->SetValue_InContainer(StructData, &ResolvedObject);
-	if(ResolvedObject)
+	UObject* NonConstObject = const_cast<UObject*>(Object.Get());
+	PropertyDescs[0].CachedProperty->SetValue_InContainer(StructData, &Object);
+	if(Object != nullptr)
 	{
 		for(int32 ParameterIndex = 0; ParameterIndex < ParametersToUpdate.Num(); ++ParameterIndex)
 		{
@@ -42,7 +40,7 @@ void FObjectProxy::Update(float DeltaTime)
 					checkSlow(SourceProperty);
 					checkSlow(SourceProperty->GetClass() == ResultProperty->GetClass());
 
-					const void* SourceBuffer = SourceProperty->ContainerPtrToValuePtr<const void>(ResolvedObject);
+					const void* SourceBuffer = SourceProperty->ContainerPtrToValuePtr<const void>(Object);
 					SourceProperty->CopyCompleteValue(ResultBuffer, SourceBuffer);
 					break;
 				}
@@ -50,10 +48,10 @@ void FObjectProxy::Update(float DeltaTime)
 				{
 					UFunction* Function = ParameterToUpdate.GetFunction();
 					checkSlow(Function);
-					checkSlow(ResolvedObject->GetClass()->IsChildOf(Function->GetOuterUClass()));
+					checkSlow(Object->GetClass()->IsChildOf(Function->GetOuterUClass()));
 
-					FFrame Stack(ResolvedObject, Function, nullptr, nullptr, Function->ChildProperties);
-					Function->Invoke(ResolvedObject, Stack, ResultBuffer);
+					FFrame Stack(NonConstObject, Function, nullptr, nullptr, Function->ChildProperties);
+					Function->Invoke(NonConstObject, Stack, ResultBuffer);
 					break;
 				}
 			case EClassProxyParameterAccessType::HoistedFunction:
@@ -62,12 +60,12 @@ void FObjectProxy::Update(float DeltaTime)
 					checkSlow(Function);
 
 					const FObjectProperty* ObjectProperty = CastFieldChecked<FObjectProperty>(PropertyDescs[0].CachedProperty);
-					check(ResolvedObject->GetClass()->IsChildOf(ObjectProperty->PropertyClass));
+					check(Object->GetClass()->IsChildOf(ObjectProperty->PropertyClass));
 					UObject** ObjectBuffer = ObjectProperty->ContainerPtrToValuePtr<UObject*>(StructData);
-					*ObjectBuffer = ResolvedObject;
+					*ObjectBuffer = NonConstObject;
 
-					FFrame Stack(ResolvedObject, Function, ObjectBuffer, nullptr, Function->ChildProperties);
-					Function->Invoke(ResolvedObject, Stack, ResultBuffer);
+					FFrame Stack(NonConstObject, Function, ObjectBuffer, nullptr, Function->ChildProperties);
+					Function->Invoke(NonConstObject, Stack, ResultBuffer);
 					break;
 				}
 			default:
@@ -91,9 +89,9 @@ void FObjectProxy::RequestParameterCache(TConstArrayView<FName> InParameterNames
 	{
 		if(!ParameterNameMap.Contains(ParameterName))
 		{
-			if(const int32* ParameterIndexPtr = ObjectAccessor->RemappedParametersMap.Find(ParameterName))
+			if(const int32* ParameterIndexPtr = ClassProxy->ParameterNameMap.Find(ParameterName))
 			{
-				const FClassProxyParameter& ClassProxyParameter = ObjectAccessor->ClassProxy->Parameters[*ParameterIndexPtr];
+				const FClassProxyParameter& ClassProxyParameter = ClassProxy->Parameters[*ParameterIndexPtr];
 
 				FAnimNextObjectProxyParameter& NewParameterToUpdate = ParametersToUpdate.AddDefaulted_GetRef();
 				NewParameterToUpdate.AccessType = ClassProxyParameter.AccessType;
@@ -111,11 +109,12 @@ void FObjectProxy::RequestParameterCache(TConstArrayView<FName> InParameterNames
 	ParameterCache.AddProperties(PropertyDescsToAdd);
 
 	// Recreate layer handle as layout has changed
-	LayerHandle = FParamStack::MakeReferenceLayer(ParameterCache);
+	LayerHandle = FParamStack::MakeReferenceLayer(RootParameterName, ParameterCache);
 }
 
 void FObjectProxy::AddReferencedObjects(FReferenceCollector& Collector)
 {
+	Collector.AddReferencedObject(Object);
 	ParameterCache.AddStructReferencedObjects(Collector);
 }
 
diff --git a/Engine/Plugins/Experimental/AnimNext/Source/AnimNext/Private/Param/ObjectProxy.h b/Engine/Plugins/Experimental/AnimNext/Source/AnimNext/Private/Param/ObjectProxy.h
--- a/Engine/Plugins/Experimental/AnimNext/Source/AnimNext/Private/Param/ObjectProxy.h
+++ b/Engine/Plugins/Experimental/AnimNext/Source/AnimNext/Private/Param/ObjectProxy.h
@@ -2,12 +2,15 @@
 
 #pragma once
 
+#include "Param/ClassProxy.h"
 #include "Param/IParameterSource.h"
 #include "PropertyBag.h"
+#include "UObject/GCObject.h"
+
+class UAnimNextParameterSchema;
 
 namespace UE::AnimNext
 {
-	struct FObjectAccessor;
 	enum class EClassProxyParameterAccessType : int32;
 }
 
@@ -63,7 +66,7 @@ struct FObjectProxy : public IParameterSource
 {
 	FObjectProxy() = delete;
 
-	FObjectProxy(UObject* InObject, const TSharedRef<FObjectAccessor>& InObjectAccessor);
+	FObjectProxy(const UObject* InObject, FStringView InObjectLocatorPath, const TSharedRef<FClassProxy>& InClassProxy);
 
 	// IParameterSource interface
 	virtual void Update(float DeltaTime) override;
@@ -74,7 +77,7 @@ struct FObjectProxy : public IParameterSource
 	void RequestParameterCache(TConstArrayView<FName> InParameterNames);
 
 	// The object that this proxy wraps
-	TWeakObjectPtr<UObject> Object;
+	TObjectPtr<const UObject> Object;
 
 	// Cache of properties, fetched from Object
 	FInstancedPropertyBag ParameterCache;
@@ -88,8 +91,8 @@ struct FObjectProxy : public IParameterSource
 	// Map of parameter name to index in ParametersToUpdate array
 	TMap<FName, int32> ParameterNameMap;
 
-	// Object accessor used to lookup parameters
-	TSharedRef<FObjectAccessor> ObjectAccessor;
+	// Class proxy defining the 'layout' of the object
+	TSharedRef<FClassProxy> ClassProxy;
 
 	// The name of the root parameter
 	FName RootParameterName;
diff --git a/Engine/Plugins/Experimental/AnimNext/Source/AnimNext/Private/Param/ObjectProxyFactory.cpp b/Engine/Plugins/Experimental/AnimNext/Source/AnimNext/Private/Param/ObjectProxyFactory.cpp
--- a/Engine/Plugins/Experimental/AnimNext/Source/AnimNext/Private/Param/ObjectProxyFactory.cpp
+++ b/Engine/Plugins/Experimental/AnimNext/Source/AnimNext/Private/Param/ObjectProxyFactory.cpp
@@ -2,231 +2,77 @@
 
 #include "Param/ObjectProxyFactory.h"
 
-#include "AnimNextConfig.h"
-#include "ClassProxy.h"
-#include "ObjectAccessor.h"
+#include "Param/ClassProxy.h"
 #include "ObjectProxy.h"
-#include "Components/SkeletalMeshComponent.h"
-#include "Logging/StructuredLog.h"
-#include "Param/AnimNextObjectAccessorConfig.h"
-#include "RigVMCore/RigVMRegistry.h"
-#include "Features/IModularFeatures.h"
-#include "Engine/World.h"
+#include "Param/AnimNextParamUniversalObjectLocator.h"
 
 namespace UE::AnimNext
 {
 
-static FObjectProxyFactory GObjectProxyFactory;
-
-void FObjectProxyFactory::Init()
-{
-	Refresh();
-}
-
-void FObjectProxyFactory::Destroy()
-{
-	GObjectProxyFactory.Reset();
-
-	IModularFeatures::Get().UnregisterModularFeature(IParameterSourceFactory::FeatureName, &GObjectProxyFactory);
-}
-
-void FObjectProxyFactory::Reset()
-{
-	ClassMap.Reset();
-	ObjectAccessors.Reset();
-	ParameterMap.Reset();
-}
-
-void FObjectProxyFactory::Refresh()
+FObjectProxyFactory::FObjectProxyFactory()
 {
-	IModularFeatures::Get().UnregisterModularFeature(IParameterSourceFactory::FeatureName, &GObjectProxyFactory);
-
-	GObjectProxyFactory.Reset();
-
-	for(const FAnimNextObjectAccessorConfig& AccessorConfig : GetDefault<UAnimNextConfig>()->ExposedClasses)
-	{
-		UClass* Class = AccessorConfig.Class.LoadSynchronous();
-		if(Class && AccessorConfig.AccessorName != NAME_None)
-		{
-			GObjectProxyFactory.RegisterObjectAccessor(AccessorConfig.AccessorName, Class);
-		}
-	}
-
-	// Refresh RigVM registry as we have updated allowed types
-	FRigVMRegistry::Get().RefreshEngineTypes();
-
-	IModularFeatures::Get().RegisterModularFeature(IParameterSourceFactory::FeatureName, &GObjectProxyFactory);
-}
-
-void FObjectProxyFactory::ForEachSource(TFunctionRef<void(FName)> InFunction) const
-{
-	for(const TPair<FName, TSharedPtr<FObjectAccessor>>& ObjectAccessorPair : ObjectAccessors)
-	{
-		InFunction(ObjectAccessorPair.Key);
-	}
-}
-
-TUniquePtr<IParameterSource> FObjectProxyFactory::CreateParameterSource(const FExternalParameterContext& InContext, FName InSourceName, TConstArrayView<FName> InRequiredParameters) const
-{
-	using namespace UE::AnimNext;
-
-	if(TSharedPtr<FObjectAccessor> ObjectAccessor = FindObjectAccessor(InSourceName))
-	{
-		if(UObject* Object = ObjectAccessor->Function(InContext))
-		{
-			TUniquePtr<FObjectProxy> ObjectProxy = MakeUnique<FObjectProxy>(Object, ObjectAccessor.ToSharedRef());
-			ObjectProxy->RequestParameterCache(InRequiredParameters);
-
-			return MoveTemp(ObjectProxy);
-		}
-	}
-	return nullptr;
-}
-
 #if WITH_EDITOR
-
-bool FObjectProxyFactory::FindParameterInfo(FName InParameterName, FParameterInfo& OutInfo) const
-{
-	if(const TWeakPtr<FObjectAccessor>* WeakAccessor = ParameterMap.Find(InParameterName))
+	OnObjectsReinstancedHandle = FCoreUObjectDelegates::OnObjectsReinstanced.AddLambda([this](const FCoreUObjectDelegates::FReplacementObjectMap& InMap)
 	{
-		if(TSharedPtr<FObjectAccessor> ObjectAccessor = WeakAccessor->Pin())
-		{
-			if(int32* ParamIndexPtr = ObjectAccessor->RemappedParametersMap.Find(InParameterName))
-			{
-				OutInfo.Type = ObjectAccessor->ClassProxy->Parameters[*ParamIndexPtr].Type;
-				OutInfo.Tooltip = ObjectAccessor->ClassProxy->Parameters[*ParamIndexPtr].Tooltip;
-				OutInfo.bThreadSafe = ObjectAccessor->ClassProxy->Parameters[*ParamIndexPtr].bThreadSafe;
-				return true;
-			}
-		}
-	}
-	return false;
+		Refresh();
+	});
+#endif
 }
 
-void FObjectProxyFactory::ForEachParameter(FName InSourceName, TFunctionRef<void(FName, const FParameterInfo&)> InFunction) const
+FObjectProxyFactory::~FObjectProxyFactory()
 {
-	if(TSharedPtr<FObjectAccessor> ObjectAccessor = FindObjectAccessor(InSourceName))
-	{
-		TSharedPtr<FClassProxy> ClassProxy = ObjectAccessor->ClassProxy;
-		check(ClassProxy->Parameters.Num() == ObjectAccessor->RemappedParameters.Num());
-		for(int32 ParameterIndex = 0, ParameterCount = ObjectAccessor->RemappedParameters.Num(); ParameterIndex < ParameterCount; ++ParameterIndex)
-		{
-			FParameterInfo ParameterInfo;
-			ParameterInfo.Type = ClassProxy->Parameters[ParameterIndex].Type;
-			ParameterInfo.Tooltip = ClassProxy->Parameters[ParameterIndex].Tooltip;
-			ParameterInfo.bThreadSafe = ClassProxy->Parameters[ParameterIndex].bThreadSafe;
-			InFunction(ObjectAccessor->RemappedParameters[ParameterIndex], ParameterInfo);
-		}
-	}
-}
-
+#if WITH_EDITOR
+	FCoreUObjectDelegates::OnObjectsReinstanced.Remove(OnObjectsReinstancedHandle);
 #endif
-
-void FObjectProxyFactory::RegisterObjectAccessor(FName InAccessorName, const UClass* InTargetClass)
-{
-	check(IsInGameThread());
-	check(InAccessorName != NAME_None);
-	check(InTargetClass != nullptr);
-
-	if(InTargetClass->IsChildOf<UActorComponent>())
-	{
-		RegisterActorComponentAccessor(InAccessorName, const_cast<UClass*>(InTargetClass));
-	}
-	else if(InTargetClass->IsChildOf<AActor>())
-	{
-		RegisterActorAccessor(InAccessorName, const_cast<UClass*>(InTargetClass));
-	}
-	else if(InTargetClass->IsChildOf<UWorld>())
-	{
-		RegisterObjectAccessor(InAccessorName, const_cast<UClass*>(InTargetClass), [](const FExternalParameterContext& InContext) -> UObject*
-		{
-			return InContext.Object->GetWorld();
-		});
-	}
-	else 
-	{
-		UE_LOGFMT(LogAnimation, Warning, "Class {ClassName} passed to FObjectProxyFactory::RegisterObjectAccessor that does not have a built-in accessor, please use the custom accessor overload.", InTargetClass->GetFName());
-	}
 }
 
-void FObjectProxyFactory::RegisterActorAccessor(FName InAccessorName, TSubclassOf<AActor> InTargetClass)
+void FObjectProxyFactory::Refresh()
 {
-	RegisterObjectAccessor(InAccessorName, InTargetClass.Get(), [InTargetClass](const FExternalParameterContext& InContext) -> UObject*
+	// Remove any invalid classes, refresh any valid ones
+	for(auto It = ClassMap.CreateIterator(); It; ++It)
 	{
-		if(UActorComponent* ContextComponent = Cast<UActorComponent>(InContext.Object))
+		const UClass* Class = It.Key().ResolveObjectPtr();
+		if(Class == nullptr || Class->HasAllClassFlags(CLASS_NewerVersionExists))
 		{
-			if(ContextComponent->GetOwner()->GetClass()->IsChildOf(InTargetClass.Get()))
-			{
-				return ContextComponent->GetOwner();
-			}
+			It.RemoveCurrent();
 		}
-
-		return nullptr;
-	});
-}
-
-void FObjectProxyFactory::RegisterActorComponentAccessor(FName InAccessorName, TSubclassOf<UActorComponent> InTargetClass)
-{
-	RegisterObjectAccessor(InAccessorName, InTargetClass.Get(), [InTargetClass](const FExternalParameterContext& InContext) -> UObject*
-	{
-		if(UActorComponent* ContextComponent = Cast<UActorComponent>(InContext.Object))
+		else
 		{
-			return ContextComponent->GetOwner()->FindComponentByClass(InTargetClass.Get());
+			It.Value()->Refresh(Class);
 		}
-
-		return nullptr;
-	});
+	}
 }
 
-void FObjectProxyFactory::RegisterObjectAccessor(FName InAccessorName, const UClass* InTargetClass, FObjectAccessorFunction&& InFunction)
+TUniquePtr<IParameterSource> FObjectProxyFactory::CreateParameterSource(const FParameterSourceContext& InContext, const TInstancedStruct<FAnimNextParamInstanceIdentifier>& InInstanceId, TConstArrayView<FName> InRequiredParameters) const
 {
-	check(IsInGameThread());
-	check(InAccessorName != NAME_None);
-	check(InTargetClass != nullptr);
-
-	// Ensure we can use the type as a pin in RigVM graphs 
-	FRigVMRegistry::Get().RegisterObjectTypes( {{ const_cast<UClass*>(InTargetClass), FRigVMRegistry::ERegisterObjectOperation::Class } });
+	using namespace UE::UniversalObjectLocator;
 
-	UE_MT_SCOPED_WRITE_ACCESS(ObjectAccessorsAccessDetector);
-	if(!ObjectAccessors.Contains(InAccessorName))
+	if(const FAnimNextParamUniversalObjectLocator* Locator = InInstanceId.GetPtr<FAnimNextParamUniversalObjectLocator>())
 	{
-		TSharedRef<FClassProxy> ClassProxy = FindOrCreateClassProxy(InTargetClass);
-		TSharedRef<FObjectAccessor> ObjectAccessor = MakeShared<FObjectAccessor>(InAccessorName, MoveTemp(InFunction), ClassProxy);
-		ObjectAccessors.Add(InAccessorName, ObjectAccessor);
+		// Resolve locator
+		FResolveParams ResolveParams(const_cast<UObject*>(InContext.Object));
+		FResolveResult Result = Locator->Locator.Resolve(ResolveParams);
+		FResolveResultData ResultData = Result.SyncGet();
 
-		// Add object accessor to parameter lookup
-		ParameterMap.Add(InAccessorName, ObjectAccessor);
+		// We shouldn't be loading as part of this call - if this hits we need to consider loading objects up front somehow
+		check(!(ResultData.Flags.bWasLoaded || ResultData.Flags.bWasLoadedIndirectly));	
 
-		// Add sub-parameters to parameter lookup
-		for(FName RemappedParameter : ObjectAccessor->RemappedParameters)
+		if(ResultData.Object)
 		{
-			ParameterMap.Add(RemappedParameter, ObjectAccessor);
-		}
-	}
-}
-
-void FObjectProxyFactory::UnregisterObjectAccessor(FName InAccessorName)
-{
-	check(IsInGameThread());
+			TStringBuilder<256> ScopeAsString;
+			Locator->Locator.ToString(ScopeAsString);
+			TSharedRef<FClassProxy> ClassProxy = FindOrCreateClassProxy(ResultData.Object->GetClass());
+			TUniquePtr<FObjectProxy> ObjectProxy = MakeUnique<FObjectProxy>(ResultData.Object, ScopeAsString, ClassProxy);
+			ObjectProxy->RequestParameterCache(InRequiredParameters);
 
-	UE_MT_SCOPED_WRITE_ACCESS(ObjectAccessorsAccessDetector);
-	if(const TSharedPtr<FObjectAccessor>* FoundAccessor = ObjectAccessors.Find(InAccessorName))
-	{
-		// Remove sub-parameters
-		for(FName RemappedParameter : (*FoundAccessor)->RemappedParameters)
-		{
-			ParameterMap.Remove(RemappedParameter);
+			return MoveTemp(ObjectProxy);
 		}
-
-		// Remove object accessor
-		ParameterMap.Remove(InAccessorName);
-
-		// Remove accessor itself
-		ObjectAccessors.Remove(InAccessorName);
 	}
+
+	return nullptr;
 }
 
-TSharedRef<FClassProxy> FObjectProxyFactory::FindOrCreateClassProxy(const UClass* InClass)
+TSharedRef<FClassProxy> FObjectProxyFactory::FindOrCreateClassProxy(const UClass* InClass) const 
 {
 	UE_MT_SCOPED_WRITE_ACCESS(ObjectAccessorsAccessDetector);
 	if(TSharedPtr<FClassProxy>* FoundProxy = ClassMap.Find(InClass))
@@ -237,10 +83,4 @@ TSharedRef<FClassProxy> FObjectProxyFactory::FindOrCreateClassProxy(const UClass
 	return ClassMap.Add(InClass, MakeShared<FClassProxy>(InClass)).ToSharedRef();
 }
 
-TSharedPtr<FObjectAccessor> FObjectProxyFactory::FindObjectAccessor(FName InAccessorName) const
-{
-	UE_MT_SCOPED_READ_ACCESS(ObjectAccessorsAccessDetector);
-	return ObjectAccessors.FindRef(InAccessorName);
-}
-
 }
diff --git a/Engine/Plugins/Experimental/AnimNext/Source/AnimNext/Private/Param/ParamStack.cpp b/Engine/Plugins/Experimental/AnimNext/Source/AnimNext/Private/Param/ParamStack.cpp
--- a/Engine/Plugins/Experimental/AnimNext/Source/AnimNext/Private/Param/ParamStack.cpp
+++ b/Engine/Plugins/Experimental/AnimNext/Source/AnimNext/Private/Param/ParamStack.cpp
@@ -25,7 +25,7 @@ struct FInstancedPropertyBagLayer : FParamStackLayer, FGCObject
 {
 	FInstancedPropertyBagLayer() = delete;
 
-	explicit FInstancedPropertyBagLayer(FInstancedPropertyBag& InPropertyBag, bool bInMutable)
+	explicit FInstancedPropertyBagLayer(FName InInstanceId, FInstancedPropertyBag& InPropertyBag, bool bInMutable)
 		: FParamStackLayer(InPropertyBag.GetPropertyBagStruct() ? InPropertyBag.GetPropertyBagStruct()->GetPropertyDescs().Num() : 0)
 	{
 		if (const UPropertyBag* PropertyBagStruct = InPropertyBag.GetPropertyBagStruct())
@@ -36,7 +36,7 @@ struct FInstancedPropertyBagLayer : FParamStackLayer, FGCObject
 			for (uint32 DescIndex = 0; DescIndex < static_cast<uint32>(Descs.Num()); ++DescIndex)
 			{
 				const FPropertyBagPropertyDesc& Desc = Descs[DescIndex];
-				const FParamId ParamId(Desc.Name);
+				const FParamId ParamId(Desc.Name, InInstanceId);
 				uint8* DataPtr = StructView.GetMemory() + Desc.CachedProperty->GetOffset_ForInternal();
 				uint32 ParamIndex = Params.Emplace(ParamId, FParamTypeHandle::FromPropertyBagPropertyDesc(Desc), TArrayView<uint8>(DataPtr, Desc.CachedProperty->GetSize()), true, true);
 				HashTable.Add(ParamId.GetHash(), ParamIndex);
@@ -57,7 +57,7 @@ struct FInstancedPropertyBagValueLayer : FInstancedPropertyBagLayer
 	FInstancedPropertyBagValueLayer() = delete;
 
 	explicit FInstancedPropertyBagValueLayer(FInstancedPropertyBag&& InPropertyBag, bool bInMutable)
-		: FInstancedPropertyBagLayer(InPropertyBag, bInMutable)
+		: FInstancedPropertyBagLayer(NAME_None, InPropertyBag, bInMutable)
 		, PropertyBag(MoveTemp(InPropertyBag))
 	{}
 
@@ -81,8 +81,8 @@ struct FInstancedPropertyBagReferenceLayer : FInstancedPropertyBagLayer
 {
 	FInstancedPropertyBagReferenceLayer() = delete;
 
-	explicit FInstancedPropertyBagReferenceLayer(FInstancedPropertyBag& InPropertyBag, bool bInMutable)
-		: FInstancedPropertyBagLayer(InPropertyBag, bInMutable)
+	explicit FInstancedPropertyBagReferenceLayer(FName InInstanceId, FInstancedPropertyBag& InPropertyBag, bool bInMutable)
+		: FInstancedPropertyBagLayer(InInstanceId, InPropertyBag, bInMutable)
 		, PropertyBag(InPropertyBag)
 	{}
 
@@ -309,9 +309,9 @@ FParamStackLayerHandle FParamStack::MakeValueLayer(const FInstancedPropertyBag&
 	return FParamStackLayerHandle(MoveTemp(Layer));
 }
 
-FParamStackLayerHandle FParamStack::MakeReferenceLayer(FInstancedPropertyBag& InPropertyBag)
+FParamStackLayerHandle FParamStack::MakeReferenceLayer(FName InInstanceId, FInstancedPropertyBag& InPropertyBag)
 {
-	TUniquePtr<FParamStackLayer> Layer = MakeUnique<FInstancedPropertyBagReferenceLayer>(InPropertyBag, true);
+	TUniquePtr<FParamStackLayer> Layer = MakeUnique<FInstancedPropertyBagReferenceLayer>(InInstanceId, InPropertyBag, true);
 	return FParamStackLayerHandle(MoveTemp(Layer));
 }
 
@@ -519,7 +519,7 @@ const Private::FParamEntry* FParamStack::FindParam(FParamId InId) const
 {
 	for(uint32 Index = LayerHash.First(InId.GetHash()); LayerHash.IsValid(Index); Index = LayerHash.Next(Index))
 	{
-		if (EntryStack[Index]->GetName() == InId.GetName())
+		if (EntryStack[Index]->GetName() == InId.GetName() && EntryStack[Index]->GetInstanceId() == InId.GetInstanceId())
 		{
 			return EntryStack[Index];
 		}
diff --git a/Engine/Plugins/Experimental/AnimNext/Source/AnimNext/Private/Param/ParamStackLayer.h b/Engine/Plugins/Experimental/AnimNext/Source/AnimNext/Private/Param/ParamStackLayer.h
--- a/Engine/Plugins/Experimental/AnimNext/Source/AnimNext/Private/Param/ParamStackLayer.h
+++ b/Engine/Plugins/Experimental/AnimNext/Source/AnimNext/Private/Param/ParamStackLayer.h
@@ -44,8 +44,7 @@ struct FParamStackLayer
 
 	ANIMNEXT_API FParamResult GetMutableParamData(FParamId InId, FParamTypeHandle InTypeHandle, TArrayView<uint8>& OutParamData);
 	ANIMNEXT_API FParamResult GetMutableParamData(FParamId InId, FParamTypeHandle InTypeHandle, TArrayView<uint8>& OutParamData, FParamTypeHandle& OutParamTypeHandle, FParamCompatibility InRequiredCompatibility = FParamCompatibility::Equal());
-
-	ANIMNEXT_API virtual UObject* AsUObject() { return nullptr; }
+	
 	ANIMNEXT_API virtual FInstancedPropertyBag* AsInstancedPropertyBag() { return nullptr; }
 
 	const Private::FParamEntry* FindEntry(FParamId InId) const;
diff --git a/Engine/Plugins/Experimental/AnimNext/Source/AnimNext/Private/Param/ParamStackLayerHandle.cpp b/Engine/Plugins/Experimental/AnimNext/Source/AnimNext/Private/Param/ParamStackLayerHandle.cpp
--- a/Engine/Plugins/Experimental/AnimNext/Source/AnimNext/Private/Param/ParamStackLayerHandle.cpp
+++ b/Engine/Plugins/Experimental/AnimNext/Source/AnimNext/Private/Param/ParamStackLayerHandle.cpp
@@ -84,11 +84,6 @@ FParamResult FParamStackLayerHandle::SetValuesInternal(TConstArrayView<Private::
 	return EParamResult::NotInScope;
 }
 
-UObject* FParamStackLayerHandle::GetUObjectFromLayer() const
-{
-	return Layer->AsUObject();
-}
-
 FInstancedPropertyBag* FParamStackLayerHandle::GetInstancedPropertyBagFromLayer() const
 {
 	return Layer->AsInstancedPropertyBag();
diff --git a/Engine/Plugins/Experimental/AnimNext/Source/AnimNext/Private/Param/ParamType.cpp b/Engine/Plugins/Experimental/AnimNext/Source/AnimNext/Private/Param/ParamType.cpp
--- a/Engine/Plugins/Experimental/AnimNext/Source/AnimNext/Private/Param/ParamType.cpp
+++ b/Engine/Plugins/Experimental/AnimNext/Source/AnimNext/Private/Param/ParamType.cpp
@@ -20,6 +20,77 @@ FAnimNextParamType::FAnimNextParamType(EValueType InValueType, EContainerType In
 {
 }
 
+FRigVMTemplateArgumentType FAnimNextParamType::ToRigVMTemplateArgument() const
+{
+	auto GetTypeInternal = [this]()
+	{
+		switch(ValueType)
+		{
+		case EValueType::None:
+			return FRigVMTemplateArgumentType();
+		case EValueType::Bool:
+			return FRigVMTemplateArgumentType(RigVMTypeUtils::BoolTypeName);
+		case EValueType::Byte:
+			return FRigVMTemplateArgumentType(RigVMTypeUtils::UInt8Type);
+		case EValueType::Int32:
+			return FRigVMTemplateArgumentType(RigVMTypeUtils::Int32Type);
+		case EValueType::Int64:
+			return FRigVMTemplateArgumentType(RigVMTypeUtils::Int64Type);
+		case EValueType::Float:
+			return FRigVMTemplateArgumentType(RigVMTypeUtils::FloatType);
+		case EValueType::Double:
+			return FRigVMTemplateArgumentType(RigVMTypeUtils::DoubleType);
+		case EValueType::Name:
+			return FRigVMTemplateArgumentType(RigVMTypeUtils::FNameType);
+		case EValueType::String:
+			return FRigVMTemplateArgumentType(RigVMTypeUtils::FStringType);
+		case EValueType::Text:
+			return FRigVMTemplateArgumentType(RigVMTypeUtils::FTextType);
+		case EValueType::Enum:
+			if(const UEnum* Enum = Cast<UEnum>(ValueTypeObject.Get()))
+			{
+				return FRigVMTemplateArgumentType(const_cast<UEnum*>(Enum));
+			}
+		case EValueType::Struct:
+			if(const UScriptStruct* Struct = Cast<UScriptStruct>(ValueTypeObject.Get()))
+			{
+				return FRigVMTemplateArgumentType(const_cast<UScriptStruct*>(Struct));
+			}
+			break;
+		case EValueType::Object:
+			if(const UClass* Class = Cast<UClass>(ValueTypeObject.Get()))
+			{
+				return FRigVMTemplateArgumentType(const_cast<UClass*>(Class));
+			}
+			break;
+		case EValueType::SoftObject:
+			return FRigVMTemplateArgumentType();
+		case EValueType::Class:
+			if(const UClass* Class = Cast<UClass>(ValueTypeObject.Get()))
+			{
+				return FRigVMTemplateArgumentType(const_cast<UClass*>(Class), RigVMTypeUtils::EClassArgType::AsClass);
+			}
+			break;
+		case EValueType::SoftClass:
+			return FRigVMTemplateArgumentType();
+		default:
+			break;
+		}
+
+		return FRigVMTemplateArgumentType();
+	};
+
+	switch(ContainerType)
+	{
+	case EContainerType::None:
+		return GetTypeInternal();
+	case EContainerType::Array:
+		return GetTypeInternal().ConvertToArray();
+	default:
+		return FRigVMTemplateArgumentType();
+	}
+}
+
 FAnimNextParamType FAnimNextParamType::FromRigVMTemplateArgument(const FRigVMTemplateArgumentType& RigVMType)
 {
 	FAnimNextParamType Type;	
@@ -40,6 +111,10 @@ FAnimNextParamType FAnimNextParamType::FromRigVMTemplateArgument(const FRigVMTem
 	{
 		Type.ValueType = EPropertyBagPropertyType::Bool;
 	}
+	else if (CPPType == RigVMTypeUtils::UInt8Type)
+	{
+		Type.ValueType = EPropertyBagPropertyType::Byte;
+	}
 	else if (CPPType == RigVMTypeUtils::Int32TypeName || CPPType == IntTypeName)
 	{
 		Type.ValueType = EPropertyBagPropertyType::Int32;
diff --git a/Engine/Plugins/Experimental/AnimNext/Source/AnimNext/Private/Param/ParamUtils.cpp b/Engine/Plugins/Experimental/AnimNext/Source/AnimNext/Private/Param/ParamUtils.cpp
--- a/Engine/Plugins/Experimental/AnimNext/Source/AnimNext/Private/Param/ParamUtils.cpp
+++ b/Engine/Plugins/Experimental/AnimNext/Source/AnimNext/Private/Param/ParamUtils.cpp
@@ -2,9 +2,11 @@
 
 #include "Param/ParamUtils.h"
 
+#include "UniversalObjectLocator.h"
 #include "Component/AnimNextMeshComponent.h"
 #include "GameFramework/CharacterMovementComponent.h"
 #include "Animation/AnimSequence.h"
+#include "Param/AnimNextClassExtensionLibrary.h"
 #include "Param/ParamType.h"
 #include "Param/ParamTypeHandle.h"
 #include "Param/ParamCompatibility.h"
@@ -407,4 +409,98 @@ FParamCompatibility FParamUtils::GetCompatibility(const FAnimNextParamType& InLH
 	return GetCompatibility(InLHS.GetHandle(), InRHS.GetHandle());
 }
 
+bool FParamUtils::CanUseFunction(const UFunction* InFunction)
+{
+	UClass* FunctionClass = InFunction->GetOuterUClass();
+	if(FunctionClass->IsChildOf(UAnimNextClassExtensionLibrary::StaticClass()))
+	{
+		UAnimNextClassExtensionLibrary* CDO = FunctionClass->GetDefaultObject<UAnimNextClassExtensionLibrary>();
+		const UClass* ExtendedClass = CDO->GetSupportedClass();
+		
+		// Check 'hoisted' functions on BPFLs
+		if(!InFunction->HasAllFunctionFlags(FUNC_BlueprintCallable | FUNC_Static | FUNC_Native | FUNC_Public))
+		{
+			return false;
+		}
+
+		if(InFunction->NumParms != 2)
+		{
+			return false;
+		}
+		
+		int32 ParamIndex = 0;
+		for(TFieldIterator<FProperty> It(InFunction); It && (It->PropertyFlags & CPF_Parm); ++It, ++ParamIndex)
+		{
+			// Check first parameter is an object of the expected class
+			if(ParamIndex == 0)
+			{
+				FObjectProperty* ObjectProperty = CastField<FObjectProperty>(*It);
+				if(ObjectProperty == nullptr)
+				{ 
+					return false;
+				}
+				
+				if(!ObjectProperty->PropertyClass->IsChildOf(ExtendedClass))
+				{
+					return false;
+				}
+			}
+
+			// Check return value
+			if(ParamIndex == 1 && !It->HasAnyPropertyFlags(CPF_ReturnParm))
+			{
+				return false;
+			}
+		}
+	}
+	else
+	{
+		// We add only 'accessor' functions (no params apart from the return value) that have valid return types
+		const FProperty* ReturnProperty = InFunction->GetReturnProperty();
+		if(ReturnProperty == nullptr || InFunction->NumParms != 1 || !InFunction->HasAnyFunctionFlags(FUNC_BlueprintCallable))
+		{
+			return false;
+		}
+
+		const FParamTypeHandle TypeHandle = FParamTypeHandle::FromProperty(ReturnProperty);
+		if(!TypeHandle.IsValid())
+		{
+			return false;
+		}
+	}
+
+	return true;
+}
+
+bool FParamUtils::CanUseProperty(const FProperty* InProperty)
+{
+	if(!InProperty->HasAnyPropertyFlags(CPF_Edit | CPF_EditConst | CPF_BlueprintVisible) || InProperty->HasAnyPropertyFlags(CPF_Deprecated | CPF_EditorOnly))
+	{
+		return false;
+	}
+
+	const FParamTypeHandle TypeHandle = FParamTypeHandle::FromProperty(InProperty);
+	if(!TypeHandle.IsValid())
+	{
+		return false;
+	}
+
+	return true;
+}
+
+FName FParamUtils::LocatorToName(const FUniversalObjectLocator& InLocator)
+{
+	// By default the string representation of an empty UOL is "uobj://none", so we shortcut here for FName consistency 
+	if(InLocator.IsEmpty())
+	{
+		return NAME_None; 
+	}
+
+	TStringBuilder<1024> StringBuilder;
+	InLocator.ToString(StringBuilder);
+	ensure(StringBuilder.Len() < NAME_SIZE);
+	return FName(StringBuilder.ToView());
+}
+
+
 }
\ No newline at end of file
