AnimNext: Parameters rework - UOL integration

diff --git a/Engine/Plugins/Experimental/AnimNext/Source/AnimNext/Private/Param/ParametersProxy.cpp b/Engine/Plugins/Experimental/AnimNext/Source/AnimNext/Private/Param/ParametersProxy.cpp
--- a/Engine/Plugins/Experimental/AnimNext/Source/AnimNext/Private/Param/ParametersProxy.cpp
+++ b/Engine/Plugins/Experimental/AnimNext/Source/AnimNext/Private/Param/ParametersProxy.cpp
@@ -10,7 +10,7 @@ namespace UE::AnimNext
 FParametersProxy::FParametersProxy(UAnimNextGraph* InGraph)
 	: Graph(InGraph)
 	, PropertyBag(InGraph->PropertyBag)
-	, LayerHandle(FParamStack::MakeReferenceLayer(PropertyBag))
+	, LayerHandle(FParamStack::MakeReferenceLayer(NAME_None, PropertyBag))
 {
 	check(Graph);
 }
@@ -22,7 +22,7 @@ void FParametersProxy::Update(float DeltaTime)
 	if(HandlePropertyBag == nullptr || HandlePropertyBag->GetPropertyBagStruct() != Graph->PropertyBag.GetPropertyBagStruct())
 	{
 		PropertyBag = Graph->PropertyBag;
-		LayerHandle = FParamStack::MakeReferenceLayer(PropertyBag);
+		LayerHandle = FParamStack::MakeReferenceLayer(NAME_None, PropertyBag);
 	}
 #endif
 
diff --git a/Engine/Plugins/Experimental/AnimNext/Source/AnimNext/Private/Param/PropertyBagProxy.cpp b/Engine/Plugins/Experimental/AnimNext/Source/AnimNext/Private/Param/PropertyBagProxy.cpp
--- a/Engine/Plugins/Experimental/AnimNext/Source/AnimNext/Private/Param/PropertyBagProxy.cpp
+++ b/Engine/Plugins/Experimental/AnimNext/Source/AnimNext/Private/Param/PropertyBagProxy.cpp
@@ -8,7 +8,7 @@ namespace UE::AnimNext
 
 FPropertyBagProxy::FPropertyBagProxy()
 {
-	LayerHandle = FParamStack::MakeReferenceLayer(PropertyBag);
+	LayerHandle = FParamStack::MakeReferenceLayer(NAME_None, PropertyBag);
 }
 
 void FPropertyBagProxy::AddReferencedObjects(FReferenceCollector& Collector)
@@ -22,7 +22,7 @@ void FPropertyBagProxy::AddPropertyAndValue(FName InName, const FProperty* InPro
 	PropertyBag.SetValue(InName, InProperty, InContainerPtr);
 
 	// Recreate the layer handle as the bag layout has changed
-	LayerHandle = FParamStack::MakeReferenceLayer(PropertyBag);
+	LayerHandle = FParamStack::MakeReferenceLayer(NAME_None, PropertyBag);
 }
 
 void FPropertyBagProxy::AddPropertiesAndValues(TConstArrayView<FPropertyAndValue> InPropertiesAndValues)
@@ -45,7 +45,7 @@ void FPropertyBagProxy::AddPropertiesAndValues(TConstArrayView<FPropertyAndValue
 	}
 
 	// Recreate the layer handle as the bag layout has changed
-	LayerHandle = FParamStack::MakeReferenceLayer(PropertyBag);
+	LayerHandle = FParamStack::MakeReferenceLayer(NAME_None, PropertyBag);
 }
 
 }
diff --git a/Engine/Plugins/Experimental/AnimNext/Source/AnimNext/Private/Param/RigVMDispatch_GetLayerParameter.cpp b/Engine/Plugins/Experimental/AnimNext/Source/AnimNext/Private/Param/RigVMDispatch_GetLayerParameter.cpp
--- a/Engine/Plugins/Experimental/AnimNext/Source/AnimNext/Private/Param/RigVMDispatch_GetLayerParameter.cpp
+++ b/Engine/Plugins/Experimental/AnimNext/Source/AnimNext/Private/Param/RigVMDispatch_GetLayerParameter.cpp
@@ -86,7 +86,7 @@ void FRigVMDispatch_GetLayerParameter::Execute(FRigVMExtendedExecuteContext& InC
 	uint32& ParameterHash = *(uint32*)Handles[2].GetData();
 	if (ParameterHash == 0 && Parameter != NAME_None)
 	{
-		ParameterHash = GetTypeHash(Parameter);
+		ParameterHash = FParamId::CalculateHash(Parameter, NAME_None);
 	}
 
 	uint32& TypeHandle = *(uint32*)Handles[3].GetData();
diff --git a/Engine/Plugins/Experimental/AnimNext/Source/AnimNext/Private/Param/RigVMDispatch_GetScopedParameter.cpp b/Engine/Plugins/Experimental/AnimNext/Source/AnimNext/Private/Param/RigVMDispatch_GetScopedParameter.cpp
new file mode 100644
--- /dev/null
+++ b/Engine/Plugins/Experimental/AnimNext/Source/AnimNext/Private/Param/RigVMDispatch_GetScopedParameter.cpp
@@ -0,0 +1,105 @@
+// Copyright Epic Games, Inc. All Rights Reserved.
+
+#include "Param/RigVMDispatch_GetScopedParameter.h"
+#include "Param/ParamStack.h"
+#include "RigVMCore/RigVMStruct.h"
+#include "Param/AnimNextParam.h"
+
+const FName FRigVMDispatch_GetScopedParameter::ParameterName = TEXT("Parameter");
+const FName FRigVMDispatch_GetScopedParameter::ValueName = TEXT("Value");
+const FName FRigVMDispatch_GetScopedParameter::ParameterIdName = TEXT("ParameterId");
+const FName FRigVMDispatch_GetScopedParameter::TypeHandleName = TEXT("Type");
+
+FRigVMDispatch_GetScopedParameter::FRigVMDispatch_GetScopedParameter()
+{
+	FactoryScriptStruct = StaticStruct();
+}
+
+FName FRigVMDispatch_GetScopedParameter::GetArgumentNameForOperandIndex(int32 InOperandIndex, int32 InTotalOperands) const
+{
+	static const FName ArgumentNames[] =
+	{
+		ParameterName,
+		ValueName,
+		ParameterIdName,
+		TypeHandleName,
+	};
+	check(InTotalOperands == UE_ARRAY_COUNT(ArgumentNames));
+	return ArgumentNames[InOperandIndex];
+}
+
+#if WITH_EDITOR
+FString FRigVMDispatch_GetScopedParameter::GetArgumentMetaData(const FName& InArgumentName, const FName& InMetaDataKey) const
+{
+	if ((InArgumentName == TypeHandleName || InArgumentName == ParameterIdName) &&
+		InMetaDataKey == FRigVMStruct::SingletonMetaName)
+	{
+		return TEXT("True");
+	}
+	else if(InArgumentName == ParameterName && InMetaDataKey == FRigVMStruct::HideSubPinsMetaName)
+	{
+		return TEXT("True");
+	}
+	return Super::GetArgumentMetaData(InArgumentName, InMetaDataKey);
+}
+#endif
+
+const TArray<FRigVMTemplateArgumentInfo>& FRigVMDispatch_GetScopedParameter::GetArgumentInfos() const
+{
+	static TArray<FRigVMTemplateArgumentInfo> Infos;
+	if(Infos.IsEmpty())
+	{
+		static const TArray<FRigVMTemplateArgument::ETypeCategory> ValueCategories =
+		{
+			FRigVMTemplateArgument::ETypeCategory_SingleAnyValue,
+			FRigVMTemplateArgument::ETypeCategory_ArrayAnyValue
+		};
+
+		const FRigVMRegistry_NoLock& Registry = FRigVMRegistry_NoLock::GetForRead();
+		Infos.Emplace(ParameterName, ERigVMPinDirection::Input, Registry.GetTypeIndex_NoLock<FAnimNextParam>());
+		Infos.Emplace(ValueName, ERigVMPinDirection::Output, ValueCategories);
+		Infos.Emplace(ParameterIdName, ERigVMPinDirection::Hidden, RigVMTypeUtils::TypeIndex::UInt32);
+		Infos.Emplace(TypeHandleName, ERigVMPinDirection::Hidden, RigVMTypeUtils::TypeIndex::UInt32);
+	}
+
+	return Infos;
+}
+ 
+FRigVMTemplateTypeMap FRigVMDispatch_GetScopedParameter::OnNewArgumentType(const FName& InArgumentName, TRigVMTypeIndex InTypeIndex) const
+{
+	FRigVMTemplateTypeMap Types;
+	const FRigVMRegistry_NoLock& Registry = FRigVMRegistry_NoLock::GetForRead();
+	Types.Add(ParameterName, Registry.GetTypeIndex_NoLock<FAnimNextParam>());
+	Types.Add(ValueName, InTypeIndex);
+	Types.Add(ParameterIdName, RigVMTypeUtils::TypeIndex::UInt32);
+	Types.Add(TypeHandleName, RigVMTypeUtils::TypeIndex::UInt32);
+	return Types;
+}
+
+void FRigVMDispatch_GetScopedParameter::Execute(FRigVMExtendedExecuteContext& InContext, FRigVMMemoryHandleArray Handles, FRigVMPredicateBranchArray RigVMBranches)
+{
+	using namespace UE::AnimNext;
+
+	const FAnimNextParam& Parameter = *(FAnimNextParam*)Handles[0].GetData();
+	const FProperty* ValueProperty = Handles[1].GetResolvedProperty();
+	check(ValueProperty);
+	uint8* TargetDataPtr = Handles[1].GetData();
+
+	uint32& ParameterHash = *(uint32*)Handles[2].GetData();
+	if (ParameterHash == 0 && Parameter.Name != NAME_None)
+	{
+		ParameterHash = FParamId::CalculateHash(Parameter.Name, Parameter.InstanceId);
+	}
+
+	uint32& TypeHandle = *(uint32*)Handles[3].GetData();
+	if (TypeHandle == 0)
+	{
+		TypeHandle = FParamTypeHandle::FromProperty(ValueProperty).ToRaw();
+	}
+
+	TConstArrayView<uint8> SourceData;
+	if (FParamStack::Get().GetParamData(FParamId(Parameter.Name, Parameter.InstanceId, ParameterHash), FParamTypeHandle::FromRaw(TypeHandle), SourceData).IsSuccessful())
+	{
+		ValueProperty->CopyCompleteValue(TargetDataPtr, SourceData.GetData());
+	}
+}
diff --git a/Engine/Plugins/Experimental/AnimNext/Source/AnimNext/Private/Param/RigVMDispatch_SetLayerParameter.cpp b/Engine/Plugins/Experimental/AnimNext/Source/AnimNext/Private/Param/RigVMDispatch_SetLayerParameter.cpp
--- a/Engine/Plugins/Experimental/AnimNext/Source/AnimNext/Private/Param/RigVMDispatch_SetLayerParameter.cpp
+++ b/Engine/Plugins/Experimental/AnimNext/Source/AnimNext/Private/Param/RigVMDispatch_SetLayerParameter.cpp
@@ -97,7 +97,7 @@ void FRigVMDispatch_SetLayerParameter::Execute(FRigVMExtendedExecuteContext& InC
 	uint32& ParameterHash = *(uint32*)Handles[2].GetData();
 	if (ParameterHash == 0 && Parameter != NAME_None)
 	{
-		ParameterHash = GetTypeHash(Parameter);
+		ParameterHash = FParamId::CalculateHash(Parameter, NAME_None);
 	}
 
 	uint32& RawTypeHandle = *(uint32*)Handles[3].GetData();
diff --git a/Engine/Plugins/Experimental/AnimNext/Source/AnimNext/Private/Scheduler/AnimNextScheduleGraphTask.cpp b/Engine/Plugins/Experimental/AnimNext/Source/AnimNext/Private/Scheduler/AnimNextScheduleGraphTask.cpp
--- a/Engine/Plugins/Experimental/AnimNext/Source/AnimNext/Private/Scheduler/AnimNextScheduleGraphTask.cpp
+++ b/Engine/Plugins/Experimental/AnimNext/Source/AnimNext/Private/Scheduler/AnimNextScheduleGraphTask.cpp
@@ -8,17 +8,53 @@
 #include "Context.h"
 #include "Graph/AnimNext_LODPose.h"
 #include "AnimNextStats.h"
+#include "Component/AnimNextComponent.h"
+#include "Component/AnimNextMeshComponent.h"
+#include "GameFramework/Character.h"
 #include "Logging/StructuredLog.h"
+#include "Param/AnimNextObjectCastLocatorFragment.h"
+#include "Param/AnimNextObjectFunctionLocatorFragment.h"
 #include "Param/AnimNextParam.h"
+#include "Param/AnimNextParamUniversalObjectLocator.h"
 
 DEFINE_STAT(STAT_AnimNext_Task_Graph);
 
+namespace UE::AnimNext
+{
+	// TODO: Currently we hard-code the ACharacter mesh component as the source of the reference pose and LOD index, but this should be a pin-input in the
+	// final schedule incarnation.
+	TInstancedStruct<FAnimNextParamUniversalObjectLocator> GetCharacterInstanceId()
+	{
+		TInstancedStruct<FAnimNextParamUniversalObjectLocator> Locator = TInstancedStruct<FAnimNextParamUniversalObjectLocator>::Make();
+		Locator.GetMutable().Locator.AddFragment<FAnimNextObjectFunctionLocatorFragment>(UAnimNextComponent::StaticClass()->FindFunctionByName("GetOwner"));
+		Locator.GetMutable().Locator.AddFragment<FAnimNextObjectCastLocatorFragment>(ACharacter::StaticClass());
+		return Locator;
+	};
+
+	FName GetCharacterInstanceIdName()
+	{
+		static FName Name(NAME_None);
+		if(Name.IsNone())
+		{
+			const TInstancedStruct<FAnimNextParamUniversalObjectLocator>& Locator = GetCharacterInstanceId();
+			Name = Locator.Get().ToName();
+		}
+		return Name;
+	};
+
+	FParamId GetMeshComponentParamId()
+	{
+		static const FParamId MeshComponentParamId("/Script/Engine.Character:Mesh", GetCharacterInstanceIdName());
+		return MeshComponentParamId;
+	}
+}
+
 UAnimNextGraph* FAnimNextScheduleGraphTask::GetGraphToRun(UE::AnimNext::FParamStack& ParamStack) const
 {
 	UAnimNextGraph* GraphToRun = Graph;
-	if (GraphToRun == nullptr && DynamicGraph != NAME_None)
+	if (GraphToRun == nullptr && DynamicGraph.IsValid())
 	{
-		if(const TObjectPtr<UAnimNextGraph>* FoundGraph = ParamStack.GetParamPtr<TObjectPtr<UAnimNextGraph>>(DynamicGraph))
+		if(const TObjectPtr<UAnimNextGraph>* FoundGraph = ParamStack.GetParamPtr<TObjectPtr<UAnimNextGraph>>(DynamicGraph.GetParamId()))
 		{
 			GraphToRun = *FoundGraph;
 		}
@@ -40,7 +76,7 @@ void FAnimNextScheduleGraphTask::VerifyRequiredParameters(UAnimNextGraph* InGrap
 			FAnimNextParamType SuppliedParameterType;
 			for(const FAnimNextParam& SuppliedParameter : SuppliedParameters)
 			{
-				if(RequiredParameter.Name == SuppliedParameter.Name)
+				if(RequiredParameter.Name == SuppliedParameter.Name && RequiredParameter.InstanceId == SuppliedParameter.InstanceId)
 				{
 					if(RequiredParameter.Type != SuppliedParameter.Type)
 					{
@@ -60,7 +96,7 @@ void FAnimNextScheduleGraphTask::VerifyRequiredParameters(UAnimNextGraph* InGrap
 			
 			if(!bFound)
 			{
-				UE_LOGFMT(LogAnimation, Warning, "    Not Found: {Name}", RequiredParameter.Name);
+				UE_LOGFMT(LogAnimation, Warning, "    Not Found: {Name} (Instance: {Instance}))", RequiredParameter.Name, RequiredParameter.InstanceId);
 			}
 			else if(!bFoundCorrectType)
 			{
@@ -97,23 +133,29 @@ void FAnimNextScheduleGraphTask::RunGraph(const UE::AnimNext::FScheduleContext&
 	// Allocate our graph instance data
 	if (!GraphCache.GraphInstanceData.IsValid())
 	{
-		GraphToRun->AllocateInstance(GraphCache.GraphInstanceData, EntryPoint);
+		GraphToRun->AllocateInstance(GraphCache.GraphInstanceData, EntryPoint.Name);
 
 		// Only do dynamic verification for dynamic graphs. Static graphs get verified at compile time. 
-		if (Graph == nullptr && DynamicGraph != NAME_None)
+		if (Graph == nullptr && DynamicGraph.IsValid())
 		{
 			VerifyRequiredParameters(GraphToRun);
 		}
 	}
 
-	const FAnimNextGraphReferencePose* GraphReferencePose = ParamStack.GetParamPtr<FAnimNextGraphReferencePose>(GraphToRun->GetReferencePoseParam());
-	if(GraphReferencePose == nullptr || !GraphReferencePose->ReferencePose.IsValid())
+	const TObjectPtr<USkeletalMeshComponent>* ComponentPtr = ParamStack.GetParamPtr<TObjectPtr<USkeletalMeshComponent>>(GetMeshComponentParamId());
+	if(ComponentPtr == nullptr)
 	{
 		return;
 	}
 
-	const int32* GraphLODLevel = ParamStack.GetParamPtr<int32>(GraphToRun->GetCurrentLODParam());
-	if(GraphLODLevel == nullptr)
+	TObjectPtr<UAnimNextMeshComponent> Component = Cast<UAnimNextMeshComponent>(*ComponentPtr);
+	if(Component == nullptr)
+	{
+		return;
+	}
+	
+	const FAnimNextGraphReferencePose GraphReferencePose = Component->GetReferencePose();
+	if(!GraphReferencePose.ReferencePose.IsValid())
 	{
 		return;
 	}
@@ -143,16 +185,18 @@ void FAnimNextScheduleGraphTask::RunGraph(const UE::AnimNext::FScheduleContext&
 		return;
 	}
 	
-	const UE::AnimNext::FReferencePose& RefPose = GraphReferencePose->ReferencePose.GetRef<UE::AnimNext::FReferencePose>();
+	const UE::AnimNext::FReferencePose& RefPose = GraphReferencePose.ReferencePose.GetRef<UE::AnimNext::FReferencePose>();
+
+	const int32 LODIndex = Component->GetPredictedLODLevel();
 
 	// Create or update our result pose
 	// TODO: Currently forcing additive flag to false here
-	if (OutputPose->LODPose.ShouldPrepareForLOD(RefPose, *GraphLODLevel, false))
+	if (OutputPose->LODPose.ShouldPrepareForLOD(RefPose, LODIndex, false))
 	{
-		OutputPose->LODPose.PrepareForLOD(RefPose, *GraphLODLevel, true, false);
+		OutputPose->LODPose.PrepareForLOD(RefPose, LODIndex, true, false);
 	}
 
-	check(OutputPose->LODPose.LODLevel == *GraphLODLevel);
+	check(OutputPose->LODPose.LODLevel == LODIndex);
 
 	// Internally we use memstack allocation, so we need a mark here
 	FMemStack& MemStack = FMemStack::Get();
@@ -163,5 +207,21 @@ void FAnimNextScheduleGraphTask::RunGraph(const UE::AnimNext::FScheduleContext&
 	MemStack.Alloc(size_t(FPageAllocator::SmallPageSize) + 1, 16);
 
 	IAnimNextModule::Get().UpdateGraph(GraphCache.GraphInstanceData, InContext.GetDeltaTime());
-	IAnimNextModule::Get().EvaluateGraph(GraphCache.GraphInstanceData, RefPose, *GraphLODLevel, OutputPose->LODPose);
+	IAnimNextModule::Get().EvaluateGraph(GraphCache.GraphInstanceData, RefPose, LODIndex, OutputPose->LODPose);
+}
+
+#if WITH_EDITORONLY_DATA 
+
+TArray<FAnimNextEditorParam> FAnimNextScheduleGraphTask::GetRequiredParametersInternal()
+{
+	using namespace UE::AnimNext;
+
+	const TArray<FAnimNextEditorParam> Params =
+	{
+		FAnimNextEditorParam(GetMeshComponentParamId().GetName(), FAnimNextParamType::GetType<TObjectPtr<UAnimNextMeshComponent>>(), GetCharacterInstanceId()),
+	};
+
+	return Params;
 }
+
+#endif
\ No newline at end of file
diff --git a/Engine/Plugins/Experimental/AnimNext/Source/AnimNext/Private/Scheduler/AnimNextScheduleParamScopeTask.cpp b/Engine/Plugins/Experimental/AnimNext/Source/AnimNext/Private/Scheduler/AnimNextScheduleParamScopeTask.cpp
--- a/Engine/Plugins/Experimental/AnimNext/Source/AnimNext/Private/Scheduler/AnimNextScheduleParamScopeTask.cpp
+++ b/Engine/Plugins/Experimental/AnimNext/Source/AnimNext/Private/Scheduler/AnimNextScheduleParamScopeTask.cpp
@@ -23,10 +23,10 @@ void FAnimNextScheduleParamScopeEntryTask::RunParamScopeEntry(const UE::AnimNext
 	check(ScopeCache.PushedLayers.Num() == 0);
 
 	// Update & push user params (before scope)
-	FScheduleInstanceData::FUserScope* FoundUserScope = InstanceData.UserScopes.Find(Scope);
+	FScheduleInstanceData::FUserScope* FoundUserScope = InstanceData.UserScopes.Find(Scope.Name);
 	if (FoundUserScope && FoundUserScope->BeforeSource.IsValid())
 	{
-		FPropertyBagProxy& UserParameterSource = *FoundUserScope->BeforeSource.Get();
+		IParameterSource& UserParameterSource = *FoundUserScope->BeforeSource.Get();
 		UserParameterSource.Update(DeltaTime);
 		ScopeCache.PushedLayers.Add(ParamStack.PushLayer(UserParameterSource.GetLayerHandle()));
 	}
@@ -43,7 +43,7 @@ void FAnimNextScheduleParamScopeEntryTask::RunParamScopeEntry(const UE::AnimNext
 	// Update & push user params (after scope)
 	if (FoundUserScope && FoundUserScope->AfterSource.IsValid())
 	{
-		FPropertyBagProxy& UserParameterSource = *FoundUserScope->AfterSource.Get();
+		IParameterSource& UserParameterSource = *FoundUserScope->AfterSource.Get();
 		UserParameterSource.Update(DeltaTime);
 		ScopeCache.PushedLayers.Add(ParamStack.PushLayer(UserParameterSource.GetLayerHandle()));
 	}
diff --git a/Engine/Plugins/Experimental/AnimNext/Source/AnimNext/Private/Scheduler/AnimNextSchedulerEntry.cpp b/Engine/Plugins/Experimental/AnimNext/Source/AnimNext/Private/Scheduler/AnimNextSchedulerEntry.cpp
--- a/Engine/Plugins/Experimental/AnimNext/Source/AnimNext/Private/Scheduler/AnimNextSchedulerEntry.cpp
+++ b/Engine/Plugins/Experimental/AnimNext/Source/AnimNext/Private/Scheduler/AnimNextSchedulerEntry.cpp
@@ -9,10 +9,11 @@
 #include "AnimNextTickFunctionBinding.h"
 #include "Logging/StructuredLog.h"
 #include "Param/PropertyBagProxy.h"
+#include "Scheduler/ScheduleInitializationContext.h"
 
 DEFINE_STAT(STAT_AnimNext_InitializeEntry);
 
-FAnimNextSchedulerEntry::FAnimNextSchedulerEntry(const UAnimNextSchedule* InSchedule, UObject* InObject, UE::AnimNext::FScheduleHandle InHandle, EAnimNextScheduleInitMethod InInitMethod, TUniqueFunction<void(const UE::AnimNext::FScheduleContext&)>&& InInitializeCallback)
+FAnimNextSchedulerEntry::FAnimNextSchedulerEntry(const UAnimNextSchedule* InSchedule, UObject* InObject, UE::AnimNext::FScheduleHandle InHandle, EAnimNextScheduleInitMethod InInitMethod, TUniqueFunction<void(const UE::AnimNext::FScheduleInitializationContext&)>&& InInitializeCallback)
 	: Schedule(InSchedule)
 	, WeakObject(InObject)
 	, Handle(InHandle)
@@ -98,10 +99,10 @@ void FAnimNextSchedulerEntry::Initialize()
 						TConstArrayView<FAnimNextScheduleInstruction>(&Instruction, 1),
 						TConstArrayView<TWeakObjectPtr<UObject>>(&TargetObject, 1)));
 
-					const FName ExternalTaskName = Schedule->ExternalTasks[Instruction.Operand].ExternalTask;
-					if(ExternalTaskName != NAME_None)
+					const FAnimNextParam& ExternalTaskParam = Schedule->ExternalTasks[Instruction.Operand].ExternalTask;
+					if(ExternalTaskParam.IsValid())
 					{
-						FAnimNextTickFunctionBinding* FoundBinding = ParamStack.GetMutableParamPtr<FAnimNextTickFunctionBinding>(ExternalTaskName);
+						FAnimNextTickFunctionBinding* FoundBinding = ParamStack.GetMutableParamPtr<FAnimNextTickFunctionBinding>(ExternalTaskParam.GetParamId());
 						if(FoundBinding)
 						{
 							check(FoundBinding->Object.Get());
@@ -112,12 +113,12 @@ void FAnimNextSchedulerEntry::Initialize()
 						}
 						else
 						{
-							UE_LOGFMT(LogAnimation, Warning, "AnimNext: Could not bind to external tick function using binding parameters (ExternalTaskName={ExternalTaskName})", ExternalTaskName);
+							UE_LOGFMT(LogAnimation, Warning, "AnimNext: Could not bind to external tick function using binding parameters (ExternalTaskName={ExternalTaskName})", ExternalTaskParam.Name);
 						}
 					}
 					else
 					{
-						UE_LOGFMT(LogAnimation, Warning, "AnimNext: Invalid external tick function binding parameters (ExternalTaskName={ExternalTaskName})", ExternalTaskName);
+						UE_LOGFMT(LogAnimation, Warning, "AnimNext: Invalid external tick function binding parameters (ExternalTaskName={ExternalTaskName})", ExternalTaskParam.Name);
 					}
 					break;
 				}
@@ -128,10 +129,10 @@ void FAnimNextSchedulerEntry::Initialize()
 						TConstArrayView<FAnimNextScheduleInstruction>(&Instruction, 1),
 						TConstArrayView<TWeakObjectPtr<UObject>>(&TargetObject, 1)));
 
-					const FName ExternalTaskName = Schedule->ExternalTasks[Instruction.Operand].ExternalTask;
-					if(ExternalTaskName != NAME_None)
+					const FAnimNextParam& ExternalTaskParam = Schedule->ExternalTasks[Instruction.Operand].ExternalTask;
+					if(ExternalTaskParam.IsValid())
 					{
-						FAnimNextTickFunctionBinding* FoundBinding = ParamStack.GetMutableParamPtr<FAnimNextTickFunctionBinding>(ExternalTaskName);
+						FAnimNextTickFunctionBinding* FoundBinding = ParamStack.GetMutableParamPtr<FAnimNextTickFunctionBinding>(ExternalTaskParam.GetParamId());
 						if(FoundBinding)
 						{
 							TargetObject = FoundBinding->Object.Get();
@@ -204,7 +205,7 @@ void FAnimNextSchedulerEntry::Initialize()
 
 		if(InitializeCallback)
 		{
-			InitializeCallback(Context);
+			InitializeCallback(FScheduleInitializationContext(Context));
 		}
 
 		// Just pause now if we arent needing an initial update
@@ -342,7 +343,7 @@ void FAnimNextSchedulerEntry::OnScheduleCompiled()
 	using namespace UE::AnimNext;
 
 	// Store any user-defined scopes, as the instance data will be going away
-	TUniquePtr<FPropertyBagProxy> RootUserScope = MoveTemp(Context.InstanceData->RootUserScope);
+	TUniquePtr<IParameterSource> RootUserScope = MoveTemp(Context.InstanceData->RootUserScope);
 	TMap<FName, FScheduleInstanceData::FUserScope> UserScopes = MoveTemp(Context.InstanceData->UserScopes);
 
 	ResetBindingsAndInstanceData();
diff --git a/Engine/Plugins/Experimental/AnimNext/Source/AnimNext/Private/Scheduler/AnimNextSchedulerEntry.h b/Engine/Plugins/Experimental/AnimNext/Source/AnimNext/Private/Scheduler/AnimNextSchedulerEntry.h
--- a/Engine/Plugins/Experimental/AnimNext/Source/AnimNext/Private/Scheduler/AnimNextSchedulerEntry.h
+++ b/Engine/Plugins/Experimental/AnimNext/Source/AnimNext/Private/Scheduler/AnimNextSchedulerEntry.h
@@ -14,6 +14,7 @@ namespace UE::AnimNext
 {
 	struct FScheduleHandle;
 	struct FParamStack;
+	struct FScheduleInitializationContext;
 }
 
 // Root memory owner of a parameterized schedule 
@@ -23,7 +24,7 @@ struct FAnimNextSchedulerEntry
 	GENERATED_BODY()
 
 	FAnimNextSchedulerEntry() = default;
-	FAnimNextSchedulerEntry(const UAnimNextSchedule* InSchedule, UObject* InObject, UE::AnimNext::FScheduleHandle InHandle, EAnimNextScheduleInitMethod InInitMethod, TUniqueFunction<void(const UE::AnimNext::FScheduleContext&)>&& InInitializeCallback);
+	FAnimNextSchedulerEntry(const UAnimNextSchedule* InSchedule, UObject* InObject, UE::AnimNext::FScheduleHandle InHandle, EAnimNextScheduleInitMethod InInitMethod, TUniqueFunction<void(const UE::AnimNext::FScheduleInitializationContext&)>&& InInitializeCallback);
 	~FAnimNextSchedulerEntry();
 
 	// Setup the entry
@@ -76,7 +77,7 @@ struct FAnimNextSchedulerEntry
 	TArray<TUniquePtr<UE::AnimNext::FScheduleTickFunction>> TickFunctions;
 
 	// Callback to run on intialization
-	TUniqueFunction<void(const UE::AnimNext::FScheduleContext&)> InitializeCallback;
+	TUniqueFunction<void(const UE::AnimNext::FScheduleInitializationContext&)> InitializeCallback;
 
 	// Current delta time, updated each time the schedule runs
 	float DeltaTime = 0.0f;
diff --git a/Engine/Plugins/Experimental/AnimNext/Source/AnimNext/Private/Scheduler/AnimNextSchedulerWorldSubsystem.cpp b/Engine/Plugins/Experimental/AnimNext/Source/AnimNext/Private/Scheduler/AnimNextSchedulerWorldSubsystem.cpp
--- a/Engine/Plugins/Experimental/AnimNext/Source/AnimNext/Private/Scheduler/AnimNextSchedulerWorldSubsystem.cpp
+++ b/Engine/Plugins/Experimental/AnimNext/Source/AnimNext/Private/Scheduler/AnimNextSchedulerWorldSubsystem.cpp
@@ -97,7 +97,7 @@ void UAnimNextSchedulerWorldSubsystem::FlushPendingActions()
 	}
 }
 
-UE::AnimNext::FScheduleHandle UAnimNextSchedulerWorldSubsystem::AcquireHandle(UObject* InObject, UAnimNextSchedule* InSchedule, EAnimNextScheduleInitMethod InInitMethod, TUniqueFunction<void(const UE::AnimNext::FScheduleContext&)>&& InInitializeCallback)
+UE::AnimNext::FScheduleHandle UAnimNextSchedulerWorldSubsystem::AcquireHandle(UObject* InObject, UAnimNextSchedule* InSchedule, EAnimNextScheduleInitMethod InInitMethod, TUniqueFunction<void(const UE::AnimNext::FScheduleInitializationContext&)>&& InInitializeCallback)
 {
 	using namespace UE::AnimNext;
 
@@ -168,7 +168,7 @@ void UAnimNextSchedulerWorldSubsystem::EnableHandle(UE::AnimNext::FScheduleHandl
 	}
 }
 
-void UAnimNextSchedulerWorldSubsystem::QueueTask(UE::AnimNext::FScheduleHandle InHandle, FName InScheduleTaskName, TUniqueFunction<void(const UE::AnimNext::FScheduleContext&)>&& InTaskFunction, UE::AnimNext::FScheduler::ETaskRunLocation InLocation)
+void UAnimNextSchedulerWorldSubsystem::QueueTask(UE::AnimNext::FScheduleHandle InHandle, FName InScheduleTaskName, TUniqueFunction<void(const UE::AnimNext::FScheduleTaskContext&)>&& InTaskFunction, UE::AnimNext::FScheduler::ETaskRunLocation InLocation)
 {
 	using namespace UE::AnimNext;
 
@@ -177,7 +177,7 @@ void UAnimNextSchedulerWorldSubsystem::QueueTask(UE::AnimNext::FScheduleHandle I
 		TUniquePtr<FAnimNextSchedulerEntry>& Entry = Entries[InHandle.Index];
 
 		// TODO: Only supporting scope tasks or "None" for root for now
-		TSpscQueue<TUniqueFunction<void(const UE::AnimNext::FScheduleContext&)>>* Queue = nullptr;
+		TSpscQueue<TUniqueFunction<void(const UE::AnimNext::FScheduleTaskContext&)>>* Queue = nullptr;
 		if(InScheduleTaskName == NAME_None)
 		{
 			Queue = &Entry->BeginTickFunction->PreExecuteTasks;
@@ -186,7 +186,7 @@ void UAnimNextSchedulerWorldSubsystem::QueueTask(UE::AnimNext::FScheduleHandle I
 		{
 			const FAnimNextScheduleParamScopeEntryTask* FoundScope = Entry->Schedule->ParamScopeEntryTasks.FindByPredicate([&InScheduleTaskName](const FAnimNextScheduleParamScopeEntryTask& InTask)
 			{
-				return InTask.Scope == InScheduleTaskName;
+				return InTask.Scope.Name == InScheduleTaskName;
 			});
 
 			if (FoundScope)
diff --git a/Engine/Plugins/Experimental/AnimNext/Source/AnimNext/Private/Scheduler/AnimNextSchedulerWorldSubsystem.h b/Engine/Plugins/Experimental/AnimNext/Source/AnimNext/Private/Scheduler/AnimNextSchedulerWorldSubsystem.h
--- a/Engine/Plugins/Experimental/AnimNext/Source/AnimNext/Private/Scheduler/AnimNextSchedulerWorldSubsystem.h
+++ b/Engine/Plugins/Experimental/AnimNext/Source/AnimNext/Private/Scheduler/AnimNextSchedulerWorldSubsystem.h
@@ -63,7 +63,7 @@ class UAnimNextSchedulerWorldSubsystem : public UWorldSubsystem
 	void FlushPendingActions();
 
 	// Acquire a handle that binds a schedule with the supplied parameters
-	UE::AnimNext::FScheduleHandle AcquireHandle(UObject* InObject, UAnimNextSchedule* InSchedule, EAnimNextScheduleInitMethod InInitMethod, TUniqueFunction<void(const UE::AnimNext::FScheduleContext&)>&& InInitializeCallback);
+	UE::AnimNext::FScheduleHandle AcquireHandle(UObject* InObject, UAnimNextSchedule* InSchedule, EAnimNextScheduleInitMethod InInitMethod, TUniqueFunction<void(const UE::AnimNext::FScheduleInitializationContext&)>&& InInitializeCallback);
 
 	// Release an already acquired handle
 	// The full release of the binding referenced by the handle will be deferred after this call is made
@@ -78,7 +78,7 @@ class UAnimNextSchedulerWorldSubsystem : public UWorldSubsystem
 	// @param	InTaskName		The name of the task in the schedule to run the supplied task relative to
 	// @param	InTaskFunction	The function to run
 	// @param	InLocation		Where to run the task, before or after
-	void QueueTask(UE::AnimNext::FScheduleHandle InHandle, FName InScheduleTaskName, TUniqueFunction<void(const UE::AnimNext::FScheduleContext&)>&& InTaskFunction, UE::AnimNext::FScheduler::ETaskRunLocation InLocation);
+	void QueueTask(UE::AnimNext::FScheduleHandle InHandle, FName InScheduleTaskName, TUniqueFunction<void(const UE::AnimNext::FScheduleTaskContext&)>&& InTaskFunction, UE::AnimNext::FScheduler::ETaskRunLocation InLocation);
 
 #if WITH_EDITOR
 	// Refresh any entries that use the provided schedule as it has been recompiled.
diff --git a/Engine/Plugins/Experimental/AnimNext/Source/AnimNext/Private/Scheduler/Ports/AnimNextSchedulePort_AnimNextMeshComponentPose.cpp b/Engine/Plugins/Experimental/AnimNext/Source/AnimNext/Private/Scheduler/Ports/AnimNextSchedulePort_AnimNextMeshComponentPose.cpp
--- a/Engine/Plugins/Experimental/AnimNext/Source/AnimNext/Private/Scheduler/Ports/AnimNextSchedulePort_AnimNextMeshComponentPose.cpp
+++ b/Engine/Plugins/Experimental/AnimNext/Source/AnimNext/Private/Scheduler/Ports/AnimNextSchedulePort_AnimNextMeshComponentPose.cpp
@@ -9,14 +9,14 @@
 #include "Components/SkeletalMeshComponent.h"
 #include "Engine/SkeletalMesh.h"
 #include "AnimNextStats.h"
+#include "Component/AnimNextComponent.h"
 #include "Component/AnimNextMeshComponent.h"
-#include "Param/AnimNextParam.h"
+#include "GameFramework/Character.h"
+#include "Scheduler/AnimNextScheduleGraphTask.h"
+#include "Param/AnimNextParamUniversalObjectLocator.h"
 
 DEFINE_STAT(STAT_AnimNext_Port_SkeletalMeshComponent);
 
-UE::AnimNext::FParamId UAnimNextSchedulePort_AnimNextMeshComponentPose::ComponentParamId("UE_AnimNextMeshComponent");
-UE::AnimNext::FParamId UAnimNextSchedulePort_AnimNextMeshComponentPose::ReferencePoseParamId("UE_AnimNextMeshComponent_ReferencePose");
-
 void UAnimNextSchedulePort_AnimNextMeshComponentPose::Run(const UE::AnimNext::FScheduleTermContext& InContext) const
 {
 	SCOPE_CYCLE_COUNTER(STAT_AnimNext_Port_SkeletalMeshComponent);
@@ -25,38 +25,36 @@ void UAnimNextSchedulePort_AnimNextMeshComponentPose::Run(const UE::AnimNext::FS
 
 	const FParamStack& ParamStack = FParamStack::Get();
 
-	const TObjectPtr<UAnimNextMeshComponent>* ComponentPtr = ParamStack.GetParamPtr<TObjectPtr<UAnimNextMeshComponent>>(ComponentParamId);
+	const TObjectPtr<USkeletalMeshComponent>* ComponentPtr = ParamStack.GetParamPtr<TObjectPtr<USkeletalMeshComponent>>(GetMeshComponentParamId());
 	if(ComponentPtr == nullptr)
 	{
 		return;
 	}
 
-	const FAnimNextGraphLODPose* InputPose = InContext.GetLayerHandle().GetParamPtr<FAnimNextGraphLODPose>(GetTerms()[0].GetId());
-	if(InputPose == nullptr)
+	TObjectPtr<UAnimNextMeshComponent> Component = Cast<UAnimNextMeshComponent>(*ComponentPtr);
+	if(Component == nullptr)
 	{
 		return;
 	}
-
-	if(!InputPose->LODPose.IsValid())
+	
+	const FAnimNextGraphLODPose* InputPose = InContext.GetLayerHandle().GetParamPtr<FAnimNextGraphLODPose>(GetTerms()[0].GetId());
+	if(InputPose == nullptr)
 	{
-		return;
+		return; 
 	}
 
-	const FAnimNextGraphReferencePose* GraphReferencePose = ParamStack.GetParamPtr<FAnimNextGraphReferencePose>(ReferencePoseParamId);
-	if(GraphReferencePose == nullptr)
+	if(!InputPose->LODPose.IsValid())
 	{
 		return;
 	}
-
-	UAnimNextMeshComponent* Component = *ComponentPtr;
-
+	
 	USkeletalMesh* SkeletalMesh = Component->GetSkeletalMeshAsset();
 	if(SkeletalMesh == nullptr)
 	{
 		return;
 	}
 
-	const UE::AnimNext::FReferencePose& RefPose = GraphReferencePose->ReferencePose.GetRef<UE::AnimNext::FReferencePose>();
+	const UE::AnimNext::FReferencePose& RefPose = Component->GetReferencePose().ReferencePose.GetRef<UE::AnimNext::FReferencePose>();
 
 	FMemMark MemMark(FMemStack::Get());
 
@@ -87,15 +85,14 @@ TConstArrayView<UE::AnimNext::FScheduleTerm> UAnimNextSchedulePort_AnimNextMeshC
 	return Terms;
 }
 
-TConstArrayView<FAnimNextParam> UAnimNextSchedulePort_AnimNextMeshComponentPose::GetRequiredParameters() const
+TConstArrayView<FAnimNextEditorParam> UAnimNextSchedulePort_AnimNextMeshComponentPose::GetRequiredParameters() const
 {
 	using namespace UE::AnimNext;
 
-	static const FAnimNextParam Params[] =
+	if(RequiredParams.Num() == 0)
 	{
-		FAnimNextParam(ComponentParamId.GetName(), FAnimNextParamType::GetType<TObjectPtr<UAnimNextMeshComponent>>()),
-		FAnimNextParam(ReferencePoseParamId.GetName(), FAnimNextParamType::GetType<FAnimNextGraphReferencePose>()),
-	};
+		RequiredParams.Emplace(GetMeshComponentParamId().GetName(), FAnimNextParamType::GetType<TObjectPtr<UAnimNextMeshComponent>>(), GetCharacterInstanceId());
+	}
 
-	return Params;
+	return RequiredParams;
 }
diff --git a/Engine/Plugins/Experimental/AnimNext/Source/AnimNext/Private/Scheduler/Ports/AnimNextSchedulePort_AnimNextMeshComponentPose.h b/Engine/Plugins/Experimental/AnimNext/Source/AnimNext/Private/Scheduler/Ports/AnimNextSchedulePort_AnimNextMeshComponentPose.h
--- a/Engine/Plugins/Experimental/AnimNext/Source/AnimNext/Private/Scheduler/Ports/AnimNextSchedulePort_AnimNextMeshComponentPose.h
+++ b/Engine/Plugins/Experimental/AnimNext/Source/AnimNext/Private/Scheduler/Ports/AnimNextSchedulePort_AnimNextMeshComponentPose.h
@@ -12,11 +12,11 @@ class UAnimNextSchedulePort_AnimNextMeshComponentPose : public UAnimNextSchedule
 
 	// UAnimNextSchedulePort interface
 	virtual void Run(const UE::AnimNext::FScheduleTermContext& InContext) const override;
-	virtual TConstArrayView<FAnimNextParam> GetRequiredParameters() const override;
+	virtual TConstArrayView<FAnimNextEditorParam> GetRequiredParameters() const override;
 	
 	// IAnimNextScheduleTermInterface interface
 	virtual TConstArrayView<UE::AnimNext::FScheduleTerm> GetTerms() const override;
 
-	static UE::AnimNext::FParamId ComponentParamId;
-	static UE::AnimNext::FParamId ReferencePoseParamId;
+private:
+	mutable TArray<FAnimNextEditorParam> RequiredParams;
 };
\ No newline at end of file
diff --git a/Engine/Plugins/Experimental/AnimNext/Source/AnimNext/Private/Scheduler/ScheduleInitializationContext.cpp b/Engine/Plugins/Experimental/AnimNext/Source/AnimNext/Private/Scheduler/ScheduleInitializationContext.cpp
new file mode 100644
--- /dev/null
+++ b/Engine/Plugins/Experimental/AnimNext/Source/AnimNext/Private/Scheduler/ScheduleInitializationContext.cpp
@@ -0,0 +1,20 @@
+// Copyright Epic Games, Inc. All Rights Reserved.
+
+#include "Scheduler/ScheduleInitializationContext.h"
+#include "Scheduler/ScheduleContext.h"
+
+namespace UE::AnimNext
+{
+
+FScheduleInitializationContext::FScheduleInitializationContext(const FScheduleContext& InContext)
+	: Context(InContext)
+{
+}
+
+void FScheduleInitializationContext::ApplyParametersToScope(FName InScope, EParameterScopeOrdering InOrdering, TUniquePtr<IParameterSource>&& InParameters) const
+{
+	FScheduleInstanceData& InstanceData = Context.GetInstanceData();
+	InstanceData.ApplyParametersToScope(InScope, InOrdering, MoveTemp(InParameters));
+}
+
+}
\ No newline at end of file
diff --git a/Engine/Plugins/Experimental/AnimNext/Source/AnimNext/Private/Scheduler/ScheduleInstanceData.cpp b/Engine/Plugins/Experimental/AnimNext/Source/AnimNext/Private/Scheduler/ScheduleInstanceData.cpp
--- a/Engine/Plugins/Experimental/AnimNext/Source/AnimNext/Private/Scheduler/ScheduleInstanceData.cpp
+++ b/Engine/Plugins/Experimental/AnimNext/Source/AnimNext/Private/Scheduler/ScheduleInstanceData.cpp
@@ -7,11 +7,13 @@
 #include "Scheduler/AnimNextSchedulerEntry.h"
 #include "Scheduler/AnimNextSchedulePort.h"
 #include "AnimNextStats.h"
-#include "Param/ExternalParameterRegistry.h"
+#include "IAnimNextModule.h"
+#include "Modules/ModuleManager.h"
 #include "Param/IParameterSourceFactory.h"
 #include "Param/ParametersProxy.h"
 #include "Param/PropertyBagProxy.h"
 #include "Scheduler/AnimNextScheduleExternalParamTask.h"
+#include "Scheduler/ScheduleInitializationContext.h"
 
 DEFINE_STAT(STAT_AnimNext_CreateInstanceData);
 
@@ -72,8 +74,10 @@ FScheduleInstanceData::FScheduleInstanceData(const FScheduleContext& InScheduleC
 	}
 
 	// Set up external parameters
-	FExternalParameterContext ExternalParameterContext;
-	ExternalParameterContext.Object = Entry->WeakObject.Get();
+	IAnimNextModule& AnimNextModule = FModuleManager::GetModuleChecked<IAnimNextModule>("AnimNext");
+
+	FParameterSourceContext ParameterSourceContext;
+	ParameterSourceContext.Object = Entry->WeakObject.Get();
 
 	const float DeltaTime = InScheduleContext.GetDeltaTime();
 	for(int32 ExternalParamSourceIndex = 0; ExternalParamSourceIndex < ExternalParamCaches.Num(); ++ExternalParamSourceIndex)
@@ -82,7 +86,7 @@ FScheduleInstanceData::FScheduleInstanceData(const FScheduleContext& InScheduleC
 
 		for(const FAnimNextScheduleExternalParameterSource& ParameterSource : InSchedule->ExternalParamTasks[ExternalParamSourceIndex].ParameterSources)
 		{
-			if(TUniquePtr<IParameterSource> NewParameterSource = FExternalParameterRegistry::CreateParameterSource(ExternalParameterContext, ParameterSource.ParameterSource, ParameterSource.Parameters))
+			if(TUniquePtr<IParameterSource> NewParameterSource = AnimNextModule.CreateParameterSource(ParameterSourceContext, ParameterSource.InstanceId, ParameterSource.Parameters))
 			{
 				// Initial update is required to populate the cache
 				// TODO: This needs to move outside this function once we run initialization off the game thread, depending on thread-safety
@@ -99,7 +103,7 @@ FScheduleInstanceData::FScheduleInstanceData(const FScheduleContext& InScheduleC
 	IntermediatesData = InSchedule->IntermediatesData;
 
 	// Make a hosting layer for the intermediates
-	IntermediatesLayer = FParamStack::MakeReferenceLayer(IntermediatesData);
+	IntermediatesLayer = FParamStack::MakeReferenceLayer(NAME_None, IntermediatesData);
 
 	// Resize remapped intermediate data layers for port tasks, they will be allocated lazily later
 	PortTermLayers.SetNum(InSchedule->Ports.Num());
@@ -109,7 +113,7 @@ FScheduleInstanceData::~FScheduleInstanceData() = default;
 
 void FScheduleInstanceData::AddReferencedObjects(FReferenceCollector& Collector)
 {
-	for (TPair<FName, FUserScope>& ParamPair : UserScopes)
+	for (const TPair<FName, FUserScope>& ParamPair : UserScopes)
 	{
 		if(ParamPair.Value.AfterSource.IsValid())
 		{
@@ -125,6 +129,14 @@ void FScheduleInstanceData::AddReferencedObjects(FReferenceCollector& Collector)
 	{
 		Collector.AddPropertyReferencesWithStructARO(FAnimNextGraphInstancePtr::StaticStruct(), &GraphCache.GraphInstanceData);
 	}
+
+	for(const FExternalParamCache& ExternalParamCache : ExternalParamCaches)
+	{
+		for(const TUniquePtr<IParameterSource>& ParameterSource : ExternalParamCache.ParameterSources)
+		{
+			ParameterSource->AddReferencedObjects(Collector);
+		}
+	}
 }
 
 FString FScheduleInstanceData::GetReferencerName() const
@@ -144,4 +156,26 @@ TSharedPtr<FParamStack> FScheduleInstanceData::GetParamStack(uint32 InIndex) con
 	}
 }
 
+void FScheduleInstanceData::ApplyParametersToScope(FName InScope, EParameterScopeOrdering InOrdering, TUniquePtr<IParameterSource>&& InParameters)
+{
+	if (InScope == NAME_None)
+	{
+		RootUserScope = MoveTemp(InParameters);
+	}
+	else // apply to specified scope
+	{
+		FScheduleInstanceData::FUserScope& ScopeSource = UserScopes.FindOrAdd(InScope);
+		switch (InOrdering)
+		{
+		default:
+		case EParameterScopeOrdering::Before:
+			ScopeSource.BeforeSource = MoveTemp(InParameters);
+			break;
+		case EParameterScopeOrdering::After:
+			ScopeSource.AfterSource = MoveTemp(InParameters);
+			break;
+		}
+	}
+}
+
 }
\ No newline at end of file
diff --git a/Engine/Plugins/Experimental/AnimNext/Source/AnimNext/Private/Scheduler/ScheduleInstanceData.h b/Engine/Plugins/Experimental/AnimNext/Source/AnimNext/Private/Scheduler/ScheduleInstanceData.h
--- a/Engine/Plugins/Experimental/AnimNext/Source/AnimNext/Private/Scheduler/ScheduleInstanceData.h
+++ b/Engine/Plugins/Experimental/AnimNext/Source/AnimNext/Private/Scheduler/ScheduleInstanceData.h
@@ -17,6 +17,7 @@ namespace UE::AnimNext
 	struct FParamStackLayerHandle;
 	class IParameterSource;
 	struct FPropertyBagProxy;
+	enum class EParameterScopeOrdering : int32;
 }
 
 namespace UE::AnimNext
@@ -33,6 +34,9 @@ struct FScheduleInstanceData : public FGCObject
 	virtual void AddReferencedObjects(FReferenceCollector& Collector) override;
 	virtual FString GetReferencerName() const override;
 
+	// Apply the supplied parameter source to the specified scope, evicting any source that was there previously
+	void ApplyParametersToScope(FName InScope, EParameterScopeOrdering InOrdering, TUniquePtr<IParameterSource>&& InParameters);
+
 	// Get the appropriate params stack given the ID
 	TSharedPtr<FParamStack> GetParamStack(uint32 InIndex) const;
 	
@@ -43,7 +47,7 @@ struct FScheduleInstanceData : public FGCObject
 	FAnimNextSchedulerEntry* Entry = nullptr;
 
 	// Scope for user parameters to be applied at the root of the schedule
-	TUniquePtr<FPropertyBagProxy> RootUserScope;
+	TUniquePtr<IParameterSource> RootUserScope;
 
 	// Pushed layer for the root scope
 	FParamStack::FPushedLayerHandle PushedRootUserLayer;
@@ -51,10 +55,10 @@ struct FScheduleInstanceData : public FGCObject
 	struct FUserScope
 	{
 		// Layer that will be pushed before the scope, allowing the static scope to override the layer
-		TUniquePtr<FPropertyBagProxy> BeforeSource;
+		TUniquePtr<IParameterSource> BeforeSource;
 
 		// Layer that will be pushed after the scope, overriding the static scope
-		TUniquePtr<FPropertyBagProxy> AfterSource;
+		TUniquePtr<IParameterSource> AfterSource;
 	};
 	
 	// Set of dynamic parameter scopes supplied by the user
diff --git a/Engine/Plugins/Experimental/AnimNext/Source/AnimNext/Private/Scheduler/ScheduleTaskContext.cpp b/Engine/Plugins/Experimental/AnimNext/Source/AnimNext/Private/Scheduler/ScheduleTaskContext.cpp
new file mode 100644
--- /dev/null
+++ b/Engine/Plugins/Experimental/AnimNext/Source/AnimNext/Private/Scheduler/ScheduleTaskContext.cpp
@@ -0,0 +1,20 @@
+// Copyright Epic Games, Inc. All Rights Reserved.
+
+#include "Scheduler/ScheduleTaskContext.h"
+#include "Scheduler/ScheduleContext.h"
+
+namespace UE::AnimNext
+{
+
+FScheduleTaskContext::FScheduleTaskContext(const FScheduleContext& InContext)
+	: Context(InContext)
+{
+}
+
+void FScheduleTaskContext::ApplyParametersToScope(FName InScope, EParameterScopeOrdering InOrdering, TUniquePtr<IParameterSource>&& InParameters) const
+{
+	FScheduleInstanceData& InstanceData = Context.GetInstanceData();
+	InstanceData.ApplyParametersToScope(InScope, InOrdering, MoveTemp(InParameters));
+}
+
+}
\ No newline at end of file
diff --git a/Engine/Plugins/Experimental/AnimNext/Source/AnimNext/Private/Scheduler/ScheduleTickFunction.cpp b/Engine/Plugins/Experimental/AnimNext/Source/AnimNext/Private/Scheduler/ScheduleTickFunction.cpp
--- a/Engine/Plugins/Experimental/AnimNext/Source/AnimNext/Private/Scheduler/ScheduleTickFunction.cpp
+++ b/Engine/Plugins/Experimental/AnimNext/Source/AnimNext/Private/Scheduler/ScheduleTickFunction.cpp
@@ -11,6 +11,7 @@
 #include "Scheduler/AnimNextScheduleParamScopeTask.h"
 #include "Scheduler/ScheduleInstanceData.h"
 #include "Scheduler/AnimNextSchedulerWorldSubsystem.h"
+#include "Scheduler/ScheduleTaskContext.h"
 
 namespace UE::AnimNext
 {
@@ -24,9 +25,9 @@ void FScheduleBeginTickFunction::Run(float DeltaTime)
 {
 	while (!PreExecuteTasks.IsEmpty())
 	{
-		TOptional<TUniqueFunction<void(const UE::AnimNext::FScheduleContext&)>> Function = PreExecuteTasks.Dequeue();
+		TOptional<TUniqueFunction<void(const UE::AnimNext::FScheduleTaskContext&)>> Function = PreExecuteTasks.Dequeue();
 		check(Function.IsSet());
-		Function.GetValue()(Entry.Context);
+		Function.GetValue()(FScheduleTaskContext(Entry.Context));
 	}
 
 	// Push any user layer we have at the root
@@ -113,18 +114,18 @@ void FScheduleTickFunction::Run()
 		{
 			while (!PreExecuteTasks.IsEmpty())
 			{
-				TOptional<TUniqueFunction<void(const FScheduleContext&)>> Function = PreExecuteTasks.Dequeue();
+				TOptional<TUniqueFunction<void(const FScheduleTaskContext&)>> Function = PreExecuteTasks.Dequeue();
 				check(Function.IsSet());
-				Function.GetValue()(ScheduleContext);
+				Function.GetValue()(FScheduleTaskContext(ScheduleContext));
 			}
 		}, 
 		[this]()
 		{
 			while (!PostExecuteTasks.IsEmpty())
 			{
-				TOptional<TUniqueFunction<void(const FScheduleContext&)>> Function = PostExecuteTasks.Dequeue();
+				TOptional<TUniqueFunction<void(const FScheduleTaskContext&)>> Function = PostExecuteTasks.Dequeue();
 				check(Function.IsSet());
-				Function.GetValue()(ScheduleContext);
+				Function.GetValue()(FScheduleTaskContext(ScheduleContext));
 			}
 		});
 }
@@ -243,7 +244,33 @@ void FScheduleTickFunction::RunScheduleHelper(const FScheduleContext& InSchedule
 
 FString FScheduleTickFunction::DiagnosticMessage()
 {
-	return TEXT("AnimNextScheduleTickFunction");
+	if(Instructions.Num() == 1)
+	{
+		const FAnimNextScheduleInstruction& Instruction = Instructions[0];
+		switch (Instruction.Opcode)
+		{
+		case EAnimNextScheduleScheduleOpcode::RunGraphTask:
+			return TEXT("AnimNextSchedule::RunGraphTask");
+		case EAnimNextScheduleScheduleOpcode::BeginRunExternalTask:
+			return TEXT("AnimNextSchedule::BeginRunExternalTask");
+		case EAnimNextScheduleScheduleOpcode::EndRunExternalTask:
+			return TEXT("AnimNextSchedule::EndRunExternalTask");
+		case EAnimNextScheduleScheduleOpcode::RunPort:
+			return TEXT("AnimNextSchedule::RunPort");
+		case EAnimNextScheduleScheduleOpcode::RunParamScopeEntry:
+			return TEXT("AnimNextSchedule::RunParamScopeEntry");
+		case EAnimNextScheduleScheduleOpcode::RunParamScopeExit:
+			return TEXT("AnimNextSchedule::RunParamScopeExit");
+		case EAnimNextScheduleScheduleOpcode::RunExternalParamTask:
+			return TEXT("AnimNextSchedule::RunExternalParamTask");
+		default:
+			return TEXT("AnimNextScheduleTickFunction");
+		}
+	}
+	else
+	{
+		return TEXT("AnimNextScheduleTickFunction");
+	}
 }
 
 }
\ No newline at end of file
diff --git a/Engine/Plugins/Experimental/AnimNext/Source/AnimNext/Private/Scheduler/ScheduleTickFunction.h b/Engine/Plugins/Experimental/AnimNext/Source/AnimNext/Private/Scheduler/ScheduleTickFunction.h
--- a/Engine/Plugins/Experimental/AnimNext/Source/AnimNext/Private/Scheduler/ScheduleTickFunction.h
+++ b/Engine/Plugins/Experimental/AnimNext/Source/AnimNext/Private/Scheduler/ScheduleTickFunction.h
@@ -12,6 +12,7 @@ struct FAnimNextSchedulerEntry;
 namespace UE::AnimNext
 {
 	struct FScheduleContext;
+	struct FScheduleTaskContext;
 	struct FParamStackLayerHandle;
 }
 
@@ -34,7 +35,7 @@ struct FScheduleBeginTickFunction : public FTickFunction
 
 	void Run(float DeltaTime);
 
-	TSpscQueue<TUniqueFunction<void(const UE::AnimNext::FScheduleContext&)>> PreExecuteTasks;
+	TSpscQueue<TUniqueFunction<void(const UE::AnimNext::FScheduleTaskContext&)>> PreExecuteTasks;
 	FAnimNextSchedulerEntry& Entry;
 	FTickPrerequisite Subsequent;
 };
@@ -86,8 +87,8 @@ struct FScheduleTickFunction : public FTickFunction
 	TConstArrayView<FAnimNextScheduleInstruction> Instructions;
 	TArray<FTickPrerequisite> Subsequents;
 	TConstArrayView<TWeakObjectPtr<UObject>> TargetObjects;
-	TSpscQueue<TUniqueFunction<void(const UE::AnimNext::FScheduleContext&)>> PreExecuteTasks;
-	TSpscQueue<TUniqueFunction<void(const UE::AnimNext::FScheduleContext&)>> PostExecuteTasks;
+	TSpscQueue<TUniqueFunction<void(const UE::AnimNext::FScheduleTaskContext&)>> PreExecuteTasks;
+	TSpscQueue<TUniqueFunction<void(const UE::AnimNext::FScheduleTaskContext&)>> PostExecuteTasks;
 };
 
 }
\ No newline at end of file
diff --git a/Engine/Plugins/Experimental/AnimNext/Source/AnimNext/Private/Scheduler/Scheduler.cpp b/Engine/Plugins/Experimental/AnimNext/Source/AnimNext/Private/Scheduler/Scheduler.cpp
--- a/Engine/Plugins/Experimental/AnimNext/Source/AnimNext/Private/Scheduler/Scheduler.cpp
+++ b/Engine/Plugins/Experimental/AnimNext/Source/AnimNext/Private/Scheduler/Scheduler.cpp
@@ -6,12 +6,8 @@
 #include "Scheduler/ScheduleContext.h"
 #include "Scheduler/AnimNextSchedule.h"
 #include "Scheduler/AnimNextSchedulerWorldSubsystem.h"
-#include "Tasks/Task.h"
-#include "UObject/GCObject.h"
-#include "LODPose.h"
 #include "Engine/World.h"
 #include "Components/SkeletalMeshComponent.h"
-#include "UObject/ObjectKey.h"
 #include "UObject/UObjectIterator.h"
 
 namespace UE::AnimNext
@@ -50,10 +46,8 @@ void FScheduler::Destroy()
 	FWorldDelegates::OnWorldPreActorTick.Remove(Impl.OnWorldPreActorTickHandle);
 }
 
-FScheduleHandle FScheduler::AcquireHandle(UObject* InObject, UAnimNextSchedule* InSchedule, EAnimNextScheduleInitMethod InInitMethod, TUniqueFunction<void(const FScheduleContext&)>&& InInitializeCallback)
+FScheduleHandle FScheduler::AcquireHandle(UObject* InObject, UAnimNextSchedule* InSchedule, EAnimNextScheduleInitMethod InInitMethod, TUniqueFunction<void(const FScheduleInitializationContext&)>&& InInitializeCallback)
 {
-	FScheduleHandle Handle;
-
 	// Check parameters
 	if (InSchedule == nullptr)
 	{
@@ -127,7 +121,7 @@ void FScheduler::EnableHandle(UObject* InObject, FScheduleHandle InHandle, bool
 	Subsystem->EnableHandle(InHandle, bInEnabled);
 }
 
-void FScheduler::QueueTask(UObject* InObject, FScheduleHandle InHandle, FName InScheduleTaskName, TUniqueFunction<void(const FScheduleContext&)>&& InTaskFunction, ETaskRunLocation InLocation)
+void FScheduler::QueueTask(UObject* InObject, FScheduleHandle InHandle, FName InScheduleTaskName, TUniqueFunction<void(const FScheduleTaskContext&)>&& InTaskFunction, ETaskRunLocation InLocation)
 {
 	if(!InHandle.IsValid())
 	{
diff --git a/Engine/Plugins/Experimental/AnimNext/Source/AnimNext/Public/IAnimNextModule.h b/Engine/Plugins/Experimental/AnimNext/Source/AnimNext/Public/IAnimNextModule.h
--- a/Engine/Plugins/Experimental/AnimNext/Source/AnimNext/Public/IAnimNextModule.h
+++ b/Engine/Plugins/Experimental/AnimNext/Source/AnimNext/Public/IAnimNextModule.h
@@ -8,32 +8,63 @@
 #include "ReferencePose.h"
 
 struct FAnimNextGraphInstancePtr;
+struct FAnimNextParamInstanceIdentifier;
 
 namespace UE::AnimNext
 {
-	class ANIMNEXT_API IAnimNextAnimGraph
-	{
-	public:
-		virtual ~IAnimNextAnimGraph() = default;
-
-		virtual void UpdateGraph(FAnimNextGraphInstancePtr& GraphInstance, float DeltaTime) const = 0;
-		virtual void EvaluateGraph(FAnimNextGraphInstancePtr& GraphInstance, const UE::AnimNext::FReferencePose& RefPose, int32 GraphLODLevel, FLODPoseHeap& OutputPose) const = 0;
-	};
-
-	class ANIMNEXT_API IAnimNextModule : public IModuleInterface
-	{
-	public:
-		// WARNING
-		// These functions are here because the schedule currently cannot have tasks defined in
-		// external plugins. And as such, we must be able to update and evaluate graphs from
-		// AnimNextCore even though they are implementation details of the AnimNextAnimGraph.
-		// We thus introduce this interface.
-		// When AnimNextAnimGraph's module will load, it will register callbacks that
-		// AnimNextCore will call when we need to update/evaluate graphs.
-		static IAnimNextModule& Get();
-		virtual void RegisterAnimNextAnimGraph(const IAnimNextAnimGraph& InAnimGraphImpl) = 0;
-		virtual void UnregisterAnimNextAnimGraph() = 0;
-		virtual void UpdateGraph(FAnimNextGraphInstancePtr& GraphInstance, float DeltaTime) = 0;
-		virtual void EvaluateGraph(FAnimNextGraphInstancePtr& GraphInstance, const UE::AnimNext::FReferencePose& RefPose, int32 GraphLODLevel, FLODPoseHeap& OutputPose) const = 0;
-	};
+	class IParameterSource;
+	class IParameterSourceFactory;
+	struct FParameterSourceContext;
+}
+
+namespace UE::AnimNext
+{
+
+class IAnimNextAnimGraph
+{
+public:
+	virtual ~IAnimNextAnimGraph() = default;
+
+	virtual void UpdateGraph(FAnimNextGraphInstancePtr& GraphInstance, float DeltaTime) const = 0;
+	virtual void EvaluateGraph(FAnimNextGraphInstancePtr& GraphInstance, const UE::AnimNext::FReferencePose& RefPose, int32 GraphLODLevel, FLODPoseHeap& OutputPose) const = 0;
+};
+
+class IAnimNextModule : public IModuleInterface
+{
+public:
+	// Factory method used to create a parameter source to access the specified scope locator, with a set of parameters that are initially required
+	// @param   InContext             Context used to set up the parameter source
+	// @param   InInstanceId          The instance identifier associated with the parameters that are required
+	// @param   InRequiredParameters  Any required parameters that the source should initially supply, can be empty
+	// @return a new parameter source, or nullptr if the scope was not found or could not be handled
+	virtual TUniquePtr<IParameterSource> CreateParameterSource(const FParameterSourceContext& InContext, const TInstancedStruct<FAnimNextParamInstanceIdentifier>& InInstanceId, TConstArrayView<FName> InRequiredParameters) const = 0;
+
+	// Register a factory that can be used to generate parameter sources
+	// @param   InName                Identifier for the factory
+	// @param   InFactory             The factory to register
+	virtual void RegisterParameterSourceFactory(FName InName, TSharedRef<IParameterSourceFactory> InFactory) = 0;
+
+	// Unregister a factory previously passed to RegisterParameterSourceFactory
+	// @param   InName                Identifier for the factory
+	virtual void UnregisterParameterSourceFactory(FName InName) = 0;
+
+	// Find a factory previously passed to RegisterParameterSourceFactory
+	// @param   InName                Identifier for the factory
+	// @return a shared ptr to the factory, if found, otherwise an invalid shared ptr
+	virtual TSharedPtr<IParameterSourceFactory> FindParameterSourceFactory(FName InName) = 0;
+
+	// WARNING
+	// These functions are here because the schedule currently cannot have tasks defined in
+	// external plugins. And as such, we must be able to update and evaluate graphs from
+	// AnimNextCore even though they are implementation details of the AnimNextAnimGraph.
+	// We thus introduce this interface.
+	// When AnimNextAnimGraph's module will load, it will register callbacks that
+	// AnimNextCore will call when we need to update/evaluate graphs.
+	static ANIMNEXT_API IAnimNextModule& Get();
+	virtual void RegisterAnimNextAnimGraph(const IAnimNextAnimGraph& InAnimGraphImpl) = 0;
+	virtual void UnregisterAnimNextAnimGraph() = 0;
+	virtual void UpdateGraph(FAnimNextGraphInstancePtr& GraphInstance, float DeltaTime) = 0;
+	virtual void EvaluateGraph(FAnimNextGraphInstancePtr& GraphInstance, const FReferencePose& RefPose, int32 GraphLODLevel, FLODPoseHeap& OutputPose) const = 0;
+};
+
 }
diff --git a/Engine/Plugins/Experimental/AnimNext/Source/AnimNext/Public/Param/AnimNextClassExtensionLibrary.h b/Engine/Plugins/Experimental/AnimNext/Source/AnimNext/Public/Param/AnimNextClassExtensionLibrary.h
--- a/Engine/Plugins/Experimental/AnimNext/Source/AnimNext/Public/Param/AnimNextClassExtensionLibrary.h
+++ b/Engine/Plugins/Experimental/AnimNext/Source/AnimNext/Public/Param/AnimNextClassExtensionLibrary.h
@@ -8,7 +8,7 @@
 // Base class for function library-like classes that act to adapt UObjects of a specified type, exposing them to AnimNext parameters.
 // To do this, implement GetSupportedClass, then provide static UFUNCTION members that take the type as an input, e.g.
 // static float GetDeltaSeconds(UWorld* InWorld); 
-UCLASS(Abstract)
+UCLASS(MinimalAPI, Abstract)
 class UAnimNextClassExtensionLibrary : public UObject
 {
 	GENERATED_BODY()
diff --git a/Engine/Plugins/Experimental/AnimNext/Source/AnimNext/Public/Param/IParameterSourceFactory.h b/Engine/Plugins/Experimental/AnimNext/Source/AnimNext/Public/Param/IParameterSourceFactory.h
--- a/Engine/Plugins/Experimental/AnimNext/Source/AnimNext/Public/Param/IParameterSourceFactory.h
+++ b/Engine/Plugins/Experimental/AnimNext/Source/AnimNext/Public/Param/IParameterSourceFactory.h
@@ -3,7 +3,8 @@
 #pragma once
 
 #include "Param/IParameterSource.h"
-#include "Features/IModularFeature.h"
+
+struct FAnimNextParamInstanceIdentifier;
 
 namespace UE::AnimNext
 {
@@ -14,53 +15,32 @@ namespace UE::AnimNext
 {
 
 // Context passed to object accessor functions registered to RegisterObjectAccessor
-struct FExternalParameterContext
+struct FParameterSourceContext
 {
+	FParameterSourceContext() = default;
+
+	FParameterSourceContext(const UObject* InObject)
+		: Object(InObject)
+	{
+		check(Object != nullptr);
+	}
+
 	// The object that the AnimNext entry is bound to (e.g. a UAnimNextComponent)
-	UObject* Object = nullptr;
+	const UObject* Object = nullptr;
 };
 
-// Modular feature allowing other modules to extend and add to external parameter system
-class IParameterSourceFactory : public IModularFeature
+// Interface allowing other modules to extend and add to external parameter system
+class IParameterSourceFactory
 {
 public:
-	inline static const FName FeatureName = TEXT("AnimNextParameterSourceFactory");
-
 	virtual ~IParameterSourceFactory() = default;
 
-	// Iterate over all the source names that this factory provides
-	virtual void ForEachSource(TFunctionRef<void(FName)> InFunction) const = 0;
-
 	// Factory method used to create a parameter source of the specified name, with a set of parameters that are initially required
-	// @param   InContext              Context used to set up the parameter source
-	// @param   InSourceName           The name of the parameter source. This is conventionally the 'root' parameter of a set of parameters (e.g.
-	//                                 Game_MyComponent, which contains Game_MyComponent_Param etc.)
-	// @param   InRequiredParameters   Any required parameters that the source should initially supply, can be empty
+	// @param  InContext              Context used to resolve the scope
+	// @param  InInstanceId           The instance identifier associated with the parameters that are required
+	// @param  InRequiredParameters   Any required parameters that the source should initially supply, can be empty
 	// @return a new parameter source, or nullptr if the source was not found
-	virtual TUniquePtr<IParameterSource> CreateParameterSource(const FExternalParameterContext& InContext, FName InSourceName, TConstArrayView<FName> InRequiredParameters) const = 0;
-
-#if WITH_EDITOR
-	// Info about parameters gleaned from FindParameterInfo and ForEachParameter
-	struct FParameterInfo
-	{
-		FAnimNextParamType Type;
-		FText Tooltip;
-		bool bThreadSafe = false;
-	};
-
-	// Given a parameter name, find associated info for that parameter
-	// @param   InParameterName        The parameter name to find
-	// @param   OutType                The parameter type
-	// @param   OutTooltip             Optional: The parameter's tooltip
-	// @return true if the parameter info was found
-	virtual bool FindParameterInfo(FName InParameterName, FParameterInfo& OutInfo) const = 0;
-
-	// Iterate over all the known parameters for a particular source
-	// @param   InSourceName           The name of the parameter source to iterate
-	// @param   InFunction             Function called for each known parameter
-	// @return an array of parameter names
-	virtual void ForEachParameter(FName InSourceName, TFunctionRef<void(FName, const FParameterInfo&)> InFunction) const = 0;
-#endif
+	virtual TUniquePtr<IParameterSource> CreateParameterSource(const FParameterSourceContext& InContext, const TInstancedStruct<FAnimNextParamInstanceIdentifier>& InInstanceId, TConstArrayView<FName> InRequiredParameters) const = 0;
 };
 
 }
\ No newline at end of file
diff --git a/Engine/Plugins/Experimental/AnimNext/Source/AnimNext/Public/Param/ParamEntry.h b/Engine/Plugins/Experimental/AnimNext/Source/AnimNext/Public/Param/ParamEntry.h
--- a/Engine/Plugins/Experimental/AnimNext/Source/AnimNext/Public/Param/ParamEntry.h
+++ b/Engine/Plugins/Experimental/AnimNext/Source/AnimNext/Public/Param/ParamEntry.h
@@ -121,6 +121,12 @@ private:
 		return Id.GetName();
 	}
 
+	// Get the instance ID of this parameter
+	FName GetInstanceId() const
+	{
+		return Id.GetInstanceId();
+	}
+
 	// Get the hash of this parameter
 	uint32 GetHash() const
 	{
@@ -138,7 +144,7 @@ private:
 
 	// ID of the parameter
 	FParamId Id;
-	
+
 	// The type of the param
 	FParamTypeHandle TypeHandle;
 
diff --git a/Engine/Plugins/Experimental/AnimNext/Source/AnimNext/Public/Param/ParamId.h b/Engine/Plugins/Experimental/AnimNext/Source/AnimNext/Public/Param/ParamId.h
--- a/Engine/Plugins/Experimental/AnimNext/Source/AnimNext/Public/Param/ParamId.h
+++ b/Engine/Plugins/Experimental/AnimNext/Source/AnimNext/Public/Param/ParamId.h
@@ -44,7 +44,16 @@ struct ANIMNEXT_API FParamId
 	// Make a parameter ID from an FName, generating the hash
 	explicit FParamId(FName InName)
 		: Name(InName)
-		, Hash(GetTypeHash(InName))
+		, InstanceId(NAME_None)
+		, Hash(CalculateHash(InName, NAME_None))
+	{
+	}
+
+	// Make a parameter ID from an FName and instance ID, generating the hash
+	explicit FParamId(FName InName, FName InInstanceId)
+		: Name(InName)
+		, InstanceId(InInstanceId)
+		, Hash(CalculateHash(InName, InInstanceId))
 	{
 	}
 
@@ -53,7 +62,16 @@ struct ANIMNEXT_API FParamId
 		: Name(InName)
 		, Hash(InHash)
 	{
-		checkSlow(Hash == GetTypeHash(Name));
+		checkSlow(CalculateHash(InName, NAME_None) == InHash);
+	}
+
+	// Make a parameter ID from a name, an instance ID and hash
+	explicit FParamId(FName InName, FName InInstanceId, uint32 InHash)
+		: Name(InName)
+		, InstanceId(InInstanceId)
+		, Hash(InHash)
+	{
+		checkSlow(CalculateHash(InName, InInstanceId) == InHash);
 	}
 
 	// Get the name of this param
@@ -62,6 +80,12 @@ struct ANIMNEXT_API FParamId
 		return Name;
 	}
 
+	// Get the instance ID of this param
+	FName GetInstanceId() const
+	{
+		return InstanceId;
+	}
+	
 	// Get the hash of this param
 	uint32 GetHash() const
 	{
@@ -74,10 +98,19 @@ struct ANIMNEXT_API FParamId
 		return Hash != 0;
 	}
 
+	// Get the hash of a parameter name/instance ID combination
+	static uint32 CalculateHash(FName InName, FName InInstanceId)
+	{
+		return HashCombineFast(GetTypeHash(InName), GetTypeHash(InInstanceId));
+	}
+	
 private:
 	// Parameter name
 	FName Name;
 
+	// Parameter instance ID
+	FName InstanceId;
+
 	// Name hash
 	uint32 Hash = 0;
 };
diff --git a/Engine/Plugins/Experimental/AnimNext/Source/AnimNext/Public/Param/ParamStack.h b/Engine/Plugins/Experimental/AnimNext/Source/AnimNext/Public/Param/ParamStack.h
--- a/Engine/Plugins/Experimental/AnimNext/Source/AnimNext/Public/Param/ParamStack.h
+++ b/Engine/Plugins/Experimental/AnimNext/Source/AnimNext/Public/Param/ParamStack.h
@@ -198,7 +198,9 @@ public:
 	static ANIMNEXT_API FParamStackLayerHandle MakeValueLayer(const FInstancedPropertyBag& InInstancedPropertyBag);
 
 	// Create a cached parameter layer from an instanced property bag. This layer will reference the supplied property bag and does not transfer ownership.
-	static ANIMNEXT_API FParamStackLayerHandle MakeReferenceLayer(FInstancedPropertyBag& InInstancedPropertyBag);
+	// @param	InInstanceId					The instance ID to apply to the supplied parameters (as arguments to their FParamId)
+	// @param	InInstancedPropertyBag	THe property bag to reference
+	static ANIMNEXT_API FParamStackLayerHandle MakeReferenceLayer(FName InInstanceId, FInstancedPropertyBag& InInstancedPropertyBag);
 
 	// Create a cached parameter layer by remapping the entries from another layer
 	// Original layer continues to own the memory
diff --git a/Engine/Plugins/Experimental/AnimNext/Source/AnimNext/Public/Param/ParamStackLayerHandle.h b/Engine/Plugins/Experimental/AnimNext/Source/AnimNext/Public/Param/ParamStackLayerHandle.h
--- a/Engine/Plugins/Experimental/AnimNext/Source/AnimNext/Public/Param/ParamStackLayerHandle.h
+++ b/Engine/Plugins/Experimental/AnimNext/Source/AnimNext/Public/Param/ParamStackLayerHandle.h
@@ -151,11 +151,7 @@ public:
 	template<typename WrappedType>
 	WrappedType* As() const
 	{
-		if constexpr (TModels<CStaticClassProvider, WrappedType>::Value)
-		{
-			return Cast<WrappedType>(GetUObjectFromLayer());
-		}
-		else if constexpr (std::is_same_v<WrappedType, FInstancedPropertyBag>)
+		if constexpr (std::is_same_v<WrappedType, FInstancedPropertyBag>)
 		{
 			return GetInstancedPropertyBagFromLayer();
 		}
@@ -171,8 +167,7 @@ private:
 	// Set parameter values
 	ANIMNEXT_API FParamResult SetValuesInternal(TConstArrayView<Private::FParamEntry> InParams);
 
-	// Helpers for As() to hide Layer
-	ANIMNEXT_API UObject* GetUObjectFromLayer() const;
+	// Helper for As() to hide Layer
 	ANIMNEXT_API FInstancedPropertyBag* GetInstancedPropertyBagFromLayer() const;
 
 	// Recursive helper function for SetValues
diff --git a/Engine/Plugins/Experimental/AnimNext/Source/AnimNext/Public/Param/ParamType.h b/Engine/Plugins/Experimental/AnimNext/Source/AnimNext/Public/Param/ParamType.h
--- a/Engine/Plugins/Experimental/AnimNext/Source/AnimNext/Public/Param/ParamType.h
+++ b/Engine/Plugins/Experimental/AnimNext/Source/AnimNext/Public/Param/ParamType.h
@@ -71,8 +71,6 @@ public:
 	/** Construct a parameter type from the passed in value, container and object type. */
 	FAnimNextParamType(EValueType InValueType, EContainerType InContainerType = EContainerType::None, const UObject* InValueTypeObject = nullptr);
 
-	/** Construct a parameter type from the passed in FRigVMTemplateArgumentType. */
-	static FAnimNextParamType FromRigVMTemplateArgument(const FRigVMTemplateArgumentType& RigVMType);
 private:
 	/** Pointer to object that defines the Enum, Struct, or Class. */
 	UPROPERTY()
@@ -284,6 +282,12 @@ public:
 	/** Get a type from a string */
 	static FAnimNextParamType FromString(const FString& InString);
 
+	/** Get a FRigVMTemplateArgumentType from this type */
+	FRigVMTemplateArgumentType ToRigVMTemplateArgument() const;
+	
+	/** Construct a parameter type from the passed in FRigVMTemplateArgumentType. */
+	static FAnimNextParamType FromRigVMTemplateArgument(const FRigVMTemplateArgumentType& RigVMType);
+	
 	/** Equality operator */
 	friend bool operator==(const FAnimNextParamType& InLHS, const FAnimNextParamType& InRHS)
 	{
@@ -301,7 +305,13 @@ public:
 	{
 		return HashCombineFast(GetTypeHash((uint32)InType.ValueType | ((uint32)InType.ContainerType << 8)), GetTypeHash(InType.ValueTypeObject));
 	}
-	
+
+	/** @return whether this type is explicitly none */
+	bool IsNone() const
+	{
+		return ValueType == EValueType::None && ContainerType == EContainerType::None && !ValueTypeObject;
+	}
+
 	/** @return whether this type actually describes a type */
 	bool IsValid() const
 	{
diff --git a/Engine/Plugins/Experimental/AnimNext/Source/AnimNext/Public/Scheduler/AnimNextSchedulePort.h b/Engine/Plugins/Experimental/AnimNext/Source/AnimNext/Public/Scheduler/AnimNextSchedulePort.h
--- a/Engine/Plugins/Experimental/AnimNext/Source/AnimNext/Public/Scheduler/AnimNextSchedulePort.h
+++ b/Engine/Plugins/Experimental/AnimNext/Source/AnimNext/Public/Scheduler/AnimNextSchedulePort.h
@@ -8,7 +8,7 @@
 
 class UAnimNextSchedule;
 struct FAnimNextSchedulePortTask;
-struct FAnimNextParam;
+struct FAnimNextEditorParam;
 
 namespace UE::AnimNext
 {
@@ -28,5 +28,5 @@ private:
 	virtual void Run(const UE::AnimNext::FScheduleTermContext& InContext) const PURE_VIRTUAL(UAnimNextSchedulePort::Run, )
 
 	// Get any required parameters for this port
-	virtual TConstArrayView<FAnimNextParam> GetRequiredParameters() const PURE_VIRTUAL(UAnimNextSchedulePort::GetRequiredParameters, return TConstArrayView<FAnimNextParam>(); )
+	virtual TConstArrayView<FAnimNextEditorParam> GetRequiredParameters() const PURE_VIRTUAL(UAnimNextSchedulePort::GetRequiredParameters, return TConstArrayView<FAnimNextEditorParam>(); )
 };
diff --git a/Engine/Plugins/Experimental/AnimNext/Source/AnimNextEditor/AnimNextEditor.Build.cs b/Engine/Plugins/Experimental/AnimNext/Source/AnimNextEditor/AnimNextEditor.Build.cs
--- a/Engine/Plugins/Experimental/AnimNext/Source/AnimNextEditor/AnimNextEditor.Build.cs
+++ b/Engine/Plugins/Experimental/AnimNext/Source/AnimNextEditor/AnimNextEditor.Build.cs
@@ -45,6 +45,10 @@ namespace UnrealBuildTool.Rules
 					"EditorWidgets",
 					"WorkspaceEditor",
 					"ContentBrowser",
+					"UniversalObjectLocator",
+					"UniversalObjectLocatorEditor",
+					"Kismet",
+					"AdvancedWidgets",
 				}
 			);
 		}
diff --git a/Engine/Plugins/Experimental/AnimNext/Source/AnimNextEditor/Internal/AnimNextEditorModule.h b/Engine/Plugins/Experimental/AnimNext/Source/AnimNextEditor/Internal/AnimNextEditorModule.h
--- a/Engine/Plugins/Experimental/AnimNext/Source/AnimNextEditor/Internal/AnimNextEditorModule.h
+++ b/Engine/Plugins/Experimental/AnimNext/Source/AnimNextEditor/Internal/AnimNextEditorModule.h
@@ -16,6 +16,7 @@ class IModule : public IModuleInterface
 {
 public:
 	// Create a parameter picker
+	// @param  InArgs          Arguments used for configuring the picker
 	virtual TSharedRef<SWidget> CreateParameterPicker(const FParameterPickerArgs& InArgs) = 0;
 };
 
diff --git a/Engine/Plugins/Experimental/AnimNext/Source/AnimNextEditor/Internal/Param/ParameterPickerArgs.h b/Engine/Plugins/Experimental/AnimNext/Source/AnimNextEditor/Internal/Param/ParameterPickerArgs.h
--- a/Engine/Plugins/Experimental/AnimNext/Source/AnimNextEditor/Internal/Param/ParameterPickerArgs.h
+++ b/Engine/Plugins/Experimental/AnimNext/Source/AnimNextEditor/Internal/Param/ParameterPickerArgs.h
@@ -2,11 +2,16 @@
 
 #pragma once
 
+#include "UniversalObjectLocator.h"
 #include "UObject/NameTypes.h"
 #include "Delegates/DelegateCombinations.h"
 #include "AssetRegistry/AssetData.h"
+#include "Param/AnimNextParam.h"
 #include "Param/ParamTypeHandle.h"
 
+struct FAnimNextParamInstanceIdentifier;
+class UAnimNextRigVMAsset;
+
 namespace UE::AnimNext::Editor
 {
 
@@ -38,9 +43,10 @@ struct FParameterBindingReference
 {
 	FParameterBindingReference() = default;
 
-	FParameterBindingReference(FName InParameter, const FAnimNextParamType& InType, const FAssetData& InGraph = FAssetData())
+	FParameterBindingReference(FName InParameter, const FAnimNextParamType& InType, const TInstancedStruct<FAnimNextParamInstanceIdentifier>& InInstanceId,  const FAssetData& InGraph = FAssetData())
 		: Parameter(InParameter)
 		, Type(InType)
+		, InstanceId(InInstanceId)
 		, Graph(InGraph)
 	{
 	}
@@ -51,21 +57,21 @@ struct FParameterBindingReference
 	// Parameter type
 	FAnimNextParamType Type;
 
-	// Asset (first found in asset registry) that the parameter is used in
-	FAssetData Asset;
+	// Instance Id used to disambiguate the parameter
+	TInstancedStruct<FAnimNextParamInstanceIdentifier> InstanceId;
 
 	// Optional graph asset that the parameter is bound in
 	FAssetData Graph;
 };
 
 // Delegate called when a parameter has been picked. Graph argument is invalid when an unbound parameter is chosen.
-DECLARE_DELEGATE_OneParam(FOnGetParameterBindings, TArray<FParameterBindingReference>& /*OutParameterBindings*/);
+using FOnGetParameterBindings = TDelegate<void(TArray<FParameterBindingReference>& /*OutParameterBindings*/)>;
 
 // Delegate called when a parameter has been picked. Graph argument is invalid when an unbound parameter is chosen.
-DECLARE_DELEGATE_OneParam(FOnParameterPicked, const FParameterBindingReference& /*InParameterBinding*/);
+using FOnParameterPicked = TDelegate<void(const FParameterBindingReference& /*InParameterBinding*/)>;
 
 // Delegate called when a parameter is due to be added.
-DECLARE_DELEGATE_OneParam(FOnAddParameter, const FParameterToAdd& /*InParameterToAdd*/);
+using FOnAddParameter = TDelegate<void(const FParameterToAdd& /*InParameterToAdd*/)>;
 
 // Result of a filter operation via FOnFilterParameter
 enum class EFilterParameterResult : int32
@@ -75,17 +81,22 @@ enum class EFilterParameterResult : int32
 };
 
 // Delegate called to filter parameters for display to the user
-DECLARE_DELEGATE_RetVal_OneParam(EFilterParameterResult, FOnFilterParameter, const FParameterBindingReference& /*InParameterBinding*/);
+using FOnFilterParameter = TDelegate<EFilterParameterResult(const FParameterBindingReference& /*InParameterBinding*/)>;
+
+// Delegate called to filter parameters by type for display to the user
+using FOnFilterParameterType = TDelegate<EFilterParameterResult(const FAnimNextParamType& /*InParameterType*/)>;
 
 // Delegate called to filter parameters by type for display to the user
-DECLARE_DELEGATE_RetVal_OneParam(EFilterParameterResult, FOnFilterParameterType, const FAnimNextParamType& /*InParameterType*/);
+using FOnFilterParameterType = TDelegate<EFilterParameterResult(const FAnimNextParamType& /*InParameterType*/)>;
+
+// Delegate called when the selected instance ID changes 
+using FOnInstanceIdChanged = TDelegate<void(const TInstancedStruct<FAnimNextParamInstanceIdentifier>& /*InInstanceId*/)>;
 
 struct FParameterPickerArgs
 {
 	FParameterPickerArgs() = default;
 
-	// Ptr to existing called delegate to which the picker will register a function which returns the selected parameter
-	// bindings
+	// Ptr to existing called delegate to which the picker will register a function which returns the selected parameter bindings
 	FOnGetParameterBindings* OnGetParameterBindings = nullptr;
 
 	// Delegate used to signal whether selection has changed
@@ -106,29 +117,26 @@ struct FParameterPickerArgs
 	// Type to use for any new parameters generated through the picker
 	FAnimNextParamType NewParameterType;
 
-	// Whether we allow selecting multiple parameters or just one
-	bool bMultiSelect = true;
-
-	// Whether we should show parameters that are bound in a graph
-	bool bShowBoundParameters = true;
+	// Ptr to existing called delegate to which the picker will register a function which returns the selected instance ID
+	FOnInstanceIdChanged* OnSetInstanceId = nullptr;
 
-	// Whether we should show parameters that are not bound in a graph (if bShowBoundParameters is false this will show all parameters)
-	bool bShowUnboundParameters = true;	
+	// Delegate called when the selected instance ID changes in the picker
+	FOnInstanceIdChanged OnInstanceIdChanged;
 
-	// Whether we should show parameters that are built in
-	bool bShowBuiltInParameters = true;
+	// The initial instance ID to use
+	TInstancedStruct<FAnimNextParamInstanceIdentifier> InstanceId;
 
-	// Whether we should show the source graph alongside bound parameters
-	bool bShowSourceGraph = true;
+	// The context to use when resolving instance IDs
+	const UObject* Context = nullptr;
 
-	// Whether we should allow new parameters to be created by this widget
-	bool bAllowNew = true;
+	// Whether we allow selecting multiple parameters or just one
+	bool bMultiSelect = true;
 
 	// Whether the search box should be focussed on widget creation
 	bool bFocusSearchWidget = true;
 
-	// Whether 'none' can be selected
-	bool bAllowNone = false;
+	// Whether the instance ID chooser is displayed
+	bool bShowInstanceId = true;
 };
 
 }
diff --git a/Engine/Plugins/Experimental/AnimNext/Source/AnimNextEditor/Private/AnimNextEditorModule.cpp b/Engine/Plugins/Experimental/AnimNext/Source/AnimNextEditor/Private/AnimNextEditorModule.cpp
--- a/Engine/Plugins/Experimental/AnimNext/Source/AnimNextEditor/Private/AnimNextEditorModule.cpp
+++ b/Engine/Plugins/Experimental/AnimNext/Source/AnimNextEditor/Private/AnimNextEditorModule.cpp
@@ -30,6 +30,14 @@
 #include "Common/SActionMenu.h"
 #include "AnimNextRigVMAssetEntry.h"
 #include "Editor/RigVMGraphDetailCustomization.h"
+#include "EditorUtils.h"
+#include "IUniversalObjectLocatorEditorModule.h"
+#include "Param/AnimNextParam.h"
+#include "Param/AnimNextLocatorContext.h"
+#include "Param/ObjectCastLocatorEditor.h"
+#include "Param/ObjectFunctionLocatorEditor.h"
+#include "Param/ObjectPropertyLocatorEditor.h"
+#include "Param/SAddParametersDialog.h"
 
 #define LOCTEXT_NAMESPACE "AnimNextEditorModule"
 
@@ -56,8 +64,12 @@ class FModule : public IModule
 
 		PropertyModule.RegisterCustomPropertyTypeLayout(
 			"AnimNextParam",
-			FOnGetPropertyTypeCustomizationInstance::CreateLambda([] { return MakeShared<FParamPropertyTypeCustomization>(); }));
+			FOnGetPropertyTypeCustomizationInstance::CreateLambda([] { return MakeShared<FParamPropertyCustomization>(); }));
 
+		PropertyModule.RegisterCustomPropertyTypeLayout(
+			"AnimNextEditorParam",
+			FOnGetPropertyTypeCustomizationInstance::CreateLambda([] { return MakeShared<FParamPropertyCustomization>(); }));
+		
 		Identifier = MakeShared<FParamNamePropertyTypeIdentifier>();
 		PropertyModule.RegisterCustomPropertyTypeLayout(
 			FNameProperty::StaticClass()->GetFName(),
@@ -104,54 +116,29 @@ class FModule : public IModule
 		SRigVMAssetView::RegisterCategoryFactory("Parameters", [](UAnimNextRigVMAssetEditorData* InEditorData)
 		{
 			UAnimNextGraph_EditorData* EditorData = CastChecked<UAnimNextGraph_EditorData>(InEditorData);
-			return SNew(SSimpleComboButton)
+			UAnimNextRigVMAsset* Asset = UncookedOnly::FUtils::GetAsset(InEditorData);
+			return SNew(SSimpleButton)
 				.Text(LOCTEXT("AddParameterButton", "Add Parameter"))
 				.Icon(FAppStyle::Get().GetBrush("Icons.Plus"))
-				.HasDownArrow(true)
-				.OnGetMenuContent_Lambda([EditorData]()
+				.OnClicked_Lambda([EditorData, Asset]()
 				{
-					FAssetData AssetData(UncookedOnly::FUtils::GetAsset(EditorData));
-					
-					FParameterPickerArgs Args;
-					Args.bMultiSelect = false;
-					Args.bShowSourceGraph = false;
-					Args.bShowBoundParameters = false;
-					Args.bShowBuiltInParameters = false; // Built-In parameters disabled for MVP
-					Args.OnFilterParameter = FOnFilterParameter::CreateLambda([EditorData, AssetData](const FParameterBindingReference& InParameterBinding)
-					{
-						// Skip params that are already bound in this graph
-						if(InParameterBinding.Graph == AssetData)
-						{
-							return EFilterParameterResult::Exclude;
-						}
-						
-						return EFilterParameterResult::Include;
-					});
+					TSharedRef<SAddParametersDialog> AddParametersDialog =
+						SNew(SAddParametersDialog, FAssetData(EditorData));
 
-					Args.OnAddParameter = FOnAddParameter::CreateLambda([EditorData](const FParameterToAdd& ParameterToAdd)
+					TArray<FParameterToAdd> ParametersToAdd;
+					if(AddParametersDialog->ShowModal(ParametersToAdd))
 					{
-						FSlateApplication::Get().DismissAllMenus();
-
-						check(EditorData->FindEntry(ParameterToAdd.Name) == nullptr);
-						FScopedTransaction Transaction(LOCTEXT("AddParameter", "Add parameter"));
-						EditorData->AddParameter(ParameterToAdd.Name, ParameterToAdd.Type);
-					});
-					Args.OnParameterPicked = FOnParameterPicked::CreateLambda([EditorData](const FParameterBindingReference& InParameterBinding)
-					{
-						FSlateApplication::Get().DismissAllMenus();
-
-						if (EditorData->FindEntry(InParameterBinding.Parameter) == nullptr)
+						if(ParametersToAdd.Num() > 0)
 						{
-							const FAnimNextParamType Type = UncookedOnly::FUtils::GetParameterTypeFromName(InParameterBinding.Parameter);
-							if (Type.IsValid())
+							FScopedTransaction Transaction(LOCTEXT("AddParameter", "Add parameter"));
+							for (const FParameterToAdd& ParameterToAdd : ParametersToAdd)
 							{
-								EditorData->AddParameter(InParameterBinding.Parameter, Type);
-							};
+								check(EditorData->FindEntry(ParameterToAdd.Name) == nullptr);
+								EditorData->AddParameter(ParameterToAdd.Name, ParameterToAdd.Type);
+							}
 						}
-					});
-					
-					return SNew(SParameterPicker)
-						.Args(Args);
+					}
+					return FReply::Handled();
 				});
 		});
 
@@ -188,6 +175,13 @@ class FModule : public IModule
 					return FReply::Handled();
 				});
 		});
+
+		UE::UniversalObjectLocator::IUniversalObjectLocatorEditorModule& UolEditorModule = FModuleManager::LoadModuleChecked<UE::UniversalObjectLocator::IUniversalObjectLocatorEditorModule>("UniversalObjectLocatorEditor");
+		UolEditorModule.RegisterLocatorEditor("AnimNextObjectFunction", MakeShared<FObjectFunctionLocatorEditor>());
+		UolEditorModule.RegisterLocatorEditor("AnimNextObjectProperty", MakeShared<FObjectPropertyLocatorEditor>());
+		UolEditorModule.RegisterLocatorEditor("AnimNextObjectCast", MakeShared<FObjectCastLocatorEditor>());
+
+		UolEditorModule.RegisterEditorContext("AnimNextContext", MakeShared<FAnimNextLocatorContext>());
 	}
 
 	virtual void ShutdownModule() override
@@ -197,6 +191,7 @@ class FModule : public IModule
 			FPropertyEditorModule& PropertyModule = FModuleManager::LoadModuleChecked<FPropertyEditorModule>("PropertyEditor");
 			PropertyModule.UnregisterCustomPropertyTypeLayout("AnimNextParamType");
 			PropertyModule.UnregisterCustomPropertyTypeLayout("AnimNextParam");
+			PropertyModule.UnregisterCustomPropertyTypeLayout("AnimNextEditorParam");
 			PropertyModule.UnregisterCustomPropertyTypeLayout("NameProperty");
 			PropertyModule.UnregisterCustomClassLayout("AnimNextGraph_Parameter");
 			PropertyModule.UnregisterCustomClassLayout("AnimNextGraph_EdGraphNode");
@@ -209,7 +204,18 @@ class FModule : public IModule
 		UnregisterWorkspaceDocumentTypes();
 
 		SRigVMAssetView::UnregisterCategoryFactory("Parameters");
-		SRigVMAssetView::UnregisterCategoryFactory("Parameter Graphs");
+		SRigVMAssetView::UnregisterCategoryFactory("Event Graphs");
+		SRigVMAssetView::UnregisterCategoryFactory("Animation Graphs");
+
+		if(FModuleManager::Get().IsModuleLoaded("UniversalObjectLocatorEditor"))
+		{
+			UE::UniversalObjectLocator::IUniversalObjectLocatorEditorModule& UolEditorModule = FModuleManager::GetModuleChecked<UE::UniversalObjectLocator::IUniversalObjectLocatorEditorModule>("UniversalObjectLocatorEditor");
+			UolEditorModule.UnregisterLocatorEditor("AnimNextObjectCast");
+			UolEditorModule.UnregisterLocatorEditor("AnimNextObjectFunction");
+			UolEditorModule.UnregisterLocatorEditor("AnimNextObjectProperty");
+
+			UolEditorModule.UnregisterEditorContext("AnimNextContext");
+		}
 	}
 
 	virtual TSharedRef<SWidget> CreateParameterPicker(const FParameterPickerArgs& InArgs) override
diff --git a/Engine/Plugins/Experimental/AnimNext/Source/AnimNextEditor/Private/Common/SRigVMAssetView.cpp b/Engine/Plugins/Experimental/AnimNext/Source/AnimNextEditor/Private/Common/SRigVMAssetView.cpp
--- a/Engine/Plugins/Experimental/AnimNext/Source/AnimNextEditor/Private/Common/SRigVMAssetView.cpp
+++ b/Engine/Plugins/Experimental/AnimNext/Source/AnimNextEditor/Private/Common/SRigVMAssetView.cpp
@@ -802,7 +802,7 @@ class SRigVMAssetViewRow : public SMultiColumnTableRow<TSharedRef<FRigVMAssetVie
 					.OnVerifyTextChanged_Lambda([this](const FText& InNewText, FText& OutErrorText)
 					{
 						const FString NewString = InNewText.ToString();
-						if (!FUtils::IsValidEntryNameString(NewString, OutErrorText))
+						if (!FUtils::IsValidParameterNameString(NewString, OutErrorText))
 						{
 							return false;
 						}
@@ -822,14 +822,7 @@ class SRigVMAssetViewRow : public SMultiColumnTableRow<TSharedRef<FRigVMAssetVie
 					{
 						if(UAnimNextRigVMAssetEntry* AssetEntry = Cast<UAnimNextRigVMAssetEntry>(Entry->WeakEntry.Get()))
 						{
-							if(IAnimNextRigVMParameterInterface* ParameterInterface = Cast<IAnimNextRigVMParameterInterface>(AssetEntry))
-							{
-								return UncookedOnly::FUtils::GetParameterDisplayNameText(AssetEntry->GetEntryName());
-							}
-							else
-							{
-								return FText::FromName(AssetEntry->GetEntryName());
-							}
+							return FText::FromName(AssetEntry->GetEntryName());
 						}
 						return FText::GetEmpty();
 					})
@@ -854,9 +847,7 @@ class SRigVMAssetViewRow : public SMultiColumnTableRow<TSharedRef<FRigVMAssetVie
 					if (const IAnimNextRigVMParameterInterface* ParameterInterface = Cast<IAnimNextRigVMParameterInterface>(AssetEntry))
 					{
 						FInstancedPropertyBag& PropertyBag = ParameterInterface->GetPropertyBag();
-						const FName ParameterName = AssetEntry->GetEntryName();
-
-
+						const FName ParameterName = ParameterInterface->GetParamName();
 						if (const FPropertyBagPropertyDesc* PropertyDesc = PropertyBag.FindPropertyDescByName(ParameterName))
 						{
 							if (PropertyDesc->ContainerTypes.IsEmpty()) // avoid trying to inline containers
diff --git a/Engine/Plugins/Experimental/AnimNext/Source/AnimNextEditor/Private/EditorUtils.cpp b/Engine/Plugins/Experimental/AnimNext/Source/AnimNextEditor/Private/EditorUtils.cpp
--- a/Engine/Plugins/Experimental/AnimNext/Source/AnimNextEditor/Private/EditorUtils.cpp
+++ b/Engine/Plugins/Experimental/AnimNext/Source/AnimNextEditor/Private/EditorUtils.cpp
@@ -10,6 +10,7 @@
 #include "PropertyBagDetails.h"
 #include "UncookedOnlyUtils.h"
 #include "Param/RigVMDispatch_GetParameter.h"
+#include "String/ParseTokens.h"
 
 #define LOCTEXT_NAMESPACE "AnimNextEditorUtils"
 
@@ -310,54 +311,51 @@ void FUtils::GetFilteredVariableTypeTree(TArray<TSharedPtr<UEdGraphSchema_K2::FP
 };
 
 
-FName FUtils::GetNewParameterName(const TCHAR* InBaseName, TArrayView<FName> InAdditionalExistingNames)
+FName FUtils::GetNewParameterName(FName InBaseName, const FAssetData& InAssetData, TArrayView<FName> InExistingNames)
 {
-	FAnimNextParameterProviderAssetRegistryExports Exports;
-	UE::AnimNext::UncookedOnly::FUtils::GetExportedParametersFromAssetRegistry(Exports);
-
-	auto NameExists = [&Exports, &InAdditionalExistingNames](const TCHAR* InName)
+	auto NameExists = [&InExistingNames, &InAssetData](FName InName)
 	{
-		for(const FAnimNextParameterAssetRegistryExportEntry& Parameter : Exports.Parameters)
+		for(FName AdditionalName : InExistingNames)
 		{
-			if(Parameter.Name.ToString() == InName)
+			if(AdditionalName == InName)
 			{
 				return true;
 			}
 		}
 
-		for(FName AdditionalName : InAdditionalExistingNames)
+		if(DoesParameterNameExistInAsset(InName, InAssetData))
 		{
-			if(AdditionalName.ToString() == InName)
-			{
-				return true;
-			}
+			return true;
 		}
-
+		
 		return false;
 	};
 
 	if(!NameExists(InBaseName))
 	{
 		// Early out - name is valid
-		return FName(InBaseName);
+		return InBaseName;
 	}
-	
+
 	int32 PostFixIndex = 0;
+	TStringBuilder<128> StringBuilder;
 	while(true)
 	{
-		TStringBuilder<128> StringBuilder;
-		StringBuilder.Appendf(TEXT("%s_%d"), InBaseName, PostFixIndex++);
+		StringBuilder.Reset();
+		InBaseName.GetDisplayNameEntry()->AppendNameToString(StringBuilder);
+		StringBuilder.Appendf(TEXT("_%d"), PostFixIndex++);
 
-		if(!NameExists(StringBuilder.ToString()))
+		FName TestName(StringBuilder.ToString()); 
+		if(!NameExists(TestName))
 		{
-			return FName(StringBuilder.ToString());
+			return TestName;
 		}
 	}
 
 	return NAME_None;
 }
 
-bool FUtils::IsValidEntryNameString(FStringView InStringView, FText& OutErrorText)
+bool FUtils::IsValidParameterNameString(FStringView InStringView, FText& OutErrorText)
 {
 	// See if this can be represented as an FName
 	if(!FName::IsValidXName(InStringView, INVALID_NAME_CHARACTERS, &OutErrorText))
@@ -365,10 +363,10 @@ bool FUtils::IsValidEntryNameString(FStringView InStringView, FText& OutErrorTex
 		return false;
 	}
 
-	return IsValidEntryName(FName(InStringView), OutErrorText);
+	return IsValidParameterName(FName(InStringView), OutErrorText);
 }
 
-bool FUtils::IsValidEntryName(const FName InName, FText& OutErrorText)
+bool FUtils::IsValidParameterName(const FName InName, FText& OutErrorText)
 {
 	const FString NewString = InName.ToString();
 
@@ -379,11 +377,9 @@ bool FUtils::IsValidEntryName(const FName InName, FText& OutErrorText)
 	}
 
 	// Check start
-	if (NewString[0] == TEXT('.') ||
-		FChar::IsUnderscore(NewString[0]) ||
-		FChar::IsDigit(NewString[0]))
+	if (FChar::IsDigit(NewString[0]))
 	{
-		OutErrorText = LOCTEXT("Error_Start", "Name cannot start with an underscore, period or digit");
+		OutErrorText = LOCTEXT("Error_Start", "Name cannot start with a digit");
 		return false;
 	}
 
@@ -391,27 +387,19 @@ bool FUtils::IsValidEntryName(const FName InName, FText& OutErrorText)
 	for (int32 CharIndex = 0; bAllowed && CharIndex < NewString.Len(); ++CharIndex)
 	{
 		bAllowed &= FChar::IsAlnum(NewString[CharIndex]) ||
-					FChar::IsUnderscore(NewString[CharIndex]) ||
-					NewString[CharIndex] == TEXT('.');
+					FChar::IsUnderscore(NewString[CharIndex]);
 	}
 
 	// Make sure the new name only contains valid characters
 	if (!bAllowed)
 	{
-		OutErrorText = LOCTEXT("Error_CharacterNotAllowed", "Only alpha-numerical, underscore or period characters are allowed");
+		OutErrorText = LOCTEXT("Error_CharacterNotAllowed", "Only alpha-numerical or underscore characters are allowed");
 		return false;
 	}
 
 	return true;
 }
 
-bool FUtils::DoesParameterNameExist(const FName InName)
-{
-	FAnimNextParameterProviderAssetRegistryExports Exports;
-	UncookedOnly::FUtils::GetExportedParametersFromAssetRegistry(Exports);
-	return Exports.Parameters.ContainsByPredicate([InName](const FAnimNextParameterAssetRegistryExportEntry& Entry) { return Entry.Name == InName; });
-}
-
 bool FUtils::DoesParameterNameExistInAsset(const FName InName, const FAssetData& InAsset)
 {
 	FAnimNextParameterProviderAssetRegistryExports Exports;
diff --git a/Engine/Plugins/Experimental/AnimNext/Source/AnimNextEditor/Private/EditorUtils.h b/Engine/Plugins/Experimental/AnimNext/Source/AnimNextEditor/Private/EditorUtils.h
--- a/Engine/Plugins/Experimental/AnimNext/Source/AnimNextEditor/Private/EditorUtils.h
+++ b/Engine/Plugins/Experimental/AnimNext/Source/AnimNextEditor/Private/EditorUtils.h
@@ -27,14 +27,12 @@ struct FUtils
 
 	static void GetFilteredVariableTypeTree(TArray<TSharedPtr<UEdGraphSchema_K2::FPinTypeTreeInfo>>& TypeTree, ETypeTreeFilter TypeTreeFilter);
 
-	static FName GetNewParameterName(const TCHAR* InBaseName, TArrayView<FName> InAdditionalExistingNames);
+	static FName GetNewParameterName(FName InBaseName, const FAssetData& InAssetData, TArrayView<FName> InExistingNames);
 
-	static bool IsValidEntryNameString(FStringView InStringView, FText& OutErrorText);
+	static bool IsValidParameterNameString(FStringView InStringView, FText& OutErrorText);
 
-	static bool IsValidEntryName(const FName InName, FText& OutErrorText);
+	static bool IsValidParameterName(const FName InName, FText& OutErrorText);
 
-	static bool DoesParameterNameExist(const FName InName);
-	
 	static bool DoesParameterNameExistInAsset(const FName InName, const FAssetData& InAsset);
 };
 
diff --git a/Engine/Plugins/Experimental/AnimNext/Source/AnimNextEditor/Private/Graph/AnimNextGraphPanelNodeFactory.cpp b/Engine/Plugins/Experimental/AnimNext/Source/AnimNextEditor/Private/Graph/AnimNextGraphPanelNodeFactory.cpp
--- a/Engine/Plugins/Experimental/AnimNext/Source/AnimNextEditor/Private/Graph/AnimNextGraphPanelNodeFactory.cpp
+++ b/Engine/Plugins/Experimental/AnimNext/Source/AnimNextEditor/Private/Graph/AnimNextGraphPanelNodeFactory.cpp
@@ -8,13 +8,16 @@ TSharedPtr<SGraphNode> FAnimNextGraphPanelNodeFactory::CreateNode(UEdGraphNode*
 {
 	if (UAnimNextGraph_EdGraphNode* AnimNextGraphNode = Cast<UAnimNextGraph_EdGraphNode>(Node))
 	{
-		TSharedPtr<SGraphNode> GraphNode =
-			SNew(SAnimNextGraphNode)
-			.GraphNodeObj(AnimNextGraphNode);
+		if(AnimNextGraphNode->IsTraitStack())
+		{
+			TSharedPtr<SGraphNode> GraphNode =
+				SNew(SAnimNextGraphNode)
+				.GraphNodeObj(AnimNextGraphNode);
 
-		GraphNode->SlatePrepass();
-		AnimNextGraphNode->SetDimensions(GraphNode->GetDesiredSize());
-		return GraphNode;
+			GraphNode->SlatePrepass();
+			AnimNextGraphNode->SetDimensions(GraphNode->GetDesiredSize());
+			return GraphNode;
+		}
 	}
 
 	return nullptr;
