AnimNext: Parameters rework - UOL integration

diff --git a/Engine/Plugins/Experimental/AnimNext/Source/AnimNextEditor/Private/Graph/SAnimNextGraphNode.cpp b/Engine/Plugins/Experimental/AnimNext/Source/AnimNextEditor/Private/Graph/SAnimNextGraphNode.cpp
--- a/Engine/Plugins/Experimental/AnimNext/Source/AnimNextEditor/Private/Graph/SAnimNextGraphNode.cpp
+++ b/Engine/Plugins/Experimental/AnimNext/Source/AnimNextEditor/Private/Graph/SAnimNextGraphNode.cpp
@@ -2,12 +2,6 @@
 
 #include "Graph/SAnimNextGraphNode.h"
 #include "Graph/AnimNextGraph_EdGraphNode.h"
-#include "SGraphPin.h"
-#include "RigVMModel/Nodes/RigVMDispatchNode.h"
-#include "TraitCore/TraitHandle.h"
-#include "Widgets/SBoxPanel.h"
-#include "Widgets/Input/SButton.h"
-#include "Widgets/Layout/SSpacer.h"
 
 #define LOCTEXT_NAMESPACE "SAnimNextGraphNode"
 
@@ -16,644 +10,4 @@ void SAnimNextGraphNode::Construct( const FArguments& InArgs )
 	SRigVMGraphNode::Construct(SRigVMGraphNode::FArguments().GraphNodeObj(InArgs._GraphNodeObj));
 }
 
-void SAnimNextGraphNode::UpdatePinTreeView()
-{
-	static const float PinWidgetSidePadding = 6.f;
-	static const float EmptySidePadding = 60.f;
-	static const float TopPadding = 2.f; 
-	static const float MaxHeight = 30.f;
-
-	check(GraphNode);
-
-	// remove all existing content in the Left
-	LeftNodeBox->ClearChildren();
-
-	URigVMEdGraphNode* RigGraphNode = Cast<URigVMEdGraphNode>(GraphNode);
-	URigVMNode* Node = RigGraphNode->GetModelNode();
-	const FRigVMTemplate* Template = nullptr;
-	if (URigVMTemplateNode* TemplateNode = Cast<URigVMTemplateNode>(Node))
-	{
-		Template = TemplateNode->GetTemplate();
-	}
-
-	TMap<UEdGraphPin*, int32> EdGraphPinToInputPin;
-	for(int32 InputPinIndex = 0; InputPinIndex < InputPins.Num(); InputPinIndex++)
-	{
-		EdGraphPinToInputPin.Add(InputPins[InputPinIndex]->GetPinObj(), InputPinIndex);
-	}
-	TMap<UEdGraphPin*, int32> EdGraphPinToOutputPin;
-	for(int32 OutputPinIndex = 0; OutputPinIndex < OutputPins.Num(); OutputPinIndex++)
-	{
-		EdGraphPinToOutputPin.Add(OutputPins[OutputPinIndex]->GetPinObj(), OutputPinIndex);
-	}
-
-	TArray<URigVMPin*> RootModelPins = ModelNode->GetPins();
-	// orphaned pins are appended to the end of pin list on each side of the node
-	RootModelPins.Append(ModelNode->GetOrphanedPins());
-	TArray<URigVMPin*> ModelPins;
-
-	const bool bSupportSubPins = !RigGraphNode->DrawAsCompactNode();
-
-	// sort model pins
-	// a) execute IOs, b) IO pins, c) input / visible pins, d) output pins
-	struct Local
-	{
-		static void VisitPinRecursively(URigVMPin* InPin, TArray<URigVMPin*>& OutPins, bool bSupportSubPins)
-		{
-			OutPins.Add(InPin);
-			if(!bSupportSubPins)
-			{
-				return;
-			}
-
-			if(InPin->ShouldHideSubPins())
-			{
-				return;
-			}
-
-			if (InPin->GetCPPType() == TEXT("FRotator"))
-			{
-				const TArray<URigVMPin*>& SubPins = InPin->GetSubPins();
-				if (SubPins.Num() == 3)
-				{
-					OutPins.Add(SubPins[2]);	
-					OutPins.Add(SubPins[0]);	
-					OutPins.Add(SubPins[1]);
-				}	
-			}
-			else
-			{				
-				for (URigVMPin* SubPin : InPin->GetSubPins())
-				{
-					VisitPinRecursively(SubPin, OutPins, bSupportSubPins);
-				}
-			}
-		}
-	};
-	
-	for(int32 SortPhase = 0; SortPhase < 4; SortPhase++)
-	{
-		for(URigVMPin* RootPin : RootModelPins)
-		{
-			switch (SortPhase)
-			{
-				case 0: // execute IO pins
-				{
-					if(RootPin->IsExecuteContext() && RootPin->GetDirection() == ERigVMPinDirection::IO)
-					{
-						Local::VisitPinRecursively(RootPin, ModelPins, bSupportSubPins);
-					}
-					break;
-				}
-				case 1: // IO pins
-				{
-					if(!RootPin->IsExecuteContext() && RootPin->GetDirection() == ERigVMPinDirection::IO)
-					{
-						Local::VisitPinRecursively(RootPin, ModelPins, bSupportSubPins);
-					}
-					break;
-				}
-				case 2: // input / visible pins
-				{
-					if(RootPin->GetDirection() == ERigVMPinDirection::Input || RootPin->GetDirection() == ERigVMPinDirection::Visible)
-					{
-						Local::VisitPinRecursively(RootPin, ModelPins, bSupportSubPins);
-					}
-					break;
-				}
-				case 3: // output pins
-				default:
-				{
-					if(RootPin->GetDirection() == ERigVMPinDirection::Output)
-					{
-						Local::VisitPinRecursively(RootPin, ModelPins, bSupportSubPins);
-					}
-					break;
-				}
-			}
-		}
-	}
-
-	const URigVMEdGraphSchema* RigSchema = Cast<URigVMEdGraphSchema>(RigGraphNode->GetSchema());
-
-	FRigVMDispatchContext DispatchContext;
-	if(const URigVMDispatchNode* DispatchNode = Cast<URigVMDispatchNode>(ModelNode))
-	{
-		DispatchContext = DispatchNode->GetDispatchContext();
-	}
-	
-	TMap<URigVMPin*, int32> ModelPinToInfoIndex;
-	for(URigVMPin* ModelPin : ModelPins)
-	{
-		FPinInfo PinInfo;
-		PinInfo.Index = PinInfos.Num();
-		PinInfo.ParentIndex = INDEX_NONE;
-		PinInfo.bHasChildren = (ModelPin->GetSubPins().Num() > 0);
-		PinInfo.bIsContainer = ModelPin->IsArray();
-		PinInfo.Depth = 0;
-		PinInfo.bExpanded = ModelPin->IsExpanded();
-		PinInfo.ModelPinPath = ModelPin->GetPinPath();
-		PinInfo.bAutoHeight = false;
-		PinInfo.bShowOnlySubPins = ModelPin->ShouldOnlyShowSubPins();
-
-		if(!bSupportSubPins)
-		{
-			PinInfo.bHasChildren = false;
-			PinInfo.bIsContainer = false;
-		}
-
-		// Trait handle pins have special packing that are not user editable, hide their innards
-		if (ModelPin->GetCPPTypeObject() == FAnimNextTraitHandle::StaticStruct())
-		{
-			PinInfo.bHasChildren = false;
-			PinInfo.bIsContainer = false;
-		}
-		
-		const bool bAskSchemaForEdition = RigSchema && ModelPin->IsStruct() && !ModelPin->IsBoundToVariable();
-		PinInfo.bHideInputWidget = (!ModelPin->IsBoundToVariable()) && PinInfo.bIsContainer;
-		if (!PinInfo.bHideInputWidget)
-		{
-			if (bAskSchemaForEdition && !PinInfo.bHasChildren)
-			{
-				const bool bIsStructEditable = RigSchema->IsStructEditable(ModelPin->GetScriptStruct()); 
-				PinInfo.bHideInputWidget = !bIsStructEditable;
-				PinInfo.bAutoHeight = bIsStructEditable;
-			}
-			else if(PinInfo.bHasChildren && !ModelPin->IsBoundToVariable())
-			{
-				PinInfo.bHideInputWidget = true;
-			}
-		}
-		
-		if(URigVMPin* ParentPin = ModelPin->GetParentPin())
-		{
-			const int32* ParentIndexPtr = ModelPinToInfoIndex.Find(ParentPin);
-			if(ParentIndexPtr == nullptr)
-			{
-				continue;
-			}
-			PinInfo.ParentIndex = *ParentIndexPtr;
-			PinInfo.Depth = PinInfos[PinInfo.ParentIndex].Depth + 1;
-			if(PinInfos[PinInfo.ParentIndex].bShowOnlySubPins)
-			{
-				PinInfo.Depth--;
-			}
-		}
-
-		TAttribute<EVisibility> PinVisibilityAttribute = TAttribute<EVisibility>::CreateSP(this, &SAnimNextGraphNode::GetPinVisibility, PinInfo.Index, false);
-
-		bool bPinWidgetForExpanderLeft = false;
-		TSharedPtr<SGraphPin> PinWidgetForExpander;
-
-		bool bPinInfoIsValid = false;
-		if(UEdGraphPin* OutputEdGraphPin = RigGraphNode->FindPin(ModelPin->GetPinPath(), EEdGraphPinDirection::EGPD_Output))
-		{
-			if(const int32* PinIndexPtr = EdGraphPinToOutputPin.Find(OutputEdGraphPin))
-			{
-				PinInfo.OutputPinWidget = OutputPins[*PinIndexPtr];
-				PinInfo.OutputPinWidget->SetVisibility(PinVisibilityAttribute);
-				PinWidgetForExpander = PinInfo.OutputPinWidget;
-				bPinWidgetForExpanderLeft = false;
-				bPinInfoIsValid = true;
-				if (Template && !ModelPin->IsExecuteContext())
-				{
-					if (URigVMPin* RootPin = ModelPin->GetRootPin())
-					{
-						FLinearColor PinColorAndOpacity = PinInfo.OutputPinWidget->GetColorAndOpacity();
-						if (Template->FindArgument(RootPin->GetFName()) == nullptr && Template->FindExecuteArgument(RootPin->GetFName(), DispatchContext) == nullptr)
-						{
-							PinColorAndOpacity.A = 0.2f;
-						}
-						else
-						{
-							PinColorAndOpacity.A = 1.0f;
-						}
-						PinInfo.OutputPinWidget->SetColorAndOpacity(PinColorAndOpacity);
-					}
-				}
-			}
-		}
-		
-		if(UEdGraphPin* InputEdGraphPin = RigGraphNode->FindPin(ModelPin->GetPinPath(), EEdGraphPinDirection::EGPD_Input))
-		{
-			// Make sure to handle input pins that we explicitly hide, like trait handles
-			if (!InputEdGraphPin->bHidden)
-			{
-				if (const int32* PinIndexPtr = EdGraphPinToInputPin.Find(InputEdGraphPin))
-				{
-					PinInfo.InputPinWidget = InputPins[*PinIndexPtr];
-					PinInfo.InputPinWidget->SetVisibility(PinVisibilityAttribute);
-					PinWidgetForExpander = PinInfo.InputPinWidget;
-					bPinWidgetForExpanderLeft = true;
-					bPinInfoIsValid = true;
-					if (Template && !ModelPin->IsExecuteContext())
-					{
-						if (URigVMPin* RootPin = ModelPin->GetRootPin())
-						{
-							FLinearColor PinColorAndOpacity = PinInfo.InputPinWidget->GetColorAndOpacity();
-							if (Template->FindArgument(RootPin->GetFName()) == nullptr && Template->FindExecuteArgument(RootPin->GetFName(), DispatchContext) == nullptr)
-							{
-								PinColorAndOpacity.A = 0.2f;
-							}
-							else
-							{
-								PinColorAndOpacity.A = 1.0f;
-							}
-							PinInfo.InputPinWidget->SetColorAndOpacity(PinColorAndOpacity);
-						}
-					}
-				}
-			}
-		}
-
-		if(!bPinInfoIsValid)
-		{
-			continue;
-		}
-
-		ModelPinToInfoIndex.Add(ModelPin, PinInfos.Add(PinInfo));
-		
-		// check if this pin has sub pins
-		TSharedPtr<SHorizontalBox> FullPinHorizontalRowWidget = PinWidgetForExpander->GetFullPinHorizontalRowWidget().Pin();
-		if(FullPinHorizontalRowWidget.IsValid())
-		{
-			// indent the pin by padding
-			const float DepthIndentation = 12.f * float(PinInfo.Depth + (PinInfo.bHasChildren ? 0 : 1));
-			const float LeftIndentation = bPinWidgetForExpanderLeft ? DepthIndentation : 0.f;
-			const float RightIndentation = bPinWidgetForExpanderLeft ? 0.f : DepthIndentation;
-			const FMargin LineIndentation = FMargin(RightIndentation, 0, LeftIndentation, 0); 
-
-			static const TSharedRef<FTagMetaData> ExpanderButtonMetadata = MakeShared<FTagMetaData>(TEXT("SRigVMGraphNode.ExpanderButton"));
-
-			// check if this pin widget may already have the expander button 
-			int32 ExpanderSlotIndex = INDEX_NONE;
-			int32 PaddedSlotIndex = INDEX_NONE;
-
-			for(int32 SlotIndex=0; SlotIndex<FullPinHorizontalRowWidget->NumSlots(); SlotIndex++)
-			{
-				const TSharedRef<SWidget> Widget = FullPinHorizontalRowWidget->GetSlot(SlotIndex).GetWidget();
-				const TSharedPtr<FTagMetaData> Metadata = Widget->GetMetaData<FTagMetaData>();
-				if(Metadata.IsValid())
-				{
-					if(Metadata->Tag == ExpanderButtonMetadata->Tag)
-					{
-						ExpanderSlotIndex = SlotIndex;
-						break;
-					}
-				}
-			}
-
-			// The expander needs to be recreated to adjust for PinInfo.Index changes
-			if (FullPinHorizontalRowWidget->IsValidSlotIndex(ExpanderSlotIndex))
-			{
-				SHorizontalBox::FSlot& Slot = FullPinHorizontalRowWidget->GetSlot(ExpanderSlotIndex);
-				const TSharedRef<SWidget> Widget = Slot.GetWidget();
-				Widget->RemoveMetaData(ExpanderButtonMetadata);
-				FullPinHorizontalRowWidget->RemoveSlot(Widget);
-				ExpanderSlotIndex = INDEX_NONE;
-			}
-	
-			if(PinInfo.bHasChildren && ExpanderSlotIndex == INDEX_NONE)
-			{
-				// only inject the expander arrow for inputs on input / IO
-				// or for output pins
-				if(
-					(
-						(
-							(ModelPin->GetDirection() == ERigVMPinDirection::Input) ||
-							(ModelPin->GetDirection() == ERigVMPinDirection::IO)
-						) &&
-						bPinWidgetForExpanderLeft
-					) ||
-					(
-						(ModelPin->GetDirection() == ERigVMPinDirection::Output) &&
-						(!bPinWidgetForExpanderLeft)
-					)
-				)
-				{
-					// Add the expander arrow
-					FullPinHorizontalRowWidget->InsertSlot(bPinWidgetForExpanderLeft ? 1 : FullPinHorizontalRowWidget->GetChildren()->Num() - 1)
-					.Padding(FMargin(0, 0, 0, 0))
-					.AutoWidth()
-					[
-						SNew(SButton)
-						.ButtonStyle(FAppStyle::Get(), TEXT("SimpleButton"))
-						.ContentPadding(0)
-						.VAlign(VAlign_Center)
-						.HAlign(HAlign_Center)
-						.ClickMethod( EButtonClickMethod::MouseDown )
-						.OnClicked(this, &SAnimNextGraphNode::OnExpanderArrowClicked, PinInfo.Index)
-						.ToolTipText(LOCTEXT("ExpandSubPin", "Expand Pin"))
-						[
-							SNew(SImage)
-							.Image(this, &SAnimNextGraphNode::GetExpanderImage, PinInfo.Index, bPinWidgetForExpanderLeft, false)
-							.ColorAndOpacity(FSlateColor::UseForeground())
-						]
-						.AddMetaData(ExpanderButtonMetadata)
-					];
-				}
-			}
-
-			// adjust the padding
-			{
-				const int32 SlotToAdjustIndex = bPinWidgetForExpanderLeft ? 0 : FullPinHorizontalRowWidget->NumSlots() - 1;
-				SHorizontalBox::FSlot& Slot = FullPinHorizontalRowWidget->GetSlot(SlotToAdjustIndex);
-				Slot.SetPadding(LineIndentation);
-			}
-		}
-	}
-
-	// add spacer widget at the start
-	LeftNodeBox->AddSlot()
-	.HAlign(HAlign_Fill)
-	.VAlign(VAlign_Center)
-	.AutoHeight()
-	[
-		SNew(SSpacer)
-		.Size(FVector2D(1.f, 2.f))
-	];
-
-	auto AddArrayPlusButtonLambda = [this](URigVMPin* InModelPin, TSharedPtr<SHorizontalBox> InSlotLayout, const float InEmptySidePadding)
-	{
-		// add array plus button
-		InSlotLayout->AddSlot()
-		.AutoWidth()
-		.HAlign(HAlign_Left)
-		.Padding(PinWidgetSidePadding, TopPadding, InEmptySidePadding, 0.f)
-		[
-			SNew(SButton)
-			.ContentPadding(0.0f)
-			.ButtonStyle(FAppStyle::Get(), "NoBorder")
-			.OnClicked(this, &SAnimNextGraphNode::HandleAddArrayElement, InModelPin->GetPinPath())
-			.IsEnabled(this, &SGraphNode::IsNodeEditable)
-			.Cursor(EMouseCursor::Default)
-			.Visibility(this, &SAnimNextGraphNode::GetArrayPlusButtonVisibility, InModelPin)
-			.ToolTipText(LOCTEXT("AddArrayElement", "Add Array Element"))
-			[
-				SNew(SHorizontalBox)
-				+SHorizontalBox::Slot()
-				.AutoWidth()
-				.VAlign(VAlign_Center)
-				[
-					SNew(SImage)
-					.Image(FAppStyle::GetBrush(TEXT("Icons.PlusCircle")))
-				]
-			]
-		];
-	};
-
-	struct FPinRowInfo
-	{
-		int32 InputPinInfo;
-		int32 OutputPinInfo;
-		
-		FPinRowInfo()
-			: InputPinInfo(INDEX_NONE)
-			, OutputPinInfo(INDEX_NONE)
-		{}
-	};
-
-	TArray<FPinRowInfo> Rows;
-	Rows.Reserve(PinInfos.Num());
-
-	// build a tree with the pin infos in order as before
-	for(const FPinInfo& PinInfo : PinInfos)
-	{
-		FPinRowInfo Row;
-		if(PinInfo.InputPinWidget.IsValid())
-		{
-			Row.InputPinInfo = PinInfo.Index;
-		}
-		if(PinInfo.OutputPinWidget.IsValid())
-		{
-			Row.OutputPinInfo = PinInfo.Index;
-		}
-		Rows.Add(Row);
-	}
-
-	// compact the rows - re-use input rows for outputs
-	// if their sub-pins count matches
-	if(RigGraphNode->DrawAsCompactNode())
-	{
-		for(int32 RowIndex = 0; RowIndex < Rows.Num();)
-		{
-			if(Rows[RowIndex].OutputPinInfo == INDEX_NONE)
-			{
-				RowIndex++;
-				continue;
-			}
-
-			bool bFoundMatch = false;
-			for(int32 InputRowIndex = 0; InputRowIndex < RowIndex; InputRowIndex++)
-			{
-				if(Rows[InputRowIndex].OutputPinInfo == INDEX_NONE)
-				{
-					const FPinInfo& OutputPinInfo = PinInfos[Rows[RowIndex].OutputPinInfo];
-					const FPinInfo& InputPinInfo = PinInfos[Rows[InputRowIndex].InputPinInfo];
-					if(OutputPinInfo.bHasChildren == false &&
-						InputPinInfo.bHasChildren == false)
-					{
-						Rows[InputRowIndex].OutputPinInfo = Rows[RowIndex].OutputPinInfo;
-						bFoundMatch = true;
-						break;
-					}
-				}
-			}
-
-			if(bFoundMatch)
-			{
-				Rows.RemoveAt(RowIndex);
-			}
-			else
-			{
-				RowIndex++;
-			}
-		}
-	}
-	
-	for(const FPinRowInfo& Row : Rows)
-	{
-		const FPinInfo& InputPinInfo = PinInfos[Row.InputPinInfo == INDEX_NONE ? Row.OutputPinInfo : Row.InputPinInfo];
-		const FPinInfo& OutputPinInfo = PinInfos[Row.OutputPinInfo == INDEX_NONE ? Row.InputPinInfo : Row.OutputPinInfo];
-
-		if(InputPinInfo.InputPinWidget.IsValid())
-		{
-			if(InputPinInfo.bHideInputWidget)
-			{
-				if(InputPinInfo.InputPinWidget->GetValueWidget() != SNullWidget::NullWidget)
-				{
-					InputPinInfo.InputPinWidget->GetValueWidget()->SetVisibility(EVisibility::Collapsed);
-				}
-			}
-				
-			// input pins
-			if(!OutputPinInfo.OutputPinWidget.IsValid())
-			{
-				TSharedPtr<SHorizontalBox> SlotLayout;
-				SHorizontalBox::FSlot* FirstSlot = nullptr;
-
-				const float MyEmptySidePadding = InputPinInfo.bHideInputWidget ? EmptySidePadding : 0.f; 
-				
-				LeftNodeBox->AddSlot()
-                .HAlign(HAlign_Fill)
-				.VAlign(VAlign_Center)
-				.AutoHeight()
-				.MaxHeight(InputPinInfo.bAutoHeight ? TAttribute<float>() : MaxHeight)
-                [
-                    SAssignNew(SlotLayout, SHorizontalBox)
-                    .Visibility(this, &SAnimNextGraphNode::GetPinVisibility, InputPinInfo.Index, false)
-		            
-                    +SHorizontalBox::Slot()
-                    .Expose(FirstSlot)
-                    .FillWidth(1.f)
-                    .HAlign(HAlign_Left)
-                    .Padding(PinWidgetSidePadding, TopPadding, InputPinInfo.bIsContainer ? 0.f : MyEmptySidePadding, 0.f)
-                    [
-                        InputPinInfo.InputPinWidget.ToSharedRef()
-                    ]
-                ];
-
-				if(InputPinInfo.bIsContainer)
-				{
-					URigVMPin* ModelPin = ModelNode->GetGraph()->FindPin(InputPinInfo.ModelPinPath);
-					if(ModelPin)
-					{
-						// make sure to minimize the width of the label
-						FirstSlot->SetAutoWidth();
-						AddArrayPlusButtonLambda(ModelPin, SlotLayout, MyEmptySidePadding);
-					}
-				}
-			}
-			// io pins
-			else
-			{
-				TSharedPtr<SHorizontalBox> SlotLayout;
-				SHorizontalBox::FSlot* FirstSlot = nullptr;
-
-				OutputPinInfo.OutputPinWidget->SetShowLabel(false);
-			
-				LeftNodeBox->AddSlot()
-                .HAlign(HAlign_Fill)
-				.VAlign(VAlign_Center)
-                .AutoHeight()
-				.MaxHeight(InputPinInfo.bAutoHeight ? TAttribute<float>() : MaxHeight)
-                [
-                    SAssignNew(SlotLayout, SHorizontalBox)
-                    .Visibility(this, &SAnimNextGraphNode::GetPinVisibility, OutputPinInfo.Index, false)
-
-                    +SHorizontalBox::Slot()
-					.Expose(FirstSlot)
-                    .FillWidth(1.f)
-                    .HAlign(HAlign_Left)
-                    .VAlign(VAlign_Center)
-                    .Padding(PinWidgetSidePadding, TopPadding, 0.f, 0.f)
-                    [
-                        InputPinInfo.InputPinWidget.ToSharedRef()
-                    ]
-                ];
-
-				if(InputPinInfo.bIsContainer)
-				{
-					URigVMPin* ModelPin = ModelNode->GetGraph()->FindPin(InputPinInfo.ModelPinPath);
-					if(ModelPin)
-					{
-						// make sure to minimize the width of the label
-						FirstSlot->SetAutoWidth();
-						AddArrayPlusButtonLambda(ModelPin, SlotLayout, EmptySidePadding);
-					}
-				}
-
-				SlotLayout->AddSlot()
-				.FillWidth(1.f)
-				.HAlign(HAlign_Right)
-				.VAlign(VAlign_Center)
-				.Padding(0.f, TopPadding, PinWidgetSidePadding, 0.f)
-				[
-					OutputPinInfo.OutputPinWidget.ToSharedRef()
-				];
-			}
-		}
-		// output pins
-		else if(OutputPinInfo.OutputPinWidget.IsValid())
-		{
-			LeftNodeBox->AddSlot()
-            .HAlign(HAlign_Fill)
-			.VAlign(VAlign_Center)
-            .AutoHeight()
-			.MaxHeight(OutputPinInfo.bAutoHeight ? TAttribute<float>() : MaxHeight)
-            [
-            	SNew(SHorizontalBox)
-	            .Visibility(this, &SAnimNextGraphNode::GetPinVisibility, OutputPinInfo.Index, false)
-	            
-				+SHorizontalBox::Slot()
-				.FillWidth(1.f)
-				.HAlign(HAlign_Right)
-				.VAlign(VAlign_Center)
-	            .Padding(EmptySidePadding, TopPadding, PinWidgetSidePadding, 0.f)
-				[
-	                OutputPinInfo.OutputPinWidget.ToSharedRef()
-				]
-            ];
-		}
-	}
-
-#if 0
-	if(URigVMFunctionReferenceNode* FunctionReferenceNode = Cast<URigVMFunctionReferenceNode>(ModelNode))
-	{
-		TWeakObjectPtr<URigVMFunctionReferenceNode> WeakFunctionReferenceNode = FunctionReferenceNode;
-		TWeakObjectPtr<URigVMBlueprint> WeakRigVMBlueprint = Blueprint;
-
-		// add the entries for the variable remapping
-		for(const TSharedPtr<FRigVMExternalVariable>& ExternalVariable : RigGraphNode->ExternalVariables)	// TODO: FIX THIS, PRIVATE?
-		{
-			LeftNodeBox->AddSlot()
-			.HAlign(HAlign_Fill)
-			.VAlign(VAlign_Center)
-			.AutoHeight()
-			.MaxHeight(MaxHeight)
-			[
-				SNew(SHorizontalBox)
-	            
-				+SHorizontalBox::Slot()
-				.AutoWidth()
-				.HAlign(HAlign_Left)
-				.VAlign(VAlign_Center)
-				.Padding(PinWidgetSidePadding, TopPadding, PinWidgetSidePadding, 0.f)
-				[
-					SNew(STextBlock)
-					.Text(FText::FromName(ExternalVariable->Name))
-					.TextStyle(FAppStyle::Get(), NAME_DefaultPinLabelStyle)
-					.ColorAndOpacity(this, &SAnimNextGraphNode::GetVariableLabelTextColor, WeakFunctionReferenceNode, ExternalVariable->Name)
-					.ToolTipText(this, &SAnimNextGraphNode::GetVariableLabelTooltipText, WeakControlRigBlueprint, ExternalVariable->Name)
-				]
-
-				+SHorizontalBox::Slot()
-				.AutoWidth()
-				.HAlign(HAlign_Left)
-				.VAlign(VAlign_Center)
-				.Padding(PinWidgetSidePadding, TopPadding, PinWidgetSidePadding, 0.f)
-				[
-					SNew(SRigVMGraphVariableBinding)
-					.Blueprint(Blueprint.Get())
-					.FunctionReferenceNode(FunctionReferenceNode)
-					.InnerVariableName(ExternalVariable->Name)
-				]
-			];
-		}
-	}
-#endif
-
-	CreateAddPinButton();
-	
-	// add spacer widget at the end
-	LeftNodeBox->AddSlot()
-	.HAlign(HAlign_Fill)
-	.VAlign(VAlign_Center)
-	.AutoHeight()
-	[
-		SNew(SSpacer)
-		.Size(FVector2D(1.f, 4.f))
-	];
-}
-
 #undef LOCTEXT_NAMESPACE
diff --git a/Engine/Plugins/Experimental/AnimNext/Source/AnimNextEditor/Private/Graph/SAnimNextGraphNode.h b/Engine/Plugins/Experimental/AnimNext/Source/AnimNextEditor/Private/Graph/SAnimNextGraphNode.h
--- a/Engine/Plugins/Experimental/AnimNext/Source/AnimNextEditor/Private/Graph/SAnimNextGraphNode.h
+++ b/Engine/Plugins/Experimental/AnimNext/Source/AnimNextEditor/Private/Graph/SAnimNextGraphNode.h
@@ -19,7 +19,4 @@ public:
 	SLATE_END_ARGS()
 
 	void Construct(const FArguments& InArgs);
-
-protected:
-	virtual void UpdatePinTreeView() override;
 };
diff --git a/Engine/Plugins/Experimental/AnimNext/Source/AnimNextEditor/Private/Param/AnimNextLocatorContext.h b/Engine/Plugins/Experimental/AnimNext/Source/AnimNextEditor/Private/Param/AnimNextLocatorContext.h
new file mode 100644
--- /dev/null
+++ b/Engine/Plugins/Experimental/AnimNext/Source/AnimNextEditor/Private/Param/AnimNextLocatorContext.h
@@ -0,0 +1,32 @@
+﻿// Copyright Epic Games, Inc. All Rights Reserved.
+
+#pragma once
+
+#include "UniversalObjectLocatorEditorContext.h"
+#include "Component/AnimNextComponent.h"
+
+namespace UE::AnimNext
+{
+
+class FAnimNextLocatorContext : public UE::UniversalObjectLocator::ILocatorFragmentEditorContext
+{
+	// ILocatorFragmentEditorContext interface
+	virtual UObject* GetContext(const IPropertyHandle& InPropertyHandle) const override
+	{
+		// TODO: This needs to defer to project/schedule/workspace defaults similar to SParameterPicker
+		return UAnimNextComponent::StaticClass()->GetDefaultObject();
+	}
+
+	virtual bool IsFragmentAllowed(FName InFragmentName) const override
+	{
+		return (InFragmentName == "Actor" ||
+				InFragmentName == "Asset" ||
+				InFragmentName == "AnimNextScope" ||
+				InFragmentName == "AnimNextGraph" ||
+				InFragmentName == "AnimNextObjectFunction" ||
+				InFragmentName == "AnimNextObjectProperty" ||
+				InFragmentName == "AnimNextObjectCast");
+	}
+};
+
+}
diff --git a/Engine/Plugins/Experimental/AnimNext/Source/AnimNextEditor/Private/Param/AnimNextParameterSettings.cpp b/Engine/Plugins/Experimental/AnimNext/Source/AnimNextEditor/Private/Param/AnimNextParameterSettings.cpp
--- a/Engine/Plugins/Experimental/AnimNext/Source/AnimNextEditor/Private/Param/AnimNextParameterSettings.cpp
+++ b/Engine/Plugins/Experimental/AnimNext/Source/AnimNextEditor/Private/Param/AnimNextParameterSettings.cpp
@@ -20,4 +20,14 @@ const FAnimNextParamType& UAnimNextParameterSettings::GetLastParameterType() con
 void UAnimNextParameterSettings::SetLastParameterType(const FAnimNextParamType& InLastParameterType)
 {
 	LastParameterType = InLastParameterType;
+}
+
+FName UAnimNextParameterSettings::GetLastParameterName() const
+{
+	return LastParameterName;
+}
+
+void UAnimNextParameterSettings::SetLastParameterName(FName InLastParameterName)
+{
+	LastParameterName = InLastParameterName;
 }
\ No newline at end of file
diff --git a/Engine/Plugins/Experimental/AnimNext/Source/AnimNextEditor/Private/Param/AnimNextParameterSettings.h b/Engine/Plugins/Experimental/AnimNext/Source/AnimNextEditor/Private/Param/AnimNextParameterSettings.h
--- a/Engine/Plugins/Experimental/AnimNext/Source/AnimNextEditor/Private/Param/AnimNextParameterSettings.h
+++ b/Engine/Plugins/Experimental/AnimNext/Source/AnimNextEditor/Private/Param/AnimNextParameterSettings.h
@@ -15,12 +15,21 @@ class UAnimNextParameterSettings : public UObject
 	UAnimNextParameterSettings();
 
 public:
-	// Get the type of the last parameter type that we created
+	// Get the type of the last parameter that we created
 	const FAnimNextParamType& GetLastParameterType() const;
 
-	// Set the type of the last parameter type that we created
+	// Set the type of the last parameter that we created
 	void SetLastParameterType(const FAnimNextParamType& InLastParameterType);
+
+	// Get the name of the last parameter that we created
+	FName GetLastParameterName() const;
+
+	// Set the name of the last parameter that we created
+	void SetLastParameterName(FName InLastParameterName);
 private:
 	UPROPERTY(Transient)
 	FAnimNextParamType LastParameterType = FAnimNextParamType::GetType<bool>();
+
+	UPROPERTY(Transient)
+	FName LastParameterName = FName("NewParameter");
 };
\ No newline at end of file
diff --git a/Engine/Plugins/Experimental/AnimNext/Source/AnimNextEditor/Private/Param/ParamNamePropertyCustomization.cpp b/Engine/Plugins/Experimental/AnimNext/Source/AnimNextEditor/Private/Param/ParamNamePropertyCustomization.cpp
--- a/Engine/Plugins/Experimental/AnimNext/Source/AnimNextEditor/Private/Param/ParamNamePropertyCustomization.cpp
+++ b/Engine/Plugins/Experimental/AnimNext/Source/AnimNextEditor/Private/Param/ParamNamePropertyCustomization.cpp
@@ -53,7 +53,6 @@ void FParamNamePropertyTypeCustomization::CustomizeHeader(TSharedRef<IPropertyHa
 	FAnimNextParamType FilterType = FAnimNextParamType::FromString(ParamTypeString);
 	
 	FParameterPickerArgs PickerArgs;
-	PickerArgs.bShowSourceGraph = false;
 	PickerArgs.bMultiSelect = false;
 	PickerArgs.OnFilterParameterType = FOnFilterParameterType::CreateLambda([FilterType](const FAnimNextParamType& InParameterType)
 	{
@@ -64,7 +63,7 @@ void FParamNamePropertyTypeCustomization::CustomizeHeader(TSharedRef<IPropertyHa
 		return EFilterParameterResult::Exclude;
 	});
 	PickerArgs.NewParameterType = FilterType;
-	PickerArgs.OnParameterPicked = FOnParameterPicked::CreateLambda([this, WeakPropertyHandle = WeakPropertyHandle](const FParameterBindingReference& InParameterBinding)
+	PickerArgs.OnParameterPicked = FOnParameterPicked::CreateLambda([this](const FParameterBindingReference& InParameterBinding)
 	{
 		if(TSharedPtr<IPropertyHandle> PropertyHandle = WeakPropertyHandle.Pin())
 		{
@@ -75,7 +74,6 @@ void FParamNamePropertyTypeCustomization::CustomizeHeader(TSharedRef<IPropertyHa
 			Refresh();
 		}
 	});
-	PickerArgs.bAllowNone = bAllowNone;
 	
 	InHeaderRow
 	.NameContent()
diff --git a/Engine/Plugins/Experimental/AnimNext/Source/AnimNextEditor/Private/Param/ParamPropertyCustomization.cpp b/Engine/Plugins/Experimental/AnimNext/Source/AnimNextEditor/Private/Param/ParamPropertyCustomization.cpp
--- a/Engine/Plugins/Experimental/AnimNext/Source/AnimNextEditor/Private/Param/ParamPropertyCustomization.cpp
+++ b/Engine/Plugins/Experimental/AnimNext/Source/AnimNextEditor/Private/Param/ParamPropertyCustomization.cpp
@@ -7,16 +7,20 @@
 #include "EditorUtils.h"
 #include "PropertyHandle.h"
 #include "SParameterPickerCombo.h"
+#include "Modules/ModuleManager.h"
+#include "Param/AnimNextEditorParam.h"
 #include "Param/AnimNextParam.h"
+#include "Param/AnimNextParamInstanceIdentifier.h"
 #include "Param/ParamCompatibility.h"
 #include "Param/ParamUtils.h"
+#include "HAL/PlatformApplicationMisc.h"
 
 #define LOCTEXT_NAMESPACE "ParamPropertyCustomization"
 
 namespace UE::AnimNext::Editor
 {
 
-void FParamPropertyTypeCustomization::CustomizeHeader(TSharedRef<IPropertyHandle> InPropertyHandle, FDetailWidgetRow& InHeaderRow, IPropertyTypeCustomizationUtils& InCustomizationUtils)
+void FParamPropertyCustomization::CustomizeHeader(TSharedRef<IPropertyHandle> InPropertyHandle, FDetailWidgetRow& InHeaderRow, IPropertyTypeCustomizationUtils& InCustomizationUtils)
 {
 	auto GetMetadataProperty = [](const FProperty* InProperty)
 	{
@@ -34,16 +38,13 @@ void FParamPropertyTypeCustomization::CustomizeHeader(TSharedRef<IPropertyHandle
 	};
 
 	PropertyHandle = InPropertyHandle;
-	NamePropertyHandle = InPropertyHandle->GetChildHandle(GET_MEMBER_NAME_CHECKED(FAnimNextParam, Name));
-	TypePropertyHandle = InPropertyHandle->GetChildHandle(GET_MEMBER_NAME_CHECKED(FAnimNextParam, Type));
+	ParamStruct = CastFieldChecked<FStructProperty>(PropertyHandle->GetProperty())->Struct;
 
 	const FProperty* Property = GetMetadataProperty(InPropertyHandle->GetProperty());
 	const FString ParamTypeString = Property->GetMetaData("AllowedParamType");
-	const bool bAllowNone = Property->HasMetaData("AllowNone");
 	FAnimNextParamType FilterType = FAnimNextParamType::FromString(ParamTypeString);
 
 	FParameterPickerArgs PickerArgs;
-	PickerArgs.bShowSourceGraph = false;
 	PickerArgs.bMultiSelect = false;
 	PickerArgs.OnFilterParameterType = FOnFilterParameterType::CreateLambda([FilterType](const FAnimNextParamType& InParameterType)
 	{
@@ -54,61 +55,225 @@ void FParamPropertyTypeCustomization::CustomizeHeader(TSharedRef<IPropertyHandle
 		return EFilterParameterResult::Exclude;
 	});
 	PickerArgs.NewParameterType = FilterType;
-	PickerArgs.OnParameterPicked = FOnParameterPicked::CreateLambda([this](const FParameterBindingReference& InParameterBinding)
+	PickerArgs.OnInstanceIdChanged = FOnInstanceIdChanged::CreateSPLambda(this, [this](const TInstancedStruct<FAnimNextParamInstanceIdentifier>& InInstanceId)
 	{
-		if(PropertyHandle && NamePropertyHandle && TypePropertyHandle)
+		if(PropertyHandle)
 		{
 			PropertyHandle->NotifyPreChange();
-			NamePropertyHandle->SetValue(InParameterBinding.Parameter);
-			TArray<void*> TypeValues;
-			TypePropertyHandle->AccessRawData(TypeValues);
-			for(void* TypePtr : TypeValues)
+			if(ParamStruct == FAnimNextEditorParam::StaticStruct())
 			{
-				FAnimNextParamType& Type = *static_cast<FAnimNextParamType*>(TypePtr);
-				Type = InParameterBinding.Type;
+				PropertyHandle->EnumerateRawData([&InInstanceId](void* RawData, const int32 DataIndex, const int32 NumDatas)
+				{
+					FAnimNextEditorParam& Param = *static_cast<FAnimNextEditorParam*>(RawData);
+					Param.InstanceId = InInstanceId;
+					return true;
+				});
+			}
+			else if(ParamStruct == FAnimNextParam::StaticStruct())
+			{
+				FAnimNextParam ScheduleParam(NAME_None, FAnimNextParamType(), InInstanceId);
+				PropertyHandle->EnumerateRawData([&ScheduleParam](void* RawData, const int32 DataIndex, const int32 NumDatas)
+				{
+					FAnimNextParam& Param = *static_cast<FAnimNextParam*>(RawData);
+					Param.InstanceId = ScheduleParam.InstanceId;
+					return true;
+				});
 			}
 			PropertyHandle->NotifyPostChange(EPropertyChangeType::ValueSet);
+			PropertyHandle->NotifyFinishedChangingProperties();
+			Refresh();
+		}
+	});
+
+	PickerArgs.OnParameterPicked = FOnParameterPicked::CreateSPLambda(this, [this](const FParameterBindingReference& InParameterBinding)
+	{
+		if(PropertyHandle)
+		{
+			PropertyHandle->NotifyPreChange();
 
+			if(ParamStruct == FAnimNextEditorParam::StaticStruct())
+			{
+				FAnimNextEditorParam ParamValue(*InParameterBinding.Parameter.ToString(), InParameterBinding.Type, InParameterBinding.InstanceId);
+				PropertyHandle->EnumerateRawData([&ParamValue](void* RawData, const int32 DataIndex, const int32 NumDatas)
+				{
+					FAnimNextEditorParam& Param = *static_cast<FAnimNextEditorParam*>(RawData);
+					Param = ParamValue;
+					return true;
+				});
+			}
+			else if(ParamStruct == FAnimNextParam::StaticStruct())
+			{
+				FAnimNextParam ParamValue(*InParameterBinding.Parameter.ToString(), InParameterBinding.Type, InParameterBinding.InstanceId);
+				PropertyHandle->EnumerateRawData([&ParamValue](void* RawData, const int32 DataIndex, const int32 NumDatas)
+				{
+					FAnimNextParam& Param = *static_cast<FAnimNextParam*>(RawData);
+					Param = ParamValue;
+					return true;
+				});
+			}
+			PropertyHandle->NotifyPostChange(EPropertyChangeType::ValueSet);
+			PropertyHandle->NotifyFinishedChangingProperties();
 			Refresh();
 		}
 	});
-	PickerArgs.bAllowNone = bAllowNone;
-	
+
+	if(ParamStruct == FAnimNextEditorParam::StaticStruct())
+	{
+		TOptional<TInstancedStruct<FAnimNextParamInstanceIdentifier>> CommonInstanceId;
+		PropertyHandle->EnumerateConstRawData([&CommonInstanceId](const void* RawData, const int32 DataIndex, const int32 NumDatas)
+		{
+			const FAnimNextEditorParam& Param = *static_cast<const FAnimNextEditorParam*>(RawData);
+			if(!CommonInstanceId.IsSet())
+			{
+				CommonInstanceId = Param.InstanceId;
+			}
+			else if(CommonInstanceId.GetValue() != Param.InstanceId)
+			{
+				// No common scope, so use a null instance
+				CommonInstanceId = TInstancedStruct<FAnimNextParamInstanceIdentifier>();
+				return false;
+			}
+			return true;
+		});
+
+		PickerArgs.InstanceId = CommonInstanceId.IsSet() ? CommonInstanceId.GetValue() : TInstancedStruct<FAnimNextParamInstanceIdentifier>();
+	}
+	else if(ParamStruct == FAnimNextParam::StaticStruct())
+	{
+		struct FInstanceId
+		{
+			FInstanceId(FName InInstanceId, const UScriptStruct* InInstanceIdType)
+				: InstanceId(InInstanceId)
+				, InstanceIdType(InInstanceIdType)
+			{}
+
+			FName InstanceId;
+			const UScriptStruct* InstanceIdType;
+		};
+		
+		TOptional<FInstanceId> CommonInstanceId;
+		PropertyHandle->EnumerateConstRawData([&CommonInstanceId](const void* RawData, const int32 DataIndex, const int32 NumDatas)
+		{
+			const FAnimNextParam& Param = *static_cast<const FAnimNextParam*>(RawData);
+			if(!CommonInstanceId.IsSet())
+			{
+				if(!Param.InstanceId.IsNone() && Param.InstanceIdType != nullptr)
+				{
+					CommonInstanceId = FInstanceId(Param.InstanceId, Param.InstanceIdType);
+				}
+			}
+			else if(CommonInstanceId.GetValue().InstanceId != Param.InstanceId && CommonInstanceId.GetValue().InstanceIdType != Param.InstanceIdType)
+			{
+				// No common instance ID, so use NAME_None
+				CommonInstanceId = FInstanceId(NAME_None, FAnimNextParamInstanceIdentifier::StaticStruct());
+				return false;
+			}
+			return true;
+		});
+
+		if(CommonInstanceId.IsSet())
+		{
+			PickerArgs.InstanceId.InitializeAsScriptStruct(CommonInstanceId.GetValue().InstanceIdType);
+			PickerArgs.InstanceId.GetMutable().FromName(CommonInstanceId.GetValue().InstanceId);
+		}
+		else
+		{
+			PickerArgs.InstanceId = TInstancedStruct<FAnimNextParamInstanceIdentifier>();
+		}
+	}
+
+	FUIAction CopyAction, PasteAction;
+	PropertyHandle->CreateDefaultPropertyCopyPasteActions(CopyAction, PasteAction);
+
+	DefaultCopyAction = CopyAction.ExecuteAction;
+	CopyAction.ExecuteAction = FExecuteAction::CreateSP(this, &FParamPropertyCustomization::HandleCopy);
+
 	InHeaderRow
+	.CopyAction(CopyAction)
+	.PasteAction(PasteAction)
 	.NameContent()
 	[
-		InPropertyHandle->CreatePropertyNameWidget()
+		PropertyHandle->CreatePropertyNameWidget()
 	]
 	.ValueContent()
 	[
-		SNew(SParameterPickerCombo)
-		.PickerArgs(PickerArgs)
-		.OnGetParameterName_Lambda([this]()
-		{
-			return CachedName;
-		})
-		.OnGetParameterType_Lambda([this]()
-		{
-			return CachedType;
-		})
+		SNew(SVerticalBox)
+		+SVerticalBox::Slot()
+		.AutoHeight()
+		[
+			SNew(SParameterPickerCombo)
+			.PickerArgs(PickerArgs)
+			.OnGetParameterName_Lambda([this]()
+			{
+				return CachedParam.Name;
+			})
+			.OnGetParameterType_Lambda([this]()
+			{
+				return CachedParam.Type;
+			})
+			.OnGetParameterInstanceId_Lambda([this]()
+			{
+				return CachedParam.InstanceId;
+			})
+		]
 	];
 
 	Refresh();
 }
 
-void FParamPropertyTypeCustomization::CustomizeChildren(TSharedRef<IPropertyHandle> InPropertyHandle, IDetailChildrenBuilder& ChildBuilder, IPropertyTypeCustomizationUtils& CustomizationUtils)
+void FParamPropertyCustomization::CustomizeChildren(TSharedRef<IPropertyHandle> InPropertyHandle, IDetailChildrenBuilder& ChildBuilder, IPropertyTypeCustomizationUtils& CustomizationUtils)
 {
 }
 
-void FParamPropertyTypeCustomization::Refresh()
+void FParamPropertyCustomization::Refresh()
 {
-	TArray<void*> TypeValues;
-	TypePropertyHandle->AccessRawData(TypeValues);
-	if(TypeValues.Num() == 1 && TypeValues[0] != nullptr)
+	CachedParam = FAnimNextEditorParam();
+
+	if(ParamStruct == FAnimNextEditorParam::StaticStruct())
 	{
-		CachedType = *static_cast<FAnimNextParamType*>(TypeValues[0]);
+		PropertyHandle->EnumerateConstRawData([this](const void* RawData, const int32 DataIndex, const int32 NumDatas)
+		{
+			if(NumDatas == 1 && RawData != nullptr)
+			{
+				CachedParam = *static_cast<const FAnimNextEditorParam*>(RawData);
+			}
+			return false;
+		});
+	}
+	else if(ParamStruct == FAnimNextParam::StaticStruct())
+	{
+		PropertyHandle->EnumerateConstRawData([this](const void* RawData, const int32 DataIndex, const int32 NumDatas)
+		{
+			if(NumDatas == 1 && RawData != nullptr)
+			{
+				CachedParam = FAnimNextEditorParam(*static_cast<const FAnimNextParam*>(RawData));
+			}
+			return false;
+		});
+	}
+}
+
+void FParamPropertyCustomization::HandleCopy()
+{
+	DefaultCopyAction.Execute();
+
+	// We always copy-paste as a FAnimNextEditorParam
+	if(ParamStruct == FAnimNextParam::StaticStruct())
+	{
+		// Grab the clipboard text
+		FString ImportClipboardString;
+		FPlatformApplicationMisc::ClipboardPaste(ImportClipboardString);
+
+		// Import as FAnimNextParam
+		FAnimNextParam ImportedParam;
+		FAnimNextParam::StaticStruct()->ImportText(*ImportClipboardString, &ImportedParam, nullptr, PPF_None, nullptr, FAnimNextParam::StaticStruct()->GetName());
+
+		// Re-export as FAnimNextEditorParam
+		FAnimNextEditorParam ParamToExport(ImportedParam);
+		FString ExportClipboardString;
+		FAnimNextEditorParam::StaticStruct()->ExportText(ExportClipboardString, &ParamToExport, nullptr, nullptr, PPF_None, nullptr);
+		FPlatformApplicationMisc::ClipboardCopy(*ExportClipboardString);
 	}
-	NamePropertyHandle->GetValue(CachedName);
 }
 
 }
diff --git a/Engine/Plugins/Experimental/AnimNext/Source/AnimNextEditor/Private/Param/ParamPropertyCustomization.h b/Engine/Plugins/Experimental/AnimNext/Source/AnimNextEditor/Private/Param/ParamPropertyCustomization.h
--- a/Engine/Plugins/Experimental/AnimNext/Source/AnimNextEditor/Private/Param/ParamPropertyCustomization.h
+++ b/Engine/Plugins/Experimental/AnimNext/Source/AnimNextEditor/Private/Param/ParamPropertyCustomization.h
@@ -4,12 +4,12 @@
 
 #include "IPropertyTypeCustomization.h"
 #include "PropertyEditorModule.h"
-#include "Param/ParamType.h"
+#include "Param/AnimNextEditorParam.h"
 
 namespace UE::AnimNext::Editor
 {
 
-class FParamPropertyTypeCustomization : public IPropertyTypeCustomization
+class FParamPropertyCustomization : public IPropertyTypeCustomization
 {
 private:
 	// IPropertyTypeCustomization interface
@@ -18,13 +18,14 @@ private:
 
 	void Refresh();
 
+	void HandleCopy();
+
 private:
 	TSharedPtr<IPropertyHandle> PropertyHandle;
-	TSharedPtr<IPropertyHandle> NamePropertyHandle;
-	TSharedPtr<IPropertyHandle> TypePropertyHandle;
-	FName CachedName = NAME_None;
+	FAnimNextEditorParam CachedParam;
 	FText CachedNameText;
-	FAnimNextParamType CachedType;
+	UScriptStruct* ParamStruct = nullptr;
+	FExecuteAction DefaultCopyAction;
 };
 
 }
diff --git a/Engine/Plugins/Experimental/AnimNext/Source/AnimNextEditor/Private/Param/ParametersGraphPanelPinFactory.cpp b/Engine/Plugins/Experimental/AnimNext/Source/AnimNextEditor/Private/Param/ParametersGraphPanelPinFactory.cpp
--- a/Engine/Plugins/Experimental/AnimNext/Source/AnimNextEditor/Private/Param/ParametersGraphPanelPinFactory.cpp
+++ b/Engine/Plugins/Experimental/AnimNext/Source/AnimNextEditor/Private/Param/ParametersGraphPanelPinFactory.cpp
@@ -1,10 +1,13 @@
 ﻿// Copyright Epic Games, Inc. All Rights Reserved.
 
 #include "ParametersGraphPanelPinFactory.h"
+#include "SGraphPinParam.h"
 #include "SGraphPinParamName.h"
 #include "EdGraphSchema_K2.h"
 #include "EdGraph/RigVMEdGraph.h"
 #include "EdGraph/RigVMEdGraphNode.h"
+#include "Param/AnimNextParam.h"
+#include "Param/AnimNextEditorParam.h"
 
 namespace UE::AnimNext::Editor
 {
@@ -18,12 +21,19 @@ TSharedPtr<SGraphPin> FParametersGraphPanelPinFactory::CreatePin_Internal(UEdGra
 {
 	if (URigVMEdGraphNode* RigNode = Cast<URigVMEdGraphNode>(InPin->GetOwningNode()))
 	{
-		URigVMEdGraph* RigGraph = Cast<URigVMEdGraph>(RigNode->GetGraph());
-
 		URigVMPin* ModelPin = RigNode->GetModelPinFromPinPath(InPin->GetName());
 		if (ModelPin)
 		{
-			if(ModelPin->GetCustomWidgetName() == "ParamName")
+			if(InPin->PinType.PinCategory == UEdGraphSchema_K2::PC_Struct && (InPin->PinType.PinSubCategoryObject.Get() == FAnimNextEditorParam::StaticStruct() || InPin->PinType.PinSubCategoryObject.Get() == FAnimNextParam::StaticStruct()))
+			{
+				const FString ParamTypeString = ModelPin->GetMetaData("AllowedParamType");
+				FAnimNextParamType FilterType = FAnimNextParamType::FromString(ParamTypeString);
+				return SNew(SGraphPinParam, InPin)
+					.ModelPin(ModelPin)
+					.GraphNode(RigNode)
+					.FilterType(FilterType);
+			}
+			else if(ModelPin->GetCustomWidgetName() == "ParamName")
 			{
 				const FString ParamTypeString = ModelPin->GetMetaData("AllowedParamType");
 				FAnimNextParamType FilterType = FAnimNextParamType::FromString(ParamTypeString);
@@ -36,7 +46,14 @@ TSharedPtr<SGraphPin> FParametersGraphPanelPinFactory::CreatePin_Internal(UEdGra
 	}
 	else if(UEdGraphNode* EdGraphNode = InPin->GetOwningNode())
 	{
-		if(EdGraphNode->GetPinMetaData(InPin->GetFName(), "CustomWidget") == "ParamName")
+		if(InPin->PinType.PinCategory == UEdGraphSchema_K2::PC_Struct && (InPin->PinType.PinSubCategoryObject.Get() == FAnimNextEditorParam::StaticStruct() || InPin->PinType.PinSubCategoryObject.Get() == FAnimNextParam::StaticStruct()))
+		{
+			const FString ParamTypeString = EdGraphNode->GetPinMetaData(InPin->GetFName(), "AllowedParamType");
+			FAnimNextParamType FilterType = FAnimNextParamType::FromString(ParamTypeString);
+			return SNew(SGraphPinParam, InPin)
+				.FilterType(FilterType);
+		}
+		else if(EdGraphNode->GetPinMetaData(InPin->GetFName(), "CustomWidget") == "ParamName")
 		{
 			const FString ParamTypeString = EdGraphNode->GetPinMetaData(InPin->GetFName(), "AllowedParamType");
 			FAnimNextParamType FilterType = FAnimNextParamType::FromString(ParamTypeString);
diff --git a/Engine/Plugins/Experimental/AnimNext/Source/AnimNextEditor/Private/Param/SAddParametersDialog.cpp b/Engine/Plugins/Experimental/AnimNext/Source/AnimNextEditor/Private/Param/SAddParametersDialog.cpp
--- a/Engine/Plugins/Experimental/AnimNext/Source/AnimNextEditor/Private/Param/SAddParametersDialog.cpp
+++ b/Engine/Plugins/Experimental/AnimNext/Source/AnimNextEditor/Private/Param/SAddParametersDialog.cpp
@@ -51,11 +51,12 @@ bool FParameterToAdd::IsValid(FText& OutReason) const
 	return true; 
 }
 
-void SAddParametersDialog::Construct(const FArguments& InArgs)
+void SAddParametersDialog::Construct(const FArguments& InArgs, const FAssetData& InAsset)
 {
 	using namespace AddParametersDialog;
 
 	OnFilterParameterType = InArgs._OnFilterParameterType;
+	Asset = InAsset;
 
 	SWindow::Construct(SWindow::FArguments()
 		.Title(LOCTEXT("WindowTitle", "Add Parameters"))
@@ -87,7 +88,7 @@ void SAddParametersDialog::Construct(const FArguments& InArgs)
 				+SVerticalBox::Slot()
 				.FillHeight(1.0f)
 				[
-					SAssignNew(EntriesList, SListView<TSharedRef<FParameterToAdd>>)
+					SAssignNew(EntriesList, SListView<TSharedRef<FParameterToAddEntry>>)
 					.ListItemsSource(&Entries)
 					.OnGenerateRow(this, &SAddParametersDialog::HandleGenerateRow)
 					.ItemHeight(20.0f)
@@ -149,6 +150,7 @@ void SAddParametersDialog::Construct(const FArguments& InArgs)
 						})
 						.OnClicked_Lambda([this]()
 						{
+							bOKPressed = true;
 							RequestDestroyWindow();
 							return FReply::Handled();
 						})
@@ -162,7 +164,6 @@ void SAddParametersDialog::Construct(const FArguments& InArgs)
 						.ToolTipText(LOCTEXT("CancelButtonTooltip", "Cancel adding new parameters"))
 						.OnClicked_Lambda([this]()
 						{
-							bCancelPressed = true;
 							RequestDestroyWindow();
 							return FReply::Handled();
 						})
@@ -175,18 +176,29 @@ void SAddParametersDialog::Construct(const FArguments& InArgs)
 	AddEntry(InArgs._InitialParamType);
 }
 
+FReply SAddParametersDialog::OnKeyDown(const FGeometry& MyGeometry, const FKeyEvent& InKeyEvent)
+{
+	if(InKeyEvent.GetKey() == EKeys::Escape)
+	{
+		RequestDestroyWindow();
+		return FReply::Handled();
+	}
+
+	return FReply::Unhandled();
+}
+
 void SAddParametersDialog::AddEntry(const FAnimNextParamType& InParamType)
 {
 	const UAnimNextParameterSettings* Settings = GetDefault<UAnimNextParameterSettings>();
 	
 	TArray<FName> PendingNames;
 	PendingNames.Reserve(Entries.Num());
-	for(const TSharedRef<FParameterToAdd>& QueuedAdd : Entries)
+	for(const TSharedRef<FParameterToAddEntry>& QueuedAdd : Entries)
 	{
 		PendingNames.Add(QueuedAdd->Name);
 	}
-	FName ParameterName = FUtils::GetNewParameterName(TEXT("NewParameter"), PendingNames);
-	Entries.Add(MakeShared<FParameterToAdd>(InParamType.IsValid() ? InParamType : Settings->GetLastParameterType(), ParameterName));
+	FName ParameterName = FUtils::GetNewParameterName(Settings->GetLastParameterName(), Asset, PendingNames);
+	Entries.Add(MakeShared<FParameterToAddEntry>(InParamType.IsValid() ? InParamType : Settings->GetLastParameterType(), ParameterName));
 
 	RefreshEntries();
 }
@@ -196,17 +208,17 @@ void SAddParametersDialog::RefreshEntries()
 	EntriesList->RequestListRefresh();
 }
 
-class SParameterToAdd : public SMultiColumnTableRow<TSharedRef<FParameterToAdd>>
+class SParameterToAdd : public SMultiColumnTableRow<TSharedRef<SAddParametersDialog::FParameterToAddEntry>>
 {
 	SLATE_BEGIN_ARGS(SParameterToAdd) {}
 	SLATE_END_ARGS()
 
-	void Construct(const FArguments& InArgs, const TSharedRef<STableViewBase>& InOwnerTableView, TSharedRef<FParameterToAdd> InEntry, TSharedRef<SAddParametersDialog> InDialog)
+	void Construct(const FArguments& InArgs, const TSharedRef<STableViewBase>& InOwnerTableView, TSharedRef<SAddParametersDialog::FParameterToAddEntry> InEntry, TSharedRef<SAddParametersDialog> InDialog)
 	{
 		Entry = InEntry;
 		WeakDialog = InDialog;
 		
-		SMultiColumnTableRow<TSharedRef<FParameterToAdd>>::Construct( SMultiColumnTableRow<TSharedRef<FParameterToAdd>>::FArguments(), InOwnerTableView);
+		SMultiColumnTableRow<TSharedRef<SAddParametersDialog::FParameterToAddEntry>>::Construct( SMultiColumnTableRow<TSharedRef<SAddParametersDialog::FParameterToAddEntry>>::FArguments(), InOwnerTableView);
 	}
 
 	virtual TSharedRef<SWidget> GenerateWidgetForColumn(const FName& InColumnName) override
@@ -215,53 +227,67 @@ class SParameterToAdd : public SMultiColumnTableRow<TSharedRef<FParameterToAdd>>
 
 		if(InColumnName == Column_Name)
 		{
-			return
+			TSharedPtr<SInlineEditableTextBlock> EditableText;
+			TSharedRef<SWidget> Widget =
 				SNew(SBox)
 				.HAlign(HAlign_Left)
 				.VAlign(VAlign_Center)
 				[
-					SNew(SInlineEditableTextBlock)
+					SAssignNew(EditableText, SInlineEditableTextBlock)
 					.Font(IDetailLayoutBuilder::GetDetailFont())
 					.IsSelected(this, &SParameterToAdd::IsSelectedExclusively)
 					.ToolTipText(LOCTEXT("NameTooltip", "The name of the new parameter"))
 					.Text_Lambda([this]()
 					{
-						return UncookedOnly::FUtils::GetParameterDisplayNameText(Entry->Name);
+						return FText::FromName(Entry->Name);
 					})
 					.OnTextCommitted_Lambda([this](const FText& InText, ETextCommit::Type InCommitType)
 					{
-						const FString UserInput = InText.ToString();
-						// Parse out segments to collapse adjacent delimiters
-						TStringBuilder<128> RebuiltInput;
-						UE::String::ParseTokensMultiple(UserInput, { TEXT('.'), TEXT('_') }, [&RebuiltInput](const FStringView InToken)
-						{
-							if(RebuiltInput.Len() != 0)
-							{
-								RebuiltInput.Append(TEXT("_"));
-							}
-							RebuiltInput.Append(InToken);
-						}, String::EParseTokensOptions::SkipEmpty);
-						Entry->Name = RebuiltInput.ToString();
+						Entry->Name = *InText.ToString();
+
+						UAnimNextParameterSettings* Settings = GetMutableDefault<UAnimNextParameterSettings>();
+						Settings->SetLastParameterName(Entry->Name);
 					})
 					.OnVerifyTextChanged_Lambda([this](const FText& InNewText, FText& OutErrorText)
 					{
 						const FString NewString = InNewText.ToString();
 
-						if(!FUtils::IsValidEntryNameString(NewString, OutErrorText))
+						if(!FUtils::IsValidParameterNameString(NewString, OutErrorText))
 						{
 							return false;
 						}
 
-						const FName Name(*NewString);
-						if(FUtils::DoesParameterNameExist(Name))
+						if(TSharedPtr<SAddParametersDialog> Dialog = WeakDialog.Pin())
 						{
-							OutErrorText = LOCTEXT("Error_NameExists", "This name already exists in the project");
-							return false;
+							const FName Name(*NewString);
+							if(FUtils::DoesParameterNameExistInAsset(Name, Dialog->Asset))
+							{
+								OutErrorText = LOCTEXT("Error_NameExists", "This name already exists in the project");
+								return false;
+							}
+
+							return true;
 						}
 
-						return true;
+						return false;
 					})
 				];
+
+			if(Entry->bIsNew)
+			{
+				EditableText->RegisterActiveTimer(1/60.0f, FWidgetActiveTimerDelegate::CreateSPLambda(EditableText.Get(), [WeakEditableText = TWeakPtr<SInlineEditableTextBlock>(EditableText)](double, float)
+				{
+					if(TSharedPtr<SInlineEditableTextBlock> PinnedEditableText = WeakEditableText.Pin())
+					{
+						PinnedEditableText->EnterEditingMode();
+					}
+					return EActiveTimerReturnType::Stop;
+				}));
+	
+				Entry->bIsNew = false;
+			}
+
+			return Widget;
 		}
 		else if(InColumnName == Column_Type)
 		{
@@ -345,11 +371,11 @@ class SParameterToAdd : public SMultiColumnTableRow<TSharedRef<FParameterToAdd>>
 		return SNullWidget::NullWidget;
 	}
 
-	TSharedPtr<FParameterToAdd> Entry;
+	TSharedPtr<SAddParametersDialog::FParameterToAddEntry> Entry;
 	TWeakPtr<SAddParametersDialog> WeakDialog;
 };
 
-TSharedRef<ITableRow> SAddParametersDialog::HandleGenerateRow(TSharedRef<FParameterToAdd> InEntry, const TSharedRef<STableViewBase>& InOwnerTable)
+TSharedRef<ITableRow> SAddParametersDialog::HandleGenerateRow(TSharedRef<FParameterToAddEntry> InEntry, const TSharedRef<STableViewBase>& InOwnerTable)
 {
 	return SNew(SParameterToAdd, InOwnerTable, InEntry, SharedThis(this));
 }
@@ -358,10 +384,10 @@ bool SAddParametersDialog::ShowModal(TArray<FParameterToAdd>& OutParameters)
 {
 	FSlateApplication::Get().AddModalWindow(SharedThis(this), FGlobalTabmanager::Get()->GetRootWindow());
 
-	if(!bCancelPressed)
+	if(bOKPressed)
 	{
 		bool bHasValid = false;
-		for(TSharedRef<FParameterToAdd>& Entry : Entries)
+		for(TSharedRef<FParameterToAddEntry>& Entry : Entries)
 		{
 			if(Entry->IsValid())
 			{
@@ -374,7 +400,7 @@ bool SAddParametersDialog::ShowModal(TArray<FParameterToAdd>& OutParameters)
 	return false;
 }
 
-TSharedRef<SWidget> SAddParametersDialog::HandleGetAddParameterMenuContent(TSharedPtr<FParameterToAdd> InEntry)
+TSharedRef<SWidget> SAddParametersDialog::HandleGetAddParameterMenuContent(TSharedPtr<FParameterToAddEntry> InEntry)
 {
 	using namespace AddParametersDialog;
 
diff --git a/Engine/Plugins/Experimental/AnimNext/Source/AnimNextEditor/Private/Param/SAddParametersDialog.h b/Engine/Plugins/Experimental/AnimNext/Source/AnimNextEditor/Private/Param/SAddParametersDialog.h
--- a/Engine/Plugins/Experimental/AnimNext/Source/AnimNextEditor/Private/Param/SAddParametersDialog.h
+++ b/Engine/Plugins/Experimental/AnimNext/Source/AnimNextEditor/Private/Param/SAddParametersDialog.h
@@ -35,31 +35,44 @@ public:
 
 	SLATE_END_ARGS()
 
-	void Construct(const FArguments& InArgs);
+	void Construct(const FArguments& InArgs, const FAssetData& InAsset);
 
 	bool ShowModal(TArray<FParameterToAdd>& OutParameters);
 
 private:
+	virtual FReply OnKeyDown(const FGeometry& MyGeometry, const FKeyEvent& InKeyEvent) override;
+
 	void AddEntry(const FAnimNextParamType& InParamType = FAnimNextParamType());
 
 	void RefreshEntries();
 
-	TSharedRef<ITableRow> HandleGenerateRow(TSharedRef<FParameterToAdd> InEntry, const TSharedRef<STableViewBase>& InOwnerTable);
+	struct FParameterToAddEntry : FParameterToAdd
+	{
+		FParameterToAddEntry() = default;
+
+		FParameterToAddEntry(const FAnimNextParamType& InType, FName InName)
+			: FParameterToAdd(InType, InName)
+		{}
+
+		bool bIsNew = true;
+	};
 
-	TSharedRef<SWidget> HandleGetAddParameterMenuContent(TSharedPtr<FParameterToAdd> InEntry);
+	TSharedRef<ITableRow> HandleGenerateRow(TSharedRef<FParameterToAddEntry> InEntry, const TSharedRef<STableViewBase>& InOwnerTable);
+
+	TSharedRef<SWidget> HandleGetAddParameterMenuContent(TSharedPtr<FParameterToAddEntry> InEntry);
 	
 private:
 	friend class SParameterToAdd;
 
-	TSharedPtr<SWrapBox> QueuedParametersBox;
-
-	TSharedPtr<SListView<TSharedRef<FParameterToAdd>>> EntriesList;
+	TSharedPtr<SListView<TSharedRef<FParameterToAddEntry>>> EntriesList;
 
-	TArray<TSharedRef<FParameterToAdd>> Entries;
+	TArray<TSharedRef<FParameterToAddEntry>> Entries;
 
 	FOnFilterParameterType OnFilterParameterType;
 
-	bool bCancelPressed = false;
+	FAssetData Asset;
+
+	bool bOKPressed = false;
 };
 
 }
diff --git a/Engine/Plugins/Experimental/AnimNext/Source/AnimNextEditor/Private/Param/SGraphPinParam.cpp b/Engine/Plugins/Experimental/AnimNext/Source/AnimNextEditor/Private/Param/SGraphPinParam.cpp
new file mode 100644
--- /dev/null
+++ b/Engine/Plugins/Experimental/AnimNext/Source/AnimNextEditor/Private/Param/SGraphPinParam.cpp
@@ -0,0 +1,211 @@
+﻿// Copyright Epic Games, Inc. All Rights Reserved.
+
+#include "SGraphPinParam.h"
+#include "EditorUtils.h"
+#include "SParameterPickerCombo.h"
+#include "ScopedTransaction.h"
+#include "UncookedOnlyUtils.h"
+#include "EdGraph/RigVMEdGraphNode.h"
+#include "Graph/AnimNextGraph.h"
+#include "Graph/AnimNextGraph_Controller.h"
+#include "Param/ParamCompatibility.h"
+#include "Param/ParamUtils.h"
+#include "Param/RigVMDispatch_GetLayerParameter.h"
+#include "Param/RigVMDispatch_GetParameter.h"
+#include "Param/RigVMDispatch_GetScopedParameter.h"
+#include "Param/RigVMDispatch_SetLayerParameter.h"
+#include "RigVMModel/Nodes/RigVMTemplateNode.h"
+#include "RigVMModel/RigVMPin.h"
+#include "RigVMCore/RigVMRegistry.h"
+#include "RigVMModel/RigVMController.h"
+#include "Param/AnimNextParam.h"
+#include "Param/AnimNextEditorParam.h"
+
+#define LOCTEXT_NAMESPACE "SGraphPinParam"
+
+namespace UE::AnimNext::Editor
+{
+
+void SGraphPinParam::Construct(const FArguments& InArgs, UEdGraphPin* InPin)
+{
+	ModelPin = InArgs._ModelPin;
+	Node = InArgs._GraphNode;
+	FilterType = InArgs._FilterType;
+	ParamStruct = CastChecked<UScriptStruct>(InPin->PinType.PinSubCategoryObject.Get());
+
+	if(ModelPin && ModelPin->GetGraph())
+	{
+		ModelPin->GetGraph()->OnModified().AddSP(this, &SGraphPinParam::HandleGraphModified);
+	}
+	
+	SGraphPin::Construct(SGraphPin::FArguments(), InPin);
+}
+
+TSharedRef<SWidget> SGraphPinParam::GetDefaultValueWidget()
+{
+	FParameterPickerArgs Args;
+	Args.bMultiSelect = false;
+
+	UpdateCachedParam();
+
+	Args.OnParameterPicked = FOnParameterPicked::CreateLambda([this](const FParameterBindingReference& InParameterBinding)
+	{
+		{
+			FScopedTransaction Transaction(LOCTEXT("SelectParameter", "Select Parameter"));
+
+			if(Node && ModelPin)
+			{
+				UAnimNextGraph_Controller* Controller = CastChecked<UAnimNextGraph_Controller>(Node->GetController());
+				Controller->SetAnimNextParameterNode(Node->GetModelNode(), InParameterBinding.Parameter, InParameterBinding.Type.GetValueType(), InParameterBinding.Type.GetContainerType(), InParameterBinding.Type.GetValueTypeObject(), InParameterBinding.InstanceId);
+			}
+			else
+			{
+				FString ValueAsString;
+				if(ParamStruct == FAnimNextEditorParam::StaticStruct())
+				{
+					FAnimNextEditorParam ParamValue(*InParameterBinding.Parameter.ToString(), InParameterBinding.Type, InParameterBinding.InstanceId);
+					FAnimNextEditorParam::StaticStruct()->ExportText(ValueAsString, &ParamValue, nullptr, nullptr, PPF_None, nullptr);
+				}
+				else if(ParamStruct == FAnimNextParam::StaticStruct())
+				{
+					FAnimNextParam ParamValue(*InParameterBinding.Parameter.ToString(), InParameterBinding.Type, InParameterBinding.InstanceId);
+					FAnimNextParam::StaticStruct()->ExportText(ValueAsString, &ParamValue, nullptr, nullptr, PPF_None, nullptr);
+				}
+
+				GraphPinObj->Modify();
+				GraphPinObj->GetSchema()->TrySetDefaultValue(*GraphPinObj, ValueAsString);
+			}
+		}
+
+		UpdateCachedParam();
+	});
+	
+	Args.OnFilterParameterType = FOnFilterParameterType::CreateLambda([this](const FAnimNextParamType& InParamType)-> EFilterParameterResult
+	{
+		if(FilterType.IsValid())
+		{
+			if(!FParamUtils::GetCompatibility(FilterType, InParamType).IsCompatible())
+			{
+				return EFilterParameterResult::Exclude;
+			}
+		}
+
+		if(InParamType.IsValid())
+		{
+			const FRigVMTemplateArgumentType RigVMType = InParamType.ToRigVMTemplateArgument();
+			if(!RigVMType.IsValid() || FRigVMRegistry::Get().GetTypeIndex(RigVMType) == INDEX_NONE)
+			{
+				return EFilterParameterResult::Exclude;
+			}
+		}
+
+		if(Node && ModelPin->IsLinked())
+		{
+			const FAnimNextParamType Type = FAnimNextParamType::FromRigVMTemplateArgument(ModelPin->GetTemplateArgumentType());
+			if(!Type.IsValid() || !FParamUtils::GetCompatibility(Type, InParamType).IsCompatible())
+			{
+				return EFilterParameterResult::Exclude;
+			}
+		}
+
+		return EFilterParameterResult::Include;
+	});
+	Args.NewParameterType = FilterType;
+
+	Args.OnInstanceIdChanged = FOnInstanceIdChanged::CreateLambda([this](const TInstancedStruct<FAnimNextParamInstanceIdentifier>& InInstanceId)
+	{
+		{
+			FScopedTransaction Transaction(LOCTEXT("SelectParameterScope", "Select Parameter Scope"));
+			
+			if(Node && ModelPin)
+			{
+				UAnimNextGraph_Controller* Controller = CastChecked<UAnimNextGraph_Controller>(Node->GetController());
+				Controller->SetAnimNextParameterNode(Node->GetModelNode(), NAME_None, EPropertyBagPropertyType::None, EPropertyBagContainerType::None, nullptr, InInstanceId);
+			}
+			else
+			{
+				FString ValueAsString;
+				if(ParamStruct == FAnimNextEditorParam::StaticStruct())
+				{
+					FAnimNextEditorParam ParamValue(NAME_None, FAnimNextParamType(), InInstanceId);
+					FAnimNextEditorParam::StaticStruct()->ExportText(ValueAsString, &ParamValue, nullptr, nullptr, PPF_None, nullptr);
+				}
+				else if(ParamStruct == FAnimNextParam::StaticStruct())
+				{
+					FAnimNextParam ParamValue(NAME_None, FAnimNextParamType(), InInstanceId);
+					FAnimNextParam::StaticStruct()->ExportText(ValueAsString, &ParamValue, nullptr, nullptr, PPF_None, nullptr);
+				}
+
+				GraphPinObj->Modify();
+				GraphPinObj->GetSchema()->TrySetDefaultValue(*GraphPinObj, ValueAsString);
+			}
+		}
+
+		UpdateCachedParam();
+	});
+
+	Args.InstanceId = CachedParam.InstanceId;
+
+	return SAssignNew(PickerCombo, SParameterPickerCombo)
+		.PickerArgs(Args)
+		.OnGetParameterName_Lambda([this]()
+		{
+			return CachedParam.Name;
+		})
+		.OnGetParameterType_Lambda([this]()
+		{
+			return CachedParam.Type;
+		})
+		.OnGetParameterInstanceId_Lambda([this]()
+		{
+			return CachedParam.InstanceId;
+		});
+}
+
+void SGraphPinParam::UpdateCachedParam()
+{
+	TStringBuilder<256> DefaultValueBuilder;
+	if(ModelPin)
+	{
+		DefaultValueBuilder.Append(ModelPin->GetDefaultValue());
+	}
+	else if(GraphPinObj)
+	{
+		DefaultValueBuilder.Append(GraphPinObj->DefaultValue);
+	}
+
+	if(DefaultValueBuilder.Len() > 0)
+	{
+		if(ParamStruct == FAnimNextEditorParam::StaticStruct())
+		{
+			FAnimNextEditorParam::StaticStruct()->ImportText(DefaultValueBuilder.ToString(), &CachedParam, nullptr, PPF_None, nullptr, FAnimNextEditorParam::StaticStruct()->GetName());
+		}
+		else if(ParamStruct == FAnimNextParam::StaticStruct())
+		{
+			FAnimNextParam Param;
+			FAnimNextParam::StaticStruct()->ImportText(DefaultValueBuilder.ToString(), &Param, nullptr, PPF_None, nullptr, FAnimNextParam::StaticStruct()->GetName());
+			CachedParam = FAnimNextEditorParam(Param);
+		}
+	}
+
+	if(PickerCombo.IsValid())
+	{
+		PickerCombo->RequestRefresh();
+	}
+}
+
+void SGraphPinParam::HandleGraphModified(ERigVMGraphNotifType InType, URigVMGraph* InGraph, UObject* InSubject)
+{
+	switch (InType)
+	{
+	case ERigVMGraphNotifType::PinDefaultValueChanged:
+		UpdateCachedParam();
+		break;
+	default:
+		break;
+	}
+}
+
+}
+
+#undef LOCTEXT_NAMESPACE
\ No newline at end of file
diff --git a/Engine/Plugins/Experimental/AnimNext/Source/AnimNextEditor/Private/Param/SGraphPinParam.h b/Engine/Plugins/Experimental/AnimNext/Source/AnimNextEditor/Private/Param/SGraphPinParam.h
new file mode 100644
--- /dev/null
+++ b/Engine/Plugins/Experimental/AnimNext/Source/AnimNextEditor/Private/Param/SGraphPinParam.h
@@ -0,0 +1,61 @@
+﻿// Copyright Epic Games, Inc. All Rights Reserved.
+
+#pragma once
+
+#include "SGraphPin.h"
+#include "Param/AnimNextEditorParam.h"
+#include "Param/AnimNextParam.h"
+#include "Param/ParamType.h"
+
+class URigVMGraph;
+enum class ERigVMGraphNotifType : uint8;
+class URigVMPin;
+class URigVMEdGraphNode;
+
+namespace UE::AnimNext::Editor
+{
+	class SParameterPickerCombo;
+}
+
+namespace UE::AnimNext::Editor
+{
+
+// A pin widget that allows picking using an AnimNext parameter picker
+class SGraphPinParam : public SGraphPin
+{
+	SLATE_BEGIN_ARGS(SGraphPinParam)
+		: _ModelPin(nullptr)
+		, _GraphNode(nullptr)
+	{}
+
+	SLATE_ARGUMENT(URigVMPin*, ModelPin)
+
+	SLATE_ARGUMENT(URigVMEdGraphNode*, GraphNode)
+
+	SLATE_ARGUMENT(FAnimNextParamType, FilterType)
+
+	SLATE_END_ARGS()
+
+	void Construct(const FArguments& InArgs, UEdGraphPin* InPin);
+
+private:
+	virtual TSharedRef<SWidget>	GetDefaultValueWidget() override;
+
+	void UpdateCachedParam();
+
+	void HandleGraphModified(ERigVMGraphNotifType InType, URigVMGraph* InGraph, UObject* InSubject);
+
+	URigVMPin* ModelPin = nullptr;
+
+	URigVMEdGraphNode* Node = nullptr;
+
+	FAnimNextParamType FilterType;
+
+	FAnimNextEditorParam CachedParam;
+
+	TSharedPtr<SParameterPickerCombo> PickerCombo;
+
+	UScriptStruct* ParamStruct = nullptr; 
+};
+
+}
\ No newline at end of file
diff --git a/Engine/Plugins/Experimental/AnimNext/Source/AnimNextEditor/Private/Param/SGraphPinParamName.cpp b/Engine/Plugins/Experimental/AnimNext/Source/AnimNextEditor/Private/Param/SGraphPinParamName.cpp
--- a/Engine/Plugins/Experimental/AnimNext/Source/AnimNextEditor/Private/Param/SGraphPinParamName.cpp
+++ b/Engine/Plugins/Experimental/AnimNext/Source/AnimNextEditor/Private/Param/SGraphPinParamName.cpp
@@ -13,6 +13,11 @@
 #include "RigVMModel/Nodes/RigVMTemplateNode.h"
 #include "RigVMModel/RigVMPin.h"
 #include "RigVMCore/RigVMRegistry.h"
+#include "RigVMModel/RigVMClient.h"
+#include "AnimNextRigVMAsset.h"
+#include "EdGraph/RigVMEdGraphNode.h"
+#include "Graph/AnimNextGraph_Controller.h"
+#include "UniversalObjectLocators/AssetLocatorFragment.h"
 
 #define LOCTEXT_NAMESPACE "SGraphPinParamName"
 
@@ -24,6 +29,12 @@ void SGraphPinParamName::Construct(const FArguments& InArgs, UEdGraphPin* InPin)
 	ModelPin = InArgs._ModelPin;
 	Node = InArgs._GraphNode;
 	FilterType = InArgs._FilterType;
+	InstanceId = TInstancedStruct<FAnimNextParamUniversalObjectLocator>::Make();
+
+	if(ModelPin && ModelPin->GetGraph())
+	{
+		ModelPin->GetGraph()->OnModified().AddSP(this, &SGraphPinParamName::HandleGraphModified);
+	}
 
 	SGraphPin::Construct(SGraphPin::FArguments(), InPin);
 }
@@ -31,12 +42,12 @@ void SGraphPinParamName::Construct(const FArguments& InArgs, UEdGraphPin* InPin)
 TSharedRef<SWidget> SGraphPinParamName::GetDefaultValueWidget()
 {
 	FParameterPickerArgs Args;
-	Args.bShowSourceGraph = false;
 	Args.bMultiSelect = false;
+	Args.bShowInstanceId = false;
 
-	CachedType = UncookedOnly::FUtils::GetParameterTypeFromName(FName(*GraphPinObj->DefaultValue));
+	UpdateCachedParamType();
 
-	// Check whether this is a Set/Get parameter node, and if so only show bound parameters
+	// Check whether this is a Set/Get parameter node, and if so only show parameters in the current asset
 	if (ModelPin)
 	{
 		if (const URigVMTemplateNode* TemplateNode = Cast<URigVMTemplateNode>(ModelPin->GetOuter()))
@@ -49,18 +60,29 @@ TSharedRef<SWidget> SGraphPinParamName::GetDefaultValueWidget()
 
 			if (TemplateNode->GetNotation() == GetLayerParameterNotation || TemplateNode->GetNotation() == SetLayerParameterNotation)
 			{
-				Args.bShowUnboundParameters = false;
+				AssetData = FAssetData(ModelPin->GetTypedOuter<UAnimNextRigVMAsset>());
+				InstanceId.GetMutable().Locator.Reset();
+				InstanceId.GetMutable().Locator.AddFragment<FAssetLocatorFragment>(AssetData);
 			}
 		}
 	}
-	
+
 	Args.OnParameterPicked = FOnParameterPicked::CreateLambda([this](const FParameterBindingReference& InParameterBinding)
 	{
 		FScopedTransaction Transaction(LOCTEXT("SelectParameter", "Select Parameter"));
-		GraphPinObj->Modify();
-		GraphPinObj->GetSchema()->TrySetDefaultValue(*GraphPinObj, InParameterBinding.Parameter.ToString());
 
-		CachedType = UncookedOnly::FUtils::GetParameterTypeFromName(FName(*GraphPinObj->DefaultValue));
+		if(ModelPin && Node)
+		{
+			UAnimNextGraph_Controller* Controller = CastChecked<UAnimNextGraph_Controller>(Node->GetController());
+			Controller->SetAnimNextParameterNode(Node->GetModelNode(), InParameterBinding.Parameter, InParameterBinding.Type, InstanceId);
+		}
+		else
+		{
+			GraphPinObj->Modify();
+			GraphPinObj->GetSchema()->TrySetDefaultValue(*GraphPinObj, InParameterBinding.Parameter.ToString());
+		}
+
+		UpdateCachedParamType();
 	});
 	
 	Args.OnFilterParameterType = FOnFilterParameterType::CreateLambda([this](const FAnimNextParamType& InParamType)-> EFilterParameterResult
@@ -77,20 +99,71 @@ TSharedRef<SWidget> SGraphPinParamName::GetDefaultValueWidget()
 
 		return EFilterParameterResult::Include;
 	});
+
+	Args.OnFilterParameter = FOnFilterParameter::CreateLambda([this](const FParameterBindingReference& InParameterBinding)
+	{
+		if(AssetData.IsValid())
+		{
+			return InParameterBinding.Graph == AssetData ? EFilterParameterResult::Include : EFilterParameterResult::Exclude;
+		}
+
+		return EFilterParameterResult::Include;
+	});
+	
 	Args.NewParameterType = FilterType;
 	
-	return SNew(SParameterPickerCombo)
+	return SAssignNew(PickerCombo, SParameterPickerCombo)
 		.PickerArgs(Args)
 		.OnGetParameterName_Lambda([this]()
 		{
-			return FName(*GraphPinObj->DefaultValue);
+			if(ModelPin && Node)
+			{
+				return FName(*ModelPin->GetDefaultValue());
+			}
+			else
+			{
+				return FName(*GraphPinObj->DefaultValue);
+			}
 		})
 		.OnGetParameterType_Lambda([this]()
 		{
 			return CachedType;
+		})
+		.OnGetParameterInstanceId_Lambda([this]()
+		{
+			return InstanceId;
 		});
 }
 
+void SGraphPinParamName::UpdateCachedParamType()
+{
+	if(ModelPin && Node)
+	{
+		CachedType = UncookedOnly::FUtils::GetParameterTypeFromName(FName(*ModelPin->GetDefaultValue()));
+	}
+	else
+	{
+		CachedType = UncookedOnly::FUtils::GetParameterTypeFromName(FName(*GraphPinObj->DefaultValue));
+	}
+
+	if(PickerCombo)
+	{
+		PickerCombo->RequestRefresh();
+	}
+}
+
+void SGraphPinParamName::HandleGraphModified(ERigVMGraphNotifType InType, URigVMGraph* InGraph, UObject* InSubject)
+{
+	switch (InType)
+	{
+	case ERigVMGraphNotifType::PinDefaultValueChanged:
+		UpdateCachedParamType();
+		break;
+	default:
+		break;
+	}
+}
+
 }
 
 #undef LOCTEXT_NAMESPACE
\ No newline at end of file
diff --git a/Engine/Plugins/Experimental/AnimNext/Source/AnimNextEditor/Private/Param/SGraphPinParamName.h b/Engine/Plugins/Experimental/AnimNext/Source/AnimNextEditor/Private/Param/SGraphPinParamName.h
--- a/Engine/Plugins/Experimental/AnimNext/Source/AnimNextEditor/Private/Param/SGraphPinParamName.h
+++ b/Engine/Plugins/Experimental/AnimNext/Source/AnimNextEditor/Private/Param/SGraphPinParamName.h
@@ -3,13 +3,18 @@
 #pragma once
 
 #include "SGraphPin.h"
+#include "UniversalObjectLocator.h"
+#include "Param/AnimNextParamUniversalObjectLocator.h"
 #include "Param/ParamType.h"
 
+class URigVMGraph;
+enum class ERigVMGraphNotifType : uint8;
 class URigVMPin;
 class URigVMEdGraphNode;
 
 namespace UE::AnimNext::Editor
 {
+class SParameterPickerCombo;
 
 // A pin widget that allows picking using an AnimNext parameter picker
 class SGraphPinParamName : public SGraphPin
@@ -32,6 +37,10 @@ class SGraphPinParamName : public SGraphPin
 private:
 	virtual TSharedRef<SWidget>	GetDefaultValueWidget() override;
 
+	void UpdateCachedParamType();
+
+	void HandleGraphModified(ERigVMGraphNotifType InType, URigVMGraph* InGraph, UObject* InSubject);
+	
 	URigVMPin* ModelPin = nullptr;
 
 	URigVMEdGraphNode* Node = nullptr;
@@ -39,6 +48,14 @@ private:
 	FAnimNextParamType FilterType;
 
 	FAnimNextParamType CachedType;
+
+	// The asset this pin is being edited within
+	FAssetData AssetData;
+
+	TSharedPtr<SParameterPickerCombo> PickerCombo;
+
+	// Instance Id corresponding to AssetData
+	TInstancedStruct<FAnimNextParamUniversalObjectLocator> InstanceId;
 };
 
 }
\ No newline at end of file
diff --git a/Engine/Plugins/Experimental/AnimNext/Source/AnimNextEditor/Private/Param/SLinkParametersDialog.cpp b/Engine/Plugins/Experimental/AnimNext/Source/AnimNextEditor/Private/Param/SLinkParametersDialog.cpp
--- a/Engine/Plugins/Experimental/AnimNext/Source/AnimNextEditor/Private/Param/SLinkParametersDialog.cpp
+++ b/Engine/Plugins/Experimental/AnimNext/Source/AnimNextEditor/Private/Param/SLinkParametersDialog.cpp
@@ -48,7 +48,6 @@ class SParameterToLink : public SCompoundWidget
 void SLinkParametersDialog::Construct(const FArguments& InArgs)
 {
 	FParameterPickerArgs ParameterPickerArgs;
-	ParameterPickerArgs.bShowUnboundParameters = false;
 	ParameterPickerArgs.OnGetParameterBindings = &OnGetParameterBindings;
 	ParameterPickerArgs.OnSelectionChanged = FSimpleDelegate::CreateSP(this, &SLinkParametersDialog::HandleSelectionChanged);
 	ParameterPickerArgs.OnFilterParameter = InArgs._OnFilterParameter;
