AnimNext: Parameters rework - UOL integration

diff --git a/Engine/Plugins/Experimental/AnimNext/Source/AnimNextEditor/Private/Param/SParameterPicker.cpp b/Engine/Plugins/Experimental/AnimNext/Source/AnimNextEditor/Private/Param/SParameterPicker.cpp
--- a/Engine/Plugins/Experimental/AnimNext/Source/AnimNextEditor/Private/Param/SParameterPicker.cpp
+++ b/Engine/Plugins/Experimental/AnimNext/Source/AnimNextEditor/Private/Param/SParameterPicker.cpp
@@ -8,6 +8,10 @@
 #include "Param/ParamType.h"
 #include "DetailLayoutBuilder.h"
 #include "EditorUtils.h"
+#include "IAnimNextModule.h"
+#include "IDetailTreeNode.h"
+#include "IPropertyRowGenerator.h"
+#include "IStructureDataProvider.h"
 #include "SAddParametersDialog.h"
 #include "Param/ParameterPickerArgs.h"
 #include "Widgets/Input/SSearchBox.h"
@@ -18,99 +22,102 @@
 #include "Framework/Application/SlateApplication.h"
 #include "ScopedTransaction.h"
 #include "SSimpleButton.h"
+#include "UniversalObjectLocatorEditor.h"
+#include "AnimNextUncookedOnly/Private/AnimNextUncookedOnlyModule.h"
+#include "Component/AnimNextComponent.h"
+#include "Editor/PropertyEditor/Private/SSingleProperty.h"
 #include "Graph/AnimNextGraph.h"
-#include "Param/ExternalParameterRegistry.h"
+#include "Param/AnimNextClassExtensionLibrary.h"
+#include "Param/AnimNextParam.h"
+#include "Param/IParameterSourceType.h"
+#include "Param/ParamUtils.h"
+#include "UObject/UObjectIterator.h"
+#include "Widgets/PropertyViewer/SFieldName.h"
+#include "Widgets/PropertyViewer/SPropertyViewer.h"
+#include "Param/AnimNextParamUniversalObjectLocator.h"
+#include "UniversalObjectLocators/AssetLocatorFragment.h"
 
 #define LOCTEXT_NAMESPACE "SParameterPicker"
 
 namespace UE::AnimNext::Editor
 {
 
-namespace ParameterPicker
+// IStructureDataProvider that allows us to display the instanced struct for instance IDs inline
+class FInstanceIdProvider : public IStructureDataProvider
 {
-static FName Column_Parameter(TEXT("Parameter"));
-static FName Column_Graph(TEXT("Graph"));
-static FName Column_Type(TEXT("Type"));
-static FName Column_New(TEXT("New"));
-}
-
-struct FParameterPickerEntry
-{
-	enum class EFilterResult : uint8
+public:
+	explicit FInstanceIdProvider(TInstancedStruct<FAnimNextParamInstanceIdentifier>& InInstanceId)
+		: InstanceId(InInstanceId)
 	{
-		DoesNotPassFilter	= 0x00,
-		PassesFilter		= 0x01,
-		ChildPassesFilter	= 0x02,
-	};
-
-	FRIEND_ENUM_CLASS_FLAGS(EFilterResult);
-
-	FParameterPickerEntry() = default;
+	}
 
-	FParameterPickerEntry(const FParameterBindingReference& InBinding)
-		: Binding(InBinding)
+	virtual bool IsValid() const override
 	{
-		PinType = UE::AnimNext::UncookedOnly::FUtils::GetPinTypeFromParamType(InBinding.Type);
-		PinIcon = FBlueprintEditorUtils::GetIconFromPin(PinType, /* bIsLarge = */true);
-		PinColor = GetDefault<UEdGraphSchema_K2>()->GetPinTypeColor(PinType);
-		FString ParameterString = Binding.Parameter.ToString();
-		int32 LastUnderscoreLoc = INDEX_NONE;
-		if (ParameterString.FindLastChar(TEXT('_'), LastUnderscoreLoc))
-		{
-			ParameterString.RightChopInline(LastUnderscoreLoc + 1);
-		}
-		DisplayString = MoveTemp(ParameterString);
+		return InstanceId.IsValid();
 	}
-
-	FParameterPickerEntry(FName InName)
+	
+	virtual const UStruct* GetBaseStructure() const override
 	{
-		Binding.Parameter = InName;
-		FString ParameterString = Binding.Parameter.ToString();
-		int32 LastUnderscoreLoc = INDEX_NONE;
-		if (ParameterString.FindLastChar(TEXT('_'), LastUnderscoreLoc))
-		{
-			ParameterString.RightChopInline(LastUnderscoreLoc + 1);
-		}
-		DisplayString = MoveTemp(ParameterString);
+		return InstanceId.GetScriptStruct();
 	}
 
-	bool PassesFilter(const FString& InFilterText) const
+	virtual void GetInstances(TArray<TSharedPtr<FStructOnScope>>& OutInstances, const UStruct* ExpectedBaseStructure) const override
 	{
-		return Binding.Parameter.ToString().Contains(InFilterText);
+		uint8* Memory = const_cast<uint8*>(InstanceId.GetMemory());
+		OutInstances.Add(MakeShared<FStructOnScope>(InstanceId.GetScriptStruct(), Memory));
 	}
 
-	FString DisplayString;
-
-	FParameterBindingReference Binding;
-
-	FEdGraphPinType PinType;
-
-	const FSlateBrush* PinIcon = nullptr;
-
-	FLinearColor PinColor = FLinearColor::White;
-
-	TSharedPtr<FParameterPickerEntry> Parent;
+	virtual bool IsPropertyIndirection() const override
+	{
+		return true;
+	}
 
-	TArray<TSharedRef<FParameterPickerEntry>> Children;
+	virtual uint8* GetValueBaseAddress(uint8* ParentValueAddress, const UStruct* ExpectedBaseStructure) const override
+	{
+		if (!ParentValueAddress)
+		{
+			return nullptr;
+		}
 
-	TArray<TSharedRef<FParameterPickerEntry>> FilteredChildren;
+		FInstancedStruct& InstancedStruct = *reinterpret_cast<FInstancedStruct*>(ParentValueAddress);
+		if (ExpectedBaseStructure && InstancedStruct.GetScriptStruct() && InstancedStruct.GetScriptStruct()->IsChildOf(ExpectedBaseStructure))
+		{
+			return InstancedStruct.GetMutableMemory();
+		}
+		
+		return nullptr;
+	}
 
-	EFilterResult FilterResult = EFilterResult::PassesFilter;
+	TInstancedStruct<FAnimNextParamInstanceIdentifier>& InstanceId;
 };
 
-ENUM_CLASS_FLAGS(FParameterPickerEntry::EFilterResult);
-
 void SParameterPicker::Construct(const FArguments& InArgs)
 {
-	using namespace ParameterPicker;
-
 	Args = InArgs._Args;
+	if(Args.Context == nullptr)
+	{
+		// TODO: This needs to defer to project/schedule/workspace defaults similar to FAnimNextLocatorContext
+		Args.Context = UAnimNextComponent::StaticClass()->GetDefaultObject();
+	}
+
+	SelectedInstanceId = Args.InstanceId;
+	if(!SelectedInstanceId.IsValid())
+	{
+		SelectedInstanceId = TInstancedStruct<FAnimNextParamUniversalObjectLocator>::Make();
+	}
+
+	FieldIterator = MakeUnique<FFieldIterator>(Args.OnFilterParameterType);
 
 	if(Args.OnGetParameterBindings != nullptr)
 	{
 		Args.OnGetParameterBindings->BindSP(this, &SParameterPicker::HandleGetParameterBindings);
 	}
 
+	if(Args.OnSetInstanceId != nullptr)
+	{
+		Args.OnSetInstanceId->BindSP(this, &SParameterPicker::HandleSetInstanceId);
+	}
+
 	if(Args.bFocusSearchWidget)
 	{
 		RegisterActiveTimer( 0.f, FWidgetActiveTimerDelegate::CreateLambda([this](double InCurrentTime, float InDeltaTime)
@@ -119,8 +126,11 @@ void SParameterPicker::Construct(const FArguments& InArgs)
 			{
 				FWidgetPath WidgetToFocusPath;
 				FSlateApplication::Get().GeneratePathToWidgetUnchecked(SearchBox.ToSharedRef(), WidgetToFocusPath);
-				FSlateApplication::Get().SetKeyboardFocus(WidgetToFocusPath, EFocusCause::SetDirectly);
-				WidgetToFocusPath.GetWindow()->SetWidgetToFocusOnActivate(SearchBox);
+				if(WidgetToFocusPath.IsValid())
+				{
+					FSlateApplication::Get().SetKeyboardFocus(WidgetToFocusPath, EFocusCause::SetDirectly);
+					WidgetToFocusPath.GetWindow()->SetWidgetToFocusOnActivate(SearchBox);
+				}
 				return EActiveTimerReturnType::Stop;
 			}
 
@@ -128,542 +138,362 @@ void SParameterPicker::Construct(const FArguments& InArgs)
 		}));
 	}
 
-	TSharedPtr<SHeaderRow> HeaderRow;
-	TSharedRef<SWidget> AddNewParameterWidget = SNullWidget::NullWidget;
-
-	if(Args.bAllowNew)
-	{
-		SAssignNew(AddNewParameterWidget, SSimpleButton)
-		.Text(LOCTEXT("AddNewParameterButton", "New Parameter"))
-		.ToolTipText(LOCTEXT("AddColumnHeaderTooltip", "Add a new parameter at global scope"))
-		.Icon(FAppStyle::Get().GetBrush("Icons.Plus"))
-		.OnClicked_Lambda([this]()
-		{
-			FSlateApplication::Get().DismissAllMenus();
-			TSharedRef<SAddParametersDialog> AddParametersDialog =
-				SNew(SAddParametersDialog)
-				.AllowMultiple(false)
-				.OnFilterParameterType(Args.OnFilterParameterType)
-				.InitialParamType(Args.NewParameterType);
-			TArray<FParameterToAdd> ParametersToAdd;
-			if(AddParametersDialog->ShowModal(ParametersToAdd))
-			{
-				if(ParametersToAdd.Num() > 0)
-				{
-					FScopedTransaction Transaction(LOCTEXT("AddParameter", "Add parameter"));
-					for (const FParameterToAdd& ParameterToAdd : ParametersToAdd)
-					{
-						FParameterBindingReference Reference;
-						Reference.Parameter = ParametersToAdd[0].Name;
-						Args.OnAddParameter.ExecuteIfBound(ParametersToAdd[0]);
-						Args.OnParameterPicked.ExecuteIfBound(Reference);
-					}
-				}
-			}
-			return FReply::Handled();
-		});
-	}
+	InstanceIdProvider = MakeShared<FInstanceIdProvider>(SelectedInstanceId);
 
+	FPropertyEditorModule& PropertyEditorModule = FModuleManager::Get().LoadModuleChecked<FPropertyEditorModule>("PropertyEditor");
+	FPropertyRowGeneratorArgs PropertyRowGeneratorArgs;
+	PropertyRowGeneratorArgs.NotifyHook = this;
+	PropertyRowGenerator = PropertyEditorModule.CreatePropertyRowGenerator(PropertyRowGeneratorArgs);
+	PropertyRowGenerator->SetStructure(InstanceIdProvider);
 
+	check(PropertyRowGenerator->GetRootTreeNodes().Num() > 0);
+	TSharedPtr<IDetailTreeNode> RootNode = PropertyRowGenerator->GetRootTreeNodes()[0];
+	TArray<TSharedRef<IDetailTreeNode>> Children;
+	RootNode->GetChildren(Children);
+	TSharedPtr<IDetailTreeNode> ChildNode = Children[0];
+	const FNodeWidgets NodeWidgets = ChildNode->CreateNodeWidgets();
+	
 	ChildSlot
 	[
-		SNew(SVerticalBox)
-		+SVerticalBox::Slot()
-		.AutoHeight()
-		.Padding(2.0f)
+		SNew(SBox)
+		.WidthOverride(400.0f)
+		.HeightOverride(400.0f)
 		[
-			SNew(SHorizontalBox)
-			+ SHorizontalBox::Slot()
-			.MaxWidth(250.0f)
+			SNew(SVerticalBox)
+			+SVerticalBox::Slot()
+			.AutoHeight()
+			.Padding(Args.bShowInstanceId ? FMargin(0.0f, 2.0f) : FMargin(0.0f))
 			[
-				SAssignNew(SearchBox, SSearchBox)
-				.OnTextChanged_Lambda([this](FText InText)
-				{
-					FilterText = InText;
-					RefreshFilter();
-				})
+				SNew(SHorizontalBox)
+				.Visibility(Args.bShowInstanceId ? EVisibility::Visible : EVisibility::Collapsed)
+				// TODO: slot for picker to switch instance ID types here
+				+SHorizontalBox::Slot()
+				.FillWidth(1.0f)
+				[
+					NodeWidgets.ValueWidget.ToSharedRef()
+				]
 			]
-			+ SHorizontalBox::Slot()
-			.HAlign(HAlign_Center)
-			.AutoWidth()
+			+SVerticalBox::Slot()
+			.FillHeight(1.0f)
+			.Padding(2.0f)
 			[
-				AddNewParameterWidget
+				SAssignNew(PropertyViewer, UE::PropertyViewer::SPropertyViewer)
+				.OnSelectionChanged(this, &SParameterPicker::HandleFieldPicked)
+				.OnGenerateContainer(this, &SParameterPicker::HandleGenerateContainer)
+				.FieldIterator(FieldIterator.Get())
+				.FieldExpander(&FieldExpander)
+				.bShowSearchBox(true)
 			]
 		]
-		+SVerticalBox::Slot()
-		.FillHeight(1.0f)
-		.Padding(2.0f)
-		[
-			SAssignNew(EntriesList, STreeView<TSharedRef<FParameterPickerEntry>>)
-			.TreeItemsSource(&FilteredHierarchy)
-			.SelectionMode(Args.bMultiSelect ? ESelectionMode::Multi : ESelectionMode::Single)
-			.OnGenerateRow(this, &SParameterPicker::HandleGenerateRow)
-			.OnSelectionChanged(this, &SParameterPicker::HandleSelectionChanged)
-			.OnGetChildren(this, &SParameterPicker::HandleGetChildren)
-			.OnIsSelectableOrNavigable(this, &SParameterPicker::HandleIsSelectableOrNavigable)
-			.ItemHeight(20.0f)
-			.HeaderRow(
-				SAssignNew(HeaderRow, SHeaderRow)
-				+SHeaderRow::Column(Column_Type)
-				.DefaultLabel(FText::GetEmpty())
-				.FixedWidth(24.0f)
-				.HeaderContent()
-				[
-					SNew(SBox)
-					.WidthOverride(16.0f)
-					.HeightOverride(16.0f)
-					.VAlign(VAlign_Center)
-					.HAlign(HAlign_Center)
-					[
-						SNew(SImage)
-						.ColorAndOpacity(FSlateColor::UseForeground())
-						.Image(FAppStyle::GetBrush("Kismet.VariableList.TypeIcon"))
-						.ToolTipText(LOCTEXT("TypeColumnHeaderTooltip", "The parameter's type"))
-					]
-				]
-
-				+SHeaderRow::Column(Column_Parameter)
-				.DefaultLabel(LOCTEXT("ParameterColumnHeader", "Parameter"))
-				.ToolTipText(LOCTEXT("ParameterColumnHeaderTooltip", "The parameter's name"))
-				.FillWidth(0.33f)
-			)
-		]
 	];
 
-	
-	if (Args.bShowSourceGraph)
-	{
-		HeaderRow->AddColumn(
-			SHeaderRow::Column(Column_Graph)
-			.DefaultLabel(LOCTEXT("GraphColumnHeader", "Graph"))
-			.ToolTipText(LOCTEXT("GraphColumnHeaderTooltip", "The graph that has a binding to the parameter"))
-			.FillWidth(0.33f));
-	}
-
-	if(Args.bAllowNew)
-	{
-		HeaderRow->AddColumn(
-			SHeaderRow::Column(Column_New)
-			.DefaultLabel(FText::GetEmpty())
-			.HeaderContentPadding(FMargin(0.0f))
-			.FixedWidth(24.0f)
-			.HeaderContent()
-			[
-				SNew(SButton)
-				.ToolTipText(LOCTEXT("AddColumnHeaderTooltip", "Add a new parameter at global scope"))
-				.ButtonStyle(FAppStyle::Get(), "HoverHintOnly")
-				.OnClicked_Lambda([this]()
-				{
-					FSlateApplication::Get().DismissAllMenus();
-					TSharedRef<SAddParametersDialog> AddParametersDialog =
-						SNew(SAddParametersDialog)
-						.AllowMultiple(false);
-					TArray<FParameterToAdd> ParametersToAdd;
-					if(AddParametersDialog->ShowModal(ParametersToAdd))
-					{
-						if(ParametersToAdd.Num() > 0)
-						{
-							FScopedTransaction Transaction(LOCTEXT("AddParameter", "Add parameter"));
-
-							FParameterBindingReference Reference(ParametersToAdd[0].Name, ParametersToAdd[0].Type);
-							Args.OnParameterPicked.ExecuteIfBound(Reference);
-						}
-					}
-					return FReply::Handled();
-				})
-				[
-					SNew(SBox)
-					.WidthOverride(16.0f)
-					.HeightOverride(16.0f)
-					.VAlign(VAlign_Center)
-					.HAlign(HAlign_Center)
-					[
-						SNew(SImage)
-						.ColorAndOpacity(FSlateColor::UseForeground())
-						.Image(FAppStyle::GetBrush("Icons.Plus"))
-					]
-				]
-			]);
-	}
-	
 	RefreshEntries();
 }
 
 void SParameterPicker::RefreshEntries()
 {
-	Entries.Empty();
-
-	IAssetRegistry& AssetRegistry = FModuleManager::Get().LoadModuleChecked<FAssetRegistryModule>(TEXT("AssetRegistry")).Get();
+	using namespace UE::UniversalObjectLocator;
 
-	FARFilter ARFilter;
+	PropertyViewer->RemoveAll();
+	CachedContainers.Reset();
+	ContainerMap.Reset();
 
-	TSet<TTuple<FName, FAssetData>> BoundParameters;
-
-	if(Args.bAllowNone)
+	UncookedOnly::IAnimNextUncookedOnlyModule& Module = FModuleManager::GetModuleChecked<UncookedOnly::IAnimNextUncookedOnlyModule>("AnimNextUncookedOnly");
+	const UStruct* Struct = nullptr;
+	if(SelectedInstanceId.IsValid() && SelectedInstanceId.Get().IsValid())
 	{
-		Entries.Add(MakeShared<FParameterPickerEntry>(FParameterBindingReference(NAME_None, FAnimNextParamType())));
+		TSharedPtr<UncookedOnly::IParameterSourceType> SourceType = Module.FindParameterSourceType(SelectedInstanceId.GetScriptStruct());
+		if(SourceType.IsValid())
+		{
+			Struct = SourceType->GetStruct(SelectedInstanceId);
+		}
 	}
-	
-	// Find all graphs and their bound parameters
-	if(Args.bShowBoundParameters)
-	{
-		ARFilter.ClassPaths = { UAnimNextGraph::StaticClass()->GetClassPathName() };
-		
-		TArray<FAssetData> GraphAssets;
-		AssetRegistry.GetAssets(ARFilter, GraphAssets);
 
-		for(const FAssetData& GraphAsset : GraphAssets)
+	if(Struct == nullptr || Struct == UAnimNextGraph::StaticClass())
+	{
+		// For AnimNext graphs, we add the structs that are exposed via parameters
+		TMap<FAssetData, FAnimNextParameterProviderAssetRegistryExports> Exports;
+		if(UncookedOnly::FUtils::GetExportedParametersFromAssetRegistry(Exports))
 		{
-			FAnimNextParameterProviderAssetRegistryExports Exports;
-			if(UncookedOnly::FUtils::GetExportedParametersForAsset(GraphAsset, Exports))
+			for(const TPair<FAssetData, FAnimNextParameterProviderAssetRegistryExports>& ExportPair : Exports)
 			{
-				for(const FAnimNextParameterAssetRegistryExportEntry& Export : Exports.Parameters)
+				if(ExportPair.Value.Parameters.Num() > 0)
 				{
-					BoundParameters.Add({ Export.Name, GraphAsset });
-					FParameterBindingReference NewReference(Export.Name, Export.Type, GraphAsset);
-					if(!Args.OnFilterParameter.IsBound() || Args.OnFilterParameter.Execute(NewReference) == EFilterParameterResult::Include)
+					// Add a placeholder struct for this asset's properties
+					TArray<FPropertyBagPropertyDesc> PropertyDescs;
+					PropertyDescs.Reserve(ExportPair.Value.Parameters.Num());
+					for(const FAnimNextParameterAssetRegistryExportEntry& ParameterEntry : ExportPair.Value.Parameters)
 					{
-						FAnimNextParamType ParamType = UE::AnimNext::UncookedOnly::FUtils::GetParameterTypeFromName(Export.Name);
-						if(!Args.OnFilterParameterType.IsBound() || Args.OnFilterParameterType.Execute(ParamType) == EFilterParameterResult::Include)
+						if(EnumHasAnyFlags(ParameterEntry.GetFlags(), EAnimNextParameterFlags::Bound) && ParameterEntry.Name != NAME_None)
 						{
-							if (Args.bShowBoundParameters && EnumHasAnyFlags(Export.Flags, EAnimNextParameterFlags::Bound))
+							FParameterBindingReference ParameterBinding;
+							ParameterBinding.Type = ParameterEntry.Type;
+							ParameterBinding.Parameter = UncookedOnly::FUtils::GetParameterNameFromQualifiedName(ParameterEntry.Name);
+							ParameterBinding.Graph = ExportPair.Key;
+
+							if (!Args.OnFilterParameter.IsBound() || Args.OnFilterParameter.Execute(ParameterBinding) == EFilterParameterResult::Include)
 							{
-								TSharedRef<FParameterPickerEntry> NewEntry = MakeShared<FParameterPickerEntry>(NewReference);
-								Entries.Add(NewEntry);
+								if (!Args.OnFilterParameterType.IsBound() || Args.OnFilterParameterType.Execute(ParameterEntry.Type) == EFilterParameterResult::Include)
+								{
+									PropertyDescs.Emplace(ParameterBinding.Parameter, ParameterEntry.Type.GetContainerType(), ParameterEntry.Type.GetValueType(), ParameterEntry.Type.GetValueTypeObject());
+								}
 							}
 						}
 					}
+
+					if(PropertyDescs.Num() > 0)
+					{
+						const FText DisplayName = FText::FromName(ExportPair.Key.AssetName);
+						const FText TooltipText = FText::FromString(ExportPair.Key.GetObjectPathString());
+						FContainerInfo& ContainerInfo = CachedContainers.Emplace_GetRef(DisplayName, TooltipText, ExportPair.Key, MakeUnique<FInstancedPropertyBag>());
+						ContainerInfo.PropertyBag->AddProperties(PropertyDescs);
+						UE::PropertyViewer::SPropertyViewer::FHandle Handle = PropertyViewer->AddContainer(ContainerInfo.PropertyBag.Get()->GetPropertyBagStruct(), DisplayName);
+						ContainerMap.Add(Handle, CachedContainers.Num() - 1);
+					}
 				}
-			};
+			}
 		}
 	}
-
-	// Find all parameters (that have not already been added as bound above)
-	if(Args.bShowUnboundParameters)
+	else if(Struct != nullptr)
 	{
-		FAnimNextParameterProviderAssetRegistryExports AllExports;
-		if(UE::AnimNext::UncookedOnly::FUtils::GetExportedParametersFromAssetRegistry(AllExports))
+		CachedContainers.Emplace_GetRef(Struct->GetDisplayNameText(), Struct->GetToolTipText(), Struct);
+		if(const UScriptStruct* ScriptStruct = Cast<UScriptStruct>(Struct))
 		{
-			for(const FAnimNextParameterAssetRegistryExportEntry& ExportEntry : AllExports.Parameters)
+			UE::PropertyViewer::SPropertyViewer::FHandle Handle = PropertyViewer->AddContainer(ScriptStruct);
+			ContainerMap.Add(Handle, CachedContainers.Num() - 1);
+		}
+		else if(const UClass* Class = Cast<UClass>(Struct))
+		{
+			UE::PropertyViewer::SPropertyViewer::FHandle Handle = PropertyViewer->AddContainer(Class);
+			ContainerMap.Add(Handle, CachedContainers.Num() - 1);
+
+			// Find any UAnimNextClassExtensionLibrary classes to extend this class
+			for(TObjectIterator<UClass> It; It; ++It)
 			{
-				if(!BoundParameters.Contains( { ExportEntry.Name, ExportEntry.ReferencingAsset } ))
+				UClass* LibraryClass = *It;
+				if(LibraryClass->HasAnyClassFlags(CLASS_Abstract) || !LibraryClass->HasAnyClassFlags(CLASS_Native) || !LibraryClass->IsChildOf(UAnimNextClassExtensionLibrary::StaticClass()))
 				{
-					FParameterBindingReference NewReference(ExportEntry.Name, ExportEntry.Type, ExportEntry.ReferencingAsset);
+					continue;
+				}
 
-					if(!Args.OnFilterParameter.IsBound() || Args.OnFilterParameter.Execute(NewReference) == EFilterParameterResult::Include)
-					{
-						if (!Args.OnFilterParameterType.IsBound() || Args.OnFilterParameterType.Execute(ExportEntry.Type) == EFilterParameterResult::Include)
-						{
-							TSharedRef<FParameterPickerEntry> NewEntry = MakeShared<FParameterPickerEntry>(NewReference);
-							Entries.Add(NewEntry);
-						}
-					}
+				UClass* ExtendedClass = LibraryClass->GetDefaultObject<UAnimNextClassExtensionLibrary>()->GetSupportedClass();
+				if(ExtendedClass == nullptr)
+				{
+					continue;
 				}
+
+				if(!Class->IsChildOf(ExtendedClass))
+				{
+					continue;
+				}
+
+				CachedContainers.Emplace_GetRef(LibraryClass->GetDisplayNameText(), LibraryClass->GetToolTipText(), LibraryClass);
+				Handle = PropertyViewer->AddContainer(LibraryClass);
+				ContainerMap.Add(Handle, CachedContainers.Num() - 1);
 			}
 		}
 	}
+}
 
-	if (Args.bShowBuiltInParameters)
+bool SParameterPicker::GetFieldInfo(UE::PropertyViewer::SPropertyViewer::FHandle InHandle, const FFieldVariant& InField, FName& OutName, TInstancedStruct<FAnimNextParamInstanceIdentifier>& OutInstanceId, FAnimNextParamType& OutType) const
+{
+	if(const int32* ContainerIndexPtr = ContainerMap.Find(InHandle))
 	{
-		FExternalParameterRegistry::ForEachParameter([this](FName InParameterName, const IParameterSourceFactory::FParameterInfo& InInfo)
+		const FContainerInfo& ContainerInfo = CachedContainers[*ContainerIndexPtr];
+		if(UFunction* Function = InField.Get<UFunction>())
+		{
+			check(Function->GetReturnProperty());
+			OutType = FParamTypeHandle::FromProperty(Function->GetReturnProperty()).GetType();
+			OutName = *Function->GetPathName();
+			OutInstanceId = SelectedInstanceId;
+		}
+		else if(FProperty* Property = InField.Get<FProperty>())
 		{
-			FParameterBindingReference NewReference(InParameterName, InInfo.Type);
-			if (!Args.OnFilterParameter.IsBound() || Args.OnFilterParameter.Execute(NewReference) == EFilterParameterResult::Include)
+			OutType = FParamTypeHandle::FromProperty(Property).GetType();
+			if(ContainerInfo.PropertyBag.IsValid())
 			{
-				if (!Args.OnFilterParameterType.IsBound() || Args.OnFilterParameterType.Execute(InInfo.Type) == EFilterParameterResult::Include)
-				{
-					Entries.Add(MakeShared<FParameterPickerEntry>(NewReference));
-				}
+				// Properties from property bags are assumed to use the asset that they come from
+				ensure(ContainerInfo.AssetData.IsValid());
+				TStringBuilder<256> StringBuilder;
+				ContainerInfo.AssetData.AppendObjectPath(StringBuilder);
+				StringBuilder.Append(TEXT(":"));
+				Property->GetFName().AppendString(StringBuilder);
+				OutName = FName(StringBuilder.ToView());
+				OutInstanceId = TInstancedStruct<FAnimNextParamUniversalObjectLocator>::Make();
+				OutInstanceId.GetMutable<FAnimNextParamUniversalObjectLocator>().Locator.Reset();
+				OutInstanceId.GetMutable<FAnimNextParamUniversalObjectLocator>().Locator.AddFragment<FAssetLocatorFragment>(ContainerInfo.AssetData);
 			}
-		});
-	}
+			else
+			{
+				OutName = *Property->GetPathName();
+				OutInstanceId = SelectedInstanceId;
+			}
+		}
 
-	BuildHierarchy();
+		return true;
+	}
 
-	RefreshFilter();
+	return false;
 }
 
-void SParameterPicker::BuildHierarchy()
+void SParameterPicker::HandleGetParameterBindings(TArray<FParameterBindingReference>& OutParameterBindings) const
 {
-	Hierarchy.Reset();
+	TArray<UE::PropertyViewer::SPropertyViewer::FSelectedItem> SelectedItems = PropertyViewer->GetSelectedItems();
 
-	TMap<FName, TSharedRef<FParameterPickerEntry>> NameMap;
-	NameMap.Reserve(Entries.Num());
-
-	for (const TSharedRef<FParameterPickerEntry>& Entry : Entries)
+	for(UE::PropertyViewer::SPropertyViewer::FSelectedItem& SelectedItem : SelectedItems)
 	{
-		NameMap.Add(Entry->Binding.Parameter, Entry);
-	}
-
-	TArray< TSharedRef<FParameterPickerEntry>> NewEntries;
-	for (const TSharedRef<FParameterPickerEntry>& Entry : Entries)
-	{
-		// Parse name into separators
-		TStringBuilder<256> WholeParameterString;
-		Entry->Binding.Parameter.ToString(WholeParameterString);
-
-		TStringBuilder<256> PartialParameterString;
-
-		// We use '_' as a separator here as:
-		// - Each param is a UObject in editor and uses its object name, so we cannot use '.'
-		// - This maps nicely to Verse tags that use '_' to hierarchically define their relationship
-		TSharedPtr<FParameterPickerEntry> Parent = nullptr;
-		UE::String::ParseTokens(WholeParameterString, TEXT('_'), [this, &PartialParameterString, &NameMap, &Parent, &NewEntries](FStringView InStringView)
+		if(SelectedItem.Fields.Num() > 0 && SelectedItem.Fields.Last().Num() > 0)
 		{
-			if(Parent.IsValid())
+			if(const int32* ContainerIndexPtr = ContainerMap.Find(SelectedItem.Handle))
 			{
-				// Have a parent, so add child if it doesnt exist already
-				PartialParameterString += TEXT('_');
-				PartialParameterString += InStringView;
-
-				const FName ParameterName(InStringView);
-				const FName PartialParameterName(PartialParameterString);
-				if (TSharedRef<FParameterPickerEntry>* ExistingEntry = NameMap.Find(PartialParameterName))
+				const FContainerInfo& ContainerInfo = CachedContainers[*ContainerIndexPtr];
+				FAnimNextParamType Type;
+				FName Name;
+				TInstancedStruct<FAnimNextParamInstanceIdentifier> InstanceId;
+				if(GetFieldInfo(SelectedItem.Handle, SelectedItem.Fields.Last().Last(), Name, InstanceId, Type))
 				{
-					(*ExistingEntry)->Parent = Parent;
-					if (!Parent->Children.ContainsByPredicate([&ExistingEntry](const TSharedRef<FParameterPickerEntry>& InEntry) { return (*ExistingEntry)->Binding.Parameter == InEntry->Binding.Parameter; }))
+					if(ensure(Type.IsValid() && !Name.IsNone()))
 					{
-						Parent->Children.Add(*ExistingEntry);
+						FParameterBindingReference NewReference;
+						NewReference.Parameter = Name;
+						NewReference.InstanceId = InstanceId;
+						NewReference.Type = Type;
+						NewReference.Graph = ContainerInfo.AssetData;
+						OutParameterBindings.Emplace(NewReference);
 					}
-					Parent = *ExistingEntry;
-				}
-				else
-				{
-					TSharedRef<FParameterPickerEntry> NewEntry = MakeShared<FParameterPickerEntry>(PartialParameterName);
-					NewEntry->Parent = Parent;
-					if (!Parent->Children.ContainsByPredicate([&NewEntry](const TSharedRef<FParameterPickerEntry>& InEntry) { return NewEntry->Binding.Parameter == InEntry->Binding.Parameter; }))
-					{
-						Parent->Children.Add(NewEntry);
-					}
-					Parent = NameMap.Add(PartialParameterName, NewEntry);
-					NewEntries.Add(NewEntry);
 				}
 			}
-			else
-			{
-				PartialParameterString += InStringView;
-
-				// Add root item if it doesnt exist already
-				const FName ParameterName(InStringView);
-				if (TSharedRef<FParameterPickerEntry>* ExistingParent = NameMap.Find(ParameterName))
-				{
-					Parent = *ExistingParent;
-				}
-				else
-				{
-					TSharedRef<FParameterPickerEntry> NewEntry = MakeShared<FParameterPickerEntry>(ParameterName);
-					Parent = NameMap.Add(ParameterName, NewEntry);
-					NewEntries.Add(NewEntry);
-				}
-
-				if (!Hierarchy.ContainsByPredicate([&Parent](const TSharedRef<FParameterPickerEntry>& InEntry){ return Parent->Binding.Parameter == InEntry->Binding.Parameter; }))
-				{
-					Hierarchy.Add(Parent.ToSharedRef());
-				}
-			}
-		}, UE::String::EParseTokensOptions::SkipEmpty);
+		}
 	}
-
-	Entries.Append(NewEntries);
 }
 
-void SParameterPicker::RefreshFilter()
+void SParameterPicker::HandleSetInstanceId(const TInstancedStruct<FAnimNextParamInstanceIdentifier>& InInstanceId)
 {
-	FilteredEntries.Reset();
-	FilteredHierarchy.Reset();
-	const FString FilterTextAsString = FilterText.ToString();
-
-	for (const TSharedRef<FParameterPickerEntry>& Entry : Entries)
-	{
-		Entry->FilterResult = FParameterPickerEntry::EFilterResult::DoesNotPassFilter;
-	}
+	SelectedInstanceId = InInstanceId;
+	RefreshEntries();
+}
 
-	for(const TSharedRef<FParameterPickerEntry>& Entry : Entries)
+void SParameterPicker::HandleFieldPicked(UE::PropertyViewer::SPropertyViewer::FHandle InHandle, TArrayView<const FFieldVariant> InFields, ESelectInfo::Type InSelectionType)
+{
+	if(InFields.Num() == 1)
 	{
-		Entry->FilteredChildren.Reset();
-		if(Entry->PassesFilter(FilterTextAsString))
+		FAnimNextParamType Type;
+		FName Name;
+		TInstancedStruct<FAnimNextParamInstanceIdentifier> InstanceId;
+		if(GetFieldInfo(InHandle, InFields.Last(), Name, InstanceId, Type))
 		{
-			Entry->FilterResult |= FParameterPickerEntry::EFilterResult::PassesFilter;
-			FilteredEntries.Add(Entry);
-
-			if(FilterTextAsString.Len() > 0)
+			if(ensure(Type.IsValid() && !Name.IsNone()))
 			{
-				TSharedPtr< FParameterPickerEntry> ParentEntry = Entry->Parent;
-				while (ParentEntry.IsValid())
+				UncookedOnly::IAnimNextUncookedOnlyModule& Module = FModuleManager::GetModuleChecked<UncookedOnly::IAnimNextUncookedOnlyModule>("AnimNextUncookedOnly");
+				if(InstanceId.IsValid() && InstanceId.Get().IsValid())
 				{
-					ParentEntry->FilterResult |= FParameterPickerEntry::EFilterResult::ChildPassesFilter;
-					ParentEntry = ParentEntry->Parent;
+					TSharedPtr<UncookedOnly::IParameterSourceType> SourceType = Module.FindParameterSourceType(InstanceId.GetScriptStruct());
+					if(SourceType.IsValid())
+					{
+						const UStruct* Struct = SourceType->GetStruct(InstanceId);
+						if(Struct && Struct == UAnimNextGraph::StaticClass())
+						{
+							// Invalidate the instance ID if this is an AnimNext graph, as they dont have instances
+							InstanceId.Reset();
+						}
+					}
 				}
+
+				FParameterBindingReference Reference(Name, Type, InstanceId);
+				Args.OnParameterPicked.ExecuteIfBound(Reference);
 			}
 		}
 	}
+}
 
-	TArray<TSharedRef<FParameterPickerEntry>> Stack;
-	Stack.Reserve(16);
-	for (const TSharedRef<FParameterPickerEntry>& Entry : Hierarchy)
-	{
-		Stack.Add(Entry);
-	}
-
-	while (Stack.Num() > 0)
+TSharedRef<SWidget> SParameterPicker::HandleGenerateContainer(UE::PropertyViewer::SPropertyViewer::FHandle InHandle, TOptional<FText> InDisplayName)
+{
+	if(int32* ContainerIndexPtr = ContainerMap.Find(InHandle))
 	{
-		TSharedRef<FParameterPickerEntry> Top = Stack.Top();
-		Stack.Pop(EAllowShrinking::No);
-
-		if (Top->FilterResult != FParameterPickerEntry::EFilterResult::DoesNotPassFilter)
+		if(CachedContainers.IsValidIndex(*ContainerIndexPtr))
 		{
-			if (!Top->Parent.IsValid())
-			{
-				FilteredHierarchy.Add(Top);
-			}
-			else
-			{
-				Top->Parent->FilteredChildren.Add(Top);
-			}
-
-			for (const TSharedRef<FParameterPickerEntry>& ChildEntry : Top->Children)
-			{
-				Stack.Add(ChildEntry);
-			}
+			FContainerInfo& ContainerInfo = CachedContainers[*ContainerIndexPtr];
 
-			if (EnumHasAnyFlags(Top->FilterResult, FParameterPickerEntry::EFilterResult::ChildPassesFilter))
-			{
-				EntriesList->SetItemExpansion(Top, true);
-			}
+			return SNew(SHorizontalBox)
+			+SHorizontalBox::Slot()
+			.AutoWidth()
+			.HAlign(HAlign_Right)
+			.VAlign(VAlign_Center)
+			[
+				SNew(SImage)
+				.Image(FAppStyle::GetBrush("ClassIcon.Object"))
+			]
+			+SHorizontalBox::Slot()
+			.Padding(4.0f)
+			[
+				SNew(STextBlock)
+				.Text(ContainerInfo.DisplayName)
+				.ToolTipText(ContainerInfo.TooltipText)
+			];
 		}
 	}
 
-	EntriesList->RequestTreeRefresh();
+	return SNullWidget::NullWidget;
 }
 
-class SParameterPickerRow : public SMultiColumnTableRow<TSharedRef<FParameterPickerEntry>>
+void SParameterPicker::NotifyPostChange(const FPropertyChangedEvent& PropertyChangedEvent, FEditPropertyChain* PropertyThatChanged)
 {
-	SLATE_BEGIN_ARGS(SParameterPickerRow) {}
-	SLATE_END_ARGS()
-
-	void Construct(const FArguments& InArgs, const TSharedRef<STableViewBase>& InOwnerTableView, TSharedRef<FParameterPickerEntry> InEntry, TSharedRef<const SParameterPicker> InParameterPicker)
-	{
-		Entry = InEntry;
-		ParameterPicker = InParameterPicker;
-		SMultiColumnTableRow<TSharedRef<FParameterPickerEntry>>::Construct( SMultiColumnTableRow<TSharedRef<FParameterPickerEntry>>::FArguments(), InOwnerTableView);
-	}
+	Args.OnInstanceIdChanged.ExecuteIfBound(SelectedInstanceId);
+	RefreshEntries();
+}
 
-	virtual TSharedRef<SWidget> GenerateWidgetForColumn(const FName& InColumnName) override
+TArray<FFieldVariant> SParameterPicker::FFieldIterator::GetFields(const UStruct* InStruct) const
+{
+	auto PassesFilterChecks = [this](const FProperty* InProperty)
 	{
-		using namespace ParameterPicker;
-
-		if (InColumnName == Column_Type)
+		if(InProperty && OnFilterParameterType.IsBound())
 		{
-			return
-				SNew(SHorizontalBox)
-				.Visibility(Entry->Binding.Type.IsValid() ? EVisibility::Visible : EVisibility::Hidden)
-				+ SHorizontalBox::Slot()
-				.HAlign(HAlign_Center)
-				.VAlign(VAlign_Center)
-				.Padding(2.0f)
-				.AutoWidth()
-				[
-					SNew(SImage)
-					.Image(Entry->PinIcon)
-					.ColorAndOpacity(Entry->PinColor)
-					.ToolTipText_Lambda([this]() -> FText
-					{
-						return FText::FromString(Entry->Binding.Type.ToString());
-					})
-				];
-		}
-		else if(InColumnName == Column_Parameter)
-		{
-			return
-				SNew(SHorizontalBox)
-				+SHorizontalBox::Slot()
-				.HAlign(HAlign_Center)
-				.VAlign(VAlign_Center)
-				.Padding(2.0f)
-				.AutoWidth()
-				[
-					SNew(SExpanderArrow, SharedThis(this))
-					.IndentAmount(8)
-				]
-				+SHorizontalBox::Slot()
-				.VAlign(VAlign_Center)
-				.Padding(2.0f)
-				.FillWidth(1.0f)
-				[
-					SNew(STextBlock)
-					.Font(IDetailLayoutBuilder::GetDetailFont())
-					.Text(FText::FromString(Entry->DisplayString))
-					.ToolTipText(FText::Format(LOCTEXT("ParameterPathTooltip", "{1}::{0}"), UncookedOnly::FUtils::GetParameterDisplayNameText(Entry->Binding.Parameter), FText::FromName(Entry->Binding.Asset.AssetName)))
-					.HighlightText_Lambda([this]()
-					{
-						return ParameterPicker.Pin()->FilterText;
-					})
-				];
+			FAnimNextParamType Type = FParamTypeHandle::FromProperty(InProperty).GetType();
+			return OnFilterParameterType.Execute(Type) == EFilterParameterResult::Include;
 		}
-		else if(InColumnName == Column_Graph)
+
+		return false;
+	};
+	
+	TArray<FFieldVariant> Result;
+	for (TFieldIterator<FProperty> PropertyIt(InStruct, EFieldIteratorFlags::IncludeSuper); PropertyIt; ++PropertyIt)
+	{
+		FProperty* Property = *PropertyIt;
+		if (FParamUtils::CanUseProperty(Property))
 		{
-			if(Entry->Binding.Graph.IsValid())
+			if(PassesFilterChecks(Property))
 			{
-				return
-					SNew(SBox)
-					.VAlign(VAlign_Center)
-					[
-						SNew(STextBlock)
-						.Font(IDetailLayoutBuilder::GetDetailFont())
-						.Text(FText::FromName(Entry->Binding.Graph.AssetName))
-						.ToolTipText(FText::FromName(Entry->Binding.Graph.PackageName))
-					];
+				Result.Add(FFieldVariant(Property));
 			}
 		}
-
-		return SNullWidget::NullWidget;
 	}
-	
-	TWeakPtr<const SParameterPicker> ParameterPicker;
-	TSharedPtr<FParameterPickerEntry> Entry;
-};
-
-TSharedRef<ITableRow> SParameterPicker::HandleGenerateRow(TSharedRef<FParameterPickerEntry> InEntry, const TSharedRef<STableViewBase>& InOwnerTable) const
-{
-	return SNew(SParameterPickerRow, InOwnerTable, InEntry, SharedThis(this));
-}
-
-void SParameterPicker::HandleGetChildren(TSharedRef<FParameterPickerEntry> InEntry, TArray<TSharedRef<FParameterPickerEntry>>& OutChildren) const
-{
-	OutChildren = InEntry->FilteredChildren;
-}
-
-void SParameterPicker::HandleSelectionChanged(TSharedPtr<FParameterPickerEntry> InEntry, ESelectInfo::Type InSelectInfo)
-{
-	Args.OnSelectionChanged.ExecuteIfBound();
-
-	if(!Args.bMultiSelect && EntriesList->GetNumItemsSelected() == 1 && Args.OnParameterPicked.IsBound())
+	for (TFieldIterator<UFunction> FunctionIt(InStruct, EFieldIteratorFlags::IncludeSuper); FunctionIt; ++FunctionIt)
 	{
-		TArray<TSharedRef<FParameterPickerEntry>> SelectedEntries;
-		EntriesList->GetSelectedItems(SelectedEntries);
-
-		if(SelectedEntries[0]->Binding.Type.IsValid() || Args.bAllowNone)
+		UFunction* Function = *FunctionIt;
+		if (FParamUtils::CanUseFunction(Function))
 		{
-			Args.OnParameterPicked.ExecuteIfBound(SelectedEntries[0]->Binding);
+			if(PassesFilterChecks(Function->GetReturnProperty()))
+			{
+				Result.Add(FFieldVariant(Function));
+			}
 		}
 	}
+	return Result;
 }
 
-void SParameterPicker::HandleGetParameterBindings(TArray<FParameterBindingReference>& OutParameterBindings) const
+TOptional<const UClass*> SParameterPicker::FFieldExpander::CanExpandObject(const FObjectPropertyBase* Property, const UObject* Instance) const
 {
-	TArray<TSharedRef<FParameterPickerEntry>> SelectedEntries;
-	EntriesList->GetSelectedItems(SelectedEntries);
+	return TOptional<const UClass*>();
+}
 
-	for(TSharedRef<FParameterPickerEntry>& Entry : SelectedEntries)
-	{
-		OutParameterBindings.Emplace(Entry->Binding);
-	}
+bool SParameterPicker::FFieldExpander::CanExpandScriptStruct(const FStructProperty* StructProperty) const
+{
+	return false;
 }
 
-bool SParameterPicker::HandleIsSelectableOrNavigable(TSharedRef<FParameterPickerEntry> InEntry) const
+TOptional<const UStruct*> SParameterPicker::FFieldExpander::GetExpandedFunction(const UFunction* Function) const
 {
-	// Only allow selecting items that have valid parameters
-	return InEntry->Binding.Type.IsValid();
+	return TOptional<const UStruct*>();
 }
 
 }
diff --git a/Engine/Plugins/Experimental/AnimNext/Source/AnimNextEditor/Private/Param/SParameterPicker.h b/Engine/Plugins/Experimental/AnimNext/Source/AnimNextEditor/Private/Param/SParameterPicker.h
--- a/Engine/Plugins/Experimental/AnimNext/Source/AnimNextEditor/Private/Param/SParameterPicker.h
+++ b/Engine/Plugins/Experimental/AnimNext/Source/AnimNextEditor/Private/Param/SParameterPicker.h
@@ -2,18 +2,25 @@
 
 #pragma once
 
+#include "Framework/PropertyViewer/IFieldExpander.h"
+#include "Framework/PropertyViewer/IFieldIterator.h"
+#include "Param/AnimNextParam.h"
 #include "Widgets/SCompoundWidget.h"
 #include "Param/ParameterPickerArgs.h"
+#include "Widgets/PropertyViewer/SPropertyViewer.h"
 #include "Widgets/Views/STreeView.h"
+#include "PropertyBag.h"
+#include "Misc/NotifyHook.h"
+#include "Param/AnimNextEditorParam.h"
 
+class IPropertyRowGenerator;
+class IStructureDataProvider;
 class SSearchBox;
 
 namespace UE::AnimNext::Editor
 {
 
-struct FParameterPickerEntry;
-
-class SParameterPicker : public SCompoundWidget
+class SParameterPicker : public SCompoundWidget, public FNotifyHook
 {
 public:
 	SLATE_BEGIN_ARGS(SParameterPicker) {}
@@ -27,38 +34,82 @@ public:
 private:
 	void RefreshEntries();
 
-	void BuildHierarchy();
-
-	void RefreshFilter();
-
-	TSharedRef<ITableRow> HandleGenerateRow(TSharedRef<FParameterPickerEntry> InEntry, const TSharedRef<STableViewBase>& InOwnerTable) const;
+	bool GetFieldInfo(UE::PropertyViewer::SPropertyViewer::FHandle InHandle, const FFieldVariant& InField, FName& OutName, TInstancedStruct<FAnimNextParamInstanceIdentifier>& OutInstanceId, FAnimNextParamType& OutType) const;
 
-	void HandleGetChildren(TSharedRef<FParameterPickerEntry> InEntry, TArray<TSharedRef<FParameterPickerEntry>>& OutChildren) const;
+	void HandleGetParameterBindings(TArray<FParameterBindingReference>& OutParameterBindings) const;
 
-	void HandleSelectionChanged(TSharedPtr<FParameterPickerEntry> InEntry, ESelectInfo::Type InSelectInfo);
+	void HandleSetInstanceId(const TInstancedStruct<FAnimNextParamInstanceIdentifier>& InInstanceId);
 
-	void HandleGetParameterBindings(TArray<FParameterBindingReference>& OutParameterBindings) const;
+	void HandleFieldPicked(UE::PropertyViewer::SPropertyViewer::FHandle InHandle, TArrayView<const FFieldVariant> InFields, ESelectInfo::Type InSelectionType);
 
-	bool HandleIsSelectableOrNavigable(TSharedRef<FParameterPickerEntry> InEntry) const;
+	TSharedRef<SWidget> HandleGenerateContainer(UE::PropertyViewer::SPropertyViewer::FHandle InHandle, TOptional<FText> InDisplayName);
 
+	// FNotifyHook interface
+	virtual void NotifyPostChange(const FPropertyChangedEvent& PropertyChangedEvent, FEditPropertyChain* PropertyThatChanged) override;
+	
 private:
 	friend class SParameterPickerRow;
 
 	FParameterPickerArgs Args;
-
-	TSharedPtr<STreeView<TSharedRef<FParameterPickerEntry>>> EntriesList;
-
-	TArray<TSharedRef<FParameterPickerEntry>> Entries;
-
-	TArray<TSharedRef<FParameterPickerEntry>> FilteredEntries;
-
-	TArray<TSharedRef<FParameterPickerEntry>> Hierarchy;
-
-	TArray<TSharedRef<FParameterPickerEntry>> FilteredHierarchy;
+	
+	TSharedPtr<UE::PropertyViewer::SPropertyViewer> PropertyViewer;
 
 	FText FilterText;
 
 	TSharedPtr<SSearchBox> SearchBox;
+
+	TInstancedStruct<FAnimNextParamInstanceIdentifier> SelectedInstanceId;
+
+	TSharedPtr<IPropertyRowGenerator> PropertyRowGenerator;
+
+	TSharedPtr<IStructureDataProvider> InstanceIdProvider;
+
+	struct FFieldIterator : UE::PropertyViewer::IFieldIterator
+	{
+		FFieldIterator(FOnFilterParameterType InOnFilterParameterType)
+			: OnFilterParameterType(InOnFilterParameterType)
+		{}
+		
+		virtual TArray<FFieldVariant> GetFields(const UStruct* InStruct) const override;
+
+		FOnFilterParameterType OnFilterParameterType;
+	};
+
+	struct FFieldExpander : UE::PropertyViewer::IFieldExpander
+	{
+		virtual TOptional<const UClass*> CanExpandObject(const FObjectPropertyBase* Property, const UObject* Instance) const override;
+		virtual bool CanExpandScriptStruct(const FStructProperty* StructProperty) const override;
+		virtual TOptional<const UStruct*> GetExpandedFunction(const UFunction* Function) const override;
+	} FieldExpander;
+
+	TUniquePtr<FFieldIterator> FieldIterator;
+
+	struct FContainerInfo
+	{
+		FContainerInfo(const FText& InDisplayName, const FText& InTooltipText, const UStruct* InStruct)
+			: DisplayName(InDisplayName)
+			, TooltipText(InTooltipText)
+			, Struct(InStruct)
+		{}
+
+		FContainerInfo(const FText& InDisplayName, const FText& InTooltipText, const FAssetData& InAssetData, TUniquePtr<FInstancedPropertyBag>&& InPropertyBag)
+			: DisplayName(InDisplayName)
+			, TooltipText(InTooltipText)
+			, Struct(InPropertyBag.Get()->GetPropertyBagStruct())
+			, PropertyBag(MoveTemp(InPropertyBag))
+			, AssetData(InAssetData)
+		{}
+
+		FText DisplayName;
+		FText TooltipText;
+		const UStruct* Struct = nullptr;
+		TUniquePtr<FInstancedPropertyBag> PropertyBag;
+		FAssetData AssetData;
+	};
+	
+	TArray<FContainerInfo> CachedContainers;
+
+	TMap<UE::PropertyViewer::SPropertyViewer::FHandle, int32> ContainerMap;
 };
 
 }
\ No newline at end of file
diff --git a/Engine/Plugins/Experimental/AnimNext/Source/AnimNextEditor/Private/Param/SParameterPickerCombo.cpp b/Engine/Plugins/Experimental/AnimNext/Source/AnimNextEditor/Private/Param/SParameterPickerCombo.cpp
--- a/Engine/Plugins/Experimental/AnimNext/Source/AnimNextEditor/Private/Param/SParameterPickerCombo.cpp
+++ b/Engine/Plugins/Experimental/AnimNext/Source/AnimNextEditor/Private/Param/SParameterPickerCombo.cpp
@@ -7,7 +7,6 @@
 #include "Widgets/Input/SComboButton.h"
 #include "Framework/Application/SlateApplication.h"
 #include "Widgets/Text/SInlineEditableTextBlock.h"
-#include "Widgets/Input/SSearchBox.h"
 
 namespace UE::AnimNext::Editor
 {
@@ -16,32 +15,52 @@ void SParameterPickerCombo::Construct(const FArguments& InArgs)
 {
 	OnGetParameterNameDelegate = InArgs._OnGetParameterName;
 	OnGetParameterTypeDelegate = InArgs._OnGetParameterType;
+	OnGetParameterInstanceIdDelegate = InArgs._OnGetParameterInstanceId;
 
-	FParameterPickerArgs PickerArgs = InArgs._PickerArgs;
-	PickerArgs.OnParameterPicked = FOnParameterPicked::CreateLambda([this, OrigPickerArgs = InArgs._PickerArgs](const FParameterBindingReference& InReference)
+	FOnParameterPicked OnParameterPicked = InArgs._PickerArgs.OnParameterPicked;
+	FOnInstanceIdChanged OnInstanceIdChanged = InArgs._PickerArgs.OnInstanceIdChanged;
+
+	PickerArgs = InArgs._PickerArgs;
+
+	PickerArgs.OnParameterPicked = FOnParameterPicked::CreateSPLambda(this, [this, OnParameterPicked](const FParameterBindingReference& InReference)
 	{
 		FSlateApplication::Get().DismissAllMenus();
 
 		// Forward to the original delegate
-		OrigPickerArgs.OnParameterPicked.ExecuteIfBound(InReference);
+		OnParameterPicked.ExecuteIfBound(InReference);
+
+		RequestRefresh();
+	});
+	PickerArgs.OnInstanceIdChanged = FOnInstanceIdChanged::CreateSPLambda(this, [this, OnInstanceIdChanged](const TInstancedStruct<FAnimNextParamInstanceIdentifier>& InInstanceId)
+	{
+		// Forward to the original delegate
+		OnInstanceIdChanged.ExecuteIfBound(InInstanceId);
 
 		RequestRefresh();
 	});
+	if(OnGetParameterInstanceIdDelegate.IsBound())
+	{
+		PickerArgs.InstanceId = OnGetParameterInstanceIdDelegate.Execute();
+	}
 	PickerArgs.bFocusSearchWidget = true;
 
 	ChildSlot
 	[
 		SNew(SComboButton)
-		.OnGetMenuContent_Lambda([PickerArgs, this]()
+		.ToolTipText_Lambda([this]()
+		{
+			return ParameterNameTooltipText;
+		})
+		.OnGetMenuContent_Lambda([this]()
 		{
+			if(OnGetParameterInstanceIdDelegate.IsBound())
+			{
+				PickerArgs.InstanceId = OnGetParameterInstanceIdDelegate.Execute();
+			}
+
 			return
-				SNew(SBox)
-				.WidthOverride(300.0f)
-				.HeightOverride(300.0f)
-				[
-					SNew(SParameterPicker)
-					.Args(PickerArgs)
-				];
+				SNew(SParameterPicker)
+				.Args(PickerArgs);
 		})
 		.ButtonContent()
 		[
@@ -83,25 +102,36 @@ void SParameterPickerCombo::Construct(const FArguments& InArgs)
 
 void SParameterPickerCombo::RequestRefresh()
 {
-	RegisterActiveTimer(1.0f/60.0f, FWidgetActiveTimerDelegate::CreateLambda([this](double InCurrentTime, float InDeltaTime)
+	if(!bRefreshRequested)
 	{
-		if(OnGetParameterNameDelegate.IsBound())
+		bRefreshRequested = true;
+		RegisterActiveTimer(1.0f/60.0f, FWidgetActiveTimerDelegate::CreateLambda([this](double InCurrentTime, float InDeltaTime)
 		{
-			ParameterName = OnGetParameterNameDelegate.Execute();
-		}
+			if(OnGetParameterNameDelegate.IsBound())
+			{
+				ParameterName = OnGetParameterNameDelegate.Execute();
+			}
 
-		if(OnGetParameterTypeDelegate.IsBound())
-		{
-			ParameterType = OnGetParameterTypeDelegate.Execute();
-		}
+			if(OnGetParameterTypeDelegate.IsBound())
+			{
+				ParameterType = OnGetParameterTypeDelegate.Execute();
+			}
+
+			if(OnGetParameterInstanceIdDelegate.IsBound())
+			{
+				ParameterInstanceId = OnGetParameterInstanceIdDelegate.Execute();
+			}
 
-		ParameterNameText = UE::AnimNext::UncookedOnly::FUtils::GetParameterDisplayNameText(ParameterName);
-		PinType = UncookedOnly::FUtils::GetPinTypeFromParamType(ParameterType);
-		Icon = FBlueprintEditorUtils::GetIconFromPin(PinType, true);
-		IconColor = GetDefault<UEdGraphSchema_K2>()->GetPinTypeColor(PinType);
+			ParameterNameText = UncookedOnly::FUtils::GetParameterDisplayNameText(ParameterName, ParameterInstanceId);
+			ParameterNameTooltipText = UncookedOnly::FUtils::GetParameterTooltipText(ParameterName, ParameterInstanceId);
+			PinType = UncookedOnly::FUtils::GetPinTypeFromParamType(ParameterType);
+			Icon = FBlueprintEditorUtils::GetIconFromPin(PinType, true);
+			IconColor = GetDefault<UEdGraphSchema_K2>()->GetPinTypeColor(PinType);
 
-		return EActiveTimerReturnType::Stop;
-	}));
+			bRefreshRequested = false;
+			return EActiveTimerReturnType::Stop;
+		}));
+	}
 }
 
 }
diff --git a/Engine/Plugins/Experimental/AnimNext/Source/AnimNextEditor/Private/Param/SParameterPickerCombo.h b/Engine/Plugins/Experimental/AnimNext/Source/AnimNextEditor/Private/Param/SParameterPickerCombo.h
--- a/Engine/Plugins/Experimental/AnimNext/Source/AnimNextEditor/Private/Param/SParameterPickerCombo.h
+++ b/Engine/Plugins/Experimental/AnimNext/Source/AnimNextEditor/Private/Param/SParameterPickerCombo.h
@@ -6,11 +6,13 @@
 #include "Param/ParameterPickerArgs.h"
 #include "EdGraph/EdGraphPin.h"
 
+struct FAnimNextParamInstanceIdentifier;
 struct FAnimNextParamType;
 
 namespace UE::AnimNext::Editor
 {
-class SParameterPicker;
+	class SParameterPicker;
+	class SGraphPinParam;
 }
 
 namespace UE::AnimNext::Editor
@@ -22,6 +24,9 @@ using FOnGetParameterName = TDelegate<FName(void)>;
 /** Retrieves the parameter type to display */
 using FOnGetParameterType = TDelegate<FAnimNextParamType(void)>;
 
+/** Retrieves the parameter scope to display */
+using FOnGetParameterInstanceId = TDelegate<TInstancedStruct<FAnimNextParamInstanceIdentifier>(void)>;
+
 class SParameterPickerCombo : public SCompoundWidget
 {
 public:
@@ -36,6 +41,9 @@ public:
 	/** Retrieves the parameter type to display */
 	SLATE_EVENT(FOnGetParameterType, OnGetParameterType)
 
+	/** Retrieves the parameter instance ID to display */
+	SLATE_EVENT(FOnGetParameterInstanceId, OnGetParameterInstanceId)
+
 	SLATE_END_ARGS()
 
 	void Construct(const FArguments& InArgs);
@@ -49,6 +57,9 @@ private:
 	// Retrieves the parameter type to display
 	FOnGetParameterType OnGetParameterTypeDelegate;
 
+	// Retrieves the parameter instance ID to display
+	FOnGetParameterInstanceId OnGetParameterInstanceIdDelegate;
+
 	// Cached pin type
 	FEdGraphPinType PinType;
 
@@ -58,9 +69,15 @@ private:
 	// Cached display name
 	FText ParameterNameText;
 
+	// Cached tooltip
+	FText ParameterNameTooltipText;
+
 	// Cached parameter type
 	FAnimNextParamType ParameterType;
 
+	// Cached parameter instance ID
+	TInstancedStruct<FAnimNextParamInstanceIdentifier> ParameterInstanceId;
+
 	// Cached icon
 	const FSlateBrush* Icon = nullptr;
 
@@ -69,6 +86,14 @@ private:
 
 	// Picker widget used to focus after the popup is displayed
 	TWeakPtr<SParameterPicker> PickerWidget;
+
+	// Arguments for the picker popup
+	FParameterPickerArgs PickerArgs;
+
+	bool bRefreshRequested = false;
+
+	friend class SGraphPinParam;
+	friend class SGraphPinParamName;
 };
 
 }
\ No newline at end of file
diff --git a/Engine/Plugins/Experimental/AnimNext/Source/AnimNextTestSuite/Private/AnimNextEditorTest.cpp b/Engine/Plugins/Experimental/AnimNext/Source/AnimNextTestSuite/Private/AnimNextEditorTest.cpp
--- a/Engine/Plugins/Experimental/AnimNext/Source/AnimNextTestSuite/Private/AnimNextEditorTest.cpp
+++ b/Engine/Plugins/Experimental/AnimNext/Source/AnimNextTestSuite/Private/AnimNextEditorTest.cpp
@@ -9,7 +9,9 @@
 #include "Graph/AnimNextGraph_Parameter.h"
 #include "Graph/AnimNextGraph_EditorData.h"
 #include "Graph/AnimNextGraph_AnimationGraph.h"
+#include "Graph/AnimNextGraph_EventGraph.h"
 #include "Graph/GraphFactory.h"
+#include "Param/RigVMDispatch_GetParameter.h"
 #if WITH_EDITOR
 #include "ScopedTransaction.h"
 #include "Editor.h"
@@ -29,6 +31,8 @@ bool FEditor_Parameters::RunTest(const FString& InParameters)
 {
 	using namespace UE::AnimNext;
 
+	ON_SCOPE_EXIT{ FUtils::CleanupAfterTests(); };
+
 	const TStrongObjectPtr<UFactory> GraphFactory(NewObject<UAnimNextGraphFactory>());
 	UAnimNextGraph* Graph = Cast<UAnimNextGraph>(GraphFactory->FactoryCreateNew(UAnimNextGraph::StaticClass(), GetTransientPackage(), TEXT("TestAnimNextGraph"), RF_Transient, nullptr, nullptr, NAME_None));
 	if(Graph == nullptr)
@@ -71,14 +75,17 @@ bool FEditor_Parameters::RunTest(const FString& InParameters)
 	AddExpectedError(TEXT("UAnimNextGraph_EditorData::AddParameter: Invalid parameter name supplied."));
 	AddErrorIfFalse(EditorData->AddParameter(NAME_None, FAnimNextParamType::GetType<bool>()) == nullptr, TEXT("Expected invalid argument to fail"));
 
-	auto TestParameterType = [this, EditorData](FAnimNextParamType InType)
+	auto TestParameterType = [this, EditorData](FAnimNextParamType InType, bool bInRemove = true)
 	{
 		UAnimNextGraph_Parameter* TypedParameter = EditorData->AddParameter(TEXT("TestParam0"), InType);
 		const bool bValidParameter = TypedParameter != nullptr;
 		if (bValidParameter && AddErrorIfFalse(bValidParameter, FString::Printf(TEXT("Could not create new parameter of type %s in graph."), *InType.ToString())))
 		{
 			AddErrorIfFalse(TypedParameter->GetParamType() == InType, TEXT("Incorrect parameter type found"));
-			EditorData->RemoveEntry(TypedParameter);
+			if(bInRemove)
+			{
+				EditorData->RemoveEntry(TypedParameter);
+			}
 		}
 	};
 
@@ -121,10 +128,9 @@ bool FEditor_Parameters::RunTest(const FString& InParameters)
 	GEditor->UndoTransaction();
 
 	// Add graph
-	UAnimNextGraph_EventGraph* EventGraph = nullptr;
 	{
 		FScopedTransaction Transaction(FText::GetEmpty());
-		EventGraph = EditorData->AddEventGraph(TEXT("TestGraph"));
+		UAnimNextGraph_EventGraph* EventGraph = EditorData->AddEventGraph(TEXT("TestGraph"));
 		AddErrorIfFalse(EventGraph != nullptr, TEXT("Could not create new event graph in graph."));
 	}
 
@@ -135,12 +141,41 @@ bool FEditor_Parameters::RunTest(const FString& InParameters)
 	AddErrorIfFalse(EditorData->Entries.Num() == 1, FString::Printf(TEXT("Unexpected entry count found in graph (Have %d, expected 1)."), EditorData->Entries.Num()));
 	GEditor->UndoTransaction();
 
-	FUtils::CleanupAfterTests();
+	// Add graph and add parameter getters and setters to it, testing compilation
+	{
+		TestParameterType(FAnimNextParamType::GetType<bool>(), false);
+
+		IAnimNextRigVMParameterInterface* ParameterEntry = CastChecked<IAnimNextRigVMParameterInterface>(EditorData->FindEntry("TestParam0"));
+		UE_RETURN_ON_ERROR(ParameterEntry != nullptr, TEXT("Could not find new parameter entry."));
+
+		UAnimNextGraph_EventGraph* EventGraph = EditorData->AddEventGraph(TEXT("TestGraph1"));
+		AddErrorIfFalse(EventGraph != nullptr, TEXT("Could not create new event graph in graph."));
+
+		URigVMGraph* RigVMGraph = EventGraph->GetRigVMGraph();
+		UE_RETURN_ON_ERROR(RigVMGraph->GetNodes().Num() == 1, TEXT("Unexpected number of nodes in new event graph."));
+
+		URigVMNode* EventNode = RigVMGraph->GetNodes()[0];
+		check(EventNode);
+		URigVMPin* ExecutePin = EventNode->FindPin("ExecuteContext");
+		UE_RETURN_ON_ERROR(ExecutePin != nullptr, TEXT("Could find initial execute pin."));
+
+		UAnimNextGraph_Controller* Controller = Cast<UAnimNextGraph_Controller>(EditorData->GetController(EventGraph->GetRigVMGraph()));
+		URigVMNode* GetParameterNode = Controller->AddGetAnimNextGraphParameterNode(FVector2D::ZeroVector, ParameterEntry->GetParamName(), FAnimNextParamType::GetType<bool>());
+		UE_RETURN_ON_ERROR(GetParameterNode != nullptr, TEXT("Could not add GetParameter node."));
+
+		URigVMNode* SetParameterNode = Controller->AddSetAnimNextGraphParameterNode(FVector2D::ZeroVector, ParameterEntry->GetParamName(), FAnimNextParamType::GetType<bool>());
+		UE_RETURN_ON_ERROR(SetParameterNode != nullptr, TEXT("Could not add SetParameter node."));
+
+		UE_RETURN_ON_ERROR(Controller->AddLink(ExecutePin, SetParameterNode->FindPin("ExecuteContext")), TEXT("Could not link SetParameter node."));
+
+		const FString ValuePin = FRigVMDispatch_GetParameter::ValueName.ToString(); 
+		UE_RETURN_ON_ERROR(Controller->AddLink(GetParameterNode->FindPin(ValuePin), SetParameterNode->FindPin(ValuePin)), TEXT("Could not link value pins."));
+	}
 
 	return true;
 }
 
-IMPLEMENT_SIMPLE_AUTOMATION_TEST(FEditor_Parameters_Python, "Animation.AnimNext.Editor.Parameters.Python", EAutomationTestFlags::EditorContext | EAutomationTestFlags::EngineFilter)
+IMPLEMENT_SIMPLE_AUTOMATION_TEST(FEditor_Parameters_Python, "Animation.AnimNext.Editor.PythonParameters", EAutomationTestFlags::EditorContext | EAutomationTestFlags::EngineFilter)
 
 bool FEditor_Parameters_Python::RunTest(const FString& InParameters)
 {
diff --git a/Engine/Plugins/Experimental/AnimNext/Source/AnimNextTestSuite/Private/AnimNextParametersTest.cpp b/Engine/Plugins/Experimental/AnimNext/Source/AnimNextTestSuite/Private/AnimNextParametersTest.cpp
--- a/Engine/Plugins/Experimental/AnimNext/Source/AnimNextTestSuite/Private/AnimNextParametersTest.cpp
+++ b/Engine/Plugins/Experimental/AnimNext/Source/AnimNextTestSuite/Private/AnimNextParametersTest.cpp
@@ -1,14 +1,21 @@
 ﻿// Copyright Epic Games, Inc. All Rights Reserved.
 
 #include "CoreMinimal.h"
+#include "AnimNextTest.h"
 #include "Context.h"
+#include "UncookedOnlyUtils.h"
 #include "Misc/AutomationTest.h"
 #include "Param/ParamType.h"
 #include "Param/ParamTypeHandle.h"
 #include "Animation/AnimSequence.h"
 #include "Animation/AnimMontage.h"
 #include "Async/ParallelFor.h"
+#include "Graph/AnimNextGraph.h"
+#include "Graph/AnimNextGraph_EditorData.h"
+#include "Graph/AnimNextGraph_Parameter.h"
+#include "Graph/GraphFactory.h"
 #include "Param/ParamStack.h"
+#include "UObject/StrongObjectPtr.h"
 
 // AnimNext Parameters Tests
 
@@ -517,7 +524,7 @@ bool FParamStackTest::RunTest(const FString& InParameters)
 		PropertyBag.SetValueBool("Param1", true);
 		PropertyBag.SetValueInt32("Param2", 5);
 
-		FParamStackLayerHandle Layer = FParamStack::MakeReferenceLayer(PropertyBag);
+		FParamStackLayerHandle Layer = FParamStack::MakeReferenceLayer(NAME_None, PropertyBag);
 
 		FParamStack Stack;
 
diff --git a/Engine/Plugins/Experimental/AnimNext/Source/AnimNextTestSuite/Private/AnimNextRuntimeTest.cpp b/Engine/Plugins/Experimental/AnimNext/Source/AnimNextTestSuite/Private/AnimNextRuntimeTest.cpp
--- a/Engine/Plugins/Experimental/AnimNext/Source/AnimNextTestSuite/Private/AnimNextRuntimeTest.cpp
+++ b/Engine/Plugins/Experimental/AnimNext/Source/AnimNextTestSuite/Private/AnimNextRuntimeTest.cpp
@@ -35,7 +35,7 @@ namespace UE::AnimNext
 		Graph.EntryPoints.Reset();
 
 		FAnimNextGraphEntryPoint& EntryPoint = Graph.EntryPoints.AddDefaulted_GetRef();
-		EntryPoint.EntryPointName = FRigUnit_AnimNextGraphRoot::DefaultEntryPoint;
+		EntryPoint.EntryPointName = Graph.GetDefaultEntryPoint();
 		EntryPoint.RootTraitHandle = FAnimNextEntryPointHandle(NodeHandles[0]);
 		Graph.ExecuteDefinition = ExecuteDefinition;
 		Graph.SharedDataArchiveBuffer = SharedDataArchiveBuffer;
@@ -48,7 +48,7 @@ namespace UE::AnimNext
 		const FTraitReader::EErrorState ErrorState = TraitReader.ReadGraph(Graph.SharedDataBuffer);
 		if (ErrorState == FTraitReader::EErrorState::None)
 		{
-			Graph.ResolvedRootTraitHandles.Add(FRigUnit_AnimNextGraphRoot::DefaultEntryPoint, TraitReader.ResolveEntryPointHandle(Graph.EntryPoints[0].RootTraitHandle));
+			Graph.ResolvedRootTraitHandles.Add(Graph.GetDefaultEntryPoint(), TraitReader.ResolveEntryPointHandle(Graph.EntryPoints[0].RootTraitHandle));
 
 			for (FNodeHandle& NodeHandle : NodeHandles)
 			{
diff --git a/Engine/Plugins/Experimental/AnimNext/Source/AnimNextUncookedOnly/AnimNextUncookedOnly.Build.cs b/Engine/Plugins/Experimental/AnimNext/Source/AnimNextUncookedOnly/AnimNextUncookedOnly.Build.cs
--- a/Engine/Plugins/Experimental/AnimNext/Source/AnimNextUncookedOnly/AnimNextUncookedOnly.Build.cs
+++ b/Engine/Plugins/Experimental/AnimNext/Source/AnimNextUncookedOnly/AnimNextUncookedOnly.Build.cs
@@ -13,6 +13,8 @@ namespace UnrealBuildTool.Rules
 					"RigVMDeveloper",
 					"ControlRig",
 					"ControlRigDeveloper",
+					"AnimNext",
+					"StructUtils",
 				}
 			);
 
@@ -22,16 +24,16 @@ namespace UnrealBuildTool.Rules
 					"Core",
 					"CoreUObject",
 					"Engine",
-					"AnimNext",
 					"BlueprintGraph",	// For K2 schema
 					"AnimationCore",
 					"AnimGraph",
 					"Kismet",
 					"Slate",
 					"SlateCore",
-					"StructUtils",
 					"StructUtilsEditor",
 					"ToolMenus",
+					"UniversalObjectLocator", 
+					"UniversalObjectLocatorEditor",
 				}
 			);
 
diff --git a/Engine/Plugins/Experimental/AnimNext/Source/AnimNextUncookedOnly/Internal/Graph/AnimNextGraph_AnimationGraph.h b/Engine/Plugins/Experimental/AnimNext/Source/AnimNextUncookedOnly/Internal/Graph/AnimNextGraph_AnimationGraph.h
--- a/Engine/Plugins/Experimental/AnimNext/Source/AnimNextUncookedOnly/Internal/Graph/AnimNextGraph_AnimationGraph.h
+++ b/Engine/Plugins/Experimental/AnimNext/Source/AnimNextUncookedOnly/Internal/Graph/AnimNextGraph_AnimationGraph.h
@@ -4,6 +4,7 @@
 
 #include "CoreMinimal.h"
 #include "AnimNextRigVMAssetEntry.h"
+#include "IAnimNextRigVMExportInterface.h"
 #include "IAnimNextRigVMGraphInterface.h"
 #include "AnimNextGraph_AnimationGraph.generated.h"
 
@@ -18,13 +19,17 @@ namespace UE::AnimNext::Editor
 
 /** A single entry in an AnimNext graph asset */
 UCLASS(MinimalAPI, Category = "Animation Graphs")
-class UAnimNextGraph_AnimationGraph : public UAnimNextRigVMAssetEntry, public IAnimNextRigVMGraphInterface
+class UAnimNextGraph_AnimationGraph : public UAnimNextRigVMAssetEntry, public IAnimNextRigVMGraphInterface, public IAnimNextRigVMExportInterface
 {
 	GENERATED_BODY()
 
 	friend class UAnimNextGraph_EditorData;
 	friend struct UE::AnimNext::Editor::FUtils;	
 
+	// IAnimNextRigVMExportInterface interface
+	virtual FAnimNextParamType GetExportType() const override;
+	virtual FName GetExportName() const override;
+
 	// UAnimNextRigVMAssetEntry interface
 	virtual FName GetEntryName() const override;
 	virtual void SetEntryName(FName InName, bool bSetupUndoRedo = true) override;
diff --git a/Engine/Plugins/Experimental/AnimNext/Source/AnimNextUncookedOnly/Internal/Graph/AnimNextGraph_Controller.h b/Engine/Plugins/Experimental/AnimNext/Source/AnimNextUncookedOnly/Internal/Graph/AnimNextGraph_Controller.h
--- a/Engine/Plugins/Experimental/AnimNext/Source/AnimNextUncookedOnly/Internal/Graph/AnimNextGraph_Controller.h
+++ b/Engine/Plugins/Experimental/AnimNext/Source/AnimNextUncookedOnly/Internal/Graph/AnimNextGraph_Controller.h
@@ -3,10 +3,14 @@
 #pragma once
 
 #include "CoreMinimal.h"
+#include "Param/ParamType.h"
 #include "RigVMModel/RigVMController.h"
-
+#include "Param/AnimNextParamInstanceIdentifier.h"
+#include "InstancedStruct.h"
 #include "AnimNextGraph_Controller.generated.h"
 
+struct FAnimNextParamType;
+
 /**
   * Implements AnimNext RigVM controller extensions
   */
@@ -18,4 +22,27 @@ class UAnimNextGraph_Controller : public URigVMController
 public:
 	// Adds a unit node with a dynamic number of pins
 	URigVMUnitNode* AddUnitNodeWithPins(UScriptStruct* InScriptStruct, const FRigVMPinInfoArray& PinArray, const FName& InMethodName = TEXT("Execute"), const FVector2D& InPosition = FVector2D::ZeroVector, const FString& InNodeName = TEXT(""), bool bSetupUndoRedo = true, bool bPrintPythonCommand = false);
+
+	// Sets the parameter referenced by a set/get parameter node
+	UFUNCTION(BlueprintCallable, Category = RigVMController)
+	ANIMNEXTUNCOOKEDONLY_API bool SetAnimNextParameterNode(URigVMNode* ParameterNode, FName ParameterName, EPropertyBagPropertyType ValueType, EPropertyBagContainerType ContainerType = EPropertyBagContainerType::None, const UObject* ValueTypeObject = nullptr, const TInstancedStruct<FAnimNextParamInstanceIdentifier>& InstanceId = TInstancedStruct<FAnimNextParamInstanceIdentifier>(), bool bSetupUndoRedo = true, bool bPrintPythonCommand = false);
+	ANIMNEXTUNCOOKEDONLY_API bool SetAnimNextParameterNode(URigVMNode* ParameterNode, FName ParameterName, const FAnimNextParamType& InType, const TInstancedStruct<FAnimNextParamInstanceIdentifier>& InstanceId = TInstancedStruct<FAnimNextParamInstanceIdentifier>(), bool bSetupUndoRedo = true, bool bPrintPythonCommand = false);
+
+	// Creates a Get Parameter node
+	UFUNCTION(BlueprintCallable, Category = RigVMController)
+	ANIMNEXTUNCOOKEDONLY_API URigVMNode* AddGetAnimNextParameterNode(const FVector2D& InPosition, FName ParameterName, EPropertyBagPropertyType ValueType, EPropertyBagContainerType ContainerType = EPropertyBagContainerType::None, const UObject* ValueTypeObject = nullptr, const TInstancedStruct<FAnimNextParamInstanceIdentifier>& InstanceId = TInstancedStruct<FAnimNextParamInstanceIdentifier>(), bool bSetupUndoRedo = true, bool bPrintPythonCommand = false);
+	ANIMNEXTUNCOOKEDONLY_API URigVMNode* AddGetAnimNextParameterNode(const FVector2D& InPosition, FName ParameterName, const FAnimNextParamType& InType, const TInstancedStruct<FAnimNextParamInstanceIdentifier>& InstanceId = TInstancedStruct<FAnimNextParamInstanceIdentifier>(), bool bSetupUndoRedo = true, bool bPrintPythonCommand = false);
+
+	// Creates a Get Graph Parameter node
+	UFUNCTION(BlueprintCallable, Category = RigVMController)
+	ANIMNEXTUNCOOKEDONLY_API URigVMNode* AddGetAnimNextGraphParameterNode(const FVector2D& InPosition, FName ParameterName, EPropertyBagPropertyType ValueType, EPropertyBagContainerType ContainerType = EPropertyBagContainerType::None, const UObject* ValueTypeObject = nullptr, const TInstancedStruct<FAnimNextParamInstanceIdentifier>& InstanceId = TInstancedStruct<FAnimNextParamInstanceIdentifier>(), bool bSetupUndoRedo = true, bool bPrintPythonCommand = false);
+	ANIMNEXTUNCOOKEDONLY_API URigVMNode* AddGetAnimNextGraphParameterNode(const FVector2D& InPosition, FName ParameterName, const FAnimNextParamType& InType, const TInstancedStruct<FAnimNextParamInstanceIdentifier>& InstanceId = TInstancedStruct<FAnimNextParamInstanceIdentifier>(), bool bSetupUndoRedo = true, bool bPrintPythonCommand = false);
+
+	// Creates a Set Graph Parameter node
+	UFUNCTION(BlueprintCallable, Category = RigVMController)
+	ANIMNEXTUNCOOKEDONLY_API URigVMNode* AddSetAnimNextGraphParameterNode(const FVector2D& InPosition, FName ParameterName, EPropertyBagPropertyType ValueType, EPropertyBagContainerType ContainerType = EPropertyBagContainerType::None, const UObject* ValueTypeObject = nullptr, const TInstancedStruct<FAnimNextParamInstanceIdentifier>& InstanceId = TInstancedStruct<FAnimNextParamInstanceIdentifier>(), bool bSetupUndoRedo = true, bool bPrintPythonCommand = false);
+	ANIMNEXTUNCOOKEDONLY_API URigVMNode* AddSetAnimNextGraphParameterNode(const FVector2D& InPosition, FName ParameterName, const FAnimNextParamType& InType, const TInstancedStruct<FAnimNextParamInstanceIdentifier>& InstanceId = TInstancedStruct<FAnimNextParamInstanceIdentifier>(), bool bSetupUndoRedo = true, bool bPrintPythonCommand = false);
+
+private:
+	URigVMNode* AddGetAnimNextParameterNodeInternal(UScriptStruct* NodeStruct, const FVector2D& Position, FName ParameterName, const FAnimNextParamType& Type, const TInstancedStruct<FAnimNextParamInstanceIdentifier>& InstanceId, bool bSetupUndoRedo, bool bPrintPythonCommand);
 };
diff --git a/Engine/Plugins/Experimental/AnimNext/Source/AnimNextUncookedOnly/Internal/Graph/AnimNextGraph_EdGraphSchema.h b/Engine/Plugins/Experimental/AnimNext/Source/AnimNextUncookedOnly/Internal/Graph/AnimNextGraph_EdGraphSchema.h
--- a/Engine/Plugins/Experimental/AnimNext/Source/AnimNextUncookedOnly/Internal/Graph/AnimNextGraph_EdGraphSchema.h
+++ b/Engine/Plugins/Experimental/AnimNext/Source/AnimNextUncookedOnly/Internal/Graph/AnimNextGraph_EdGraphSchema.h
@@ -14,7 +14,8 @@ class UAnimNextGraph_EdGraphSchema : public URigVMEdGraphSchema
 
 	// URigVMEdGraphSchema interface
 	virtual TSubclassOf<URigVMEdGraphNode> GetGraphNodeClass(const URigVMEdGraph* InGraph) const override { return UAnimNextGraph_EdGraphNode::StaticClass(); }
-	
+	virtual bool IsStructEditable(UStruct* InStruct) const override;
+
 	// UEdGraphSchema interface
 	virtual void GetGraphDisplayInformation(const UEdGraph& Graph, /*out*/ FGraphDisplayInfo& DisplayInfo) const override;
 };
\ No newline at end of file
diff --git a/Engine/Plugins/Experimental/AnimNext/Source/AnimNextUncookedOnly/Internal/Graph/AnimNextGraph_EventGraph.h b/Engine/Plugins/Experimental/AnimNext/Source/AnimNextUncookedOnly/Internal/Graph/AnimNextGraph_EventGraph.h
--- a/Engine/Plugins/Experimental/AnimNext/Source/AnimNextUncookedOnly/Internal/Graph/AnimNextGraph_EventGraph.h
+++ b/Engine/Plugins/Experimental/AnimNext/Source/AnimNextUncookedOnly/Internal/Graph/AnimNextGraph_EventGraph.h
@@ -4,18 +4,30 @@
 
 #include "CoreMinimal.h"
 #include "AnimNextRigVMAssetEntry.h"
+#include "IAnimNextRigVMExportInterface.h"
 #include "IAnimNextRigVMGraphInterface.h"
+#include "IAnimNextRigVMParameterInterface.h"
 #include "AnimNextGraph_EventGraph.generated.h"
 
 class UAnimNextGraph_EditorData;
 class UAnimNextGraph_EdGraph;
 
+namespace UE::AnimNext::Tests
+{
+	class FEditor_Parameters;
+}
+
 UCLASS(MinimalAPI, Category = "Event Graphs")
-class UAnimNextGraph_EventGraph : public UAnimNextRigVMAssetEntry, public IAnimNextRigVMGraphInterface
+class UAnimNextGraph_EventGraph : public UAnimNextRigVMAssetEntry, public IAnimNextRigVMGraphInterface, public IAnimNextRigVMExportInterface
 {
 	GENERATED_BODY()
 
 	friend class UAnimNextGraph_EditorData;
+	friend class UE::AnimNext::Tests::FEditor_Parameters;
+
+	// IAnimNextRigVMExportInterface interface
+	virtual FAnimNextParamType GetExportType() const override;
+	virtual FName GetExportName() const override;
 
 	// UAnimNextRigVMAssetEntry interface
 	virtual FName GetEntryName() const override { return GraphName; }
diff --git a/Engine/Plugins/Experimental/AnimNext/Source/AnimNextUncookedOnly/Internal/Graph/AnimNextGraph_Parameter.h b/Engine/Plugins/Experimental/AnimNext/Source/AnimNextUncookedOnly/Internal/Graph/AnimNextGraph_Parameter.h
--- a/Engine/Plugins/Experimental/AnimNext/Source/AnimNextUncookedOnly/Internal/Graph/AnimNextGraph_Parameter.h
+++ b/Engine/Plugins/Experimental/AnimNext/Source/AnimNextUncookedOnly/Internal/Graph/AnimNextGraph_Parameter.h
@@ -4,6 +4,7 @@
 
 #include "CoreMinimal.h"
 #include "AnimNextRigVMAssetEntry.h"
+#include "IAnimNextRigVMExportInterface.h"
 #include "IAnimNextRigVMParameterInterface.h"
 #include "Param/ParamType.h"
 #include "AnimNextGraph_Parameter.generated.h"
@@ -21,7 +22,7 @@ namespace UE::AnimNext::Tests
 }
 
 UCLASS(MinimalAPI, Category = "Parameters")
-class UAnimNextGraph_Parameter : public UAnimNextRigVMAssetEntry, public IAnimNextRigVMParameterInterface
+class UAnimNextGraph_Parameter : public UAnimNextRigVMAssetEntry, public IAnimNextRigVMParameterInterface, public IAnimNextRigVMExportInterface
 {
 	GENERATED_BODY()
 
@@ -29,6 +30,10 @@ class UAnimNextGraph_Parameter : public UAnimNextRigVMAssetEntry, public IAnimNe
 	friend class UE::AnimNext::Tests::FEditor_Parameters;
 	friend class UE::AnimNext::Editor::FParameterCustomization;
 
+	// IAnimNextRigVMExportInterface interface
+	virtual FAnimNextParamType GetExportType() const override;
+	virtual FName GetExportName() const override;
+
 	// UAnimNextRigVMAssetEntry interface
 	virtual FName GetEntryName() const override;
 	virtual void SetEntryName(FName InName, bool bSetupUndoRedo = true) override;
@@ -38,6 +43,8 @@ class UAnimNextGraph_Parameter : public UAnimNextRigVMAssetEntry, public IAnimNe
 	// IAnimNextRigVMParameterInterface interface
 	virtual FAnimNextParamType GetParamType() const override;
 	virtual bool SetParamType(const FAnimNextParamType& InType, bool bSetupUndoRedo = true) override;
+	virtual FName GetParamName() const override;
+	virtual void SetParamName(FName InName, bool bSetupUndoRedo = true) override;
 	virtual FInstancedPropertyBag& GetPropertyBag() const override;
 	
 	/** Parameter name we reference */
diff --git a/Engine/Plugins/Experimental/AnimNext/Source/AnimNextUncookedOnly/Internal/IAnimNextRigVMExportInterface.h b/Engine/Plugins/Experimental/AnimNext/Source/AnimNextUncookedOnly/Internal/IAnimNextRigVMExportInterface.h
new file mode 100644
--- /dev/null
+++ b/Engine/Plugins/Experimental/AnimNext/Source/AnimNextUncookedOnly/Internal/IAnimNextRigVMExportInterface.h
@@ -0,0 +1,26 @@
+﻿// Copyright Epic Games, Inc. All Rights Reserved.
+
+#pragma once
+
+#include "UObject/Interface.h"
+#include "IAnimNextRigVMExportInterface.generated.h"
+
+struct FAnimNextParamType;
+
+UINTERFACE(meta=(CannotImplementInterfaceInBlueprint))
+class ANIMNEXTUNCOOKEDONLY_API UAnimNextRigVMExportInterface : public UInterface
+{
+	GENERATED_BODY()
+};
+
+class ANIMNEXTUNCOOKEDONLY_API IAnimNextRigVMExportInterface
+{
+	GENERATED_BODY()
+
+public:
+	// Get the export type
+	virtual FAnimNextParamType GetExportType() const = 0;
+
+	// Get the export name (e.g. asset path + entry name)
+	virtual FName GetExportName() const = 0;
+};
\ No newline at end of file
diff --git a/Engine/Plugins/Experimental/AnimNext/Source/AnimNextUncookedOnly/Internal/IAnimNextRigVMParameterInterface.h b/Engine/Plugins/Experimental/AnimNext/Source/AnimNextUncookedOnly/Internal/IAnimNextRigVMParameterInterface.h
--- a/Engine/Plugins/Experimental/AnimNext/Source/AnimNextUncookedOnly/Internal/IAnimNextRigVMParameterInterface.h
+++ b/Engine/Plugins/Experimental/AnimNext/Source/AnimNextUncookedOnly/Internal/IAnimNextRigVMParameterInterface.h
@@ -25,6 +25,12 @@ public:
 	// Set the parameter type
 	virtual bool SetParamType(const FAnimNextParamType& InType, bool bSetupUndoRedo = true) = 0;
 
+	// Get the parameter name (asset path + entry name)
+	virtual FName GetParamName() const = 0;
+
+	// Set the parameter type (only requires the entry name, do not include the asset path)
+	virtual void SetParamName(FName InName, bool bSetupUndoRedo = true) = 0;
+
 	// Access the backing storage property bag for the parameter
 	virtual FInstancedPropertyBag& GetPropertyBag() const = 0;
 };
\ No newline at end of file
diff --git a/Engine/Plugins/Experimental/AnimNext/Source/AnimNextUncookedOnly/Internal/UncookedOnlyUtils.h b/Engine/Plugins/Experimental/AnimNext/Source/AnimNextUncookedOnly/Internal/UncookedOnlyUtils.h
--- a/Engine/Plugins/Experimental/AnimNext/Source/AnimNextUncookedOnly/Internal/UncookedOnlyUtils.h
+++ b/Engine/Plugins/Experimental/AnimNext/Source/AnimNextUncookedOnly/Internal/UncookedOnlyUtils.h
@@ -5,11 +5,13 @@
 #include "CoreMinimal.h"
 #include "IAnimNextRigVMGraphInterface.h"
 #include "AssetRegistry/AssetData.h"
+#include "Param/AnimNextParamInstanceIdentifier.h"
 #include "Param/ParamTypeHandle.h"
 #include "RigVMCore/RigVMTemplate.h"
-
+#include "InstancedStruct.h"
 #include "UncookedOnlyUtils.generated.h"
 
+struct FAnimNextEditorParam;
 struct FAnimNextParam;
 class UAnimNextSchedule;
 class UAnimNextGraph;
@@ -32,7 +34,7 @@ namespace UE
 }
 
 UENUM()
-enum class EAnimNextParameterFlags
+enum class EAnimNextParameterFlags : uint32
 {
 	NoFlags = 0x0,
 	Private = 0x1,
@@ -51,23 +53,40 @@ struct FAnimNextParameterAssetRegistryExportEntry
 
 	FAnimNextParameterAssetRegistryExportEntry() = default;
 	
-	FAnimNextParameterAssetRegistryExportEntry(FName InName, const FAnimNextParamType& InType, EAnimNextParameterFlags InFlags = EAnimNextParameterFlags::NoFlags)
+	FAnimNextParameterAssetRegistryExportEntry(FName InName, const TInstancedStruct<FAnimNextParamInstanceIdentifier>& InInstanceId, const FAnimNextParamType& InType, EAnimNextParameterFlags InFlags = EAnimNextParameterFlags::NoFlags)
 		: Name(InName)
+		, InstanceId(InInstanceId)
 		, Type(InType)
-		, Flags(InFlags) 
+		, Flags((uint32)InFlags) 
 	{}
+
+	bool operator==(const FAnimNextParameterAssetRegistryExportEntry& Other) const
+	{
+		return Name == Other.Name && InstanceId == Other.InstanceId;
+	}
+
+	friend uint32 GetTypeHash(const FAnimNextParameterAssetRegistryExportEntry& Entry)
+	{
+		const FName InstanceIdName = Entry.InstanceId.IsValid() ? Entry.InstanceId.Get().ToName() : NAME_None;
+		return HashCombineFast(GetTypeHash(Entry.Name), GetTypeHash(InstanceIdName));
+	}
+
+	EAnimNextParameterFlags GetFlags() const
+	{
+		return (EAnimNextParameterFlags)Flags;
+	}
 	
 	UPROPERTY()
 	FName Name;
 
 	UPROPERTY()
-	FAnimNextParamType Type;
+	TInstancedStruct<FAnimNextParamInstanceIdentifier> InstanceId;
 
-	// Asset, found first in asset-registry, that references this parameter entry
-	FAssetData ReferencingAsset;
+	UPROPERTY()
+	FAnimNextParamType Type;
 
 	UPROPERTY()
-	EAnimNextParameterFlags Flags = EAnimNextParameterFlags::NoFlags;
+	uint32 Flags = (int32)EAnimNextParameterFlags::NoFlags;
 };
 
 USTRUCT()
@@ -134,27 +153,40 @@ struct ANIMNEXTUNCOOKEDONLY_API FUtils
 	
 	/** Set up a simple event graph */
 	static void SetupEventGraph(URigVMController* InController);
-	
-	/** Converts the Verse-tag-like snake_case_parameter_name to a period-separated display name similar to a gameplay tag */
-	static FText GetParameterDisplayNameText(FName InParameterName);
-	
+
+	/** Get a parameter name (e.g. MyParameter) from a fully-qualified parameter name (e.g. /Game/MyAsset.MyAsset:MyParameter) */
+	static FName GetParameterNameFromQualifiedName(FName InName);
+
+	/** Get a fully-qualified parameter name (e.g. /Game/MyAsset.MyAsset:MyParameter) from its containing asset and base name (e.g. MyParameter) */
+	static FName GetQualifiedName(UAnimNextRigVMAsset* InAsset, FName InBaseName);
+
+	/** Gets a name to display for a parameter in the editor, including scope if external */
+	static FText GetParameterDisplayNameText(FName InParameterName, const TInstancedStruct<FAnimNextParamInstanceIdentifier>& InInstanceId);
+
+	/** Gets a name to display for a parameter's tooltip in the editor, including scope if external */
+	static FText GetParameterTooltipText(FName InParameterName, const TInstancedStruct<FAnimNextParamInstanceIdentifier>& InInstanceId);
+
 	// Gets the parameters that are exported to the asset registry for an asset
 	static bool GetExportedParametersForAsset(const FAssetData& InAsset, FAnimNextParameterProviderAssetRegistryExports& OutExports);
 
 	// Gets all the parameters that are exported to the asset registry
-	static bool GetExportedParametersFromAssetRegistry(FAnimNextParameterProviderAssetRegistryExports& OutExports);
+	static bool GetExportedParametersFromAssetRegistry(TMap<FAssetData, FAnimNextParameterProviderAssetRegistryExports>& OutExports);
 
 	// Gets the exported parameters that are used by a RigVM asset
 	static void GetAssetParameters(const UAnimNextRigVMAssetEditorData* EditorData, FAnimNextParameterProviderAssetRegistryExports& OutExports);
+	static void GetAssetParameters(const UAnimNextRigVMAssetEditorData* EditorData, TSet<FAnimNextParameterAssetRegistryExportEntry>& OutExports);
 
 	// Gets the exported parameters that are used by a RigVM graph
 	static void GetGraphParameters(const URigVMGraph* Graph, FAnimNextParameterProviderAssetRegistryExports& OutExports);
+	static void GetGraphParameters(const URigVMGraph* Graph, TSet<FAnimNextParameterAssetRegistryExportEntry>& OutExports);
 
 	// Gets the parameters that are exported to the asset registry by a schedule
 	static void GetScheduleParameters(const UAnimNextSchedule* InSchedule, FAnimNextParameterProviderAssetRegistryExports& OutExports);
+	static void GetScheduleParameters(const UAnimNextSchedule* InSchedule, TSet<FAnimNextParameterAssetRegistryExportEntry>& OutExports);
 
 	// Gets the parameters that are exported to the asset registry by a blueprint
 	static void GetBlueprintParameters(const UBlueprint* InBlueprint, FAnimNextParameterProviderAssetRegistryExports& OutExports);
+	static void GetBlueprintParameters(const UBlueprint* InBlueprint, TSet<FAnimNextParameterAssetRegistryExportEntry>& OutExports);
 
 	// Attempts to determine the type from a parameter name
 	// If the name cannot be found, the returned type will be invalid
diff --git a/Engine/Plugins/Experimental/AnimNext/Source/AnimNextUncookedOnly/Private/AnimNextRigVMAssetEditorData.cpp b/Engine/Plugins/Experimental/AnimNext/Source/AnimNextUncookedOnly/Private/AnimNextRigVMAssetEditorData.cpp
--- a/Engine/Plugins/Experimental/AnimNext/Source/AnimNextUncookedOnly/Private/AnimNextRigVMAssetEditorData.cpp
+++ b/Engine/Plugins/Experimental/AnimNext/Source/AnimNextUncookedOnly/Private/AnimNextRigVMAssetEditorData.cpp
@@ -162,7 +162,7 @@ void UAnimNextRigVMAssetEditorData::GetAssetRegistryTags(FAssetRegistryTagsConte
 	}
 
 	FString TagValue;
-	FAnimNextParameterProviderAssetRegistryExports::StaticStruct()->ExportText(TagValue, &CachedExports, nullptr, nullptr, PPF_None, nullptr);
+	FAnimNextParameterProviderAssetRegistryExports::StaticStruct()->ExportText(TagValue, &CachedExports.GetValue(), nullptr, nullptr, PPF_None, nullptr);
 	Context.AddTag(FAssetRegistryTag(UE::AnimNext::ExportsAnimNextAssetRegistryTag, TagValue, FAssetRegistryTag::TT_Hidden));
 }
 
diff --git a/Engine/Plugins/Experimental/AnimNext/Source/AnimNextUncookedOnly/Private/AnimNextUncookedOnlyModule.cpp b/Engine/Plugins/Experimental/AnimNext/Source/AnimNextUncookedOnly/Private/AnimNextUncookedOnlyModule.cpp
--- a/Engine/Plugins/Experimental/AnimNext/Source/AnimNextUncookedOnly/Private/AnimNextUncookedOnlyModule.cpp
+++ b/Engine/Plugins/Experimental/AnimNext/Source/AnimNextUncookedOnly/Private/AnimNextUncookedOnlyModule.cpp
@@ -1,66 +1,80 @@
 // Copyright Epic Games, Inc. All Rights Reserved.
 
-#include "CoreMinimal.h"
-#include "K2Node_CallFunction.h"
+#include "AnimNextUncookedOnlyModule.h"
 #include "UncookedOnlyUtils.h"
-#include "Component/AnimNextComponentParameter.h"
-#include "Component/AnimNextComponent.h"
-#include "Engine/SCS_Node.h"
+#include "Engine/Blueprint.h"
 #include "Modules/ModuleManager.h"
-#include "Param/AnimNextTag.h"
+#include "Param/AnimNextParamUniversalObjectLocator.h"
+#include "Param/IParameterSourceType.h"
+#include "Param/ObjectProxyType.h"
 #include "Scheduler/AnimNextSchedule.h"
 #include "UObject/AssetRegistryTagsContext.h"
 
 namespace UE::AnimNext::UncookedOnly
 {
-	class FModule : public IModuleInterface
+
+void FModule::StartupModule()
+{
+	// TEMP: Bind the compilation function for schedules
+	UAnimNextSchedule::CompileFunction = [](UAnimNextSchedule* InSchedule)
 	{
-	private:
-		virtual void StartupModule() override
-		{
-			// TEMP: Bind the compilation function for schedules
-			UAnimNextSchedule::CompileFunction = [](UAnimNextSchedule* InSchedule)
-			{
-				FUtils::CompileSchedule(InSchedule);
-			};
+		FUtils::CompileSchedule(InSchedule);
+	};
 
-			// TEMP: Bind the asset registry tags function for schedules
-			UAnimNextSchedule::GetAssetRegistryTagsFunction = [](const UAnimNextSchedule* InSchedule, FAssetRegistryTagsContext Context)
-			{
-				FAnimNextParameterProviderAssetRegistryExports Exports;
-				FUtils::GetScheduleParameters(InSchedule, Exports);
-				
-				FString TagValue;
-				FAnimNextParameterProviderAssetRegistryExports::StaticStruct()->ExportText(TagValue, &Exports, nullptr, nullptr, PPF_None, nullptr);
-				Context.AddTag(UObject::FAssetRegistryTag(UE::AnimNext::ExportsAnimNextAssetRegistryTag, TagValue, UObject::FAssetRegistryTag::TT_Hidden));
-			};
-			
-			// Ensure that any BP components that we care about contribute to the parameter pool
-			OnGetExtraObjectTagsHandle = UObject::FAssetRegistryTag::OnGetExtraObjectTagsWithContext.AddLambda([](FAssetRegistryTagsContext Context)
-			{
-				const UObject* InObject = Context.GetObject();
-				if(const UBlueprint* Blueprint = Cast<UBlueprint>(InObject))
-				{
-					FAnimNextParameterProviderAssetRegistryExports Exports;
-					FUtils::GetBlueprintParameters(Blueprint, Exports);
+	// TEMP: Bind the asset registry tags function for schedules
+	UAnimNextSchedule::GetAssetRegistryTagsFunction = [](const UAnimNextSchedule* InSchedule, FAssetRegistryTagsContext Context)
+	{
+		FAnimNextParameterProviderAssetRegistryExports Exports;
+		FUtils::GetScheduleParameters(InSchedule, Exports);
+		
+		FString TagValue;
+		FAnimNextParameterProviderAssetRegistryExports::StaticStruct()->ExportText(TagValue, &Exports, nullptr, nullptr, PPF_None, nullptr);
+		Context.AddTag(UObject::FAssetRegistryTag(UE::AnimNext::ExportsAnimNextAssetRegistryTag, TagValue, UObject::FAssetRegistryTag::TT_Hidden));
+	};
+	
+	// Ensure that any BP components that we care about contribute to the parameter pool
+	OnGetExtraObjectTagsHandle = UObject::FAssetRegistryTag::OnGetExtraObjectTagsWithContext.AddLambda([](FAssetRegistryTagsContext Context)
+	{
+		const UObject* InObject = Context.GetObject();
+		if(const UBlueprint* Blueprint = Cast<UBlueprint>(InObject))
+		{
+			FAnimNextParameterProviderAssetRegistryExports Exports;
+			FUtils::GetBlueprintParameters(Blueprint, Exports);
 
-					FString TagValue;
-					FAnimNextParameterProviderAssetRegistryExports::StaticStruct()->ExportText(TagValue, &Exports, nullptr, nullptr, PPF_None, nullptr);
-					Context.AddTag(UObject::FAssetRegistryTag(UE::AnimNext::ExportsAnimNextAssetRegistryTag, TagValue, UObject::FAssetRegistryTag::TT_Hidden));
-				}
-			});
+			FString TagValue;
+			FAnimNextParameterProviderAssetRegistryExports::StaticStruct()->ExportText(TagValue, &Exports, nullptr, nullptr, PPF_None, nullptr);
+			Context.AddTag(UObject::FAssetRegistryTag(UE::AnimNext::ExportsAnimNextAssetRegistryTag, TagValue, UObject::FAssetRegistryTag::TT_Hidden));
 		}
+	});
 
-		virtual void ShutdownModule() override
-		{
-			UAnimNextSchedule::GetAssetRegistryTagsFunction = nullptr;
-			UAnimNextSchedule::CompileFunction = nullptr;
+	RegisterParameterSourceType(FAnimNextParamUniversalObjectLocator::StaticStruct(), MakeShared<FObjectProxyType>());
+}
 
-			UObject::FAssetRegistryTag::OnGetExtraObjectTagsWithContext.Remove(OnGetExtraObjectTagsHandle);
-		}
+void FModule::ShutdownModule()
+{
+	UnregisterParameterSourceType(FAnimNextParamUniversalObjectLocator::StaticStruct());
+
+	UAnimNextSchedule::GetAssetRegistryTagsFunction = nullptr;
+	UAnimNextSchedule::CompileFunction = nullptr;
+
+	UObject::FAssetRegistryTag::OnGetExtraObjectTagsWithContext.Remove(OnGetExtraObjectTagsHandle);
+}
+
+void FModule::RegisterParameterSourceType(const UScriptStruct* InInstanceIdStruct, TSharedPtr<IParameterSourceType> InType)
+{
+	ParameterSourceTypes.Add(InInstanceIdStruct, InType);
+}
+
+void FModule::UnregisterParameterSourceType(const UScriptStruct* InInstanceIdStruct)
+{
+	ParameterSourceTypes.Remove(InInstanceIdStruct);
+}
+
+TSharedPtr<IParameterSourceType> FModule::FindParameterSourceType(const UScriptStruct* InInstanceIdStruct) const
+{
+	return ParameterSourceTypes.FindRef(InInstanceIdStruct);
+}
 
-		FDelegateHandle OnGetExtraObjectTagsHandle;
-	};
 }
 
 IMPLEMENT_MODULE(UE::AnimNext::UncookedOnly::FModule, AnimNextUncookedOnly);
diff --git a/Engine/Plugins/Experimental/AnimNext/Source/AnimNextUncookedOnly/Private/AnimNextUncookedOnlyModule.h b/Engine/Plugins/Experimental/AnimNext/Source/AnimNextUncookedOnly/Private/AnimNextUncookedOnlyModule.h
new file mode 100644
--- /dev/null
+++ b/Engine/Plugins/Experimental/AnimNext/Source/AnimNextUncookedOnly/Private/AnimNextUncookedOnlyModule.h
@@ -0,0 +1,34 @@
+﻿// Copyright Epic Games, Inc. All Rights Reserved.
+
+#pragma once
+
+#include "CoreMinimal.h"
+#include "IAnimNextUncookedOnlyModule.h"
+
+namespace UE::AnimNext::UncookedOnly
+{
+	struct FUtils;
+}
+
+namespace UE::AnimNext::UncookedOnly
+{
+
+class FModule : public IAnimNextUncookedOnlyModule
+{
+private:
+	// IModuleInterface interface
+	virtual void StartupModule() override;
+	virtual void ShutdownModule() override;
+
+	// IAnimNextUncookedOnlyModule interface
+	virtual void RegisterParameterSourceType(const UScriptStruct* InInstanceIdStruct, TSharedPtr<IParameterSourceType> InType) override;
+	virtual void UnregisterParameterSourceType(const UScriptStruct* InInstanceIdStruct) override;
+	virtual TSharedPtr<IParameterSourceType> FindParameterSourceType(const UScriptStruct* InInstanceIdStruct) const override;
+
+	TMap<const UScriptStruct*, TSharedPtr<IParameterSourceType>> ParameterSourceTypes;
+	FDelegateHandle OnGetExtraObjectTagsHandle;
+
+	friend struct FUtils;
+};
+
+}
\ No newline at end of file
diff --git a/Engine/Plugins/Experimental/AnimNext/Source/AnimNextUncookedOnly/Private/Graph/AnimGraph/AnimGraphNodeBinding_AnimNextParameters.cpp b/Engine/Plugins/Experimental/AnimNext/Source/AnimNextUncookedOnly/Private/Graph/AnimGraph/AnimGraphNodeBinding_AnimNextParameters.cpp
--- a/Engine/Plugins/Experimental/AnimNext/Source/AnimNextUncookedOnly/Private/Graph/AnimGraph/AnimGraphNodeBinding_AnimNextParameters.cpp
+++ b/Engine/Plugins/Experimental/AnimNext/Source/AnimNextUncookedOnly/Private/Graph/AnimGraph/AnimGraphNodeBinding_AnimNextParameters.cpp
@@ -517,7 +517,7 @@ TSharedRef<SWidget> UAnimGraphNodeBinding_AnimNextParameters::MakePropertyBindin
 									if (ComparisonName == FName(Iter.Key(), 0))
 									{
 										Iter.RemoveCurrent();
-									}
+									} 
 								}
 							}
 							else if (bIsArrayElement)
@@ -536,7 +536,7 @@ TSharedRef<SWidget> UAnimGraphNodeBinding_AnimNextParameters::MakePropertyBindin
 							FAnimNextAnimGraphNodeParameterBinding NewBinding;
 							NewBinding.BindingName = InArgs.BindingName;
 							NewBinding.ParameterName = InBinding.Parameter;
-							NewBinding.CachedParameterNameText = UncookedOnly::FUtils::GetParameterDisplayNameText(NewBinding.ParameterName);
+							NewBinding.CachedParameterNameText = UncookedOnly::FUtils::GetParameterDisplayNameText(NewBinding.ParameterName, InBinding.InstanceId);
 							NewBinding.ArrayIndex = PinArrayIndex;
 							Binding->PropertyBindings.Add(InArgs.BindingName, NewBinding);
 						}
@@ -550,7 +550,6 @@ TSharedRef<SWidget> UAnimGraphNodeBinding_AnimNextParameters::MakePropertyBindin
 				});
 				Args.NewParameterType = FParamTypeHandle::FromProperty(PropertyToBindTo).GetType();
 				Args.bMultiSelect = false;
-				Args.bShowSourceGraph = false;
 			
 				MenuBuilder.AddWidget(
 					SNew(SBox)
diff --git a/Engine/Plugins/Experimental/AnimNext/Source/AnimNextUncookedOnly/Private/Graph/AnimNextGraph_AnimationGraph.cpp b/Engine/Plugins/Experimental/AnimNext/Source/AnimNextUncookedOnly/Private/Graph/AnimNextGraph_AnimationGraph.cpp
--- a/Engine/Plugins/Experimental/AnimNext/Source/AnimNextUncookedOnly/Private/Graph/AnimNextGraph_AnimationGraph.cpp
+++ b/Engine/Plugins/Experimental/AnimNext/Source/AnimNextUncookedOnly/Private/Graph/AnimNextGraph_AnimationGraph.cpp
@@ -1,8 +1,27 @@
 ﻿// Copyright Epic Games, Inc. All Rights Reserved.
 
 #include "Graph/AnimNextGraph_AnimationGraph.h"
+
+#include "AnimNextRigVMAsset.h"
+#include "UncookedOnlyUtils.h"
 #include "Graph/AnimNextGraph_EdGraph.h"
 #include "Graph/RigUnit_AnimNextGraphRoot.h"
+#include "Param/AnimNextTag.h"
+#include "Param/ParamType.h"
+
+FAnimNextParamType UAnimNextGraph_AnimationGraph::GetExportType() const
+{
+	return FAnimNextParamType::GetType<FAnimNextEntryPoint>();
+}
+
+FName UAnimNextGraph_AnimationGraph::GetExportName() const
+{
+	if(UAnimNextRigVMAsset* OuterAsset = GetTypedOuter<UAnimNextRigVMAsset>())
+	{
+		return UE::AnimNext::UncookedOnly::FUtils::GetQualifiedName(OuterAsset, GraphName);
+	}
+	return GraphName;
+}
 
 FName UAnimNextGraph_AnimationGraph::GetEntryName() const
 {
diff --git a/Engine/Plugins/Experimental/AnimNext/Source/AnimNextUncookedOnly/Private/Graph/AnimNextGraph_AnimationGraphSchema.cpp b/Engine/Plugins/Experimental/AnimNext/Source/AnimNextUncookedOnly/Private/Graph/AnimNextGraph_AnimationGraphSchema.cpp
--- a/Engine/Plugins/Experimental/AnimNext/Source/AnimNextUncookedOnly/Private/Graph/AnimNextGraph_AnimationGraphSchema.cpp
+++ b/Engine/Plugins/Experimental/AnimNext/Source/AnimNextUncookedOnly/Private/Graph/AnimNextGraph_AnimationGraphSchema.cpp
@@ -5,6 +5,7 @@
 #include "AnimNextExecuteContext.h"
 #include "Graph/RigUnit_AnimNextBase.h"
 #include "Param/RigVMDispatch_GetParameter.h"
+#include "Param/RigVMDispatch_GetScopedParameter.h"
 
 bool UAnimNextGraph_AnimationGraphSchema::SupportsUnitFunction(URigVMController* InController, const FRigVMFunction* InUnitFunction) const
 {
@@ -34,8 +35,10 @@ bool UAnimNextGraph_AnimationGraphSchema::SupportsDispatchFactory(URigVMControll
 		{
 			if(DispatchExecuteContextStruct == FAnimNextExecuteContext::StaticStruct())
 			{
-				// We only support the FRigVMDispatch_GetParameter at the moment.
-				return (InDispatchFactory->GetScriptStruct() == FRigVMDispatch_GetParameter::StaticStruct());
+				// We only support the FRigVMDispatch_GetParameter/FRigVMDispatch_GetScopedParameter at the moment.
+				return
+				   (InDispatchFactory->GetScriptStruct() == FRigVMDispatch_GetParameter::StaticStruct() ||
+					InDispatchFactory->GetScriptStruct() == FRigVMDispatch_GetScopedParameter::StaticStruct());
 			}
 		}
 	}
