AnimNext: Parameters rework - UOL integration

diff --git a/Engine/Plugins/Experimental/AnimNext/Source/AnimNextUncookedOnly/Private/Graph/AnimNextGraph_Controller.cpp b/Engine/Plugins/Experimental/AnimNext/Source/AnimNextUncookedOnly/Private/Graph/AnimNextGraph_Controller.cpp
--- a/Engine/Plugins/Experimental/AnimNext/Source/AnimNextUncookedOnly/Private/Graph/AnimNextGraph_Controller.cpp
+++ b/Engine/Plugins/Experimental/AnimNext/Source/AnimNextUncookedOnly/Private/Graph/AnimNextGraph_Controller.cpp
@@ -2,6 +2,14 @@
 
 #include "Graph/AnimNextGraph_Controller.h"
 #include "Graph/AnimNextGraph_UnitNode.h"
+#include "Param/AnimNextEditorParam.h"
+#include "Param/AnimNextParam.h"
+#include "Param/AnimNextParamInstanceIdentifier.h"
+#include "Param/ParamType.h"
+#include "Param/RigVMDispatch_GetLayerParameter.h"
+#include "Param/RigVMDispatch_GetParameter.h"
+#include "Param/RigVMDispatch_GetScopedParameter.h"
+#include "Param/RigVMDispatch_SetLayerParameter.h"
 
 URigVMUnitNode* UAnimNextGraph_Controller::AddUnitNodeWithPins(UScriptStruct* InScriptStruct, const FRigVMPinInfoArray& PinArray, const FName& InMethodName, const FVector2D& InPosition, const FString& InNodeName, bool bSetupUndoRedo, bool bPrintPythonCommand)
 {
@@ -47,3 +55,225 @@ URigVMUnitNode* UAnimNextGraph_Controller::AddUnitNodeWithPins(UScriptStruct* In
 
 	return Node;
 }
+
+bool UAnimNextGraph_Controller::SetAnimNextParameterNode(URigVMNode* ParameterNode, FName ParameterName, EPropertyBagPropertyType ValueType, EPropertyBagContainerType ContainerType, const UObject* ValueTypeObject, const TInstancedStruct<FAnimNextParamInstanceIdentifier>& InstanceId, bool bSetupUndoRedo, bool bPrintPythonCommand)
+{
+	return SetAnimNextParameterNode(ParameterNode, ParameterName, FAnimNextParamType(ValueType, ContainerType, ValueTypeObject), InstanceId, bSetupUndoRedo, bPrintPythonCommand);
+}
+
+bool UAnimNextGraph_Controller::SetAnimNextParameterNode(URigVMNode* ParameterNode, FName ParameterName, const FAnimNextParamType& Type, const TInstancedStruct<FAnimNextParamInstanceIdentifier>& InstanceId, bool bSetupUndoRedo, bool bPrintPythonCommand)
+{
+	if(!IsValidGraph())
+	{
+		return false;
+	}
+
+	if (!IsTransacting() && !IsGraphEditable())
+	{
+		return false;
+	}
+
+	if(ParameterNode == nullptr)
+	{
+		ReportError(TEXT("Invalid node."));
+		return false;
+	}
+
+	const URigVMTemplateNode* TemplateNode = Cast<URigVMTemplateNode>(ParameterNode);
+	if (TemplateNode == nullptr)
+	{
+		ReportError(TEXT("Not a template node."));
+		return false;
+	}
+
+	const FRigVMDispatchFactory* GetParameterFactory = FRigVMRegistry::Get().FindOrAddDispatchFactory(FRigVMDispatch_GetParameter::StaticStruct());
+	const FName GetParameterNotation = GetParameterFactory->GetTemplate()->GetNotation();
+	const FRigVMDispatchFactory* GetScopedParameterFactory = FRigVMRegistry::Get().FindOrAddDispatchFactory(FRigVMDispatch_GetScopedParameter::StaticStruct());
+	const FName GetScopedParameterNotation = GetScopedParameterFactory->GetTemplate()->GetNotation();
+	const FRigVMDispatchFactory* GetLayerParameterFactory = FRigVMRegistry::Get().FindOrAddDispatchFactory(FRigVMDispatch_GetLayerParameter::StaticStruct());
+	const FName GetLayerParameterNotation = GetLayerParameterFactory->GetTemplate()->GetNotation();
+	const FRigVMDispatchFactory* SetLayerParameterFactory = FRigVMRegistry::Get().FindOrAddDispatchFactory(FRigVMDispatch_SetLayerParameter::StaticStruct());
+	const FName SetLayerParameterNotation = SetLayerParameterFactory->GetTemplate()->GetNotation();
+
+	if( TemplateNode->GetNotation() != GetParameterNotation &&
+		TemplateNode->GetNotation() != GetScopedParameterNotation &&
+		TemplateNode->GetNotation() != GetLayerParameterNotation &&
+		TemplateNode->GetNotation() != SetLayerParameterNotation)
+	{
+		ReportError(TEXT("Not a parameter node."));
+		return false;
+	}
+
+	const bool bIsNamedParam = TemplateNode->GetNotation() == GetLayerParameterNotation || TemplateNode->GetNotation() == SetLayerParameterNotation;
+
+	if(ParameterName != NAME_None)
+	{
+		const FSoftObjectPath SoftObjectPath(ParameterName.ToString());
+		if(!SoftObjectPath.GetAssetPath().IsValid() || SoftObjectPath.GetSubPathString().Len() == 0)
+		{
+			ReportError(TEXT("InParameterName is an invalid format. Should be /AssetOrFieldPath/ClassOrAsset.ClassOrAsset:FieldOrParameter."));
+			return false;
+		}
+	}
+
+	if (!Type.IsNone() && !Type.IsValid())
+	{
+		ReportError(TEXT("Type is invalid."));
+		return false;
+	}
+
+	URigVMPin* ParameterPin = ParameterNode->FindPin(FRigVMDispatch_GetParameter::ParameterName.ToString());
+	if(ParameterPin == nullptr)
+	{
+		ReportError(TEXT("Parameter pin not found."));
+		return false;
+	}
+
+	if(bSetupUndoRedo)
+	{
+		OpenUndoBracket(TEXT("Set parameter"));
+	}
+
+	FString ValueAsString;
+	if(bIsNamedParam)
+	{
+		ValueAsString = ParameterName.ToString();
+	}
+	else
+	{
+		FAnimNextParam ParamValue(ParameterName, Type, InstanceId);
+		FAnimNextParam::StaticStruct()->ExportText(ValueAsString, &ParamValue, nullptr, nullptr, PPF_None, nullptr);
+	}
+
+	if(!SetPinDefaultValue(ParameterPin->GetPinPath(), ValueAsString, true, bSetupUndoRedo, true, bPrintPythonCommand))
+	{
+		if(bSetupUndoRedo)
+		{
+			CancelUndoBracket();
+		}
+		return false;
+	}
+
+	FRigVMTemplateArgumentType RigVMType = Type.ToRigVMTemplateArgument();
+	URigVMPin* OutputPin = ParameterNode->FindPin(FRigVMDispatch_GetParameter::ValueName.ToString());
+	if(RigVMType.IsValid() && OutputPin)
+	{
+		// Re-resolve the node's output
+		TArray<URigVMLink*> Links = OutputPin->GetLinks();
+		if(!TemplateNode->IsFullyUnresolved())
+		{
+			if(!UnresolveTemplateNodes({ ParameterNode->GetFName() }, bSetupUndoRedo, bPrintPythonCommand))
+			{
+				if(bSetupUndoRedo)
+				{
+					CancelUndoBracket();
+				}
+				return false;
+			}
+		}
+
+		if(TemplateNode->IsFullyUnresolved())
+		{
+			if(!ResolveWildCardPin(OutputPin, RigVMType, bSetupUndoRedo, bPrintPythonCommand))
+			{
+				if(bSetupUndoRedo)
+				{
+					CancelUndoBracket();
+				}
+				return false;
+			}
+		}
+
+		// Try to restore links
+		for(URigVMLink* Link : Links)
+		{
+			if(!AddLink(OutputPin, Link->GetOppositePin(OutputPin), bSetupUndoRedo))
+			{
+				if(bSetupUndoRedo)
+				{
+					CancelUndoBracket();
+				}
+				return false;
+			}
+		}
+	}
+
+	if(bSetupUndoRedo)
+	{
+		CloseUndoBracket();
+	}
+
+	return true;
+}
+
+URigVMNode* UAnimNextGraph_Controller::AddGetAnimNextParameterNode(const FVector2D& Position, FName ParameterName, EPropertyBagPropertyType ValueType, EPropertyBagContainerType ContainerType, const UObject* ValueTypeObject, const TInstancedStruct<FAnimNextParamInstanceIdentifier>& InstanceId, bool bSetupUndoRedo, bool bPrintPythonCommand)
+{
+	return AddGetAnimNextParameterNode(Position, ParameterName, FAnimNextParamType(ValueType, ContainerType, ValueTypeObject), InstanceId, bSetupUndoRedo, bPrintPythonCommand);
+}
+
+URigVMNode* UAnimNextGraph_Controller::AddGetAnimNextParameterNode(const FVector2D& Position, FName ParameterName, const FAnimNextParamType& Type, const TInstancedStruct<FAnimNextParamInstanceIdentifier>& InstanceId, bool bSetupUndoRedo, bool bPrintPythonCommand)
+{
+	return AddGetAnimNextParameterNodeInternal(FRigVMDispatch_GetScopedParameter::StaticStruct(), Position, ParameterName, Type, InstanceId, bSetupUndoRedo, bPrintPythonCommand);
+}
+
+URigVMNode* UAnimNextGraph_Controller::AddGetAnimNextGraphParameterNode(const FVector2D& Position, FName ParameterName, EPropertyBagPropertyType ValueType, EPropertyBagContainerType ContainerType, const UObject* ValueTypeObject, const TInstancedStruct<FAnimNextParamInstanceIdentifier>& InstanceId, bool bSetupUndoRedo, bool bPrintPythonCommand)
+{
+	return AddGetAnimNextGraphParameterNode(Position, ParameterName, FAnimNextParamType(ValueType, ContainerType, ValueTypeObject), InstanceId, bSetupUndoRedo, bPrintPythonCommand);
+}
+
+URigVMNode* UAnimNextGraph_Controller::AddGetAnimNextGraphParameterNode(const FVector2D& Position, FName ParameterName, const FAnimNextParamType& Type, const TInstancedStruct<FAnimNextParamInstanceIdentifier>& InstanceId, bool bSetupUndoRedo, bool bPrintPythonCommand)
+{
+return AddGetAnimNextParameterNodeInternal(FRigVMDispatch_GetLayerParameter::StaticStruct(), Position, ParameterName, Type, InstanceId, bSetupUndoRedo, bPrintPythonCommand);
+}
+
+URigVMNode* UAnimNextGraph_Controller::AddSetAnimNextGraphParameterNode(const FVector2D& Position, FName ParameterName, EPropertyBagPropertyType ValueType, EPropertyBagContainerType ContainerType, const UObject* ValueTypeObject, const TInstancedStruct<FAnimNextParamInstanceIdentifier>& InstanceId, bool bSetupUndoRedo, bool bPrintPythonCommand)
+{
+	return AddSetAnimNextGraphParameterNode(Position, ParameterName, FAnimNextParamType(ValueType, ContainerType, ValueTypeObject), InstanceId, bSetupUndoRedo, bPrintPythonCommand);
+}
+
+URigVMNode* UAnimNextGraph_Controller::AddSetAnimNextGraphParameterNode(const FVector2D& Position, FName ParameterName, const FAnimNextParamType& Type, const TInstancedStruct<FAnimNextParamInstanceIdentifier>& InstanceId, bool bSetupUndoRedo, bool bPrintPythonCommand)
+{
+	return AddGetAnimNextParameterNodeInternal(FRigVMDispatch_SetLayerParameter::StaticStruct(), Position, ParameterName, Type, InstanceId, bSetupUndoRedo, bPrintPythonCommand);
+}
+
+URigVMNode* UAnimNextGraph_Controller::AddGetAnimNextParameterNodeInternal(UScriptStruct* NodeStruct, const FVector2D& Position, FName ParameterName, const FAnimNextParamType& Type, const TInstancedStruct<FAnimNextParamInstanceIdentifier>& InstanceId, bool bSetupUndoRedo, bool bPrintPythonCommand)
+{
+	if(!IsValidGraph())
+	{
+		return nullptr;
+	}
+
+	URigVMGraph* Graph = GetGraph();
+	check(Graph);
+
+	if(bSetupUndoRedo)
+	{
+		OpenUndoBracket(TEXT("Add parameter node"));
+	}
+
+	const FString Name = GetSchema()->GetValidNodeName(Graph, TEXT("ParameterNode"));
+	const FRigVMDispatchFactory* Factory = FRigVMRegistry::Get().FindOrAddDispatchFactory(NodeStruct);
+	URigVMNode* Node = AddTemplateNode(Factory->GetTemplate()->GetNotation(), Position, Name, bSetupUndoRedo, bPrintPythonCommand);
+	if(Node)
+	{
+		if(!SetAnimNextParameterNode(Node, ParameterName, Type, InstanceId, bSetupUndoRedo, bPrintPythonCommand))
+		{
+			if(bSetupUndoRedo)
+			{
+				CancelUndoBracket();
+			}
+			return nullptr;
+		}
+	}
+	else if(bSetupUndoRedo)
+	{
+		CancelUndoBracket();
+	}
+
+	if(bSetupUndoRedo)
+	{
+		CloseUndoBracket();
+	}
+
+	return Node;
+}
\ No newline at end of file
diff --git a/Engine/Plugins/Experimental/AnimNext/Source/AnimNextUncookedOnly/Private/Graph/AnimNextGraph_EdGraphSchema.cpp b/Engine/Plugins/Experimental/AnimNext/Source/AnimNextUncookedOnly/Private/Graph/AnimNextGraph_EdGraphSchema.cpp
--- a/Engine/Plugins/Experimental/AnimNext/Source/AnimNextUncookedOnly/Private/Graph/AnimNextGraph_EdGraphSchema.cpp
+++ b/Engine/Plugins/Experimental/AnimNext/Source/AnimNextUncookedOnly/Private/Graph/AnimNextGraph_EdGraphSchema.cpp
@@ -4,6 +4,8 @@
 
 #include "AnimNextRigVMAsset.h"
 #include "AnimNextRigVMAssetEntry.h"
+#include "Param/AnimNextEditorParam.h"
+#include "Param/AnimNextParam.h"
 
 #define LOCTEXT_NAMESPACE "AnimNextGraph_EdGraphSchema"
 
@@ -18,4 +20,13 @@ void UAnimNextGraph_EdGraphSchema::GetGraphDisplayInformation(const UEdGraph& Gr
 	}
 }
 
+bool UAnimNextGraph_EdGraphSchema::IsStructEditable(UStruct* InStruct) const
+{
+	if (InStruct == FAnimNextEditorParam::StaticStruct() || InStruct == FAnimNextParam::StaticStruct())
+	{
+		return true;
+	}
+	return Super::IsStructEditable(InStruct);
+}
+
 #undef LOCTEXT_NAMESPACE
\ No newline at end of file
diff --git a/Engine/Plugins/Experimental/AnimNext/Source/AnimNextUncookedOnly/Private/Graph/AnimNextGraph_EventGraph.cpp b/Engine/Plugins/Experimental/AnimNext/Source/AnimNextUncookedOnly/Private/Graph/AnimNextGraph_EventGraph.cpp
--- a/Engine/Plugins/Experimental/AnimNext/Source/AnimNextUncookedOnly/Private/Graph/AnimNextGraph_EventGraph.cpp
+++ b/Engine/Plugins/Experimental/AnimNext/Source/AnimNextUncookedOnly/Private/Graph/AnimNextGraph_EventGraph.cpp
@@ -1,7 +1,26 @@
 ﻿// Copyright Epic Games, Inc. All Rights Reserved.
 
 #include "Graph/AnimNextGraph_EventGraph.h"
+
+#include "AnimNextRigVMAsset.h"
+#include "UncookedOnlyUtils.h"
 #include "Graph/AnimNextGraph_EdGraph.h"
+#include "Param/AnimNextTag.h"
+#include "Param/ParamType.h"
+
+FAnimNextParamType UAnimNextGraph_EventGraph::GetExportType() const
+{
+	return FAnimNextParamType::GetType<FAnimNextEntryPoint>();
+}
+
+FName UAnimNextGraph_EventGraph::GetExportName() const
+{
+	if(UAnimNextRigVMAsset* OuterAsset = GetTypedOuter<UAnimNextRigVMAsset>())
+	{
+		return UE::AnimNext::UncookedOnly::FUtils::GetQualifiedName(OuterAsset, GraphName);
+	}
+	return GraphName;
+}
 
 FText UAnimNextGraph_EventGraph::GetDisplayName() const
 {
diff --git a/Engine/Plugins/Experimental/AnimNext/Source/AnimNextUncookedOnly/Private/Graph/AnimNextGraph_Parameter.cpp b/Engine/Plugins/Experimental/AnimNext/Source/AnimNextUncookedOnly/Private/Graph/AnimNextGraph_Parameter.cpp
--- a/Engine/Plugins/Experimental/AnimNext/Source/AnimNextUncookedOnly/Private/Graph/AnimNextGraph_Parameter.cpp
+++ b/Engine/Plugins/Experimental/AnimNext/Source/AnimNextUncookedOnly/Private/Graph/AnimNextGraph_Parameter.cpp
@@ -1,20 +1,23 @@
 // Copyright Epic Games, Inc. All Rights Reserved.
 
 #include "Graph/AnimNextGraph_Parameter.h"
+
+#include "UncookedOnlyUtils.h"
 #include "Graph/AnimNextGraph.h"
-#include "Param/ExternalParameterRegistry.h"
+#include "Param/AnimNextTag.h"
 
-FAnimNextParamType UAnimNextGraph_Parameter::GetParamType() const
+FAnimNextParamType UAnimNextGraph_Parameter::GetExportType() const
 {
-	using namespace UE::AnimNext;
+	return GetParamType();
+}
 
-	// Look in built-in parameters
-	IParameterSourceFactory::FParameterInfo Info;
-	if(FExternalParameterRegistry::FindParameterInfo(ParameterName, Info))
-	{
-		return Info.Type;
-	}
+FName UAnimNextGraph_Parameter::GetExportName() const
+{
+	return GetParamName();
+}
 
+FAnimNextParamType UAnimNextGraph_Parameter::GetParamType() const
+{
 	return Type;
 }
 
@@ -37,6 +40,20 @@ bool UAnimNextGraph_Parameter::SetParamType(const FAnimNextParamType& InType, bo
 	return true;
 }
 
+FName UAnimNextGraph_Parameter::GetParamName() const
+{
+	if(UAnimNextRigVMAsset* OuterAsset = GetTypedOuter<UAnimNextRigVMAsset>())
+	{
+		return UE::AnimNext::UncookedOnly::FUtils::GetQualifiedName(OuterAsset, ParameterName);
+	}
+	return ParameterName;
+}
+
+void UAnimNextGraph_Parameter::SetParamName(FName InName, bool bSetupUndoRedo)
+{
+	SetEntryName(InName, bSetupUndoRedo);
+}
+
 FInstancedPropertyBag& UAnimNextGraph_Parameter::GetPropertyBag() const
 {
 	// TODO: move property bag for defaults onto this entry!
@@ -63,13 +80,5 @@ FText UAnimNextGraph_Parameter::GetDisplayName() const
 
 FText UAnimNextGraph_Parameter::GetDisplayNameTooltip() const
 {
-	using namespace UE::AnimNext;
-
-	IParameterSourceFactory::FParameterInfo Info;
-	if(FExternalParameterRegistry::FindParameterInfo(ParameterName, Info))
-	{
-		return Info.Tooltip;
-	}
-	
-	return FText::FromName(ParameterName);
+	return FText::FromString(Comment);
 }
\ No newline at end of file
diff --git a/Engine/Plugins/Experimental/AnimNext/Source/AnimNextUncookedOnly/Private/Param/ObjectProxyType.cpp b/Engine/Plugins/Experimental/AnimNext/Source/AnimNextUncookedOnly/Private/Param/ObjectProxyType.cpp
new file mode 100644
--- /dev/null
+++ b/Engine/Plugins/Experimental/AnimNext/Source/AnimNextUncookedOnly/Private/Param/ObjectProxyType.cpp
@@ -0,0 +1,207 @@
+﻿// Copyright Epic Games, Inc. All Rights Reserved.
+
+#include "ObjectProxyType.h"
+
+#include "IAnimNextModule.h"
+#include "IUniversalObjectLocatorEditorModule.h"
+#include "UniversalObjectLocatorEditor.h"
+#include "UniversalObjectLocatorFragmentType.h"
+#include "Component/AnimNextComponent.h"
+#include "Modules/ModuleManager.h"
+#include "Param/AnimNextParamUniversalObjectLocator.h"
+#include "Param/ClassProxy.h"
+#include "Param/ObjectProxyFactory.h"
+
+#define LOCTEXT_NAMESPACE "ObjectProxyType"
+
+namespace UE::AnimNext::UncookedOnly
+{
+
+const UStruct* FObjectProxyType::GetStruct(const TInstancedStruct<FAnimNextParamInstanceIdentifier>& InInstanceId) const
+{
+	using namespace UE::UniversalObjectLocator;
+
+	const FAnimNextParamUniversalObjectLocator* Locator = InInstanceId.GetPtr<FAnimNextParamUniversalObjectLocator>();
+	if(Locator == nullptr)
+	{
+		return nullptr;
+	}
+
+	IUniversalObjectLocatorEditorModule& UolEditorModule = FModuleManager::LoadModuleChecked<IUniversalObjectLocatorEditorModule>("UniversalObjectLocatorEditor");
+	const FFragmentType* FragmentType = Locator->Locator.GetLastFragmentType();
+	if(FragmentType == nullptr)
+	{
+		return nullptr;
+	}
+	
+	TSharedPtr<ILocatorFragmentEditor> LocatorEditor = UolEditorModule.FindLocatorEditor(FragmentType->PrimaryEditorType);
+	if(!LocatorEditor.IsValid())
+	{
+		return nullptr;
+	}
+
+	// TODO: This context needs to defer to project/schedule/workspace defaults similar to SParameterPicker/FAnimNextLocatorContext
+	UObject* Context = UAnimNextComponent::StaticClass()->GetDefaultObject();
+	return LocatorEditor->ResolveClass(*Locator->Locator.GetLastFragment(), Context);
+}
+
+FText FObjectProxyType::GetDisplayText(const TInstancedStruct<FAnimNextParamInstanceIdentifier>& InInstanceId) const
+{
+	using namespace UE::UniversalObjectLocator;
+	
+	const FAnimNextParamUniversalObjectLocator* Locator = InInstanceId.GetPtr<FAnimNextParamUniversalObjectLocator>();
+	if(Locator == nullptr)
+	{
+		return FText::GetEmpty();
+	}
+
+	TStringBuilder<256> StringBuilder;
+	IUniversalObjectLocatorEditorModule& UolEditorModule = FModuleManager::LoadModuleChecked<IUniversalObjectLocatorEditorModule>("UniversalObjectLocatorEditor");
+	Locator->Locator.ForEachFragment([&UolEditorModule, &StringBuilder](int32 FragmentIndex, int32 NumFragments, const FUniversalObjectLocatorFragment& InFragment)
+	{
+		if(const FFragmentType* FragmentType = InFragment.GetFragmentType())
+		{
+			TSharedPtr<ILocatorFragmentEditor> LocatorEditor = UolEditorModule.FindLocatorEditor(FragmentType->PrimaryEditorType);
+			if(LocatorEditor.IsValid())
+			{
+				FText Text = LocatorEditor->GetDisplayText(&InFragment);
+				if(FragmentIndex != 0)
+				{
+					StringBuilder.Append(TEXT("."));
+				}
+				StringBuilder.Append(Text.ToString());
+				return true;
+			}
+		}
+
+		return false;
+	});
+
+	return FText::FromStringView(StringBuilder);
+}
+
+FText FObjectProxyType::GetTooltipText(const TInstancedStruct<FAnimNextParamInstanceIdentifier>& InInstanceId) const
+{
+	using namespace UE::UniversalObjectLocator;
+	
+	const FAnimNextParamUniversalObjectLocator* Locator = InInstanceId.GetPtr<FAnimNextParamUniversalObjectLocator>();
+	if(Locator == nullptr)
+	{
+		return FText::GetEmpty();
+	}
+
+	TStringBuilder<256> StringBuilder;
+	IUniversalObjectLocatorEditorModule& UolEditorModule = FModuleManager::LoadModuleChecked<IUniversalObjectLocatorEditorModule>("UniversalObjectLocatorEditor");
+	bool bSuccess = Locator->Locator.ForEachFragment([&UolEditorModule, &StringBuilder](int32 FragmentIndex, int32 NumFragments, const FUniversalObjectLocatorFragment& InFragment)
+	{
+		if(const FFragmentType* FragmentType = InFragment.GetFragmentType())
+		{
+			TSharedPtr<ILocatorFragmentEditor> LocatorEditor = UolEditorModule.FindLocatorEditor(FragmentType->PrimaryEditorType);
+			if(LocatorEditor.IsValid())
+			{
+				FText Text = LocatorEditor->GetDisplayText(&InFragment);
+				if(FragmentIndex != 0)
+				{
+					StringBuilder.Append(TEXT("."));
+				}
+				StringBuilder.Append(Text.ToString());
+				return true;
+			}
+		}
+
+		return false;
+	});
+
+	FTextBuilder TextBuilder;
+	if(bSuccess)
+	{
+		TextBuilder.AppendLine(FText::Format(LOCTEXT("ParameterInstanceTooltipFormat", "Instance: {0}"), FText::FromStringView(StringBuilder)));
+	}
+
+	TStringBuilder<256> ScopeStringBuilder;
+	Locator->Locator.ToString(ScopeStringBuilder);
+	TextBuilder.AppendLine(FText::Format(LOCTEXT("ParameterUOLTooltipFormat", "UOL: {0}"), FText::FromStringView(ScopeStringBuilder)));
+
+	return TextBuilder.ToText();
+}
+
+bool FObjectProxyType::FindParameterInfo(const TInstancedStruct<FAnimNextParamInstanceIdentifier>& InInstanceId, TConstArrayView<FName> InParameterNames, TArrayView<FParameterSourceInfo> OutInfo) const
+{
+	check(InParameterNames.Num() == OutInfo.Num());
+
+	const UClass* Class = Cast<UClass>(GetStruct(InInstanceId));
+	if(Class == nullptr)
+	{
+		return false;
+	}
+
+	IAnimNextModule& AnimNextModule = FModuleManager::GetModuleChecked<IAnimNextModule>("AnimNext");
+	TSharedPtr<FObjectProxyFactory> ObjectProxyFactory = StaticCastSharedPtr<FObjectProxyFactory>(AnimNextModule.FindParameterSourceFactory("ObjectProxy"));
+	if(!ObjectProxyFactory.IsValid())
+	{
+		return false;
+	}
+
+	// Extract parameter info if found
+	bool bFound = false;
+	TSharedRef<FClassProxy> ClassProxy = ObjectProxyFactory->FindOrCreateClassProxy(Class);
+	for(int32 ParameterIndex = 0; ParameterIndex < InParameterNames.Num(); ++ParameterIndex)
+	{
+		FName ParameterName = InParameterNames[ParameterIndex];
+		if(const int32* FoundIndexPtr = ClassProxy->ParameterNameMap.Find(ParameterName))
+		{
+			const FClassProxyParameter& FoundParameter = ClassProxy->Parameters[*FoundIndexPtr];
+			FParameterSourceInfo& ParameterInfo = OutInfo[ParameterIndex];
+			ParameterInfo.Type = FoundParameter.Type;
+			ParameterInfo.DisplayName = FoundParameter.DisplayName;
+			ParameterInfo.Tooltip = FoundParameter.Tooltip;
+			ParameterInfo.Function = FoundParameter.Function.Get();
+			ParameterInfo.Property = FoundParameter.Property.Get();
+			ParameterInfo.bThreadSafe = FoundParameter.bThreadSafe;
+			bFound = true;
+		}
+		else
+		{
+			OutInfo[ParameterIndex] = FParameterSourceInfo();
+		}
+	}
+
+	return bFound;
+}
+
+void FObjectProxyType::ForEachParameter(const TInstancedStruct<FAnimNextParamInstanceIdentifier>& InInstanceId, TFunctionRef<void(FName, const FParameterSourceInfo&)> InFunction) const
+{
+	using namespace UE::UniversalObjectLocator;
+
+	const UClass* Class = Cast<UClass>(GetStruct(InInstanceId));
+	if(Class == nullptr)
+	{
+		return;
+	}
+
+	IAnimNextModule& AnimNextModule = FModuleManager::GetModuleChecked<IAnimNextModule>("AnimNext");
+	TSharedPtr<FObjectProxyFactory> ObjectProxyFactory = StaticCastSharedPtr<FObjectProxyFactory>(AnimNextModule.FindParameterSourceFactory("ObjectProxy"));
+	if(!ObjectProxyFactory.IsValid())
+	{
+		return;
+	}
+
+	// Extract parameter info if found
+	TSharedRef<FClassProxy> ClassProxy = ObjectProxyFactory->FindOrCreateClassProxy(Class);
+	for(const FClassProxyParameter& Parameter : ClassProxy->Parameters)
+	{
+		FParameterSourceInfo SourceInfo;
+		SourceInfo.Type = Parameter.Type;
+		SourceInfo.DisplayName = Parameter.DisplayName;
+		SourceInfo.Tooltip = Parameter.Tooltip;
+		SourceInfo.Function = Parameter.Function.Get();
+		SourceInfo.Property = Parameter.Property.Get();
+		SourceInfo.bThreadSafe = Parameter.bThreadSafe;
+
+		InFunction(Parameter.ParameterName, SourceInfo);
+	}
+}
+
+}
+
+#undef LOCTEXT_NAMESPACE
\ No newline at end of file
diff --git a/Engine/Plugins/Experimental/AnimNext/Source/AnimNextUncookedOnly/Private/Param/ObjectProxyType.h b/Engine/Plugins/Experimental/AnimNext/Source/AnimNextUncookedOnly/Private/Param/ObjectProxyType.h
new file mode 100644
--- /dev/null
+++ b/Engine/Plugins/Experimental/AnimNext/Source/AnimNextUncookedOnly/Private/Param/ObjectProxyType.h
@@ -0,0 +1,22 @@
+﻿// Copyright Epic Games, Inc. All Rights Reserved.
+
+#pragma once
+
+#include "CoreMinimal.h"
+#include "Param/IParameterSourceType.h"
+
+namespace UE::AnimNext::UncookedOnly
+{
+
+// Provides information about object proxy parameter sources
+class FObjectProxyType : public IParameterSourceType
+{
+	// IParameterSourceType interface
+	virtual const UStruct* GetStruct(const TInstancedStruct<FAnimNextParamInstanceIdentifier>& InInstanceId) const override;
+	virtual FText GetDisplayText(const TInstancedStruct<FAnimNextParamInstanceIdentifier>& InInstanceId) const override;
+	virtual FText GetTooltipText(const TInstancedStruct<FAnimNextParamInstanceIdentifier>& InInstanceId) const override;
+	virtual bool FindParameterInfo(const TInstancedStruct<FAnimNextParamInstanceIdentifier>& InInstanceId, TConstArrayView<FName> InParameterNames, TArrayView<FParameterSourceInfo> OutInfo) const override;
+	virtual void ForEachParameter(const TInstancedStruct<FAnimNextParamInstanceIdentifier>& InInstanceId, TFunctionRef<void(FName, const FParameterSourceInfo&)> InFunction) const override;
+};
+
+}
\ No newline at end of file
diff --git a/Engine/Plugins/Experimental/AnimNext/Source/AnimNextUncookedOnly/Private/UncookedOnlyUtils.cpp b/Engine/Plugins/Experimental/AnimNext/Source/AnimNextUncookedOnly/Private/UncookedOnlyUtils.cpp
--- a/Engine/Plugins/Experimental/AnimNext/Source/AnimNextUncookedOnly/Private/UncookedOnlyUtils.cpp
+++ b/Engine/Plugins/Experimental/AnimNext/Source/AnimNextUncookedOnly/Private/UncookedOnlyUtils.cpp
@@ -35,15 +35,22 @@
 #include "Param/RigVMDispatch_GetLayerParameter.h"
 #include "RigVMCompiler/RigVMCompiler.h"
 #include "RigVMCore/RigVM.h"
-#include "Param/ExternalParameterRegistry.h"
-#include "Scheduler/AnimNextExternalTaskBinding.h"
 #include "Scheduler/AnimNextSchedule.h"
 #include "Scheduler/AnimNextSchedulePort.h"
 #include "AnimNextRigVMAsset.h"
 #include "AnimNextRigVMAssetEditorData.h"
 #include "AnimNextRigVMAssetEntry.h"
+#include "AnimNextUncookedOnlyModule.h"
+#include "IAnimNextRigVMExportInterface.h"
+#include "GameFramework/Character.h"
 #include "Graph/AnimNextGraphEntryPoint.h"
 #include "Graph/AnimNextGraph_AnimationGraphSchema.h"
+#include "Logging/StructuredLog.h"
+#include "Param/AnimNextParamInstanceIdentifier.h"
+#include "Param/IParameterSourceType.h"
+#include "Param/RigVMDispatch_GetScopedParameter.h"
+
+#define LOCTEXT_NAMESPACE "AnimNextUncookedOnlyUtils"
 
 namespace UE::AnimNext::UncookedOnly
 {
@@ -94,10 +101,15 @@ namespace Private
 		URigVMNode* RootNode;
 		TArray<FTraitStackMapping> TraitStackNodes;
 
-		explicit FTraitGraph(URigVMNode* InRootNode)
-			: EntryPoint(*InRootNode->FindPin(GET_MEMBER_NAME_STRING_CHECKED(FRigUnit_AnimNextGraphRoot, EntryPoint))->GetDefaultValue())
-			, RootNode(InRootNode)
-		{}
+		explicit FTraitGraph(const UAnimNextGraph* InGraph, URigVMNode* InRootNode)
+			: RootNode(InRootNode)
+		{
+			TStringBuilder<256> StringBuilder;
+			StringBuilder.Append(InGraph->GetPathName());
+			StringBuilder.Append(TEXT(":"));
+			StringBuilder.Append(InRootNode->FindPin(GET_MEMBER_NAME_STRING_CHECKED(FRigUnit_AnimNextGraphRoot, EntryPoint))->GetDefaultValue());
+			EntryPoint = FName(StringBuilder.ToView());
+		}
 	};
 
 	template<typename TraitAction>
@@ -336,7 +348,7 @@ namespace Private
 		}
 	}
 
-	FTraitGraph CollectGraphInfo(const URigVMGraph* VMGraph, URigVMController* VMController)
+	FTraitGraph CollectGraphInfo(const UAnimNextGraph* InGraph,  const URigVMGraph* VMGraph, URigVMController* VMController)
 	{
 		const TArray<URigVMNode*>& VMNodes = VMGraph->GetNodes();
 		URigVMUnitNode* VMRootNode = FindRootNode(VMNodes);
@@ -350,7 +362,7 @@ namespace Private
 		// Make sure we don't have empty input pins
 		AddMissingInputLinks(VMGraph, VMController);
 
-		FTraitGraph TraitGraph(VMRootNode);
+		FTraitGraph TraitGraph(InGraph, VMRootNode);
 
 		TArray<const URigVMNode*> NodesToVisit;
 		NodesToVisit.Add(VMRootNode);
@@ -570,7 +582,7 @@ void FUtils::CompileVM(UAnimNextGraph* InGraph)
 		if(TempGraph->GetSchemaClass() == UAnimNextGraph_AnimationGraphSchema::StaticClass())
 		{
 			// Gather our trait stacks
-			Private::FTraitGraph& TraitGraph = TraitGraphs.Add_GetRef(Private::CollectGraphInfo(TempGraph, TempController->GetControllerForGraph(TempGraph)));
+			Private::FTraitGraph& TraitGraph = TraitGraphs.Add_GetRef(Private::CollectGraphInfo(InGraph, TempGraph, TempController->GetControllerForGraph(TempGraph)));
 			check(!TraitGraph.TraitStackNodes.IsEmpty());
 
 			FAnimNextGraphEntryPoint& EntryPoint = InGraph->EntryPoints.AddDefaulted_GetRef();
@@ -678,17 +690,16 @@ void FUtils::CompileVM(UAnimNextGraph* InGraph)
 
 	// Now that the graph has been re-compiled, re-allocate the previous live instances
 	InGraph->ThawGraphInstances();
-
-	FAssetData AssetData(InGraph);
+	
 	FAnimNextParameterProviderAssetRegistryExports Exports;
-	GetExportedParametersForAsset(AssetData, Exports);
+	GetAssetParameters(EditorData, Exports);
 
 	for(FAnimNextParameterAssetRegistryExportEntry& Entry : Exports.Parameters)
 	{
-		FName ParameterSourceName = FExternalParameterRegistry::FindSourceForParameter(Entry.Name);
-		if(ParameterSourceName != NAME_None)
+		// Required parameters are those that are read in this asset but not bound in this asset as state
+		if(EnumHasAnyFlags(Entry.GetFlags(), EAnimNextParameterFlags::Read) && !EnumHasAnyFlags(Entry.GetFlags(), EAnimNextParameterFlags::Bound))
 		{
-			InGraph->RequiredParameters.Emplace(Entry.Name, Entry.Type);
+			InGraph->RequiredParameters.Emplace(Entry.Name, Entry.Type, Entry.InstanceId);
 		}
 	}
 
@@ -845,11 +856,13 @@ void FUtils::CompileStruct(UAnimNextGraph* InGraph)
 	// Gather all parameters in this asset
 	for(const UAnimNextRigVMAssetEntry* Entry : EditorData->Entries)
 	{
-		if(const IAnimNextRigVMParameterInterface* Binding = Cast<IAnimNextRigVMParameterInterface>(Entry))
+		if(const IAnimNextRigVMParameterInterface* Parameter = Cast<IAnimNextRigVMParameterInterface>(Entry))
 		{
-			const FAnimNextParamType& Type = Binding->GetParamType();
+			const FAnimNextParamType& Type = Parameter->GetParamType();
 			ensure(Type.IsValid());
-			PropertyDescs.Emplace(Entry->GetEntryName(), Type.GetContainerType(), Type.GetValueType(), Type.GetValueTypeObject());
+			const FName Name = Parameter->GetParamName();
+
+			PropertyDescs.Emplace(Name, Type.GetContainerType(), Type.GetValueType(), Type.GetValueTypeObject());
 		}
 	}
 
@@ -1203,44 +1216,104 @@ void FUtils::SetupEventGraph(URigVMController* InController)
 	InController->AddUnitNode(FRigUnit_AnimNextParameterBeginExecution::StaticStruct(), FRigVMStruct::ExecuteName, FVector2D(-200.0f, 0.0f), FString(), false);
 }
 
-FText FUtils::GetParameterDisplayNameText(FName InParameterName)
+FName FUtils::GetParameterNameFromQualifiedName(FName InName)
 {
-	FString NameAsString = InParameterName.ToString();
-	NameAsString.ReplaceCharInline(TEXT('_'), TEXT('.'));
-	return FText::FromString(NameAsString);
+	const FSoftObjectPath SoftObjectPath(InName.ToString());
+	return FName(*SoftObjectPath.GetSubPathString());
 }
 
-FAnimNextParamType FUtils::GetParameterTypeFromName(FName InName)
+FName FUtils::GetQualifiedName(UAnimNextRigVMAsset* InAsset, FName InBaseName)
 {
-	// Check built-in params first as they are cheaper
-	IParameterSourceFactory::FParameterInfo Info;
-	if(FExternalParameterRegistry::FindParameterInfo(InName, Info))
+	if(InAsset)
 	{
-		return Info.Type;
+		TStringBuilder<256> StringBuilder;
+		StringBuilder.Append(InAsset->GetPathName());
+		StringBuilder.Append(TEXT(":"));
+		InBaseName.AppendString(StringBuilder);
+
+		return FName(StringBuilder.ToView());
 	}
+	return InBaseName;
+}
 
-	// Query the asset registry for other params
-	IAssetRegistry& AssetRegistry = FModuleManager::Get().LoadModuleChecked<FAssetRegistryModule>(TEXT("AssetRegistry")).Get();
+FText FUtils::GetParameterDisplayNameText(FName InParameterName, const TInstancedStruct<FAnimNextParamInstanceIdentifier>& InInstanceId)
+{
+	if(InInstanceId.IsValid())
+	{
+		FText ParameterText;
+		FModule& Module = FModuleManager::GetModuleChecked<FModule>("AnimNextUncookedOnly");
+		if(TSharedPtr<IParameterSourceType> SourceType = Module.FindParameterSourceType(InInstanceId.GetScriptStruct()))
+		{
+			ParameterText = SourceType->GetDisplayText(InInstanceId);
+		}
 
-	FAnimNextParameterProviderAssetRegistryExports Exports;
-	GetExportedParametersFromAssetRegistry(Exports);
-	for(const FAnimNextParameterAssetRegistryExportEntry& Export : Exports.Parameters)
+		TStringBuilder<256> StringBuilder;
+		if(!ParameterText.IsEmpty())
+		{
+			StringBuilder.Append(ParameterText.ToString());
+			StringBuilder.Append(TEXT("."));
+		}
+
+		GetParameterNameFromQualifiedName(InParameterName).AppendString(StringBuilder);
+		return FText::FromStringView(StringBuilder);
+	}
+	else
+	{
+		if(InParameterName.IsNone())
+		{
+			return FText::FromName(InParameterName);
+		}
+		else
+		{
+			const FSoftObjectPath SoftObjectPath(InParameterName.ToString());
+			return FText::Format(LOCTEXT("ParameterNameDisplayFormat", "{0}.{1}"), FText::FromString(SoftObjectPath.GetAssetName()), FText::FromString(SoftObjectPath.GetSubPathString()));
+		}
+	}
+}
+
+FText FUtils::GetParameterTooltipText(FName InParameterName, const TInstancedStruct<FAnimNextParamInstanceIdentifier>& InInstanceId)
+{
+	FTextBuilder TextBuilder;
+	TextBuilder.AppendLine(FText::Format(LOCTEXT("ParameterNameTooltipFormat", "Parameter: {0}"), FText::FromString(InParameterName.ToString())));
+
+	if(InInstanceId.IsValid())
 	{
-		if(Export.Name == InName)
+		FModule& Module = FModuleManager::GetModuleChecked<FModule>("AnimNextUncookedOnly");
+		if(TSharedPtr<IParameterSourceType> SourceType = Module.FindParameterSourceType(InInstanceId.GetScriptStruct()))
 		{
-			return Export.Type;
+			TextBuilder.AppendLine(SourceType->GetTooltipText(InInstanceId));
+		}
+	}
+
+	return TextBuilder.ToText();
+}
+
+FAnimNextParamType FUtils::GetParameterTypeFromName(FName InName)
+{
+	// Query the asset registry for other params
+	TMap<FAssetData, FAnimNextParameterProviderAssetRegistryExports> ExportMap;
+	GetExportedParametersFromAssetRegistry(ExportMap);
+	for(const TPair<FAssetData, FAnimNextParameterProviderAssetRegistryExports>& ExportPair : ExportMap)
+	{
+		for(const FAnimNextParameterAssetRegistryExportEntry& Parameter : ExportPair.Value.Parameters)
+		{
+			if(Parameter.Name == InName)
+			{
+				return Parameter.Type;
+			}
 		}
 	}
 
 	return FAnimNextParamType();
 }
+
 bool FUtils::GetExportedParametersForAsset(const FAssetData& InAsset, FAnimNextParameterProviderAssetRegistryExports& OutExports)
 {
 	const FString TagValue = InAsset.GetTagValueRef<FString>(UE::AnimNext::ExportsAnimNextAssetRegistryTag);
 	return FAnimNextParameterProviderAssetRegistryExports::StaticStruct()->ImportText(*TagValue, &OutExports, nullptr, PPF_None, nullptr, FAnimNextParameterProviderAssetRegistryExports::StaticStruct()->GetName()) != nullptr;
 }
 
-bool FUtils::GetExportedParametersFromAssetRegistry(FAnimNextParameterProviderAssetRegistryExports& OutExports)
+bool FUtils::GetExportedParametersFromAssetRegistry(TMap<FAssetData, FAnimNextParameterProviderAssetRegistryExports>& OutExports)
 {
 	TArray<FAssetData> AssetData;
 	IAssetRegistry::GetChecked().GetAssetsByTags({UE::AnimNext::ExportsAnimNextAssetRegistryTag}, AssetData);
@@ -1251,32 +1324,27 @@ bool FUtils::GetExportedParametersFromAssetRegistry(FAnimNextParameterProviderAs
 		FAnimNextParameterProviderAssetRegistryExports AssetExports;
 		if (FAnimNextParameterProviderAssetRegistryExports::StaticStruct()->ImportText(*TagValue, &AssetExports, nullptr, PPF_None, nullptr, FAnimNextParameterProviderAssetRegistryExports::StaticStruct()->GetName()) != nullptr)
 		{
-			for (FAnimNextParameterAssetRegistryExportEntry& Parameter : AssetExports.Parameters)
-			{
-				if (Parameter.Name != NAME_None && Parameter.Type.IsValid() && Parameter.Type.ValueTypeObject != FRigVMUnknownType::StaticStruct())
-				{
-					FAnimNextParameterAssetRegistryExportEntry* ExistingEntry = OutExports.Parameters.FindByPredicate([Name=Parameter.Name](const FAnimNextParameterAssetRegistryExportEntry& Entry)
-					{
-						return Entry.Name == Name;
-					});
-					
-					if (!ExistingEntry)
-					{
-						Parameter.ReferencingAsset = Asset;
-						OutExports.Parameters.Add(Parameter);
-					}
-					else
-					{
-						ensureMsgf(ExistingEntry->Type == Parameter.Type, TEXT("[%s::%s] %s vs [%s::%s] %s"), *ExistingEntry->ReferencingAsset.ToSoftObjectPath().ToString(), *ExistingEntry->Name.ToString(), *ExistingEntry->Type.ToString(), *Asset.ToSoftObjectPath().ToString(), *Parameter.Name.ToString(), *Parameter.Type.ToString());
-
-						ExistingEntry->Flags |= Parameter.Flags;
-					}
-				}
-			}
+			OutExports.Add(Asset, MoveTemp(AssetExports));
 		}
 	}
 
-	return OutExports.Parameters.Num() > 0;
+	return OutExports.Num() > 0;
+}
+
+static void AddParamToSet(const FAnimNextParameterAssetRegistryExportEntry& InNewParam, TSet<FAnimNextParameterAssetRegistryExportEntry>& OutExports)
+{
+	if(FAnimNextParameterAssetRegistryExportEntry* ExistingEntry = OutExports.Find(InNewParam))
+	{
+		if(ExistingEntry->Type != InNewParam.Type)
+		{
+			UE_LOGFMT(LogAnimation, Warning, "Type mismatch between parameter {ParameterName}. {ParamType1} vs {ParamType1}", InNewParam.Name, InNewParam.Type.ToString(), ExistingEntry->Type.ToString());
+		}
+		ExistingEntry->Flags |= InNewParam.Flags;
+	}
+	else
+	{
+		OutExports.Add(InNewParam);
+	}
 }
 
 void FUtils::GetAssetParameters(const UAnimNextRigVMAssetEditorData* EditorData, FAnimNextParameterProviderAssetRegistryExports& OutExports)
@@ -1284,24 +1352,36 @@ void FUtils::GetAssetParameters(const UAnimNextRigVMAssetEditorData* EditorData,
 	OutExports.Parameters.Reset();
 	OutExports.Parameters.Reserve(EditorData->Entries.Num());
 
+	TSet<FAnimNextParameterAssetRegistryExportEntry> ExportSet;
+	GetAssetParameters(EditorData, ExportSet);
+	OutExports.Parameters = ExportSet.Array();
+}
+
+void FUtils::GetAssetParameters(const UAnimNextRigVMAssetEditorData* EditorData, TSet<FAnimNextParameterAssetRegistryExportEntry>& OutExports)
+{
 	for(const UAnimNextRigVMAssetEntry* Entry : EditorData->Entries)
 	{
-		if(const IAnimNextRigVMParameterInterface* ParameterInterface = Cast<IAnimNextRigVMParameterInterface>(Entry))
+		if(const IAnimNextRigVMExportInterface* ExportInterface = Cast<IAnimNextRigVMExportInterface>(Entry))
 		{
 			// TODO: Public/private symbols would influence whether this would be exposed to the asset registry here
-			OutExports.Parameters.Emplace(Entry->GetEntryName(), ParameterInterface->GetParamType(), EAnimNextParameterFlags::Bound);
+			FAnimNextParameterAssetRegistryExportEntry NewParam(ExportInterface->GetExportName(), TInstancedStruct<FAnimNextParamInstanceIdentifier>(), ExportInterface->GetExportType(), EAnimNextParameterFlags::Bound);
+			AddParamToSet(NewParam, OutExports);
 		}
-		else if(const IAnimNextRigVMGraphInterface* GraphInterface = Cast<IAnimNextRigVMGraphInterface>(Entry))
+		if(const IAnimNextRigVMGraphInterface* GraphInterface = Cast<IAnimNextRigVMGraphInterface>(Entry))
 		{
-			// TODO: Public/private symbols would influence whether this would be exposed to the asset registry here
-			OutExports.Parameters.Emplace(Entry->GetEntryName(), FAnimNextParamType::GetType<FAnimNextEntryPoint>(), EAnimNextParameterFlags::Bound);
-
 			GetGraphParameters(GraphInterface->GetRigVMGraph(), OutExports);
 		}
 	}
 }
 
 void FUtils::GetGraphParameters(const URigVMGraph* Graph, FAnimNextParameterProviderAssetRegistryExports& OutExports)
+{
+	TSet<FAnimNextParameterAssetRegistryExportEntry> ExportSet;
+	GetGraphParameters(Graph, ExportSet);
+	OutExports.Parameters = ExportSet.Array();
+}
+
+void FUtils::GetGraphParameters(const URigVMGraph* Graph, TSet<FAnimNextParameterAssetRegistryExportEntry>& OutExports)
 {
 	if (Graph == nullptr)
 	{
@@ -1316,58 +1396,65 @@ void FUtils::GetGraphParameters(const URigVMGraph* Graph, FAnimNextParameterProv
 			const FRigVMDispatchFactory* GetParameterFactory = FRigVMRegistry::Get().FindOrAddDispatchFactory(FRigVMDispatch_GetParameter::StaticStruct());
 			const FName GetParameterNotation = GetParameterFactory->GetTemplate()->GetNotation();
 
+			const FRigVMDispatchFactory* GetScopedParameterFactory = FRigVMRegistry::Get().FindOrAddDispatchFactory(FRigVMDispatch_GetScopedParameter::StaticStruct());
+			const FName GetScopedParameterNotation = GetScopedParameterFactory->GetTemplate()->GetNotation();
+			
 			const FRigVMDispatchFactory* GetLayerParameterFactory = FRigVMRegistry::Get().FindOrAddDispatchFactory(FRigVMDispatch_GetLayerParameter::StaticStruct());
 			const FName GetLayerParameterNotation = GetLayerParameterFactory->GetTemplate()->GetNotation();
 
 			const FRigVMDispatchFactory* SetLayerParameterFactory = FRigVMRegistry::Get().FindOrAddDispatchFactory(FRigVMDispatch_SetLayerParameter::StaticStruct());
 			const FName SetLayerParameterNotation = SetLayerParameterFactory->GetTemplate()->GetNotation();
 
-			const bool bReadParameter = TemplateNode->GetNotation() == GetParameterNotation || TemplateNode->GetNotation() == GetLayerParameterNotation;
+			const bool bIsScopedParameter = TemplateNode->GetNotation() == GetScopedParameterNotation; 
+			const bool bReadParameter = bIsScopedParameter || TemplateNode->GetNotation() == GetParameterNotation || TemplateNode->GetNotation() == GetLayerParameterNotation;
 			const bool bWriteParameter = TemplateNode->GetNotation() == SetLayerParameterNotation;
-			if (bReadParameter || bWriteParameter)
+			const bool bUsesRuntimeStruct = bIsScopedParameter;
+			const bool bUsesName = TemplateNode->GetNotation() == GetLayerParameterNotation || TemplateNode->GetNotation() == SetLayerParameterNotation;
+			const bool bIsParameterNode = bIsScopedParameter || bReadParameter || bWriteParameter;
+
+			const URigVMPin* ParameterPin = TemplateNode->FindPin(FRigVMDispatch_GetParameter::ParameterName.ToString());
+			if (bIsParameterNode && ParameterPin != nullptr)
 			{
-				if (const URigVMPin* NamePin = TemplateNode->FindPin(FRigVMDispatch_GetParameter::ParameterName.ToString()))
+				const FString PinDefaultValue = ParameterPin->GetDefaultValue();
+				if(!PinDefaultValue.IsEmpty())
 				{
-					const FString PinDefaultValue = NamePin->GetDefaultValue();
-					if(!PinDefaultValue.IsEmpty() && PinDefaultValue != TEXT("None"))
+					FAnimNextEditorParam PinParam;
+					if(bUsesRuntimeStruct)
 					{
-						if (const URigVMPin* ValuePin = TemplateNode->FindPin(FRigVMDispatch_GetParameter::ValueName.ToString()))
+						FAnimNextParam AnimNextParam;
+						FAnimNextParam::StaticStruct()->ImportText(*PinDefaultValue, &AnimNextParam, nullptr, PPF_None, nullptr, FAnimNextParam::StaticStruct()->GetName());
+						PinParam = FAnimNextEditorParam(AnimNextParam);
+					}
+					else if(bUsesName)
+					{
+						FName ParameterName = *PinDefaultValue;
+						const URigVMPin* ValuePin = TemplateNode->FindPin(FRigVMDispatch_GetLayerParameter::ValueName.ToString());
+						if (ValuePin != nullptr)
 						{
-							FAnimNextParamType Type = FAnimNextParamType::FromRigVMTemplateArgument(ValuePin->GetTemplateArgumentType());
-							EAnimNextParameterFlags Flags = EAnimNextParameterFlags::NoFlags;
-							if (bReadParameter)
-							{
-								Flags |= EAnimNextParameterFlags::Read;
-							}
-							
-							if (bWriteParameter)
-							{
-								Flags |= EAnimNextParameterFlags::Write;
-							}
-							OutExports.Parameters.Emplace(FName(PinDefaultValue), Type, Flags);
+							const FAnimNextParamType ParamType = FAnimNextParamType::FromRigVMTemplateArgument(FRigVMTemplateArgumentType(FName(*ValuePin->GetCPPType()), ValuePin->GetCPPTypeObject()));
+							PinParam = FAnimNextEditorParam(ParameterName, ParamType, TInstancedStruct<FAnimNextParamInstanceIdentifier>());
 						}
 					}
-				}
-			}
-			else
-			{
-				for(const URigVMPin* Pin : TemplateNode->GetAllPinsRecursively())
-				{
-					if(Pin->GetCPPType() == TEXT("FName"))
+					else
+					{
+						FAnimNextEditorParam::StaticStruct()->ImportText(*PinDefaultValue, &PinParam, nullptr, PPF_None, nullptr, FAnimNextEditorParam::StaticStruct()->GetName());
+					}
+
+					if(PinParam.Type.IsValid())
 					{
-						if(Pin->GetCustomWidgetName() == "ParamName")
+						EAnimNextParameterFlags Flags = EAnimNextParameterFlags::NoFlags;
+						if (bReadParameter)
 						{
-							const FString PinDefaultValue = Pin->GetDefaultValue();
-							if(!PinDefaultValue.IsEmpty() && PinDefaultValue != TEXT("None"))
-							{
-								const FString ParamTypeString = Pin->GetMetaData("AllowedParamType");
-								FAnimNextParamType Type = FAnimNextParamType::FromString(ParamTypeString);
-								if(Type.IsValid())
-								{
-									OutExports.Parameters.Emplace(FName(Pin->GetDefaultValue()), Type, EAnimNextParameterFlags::Read);
-								}
-							}
+							Flags |= EAnimNextParameterFlags::Read;
+						}
+
+						if (bWriteParameter)
+						{
+							Flags |= EAnimNextParameterFlags::Write;
 						}
+
+						FAnimNextParameterAssetRegistryExportEntry NewEntry(PinParam.Name, PinParam.InstanceId, PinParam.Type, Flags);
+						AddParamToSet(NewEntry, OutExports);
 					}
 				}
 			}
@@ -1376,6 +1463,13 @@ void FUtils::GetGraphParameters(const URigVMGraph* Graph, FAnimNextParameterProv
 }
 
 void FUtils::GetScheduleParameters(const UAnimNextSchedule* InSchedule, FAnimNextParameterProviderAssetRegistryExports& OutExports)
+{
+	TSet<FAnimNextParameterAssetRegistryExportEntry> ExportSet;
+	GetScheduleParameters(InSchedule, ExportSet);
+	OutExports.Parameters = ExportSet.Array();
+}
+
+void FUtils::GetScheduleParameters(const UAnimNextSchedule* InSchedule, TSet<FAnimNextParameterAssetRegistryExportEntry>& OutExports)
 {
 	for(UAnimNextScheduleEntry* Entry : InSchedule->Entries)
 	{
@@ -1384,45 +1478,57 @@ void FUtils::GetScheduleParameters(const UAnimNextSchedule* InSchedule, FAnimNex
 			if(PortEntry->Port)
 			{
 				UAnimNextSchedulePort* CDO = PortEntry->Port->GetDefaultObject<UAnimNextSchedulePort>();
-				TConstArrayView<FAnimNextParam> RequiredParameters = CDO->GetRequiredParameters();
-				for(const FAnimNextParam& RequiredParameter : RequiredParameters)
+				TConstArrayView<FAnimNextEditorParam> RequiredParameters = CDO->GetRequiredParameters();
+				for(const FAnimNextEditorParam& RequiredParameter : RequiredParameters)
 				{
 					if(!RequiredParameter.Name.IsNone() && RequiredParameter.Type.IsValid())
 					{
-						OutExports.Parameters.Emplace(RequiredParameter.Name, RequiredParameter.Type, EAnimNextParameterFlags::Read);
+						FAnimNextParameterAssetRegistryExportEntry NewEntry(RequiredParameter.Name, RequiredParameter.InstanceId, RequiredParameter.Type, EAnimNextParameterFlags::Read);
+						AddParamToSet(NewEntry, OutExports);
 					}
 				}
 			}
 		}
 		else if (UAnimNextScheduleEntry_AnimNextGraph* GraphEntry = Cast<UAnimNextScheduleEntry_AnimNextGraph>(Entry))
 		{
-			if(!GraphEntry->DynamicGraph.IsNone())
+			if(GraphEntry->DynamicGraph.IsValid())
 			{
-				OutExports.Parameters.Emplace(GraphEntry->DynamicGraph, FAnimNextParamType::GetType<TObjectPtr<UAnimNextGraph>>(), EAnimNextParameterFlags::Read);
+				FAnimNextParameterAssetRegistryExportEntry NewEntry(GraphEntry->DynamicGraph.Name, GraphEntry->DynamicGraph.InstanceId, GraphEntry->DynamicGraph.Type, EAnimNextParameterFlags::Read);
+				AddParamToSet(NewEntry, OutExports);
 			}
-			for(const FAnimNextParam& ParameterName : GraphEntry->RequiredParameters)
+			for(const FAnimNextEditorParam& RequiredParameter : GraphEntry->RequiredParameters)
 			{
-				OutExports.Parameters.Emplace(ParameterName.Name, ParameterName.Type, EAnimNextParameterFlags::Read);
+				FAnimNextParameterAssetRegistryExportEntry NewEntry(RequiredParameter.Name, RequiredParameter.InstanceId, RequiredParameter.Type, EAnimNextParameterFlags::Read);
+				AddParamToSet(NewEntry, OutExports);
 			}
 		}
 		else if (UAnimNextScheduleEntry_ExternalTask* ExternalTaskEntry = Cast<UAnimNextScheduleEntry_ExternalTask>(Entry))
 		{
-			if(!ExternalTaskEntry->ExternalTask.IsNone())
+			if(ExternalTaskEntry->ExternalTask.IsValid())
 			{
-				OutExports.Parameters.Emplace(ExternalTaskEntry->ExternalTask, FAnimNextParamType::GetType<FAnimNextExternalTaskBinding>(), EAnimNextParameterFlags::Read);
+				FAnimNextParameterAssetRegistryExportEntry NewEntry(ExternalTaskEntry->ExternalTask.Name, ExternalTaskEntry->ExternalTask.InstanceId, ExternalTaskEntry->ExternalTask.Type, EAnimNextParameterFlags::Read);
+				AddParamToSet(NewEntry, OutExports);
 			}
 		}
 		else if (UAnimNextScheduleEntry_ParamScope* ParamScopeTaskEntry = Cast<UAnimNextScheduleEntry_ParamScope>(Entry))
 		{
-			if(!ParamScopeTaskEntry->Scope.IsNone())
+			if(ParamScopeTaskEntry->Scope.IsValid())
 			{
-				OutExports.Parameters.Emplace(ParamScopeTaskEntry->Scope, FAnimNextParamType::GetType<FAnimNextScope>(), EAnimNextParameterFlags::Read);
+				FAnimNextParameterAssetRegistryExportEntry NewEntry(ParamScopeTaskEntry->Scope.Name, ParamScopeTaskEntry->Scope.InstanceId, ParamScopeTaskEntry->Scope.Type, EAnimNextParameterFlags::Read);
+				AddParamToSet(NewEntry, OutExports);
 			}
 		}
 	}
 }
 
 void FUtils::GetBlueprintParameters(const UBlueprint* InBlueprint, FAnimNextParameterProviderAssetRegistryExports& OutExports)
+{
+	TSet<FAnimNextParameterAssetRegistryExportEntry> ExportSet;
+	GetBlueprintParameters(InBlueprint, ExportSet);
+	OutExports.Parameters = ExportSet.Array();
+}
+
+void FUtils::GetBlueprintParameters(const UBlueprint* InBlueprint, TSet<FAnimNextParameterAssetRegistryExportEntry>& OutExports)
 {
 	// Add 'static' params held on components
 	if(InBlueprint->SimpleConstructionScript)
@@ -1435,7 +1541,8 @@ void FUtils::GetBlueprintParameters(const UBlueprint* InBlueprint, FAnimNextPara
 				{
 					if(!Parameter->Scope.IsNone())
 					{
-						OutExports.Parameters.Emplace(Parameter->Scope, FAnimNextParamType::GetType<FAnimNextScope>(), EAnimNextParameterFlags::Read);
+						FAnimNextParameterAssetRegistryExportEntry NewEntry(Parameter->Scope, TInstancedStruct<FAnimNextParamInstanceIdentifier>(), FAnimNextParamType::GetType<FAnimNextScope>(), EAnimNextParameterFlags::Read);
+						AddParamToSet(NewEntry, OutExports);
 					}
 
 					FName Name;
@@ -1445,7 +1552,8 @@ void FUtils::GetBlueprintParameters(const UBlueprint* InBlueprint, FAnimNextPara
 					{
 						FAnimNextParamType Type = FParamTypeHandle::FromProperty(Property).GetType();
 						check(Type.IsValid());
-						OutExports.Parameters.Emplace(Name, Type, EAnimNextParameterFlags::Write);
+						FAnimNextParameterAssetRegistryExportEntry NewEntry(Name, TInstancedStruct<FAnimNextParamInstanceIdentifier>(), Type, EAnimNextParameterFlags::Write);
+						AddParamToSet(NewEntry, OutExports);
 					}
 				}
 			}
@@ -1471,7 +1579,8 @@ void FUtils::GetBlueprintParameters(const UBlueprint* InBlueprint, FAnimNextPara
 				FName ScopeName(*ScopePin->GetDefaultAsString());
 				if(!ScopeName.IsNone())
 				{
-					OutExports.Parameters.Emplace(ScopeName, FAnimNextParamType::GetType<FAnimNextScope>(), EAnimNextParameterFlags::Read);
+					FAnimNextParameterAssetRegistryExportEntry NewEntry(ScopeName, TInstancedStruct<FAnimNextParamInstanceIdentifier>(), FAnimNextParamType::GetType<FAnimNextScope>(), EAnimNextParameterFlags::Read);
+					AddParamToSet(NewEntry, OutExports);
 				}
 
 				UEdGraphPin* NamePin = FunctionNode->FindPinChecked(TEXT("Name"));
@@ -1482,7 +1591,8 @@ void FUtils::GetBlueprintParameters(const UBlueprint* InBlueprint, FAnimNextPara
 					FAnimNextParamType Type = UncookedOnly::FUtils::GetParamTypeFromPinType(ValuePin->PinType);
 					if(Type.IsValid())
 					{
-						OutExports.Parameters.Emplace(ParamName, Type, EAnimNextParameterFlags::Write);
+						FAnimNextParameterAssetRegistryExportEntry NewEntry(ParamName, TInstancedStruct<FAnimNextParamInstanceIdentifier>(), Type, EAnimNextParameterFlags::Write);
+						AddParamToSet(NewEntry, OutExports);
 					}
 				}
 			}
@@ -1663,9 +1773,9 @@ void FUtils::CompileSchedule(UAnimNextSchedule* InSchedule)
 			{
 				bool bValid = true;
 
-				if(GraphEntry->Graph == nullptr && GraphEntry->DynamicGraph == NAME_None)
+				if(GraphEntry->Graph == nullptr && !GraphEntry->DynamicGraph.IsValid())
 				{
-					UE_LOG(LogAnimation, Error, TEXT("AnimNext: Invalid graph or no parameter supplied"));
+					UE_LOG(LogAnimation, Error, TEXT("AnimNext: Invalid graph or invalid parameter supplied"));
 					bValid = false;
 				}
 				else if(GraphEntry->Graph != nullptr)
@@ -1734,12 +1844,12 @@ void FUtils::CompileSchedule(UAnimNextSchedule* InSchedule)
 					GraphTask.TaskIndex = InSchedule->GraphTasks.Num();
 					GraphTask.ParamScopeIndex = InSchedule->NumParameterScopes++;
 					GraphTask.ParamParentScopeIndex = ParentScopeIndex;
-					GraphTask.EntryPoint = GraphEntry->EntryPoint;
+					GraphTask.EntryPoint = FAnimNextParam(GraphEntry->EntryPoint);
 					GraphTask.Graph = GraphEntry->Graph;
-					GraphTask.DynamicGraph = GraphEntry->DynamicGraph;
-					if(GraphEntry->Graph == nullptr && GraphEntry->DynamicGraph != NAME_None)
+					GraphTask.DynamicGraph = FAnimNextParam(GraphEntry->DynamicGraph);
+					if(GraphEntry->Graph == nullptr && GraphEntry->DynamicGraph.IsValid())
 					{
-						GraphTask.SuppliedParameters = GraphEntry->RequiredParameters;
+						Algo::Transform(GraphEntry->RequiredParameters, GraphTask.SuppliedParameters, [](const FAnimNextEditorParam& InParam){ return FAnimNextParam(InParam); });
 						GraphTask.SuppliedParametersHash = SortAndHashParameters(GraphTask.SuppliedParameters);
 					}
 
@@ -1774,7 +1884,7 @@ void FUtils::CompileSchedule(UAnimNextSchedule* InSchedule)
 				ExternalTask.TaskIndex = InSchedule->ExternalTasks.Num();
 				ExternalTask.ParamScopeIndex = InSchedule->NumParameterScopes++;
 				ExternalTask.ParamParentScopeIndex = ParentScopeIndex;
-				ExternalTask.ExternalTask = ExternalTaskEntry->ExternalTask;
+				ExternalTask.ExternalTask = FAnimNextParam(ExternalTaskEntry->ExternalTask);
 				int32 ExternalTaskIndex = InSchedule->ExternalTasks.Add(ExternalTask);
 
 				// Emit the external task
@@ -1796,7 +1906,7 @@ void FUtils::CompileSchedule(UAnimNextSchedule* InSchedule)
 				ParamScopeEntryTask.ParamScopeIndex = ParamScopeIndex;
 				ParamScopeEntryTask.ParamParentScopeIndex = ParentScopeIndex;
 				ParamScopeEntryTask.TickFunctionIndex = InSchedule->NumTickFunctions;
-				ParamScopeEntryTask.Scope = ParamScopeTaskEntry->Scope;
+				ParamScopeEntryTask.Scope = FAnimNextParam(ParamScopeTaskEntry->Scope);
 				ParamScopeEntryTask.Parameters = ParamScopeTaskEntry->Parameters;
 				int32 ParamScopeTaskEntryIndex = InSchedule->ParamScopeEntryTasks.Add(ParamScopeEntryTask);
 
@@ -1818,7 +1928,7 @@ void FUtils::CompileSchedule(UAnimNextSchedule* InSchedule)
 				FAnimNextScheduleParamScopeExitTask ParamScopeExitTask;
 				ParamScopeExitTask.TaskIndex = InSchedule->ParamScopeExitTasks.Num();
 				ParamScopeExitTask.ParamScopeIndex = ParamScopeIndex;
-				ParamScopeExitTask.Scope = ParamScopeTaskEntry->Scope;
+				ParamScopeExitTask.Scope = FAnimNextParam(ParamScopeTaskEntry->Scope);
 				int32 ParamScopeExitTaskIndex = InSchedule->ParamScopeExitTasks.Add(ParamScopeExitTask);
 
 				Emit(EAnimNextScheduleScheduleOpcode::RunParamScopeExit, ParamScopeExitTaskIndex);
@@ -1840,11 +1950,11 @@ void FUtils::CompileSchedule(UAnimNextSchedule* InSchedule)
 		}
 	};
 
-	auto GenerateExternalParameters = [](TArray<TObjectPtr<UAnimNextScheduleEntry>>& InEntries)
+	auto GenerateExternalParameters = [&InSchedule](TArray<TObjectPtr<UAnimNextScheduleEntry>>& InEntries)
 	{
 		struct FParameterTracker
 		{
-			FName SourceName;
+			TInstancedStruct<FAnimNextParamInstanceIdentifier> InstanceId;
 			TArray<TObjectPtr<UAnimNextScheduleEntry>>* BestContainer = nullptr;
 			TSet<FName> ThreadSafeParameters;
 			TSet<FName> NonThreadSafeParameters;
@@ -1852,21 +1962,22 @@ void FUtils::CompileSchedule(UAnimNextSchedule* InSchedule)
 			uint32 BestArrayIndex = MAX_uint32;
 		};
 
+		FModule& Module = FModuleManager::Get().LoadModuleChecked<FModule>("AnimNextUncookedOnly");
+
 		// We need to find the task that is 'earliest' in the DAG for each external parameter, so we track that with this map
 		TMap<FName, FParameterTracker> TrackerMap;
-		auto TrackExternalParameters = [&TrackerMap](TConstArrayView<FAnimNextParam> InParameters, uint32 InDistance, UAnimNextScheduleEntry* InEntry, TArray<TObjectPtr<UAnimNextScheduleEntry>>* InContainer, int32 InArrayIndex, bool bInUpdateDependent)
+		auto TrackExternalParameters = [&InSchedule, &TrackerMap, &Module](TConstArrayView<FAnimNextEditorParam> InParameters, uint32 InDistance, UAnimNextScheduleEntry* InEntry, TArray<TObjectPtr<UAnimNextScheduleEntry>>* InContainer, int32 InArrayIndex, bool bInUpdateDependent)
 		{
+			using namespace UE::UniversalObjectLocator;
+
 			bool bHasExternal = false;
-			for(const FAnimNextParam& Parameter : InParameters)
+			for(const FAnimNextEditorParam& Parameter : InParameters)
 			{
-				// Only add if the parameter name is 'external'
-				FName ParameterSourceName = FExternalParameterRegistry::FindSourceForParameter(Parameter.Name);
-				if(ParameterSourceName != NAME_None)
+				// Only add if the parameter name is 'external' (i.e. it has a valid instance ID)
+				if(Parameter.InstanceId.IsValid() && Parameter.InstanceId.Get().IsValid())
 				{
-					bHasExternal = true;
-					FParameterTracker& Tracker = TrackerMap.FindOrAdd(ParameterSourceName);
-
-					Tracker.SourceName = ParameterSourceName;
+					FParameterTracker& Tracker = TrackerMap.FindOrAdd(Parameter.InstanceId.Get().ToName());
+					Tracker.InstanceId = Parameter.InstanceId;
 
 					// Only track array index if this param is update dependent
 					if(bInUpdateDependent && InDistance < Tracker.BestDistance)
@@ -1876,10 +1987,14 @@ void FUtils::CompileSchedule(UAnimNextSchedule* InSchedule)
 						Tracker.BestArrayIndex = InArrayIndex == 0 ? 0 : InArrayIndex - 1;
 					}
 
-					IParameterSourceFactory::FParameterInfo Info;
-					ensure(FExternalParameterRegistry::FindParameterInfo(Parameter.Name, Info));
+					FParameterSourceInfo Info[1];
+					if(TSharedPtr<IParameterSourceType> SourceType = Module.FindParameterSourceType(Parameter.InstanceId.GetScriptStruct()))
+					{
+						SourceType->FindParameterInfo(Parameter.InstanceId, { Parameter.Name }, Info);
+					}
 
-					if(Info.bThreadSafe)
+					// Note that if the above call to FindParameterSourceType or FindParameterInfo fails, this will still default to non-thread safe
+					if(Info[0].bThreadSafe)
 					{
 						Tracker.ThreadSafeParameters.Add(Parameter.Name);
 					}
@@ -1887,6 +2002,8 @@ void FUtils::CompileSchedule(UAnimNextSchedule* InSchedule)
 					{
 						Tracker.NonThreadSafeParameters.Add(Parameter.Name);
 					}
+
+					bHasExternal = true;
 				}
 			}
 			return bHasExternal;
@@ -1906,11 +2023,11 @@ void FUtils::CompileSchedule(UAnimNextSchedule* InSchedule)
 						FAnimNextParameterProviderAssetRegistryExports Exports;
 						if(UncookedOnly::FUtils::GetExportedParametersForAsset(FAssetData(GraphEntry->Graph), Exports))
 						{
-							TArray<FAnimNextParam> RequiredParameters;
+							TArray<FAnimNextEditorParam> RequiredParameters;
 							RequiredParameters.Reserve(Exports.Parameters.Num());
 							for(const FAnimNextParameterAssetRegistryExportEntry& ExportedParameter : Exports.Parameters)
 							{
-								RequiredParameters.Emplace(ExportedParameter.Name, ExportedParameter.Type);
+								RequiredParameters.Emplace(ExportedParameter.Name, ExportedParameter.Type, ExportedParameter.InstanceId);
 							}
 							TrackExternalParameters(RequiredParameters, InDistance, Entry, &InEntries, ArrayIndex, true);
 						}
@@ -1918,28 +2035,22 @@ void FUtils::CompileSchedule(UAnimNextSchedule* InSchedule)
 
 					TrackExternalParameters(GraphEntry->RequiredParameters, InDistance, Entry, &InEntries, ArrayIndex, true);
 
-					// All graphs require access to LOD and ref pose
 					// TODO: need to not require defaults here - use static graph params. In fact, these params should probably be defined at the schedule level.
-					const FAnimNextParam DefaultParams[] =
-					{
-						FAnimNextParam(UAnimNextGraph::DefaultCurrentLODId.GetName(), FAnimNextParamType::GetType<int32>()),
-						FAnimNextParam(UAnimNextGraph::DefaultReferencePoseId.GetName(), FAnimNextParamType::GetType<FAnimNextGraphReferencePose>())
-					};
-					TrackExternalParameters(DefaultParams, InDistance, Entry, &InEntries, ArrayIndex, true);
+					TrackExternalParameters(FAnimNextScheduleGraphTask::GetRequiredParametersInternal(), InDistance, Entry, &InEntries, ArrayIndex, true);
 				}
 				else if (UAnimNextScheduleEntry_Port* PortEntry = Cast<UAnimNextScheduleEntry_Port>(Entry))
 				{
 					if(PortEntry->Port)
 					{
 						UAnimNextSchedulePort* CDO = PortEntry->Port->GetDefaultObject<UAnimNextSchedulePort>();
-						TConstArrayView<FAnimNextParam> RequiredParameters = CDO->GetRequiredParameters();
+						TConstArrayView<FAnimNextEditorParam> RequiredParameters = CDO->GetRequiredParameters();
 						TrackExternalParameters(RequiredParameters, InDistance, Entry, &InEntries, ArrayIndex, true);
 					}
 				}
 				else if (UAnimNextScheduleEntry_ExternalTask* ExternalTaskEntry = Cast<UAnimNextScheduleEntry_ExternalTask>(Entry))
 				{
 					// Note: external task params are not update dependent as this would cause external param updates to occur before tick functions
-					TrackExternalParameters({ FAnimNextParam(ExternalTaskEntry->ExternalTask, FAnimNextParamType::GetType<FAnimNextExternalTaskBinding>() ) }, InDistance, Entry, &InEntries, ArrayIndex, false);
+					TrackExternalParameters({ ExternalTaskEntry->ExternalTask }, InDistance, Entry, &InEntries, ArrayIndex, false);
 				}
 				else if (UAnimNextScheduleEntry_ParamScope* ParamScopeTaskEntry = Cast<UAnimNextScheduleEntry_ParamScope>(Entry))
 				{
@@ -1950,11 +2061,11 @@ void FUtils::CompileSchedule(UAnimNextSchedule* InSchedule)
 							FAnimNextParameterProviderAssetRegistryExports Exports;
 							if(UncookedOnly::FUtils::GetExportedParametersForAsset(FAssetData(Parameters), Exports))
 							{
-								TArray<FAnimNextParam> RequiredParameters;
+								TArray<FAnimNextEditorParam> RequiredParameters;
 								RequiredParameters.Reserve(Exports.Parameters.Num());
 								for(const FAnimNextParameterAssetRegistryExportEntry& ExportedParameter : Exports.Parameters)
 								{
-									RequiredParameters.Emplace(ExportedParameter.Name, ExportedParameter.Type);
+									RequiredParameters.Emplace(ExportedParameter.Name, ExportedParameter.Type, ExportedParameter.InstanceId);
 								}
 								TrackExternalParameters(RequiredParameters, InDistance, Entry, &InEntries, ArrayIndex, true);
 							}
@@ -1989,7 +2100,7 @@ void FUtils::CompileSchedule(UAnimNextSchedule* InSchedule)
 			{
 				TArray<FAnimNextScheduleExternalParameterSource>& ParameterSources = InsertionMap.FindOrAdd({ Container, ArrayIndex, true });
 				FAnimNextScheduleExternalParameterSource& NewSource = ParameterSources.AddDefaulted_GetRef();
-				NewSource.ParameterSource = TrackedParameterPair.Value.SourceName;
+				NewSource.InstanceId = TrackedParameterPair.Value.InstanceId;
 				NewSource.Parameters = TrackedParameterPair.Value.ThreadSafeParameters.Array();
 			}
 
@@ -1997,7 +2108,7 @@ void FUtils::CompileSchedule(UAnimNextSchedule* InSchedule)
 			{
 				TArray<FAnimNextScheduleExternalParameterSource>& ParameterSources = InsertionMap.FindOrAdd({ Container, ArrayIndex, false });
 				FAnimNextScheduleExternalParameterSource& NewSource = ParameterSources.AddDefaulted_GetRef();
-				NewSource.ParameterSource = TrackedParameterPair.Value.SourceName;
+				NewSource.InstanceId = TrackedParameterPair.Value.InstanceId;
 				NewSource.Parameters = TrackedParameterPair.Value.NonThreadSafeParameters.Array();
 			}
 		}
@@ -2079,3 +2190,5 @@ uint64 FUtils::SortAndHashParameters(TArray<FAnimNextParam>& InParameters)
 }
 
 }
+
+#undef LOCTEXT_NAMESPACE
\ No newline at end of file
diff --git a/Engine/Plugins/Experimental/AnimNext/Source/AnimNextUncookedOnly/Public/IAnimNextUncookedOnlyModule.h b/Engine/Plugins/Experimental/AnimNext/Source/AnimNextUncookedOnly/Public/IAnimNextUncookedOnlyModule.h
new file mode 100644
--- /dev/null
+++ b/Engine/Plugins/Experimental/AnimNext/Source/AnimNextUncookedOnly/Public/IAnimNextUncookedOnlyModule.h
@@ -0,0 +1,37 @@
+﻿// Copyright Epic Games, Inc. All Rights Reserved.
+
+#pragma once
+
+#include "CoreMinimal.h"
+#include "Modules/ModuleInterface.h"
+#include "Param/AnimNextParamInstanceIdentifier.h"
+
+template<typename T> struct TInstancedStruct;
+
+namespace UE::AnimNext::UncookedOnly
+{
+	struct FParameterSourceInfo;
+	class IParameterSourceType;
+}
+
+namespace UE::AnimNext::UncookedOnly
+{
+
+class IAnimNextUncookedOnlyModule : public IModuleInterface
+{
+public:
+	// Register a parameter source type that can be used to query information about parameter sources
+	// @param   InInstanceIdStruct     The name of struct type for this parameter source's instance ID (must be a child struct of FAnimNextParamInstanceIdentifier)
+	// @param   InFactory              The type to register
+	virtual void RegisterParameterSourceType(const UScriptStruct* InInstanceIdStruct, TSharedPtr<IParameterSourceType> InType) = 0;
+	
+	// Unregister a parameter source type previously passed to RegisterParameterSourceType
+	// @param   InInstanceIdStruct     The name of struct type for this parameter source's instance ID
+	virtual void UnregisterParameterSourceType(const UScriptStruct* InInstanceIdStruct) = 0;
+
+	// Find a parameter source type previously passed to RegisterParameterSourceType
+	// @param   InInstanceIdStruct     The name of struct type for this parameter source's instance ID
+	virtual TSharedPtr<IParameterSourceType> FindParameterSourceType(const UScriptStruct* InInstanceIdStruct) const = 0;
+};
+
+}
diff --git a/Engine/Plugins/Experimental/AnimNext/Source/AnimNextUncookedOnly/Public/Param/IParameterSourceType.h b/Engine/Plugins/Experimental/AnimNext/Source/AnimNextUncookedOnly/Public/Param/IParameterSourceType.h
new file mode 100644
--- /dev/null
+++ b/Engine/Plugins/Experimental/AnimNext/Source/AnimNextUncookedOnly/Public/Param/IParameterSourceType.h
@@ -0,0 +1,77 @@
+// Copyright Epic Games, Inc. All Rights Reserved.
+
+#pragma once
+
+#include "CoreMinimal.h"
+#include "Param/ParamType.h"
+
+struct FAnimNextParamInstanceIdentifier;
+template<typename T> struct TInstancedStruct;
+
+namespace UE::AnimNext::UncookedOnly
+{
+	class FModule;
+	struct FUtils;
+}
+
+namespace UE::AnimNext::Editor
+{
+	class SParameterPicker;
+}
+
+
+namespace UE::AnimNext::UncookedOnly
+{
+
+// Info about parameters gleaned from FindParameterInfo and ForEachParameter
+struct FParameterSourceInfo
+{
+	// The parameter's type
+	FAnimNextParamType Type;
+	// Display name for editor
+	FText DisplayName;
+	// Tooltip to display in editor
+	FText Tooltip;
+	// Function used to access this parameter, if any
+	const UFunction* Function = nullptr;
+	// Property for this parameter, if any
+	const FProperty* Property = nullptr;
+	// Whether this parameter is safe to be accessed on worker threads
+	bool bThreadSafe = false;
+};
+
+// Interface used in editor/uncooked situations to determine the characteristics of a parameter source.
+// @see IParameterSource, IParameterSourceFactory
+class IParameterSourceType
+{
+public:
+	virtual ~IParameterSourceType() = default;
+
+private:
+	// Get the struct type that this FAnimNextParamInstanceIdentifier resolves to
+	virtual const UStruct* GetStruct(const TInstancedStruct<FAnimNextParamInstanceIdentifier>& InInstanceId) const = 0;
+
+	// Get the display text for the specified instance ID
+	virtual FText GetDisplayText(const TInstancedStruct<FAnimNextParamInstanceIdentifier>& InInstanceId) const = 0;
+
+	// Get the tooltip text for the specified instance ID
+	virtual FText GetTooltipText(const TInstancedStruct<FAnimNextParamInstanceIdentifier>& InInstanceId) const = 0;
+
+	// Given a parameter name and class, find associated info for that parameter
+	// @param   InInstanceId           Instance ID used to find the parameter
+	// @param   InParameterNames       The parameter names to find. Size must match OutInfo.
+	// @param   OutInfo                The parameter info found. Size must match InParameterNames.
+	// @return true if any info was found, false if no info was found
+	virtual bool FindParameterInfo(const TInstancedStruct<FAnimNextParamInstanceIdentifier>& InInstanceId, TConstArrayView<FName> InParameterNames, TArrayView<FParameterSourceInfo> OutInfo) const = 0;
+
+	// Iterate over all the known parameters for a particular class
+	// @param   InInstanceId           Instance ID used to find the parameters
+	// @param   InFunction             Function called for each known parameter
+	virtual void ForEachParameter(const TInstancedStruct<FAnimNextParamInstanceIdentifier>& InInstanceId, TFunctionRef<void(FName, const FParameterSourceInfo&)> InFunction) const = 0;
+
+	friend class FModule;
+	friend struct FUtils;
+	friend class Editor::SParameterPicker;
+};
+
+}
\ No newline at end of file
diff --git a/Engine/Plugins/Experimental/AnimNextAnimGraph/Source/AnimNextAnimGraph/Private/Graph/AnimGraph/AnimNode_AnimNextGraph.cpp b/Engine/Plugins/Experimental/AnimNextAnimGraph/Source/AnimNextAnimGraph/Private/Graph/AnimGraph/AnimNode_AnimNextGraph.cpp
--- a/Engine/Plugins/Experimental/AnimNextAnimGraph/Source/AnimNextAnimGraph/Private/Graph/AnimGraph/AnimNode_AnimNextGraph.cpp
+++ b/Engine/Plugins/Experimental/AnimNextAnimGraph/Source/AnimNextAnimGraph/Private/Graph/AnimGraph/AnimNode_AnimNextGraph.cpp
@@ -61,18 +61,7 @@ void FAnimNode_AnimNextGraph::Update_AnyThread(const FAnimationUpdateContext& Co
 
 		PropagateInputProperties(Context.AnimInstanceProxy->GetAnimInstanceObject());
 
-		const int32 LODLevel = Context.AnimInstanceProxy->GetLODLevel();
-
-		UE::AnimNext::FAnimGraphParamStackScope Scope(Context);
-
-		FParamStack& ParamStack = FParamStack::Get();
-		FParamStack::FPushedLayerHandle LayerHandle = ParamStack.PushValues(
-			AnimNextGraph->GetCurrentLODParam(), LODLevel
-		);
-
 		UE::AnimNext::UpdateGraph(GraphInstance, Context.GetDeltaTime());
-
-		ParamStack.PopLayer(LayerHandle);
 	}
 
 	FAnimNode_CustomProperty::Update_AnyThread(Context);
@@ -90,20 +79,7 @@ void FAnimNode_AnimNextGraph::Initialize_AnyThread(const FAnimationInitializeCon
 
 	if (!GraphInstance.IsValid() && AnimNextGraph)
 	{
-		// If we don't have an instance yet, create one
-		UE::AnimNext::FAnimGraphParamStackScope Scope(Context);
-
-		// Populate our param stack since our instance data might need it during construction
-		const int32 LODLevel = Context.AnimInstanceProxy->GetLODLevel();
-
-		FParamStack& ParamStack = FParamStack::Get();
-		FParamStack::FPushedLayerHandle LayerHandle = ParamStack.PushValues(
-			AnimNextGraph->GetCurrentLODParam(), LODLevel
-		);
-
 		AnimNextGraph->AllocateInstance(GraphInstance);
-
-		ParamStack.PopLayer(LayerHandle);
 	}
 
 	FAnimNode_CustomProperty::Initialize_AnyThread(Context);
@@ -137,13 +113,6 @@ void FAnimNode_AnimNextGraph::Evaluate_AnyThread(FPoseContext& Output)
 		const UE::AnimNext::FReferencePose& RefPose = RefPoseHandle.GetRef<UE::AnimNext::FReferencePose>();
 		FAnimNextGraphLODPose ResultPose(FLODPoseHeap(RefPose, LODLevel, true, Output.ExpectsAdditivePose()));
 
-		FAnimGraphParamStackScope Scope(Output);
-		FParamStack& ParamStack = FParamStack::Get();
-		FParamStack::FPushedLayerHandle LayerHandle = ParamStack.PushValues(
-			AnimNextGraph->GetReferencePoseParam(), GraphReferencePose,
-			AnimNextGraph->GetCurrentLODParam(), LODLevel
-		);
-
 		{
 			const FEvaluationProgram EvaluationProgram = UE::AnimNext::EvaluateGraph(GraphInstance);
 
@@ -171,8 +140,6 @@ void FAnimNode_AnimNextGraph::Evaluate_AnyThread(FPoseContext& Output)
 		}
 
 		FGenerationTools::RemapPose(ResultPose.LODPose, Output);
-
-		ParamStack.PopLayer(LayerHandle);
 	}
 	else
 	{
diff --git a/Engine/Plugins/Experimental/AnimNextAnimGraph/Source/AnimNextAnimGraphTestSuite/AnimNextAnimGraphTestSuite.Build.cs b/Engine/Plugins/Experimental/AnimNextAnimGraph/Source/AnimNextAnimGraphTestSuite/AnimNextAnimGraphTestSuite.Build.cs
--- a/Engine/Plugins/Experimental/AnimNextAnimGraph/Source/AnimNextAnimGraphTestSuite/AnimNextAnimGraphTestSuite.Build.cs
+++ b/Engine/Plugins/Experimental/AnimNextAnimGraph/Source/AnimNextAnimGraphTestSuite/AnimNextAnimGraphTestSuite.Build.cs
@@ -29,6 +29,7 @@ namespace UnrealBuildTool.Rules
 						"AnimNextUncookedOnly",
 						"AnimNextEditor",
 						"RigVMDeveloper",
+						"StructUtils",
 					}
 				);
 			}
diff --git a/Engine/Plugins/Experimental/AnimNextAnimGraph/Source/AnimNextAnimGraphTestSuite/Private/AnimNextAnimGraphTraitGraphTest.cpp b/Engine/Plugins/Experimental/AnimNextAnimGraph/Source/AnimNextAnimGraphTestSuite/Private/AnimNextAnimGraphTraitGraphTest.cpp
--- a/Engine/Plugins/Experimental/AnimNextAnimGraph/Source/AnimNextAnimGraphTestSuite/Private/AnimNextAnimGraphTraitGraphTest.cpp
+++ b/Engine/Plugins/Experimental/AnimNextAnimGraph/Source/AnimNextAnimGraphTestSuite/Private/AnimNextAnimGraphTraitGraphTest.cpp
@@ -46,15 +46,15 @@ namespace UE::AnimNext
 			const FSharedData* SharedData = Binding.GetSharedData<FSharedData>();
 			UE::AnimNext::FParamStack& ParamStack = UE::AnimNext::FParamStack::Get();
 
-			int32& UpdateCount = ParamStack.GetMutableParam<int32>("UpdateCount");
+			int32& UpdateCount = ParamStack.GetMutableParam<int32>("/Engine/Transient.TestAnimNextGraph:UpdateCount");
 			UpdateCount++;
 
-			ParamStack.GetMutableParam<int32>("SomeInt32") = SharedData->SomeInt32;
-			ParamStack.GetMutableParam<float>("SomeFloat") = SharedData->SomeFloat;
+			ParamStack.GetMutableParam<int32>("/Engine/Transient.TestAnimNextGraph:SomeInt32") = SharedData->SomeInt32;
+			ParamStack.GetMutableParam<float>("/Engine/Transient.TestAnimNextGraph:SomeFloat") = SharedData->SomeFloat;
 
-			ParamStack.GetMutableParam<int32>("SomeLatentInt32") = SharedData->GetSomeLatentInt32(Binding);				// MathAdd with constants, latent
-			ParamStack.GetMutableParam<int32>("SomeOtherLatentInt32") = SharedData->GetSomeOtherLatentInt32(Binding);	// GetParameter, latent
-			ParamStack.GetMutableParam<float>("SomeLatentFloat") = SharedData->GetSomeLatentFloat(Binding);				// Inline value, not latent
+			ParamStack.GetMutableParam<int32>("/Engine/Transient.TestAnimNextGraph:SomeLatentInt32") = SharedData->GetSomeLatentInt32(Binding);				// MathAdd with constants, latent
+			ParamStack.GetMutableParam<int32>("/Engine/Transient.TestAnimNextGraph:SomeOtherLatentInt32") = SharedData->GetSomeOtherLatentInt32(Binding);	// GetParameter, latent
+			ParamStack.GetMutableParam<float>("/Engine/Transient.TestAnimNextGraph:SomeLatentFloat") = SharedData->GetSomeLatentFloat(Binding);				// Inline value, not latent
 		}
 
 		// IEvaluate impl
@@ -62,9 +62,10 @@ namespace UE::AnimNext
 		{
 			IEvaluate::PostEvaluate(Context, Binding);
 
+			const FSharedData* SharedData = Binding.GetSharedData<FSharedData>();
 			UE::AnimNext::FParamStack& ParamStack = UE::AnimNext::FParamStack::Get();
-
-			int32& EvaluateCount = ParamStack.GetMutableParam<int32>("EvaluateCount");
+			
+			int32& EvaluateCount = ParamStack.GetMutableParam<int32>("/Engine/Transient.TestAnimNextGraph:EvaluateCount");
 			EvaluateCount++;
 		}
 	};
@@ -269,13 +270,13 @@ bool FAnimationAnimNextRuntimeTest_GraphExecute::RunTest(const FString& InParame
 		AnimNextGraph->AllocateInstance(GraphInstance);
 
 		FParamStack::FPushedLayerHandle LayerHandle = ParamStack->PushValues(
-			"UpdateCount", (int32)0,
-			"EvaluateCount", (int32)0,
-			"SomeInt32", (int32)0,
-			"SomeFloat", 0.0f,
-			"SomeLatentInt32", (int32)0,
-			"SomeOtherLatentInt32", (int32)0,
-			"SomeLatentFloat", 0.0f
+			"/Engine/Transient.TestAnimNextGraph:UpdateCount", (int32)0,
+			"/Engine/Transient.TestAnimNextGraph:EvaluateCount", (int32)0,
+			"/Engine/Transient.TestAnimNextGraph:SomeInt32", (int32)0,
+			"/Engine/Transient.TestAnimNextGraph:SomeFloat", 0.0f,
+			"/Engine/Transient.TestAnimNextGraph:SomeLatentInt32", (int32)0,
+			"/Engine/Transient.TestAnimNextGraph:SomeOtherLatentInt32", (int32)0,
+			"/Engine/Transient.TestAnimNextGraph:SomeLatentFloat", 0.0f
 		);
 
 		{
@@ -283,13 +284,13 @@ bool FAnimationAnimNextRuntimeTest_GraphExecute::RunTest(const FString& InParame
 			(void)UE::AnimNext::EvaluateGraph(GraphInstance);
 		}
 
-		AddErrorIfFalse(ParamStack->GetParam<int32>("UpdateCount") == 1, "FAnimationAnimNextRuntimeTest_GraphExecute -> Unexpected update count");
-		AddErrorIfFalse(ParamStack->GetParam<int32>("EvaluateCount") == 1, "FAnimationAnimNextRuntimeTest_GraphExecute -> Unexpected evaluate count");
-		AddErrorIfFalse(ParamStack->GetParam<int32>("SomeInt32") == 78, "FAnimationAnimNextRuntimeTest_GraphExecute -> Unexpected SomeInt32 value");
-		AddErrorIfFalse(ParamStack->GetParam<float>("SomeFloat") == 142.33f, "FAnimationAnimNextRuntimeTest_GraphExecute -> Unexpected SomeFloat value");
-		AddErrorIfFalse(ParamStack->GetParam<int32>("SomeLatentInt32") == 3, "FAnimationAnimNextRuntimeTest_GraphExecute -> Unexpected SomeLatentInt32 value");
-		AddErrorIfFalse(ParamStack->GetParam<int32>("SomeOtherLatentInt32") == 3, "FAnimationAnimNextRuntimeTest_GraphExecute -> Unexpected SomeOtherLatentInt32 value");
-		AddErrorIfFalse(ParamStack->GetParam<float>("SomeLatentFloat") == 34.0f, "FAnimationAnimNextRuntimeTest_GraphExecute -> Unexpected SomeLatentFloat value");
+		AddErrorIfFalse(ParamStack->GetParam<int32>("/Engine/Transient.TestAnimNextGraph:UpdateCount") == 1, "FAnimationAnimNextRuntimeTest_GraphExecute -> Unexpected update count");
+		AddErrorIfFalse(ParamStack->GetParam<int32>("/Engine/Transient.TestAnimNextGraph:EvaluateCount") == 1, "FAnimationAnimNextRuntimeTest_GraphExecute -> Unexpected evaluate count");
+		AddErrorIfFalse(ParamStack->GetParam<int32>("/Engine/Transient.TestAnimNextGraph:SomeInt32") == 78, "FAnimationAnimNextRuntimeTest_GraphExecute -> Unexpected SomeInt32 value");
+		AddErrorIfFalse(ParamStack->GetParam<float>("/Engine/Transient.TestAnimNextGraph:SomeFloat") == 142.33f, "FAnimationAnimNextRuntimeTest_GraphExecute -> Unexpected SomeFloat value");
+		AddErrorIfFalse(ParamStack->GetParam<int32>("/Engine/Transient.TestAnimNextGraph:SomeLatentInt32") == 3, "FAnimationAnimNextRuntimeTest_GraphExecute -> Unexpected SomeLatentInt32 value");
+		AddErrorIfFalse(ParamStack->GetParam<int32>("/Engine/Transient.TestAnimNextGraph:SomeOtherLatentInt32") == 3, "FAnimationAnimNextRuntimeTest_GraphExecute -> Unexpected SomeOtherLatentInt32 value");
+		AddErrorIfFalse(ParamStack->GetParam<float>("/Engine/Transient.TestAnimNextGraph:SomeLatentFloat") == 34.0f, "FAnimationAnimNextRuntimeTest_GraphExecute -> Unexpected SomeLatentFloat value");
 
 		ParamStack->PopLayer(LayerHandle);
 		GraphInstance.Release();
@@ -321,7 +322,7 @@ bool FAnimationAnimNextRuntimeTest_GraphExecuteLatent::RunTest(const FString& In
 
 		UE_RETURN_ON_ERROR(EditorData->AddAnimationGraph(FRigUnit_AnimNextGraphRoot::DefaultEntryPoint, false) != nullptr, "FAnimationAnimNextRuntimeTest_GraphExecuteLatent -> Failed to add animation graph");
 
-		URigVMController* Controller = EditorData->GetRigVMClient()->GetController(EditorData->GetRigVMClient()->GetDefaultModel());
+		UAnimNextGraph_Controller* Controller = Cast<UAnimNextGraph_Controller>(EditorData->GetRigVMClient()->GetController(EditorData->GetRigVMClient()->GetDefaultModel()));
 		UE_RETURN_ON_ERROR(Controller != nullptr, "FAnimationAnimNextRuntimeTest_GraphExecuteLatent -> Failed to get RigVM controller");
 
 		// Find graph entry point
@@ -396,15 +397,9 @@ bool FAnimationAnimNextRuntimeTest_GraphExecuteLatent::RunTest(const FString& In
 		}
 
 		{
-			const FRigVMDispatchFactory* GetParameterFactory = FRigVMRegistry::Get().RegisterFactory(FRigVMDispatch_GetParameter::StaticStruct());
-			const FName GetParameterNotation = GetParameterFactory->GetTemplate()->GetNotation();
-
-			URigVMNode* GetParameterNode = Controller->AddTemplateNode(GetParameterNotation);
+			URigVMNode* GetParameterNode = Controller->AddGetAnimNextParameterNode(FVector2D::ZeroVector, UncookedOnly::FUtils::GetQualifiedName(AnimNextGraph, "SomeSourceInt"), FAnimNextParamType::GetType<int32>());
 			UE_RETURN_ON_ERROR(GetParameterNode != nullptr, TEXT("FAnimationAnimNextRuntimeTest_GraphExecuteLatent -> Failed to create GetParameter node"));
 
-			UE_RETURN_ON_ERROR(Controller->ResolveWildCardPin(GetParameterNode->FindPin(TEXT("Value")), RigVMTypeUtils::TypeIndex::Int32), TEXT("FAnimationAnimNextRuntimeTest_GraphExecuteLatent -> Failed to resolve wildcard for GetParameter node"));
-			Controller->SetPinDefaultValue(GetParameterNode->FindPin(TEXT("Parameter"))->GetPinPath(), TEXT("SomeSourceInt"));
-
 			Controller->AddLink(
 				GetParameterNode->FindPin(TEXT("Value")),
 				DecoratorPin->GetSubPins()[4]);	// SomeOtherLatentInt32
@@ -417,14 +412,14 @@ bool FAnimationAnimNextRuntimeTest_GraphExecuteLatent::RunTest(const FString& In
 		AnimNextGraph->AllocateInstance(GraphInstance);
 
 		FParamStack::FPushedLayerHandle LayerHandle = ParamStack->PushValues(
-			"UpdateCount", (int32)0,
-			"EvaluateCount", (int32)0,
-			"SomeSourceInt", (int32)1223,
-			"SomeInt32", (int32)0,
-			"SomeFloat", 0.0f,
-			"SomeLatentInt32", (int32)0,
-			"SomeOtherLatentInt32", (int32)0,
-			"SomeLatentFloat", 0.0f
+			UncookedOnly::FUtils::GetQualifiedName(AnimNextGraph, "UpdateCount"), (int32)0,
+			UncookedOnly::FUtils::GetQualifiedName(AnimNextGraph, "EvaluateCount"), (int32)0,
+			UncookedOnly::FUtils::GetQualifiedName(AnimNextGraph, "SomeSourceInt"), (int32)1223,
+			UncookedOnly::FUtils::GetQualifiedName(AnimNextGraph, "SomeInt32"), (int32)0,
+			UncookedOnly::FUtils::GetQualifiedName(AnimNextGraph, "SomeFloat"), 0.0f,
+			UncookedOnly::FUtils::GetQualifiedName(AnimNextGraph, "SomeLatentInt32"), (int32)0,
+			UncookedOnly::FUtils::GetQualifiedName(AnimNextGraph, "SomeOtherLatentInt32"), (int32)0,
+			UncookedOnly::FUtils::GetQualifiedName(AnimNextGraph, "SomeLatentFloat"), 0.0f
 		);
 
 		{
@@ -432,13 +427,13 @@ bool FAnimationAnimNextRuntimeTest_GraphExecuteLatent::RunTest(const FString& In
 			(void)UE::AnimNext::EvaluateGraph(GraphInstance);
 		}
 
-		AddErrorIfFalse(ParamStack->GetParam<int32>("UpdateCount") == 1, "FAnimationAnimNextRuntimeTest_GraphExecuteLatent -> Unexpected update count");
-		AddErrorIfFalse(ParamStack->GetParam<int32>("EvaluateCount") == 1, "FAnimationAnimNextRuntimeTest_GraphExecuteLatent -> Unexpected evaluate count");
-		AddErrorIfFalse(ParamStack->GetParam<int32>("SomeInt32") == 78, "FAnimationAnimNextRuntimeTest_GraphExecuteLatent -> Unexpected SomeInt32 value");
-		AddErrorIfFalse(ParamStack->GetParam<float>("SomeFloat") == 142.33f, "FAnimationAnimNextRuntimeTest_GraphExecuteLatent -> Unexpected SomeFloat value");
-		AddErrorIfFalse(ParamStack->GetParam<int32>("SomeLatentInt32") == 33, "FAnimationAnimNextRuntimeTest_GraphExecuteLatent -> Unexpected SomeLatentInt32 value");
-		AddErrorIfFalse(ParamStack->GetParam<int32>("SomeOtherLatentInt32") == 1223, "FAnimationAnimNextRuntimeTest_GraphExecuteLatent -> Unexpected SomeOtherLatentInt32 value");
-		AddErrorIfFalse(ParamStack->GetParam<float>("SomeLatentFloat") == 1123.31f, "FAnimationAnimNextRuntimeTest_GraphExecuteLatent -> Unexpected SomeLatentFloat value");
+		AddErrorIfFalse(ParamStack->GetParam<int32>(UncookedOnly::FUtils::GetQualifiedName(AnimNextGraph, "UpdateCount")) == 1, "FAnimationAnimNextRuntimeTest_GraphExecuteLatent -> Unexpected update count");
+		AddErrorIfFalse(ParamStack->GetParam<int32>(UncookedOnly::FUtils::GetQualifiedName(AnimNextGraph, "EvaluateCount")) == 1, "FAnimationAnimNextRuntimeTest_GraphExecuteLatent -> Unexpected evaluate count");
+		AddErrorIfFalse(ParamStack->GetParam<int32>(UncookedOnly::FUtils::GetQualifiedName(AnimNextGraph, "SomeInt32")) == 78, "FAnimationAnimNextRuntimeTest_GraphExecuteLatent -> Unexpected SomeInt32 value");
+		AddErrorIfFalse(ParamStack->GetParam<float>(UncookedOnly::FUtils::GetQualifiedName(AnimNextGraph, "SomeFloat")) == 142.33f, "FAnimationAnimNextRuntimeTest_GraphExecuteLatent -> Unexpected SomeFloat value");
+		AddErrorIfFalse(ParamStack->GetParam<int32>(UncookedOnly::FUtils::GetQualifiedName(AnimNextGraph, "SomeLatentInt32")) == 33, "FAnimationAnimNextRuntimeTest_GraphExecuteLatent -> Unexpected SomeLatentInt32 value");
+		AddErrorIfFalse(ParamStack->GetParam<int32>(UncookedOnly::FUtils::GetQualifiedName(AnimNextGraph, "SomeOtherLatentInt32")) == 1223, "FAnimationAnimNextRuntimeTest_GraphExecuteLatent -> Unexpected SomeOtherLatentInt32 value");
+		AddErrorIfFalse(ParamStack->GetParam<float>(UncookedOnly::FUtils::GetQualifiedName(AnimNextGraph, "SomeLatentFloat")) == 1123.31f, "FAnimationAnimNextRuntimeTest_GraphExecuteLatent -> Unexpected SomeLatentFloat value");
 
 		ParamStack->PopLayer(LayerHandle);
 		GraphInstance.Release();
