CIS: Fix non-unity build

diff --git a/Engine/Plugins/Experimental/AnimNext/Source/AnimNext/Private/Scheduler/Ports/AnimNextSchedulePort_AnimNextMeshComponentPose.h b/Engine/Plugins/Experimental/AnimNext/Source/AnimNext/Private/Scheduler/Ports/AnimNextSchedulePort_AnimNextMeshComponentPose.h
--- a/Engine/Plugins/Experimental/AnimNext/Source/AnimNext/Private/Scheduler/Ports/AnimNextSchedulePort_AnimNextMeshComponentPose.h
+++ b/Engine/Plugins/Experimental/AnimNext/Source/AnimNext/Private/Scheduler/Ports/AnimNextSchedulePort_AnimNextMeshComponentPose.h
@@ -3,6 +3,7 @@
 #pragma once
 
 #include "Scheduler/AnimNextSchedulePort.h"
+#include "Param/AnimNextEditorParam.h"
 #include "AnimNextSchedulePort_AnimNextMeshComponentPose.generated.h"
 
 UCLASS(DisplayName = "AnimNext Mesh Component Pose")
