AnimNext: Fixed issue with SkeletonPoseIndex bone table in reference pose being sliced.

diff --git a/Engine/Plugins/Experimental/AnimNext/Source/AnimNext/Public/ReferencePose.h b/Engine/Plugins/Experimental/AnimNext/Source/AnimNext/Public/ReferencePose.h
--- a/Engine/Plugins/Experimental/AnimNext/Source/AnimNext/Public/ReferencePose.h
+++ b/Engine/Plugins/Experimental/AnimNext/Source/AnimNext/Public/ReferencePose.h
@@ -46,7 +46,7 @@ struct TReferencePose
 
 	// List of skeleton bone indices for each LOD
 	// Each list of skeleton bone indices is a mapping of: SkeletonBoneIndex -> LODSortedBoneIndex
-	// When fast path is enabled, we have a single LOD entry that we truncate to the number of bones for each LOD
+	// When fast path is enabled, we have a single LOD entry
 	TArray<TArray<FBoneIndexType, AllocatorType>, AllocatorType> SkeletonBoneIndexToLODBoneIndexMapPerLOD;
 
 	// Number of bones for each LOD
@@ -152,10 +152,8 @@ struct TReferencePose
 
 		if (LODLevel >= 0 && (IsFastPath() || LODLevel < SkeletonBoneIndexToLODBoneIndexMapPerLOD.Num()))
 		{
-			const int32 NumBonesForLOD = GetNumBonesForLOD(LODLevel);
 			const int32 LODIndex = IsFastPath() ? 0 : LODLevel;
-
-			ArrayView = MakeArrayView(SkeletonBoneIndexToLODBoneIndexMapPerLOD[LODIndex].GetData(), NumBonesForLOD);
+			ArrayView = SkeletonBoneIndexToLODBoneIndexMapPerLOD[LODIndex];
 		}
 
 		return ArrayView;
