AnimNext: Switched Evaluation VM to use static assert for unsupported types.

diff --git a/Engine/Plugins/Experimental/AnimNext/Source/AnimNext/Public/EvaluationVM/EvaluationVM.h b/Engine/Plugins/Experimental/AnimNext/Source/AnimNext/Public/EvaluationVM/EvaluationVM.h
--- a/Engine/Plugins/Experimental/AnimNext/Source/AnimNext/Public/EvaluationVM/EvaluationVM.h
+++ b/Engine/Plugins/Experimental/AnimNext/Source/AnimNext/Public/EvaluationVM/EvaluationVM.h
@@ -28,7 +28,8 @@ namespace UE::AnimNext
 	template<typename Type>
 	constexpr uint32 GetTypeID()
 	{
-		checkf(false, TEXT("Not implemented! See ANIM_NEXT_ENABLE_EVALUATION_STACK_USAGE for details and register your type"));
+		static_assert(sizeof(Type) != sizeof(Type),
+			"Not implemented! See ANIM_NEXT_ENABLE_EVALUATION_STACK_USAGE for details and register your type");
 		return 0;
 	}
 
