CIS: Fix static analysis

diff --git a/Engine/Plugins/Experimental/AnimNext/Source/AnimNextTestSuite/Private/AnimNextEditorTest.cpp b/Engine/Plugins/Experimental/AnimNext/Source/AnimNextTestSuite/Private/AnimNextEditorTest.cpp
--- a/Engine/Plugins/Experimental/AnimNext/Source/AnimNextTestSuite/Private/AnimNextEditorTest.cpp
+++ b/Engine/Plugins/Experimental/AnimNext/Source/AnimNextTestSuite/Private/AnimNextEditorTest.cpp
@@ -149,7 +149,7 @@ bool FEditor_Parameters::RunTest(const FString& InParameters)
 		UE_RETURN_ON_ERROR(ParameterEntry != nullptr, TEXT("Could not find new parameter entry."));
 
 		UAnimNextGraph_EventGraph* EventGraph = EditorData->AddEventGraph(TEXT("TestGraph1"));
-		AddErrorIfFalse(EventGraph != nullptr, TEXT("Could not create new event graph in graph."));
+		UE_RETURN_ON_ERROR(EventGraph != nullptr, TEXT("Could not create new event graph in graph."));
 
 		URigVMGraph* RigVMGraph = EventGraph->GetRigVMGraph();
 		UE_RETURN_ON_ERROR(RigVMGraph->GetNodes().Num() == 1, TEXT("Unexpected number of nodes in new event graph."));
