Fix AnimNext automated test

Default anim graph entry point was not getting properly cached

diff --git a/Engine/Plugins/Experimental/AnimNext/Source/AnimNext/Internal/Graph/AnimNextGraph.h b/Engine/Plugins/Experimental/AnimNext/Source/AnimNext/Internal/Graph/AnimNextGraph.h
--- a/Engine/Plugins/Experimental/AnimNext/Source/AnimNext/Internal/Graph/AnimNextGraph.h
+++ b/Engine/Plugins/Experimental/AnimNext/Source/AnimNext/Internal/Graph/AnimNextGraph.h
@@ -111,6 +111,9 @@ protected:
 	// Set the default entry point name (unqualified)
 	void SetDefaultEntryPoint(FName InEntryPoint);
 
+	// Cache the default entry point
+	void CacheDefaultEntryPoint() const;
+
 	friend class UAnimNextGraphFactory;
 	friend class UAnimNextGraph_EditorData;
 	friend class UAnimNextGraph_Parameter;
@@ -160,9 +163,9 @@ protected:
 	UPROPERTY(EditAnywhere, Category = "Graph")
 	FName DefaultEntryPoint = FRigUnit_AnimNextGraphRoot::DefaultEntryPoint;
 
-	// Cached fully-qualified entry point, set via by SetDefaultEntryPoint.
+	// Cached fully-qualified entry point, set via by SetDefaultEntryPoint/GetDefaultEntryPoint.
 	UPROPERTY(Transient)
-	FName CachedDefaultEntryPoint = NAME_None;
+	mutable FName CachedDefaultEntryPoint = NAME_None;
 
 	// Hash of required parameters
 	UPROPERTY()
diff --git a/Engine/Plugins/Experimental/AnimNext/Source/AnimNext/Private/Graph/AnimNextGraph.cpp b/Engine/Plugins/Experimental/AnimNext/Source/AnimNext/Private/Graph/AnimNextGraph.cpp
--- a/Engine/Plugins/Experimental/AnimNext/Source/AnimNext/Private/Graph/AnimNextGraph.cpp
+++ b/Engine/Plugins/Experimental/AnimNext/Source/AnimNext/Private/Graph/AnimNextGraph.cpp
@@ -160,7 +160,7 @@ void UAnimNextGraph::Serialize(FArchive& Ar)
 			}
 		}
 
-		SetDefaultEntryPoint(DefaultEntryPoint);
+		CacheDefaultEntryPoint();
 	}
 	else if (Ar.IsSaving())
 	{
@@ -189,7 +189,7 @@ void UAnimNextGraph::PostEditChangeProperty(FPropertyChangedEvent& PropertyChang
 {
 	if(PropertyChangedEvent.Property && PropertyChangedEvent.Property->GetFName() == GET_MEMBER_NAME_CHECKED(UAnimNextGraph, DefaultEntryPoint))
 	{
-		SetDefaultEntryPoint(DefaultEntryPoint); 
+		CacheDefaultEntryPoint(); 
 	}
 }
 #endif
@@ -264,14 +264,24 @@ void UAnimNextGraph::SetDefaultEntryPoint(FName InEntryPoint)
 {
 	DefaultEntryPoint = InEntryPoint;
 
-	TStringBuilder<256> StringBuilder;
-	StringBuilder.Append(GetPathName());
-	StringBuilder.Append(TEXT(":"));
-	DefaultEntryPoint.AppendString(StringBuilder);
-	CachedDefaultEntryPoint = FName(StringBuilder.ToView());
+	CacheDefaultEntryPoint();
 }
 
 FName UAnimNextGraph::GetDefaultEntryPoint() const
 {
+	if(CachedDefaultEntryPoint == NAME_None)
+	{
+		CacheDefaultEntryPoint();
+	}
+
 	return CachedDefaultEntryPoint;
+}
+
+void UAnimNextGraph::CacheDefaultEntryPoint() const
+{
+	TStringBuilder<256> StringBuilder;
+	StringBuilder.Append(GetPathName());
+	StringBuilder.Append(TEXT(":"));
+	DefaultEntryPoint.AppendString(StringBuilder);
+	CachedDefaultEntryPoint = FName(StringBuilder.ToView());
 }
\ No newline at end of file
