Fix crash in FInstancedIdProvider after InstancedStruct PropertyEditor optimization
* New provider added while optimization was in development.
* Crash is due to assumption that a PropertyNode which is a struct and has IsPropertyIndirection true would have a parent that is the indirection.  In this case, that is not so - the struct is standalone.   The provider can directly provide the address and type without enumerating instances over the parent object.

diff --git a/Engine/Plugins/Experimental/AnimNext/Source/AnimNextEditor/Private/Param/SParameterPicker.cpp b/Engine/Plugins/Experimental/AnimNext/Source/AnimNextEditor/Private/Param/SParameterPicker.cpp
--- a/Engine/Plugins/Experimental/AnimNext/Source/AnimNextEditor/Private/Param/SParameterPicker.cpp
+++ b/Engine/Plugins/Experimental/AnimNext/Source/AnimNextEditor/Private/Param/SParameterPicker.cpp
@@ -69,7 +69,7 @@ public:
 
 	virtual bool IsPropertyIndirection() const override
 	{
-		return true;
+		return false;
 	}
 
 	virtual uint8* GetValueBaseAddress(uint8* ParentValueAddress, const UStruct* ExpectedBaseStructure) const override
