Construct and destruct AnimNext latent properties
Add support for OnBecomeRelevant during update traversal

diff --git a/Engine/Plugins/Experimental/AnimNext/Source/AnimNext/Private/Graph/AnimNextGraphInstance.cpp b/Engine/Plugins/Experimental/AnimNext/Source/AnimNext/Private/Graph/AnimNextGraphInstance.cpp
--- a/Engine/Plugins/Experimental/AnimNext/Source/AnimNext/Private/Graph/AnimNextGraphInstance.cpp
+++ b/Engine/Plugins/Experimental/AnimNext/Source/AnimNext/Private/Graph/AnimNextGraphInstance.cpp
@@ -85,6 +85,11 @@ bool FAnimNextGraphInstance::IsRoot() const
 	return this == RootGraphInstance;
 }
 
+bool FAnimNextGraphInstance::HasUpdated() const
+{
+	return bHasUpdatedOnce;
+}
+
 void FAnimNextGraphInstance::AddStructReferencedObjects(FReferenceCollector& Collector)
 {
 	if (!IsRoot())
@@ -150,6 +155,11 @@ void FAnimNextGraphInstance::CollectInputTraitEvents(UE::AnimNext::FTraitEventLi
 	Swap(OutInputEvents, RootGraphInstance->InputEventList);
 }
 
+void FAnimNextGraphInstance::Update()
+{
+	bHasUpdatedOnce = true;
+}
+
 void FAnimNextGraphInstance::ExecuteLatentPins(const TConstArrayView<UE::AnimNext::FLatentPropertyHandle>& LatentHandles, void* DestinationBasePtr, bool bIsFrozen)
 {
 	SCOPE_CYCLE_COUNTER(STAT_AnimNext_Graph_RigVM);
@@ -181,6 +191,7 @@ void FAnimNextGraphInstance::Freeze()
 	GraphInstancePtr.Reset();
 	ExtendedExecuteContext.Reset();
 	Components.Empty();
+	bHasUpdatedOnce = false;
 }
 
 void FAnimNextGraphInstance::Thaw()
diff --git a/Engine/Plugins/Experimental/AnimNext/Source/AnimNext/Private/Graph/AnimNextGraphInstancePtr.cpp b/Engine/Plugins/Experimental/AnimNext/Source/AnimNext/Private/Graph/AnimNextGraphInstancePtr.cpp
--- a/Engine/Plugins/Experimental/AnimNext/Source/AnimNext/Private/Graph/AnimNextGraphInstancePtr.cpp
+++ b/Engine/Plugins/Experimental/AnimNext/Source/AnimNext/Private/Graph/AnimNextGraphInstancePtr.cpp
@@ -64,6 +64,11 @@ bool FAnimNextGraphInstancePtr::IsRoot() const
 	return Impl ? Impl->IsRoot() : true;
 }
 
+bool FAnimNextGraphInstancePtr::HasUpdated() const
+{
+	return Impl ? Impl->HasUpdated() : false;
+}
+
 void FAnimNextGraphInstancePtr::AddStructReferencedObjects(FReferenceCollector& Collector)
 {
 	if (Impl)
@@ -106,3 +111,9 @@ void FAnimNextGraphInstancePtr::CollectInputTraitEvents(UE::AnimNext::FTraitEven
 	check(Impl);
 	Impl->CollectInputTraitEvents(OutInputEvents);
 }
+
+void FAnimNextGraphInstancePtr::Update()
+{
+	check(Impl);
+	Impl->Update();
+}
diff --git a/Engine/Plugins/Experimental/AnimNext/Source/AnimNext/Public/Graph/AnimNextGraphInstance.h b/Engine/Plugins/Experimental/AnimNext/Source/AnimNext/Public/Graph/AnimNextGraphInstance.h
--- a/Engine/Plugins/Experimental/AnimNext/Source/AnimNext/Public/Graph/AnimNextGraphInstance.h
+++ b/Engine/Plugins/Experimental/AnimNext/Source/AnimNext/Public/Graph/AnimNextGraphInstance.h
@@ -75,6 +75,9 @@ struct ANIMNEXT_API FAnimNextGraphInstance
 	// Returns whether or not this graph instance is the root graph instance or false otherwise
 	bool IsRoot() const;
 
+	// Returns whether or not this graph instance has updated at least once
+	bool HasUpdated() const;
+
 	// Adds strong/hard object references during GC
 	void AddStructReferencedObjects(class FReferenceCollector& Collector);
 
@@ -105,6 +108,9 @@ struct ANIMNEXT_API FAnimNextGraphInstance
 	// Collects all currently queued input trait events (thread safe)
 	void CollectInputTraitEvents(UE::AnimNext::FTraitEventList& OutInputEvents);
 
+	// Called each time the graph updates
+	void Update();
+
 private:
 	// Returns a pointer to the specified component, or nullptr if not found
 	UE::AnimNext::FGraphInstanceComponent* TryGetComponent(int32 ComponentNameHash, FName ComponentName) const;
@@ -152,6 +158,9 @@ private:
 	// Lock to ensure input event list actions are thread safe
 	FRWLock InputEventListLock;
 
+	// Whether or not this graph has updated once
+	bool bHasUpdatedOnce = false;
+
 	friend UAnimNextGraph;					// The graph is the one that allocates instances
 	friend FRigUnit_AnimNextGraphEvaluator;	// We evaluate the instance
 	friend UE::AnimNext::FExecutionContext;
diff --git a/Engine/Plugins/Experimental/AnimNext/Source/AnimNext/Public/Graph/AnimNextGraphInstancePtr.h b/Engine/Plugins/Experimental/AnimNext/Source/AnimNext/Public/Graph/AnimNextGraphInstancePtr.h
--- a/Engine/Plugins/Experimental/AnimNext/Source/AnimNext/Public/Graph/AnimNextGraphInstancePtr.h
+++ b/Engine/Plugins/Experimental/AnimNext/Source/AnimNext/Public/Graph/AnimNextGraphInstancePtr.h
@@ -63,6 +63,9 @@ struct ANIMNEXT_API FAnimNextGraphInstancePtr
 	// Returns whether or not this graph instance is the root graph instance or false otherwise
 	bool IsRoot() const;
 
+	// Returns whether or not this graph instance has updated at least once
+	bool HasUpdated() const;
+
 	// Adds strong/hard object references during GC
 	void AddStructReferencedObjects(class FReferenceCollector& Collector);
 
@@ -93,6 +96,9 @@ struct ANIMNEXT_API FAnimNextGraphInstancePtr
 	// Collects all currently queued input trait events (thread safe)
 	void CollectInputTraitEvents(UE::AnimNext::FTraitEventList& OutInputEvents);
 
+	// Called each time the graph updates
+	void Update();
+
 private:
 	// Returns a pointer to the specified component, or nullptr if not found
 	UE::AnimNext::FGraphInstanceComponent* TryGetComponent(int32 ComponentNameHash, FName ComponentName) const;
diff --git a/Engine/Plugins/Experimental/AnimNext/Source/AnimNext/Public/TraitCore/Trait.h b/Engine/Plugins/Experimental/AnimNext/Source/AnimNext/Public/TraitCore/Trait.h
--- a/Engine/Plugins/Experimental/AnimNext/Source/AnimNext/Public/TraitCore/Trait.h
+++ b/Engine/Plugins/Experimental/AnimNext/Source/AnimNext/Public/TraitCore/Trait.h
@@ -69,13 +69,19 @@ class FArchive;
 	} \
 	void TraitName::ConstructTraitInstance(const UE::AnimNext::FExecutionContext& Context, const UE::AnimNext::FTraitBinding& Binding) const \
 	{ \
+		/* Construct the base struct first */ \
 		FInstanceData* Data = new(Binding.GetInstanceData<FInstanceData>()) FInstanceData(); \
+		/* Then construct our latent properties, the Construct implementation below might need them */ \
+		FSharedData::ConstructLatentProperties(Binding); \
+		/* Construct our typed instance last */ \
 		Data->Construct(Context, Binding); \
 	} \
 	void TraitName::DestructTraitInstance(const UE::AnimNext::FExecutionContext& Context, const UE::AnimNext::FTraitBinding& Binding) const \
 	{ \
+		/* Destruction is reverse order of construction above */ \
 		FInstanceData* Data = Binding.GetInstanceData<FInstanceData>(); \
 		Data->Destruct(Context, Binding); \
+		FSharedData::DestructLatentProperties(Binding); \
 		Data->~FInstanceData(); \
 	}
 
diff --git a/Engine/Plugins/Experimental/AnimNext/Source/AnimNext/Public/TraitCore/TraitSharedData.h b/Engine/Plugins/Experimental/AnimNext/Source/AnimNext/Public/TraitCore/TraitSharedData.h
--- a/Engine/Plugins/Experimental/AnimNext/Source/AnimNext/Public/TraitCore/TraitSharedData.h
+++ b/Engine/Plugins/Experimental/AnimNext/Source/AnimNext/Public/TraitCore/TraitSharedData.h
@@ -9,6 +9,53 @@
 
 #include "TraitSharedData.generated.h"
 
+namespace UE::AnimNext
+{
+	struct FTraitBinding;
+}
+
+#define ANIM_NEXT_IMPL_EXECUTE_LATENT_CONSTRUCTOR_FOR_PROPERTY(PropertyName) \
+	/* @see GetLatentPropertyIndex below for details */ \
+	[&Binding, LatentPropertyHandles](int32 LatentPropertyIndex) \
+	{ \
+		const UE::AnimNext::FLatentPropertyHandle LatentPropertyHandle = LatentPropertyHandles[LatentPropertyIndex - 1]; \
+		if (LatentPropertyHandle.IsOffsetValid()) \
+		{ \
+			/* We remove the 'const' explicitly to avoid exposing a mutable version of the getter */ \
+			new(const_cast<decltype(PropertyName)*>(Binding.GetLatentProperty<decltype(PropertyName)>(LatentPropertyHandle))) decltype(PropertyName)(); \
+		} \
+	}(GetLatentPropertyIndex(offsetof(Self, PropertyName))); \
+
+#define ANIM_NEXT_IMPL_DEFINE_LATENT_CONSTRUCTOR(EnumeratorMacro) \
+	/** @see FAnimNextTraitSharedData::ConstructLatentProperties */ \
+	static void ConstructLatentProperties(const UE::AnimNext::FTraitBinding& Binding) \
+	{ \
+		/* Construct our latent properties */ \
+		const UE::AnimNext::FLatentPropertyHandle* LatentPropertyHandles = Binding.GetLatentPropertyHandles(); \
+		EnumeratorMacro(ANIM_NEXT_IMPL_EXECUTE_LATENT_CONSTRUCTOR_FOR_PROPERTY) \
+	} \
+
+#define ANIM_NEXT_IMPL_EXECUTE_LATENT_DESTRUCTOR_FOR_PROPERTY(PropertyName) \
+	/* @see GetLatentPropertyIndex below for details */ \
+	[&Binding, LatentPropertyHandles](int32 LatentPropertyIndex) \
+	{ \
+		const UE::AnimNext::FLatentPropertyHandle LatentPropertyHandle = LatentPropertyHandles[LatentPropertyIndex - 1]; \
+		if (LatentPropertyHandle.IsOffsetValid()) \
+		{ \
+			/* We remove the 'const' explicitly to avoid exposing a mutable version of the getter */ \
+			const_cast<decltype(PropertyName)*>(Binding.GetLatentProperty<decltype(PropertyName)>(LatentPropertyHandle))->~decltype(PropertyName)(); \
+		} \
+	}(GetLatentPropertyIndex(offsetof(Self, PropertyName))); \
+
+#define ANIM_NEXT_IMPL_DEFINE_LATENT_DESTRUCTOR(EnumeratorMacro) \
+	/** @see FAnimNextTraitSharedData::DestructLatentProperties */ \
+	static void DestructLatentProperties(const UE::AnimNext::FTraitBinding& Binding) \
+	{ \
+		/* Destruct our latent properties */ \
+		const UE::AnimNext::FLatentPropertyHandle* LatentPropertyHandles = Binding.GetLatentPropertyHandles(); \
+		EnumeratorMacro(ANIM_NEXT_IMPL_EXECUTE_LATENT_DESTRUCTOR_FOR_PROPERTY) \
+	} \
+
 #define ANIM_NEXT_IMPL_GET_LATENT_PROPERTY_INDEX_FOR_PROPERTY(PropertyName) \
 	LatentPropertyIndex--; \
 	if (LatentPropertyOffset == offsetof(Self, PropertyName)) \
@@ -35,7 +82,7 @@
 	} \
 
 #define ANIM_NEXT_IMPL_DEFINE_LATENT_GETTER(PropertyName) \
-	decltype(PropertyName) Get##PropertyName(const UE::AnimNext::FTraitBinding& Binding) const \
+	const decltype(PropertyName)& Get##PropertyName(const UE::AnimNext::FTraitBinding& Binding) const \
 	{ \
 		/* We need a mapping of latent property name/offset to latent property index */ \
 		/* This can be built once at runtime using the UE reflection and cached on first call or using a constexpr function, see below */ \
@@ -64,6 +111,8 @@
   */
 #define GENERATE_TRAIT_LATENT_PROPERTIES(SharedDataType, EnumeratorMacro) \
 	using Self = SharedDataType; \
+	ANIM_NEXT_IMPL_DEFINE_LATENT_CONSTRUCTOR(EnumeratorMacro) \
+	ANIM_NEXT_IMPL_DEFINE_LATENT_DESTRUCTOR(EnumeratorMacro) \
 	ANIM_NEXT_IMPL_DEFINE_GET_LATENT_PROPERTY_INDEX(EnumeratorMacro) \
 	EnumeratorMacro(ANIM_NEXT_IMPL_DEFINE_LATENT_GETTER) \
 
@@ -83,6 +132,16 @@ struct FAnimNextTraitSharedData
 {
 	GENERATED_BODY()
 
+	/**
+	 * Constructs the latent properties on the bound trait instance
+	 */
+	static void ConstructLatentProperties(const UE::AnimNext::FTraitBinding& Binding) {}
+
+	/**
+	 * Destructs the latent properties on the bound trait instance
+	 */
+	static void DestructLatentProperties(const UE::AnimNext::FTraitBinding& Binding) {}
+
 	/**
 	  * Returns the latent property index from a latent property offset
 	  * If the property is found within this type, its latent property 1-based index is returned (a positive value)
diff --git a/Engine/Plugins/Experimental/AnimNextAnimGraph/Source/AnimNextAnimGraph/Private/TraitInterfaces/IUpdate.cpp b/Engine/Plugins/Experimental/AnimNextAnimGraph/Source/AnimNextAnimGraph/Private/TraitInterfaces/IUpdate.cpp
--- a/Engine/Plugins/Experimental/AnimNextAnimGraph/Source/AnimNextAnimGraph/Private/TraitInterfaces/IUpdate.cpp
+++ b/Engine/Plugins/Experimental/AnimNextAnimGraph/Source/AnimNextAnimGraph/Private/TraitInterfaces/IUpdate.cpp
@@ -156,6 +156,15 @@ namespace UE::AnimNext
 		}
 	}
 
+	void IUpdate::OnBecomeRelevant(FUpdateTraversalContext& Context, const TTraitBinding<IUpdate>& Binding, const FTraitUpdateState& TraitState) const
+	{
+		TTraitBinding<IUpdate> SuperBinding;
+		if (Binding.GetStackInterfaceSuper(SuperBinding))
+		{
+			SuperBinding.OnBecomeRelevant(Context, TraitState);
+		}
+	}
+
 	void IUpdate::PreUpdate(FUpdateTraversalContext& Context, const TTraitBinding<IUpdate>& Binding, const FTraitUpdateState& TraitState) const
 	{
 		TTraitBinding<IUpdate> SuperBinding;
@@ -462,6 +471,13 @@ namespace UE::AnimNext
 		Private::FUpdateEventBookkeepingList RootParentBookkeepingEntryList;
 		TraversalContext.RootParentBookkeepingEntryList = &RootParentBookkeepingEntryList;
 
+		const FTraitUpdateState RootState =
+			FTraitUpdateState(DeltaTime)
+			.AsNewlyRelevant(!GraphInstance.HasUpdated());
+
+		// Update the graph instance itself
+		GraphInstance.Update();
+
 		// Grab our input events, we'll propagate them down the graph as we traverse
 		GraphInstance.CollectInputTraitEvents(TraversalContext.InputEventList);
 
@@ -475,7 +491,7 @@ namespace UE::AnimNext
 		}
 
 		// Add the graph root to start the update process
-		Private::FUpdateEntry RootEntry(GraphInstance.GetGraphRootPtr(), FTraitUpdateState(DeltaTime));
+		Private::FUpdateEntry RootEntry(GraphInstance.GetGraphRootPtr(), RootState);
 		TraversalContext.PushUpdateEntry(&RootEntry);
 
 		while (Private::FUpdateEntry* Entry = TraversalContext.PopUpdateEntry())
@@ -497,11 +513,19 @@ namespace UE::AnimNext
 				const bool bIsFrozen = false;	// Not yet supported
 				Entry->TraitStack.SnapshotLatentProperties(bIsFrozen);
 
+				const bool bImplementsIUpdate = Entry->TraitStack.GetInterface(Entry->UpdateTrait);
+
+				// Before we PreUpdate, signal that we became newly relevant
+				if (bImplementsIUpdate && Entry->TraitState.IsNewlyRelevant())
+				{
+					Entry->UpdateTrait.OnBecomeRelevant(TraversalContext, Entry->TraitState);
+				}
+
 				// Raise our input events
 				Private::RaiseTraitEvents(TraversalContext, Entry, TraversalContext.InputEventList);
 
-				// If this trait stack implements IUpdate, call into it
-				if (Entry->TraitStack.GetInterface(Entry->UpdateTrait))
+				// Main update before our children
+				if (bImplementsIUpdate)
 				{
 					Entry->UpdateTrait.PreUpdate(TraversalContext, Entry->TraitState);
 				}
diff --git a/Engine/Plugins/Experimental/AnimNextAnimGraph/Source/AnimNextAnimGraph/Private/Traits/BlendByBool.cpp b/Engine/Plugins/Experimental/AnimNextAnimGraph/Source/AnimNextAnimGraph/Private/Traits/BlendByBool.cpp
--- a/Engine/Plugins/Experimental/AnimNextAnimGraph/Source/AnimNextAnimGraph/Private/Traits/BlendByBool.cpp
+++ b/Engine/Plugins/Experimental/AnimNextAnimGraph/Source/AnimNextAnimGraph/Private/Traits/BlendByBool.cpp
@@ -31,6 +31,19 @@ namespace UE::AnimNext
 		TTraitBinding<IDiscreteBlend> DiscreteBlendTrait;
 		Binding.GetStackInterface(DiscreteBlendTrait);
 
+		// If we were previously relevant, update our status
+		if (InstanceData->PreviousChildIndex != INDEX_NONE)
+		{
+			if (InstanceData->PreviousChildIndex == TRUE_CHILD_INDEX)
+			{
+				InstanceData->bWasTrueChildRelevant = true;
+			}
+			else
+			{
+				InstanceData->bWasFalseChildRelevant = true;
+			}
+		}
+
 		const int32 DestinationChildIndex = DiscreteBlendTrait.GetBlendDestinationChildIndex(Context);
 		if (InstanceData->PreviousChildIndex != DestinationChildIndex)
 		{
@@ -53,11 +66,10 @@ namespace UE::AnimNext
 		const float BlendWeightTrue = DiscreteBlendTrait.GetBlendWeight(Context, TRUE_CHILD_INDEX);
 		if (InstanceData->TrueChild.IsValid() && FAnimWeight::IsRelevant(BlendWeightTrue))
 		{
-			FTraitUpdateState TraitStateTrue = TraitState.WithWeight(BlendWeightTrue);
-			if (DestinationChildIndex != TRUE_CHILD_INDEX)
-			{
-				TraitStateTrue = TraitStateTrue.AsBlendingOut();
-			}
+			FTraitUpdateState TraitStateTrue = TraitState
+				.WithWeight(BlendWeightTrue)
+				.AsBlendingOut(DestinationChildIndex != TRUE_CHILD_INDEX)
+				.AsNewlyRelevant(!InstanceData->bWasTrueChildRelevant);
 
 			TraversalQueue.Push(InstanceData->TrueChild, TraitStateTrue);
 		}
@@ -65,11 +77,10 @@ namespace UE::AnimNext
 		const float BlendWeightFalse = 1.0f - BlendWeightTrue;
 		if (InstanceData->FalseChild.IsValid() && FAnimWeight::IsRelevant(BlendWeightFalse))
 		{
-			FTraitUpdateState TraitStateFalse = TraitState.WithWeight(BlendWeightFalse);
-			if (DestinationChildIndex != FALSE_CHILD_INDEX)
-			{
-				TraitStateFalse = TraitStateFalse.AsBlendingOut();
-			}
+			FTraitUpdateState TraitStateFalse = TraitState
+				.WithWeight(BlendWeightFalse)
+				.AsBlendingOut(DestinationChildIndex != FALSE_CHILD_INDEX)
+				.AsNewlyRelevant(!InstanceData->bWasFalseChildRelevant);
 
 			TraversalQueue.Push(InstanceData->FalseChild, TraitStateFalse);
 		}
@@ -161,10 +172,12 @@ namespace UE::AnimNext
 		if (ChildIndex == TRUE_CHILD_INDEX)
 		{
 			InstanceData->TrueChild.Reset();
+			InstanceData->bWasTrueChildRelevant = false;
 		}
 		else if (ChildIndex == FALSE_CHILD_INDEX)
 		{
 			InstanceData->FalseChild.Reset();
+			InstanceData->bWasFalseChildRelevant = false;
 		}
 	}
 }
diff --git a/Engine/Plugins/Experimental/AnimNextAnimGraph/Source/AnimNextAnimGraph/Private/Traits/BlendByBool.h b/Engine/Plugins/Experimental/AnimNextAnimGraph/Source/AnimNextAnimGraph/Private/Traits/BlendByBool.h
--- a/Engine/Plugins/Experimental/AnimNextAnimGraph/Source/AnimNextAnimGraph/Private/Traits/BlendByBool.h
+++ b/Engine/Plugins/Experimental/AnimNextAnimGraph/Source/AnimNextAnimGraph/Private/Traits/BlendByBool.h
@@ -55,6 +55,8 @@ namespace UE::AnimNext
 			FTraitPtr FalseChild;
 
 			int32 PreviousChildIndex = INDEX_NONE;
+			bool bWasTrueChildRelevant = false;
+			bool bWasFalseChildRelevant = false;
 		};
 
 		// IUpdate impl
diff --git a/Engine/Plugins/Experimental/AnimNextAnimGraph/Source/AnimNextAnimGraph/Private/Traits/BlendTwoWay.cpp b/Engine/Plugins/Experimental/AnimNextAnimGraph/Source/AnimNextAnimGraph/Private/Traits/BlendTwoWay.cpp
--- a/Engine/Plugins/Experimental/AnimNextAnimGraph/Source/AnimNextAnimGraph/Private/Traits/BlendTwoWay.cpp
+++ b/Engine/Plugins/Experimental/AnimNextAnimGraph/Source/AnimNextAnimGraph/Private/Traits/BlendTwoWay.cpp
@@ -57,11 +57,16 @@ namespace UE::AnimNext
 				// We need to blend a child that isn't instanced yet, allocate it
 				InstanceData->ChildA = Context.AllocateNodeInstance(Binding, SharedData->ChildA);
 			}
+			else
+			{
+				InstanceData->bWasChildARelevant = true;
+			}
 
 			if (!FAnimWeight::IsRelevant(BlendWeightB))
 			{
 				// We no longer need this child, release it
 				InstanceData->ChildB.Reset();
+				InstanceData->bWasChildBRelevant = false;
 			}
 		}
 
@@ -72,11 +77,16 @@ namespace UE::AnimNext
 				// We need to blend a child that isn't instanced yet, allocate it
 				InstanceData->ChildB = Context.AllocateNodeInstance(Binding, SharedData->ChildB);
 			}
+			else
+			{
+				InstanceData->bWasChildBRelevant = true;
+			}
 
 			if (FAnimWeight::IsFullWeight(BlendWeightB))
 			{
 				// We no longer need this child, release it
 				InstanceData->ChildA.Reset();
+				InstanceData->bWasChildARelevant = false;
 			}
 		}
 	}
@@ -92,12 +102,12 @@ namespace UE::AnimNext
 		if (InstanceData->ChildA.IsValid())
 		{
 			const float BlendWeightA = 1.0f - BlendWeightB;
-			TraversalQueue.Push(InstanceData->ChildA, TraitState.WithWeight(BlendWeightA));
+			TraversalQueue.Push(InstanceData->ChildA, TraitState.WithWeight(BlendWeightA).AsNewlyRelevant(!InstanceData->bWasChildARelevant));
 		}
 
 		if (InstanceData->ChildB.IsValid())
 		{
-			TraversalQueue.Push(InstanceData->ChildB, TraitState.WithWeight(BlendWeightB));
+			TraversalQueue.Push(InstanceData->ChildB, TraitState.WithWeight(BlendWeightB).AsNewlyRelevant(!InstanceData->bWasChildBRelevant));
 		}
 	}
 
diff --git a/Engine/Plugins/Experimental/AnimNextAnimGraph/Source/AnimNextAnimGraph/Private/Traits/BlendTwoWay.h b/Engine/Plugins/Experimental/AnimNextAnimGraph/Source/AnimNextAnimGraph/Private/Traits/BlendTwoWay.h
--- a/Engine/Plugins/Experimental/AnimNextAnimGraph/Source/AnimNextAnimGraph/Private/Traits/BlendTwoWay.h
+++ b/Engine/Plugins/Experimental/AnimNextAnimGraph/Source/AnimNextAnimGraph/Private/Traits/BlendTwoWay.h
@@ -54,6 +54,9 @@ namespace UE::AnimNext
 		{
 			FTraitPtr ChildA;
 			FTraitPtr ChildB;
+
+			bool bWasChildARelevant = false;
+			bool bWasChildBRelevant = false;
 		};
 
 		// IEvaluate impl
diff --git a/Engine/Plugins/Experimental/AnimNextAnimGraph/Source/AnimNextAnimGraph/Private/Traits/SequencePlayer.cpp b/Engine/Plugins/Experimental/AnimNextAnimGraph/Source/AnimNextAnimGraph/Private/Traits/SequencePlayer.cpp
--- a/Engine/Plugins/Experimental/AnimNextAnimGraph/Source/AnimNextAnimGraph/Private/Traits/SequencePlayer.cpp
+++ b/Engine/Plugins/Experimental/AnimNextAnimGraph/Source/AnimNextAnimGraph/Private/Traits/SequencePlayer.cpp
@@ -25,18 +25,6 @@ namespace UE::AnimNext
 		FTrait::FInstanceData::Construct(Context, Binding);
 
 		IGarbageCollection::RegisterWithGC(Context, Binding);
-
-		const FSharedData* SharedData = Binding.GetSharedData<FSharedData>();
-		FInstanceData* InstanceData = Binding.GetInstanceData<FInstanceData>();
-
-		// Cache the anim sequence we'll play during construction, we don't allow it to change afterwards
-		InstanceData->AnimSequence = SharedData->GetAnimSequence(Binding);
-
-		if (InstanceData->AnimSequence)
-		{
-			const float SequenceLength = InstanceData->AnimSequence->GetPlayLength();
-			InternalTimeAccumulator = FMath::Clamp(SharedData->GetStartPosition(Binding), 0.0f, SequenceLength);
-		}
 	}
 
 	void FSequencePlayerTrait::FInstanceData::Destruct(const FExecutionContext& Context, const FTraitBinding& Binding)
@@ -100,6 +88,25 @@ namespace UE::AnimNext
 		}
 	}
 
+	void FSequencePlayerTrait::OnBecomeRelevant(FUpdateTraversalContext& Context, const TTraitBinding<IUpdate>& Binding, const FTraitUpdateState& TraitState) const
+	{
+		const FSharedData* SharedData = Binding.GetSharedData<FSharedData>();
+		FInstanceData* InstanceData = Binding.GetInstanceData<FInstanceData>();
+
+		// Cache the anim sequence we'll play during construction, we don't allow it to change afterwards
+		InstanceData->AnimSequence = SharedData->GetAnimSequence(Binding);
+
+		float InternalTimeAccumulator = 0.0f;
+		if (UAnimSequence* AnimSeq = InstanceData->AnimSequence.Get())
+		{
+			const float StartPosition = SharedData->GetStartPosition(Binding);
+			const float SequenceLength = AnimSeq->GetPlayLength();
+			InternalTimeAccumulator = FMath::Clamp(StartPosition, 0.0f, SequenceLength);
+		}
+
+		InstanceData->InternalTimeAccumulator = InternalTimeAccumulator;
+	}
+
 	void FSequencePlayerTrait::PreUpdate(FUpdateTraversalContext& Context, const TTraitBinding<IUpdate>& Binding, const FTraitUpdateState& TraitState) const
 	{
 		// We just advance the timeline
diff --git a/Engine/Plugins/Experimental/AnimNextAnimGraph/Source/AnimNextAnimGraph/Private/Traits/SequencePlayer.h b/Engine/Plugins/Experimental/AnimNextAnimGraph/Source/AnimNextAnimGraph/Private/Traits/SequencePlayer.h
--- a/Engine/Plugins/Experimental/AnimNextAnimGraph/Source/AnimNextAnimGraph/Private/Traits/SequencePlayer.h
+++ b/Engine/Plugins/Experimental/AnimNextAnimGraph/Source/AnimNextAnimGraph/Private/Traits/SequencePlayer.h
@@ -77,6 +77,7 @@ namespace UE::AnimNext
 		virtual void AdvanceToRatio(const FExecutionContext& Context, const TTraitBinding<ITimeline>& Binding, float ProgressRatio) const override;
 
 		// IUpdate impl
+		virtual void OnBecomeRelevant(FUpdateTraversalContext& Context, const TTraitBinding<IUpdate>& Binding, const FTraitUpdateState& TraitState) const override;
 		virtual void PreUpdate(FUpdateTraversalContext& Context, const TTraitBinding<IUpdate>& Binding, const FTraitUpdateState& TraitState) const override;
 
 		// IGarbageCollection impl
diff --git a/Engine/Plugins/Experimental/AnimNextAnimGraph/Source/AnimNextAnimGraph/Private/Traits/SubGraphHost.cpp b/Engine/Plugins/Experimental/AnimNextAnimGraph/Source/AnimNextAnimGraph/Private/Traits/SubGraphHost.cpp
--- a/Engine/Plugins/Experimental/AnimNextAnimGraph/Source/AnimNextAnimGraph/Private/Traits/SubGraphHost.cpp
+++ b/Engine/Plugins/Experimental/AnimNextAnimGraph/Source/AnimNextAnimGraph/Private/Traits/SubGraphHost.cpp
@@ -73,9 +73,11 @@ namespace UE::AnimNext
 		FName CurrentActiveEntryPoint = NAME_None;
 		if (bHasActiveSubGraph)
 		{
-			const FSubGraphSlot& SubGraphSlot = InstanceData->SubGraphSlots[InstanceData->CurrentlyActiveSubGraphIndex];
+			FSubGraphSlot& SubGraphSlot = InstanceData->SubGraphSlots[InstanceData->CurrentlyActiveSubGraphIndex];
 			CurrentActiveSubGraph = SubGraphSlot.SubGraph;
 			CurrentActiveEntryPoint = SubGraphSlot.EntryPoint;
+
+			SubGraphSlot.bWasRelevant = true;
 		}
 
 		const TObjectPtr<const UAnimNextGraph> DesiredSubGraph = SharedData->GetSubGraph(Binding);
@@ -147,15 +149,15 @@ namespace UE::AnimNext
 
 		for (int32 SubGraphIndex = 0; SubGraphIndex < NumSubGraphs; ++SubGraphIndex)
 		{
+			const FSubGraphSlot& SubGraphSlot = InstanceData->SubGraphSlots[SubGraphIndex];
 			const float BlendWeight = DiscreteBlendTrait.GetBlendWeight(Context, SubGraphIndex);
 
-			FTraitUpdateState SubGraphTraitState = TraitState.WithWeight(BlendWeight);
-			if (SubGraphIndex != InstanceData->CurrentlyActiveSubGraphIndex)
-			{
-				SubGraphTraitState = SubGraphTraitState.AsBlendingOut();
-			}
+			FTraitUpdateState SubGraphTraitState = TraitState
+				.WithWeight(BlendWeight)
+				.AsBlendingOut(SubGraphIndex != InstanceData->CurrentlyActiveSubGraphIndex)
+				.AsNewlyRelevant(!SubGraphSlot.bWasRelevant);
 
-			TraversalQueue.Push(InstanceData->SubGraphSlots[SubGraphIndex].GraphInstance.GetGraphRootPtr(), SubGraphTraitState);
+			TraversalQueue.Push(SubGraphSlot.GraphInstance.GetGraphRootPtr(), SubGraphTraitState);
 		}
 	}
 
@@ -228,6 +230,7 @@ namespace UE::AnimNext
 			}
 
 			SubGraphEntry.State = ESlotState::Inactive;
+			SubGraphEntry.bWasRelevant = false;
 		}
 	}
 
diff --git a/Engine/Plugins/Experimental/AnimNextAnimGraph/Source/AnimNextAnimGraph/Private/Traits/SubGraphHost.h b/Engine/Plugins/Experimental/AnimNextAnimGraph/Source/AnimNextAnimGraph/Private/Traits/SubGraphHost.h
--- a/Engine/Plugins/Experimental/AnimNextAnimGraph/Source/AnimNextAnimGraph/Private/Traits/SubGraphHost.h
+++ b/Engine/Plugins/Experimental/AnimNextAnimGraph/Source/AnimNextAnimGraph/Private/Traits/SubGraphHost.h
@@ -51,7 +51,7 @@ namespace UE::AnimNext
 	{
 		DECLARE_ANIM_TRAIT(FSubGraphHostTrait, 0xb1085d7b, FBaseTrait)
 
-		enum class ESlotState
+		enum class ESlotState : uint8
 		{
 			ActiveWithGraph,
 			ActiveWithReferencePose,
@@ -71,6 +71,9 @@ namespace UE::AnimNext
 
 			// The current slot state
 			ESlotState State = ESlotState::Inactive;
+
+			// Whether or not this slot state was previously relevant
+			bool bWasRelevant = false;
 		};
 
 		using FSharedData = FAnimNextSubGraphHostTraitSharedData;
diff --git a/Engine/Plugins/Experimental/AnimNextAnimGraph/Source/AnimNextAnimGraph/Public/TraitInterfaces/IUpdate.h b/Engine/Plugins/Experimental/AnimNextAnimGraph/Source/AnimNextAnimGraph/Public/TraitInterfaces/IUpdate.h
--- a/Engine/Plugins/Experimental/AnimNextAnimGraph/Source/AnimNextAnimGraph/Public/TraitInterfaces/IUpdate.h
+++ b/Engine/Plugins/Experimental/AnimNextAnimGraph/Source/AnimNextAnimGraph/Public/TraitInterfaces/IUpdate.h
@@ -52,6 +52,9 @@ namespace UE::AnimNext
 		// Returns whether or not this trait is blending out
 		bool IsBlendingOut() const { return !!bIsBlendingOut; }
 
+		// Returns whether or not this trait is newly relevant
+		bool IsNewlyRelevant() const { return !!bIsNewlyRelevant; }
+
 		// Creates a new instance of the update state with the total weight scaled by the supplied weight
 		FTraitUpdateState WithWeight(float Weight) const
 		{
@@ -77,10 +80,18 @@ namespace UE::AnimNext
 		}
 
 		// Creates a new instance of the update state marking it as blending out
-		FTraitUpdateState AsBlendingOut() const
+		FTraitUpdateState AsBlendingOut(bool InIsBlendingOut = true) const
+		{
+			FTraitUpdateState Result(*this);
+			Result.bIsBlendingOut = InIsBlendingOut;
+			return Result;
+		}
+
+		// Creates a new instance of the update state marking it as newly relevant
+		FTraitUpdateState AsNewlyRelevant(bool InIsNewlyRelevant = true) const
 		{
 			FTraitUpdateState Result(*this);
-			Result.bIsBlendingOut = true;
+			Result.bIsNewlyRelevant = InIsNewlyRelevant;
 			return Result;
 		}
 
@@ -95,7 +106,10 @@ namespace UE::AnimNext
 		float TotalTrajectoryWeight = 1.0f;
 
 		// Whether or not we are blending out
-		int32 bIsBlendingOut = false;
+		int8 bIsBlendingOut = false;
+
+		// Whether or not we newly became relevant
+		int8 bIsNewlyRelevant = false;
 	};
 
 	static_assert(sizeof(FTraitUpdateState) <= 16, "Keep the size to 16 bytes for efficient copying");
@@ -217,6 +231,9 @@ namespace UE::AnimNext
 	{
 		DECLARE_ANIM_TRAIT_INTERFACE(IUpdate, 0x59d24dc5)
 
+		// Called before the first update when a trait stack becomes relevant
+		virtual void OnBecomeRelevant(FUpdateTraversalContext& Context, const TTraitBinding<IUpdate>& Binding, const FTraitUpdateState& TraitState) const;
+
 		// Called before a traits children are updated
 		virtual void PreUpdate(FUpdateTraversalContext& Context, const TTraitBinding<IUpdate>& Binding, const FTraitUpdateState& TraitState) const;
 
@@ -250,6 +267,12 @@ namespace UE::AnimNext
 	template<>
 	struct TTraitBinding<IUpdate> : FTraitBinding
 	{
+		// @see IUpdate::OnBecomeRelevant
+		void OnBecomeRelevant(FUpdateTraversalContext& Context, const FTraitUpdateState& TraitState) const
+		{
+			GetInterface()->OnBecomeRelevant(Context, *this, TraitState);
+		}
+
 		// @see IUpdate::PreUpdate
 		void PreUpdate(FUpdateTraversalContext& Context, const FTraitUpdateState& TraitState) const
 		{
