- Export AnimNextGraph outliner items with Entry name not graph name (clashing identifiers and incorrect display name)
- Move adding missing outliner-assetuserdata to ::Initialize rather than ::Postload so it gets added to new assets correctly

diff --git a/Engine/Plugins/Experimental/AnimNext/Source/AnimNextUncookedOnly/Private/UncookedOnlyUtils.cpp b/Engine/Plugins/Experimental/AnimNext/Source/AnimNextUncookedOnly/Private/UncookedOnlyUtils.cpp
--- a/Engine/Plugins/Experimental/AnimNext/Source/AnimNextUncookedOnly/Private/UncookedOnlyUtils.cpp
+++ b/Engine/Plugins/Experimental/AnimNext/Source/AnimNextUncookedOnly/Private/UncookedOnlyUtils.cpp
@@ -1633,7 +1633,7 @@ void FUtils::GetAssetOutlinerItems(const UAnimNextRigVMAssetEditorData* EditorDa
 		if(const IAnimNextRigVMGraphInterface* GraphInterface = Cast<IAnimNextRigVMGraphInterface>(Entry))
 		{
 			FWorkspaceOutlinerItemExport& Export = OutExports.Exports.AddDefaulted_GetRef();
-			Export.Identifier = GraphInterface->GetEdGraph()->GetFName();
+			Export.Identifier = Entry->GetEntryName();
 			Export.ParentIdentifier = EditorData->GetOuter()->GetFName();					
 			Export.AssetPath = EditorData->GetOuter();
 
