Add nullptr check for ToolMenus singleton on shutdown of AnimNextEditor module

diff --git a/Engine/Plugins/Experimental/AnimNext/Source/AnimNextEditor/Private/Graph/AnimNextGraph_OutlinerItemDetails.cpp b/Engine/Plugins/Experimental/AnimNext/Source/AnimNextEditor/Private/Graph/AnimNextGraph_OutlinerItemDetails.cpp
--- a/Engine/Plugins/Experimental/AnimNext/Source/AnimNextEditor/Private/Graph/AnimNextGraph_OutlinerItemDetails.cpp
+++ b/Engine/Plugins/Experimental/AnimNext/Source/AnimNextEditor/Private/Graph/AnimNextGraph_OutlinerItemDetails.cpp
@@ -118,7 +118,10 @@ void FAnimNextGraphItemDetails::RegisterToolMenuExtensions()
 
 void FAnimNextGraphItemDetails::UnregisterToolMenuExtensions()
 {
-	UToolMenus::Get()->UnregisterOwnerByName("FAnimNextGraphItemDetails");
+	if(UToolMenus* ToolMenus = UToolMenus::Get())
+	{
+		ToolMenus->UnregisterOwnerByName("FAnimNextGraphItemDetails");
+	}
 }
 }
 
