Slate: Deprecate SListView::ItemHeight and STreeViewItemHeight. ItemHeight and ItemWidth are only used when the panel is is as a tile or both value are assigned. With TreeView and ListView, you can't assigned both. The function are protected, a subclass of those could but it's up to them to expose the ItemHeight and ItemWidth. This behavior changed was introduced with CL 7634094.

diff --git a/Engine/Plugins/Experimental/AnimNext/Source/AnimNextEditor/Private/Common/SRigVMAssetView.cpp b/Engine/Plugins/Experimental/AnimNext/Source/AnimNextEditor/Private/Common/SRigVMAssetView.cpp
--- a/Engine/Plugins/Experimental/AnimNext/Source/AnimNextEditor/Private/Common/SRigVMAssetView.cpp
+++ b/Engine/Plugins/Experimental/AnimNext/Source/AnimNextEditor/Private/Common/SRigVMAssetView.cpp
@@ -262,7 +262,6 @@ void SRigVMAssetView::Construct(const FArguments& InArgs, UAnimNextRigVMAssetEdi
 			.OnGetChildren(this, &SRigVMAssetView::HandleGetChildren)
 			.OnItemScrolledIntoView(this, &SRigVMAssetView::HandleItemScrolledIntoView)
 			.OnSelectionChanged(this, &SRigVMAssetView::HandleSelectionChanged)
-			.ItemHeight(20.0f)
 			.HeaderRow(
 				SNew(SHeaderRow)
 				+SHeaderRow::Column(Column_RevisionControl)
diff --git a/Engine/Plugins/Experimental/AnimNext/Source/AnimNextEditor/Private/Param/SAddParametersDialog.cpp b/Engine/Plugins/Experimental/AnimNext/Source/AnimNextEditor/Private/Param/SAddParametersDialog.cpp
--- a/Engine/Plugins/Experimental/AnimNext/Source/AnimNextEditor/Private/Param/SAddParametersDialog.cpp
+++ b/Engine/Plugins/Experimental/AnimNext/Source/AnimNextEditor/Private/Param/SAddParametersDialog.cpp
@@ -91,7 +91,6 @@ void SAddParametersDialog::Construct(const FArguments& InArgs, const FAssetData&
 					SAssignNew(EntriesList, SListView<TSharedRef<FParameterToAddEntry>>)
 					.ListItemsSource(&Entries)
 					.OnGenerateRow(this, &SAddParametersDialog::HandleGenerateRow)
-					.ItemHeight(20.0f)
 					.HeaderRow(
 						SNew(SHeaderRow)
 						+SHeaderRow::Column(Column_Name)
