Fix ASAN issue with AnimNextEditor shutdown
#fix ensure that UObject system is still initialized

diff --git a/Engine/Plugins/Experimental/AnimNext/Source/AnimNextEditor/Private/AnimNextEditorModule.cpp b/Engine/Plugins/Experimental/AnimNext/Source/AnimNextEditor/Private/AnimNextEditorModule.cpp
--- a/Engine/Plugins/Experimental/AnimNext/Source/AnimNextEditor/Private/AnimNextEditorModule.cpp
+++ b/Engine/Plugins/Experimental/AnimNext/Source/AnimNextEditor/Private/AnimNextEditorModule.cpp
@@ -224,9 +224,12 @@ class FModule : public IModule
 			UolEditorModule.UnregisterEditorContext("AnimNextContext");
 		}
 
-		Workspace::IWorkspaceEditorModule& WorkspaceModule = FModuleManager::Get().LoadModuleChecked<Workspace::IWorkspaceEditorModule>("WorkspaceEditor");
-		WorkspaceModule.UnregisterWorkspaceItemDetails(Workspace::FOutlinerItemDetailsId(FAnimNextGraphOutlinerData::StaticStruct()->GetFName()));
-		FAnimNextGraphItemDetails::UnregisterToolMenuExtensions();
+		if (UObjectInitialized())
+		{
+			Workspace::IWorkspaceEditorModule& WorkspaceModule = FModuleManager::Get().LoadModuleChecked<Workspace::IWorkspaceEditorModule>("WorkspaceEditor");
+			WorkspaceModule.UnregisterWorkspaceItemDetails(Workspace::FOutlinerItemDetailsId(FAnimNextGraphOutlinerData::StaticStruct()->GetFName()));
+			FAnimNextGraphItemDetails::UnregisterToolMenuExtensions();
+		}
 	}
 
 	virtual TSharedRef<SWidget> CreateParameterPicker(const FParameterPickerArgs& InArgs) override
