AnimNext: Misc changes for scene graph

Added public/private concept to export interface. Only public exports actually get exposed to the asset registry.
State that we instantiate is now sorted into public/private members. Prior to updating internal state, we now pull any public (i.e. externally-controlled) parameters into the current scope via the parameter stack.
Animation Graphs now instantiate their state when they are created, as well as pushing it to the stack when they run. This means that they no longer need an auxilliary parameter scope in the schedule to function.
Added explicit parameters for LOD and reference pose to graph tasks (cant do this for ports without a schedule refactor, so they are still hard-coded).

We now allow multiple parmeter sources to be pushed at the root.
Expanded hoisted function lib support for parameters to include all BP function library classes with the appropriate signature.
Moved AnimNext module functionality to private/internal implementations for parameter sources. This is to allow for easier optimization later on.
Made locator context extensible, so other modules can opt-in to adding locator types to AnimNext contexts.
Removed UAnimNextMeshComponent and switched to a helper function lib (supporting engine CL also in review)

Fixed some properties not showing up in property locator fragment pickers.
Fixed SetAnimNextParameterNode in the controller failing to apply parameters when auto-relinking fails
Fixed selection getting lost when edits are performed in the details panel (now that those edits cause a modified delegate to be fired and a recompilation).

diff --git a/Engine/Plugins/Experimental/AnimNext/Config/BaseAnimNext.ini b/Engine/Plugins/Experimental/AnimNext/Config/BaseAnimNext.ini
--- a/Engine/Plugins/Experimental/AnimNext/Config/BaseAnimNext.ini
+++ b/Engine/Plugins/Experimental/AnimNext/Config/BaseAnimNext.ini
@@ -1,8 +1,3 @@
-[/Script/AnimNext.AnimNextConfig]
-+ExposedClasses=(Class="/Script/Engine.World",AccessorName="UE_World")
-+ExposedClasses=(Class="/Script/Engine.CharacterMovementComponent",AccessorName="UE_CharacterMovementComponent")
-+ExposedClasses=(Class="/Script/AnimNext.AnimNextMeshComponent",AccessorName="UE_AnimNextMeshComponent")
-
 [CoreRedirects]
 ; Traits rename
 +StructRedirects=(OldName="/Script/AnimNext.AnimNextDecoratorHandle", NewName="/Script/AnimNext.AnimNextTraitHandle")
@@ -23,4 +18,7 @@
 +PropertyRedirects=(OldName="/Script/AnimNext.AnimNextScheduleParamScopeEntryTask.ParameterBlocks", NewName="Parameters")
 
 ; Workspace plugin split
-+ClassRedirects=(OldName="/Script/AnimNextEditor.AnimNextWorkspace", NewName="/Script/WorkspaceEditor.Workspace")
\ No newline at end of file
++ClassRedirects=(OldName="/Script/AnimNextEditor.AnimNextWorkspace", NewName="/Script/WorkspaceEditor.Workspace")
+
+; Scene graph prototype compatibility
++ClassRedirects=(OldName="/Script/AnimNext.AnimNextMeshComponent", NewName="/Script/Engine.SkeletalMeshComponent")
\ No newline at end of file
diff --git a/Engine/Plugins/Experimental/AnimNext/Source/AnimNext/Internal/AnimNextModule.h b/Engine/Plugins/Experimental/AnimNext/Source/AnimNext/Internal/AnimNextModule.h
new file mode 100644
--- /dev/null
+++ b/Engine/Plugins/Experimental/AnimNext/Source/AnimNext/Internal/AnimNextModule.h
@@ -0,0 +1,69 @@
+﻿// Copyright Epic Games, Inc. All Rights Reserved.
+
+#pragma once
+
+#include "CoreMinimal.h"
+#include "IAnimNextModule.h"
+
+struct FAnimNextGraphInstancePtr;
+struct FAnimNextParamInstanceIdentifier;
+
+namespace UE::AnimNext
+{
+	class IParameterSource;
+	class IParameterSourceFactory;
+	struct FParameterSourceContext;
+}
+
+// Enable console commands only in development builds when logging is enabled
+#define WITH_ANIMNEXT_CONSOLE_COMMANDS (!UE_BUILD_SHIPPING && !NO_LOGGING)
+
+namespace UE::AnimNext
+{
+
+class FModule : public IAnimNextModule
+{
+public:
+	// IModuleInterface interface
+	virtual void StartupModule() override;
+	virtual void ShutdownModule() override;
+
+	// IAnimNextModule interface
+	virtual void RegisterAnimNextAnimGraph(const IAnimNextAnimGraph& InAnimGraphImpl) override;
+	virtual void UnregisterAnimNextAnimGraph() override;
+	virtual void UpdateGraph(FAnimNextGraphInstancePtr& GraphInstance, float DeltaTime) override;
+	virtual void EvaluateGraph(FAnimNextGraphInstancePtr& GraphInstance, const FReferencePose& RefPose, int32 GraphLODLevel, FLODPoseHeap& OutputPose) const override;
+
+	// Factory method used to create a parameter source to access the specified instance ID, with a set of parameters that are initially required
+	// @param   InContext             Context used to set up the parameter source
+	// @param   InInstanceId          The instance identifier associated with the parameters that are required
+	// @param   InRequiredParameters  Any required parameters that the source should initially supply, can be empty
+	// @return a new parameter source, or nullptr if the instance ID could not be handled
+	ANIMNEXT_API TUniquePtr<IParameterSource> CreateParameterSource(const FParameterSourceContext& InContext, const TInstancedStruct<FAnimNextParamInstanceIdentifier>& InInstanceId, TConstArrayView<FName> InRequiredParameters) const;
+
+	// Register a factory that can be used to generate parameter sources
+	// @param   InName                Identifier for the factory
+	// @param   InFactory             The factory to register
+	ANIMNEXT_API void RegisterParameterSourceFactory(FName InName, TSharedRef<IParameterSourceFactory> InFactory);
+
+	// Unregister a factory previously passed to RegisterParameterSourceFactory
+	// @param   InName                Identifier for the factory
+	ANIMNEXT_API void UnregisterParameterSourceFactory(FName InName);
+
+	// Find a factory previously passed to RegisterParameterSourceFactory
+	// @param   InName                Identifier for the factory
+	// @return a shared ptr to the factory, if found, otherwise an invalid shared ptr
+	ANIMNEXT_API TSharedPtr<IParameterSourceFactory> FindParameterSourceFactory(FName InName);
+
+#if WITH_ANIMNEXT_CONSOLE_COMMANDS
+	void ListNodeTemplates(const TArray<FString>& Args);
+	void ListAnimGraphs(const TArray<FString>& Args);
+
+	TArray<IConsoleObject*> ConsoleCommands;
+#endif
+
+	// All known factories
+	TMap<FName, TSharedRef<IParameterSourceFactory>> ParameterSourceFactories;
+};
+
+}
diff --git a/Engine/Plugins/Experimental/AnimNext/Source/AnimNext/Internal/Component/SkinnedMeshComponentExtensions.h b/Engine/Plugins/Experimental/AnimNext/Source/AnimNext/Internal/Component/SkinnedMeshComponentExtensions.h
new file mode 100644
--- /dev/null
+++ b/Engine/Plugins/Experimental/AnimNext/Source/AnimNext/Internal/Component/SkinnedMeshComponentExtensions.h
@@ -0,0 +1,21 @@
+﻿// Copyright Epic Games, Inc. All Rights Reserved.
+
+#pragma once
+
+#include "Containers/ContainersFwd.h"
+#include "Math/MathFwd.h"
+#include "BoneIndices.h"
+
+class USkinnedMeshComponent;
+
+namespace UE::Anim
+{
+
+// Extensions to skinned meshes that are not public API yet
+struct FSkinnedMeshComponentExtensions
+{
+	// Converts local to component space, flips buffers, updates bounds and dispatches to renderer
+	static ANIMNEXT_API void CompleteAndDispatch(USkinnedMeshComponent* InComponent, TConstArrayView<FBoneIndexType> InParentIndices, TConstArrayView<FBoneIndexType> InRequiredBoneIndices, TConstArrayView<FTransform> InLocalSpaceTransforms);
+};
+
+}
diff --git a/Engine/Plugins/Experimental/AnimNext/Source/AnimNext/Internal/Graph/AnimNextGraph.h b/Engine/Plugins/Experimental/AnimNext/Source/AnimNext/Internal/Graph/AnimNextGraph.h
--- a/Engine/Plugins/Experimental/AnimNext/Source/AnimNext/Internal/Graph/AnimNextGraph.h
+++ b/Engine/Plugins/Experimental/AnimNext/Source/AnimNext/Internal/Graph/AnimNextGraph.h
@@ -67,6 +67,7 @@ public:
 
 	// UObject interface
 	virtual void Serialize(FArchive& Ar) override;
+	virtual void PostLoad() override;
 #if WITH_EDITOR
 	virtual void PostEditChangeProperty(FPropertyChangedEvent& PropertyChangedEvent) override;
 #endif
@@ -148,6 +149,9 @@ protected:
 	// This is a resolved handle to the root trait in our graph, for each entry point 
 	TMap<FName, FAnimNextTraitHandle> ResolvedRootTraitHandles;
 
+	// This is an index into EntryPoints, for each entry point
+	TMap<FName, int32> ResolvedEntryPoints;
+
 	// This is the graph shared data used by the trait system, the output of FTraitReader
 	// We de-serialize manually into this buffer from the archive buffer, this is never saved on disk
 	TArray<uint8> SharedDataBuffer;
@@ -175,11 +179,14 @@ protected:
 	UPROPERTY()
 	TArray<FAnimNextParam> RequiredParameters;
 
-	// Property bag that holds state for this graph 
+	// Default state for this graph
 	UPROPERTY()
-	FInstancedPropertyBag PropertyBag;
+	FAnimNextGraphState DefaultState;
 
 #if WITH_EDITORONLY_DATA
+	UPROPERTY()
+	FInstancedPropertyBag PropertyBag_DEPRECATED;
+
 	// This buffer holds the output of the FTraitWriter post compilation
 	// We serialize it manually and it is discarded at runtime
 	TArray<uint8> SharedDataArchiveBuffer;
diff --git a/Engine/Plugins/Experimental/AnimNext/Source/AnimNext/Internal/Graph/AnimNextGraphEntryPoint.h b/Engine/Plugins/Experimental/AnimNext/Source/AnimNext/Internal/Graph/AnimNextGraphEntryPoint.h
--- a/Engine/Plugins/Experimental/AnimNext/Source/AnimNext/Internal/Graph/AnimNextGraphEntryPoint.h
+++ b/Engine/Plugins/Experimental/AnimNext/Source/AnimNext/Internal/Graph/AnimNextGraphEntryPoint.h
@@ -3,9 +3,8 @@
 #pragma once
 
 #include "CoreMinimal.h"
-
+#include "AnimNextGraphState.h"
 #include "TraitCore/EntryPointHandle.h"
-
 #include "AnimNextGraphEntryPoint.generated.h"
 
 USTRUCT()
@@ -20,5 +19,9 @@ struct FAnimNextGraphEntryPoint
 	// This is a handle to the root trait for a graph
 	UPROPERTY()
 	FAnimNextEntryPointHandle RootTraitHandle;
+
+	// Default state for this entry point
+	UPROPERTY()
+	FAnimNextGraphState DefaultState;
 };
 
diff --git a/Engine/Plugins/Experimental/AnimNext/Source/AnimNext/Internal/Graph/AnimNextGraphState.h b/Engine/Plugins/Experimental/AnimNext/Source/AnimNext/Internal/Graph/AnimNextGraphState.h
new file mode 100644
--- /dev/null
+++ b/Engine/Plugins/Experimental/AnimNext/Source/AnimNext/Internal/Graph/AnimNextGraphState.h
@@ -0,0 +1,30 @@
+﻿// Copyright Epic Games, Inc. All Rights Reserved.
+
+#pragma once
+
+#include "CoreMinimal.h"
+#include "PropertyBag.h"
+#include "AnimNextGraphState.generated.h"
+
+// Default state for a graph entry point
+USTRUCT()
+struct FAnimNextGraphState
+{
+	GENERATED_BODY()
+
+#if WITH_EDITORONLY_DATA
+	void Reset()
+	{
+		State.Reset();
+		PublicParameterStartIndex = INDEX_NONE;
+	}
+#endif
+
+	// All state, both public and private
+	UPROPERTY()
+	FInstancedPropertyBag State;
+
+	// Index of the first public parameter, of INDEX_NONE if no public parameter is present
+	UPROPERTY()
+	int32 PublicParameterStartIndex = INDEX_NONE;
+};
\ No newline at end of file
diff --git a/Engine/Plugins/Experimental/AnimNext/Source/AnimNext/Internal/Param/AnimNextActorLocatorFragment.h b/Engine/Plugins/Experimental/AnimNext/Source/AnimNext/Internal/Param/AnimNextActorLocatorFragment.h
new file mode 100644
--- /dev/null
+++ b/Engine/Plugins/Experimental/AnimNext/Source/AnimNext/Internal/Param/AnimNextActorLocatorFragment.h
@@ -0,0 +1,43 @@
+﻿// Copyright Epic Games, Inc. All Rights Reserved.
+
+#pragma once
+
+#include "CoreTypes.h"
+#include "UniversalObjectLocatorFwd.h"
+#include "AnimNextActorLocatorFragment.generated.h"
+
+// Exposes the 'current' actor as a UOL fragment
+USTRUCT()
+struct FAnimNextActorLocatorFragment
+{
+	GENERATED_BODY()
+
+	FAnimNextActorLocatorFragment() = default;
+
+	ANIMNEXT_API static UE::UniversalObjectLocator::TFragmentTypeHandle<FAnimNextActorLocatorFragment> FragmentType;
+
+	friend uint32 GetTypeHash(const FAnimNextActorLocatorFragment& A)
+	{
+		return 0;
+	}
+
+	friend bool operator==(const FAnimNextActorLocatorFragment& A, const FAnimNextActorLocatorFragment& B)
+	{
+		return true;
+	}
+
+	UE::UniversalObjectLocator::FResolveResult Resolve(const UE::UniversalObjectLocator::FResolveParams& Params) const;
+	UE::UniversalObjectLocator::FInitializeResult Initialize(const UE::UniversalObjectLocator::FInitializeParams& InParams);
+	ANIMNEXT_API void ToString(FStringBuilderBase& OutStringBuilder) const;
+	UE::UniversalObjectLocator::FParseStringResult TryParseString(FStringView InString, const UE::UniversalObjectLocator::FParseStringParams& Params);
+	static uint32 ComputePriority(const UObject* Object, const UObject* Context);
+};
+
+template<>
+struct TStructOpsTypeTraits<FAnimNextActorLocatorFragment> : public TStructOpsTypeTraitsBase2<FAnimNextActorLocatorFragment>
+{
+	enum
+	{
+		WithIdenticalViaEquality = true,
+	};
+};
\ No newline at end of file
diff --git a/Engine/Plugins/Experimental/AnimNext/Source/AnimNext/Internal/Param/AnimNextComponentLocatorFragment.h b/Engine/Plugins/Experimental/AnimNext/Source/AnimNext/Internal/Param/AnimNextComponentLocatorFragment.h
new file mode 100644
--- /dev/null
+++ b/Engine/Plugins/Experimental/AnimNext/Source/AnimNext/Internal/Param/AnimNextComponentLocatorFragment.h
@@ -0,0 +1,50 @@
+﻿// Copyright Epic Games, Inc. All Rights Reserved.
+
+#pragma once
+
+#include "CoreTypes.h"
+#include "UniversalObjectLocatorFwd.h"
+#include "AnimNextComponentLocatorFragment.generated.h"
+
+class UActorComponent;
+
+USTRUCT()
+struct FAnimNextComponentLocatorFragment
+{
+	GENERATED_BODY()
+
+	FAnimNextComponentLocatorFragment() = default;
+
+	ANIMNEXT_API explicit FAnimNextComponentLocatorFragment(TSubclassOf<UActorComponent> InClass);
+
+	// Path to the component class
+	UPROPERTY()
+	FSoftObjectPath Path;
+	
+	ANIMNEXT_API static UE::UniversalObjectLocator::TFragmentTypeHandle<FAnimNextComponentLocatorFragment> FragmentType;
+
+	friend uint32 GetTypeHash(const FAnimNextComponentLocatorFragment& A)
+	{
+		return GetTypeHash(A.Path);
+	}
+
+	friend bool operator==(const FAnimNextComponentLocatorFragment& A, const FAnimNextComponentLocatorFragment& B)
+	{
+		return A.Path == B.Path;
+	}
+
+	UE::UniversalObjectLocator::FResolveResult Resolve(const UE::UniversalObjectLocator::FResolveParams& Params) const;
+	UE::UniversalObjectLocator::FInitializeResult Initialize(const UE::UniversalObjectLocator::FInitializeParams& InParams);
+	ANIMNEXT_API void ToString(FStringBuilderBase& OutStringBuilder) const;
+	UE::UniversalObjectLocator::FParseStringResult TryParseString(FStringView InString, const UE::UniversalObjectLocator::FParseStringParams& Params);
+	static uint32 ComputePriority(const UObject* Object, const UObject* Context);
+};
+
+template<>
+struct TStructOpsTypeTraits<FAnimNextComponentLocatorFragment> : public TStructOpsTypeTraitsBase2<FAnimNextComponentLocatorFragment>
+{
+	enum
+	{
+		WithIdenticalViaEquality = true,
+	};
+};
\ No newline at end of file
diff --git a/Engine/Plugins/Experimental/AnimNext/Source/AnimNext/Internal/Param/IParameterSource.h b/Engine/Plugins/Experimental/AnimNext/Source/AnimNext/Internal/Param/IParameterSource.h
new file mode 100644
--- /dev/null
+++ b/Engine/Plugins/Experimental/AnimNext/Source/AnimNext/Internal/Param/IParameterSource.h
@@ -0,0 +1,57 @@
+// Copyright Epic Games, Inc. All Rights Reserved.
+
+#pragma once
+
+#include "CoreMinimal.h"
+#include "Param/ParamStackLayerHandle.h"
+
+struct FAnimNextScheduleExternalParamTask;
+struct FAnimNextScheduleParamScopeEntryTask;
+struct FAnimNextScheduleGraphTask;
+struct FAnimNextGraphInstancePtr;
+struct FAnimNextGraphInstance;
+
+namespace UE::AnimNext
+{
+	struct FParamContext;
+	struct FScheduleInstanceData;
+	struct FScheduleBeginTickFunction;
+	struct FSubGraphHostTrait;
+}
+
+namespace UE::AnimNext
+{
+
+/** An interface used to abstract AnimNext parameter sources */
+class IParameterSource
+{
+public:
+	virtual ~IParameterSource() = default;
+
+private:
+	/**
+	 * Get the instance ID of this parameter source. Instance IDs are used to differentiate sources from one another. They do not have to be unique,
+	 * but when applying sources to scopes, the instance ID is used to determine if the source should be replaced
+	 */
+	virtual FName GetInstanceId() const = 0;
+	
+	/** Update this parameter source */
+	virtual void Update(float DeltaTime) = 0;
+
+	/** Get a layer handle to allow this source to be used in the parameter stack */
+	virtual const FParamStackLayerHandle& GetLayerHandle() const = 0; 
+
+	/** Allow GC subscription */
+	virtual void AddReferencedObjects(FReferenceCollector& Collector) {}
+
+	friend struct FScheduleInstanceData;
+	friend struct FScheduleBeginTickFunction;
+	friend struct FSubGraphHostTrait;
+	friend struct ::FAnimNextScheduleExternalParamTask;
+	friend struct ::FAnimNextScheduleParamScopeEntryTask;
+	friend struct ::FAnimNextScheduleGraphTask;
+	friend struct ::FAnimNextGraphInstancePtr;
+	friend struct ::FAnimNextGraphInstance;
+};
+
+}
\ No newline at end of file
diff --git a/Engine/Plugins/Experimental/AnimNext/Source/AnimNext/Internal/Param/IParameterSourceFactory.h b/Engine/Plugins/Experimental/AnimNext/Source/AnimNext/Internal/Param/IParameterSourceFactory.h
new file mode 100644
--- /dev/null
+++ b/Engine/Plugins/Experimental/AnimNext/Source/AnimNext/Internal/Param/IParameterSourceFactory.h
@@ -0,0 +1,46 @@
+﻿// Copyright Epic Games, Inc. All Rights Reserved.
+
+#pragma once
+
+#include "Param/IParameterSource.h"
+
+struct FAnimNextParamInstanceIdentifier;
+
+namespace UE::AnimNext
+{
+	class IParameterSource;
+}
+
+namespace UE::AnimNext
+{
+
+// Context passed to object accessor functions registered to RegisterObjectAccessor
+struct FParameterSourceContext
+{
+	FParameterSourceContext() = default;
+
+	FParameterSourceContext(const UObject* InObject)
+		: Object(InObject)
+	{
+		check(Object != nullptr);
+	}
+
+	// The object that the AnimNext entry is bound to (e.g. a UAnimNextComponent)
+	const UObject* Object = nullptr;
+};
+
+// Interface allowing other modules to extend and add to external parameter system
+class IParameterSourceFactory
+{
+public:
+	virtual ~IParameterSourceFactory() = default;
+
+	// Factory method used to create a parameter source of the specified name, with a set of parameters that are initially required
+	// @param  InContext              Context used to resolve the scope
+	// @param  InInstanceId           The instance identifier associated with the parameters that are required
+	// @param  InRequiredParameters   Any required parameters that the source should initially supply, can be empty
+	// @return a new parameter source, or nullptr if the source was not found
+	virtual TUniquePtr<IParameterSource> CreateParameterSource(const FParameterSourceContext& InContext, const TInstancedStruct<FAnimNextParamInstanceIdentifier>& InInstanceId, TConstArrayView<FName> InRequiredParameters) const = 0;
+};
+
+}
\ No newline at end of file
diff --git a/Engine/Plugins/Experimental/AnimNext/Source/AnimNext/Internal/Param/ObjectProxyFactory.h b/Engine/Plugins/Experimental/AnimNext/Source/AnimNext/Internal/Param/ObjectProxyFactory.h
--- a/Engine/Plugins/Experimental/AnimNext/Source/AnimNext/Internal/Param/ObjectProxyFactory.h
+++ b/Engine/Plugins/Experimental/AnimNext/Source/AnimNext/Internal/Param/ObjectProxyFactory.h
@@ -32,6 +32,7 @@ class FObjectProxyFactory : public IParameterSourceFactory
 	friend class ::UAnimNextConfig;
 	friend class FModule;
 	friend class UE::AnimNext::UncookedOnly::FObjectProxyType;
+	friend struct FScheduleInitializationContext;
 
 public:
 	FObjectProxyFactory();
diff --git a/Engine/Plugins/Experimental/AnimNext/Source/AnimNext/Internal/Param/ParamUtils.h b/Engine/Plugins/Experimental/AnimNext/Source/AnimNext/Internal/Param/ParamUtils.h
--- a/Engine/Plugins/Experimental/AnimNext/Source/AnimNext/Internal/Param/ParamUtils.h
+++ b/Engine/Plugins/Experimental/AnimNext/Source/AnimNext/Internal/Param/ParamUtils.h
@@ -26,7 +26,9 @@ struct FParamUtils
 	static ANIMNEXT_API FParamCompatibility GetCompatibility(const FAnimNextParamType& InLHS, const FAnimNextParamType& InRHS);
 
 	// Check whether the supplied function can be used to access/map parameters
-	static ANIMNEXT_API bool CanUseFunction(const UFunction* InFunction);
+	// @param	InFunction	The function to check
+	// @param	InClass		The expected class to check - used to determine if hoisted functions are valid in this context.
+	static ANIMNEXT_API bool CanUseFunction(const UFunction* InFunction, const UClass* InExpectedClass);
 
 	// Check whether the supplied property can be used to access/map parameters
 	static ANIMNEXT_API bool CanUseProperty(const FProperty* InProperty);
diff --git a/Engine/Plugins/Experimental/AnimNext/Source/AnimNext/Internal/Scheduler/AnimNextSchedule.h b/Engine/Plugins/Experimental/AnimNext/Source/AnimNext/Internal/Scheduler/AnimNextSchedule.h
--- a/Engine/Plugins/Experimental/AnimNext/Source/AnimNext/Internal/Scheduler/AnimNextSchedule.h
+++ b/Engine/Plugins/Experimental/AnimNext/Source/AnimNext/Internal/Scheduler/AnimNextSchedule.h
@@ -89,6 +89,14 @@ private:
 	UPROPERTY(EditAnywhere, Category = "Graph", meta = (AllowedParamType = "FAnimNextEntryPoint"))
 	FAnimNextEditorParam EntryPoint;
 
+	// Reference pose to use
+	UPROPERTY(EditAnywhere, Category = "Graph", meta = (AllowedParamType = "FAnimNextGraphReferencePose"))
+	FAnimNextEditorParam ReferencePose;
+
+	// Pose LOD to generate to use
+	UPROPERTY(EditAnywhere, Category = "Graph", meta = (AllowedParamType = "int32"))
+	FAnimNextEditorParam LOD;
+
 	// All parameters that are required by this graph to run (only required if dynamic as static graph params can be discovered by the compiler)
 	UPROPERTY(EditAnywhere, Category = "Graph")
 	TArray<FAnimNextEditorParam> RequiredParameters;
diff --git a/Engine/Plugins/Experimental/AnimNext/Source/AnimNext/Internal/Scheduler/AnimNextScheduleGraphTask.h b/Engine/Plugins/Experimental/AnimNext/Source/AnimNext/Internal/Scheduler/AnimNextScheduleGraphTask.h
--- a/Engine/Plugins/Experimental/AnimNext/Source/AnimNext/Internal/Scheduler/AnimNextScheduleGraphTask.h
+++ b/Engine/Plugins/Experimental/AnimNext/Source/AnimNext/Internal/Scheduler/AnimNextScheduleGraphTask.h
@@ -26,13 +26,6 @@ namespace UE::AnimNext::UncookedOnly
 	struct FUtils;
 }
 
-namespace UE::AnimNext
-{
-	extern TInstancedStruct<FAnimNextParamUniversalObjectLocator> GetCharacterInstanceId();
-	extern FName GetCharacterInstanceIdName();
-	extern FParamId GetMeshComponentParamId();
-}
-
 USTRUCT()
 struct FAnimNextScheduleGraphTask
 {
@@ -53,10 +46,6 @@ private:
 	// Verify graph's required parameters are satisfied by this task's supplied parameters
 	void VerifyRequiredParameters(UAnimNextGraph* InGraphToRun) const;
 
-#if WITH_EDITORONLY_DATA 
-	// Get required parameters needed for internal operation
-	ANIMNEXT_API static TArray<FAnimNextEditorParam> GetRequiredParametersInternal();
-#endif
 private:
 	UPROPERTY()
 	uint32 TaskIndex = MAX_uint32;
@@ -76,6 +65,12 @@ private:
 	UPROPERTY()
 	FAnimNextParam DynamicGraph;
 
+	UPROPERTY()
+	FAnimNextParam ReferencePose;
+
+	UPROPERTY()
+	FAnimNextParam LOD;
+
 	// Index of each term in the schedule intermediates
 	UPROPERTY()
 	TArray<uint32> Terms;
diff --git a/Engine/Plugins/Experimental/AnimNext/Source/AnimNext/Internal/Scheduler/ScheduleInitializationContext.h b/Engine/Plugins/Experimental/AnimNext/Source/AnimNext/Internal/Scheduler/ScheduleInitializationContext.h
--- a/Engine/Plugins/Experimental/AnimNext/Source/AnimNext/Internal/Scheduler/ScheduleInitializationContext.h
+++ b/Engine/Plugins/Experimental/AnimNext/Source/AnimNext/Internal/Scheduler/ScheduleInitializationContext.h
@@ -5,11 +5,14 @@
 #include "CoreMinimal.h"
 
 struct FAnimNextSchedulerEntry;
+struct FPropertyBagPropertyDesc;
+template<typename StructType> struct TInstancedStruct;
+struct FAnimNextParamInstanceIdentifier;
 
 namespace UE::AnimNext
-{	
-	class IParameterSource;
-	struct FScheduleContext;
+{
+struct FParameterSourceContext;
+struct FScheduleContext;
 }
 
 namespace UE::AnimNext
@@ -29,7 +32,9 @@ struct FScheduleInitializationContext
 {
 public:
 	// Apply the supplied parameter source to the specified scope, evicting any source that was there previously
-	ANIMNEXT_API void ApplyParametersToScope(FName InScope, EParameterScopeOrdering InOrdering, TUniquePtr<IParameterSource>&& InParameters) const;
+	ANIMNEXT_API void ApplyParametersToScope(FName InScope, EParameterScopeOrdering InOrdering, const UObject* InObject, TConstArrayView<TTuple<FName, FName>> InRequiredParameterAliases) const;
+	ANIMNEXT_API void ApplyParametersToScope(FName InScope, EParameterScopeOrdering InOrdering, const TInstancedStruct<FAnimNextParamInstanceIdentifier>& InInstanceId, const FParameterSourceContext& InContext, TConstArrayView<FName> InRequiredParameters) const;
+	ANIMNEXT_API void ApplyParametersToScope(FName InScope, EParameterScopeOrdering InOrdering, FName InId, TConstArrayView<FPropertyBagPropertyDesc> InPropertyDescs, TConstArrayView<TConstArrayView<uint8>> InValues) const;
 
 private:
 	FScheduleInitializationContext(const FScheduleContext& InContext);
diff --git a/Engine/Plugins/Experimental/AnimNext/Source/AnimNext/Internal/Scheduler/ScheduleTaskContext.h b/Engine/Plugins/Experimental/AnimNext/Source/AnimNext/Internal/Scheduler/ScheduleTaskContext.h
--- a/Engine/Plugins/Experimental/AnimNext/Source/AnimNext/Internal/Scheduler/ScheduleTaskContext.h
+++ b/Engine/Plugins/Experimental/AnimNext/Source/AnimNext/Internal/Scheduler/ScheduleTaskContext.h
@@ -4,9 +4,10 @@
 
 #include "CoreMinimal.h"
 
+struct FInstancedPropertyBag;
+
 namespace UE::AnimNext
-{	
-	class IParameterSource;
+{
 	struct FScheduleContext;
 	enum class EParameterScopeOrdering : int32;
 	struct FScheduleBeginTickFunction;
@@ -21,7 +22,7 @@ struct FScheduleTaskContext
 {
 public:
 	// Apply the supplied parameter source to the specified scope, evicting any source that was there previously
-	ANIMNEXT_API void ApplyParametersToScope(FName InScope, EParameterScopeOrdering InOrdering, TUniquePtr<IParameterSource>&& InParameters) const;
+	ANIMNEXT_API void ApplyParametersToScope(FName InScope, EParameterScopeOrdering InOrdering, FName InInstanceId, FInstancedPropertyBag&& InPropertyBag) const;
 
 private:
 	FScheduleTaskContext(const FScheduleContext& InContext);
diff --git a/Engine/Plugins/Experimental/AnimNext/Source/AnimNext/Private/AnimNextModule.cpp b/Engine/Plugins/Experimental/AnimNext/Source/AnimNext/Private/AnimNextModule.cpp
--- a/Engine/Plugins/Experimental/AnimNext/Source/AnimNext/Private/AnimNextModule.cpp
+++ b/Engine/Plugins/Experimental/AnimNext/Source/AnimNext/Private/AnimNextModule.cpp
@@ -1,16 +1,18 @@
 // Copyright Epic Games, Inc. All Rights Reserved.
 
-#include "CoreMinimal.h"
+#include "AnimNextModule.h"
 #include "AnimNextConfig.h"
 #include "DataRegistry.h"
-#include "IAnimNextModule.h"
 #include "IUniversalObjectLocatorModule.h"
 #include "RigVMRuntimeDataRegistry.h"
 #include "Animation/AnimSequence.h"
 #include "Animation/BlendProfile.h"
+#include "Component/AnimNextComponent.h"
 #include "Curves/CurveFloat.h"
 #include "Graph/AnimNextGraph.h"
 #include "Modules/ModuleManager.h"
+#include "Param/AnimNextActorLocatorFragment.h"
+#include "Param/AnimNextComponentLocatorFragment.h"
 #include "Param/AnimNextEditorParam.h"
 #include "Param/AnimNextObjectCastLocatorFragment.h"
 #include "Param/AnimNextObjectFunctionLocatorFragment.h"
@@ -18,15 +20,13 @@
 #include "Param/AnimNextParam.h"
 #include "Param/AnimNextTag.h"
 #include "Param/ObjectProxyFactory.h"
+#include "Param/PropertyBagProxy.h"
 #include "RigVMCore/RigVMRegistry.h"
 #include "Scheduler/AnimNextTickFunctionBinding.h"
 #include "Scheduler/Scheduler.h"
 #include "TraitCore/NodeTemplateRegistry.h"
 #include "TraitCore/TraitRegistry.h"
 
-// Enable console commands only in development builds when logging is enabled
-#define WITH_ANIMNEXT_CONSOLE_COMMANDS (!UE_BUILD_SHIPPING && !NO_LOGGING)
-
 #if WITH_ANIMNEXT_CONSOLE_COMMANDS
 #include "HAL/IConsoleManager.h"
 #include "UObject/UObjectIterator.h"
@@ -36,13 +36,12 @@
 #include "TraitCore/TraitTemplate.h"
 #endif
 
+#define LOCTEXT_NAMESPACE "AnimNextModule"
+
 namespace UE::AnimNext
 {
 
-class FModule : public IAnimNextModule
-{
-public:
-	virtual void StartupModule() override
+	void FModule::StartupModule()
 	{
 		GetMutableDefault<UAnimNextConfig>()->LoadConfig();
 
@@ -54,6 +53,7 @@ public:
 			{ UBlendProfile::StaticClass(), FRigVMRegistry::ERegisterObjectOperation::Class },
 			{ UCurveFloat::StaticClass(), FRigVMRegistry::ERegisterObjectOperation::Class },
 			{ UAnimNextGraph::StaticClass(), FRigVMRegistry::ERegisterObjectOperation::Class },
+			{ UAnimNextComponent::StaticClass(), FRigVMRegistry::ERegisterObjectOperation::Class },
 		};
 
 		RigVMRegistry.RegisterObjectTypes(AllowedObjectTypes);
@@ -82,20 +82,30 @@ public:
 			[&UolModule]
 			{
 				{
-					UE::UniversalObjectLocator::FFragmentTypeParameters FragmentTypeParams("animobjfunc", NSLOCTEXT("Engine", "AnimNextObjectFunctionFragment", "Function"));
+					UE::UniversalObjectLocator::FFragmentTypeParameters FragmentTypeParams("animobjfunc", LOCTEXT("AnimNextObjectFunctionFragment", "Function"));
 					FragmentTypeParams.PrimaryEditorType = "AnimNextObjectFunction";
 					FAnimNextObjectFunctionLocatorFragment::FragmentType = UolModule.RegisterFragmentType<FAnimNextObjectFunctionLocatorFragment>(FragmentTypeParams);
 				}
 				{
-					UE::UniversalObjectLocator::FFragmentTypeParameters FragmentTypeParams("animobjprop", NSLOCTEXT("Engine", "AnimNextObjectPropertyFragment", "Property"));
+					UE::UniversalObjectLocator::FFragmentTypeParameters FragmentTypeParams("animobjprop", LOCTEXT("AnimNextObjectPropertyFragment", "Property"));
 					FragmentTypeParams.PrimaryEditorType = "AnimNextObjectProperty";
 					FAnimNextObjectPropertyLocatorFragment::FragmentType = UolModule.RegisterFragmentType<FAnimNextObjectPropertyLocatorFragment>(FragmentTypeParams);
 				}
 				{
-					UE::UniversalObjectLocator::FFragmentTypeParameters FragmentTypeParams("animobjcast", NSLOCTEXT("Engine", "AnimNextCastFragment", "Cast"));
+					UE::UniversalObjectLocator::FFragmentTypeParameters FragmentTypeParams("animobjcast", LOCTEXT("AnimNextCastFragment", "Cast"));
 					FragmentTypeParams.PrimaryEditorType = "AnimNextObjectCast";
 					FAnimNextObjectCastLocatorFragment::FragmentType = UolModule.RegisterFragmentType<FAnimNextObjectCastLocatorFragment>(FragmentTypeParams);
 				}
+				{
+					UE::UniversalObjectLocator::FFragmentTypeParameters FragmentTypeParams("animcomp", LOCTEXT("AnimNextComponentFragment", "AnimNextComponent"));
+					FragmentTypeParams.PrimaryEditorType = "AnimNextComponent";
+					FAnimNextComponentLocatorFragment::FragmentType = UolModule.RegisterFragmentType<FAnimNextComponentLocatorFragment>(FragmentTypeParams);
+				}
+				{
+					UE::UniversalObjectLocator::FFragmentTypeParameters FragmentTypeParams("animactor", LOCTEXT("AnimNextActorFragment", "AnimNextActor"));
+					FragmentTypeParams.PrimaryEditorType = "AnimNextActor";
+					FAnimNextActorLocatorFragment::FragmentType = UolModule.RegisterFragmentType<FAnimNextActorLocatorFragment>(FragmentTypeParams);
+				}
 			});
 #if WITH_ANIMNEXT_CONSOLE_COMMANDS
 		if (!IsRunningCommandlet())
@@ -116,7 +126,7 @@ public:
 #endif
 	}
 
-	virtual void ShutdownModule() override
+	void FModule::ShutdownModule()
 	{
 		FRigVMRuntimeDataRegistry::Destroy();
 		FScheduler::Destroy();
@@ -136,17 +146,17 @@ public:
 
 	const IAnimNextAnimGraph* AnimGraphImpl = nullptr;
 
-	virtual void RegisterAnimNextAnimGraph(const IAnimNextAnimGraph& InAnimGraphImpl) override
+	void FModule::RegisterAnimNextAnimGraph(const IAnimNextAnimGraph& InAnimGraphImpl)
 	{
 		AnimGraphImpl = &InAnimGraphImpl;
 	}
 
-	virtual void UnregisterAnimNextAnimGraph() override
+	void FModule::UnregisterAnimNextAnimGraph()
 	{
 		AnimGraphImpl = nullptr;
 	}
 
-	virtual void UpdateGraph(FAnimNextGraphInstancePtr& GraphInstance, float DeltaTime) override
+	void FModule::UpdateGraph(FAnimNextGraphInstancePtr& GraphInstance, float DeltaTime)
 	{
 		if (AnimGraphImpl != nullptr)
 		{
@@ -154,7 +164,7 @@ public:
 		}
 	}
 
-	virtual void EvaluateGraph(FAnimNextGraphInstancePtr& GraphInstance, const UE::AnimNext::FReferencePose& RefPose, int32 GraphLODLevel, FLODPoseHeap& OutputPose) const override
+	void FModule::EvaluateGraph(FAnimNextGraphInstancePtr& GraphInstance, const UE::AnimNext::FReferencePose& RefPose, int32 GraphLODLevel, FLODPoseHeap& OutputPose) const
 	{
 		if (AnimGraphImpl != nullptr)
 		{
@@ -162,7 +172,7 @@ public:
 		}
 	}
 
-	virtual TUniquePtr<IParameterSource> CreateParameterSource(const FParameterSourceContext& InContext, const TInstancedStruct<FAnimNextParamInstanceIdentifier>& InInstanceId, TConstArrayView<FName> InRequiredParameters) const override
+	TUniquePtr<IParameterSource> FModule::CreateParameterSource(const FParameterSourceContext& InContext, const TInstancedStruct<FAnimNextParamInstanceIdentifier>& InInstanceId, TConstArrayView<FName> InRequiredParameters) const
 	{
 		for(const TPair<FName, TSharedRef<IParameterSourceFactory>>& FactoryPair : ParameterSourceFactories)
 		{
@@ -178,17 +188,17 @@ public:
 		return nullptr;
 	}
 
-	virtual void RegisterParameterSourceFactory(FName InName, TSharedRef<IParameterSourceFactory> InFactory) override
+	void FModule::RegisterParameterSourceFactory(FName InName, TSharedRef<IParameterSourceFactory> InFactory)
 	{
 		ParameterSourceFactories.Add(InName, InFactory);
 	}
 
-	virtual void UnregisterParameterSourceFactory(FName InName) override
+	void FModule::UnregisterParameterSourceFactory(FName InName)
 	{
 		ParameterSourceFactories.Remove(InName);
 	}
 
-	virtual TSharedPtr<IParameterSourceFactory> FindParameterSourceFactory(FName InName) override
+	TSharedPtr<IParameterSourceFactory> FModule::FindParameterSourceFactory(FName InName)
 	{
 		if(TSharedRef<IParameterSourceFactory>* FoundFactory = ParameterSourceFactories.Find(InName))
 		{
@@ -199,9 +209,7 @@ public:
 	}
 
 #if WITH_ANIMNEXT_CONSOLE_COMMANDS
-	TArray<IConsoleObject*> ConsoleCommands;
-
-	void ListNodeTemplates(const TArray<FString>& Args)
+	void FModule::ListNodeTemplates(const TArray<FString>& Args)
 	{
 		// Turn off log times to make diff-ing easier
 		TGuardValue<ELogTimes::Type> DisableLogTimes(GPrintLogTimes, ELogTimes::None);
@@ -256,7 +264,7 @@ public:
 		LogAnimation.SetVerbosity(OldVerbosity);
 	}
 
-	void ListAnimGraphs(const TArray<FString>& Args)
+	void FModule::ListAnimGraphs(const TArray<FString>& Args)
 	{
 		// Turn off log times to make diff-ing easier
 		TGuardValue<ELogTimes::Type> DisableLogTimes(GPrintLogTimes, ELogTimes::None);
@@ -361,16 +369,14 @@ public:
 		LogAnimation.SetVerbosity(OldVerbosity);
 	}
 #endif
-	
-	// All known factories
-	TMap<FName, TSharedRef<IParameterSourceFactory>> ParameterSourceFactories;
-};
 
-IAnimNextModule& IAnimNextModule::Get()
-{
-	return FModuleManager::LoadModuleChecked<IAnimNextModule>(TEXT("AnimNext"));
-}
+	IAnimNextModule& IAnimNextModule::Get()
+	{
+		return FModuleManager::LoadModuleChecked<IAnimNextModule>(TEXT("AnimNext"));
+	}
 
 }
 
+#undef LOCTEXT_NAMESPACE
+
 IMPLEMENT_MODULE(UE::AnimNext::FModule, AnimNext)
diff --git a/Engine/Plugins/Experimental/AnimNext/Source/AnimNext/Private/Component/AnimNextComponent.cpp b/Engine/Plugins/Experimental/AnimNext/Source/AnimNext/Private/Component/AnimNextComponent.cpp
--- a/Engine/Plugins/Experimental/AnimNext/Source/AnimNext/Private/Component/AnimNextComponent.cpp
+++ b/Engine/Plugins/Experimental/AnimNext/Source/AnimNext/Private/Component/AnimNextComponent.cpp
@@ -8,16 +8,15 @@
 #include "Param/PropertyBagProxy.h"
 #include "Scheduler/ScheduleTaskContext.h"
 #include "Scheduler/ScheduleInitializationContext.h"
+#include "IAnimNextModule.h"
 
 namespace UE::AnimNext::Private
 {
 
 template<typename ContextType>
-static void SetValuesInScopeHelper(const ContextType& InContext, FName InScope, EAnimNextParameterScopeOrdering InOrdering, TConstArrayView<FPropertyBagProxy::FPropertyAndValue> InPropertiesAndValues)
+static void SetValuesInScopeHelper(const ContextType& InContext, FName InId, FName InScope, EAnimNextParameterScopeOrdering InOrdering, TConstArrayView<FPropertyBagPropertyDesc> InPropertyDescs, TConstArrayView<TConstArrayView<uint8>> InValues)
 {
-	TUniquePtr<FPropertyBagProxy> PropertyBagProxy = MakeUnique<FPropertyBagProxy>();
-	PropertyBagProxy->AddPropertiesAndValues(InPropertiesAndValues);
-	InContext.ApplyParametersToScope(InScope, (EParameterScopeOrdering)InOrdering, MoveTemp(PropertyBagProxy));
+	InContext.ApplyParametersToScope(InScope, (EParameterScopeOrdering)InOrdering, InId, InPropertyDescs, InValues);
 }
 
 }
@@ -46,17 +45,21 @@ void UAnimNextComponent::OnRegister()
 			// Now apply to each scope
 			for(const TPair<FName, TArray<UAnimNextComponentParameter*, TInlineAllocator<4>>>& ParamPair : ParamsByScope)
 			{
-				TArray<FPropertyBagProxy::FPropertyAndValue, TInlineAllocator<16>> PropertiesAndValues;
-				PropertiesAndValues.Reserve(ParamPair.Value.Num());
+				TArray<FPropertyBagPropertyDesc, TInlineAllocator<16>> PropertyDescs;
+				TArray<TConstArrayView<uint8>, TInlineAllocator<16>> Values;
+				PropertyDescs.Reserve(ParamPair.Value.Num());
+				Values.Reserve(ParamPair.Value.Num());
 				for(UAnimNextComponentParameter* Parameter : ParamPair.Value)
 				{
-					FPropertyBagProxy::FPropertyAndValue& PropertyAndValue = PropertiesAndValues.AddDefaulted_GetRef();
-					PropertyAndValue.ContainerPtr = Parameter;
-					Parameter->GetParamInfo(PropertyAndValue.Name, PropertyAndValue.Property);
+					FName Name;
+					const FProperty* Property;
+					Parameter->GetParamInfo(Name, Property);
+					PropertyDescs.Emplace(Name, Property);
+					Values.Emplace(Property->ContainerPtrToValuePtr<uint8>(Parameter), Property->GetSize());
 				}
 
 				// NOTE: Layer is always applied 'before' currently. If we have a use case for 'After' we can add it to UAnimNextComponentParameter
-				Private::SetValuesInScopeHelper(InContext, ParamPair.Key, EAnimNextParameterScopeOrdering::Before, PropertiesAndValues);
+				Private::SetValuesInScopeHelper(InContext, "ComponentParams", ParamPair.Key, EAnimNextParameterScopeOrdering::Before, PropertyDescs, Values);
 			}
 		};
 
@@ -121,19 +124,15 @@ DEFINE_FUNCTION(UAnimNextComponent::execSetParameterInScope)
 
 	P_NATIVE_BEGIN;
 
-	TUniquePtr<FInstancedPropertyBag> PropertyBag = MakeUnique<FInstancedPropertyBag>();
-	PropertyBag->AddProperty(Name, ValueProp);
-	const FProperty* NewProperty = PropertyBag->GetPropertyBagStruct()->GetPropertyDescs()[0].CachedProperty;
+	FInstancedPropertyBag PropertyBag;
+	PropertyBag.AddProperty(Name, ValueProp);
+	const FProperty* NewProperty = PropertyBag.GetPropertyBagStruct()->GetPropertyDescs()[0].CachedProperty;
 	const void* ValuePtr = ValueProp->ContainerPtrToValuePtr<void>(ContainerPtr);
-	NewProperty->SetValue_InContainer(PropertyBag->GetMutableValue().GetMemory(), ValuePtr);
+	NewProperty->SetValue_InContainer(PropertyBag.GetMutableValue().GetMemory(), ValuePtr);
 
-	FScheduler::QueueTask(P_THIS, P_THIS->SchedulerHandle, Scope, [Scope, Name, NewProperty, PropertyBag = MoveTemp(PropertyBag), Ordering](const FScheduleTaskContext& InContext) mutable
+	FScheduler::QueueTask(P_THIS, P_THIS->SchedulerHandle, Scope, [Scope, Name, PropertyBag = MoveTemp(PropertyBag), Ordering](const FScheduleTaskContext& InContext) mutable
 	{
-		FPropertyBagProxy::FPropertyAndValue PropertyAndValue;
-		PropertyAndValue.Name = Name;
-		PropertyAndValue.Property = NewProperty;
-		PropertyAndValue.ContainerPtr = PropertyBag->GetValue().GetMemory();
-		Private::SetValuesInScopeHelper(InContext, Scope, Ordering, { PropertyAndValue });
+		InContext.ApplyParametersToScope(Scope, (EParameterScopeOrdering)Ordering, Name, MoveTemp(PropertyBag));
 	});
 
 	P_NATIVE_END;
diff --git a/Engine/Plugins/Experimental/AnimNext/Source/AnimNext/Private/Component/AnimNextMeshComponent.cpp b/Engine/Plugins/Experimental/AnimNext/Source/AnimNext/Private/Component/AnimNextMeshComponent.cpp
deleted file mode 100644
diff --git a/Engine/Plugins/Experimental/AnimNext/Source/AnimNext/Private/Component/AnimNextMeshComponent.h b/Engine/Plugins/Experimental/AnimNext/Source/AnimNext/Private/Component/AnimNextMeshComponent.h
deleted file mode 100644
diff --git a/Engine/Plugins/Experimental/AnimNext/Source/AnimNext/Private/Component/SkinnedMeshComponentExtensions.cpp b/Engine/Plugins/Experimental/AnimNext/Source/AnimNext/Private/Component/SkinnedMeshComponentExtensions.cpp
new file mode 100644
--- /dev/null
+++ b/Engine/Plugins/Experimental/AnimNext/Source/AnimNext/Private/Component/SkinnedMeshComponentExtensions.cpp
@@ -0,0 +1,32 @@
+// Copyright Epic Games, Inc. All Rights Reserved.
+
+#include "Component/SkinnedMeshComponentExtensions.h"
+#include "GenerationTools.h"
+#include "Components/SkinnedMeshComponent.h"
+
+namespace UE::Anim
+{
+
+void FSkinnedMeshComponentExtensions::CompleteAndDispatch(USkinnedMeshComponent* InComponent, TConstArrayView<FBoneIndexType> InParentIndices, TConstArrayView<FBoneIndexType> InRequiredBoneIndices, TConstArrayView<FTransform> InLocalSpaceTransforms)
+{
+	// Fill the component space transform buffer
+	TArrayView<FTransform> ComponentSpaceTransforms = InComponent->GetEditableComponentSpaceTransforms();
+	if(ComponentSpaceTransforms.Num() > 0)
+	{
+		UE::AnimNext::FGenerationTools::ConvertLocalSpaceToComponentSpace(InParentIndices, InLocalSpaceTransforms, InRequiredBoneIndices, ComponentSpaceTransforms);
+
+		// Flag buffer for flip
+		InComponent->bNeedToFlipSpaceBaseBuffers = true;
+		InComponent->bHasValidBoneTransform = false;
+		InComponent->FlipEditableSpaceBases();
+		InComponent->bHasValidBoneTransform = true;
+
+		InComponent->InvalidateCachedBounds();
+		InComponent->UpdateBounds();
+
+		// Send updated transforms to the renderer
+		InComponent->SendRenderDynamicData_Concurrent();
+	}
+}
+
+}
diff --git a/Engine/Plugins/Experimental/AnimNext/Source/AnimNext/Private/Graph/AnimNextGraph.cpp b/Engine/Plugins/Experimental/AnimNext/Source/AnimNext/Private/Graph/AnimNextGraph.cpp
--- a/Engine/Plugins/Experimental/AnimNext/Source/AnimNext/Private/Graph/AnimNextGraph.cpp
+++ b/Engine/Plugins/Experimental/AnimNext/Source/AnimNext/Private/Graph/AnimNextGraph.cpp
@@ -19,6 +19,7 @@
 #include "Graph/AnimNextGraphEntryPoint.h"
 #include "Graph/RigUnit_AnimNextBeginExecution.h"
 #include "Graph/RigUnit_AnimNextGraphRoot.h"
+#include "Param/ParametersProxy.h"
 #include "UObject/LinkerLoad.h"
 #include "UObject/ObjectResource.h"
 
@@ -67,6 +68,7 @@ void UAnimNextGraph::AllocateInstanceImpl(FAnimNextGraphInstance* ParentGraphIns
 		InstanceImpl->Graph = this;
 		InstanceImpl->ParentGraphInstance = ParentGraphInstance;
 		InstanceImpl->EntryPoint = EntryPoint;
+		InstanceImpl->GraphState = MakeUnique<UE::AnimNext::FParametersProxy>(this);
 
 		// If we have a parent graph, use its root since we share the same root, otherwise if we have no parent, we are the root
 		InstanceImpl->RootGraphInstance = ParentGraphInstance != nullptr ? ParentGraphInstance->GetRootGraphInstance() : InstanceImpl.Get();
@@ -189,6 +191,18 @@ void UAnimNextGraph::Serialize(FArchive& Ar)
 	}
 }
 
+void UAnimNextGraph::PostLoad()
+{
+	Super::PostLoad();
+
+#if WITH_EDITORONLY_DATA
+	if(GetLinkerCustomVersion(FFortniteMainBranchObjectVersion::GUID) < FFortniteMainBranchObjectVersion::AnimNextGraphAccessSpecifiers)
+	{
+		DefaultState.State = PropertyBag_DEPRECATED;
+	}
+#endif
+}
+
 #if WITH_EDITOR
 void UAnimNextGraph::PostEditChangeProperty(FPropertyChangedEvent& PropertyChangedEvent)
 {
@@ -224,9 +238,11 @@ bool UAnimNextGraph::LoadFromArchiveBuffer(const TArray<uint8>& InSharedDataArch
 	const FTraitReader::EErrorState ErrorState = TraitReader.ReadGraph(SharedDataBuffer);
 	if (ErrorState == FTraitReader::EErrorState::None)
 	{
-		for(FAnimNextGraphEntryPoint& EntryPoint : EntryPoints)
+		for(int32 EntryPointIndex = 0; EntryPointIndex < EntryPoints.Num(); ++EntryPointIndex)
 		{
+			const FAnimNextGraphEntryPoint& EntryPoint = EntryPoints[EntryPointIndex];
 			ResolvedRootTraitHandles.Add(EntryPoint.EntryPointName, TraitReader.ResolveEntryPointHandle(EntryPoint.RootTraitHandle));
+			ResolvedEntryPoints.Add(EntryPoint.EntryPointName, EntryPointIndex);
 		}
 
 		// Make sure our execute method is registered
diff --git a/Engine/Plugins/Experimental/AnimNext/Source/AnimNext/Private/Graph/AnimNextGraphInstance.cpp b/Engine/Plugins/Experimental/AnimNext/Source/AnimNext/Private/Graph/AnimNextGraphInstance.cpp
--- a/Engine/Plugins/Experimental/AnimNext/Source/AnimNext/Private/Graph/AnimNextGraphInstance.cpp
+++ b/Engine/Plugins/Experimental/AnimNext/Source/AnimNext/Private/Graph/AnimNextGraphInstance.cpp
@@ -8,9 +8,13 @@
 #include "Graph/GC_GraphInstanceComponent.h"
 #include "Graph/RigUnit_AnimNextShimRoot.h"
 #include "Misc/ScopeRWLock.h"
+#include "Param/IParameterSource.h"
+#include "Param/ParametersProxy.h"
 
 DEFINE_STAT(STAT_AnimNext_Graph_RigVM);
 
+FAnimNextGraphInstance::FAnimNextGraphInstance() = default;
+
 FAnimNextGraphInstance::~FAnimNextGraphInstance()
 {
 	Release();
@@ -101,6 +105,11 @@ void FAnimNextGraphInstance::AddStructReferencedObjects(FReferenceCollector& Col
 	{
 		Component->AddReferencedObjects(Collector);
 	}
+
+	if(GraphState.IsValid())
+	{
+		GraphState->AddReferencedObjects(Collector);
+	}
 }
 
 UE::AnimNext::FGraphInstanceComponent* FAnimNextGraphInstance::TryGetComponent(int32 ComponentNameHash, FName ComponentName) const
@@ -189,6 +198,7 @@ void FAnimNextGraphInstance::Freeze()
 	}
 
 	GraphInstancePtr.Reset();
+	GraphState.Reset();
 	ExtendedExecuteContext.Reset();
 	Components.Empty();
 	bHasUpdatedOnce = false;
@@ -198,6 +208,8 @@ void FAnimNextGraphInstance::Thaw()
 {
 	if (const UAnimNextGraph* GraphPtr = Graph)
 	{
+		GraphState = MakeUnique<UE::AnimNext::FParametersProxy>(Graph);
+
 		ExtendedExecuteContext.CopyMemoryStorage(GraphPtr->ExtendedExecuteContext);
 		GraphPtr->VM->InitializeInstance(ExtendedExecuteContext);
 
@@ -216,3 +228,18 @@ void FAnimNextGraphInstance::Thaw()
 		}
 	}
 }
+
+UE::AnimNext::FParamStack::FPushedLayerHandle FAnimNextGraphInstance::UpdateAndPushGraphState(float InDeltaTime) const
+{
+	if(GraphState.IsValid())
+	{
+		GraphState->Update(InDeltaTime);
+		return UE::AnimNext::FParamStack::Get().PushLayer(GraphState->GetLayerHandle());
+	}
+	return UE::AnimNext::FParamStack::FPushedLayerHandle();
+}
+
+void FAnimNextGraphInstance::PopGraphState(UE::AnimNext::FParamStack::FPushedLayerHandle InHandle) const
+{
+	UE::AnimNext::FParamStack::Get().PopLayer(InHandle);
+}
\ No newline at end of file
diff --git a/Engine/Plugins/Experimental/AnimNext/Source/AnimNext/Private/Graph/AnimNextGraphInstancePtr.cpp b/Engine/Plugins/Experimental/AnimNext/Source/AnimNext/Private/Graph/AnimNextGraphInstancePtr.cpp
--- a/Engine/Plugins/Experimental/AnimNext/Source/AnimNext/Private/Graph/AnimNextGraphInstancePtr.cpp
+++ b/Engine/Plugins/Experimental/AnimNext/Source/AnimNext/Private/Graph/AnimNextGraphInstancePtr.cpp
@@ -4,6 +4,7 @@
 
 #include "Graph/AnimNextGraph.h"
 #include "Graph/AnimNextGraphInstance.h"
+#include "Param/IParameterSource.h"
 
 FAnimNextGraphInstancePtr::FAnimNextGraphInstancePtr() = default;
 FAnimNextGraphInstancePtr::FAnimNextGraphInstancePtr(FAnimNextGraphInstancePtr&&) = default;
@@ -117,3 +118,15 @@ void FAnimNextGraphInstancePtr::Update()
 	check(Impl);
 	Impl->Update();
 }
+
+UE::AnimNext::FParamStack::FPushedLayerHandle FAnimNextGraphInstancePtr::UpdateAndPushGraphState(float InDeltaTime) const
+{
+	check(Impl);
+	return Impl->UpdateAndPushGraphState(InDeltaTime);
+}
+
+void FAnimNextGraphInstancePtr::PopGraphState(UE::AnimNext::FParamStack::FPushedLayerHandle InHandle) const
+{
+	check(Impl);
+	Impl->PopGraphState(InHandle);
+}
diff --git a/Engine/Plugins/Experimental/AnimNext/Source/AnimNext/Private/Param/AnimNextActorComponentLibrary.cpp b/Engine/Plugins/Experimental/AnimNext/Source/AnimNext/Private/Param/AnimNextActorComponentLibrary.cpp
--- a/Engine/Plugins/Experimental/AnimNext/Source/AnimNext/Private/Param/AnimNextActorComponentLibrary.cpp
+++ b/Engine/Plugins/Experimental/AnimNext/Source/AnimNext/Private/Param/AnimNextActorComponentLibrary.cpp
@@ -1,14 +1,9 @@
-﻿// Copyright Epic Games, Inc. All Rights Reserved.
+// Copyright Epic Games, Inc. All Rights Reserved.
 
 #include "AnimNextActorComponentLibrary.h"
 #include "Scheduler/AnimNextTickFunctionBinding.h"
 #include "Components/ActorComponent.h"
 
-UClass* UAnimNextActorComponentLibrary::GetSupportedClass() const
-{
-	return UActorComponent::StaticClass();
-}
-
 FAnimNextTickFunctionBinding UAnimNextActorComponentLibrary::GetTick(UActorComponent* InComponent)
 {
 	FAnimNextTickFunctionBinding Binding;
diff --git a/Engine/Plugins/Experimental/AnimNext/Source/AnimNext/Private/Param/AnimNextActorComponentLibrary.h b/Engine/Plugins/Experimental/AnimNext/Source/AnimNext/Private/Param/AnimNextActorComponentLibrary.h
--- a/Engine/Plugins/Experimental/AnimNext/Source/AnimNext/Private/Param/AnimNextActorComponentLibrary.h
+++ b/Engine/Plugins/Experimental/AnimNext/Source/AnimNext/Private/Param/AnimNextActorComponentLibrary.h
@@ -1,23 +1,20 @@
-﻿// Copyright Epic Games, Inc. All Rights Reserved.
+// Copyright Epic Games, Inc. All Rights Reserved.
 
 #pragma once
 
 #include "CoreMinimal.h"
-#include "Param/AnimNextClassExtensionLibrary.h"
 #include "Scheduler/AnimNextTickFunctionBinding.h"
+#include "Kismet/BlueprintFunctionLibrary.h"
 #include "AnimNextActorComponentLibrary.generated.h"
 
 class UActorComponent;
 
 // Access to non-UProperty/UFunction data on UActorComponent
 UCLASS()
-class UAnimNextActorComponentLibrary : public UAnimNextClassExtensionLibrary
+class UAnimNextActorComponentLibrary : public UBlueprintFunctionLibrary
 {
 	GENERATED_BODY()
 
-	// UAnimNextClassProxy interface
-	virtual UClass* GetSupportedClass() const override;
-
 public:
 	// Returns the component's tick function
 	UFUNCTION(BlueprintCallable, BlueprintInternalUseOnly)
diff --git a/Engine/Plugins/Experimental/AnimNext/Source/AnimNext/Private/Param/AnimNextActorLocatorFragment.cpp b/Engine/Plugins/Experimental/AnimNext/Source/AnimNext/Private/Param/AnimNextActorLocatorFragment.cpp
new file mode 100644
--- /dev/null
+++ b/Engine/Plugins/Experimental/AnimNext/Source/AnimNext/Private/Param/AnimNextActorLocatorFragment.cpp
@@ -0,0 +1,57 @@
+﻿// Copyright Epic Games, Inc. All Rights Reserved.
+
+#include "Param/AnimNextActorLocatorFragment.h"
+#include "GameFramework/Actor.h"
+#include "Components/ActorComponent.h"
+#include "UniversalObjectLocatorFragmentTypeHandle.h"
+#include "UniversalObjectLocatorInitializeParams.h"
+#include "UniversalObjectLocatorResolveParams.h"
+#include "UniversalObjectLocatorStringParams.h"
+#include "UniversalObjectLocatorInitializeResult.h"
+
+UE::UniversalObjectLocator::TFragmentTypeHandle<FAnimNextActorLocatorFragment> FAnimNextActorLocatorFragment::FragmentType;
+
+UE::UniversalObjectLocator::FResolveResult FAnimNextActorLocatorFragment::Resolve(const UE::UniversalObjectLocator::FResolveParams& InParams) const
+{
+	using namespace UE::UniversalObjectLocator;
+
+	UObject* Result = nullptr;
+
+	if(InParams.Context)
+	{
+		// Called on a component context - return the actor
+		if(UActorComponent* Component = Cast<UActorComponent>(InParams.Context))
+		{
+			Result = Component->GetOwner();
+		}
+	}
+
+	return FResolveResultData(Result);
+}
+
+void FAnimNextActorLocatorFragment::ToString(FStringBuilderBase& OutStringBuilder) const
+{
+}
+
+UE::UniversalObjectLocator::FParseStringResult FAnimNextActorLocatorFragment::TryParseString(FStringView InString, const UE::UniversalObjectLocator::FParseStringParams& Params)
+{
+	return UE::UniversalObjectLocator::FParseStringResult().Success();
+}
+
+UE::UniversalObjectLocator::FInitializeResult FAnimNextActorLocatorFragment::Initialize(const UE::UniversalObjectLocator::FInitializeParams& InParams)
+{
+	using namespace UE::UniversalObjectLocator;
+
+	if(InParams.Context == nullptr)
+	{
+		return FInitializeResult::Failure();
+	}
+
+	return FInitializeResult::Absolute();
+}
+
+uint32 FAnimNextActorLocatorFragment::ComputePriority(const UObject* ObjectToReference, const UObject* Context)
+{
+	// We can't use this at all unless explicitly added by code
+	return 0;
+}
diff --git a/Engine/Plugins/Experimental/AnimNext/Source/AnimNext/Private/Param/AnimNextComponentLocatorFragment.cpp b/Engine/Plugins/Experimental/AnimNext/Source/AnimNext/Private/Param/AnimNextComponentLocatorFragment.cpp
new file mode 100644
--- /dev/null
+++ b/Engine/Plugins/Experimental/AnimNext/Source/AnimNext/Private/Param/AnimNextComponentLocatorFragment.cpp
@@ -0,0 +1,72 @@
+﻿// Copyright Epic Games, Inc. All Rights Reserved.
+
+#include "Param/AnimNextComponentLocatorFragment.h"
+#include "Component/AnimNextComponent.h"
+#include "UniversalObjectLocatorFragmentTypeHandle.h"
+#include "UniversalObjectLocatorInitializeParams.h"
+#include "UniversalObjectLocatorResolveParams.h"
+#include "UniversalObjectLocatorStringParams.h"
+#include "UniversalObjectLocatorInitializeResult.h"
+
+UE::UniversalObjectLocator::TFragmentTypeHandle<FAnimNextComponentLocatorFragment> FAnimNextComponentLocatorFragment::FragmentType;
+
+FAnimNextComponentLocatorFragment::FAnimNextComponentLocatorFragment(TSubclassOf<UActorComponent> InClass)
+	: Path(InClass)
+{
+}
+
+UE::UniversalObjectLocator::FResolveResult FAnimNextComponentLocatorFragment::Resolve(const UE::UniversalObjectLocator::FResolveParams& InParams) const
+{
+	using namespace UE::UniversalObjectLocator;
+
+	UObject* Result = nullptr;
+
+	if (UActorComponent* Component = Cast<UActorComponent>(InParams.Context))
+	{
+		if(AActor* Actor = Component->GetOwner())
+		{
+			if (UClass* ComponentClass = Cast<UClass>(Path.ResolveObject()))
+			{
+				Result = Actor->FindComponentByClass(ComponentClass);
+			}
+		}
+	}
+	else if(AActor* Actor = Cast<AActor>(InParams.Context))
+	{
+		if(UClass* ComponentClass = Cast<UClass>(Path.ResolveObject()))
+		{
+			Result = Actor->FindComponentByClass(ComponentClass);
+		}
+	}
+
+	return FResolveResultData(Result);
+}
+
+void FAnimNextComponentLocatorFragment::ToString(FStringBuilderBase& OutStringBuilder) const
+{
+	Path.AppendString(OutStringBuilder);
+}
+
+UE::UniversalObjectLocator::FParseStringResult FAnimNextComponentLocatorFragment::TryParseString(FStringView InString, const UE::UniversalObjectLocator::FParseStringParams& Params)
+{
+	Path = InString;
+	return UE::UniversalObjectLocator::FParseStringResult().Success();
+}
+
+UE::UniversalObjectLocator::FInitializeResult FAnimNextComponentLocatorFragment::Initialize(const UE::UniversalObjectLocator::FInitializeParams& InParams)
+{
+	using namespace UE::UniversalObjectLocator;
+
+	if(InParams.Object == nullptr || !InParams.Object->GetClass()->IsChildOf(UActorComponent::StaticClass()))
+	{
+		return FInitializeResult::Failure();
+	}
+	
+	return FInitializeResult::Relative(InParams.Context);
+}
+
+uint32 FAnimNextComponentLocatorFragment::ComputePriority(const UObject* ObjectToReference, const UObject* Context)
+{
+	// We can't use this at all unless explicitly added by code
+	return 0;
+}
diff --git a/Engine/Plugins/Experimental/AnimNext/Source/AnimNext/Private/Param/AnimNextObjectFunctionLocatorFragment.cpp b/Engine/Plugins/Experimental/AnimNext/Source/AnimNext/Private/Param/AnimNextObjectFunctionLocatorFragment.cpp
--- a/Engine/Plugins/Experimental/AnimNext/Source/AnimNext/Private/Param/AnimNextObjectFunctionLocatorFragment.cpp
+++ b/Engine/Plugins/Experimental/AnimNext/Source/AnimNext/Private/Param/AnimNextObjectFunctionLocatorFragment.cpp
@@ -1,4 +1,4 @@
-﻿// Copyright Epic Games, Inc. All Rights Reserved.
+// Copyright Epic Games, Inc. All Rights Reserved.
 
 #include "Param/AnimNextObjectFunctionLocatorFragment.h"
 #include "UniversalObjectLocatorFragmentTypeHandle.h"
@@ -6,6 +6,7 @@
 #include "UniversalObjectLocatorStringParams.h"
 #include "UniversalObjectLocatorInitializeParams.h"
 #include "UniversalObjectLocatorInitializeResult.h"
+#include "Param/ParamUtils.h"
 
 UE::UniversalObjectLocator::TFragmentTypeHandle<FAnimNextObjectFunctionLocatorFragment> FAnimNextObjectFunctionLocatorFragment::FragmentType;
 
@@ -22,11 +23,23 @@ UE::UniversalObjectLocator::FResolveResult FAnimNextObjectFunctionLocatorFragmen
 
 	if(UFunction* Function = Cast<UFunction>(Path.ResolveObject()))
 	{
-		if(Params.Context && Function->NumParms == 1 && Params.Context->GetClass()->IsChildOf(Function->GetOuterUClass()))
+		if(Params.Context && UE::AnimNext::FParamUtils::CanUseFunction(Function, nullptr))
 		{
-			if(FObjectProperty* ObjectProperty = CastField<FObjectProperty>(Function->GetReturnProperty()))
+			if (FObjectProperty* ObjectProperty = CastField<FObjectProperty>(Function->GetReturnProperty()))
 			{
-				Params.Context->ProcessEvent(Function, &Result);
+				if (Function->NumParms == 1)
+				{
+					check(Params.Context->GetClass()->IsChildOf(Function->GetOuterUClass()));
+					FFrame Stack(Params.Context, Function, nullptr, nullptr, Function->ChildProperties);
+					Function->Invoke(Params.Context, Stack, &Result);
+				}
+				else if (Function->NumParms == 2)
+				{
+					UObject* CDO = Function->GetOuterUClass()->GetDefaultObject();
+					UObject* HoistedObject = Params.Context;
+					FFrame Stack(CDO, Function, &HoistedObject, nullptr, Function->ChildProperties);
+					Function->Invoke(CDO, Stack, &Result);
+				}
 			}
 		}
 	}
@@ -51,7 +64,7 @@ UE::UniversalObjectLocator::FInitializeResult FAnimNextObjectFunctionLocatorFrag
 
 	if(const UFunction* Function = Cast<UFunction>(InParams.Object))
 	{
-		if(Function->NumParms == 1)
+		if(UE::AnimNext::FParamUtils::CanUseFunction(Function, nullptr))
 		{
 			if(FObjectProperty* ObjectProperty = CastField<FObjectProperty>(Function->GetReturnProperty()))
 			{
diff --git a/Engine/Plugins/Experimental/AnimNext/Source/AnimNext/Private/Param/AnimNextSkeletalMeshComponentLibrary.cpp b/Engine/Plugins/Experimental/AnimNext/Source/AnimNext/Private/Param/AnimNextSkeletalMeshComponentLibrary.cpp
new file mode 100644
--- /dev/null
+++ b/Engine/Plugins/Experimental/AnimNext/Source/AnimNext/Private/Param/AnimNextSkeletalMeshComponentLibrary.cpp
@@ -0,0 +1,10 @@
+﻿// Copyright Epic Games, Inc. All Rights Reserved.
+
+#include "AnimNextSkeletalMeshComponentLibrary.h"
+#include "Components/SkeletalMeshComponent.h"
+
+FAnimNextGraphReferencePose UAnimNextSkeletalMeshComponentLibrary::GetReferencePose(USkeletalMeshComponent* InComponent)
+{
+	UE::AnimNext::FDataHandle RefPoseHandle = UE::AnimNext::FDataRegistry::Get()->GetOrGenerateReferencePose(InComponent);
+	return FAnimNextGraphReferencePose(RefPoseHandle);
+}
diff --git a/Engine/Plugins/Experimental/AnimNext/Source/AnimNext/Private/Param/AnimNextSkeletalMeshComponentLibrary.h b/Engine/Plugins/Experimental/AnimNext/Source/AnimNext/Private/Param/AnimNextSkeletalMeshComponentLibrary.h
new file mode 100644
--- /dev/null
+++ b/Engine/Plugins/Experimental/AnimNext/Source/AnimNext/Private/Param/AnimNextSkeletalMeshComponentLibrary.h
@@ -0,0 +1,22 @@
+﻿// Copyright Epic Games, Inc. All Rights Reserved.
+
+#pragma once
+
+#include "CoreMinimal.h"
+#include "Kismet/BlueprintFunctionLibrary.h"
+#include "Graph/AnimNext_LODPose.h"
+#include "AnimNextSkeletalMeshComponentLibrary.generated.h"
+
+class USkeletalMeshComponent;
+
+// Access to non-UProperty/UFunction data on USkeletalMeshComponent
+UCLASS()
+class UAnimNextSkeletalMeshComponentLibrary : public UBlueprintFunctionLibrary
+{
+	GENERATED_BODY()
+
+public:
+	// Returns the reference pose for the specified component
+	UFUNCTION(BlueprintCallable, BlueprintInternalUseOnly)
+	static FAnimNextGraphReferencePose GetReferencePose(USkeletalMeshComponent* InComponent);
+};
diff --git a/Engine/Plugins/Experimental/AnimNext/Source/AnimNext/Private/Param/AnimNextWorldLibrary.cpp b/Engine/Plugins/Experimental/AnimNext/Source/AnimNext/Private/Param/AnimNextWorldLibrary.cpp
--- a/Engine/Plugins/Experimental/AnimNext/Source/AnimNext/Private/Param/AnimNextWorldLibrary.cpp
+++ b/Engine/Plugins/Experimental/AnimNext/Source/AnimNext/Private/Param/AnimNextWorldLibrary.cpp
@@ -1,14 +1,9 @@
-﻿// Copyright Epic Games, Inc. All Rights Reserved.
+// Copyright Epic Games, Inc. All Rights Reserved.
 
 #include "AnimNextWorldLibrary.h"
 #include "Engine/World.h"
 #include "Scheduler/AnimNextTickFunctionBinding.h"
 
-UClass* UAnimNextWorldLibrary::GetSupportedClass() const
-{
-	return UWorld::StaticClass();
-}
-
 double UAnimNextWorldLibrary::GetTimeSeconds(UWorld* InWorld)
 {
 	return InWorld->GetTimeSeconds();
diff --git a/Engine/Plugins/Experimental/AnimNext/Source/AnimNext/Private/Param/AnimNextWorldLibrary.h b/Engine/Plugins/Experimental/AnimNext/Source/AnimNext/Private/Param/AnimNextWorldLibrary.h
--- a/Engine/Plugins/Experimental/AnimNext/Source/AnimNext/Private/Param/AnimNextWorldLibrary.h
+++ b/Engine/Plugins/Experimental/AnimNext/Source/AnimNext/Private/Param/AnimNextWorldLibrary.h
@@ -1,9 +1,9 @@
-﻿// Copyright Epic Games, Inc. All Rights Reserved.
+// Copyright Epic Games, Inc. All Rights Reserved.
 
 #pragma once
 
 #include "CoreMinimal.h"
-#include "Param/AnimNextClassExtensionLibrary.h"
+#include "Kismet/BlueprintFunctionLibrary.h"
 #include "AnimNextWorldLibrary.generated.h"
 
 class UWorld;
@@ -11,13 +11,10 @@ struct FAnimNextTickFunctionBinding;
 
 // Access to non-UProperty/UFunction data on UWorld
 UCLASS()
-class UAnimNextWorldLibrary : public UAnimNextClassExtensionLibrary
+class UAnimNextWorldLibrary : public UBlueprintFunctionLibrary
 {
 	GENERATED_BODY()
 
-	// UAnimNextClassProxy interface
-	virtual UClass* GetSupportedClass() const override;
-
 public:
 	// Returns time in seconds since world was brought up for play, IS stopped when game pauses, IS dilated/clamped
 	UFUNCTION(BlueprintCallable, BlueprintInternalUseOnly)
diff --git a/Engine/Plugins/Experimental/AnimNext/Source/AnimNext/Private/Param/ClassProxy.cpp b/Engine/Plugins/Experimental/AnimNext/Source/AnimNext/Private/Param/ClassProxy.cpp
--- a/Engine/Plugins/Experimental/AnimNext/Source/AnimNext/Private/Param/ClassProxy.cpp
+++ b/Engine/Plugins/Experimental/AnimNext/Source/AnimNext/Private/Param/ClassProxy.cpp
@@ -1,8 +1,8 @@
-﻿// Copyright Epic Games, Inc. All Rights Reserved.
+// Copyright Epic Games, Inc. All Rights Reserved.
 
 #include "Param/ClassProxy.h"
 #include "Logging/StructuredLog.h"
-#include "Param/AnimNextClassExtensionLibrary.h"
+#include "Kismet/BlueprintFunctionLibrary.h"
 #include "Param/ParamTypeHandle.h"
 #include "Param/ParamUtils.h"
 
@@ -17,38 +17,34 @@ FClassProxy::FClassProxy(const UClass* InClass)
 void FClassProxy::Refresh(const UClass* InClass)
 {
 	Class = InClass;
+	ParameterNameMap.Reset();
 
 	// Add any additional extension libraries that extend this class first
 	{
 		TArray<UClass*> Classes;
-		GetDerivedClasses(UAnimNextClassExtensionLibrary::StaticClass(), Classes);
+		GetDerivedClasses(UBlueprintFunctionLibrary::StaticClass(), Classes);
 		for(UClass* ProxyClass : Classes)
 		{
-			UAnimNextClassExtensionLibrary* CDO = ProxyClass->GetDefaultObject<UAnimNextClassExtensionLibrary>();
-			const UClass* ExtendedClass = CDO->GetSupportedClass();
-			if(InClass->IsChildOf(ExtendedClass))
+			for(TFieldIterator<UFunction> It(ProxyClass); It; ++It)
 			{
-				for(TFieldIterator<UFunction> It(ProxyClass); It; ++It)
+				UFunction* Function = *It;
+				if(FParamUtils::CanUseFunction(Function, InClass))
 				{
-					UFunction* Function = *It;
-					if(FParamUtils::CanUseFunction(Function))
+					const FProperty* ReturnProperty = Function->GetReturnProperty();
+					FParamTypeHandle TypeHandle = FParamTypeHandle::FromProperty(ReturnProperty);
+					if(TypeHandle.IsValid())
 					{
-						const FProperty* ReturnProperty = Function->GetReturnProperty();
-						FParamTypeHandle TypeHandle = FParamTypeHandle::FromProperty(ReturnProperty);
-						if(TypeHandle.IsValid())
-						{
-							FClassProxyParameter Parameter;
-							Parameter.AccessType = EClassProxyParameterAccessType::HoistedFunction;
-							Parameter.ParameterName = *Function->GetPathName();
-							Parameter.Function = Function;
-							Parameter.Type = TypeHandle.GetType();
+						FClassProxyParameter Parameter;
+						Parameter.AccessType = EClassProxyParameterAccessType::HoistedFunction;
+						Parameter.ParameterName = *Function->GetPathName();
+						Parameter.Function = Function;
+						Parameter.Type = TypeHandle.GetType();
 #if WITH_EDITOR
-							Parameter.DisplayName = Function->GetDisplayNameText();
-							Parameter.Tooltip = Function->GetToolTipText();
-							Parameter.bThreadSafe = Function->HasMetaData("BlueprintThreadSafe");
+						Parameter.DisplayName = Function->GetDisplayNameText();
+						Parameter.Tooltip = Function->GetToolTipText();
+						Parameter.bThreadSafe = Function->HasMetaData("BlueprintThreadSafe");
 #endif
-							ParameterNameMap.Add(Parameter.ParameterName, Parameters.Add(Parameter));
-						}
+						ParameterNameMap.Add(Parameter.ParameterName, Parameters.Add(Parameter));
 					}
 				}
 			}
@@ -59,7 +55,7 @@ void FClassProxy::Refresh(const UClass* InClass)
 	for(TFieldIterator<UFunction> It(InClass, EFieldIterationFlags::IncludeSuper | EFieldIterationFlags::IncludeInterfaces); It; ++It)
 	{
 		UFunction* Function = *It;
-		if(FParamUtils::CanUseFunction(Function))
+		if(FParamUtils::CanUseFunction(Function, InClass))
 		{
 			const FProperty* ReturnProperty = Function->GetReturnProperty();
 			FParamTypeHandle TypeHandle = FParamTypeHandle::FromProperty(ReturnProperty);
diff --git a/Engine/Plugins/Experimental/AnimNext/Source/AnimNext/Private/Param/ObjectProxy.cpp b/Engine/Plugins/Experimental/AnimNext/Source/AnimNext/Private/Param/ObjectProxy.cpp
--- a/Engine/Plugins/Experimental/AnimNext/Source/AnimNext/Private/Param/ObjectProxy.cpp
+++ b/Engine/Plugins/Experimental/AnimNext/Source/AnimNext/Private/Param/ObjectProxy.cpp
@@ -7,6 +7,13 @@
 namespace UE::AnimNext
 {
 
+FObjectProxy::FObjectProxy(const UObject* InObject, const TSharedRef<FClassProxy>& InClassProxy)
+	: Object(InObject)
+	, ClassProxy(InClassProxy)
+	, RootParameterName(NAME_None)
+{
+}
+
 FObjectProxy::FObjectProxy(const UObject* InObject, FStringView InObjectLocatorPath, const TSharedRef<FClassProxy>& InClassProxy)
 	: Object(InObject)
 	, ClassProxy(InClassProxy)
@@ -17,13 +24,22 @@ FObjectProxy::FObjectProxy(const UObject* InObject, FStringView InObjectLocatorP
 	ParameterCache.SetValueObject(RootParameterName, InObject);
 }
 
+FName FObjectProxy::GetInstanceId() const
+{
+	return RootParameterName;
+}
+
 void FObjectProxy::Update(float DeltaTime)
 {
 	const UPropertyBag* PropertyBag = ParameterCache.GetPropertyBagStruct();
 	TConstArrayView<FPropertyBagPropertyDesc> PropertyDescs = PropertyBag->GetPropertyDescs();
 	uint8* StructData = ParameterCache.GetMutableValue().GetMemory();
 	UObject* NonConstObject = const_cast<UObject*>(Object.Get());
-	PropertyDescs[0].CachedProperty->SetValue_InContainer(StructData, &Object);
+	if(RootParameterName != NAME_None)
+	{
+		PropertyDescs[0].CachedProperty->SetValue_InContainer(StructData, &Object);
+	}
+
 	if(Object != nullptr)
 	{
 		for(int32 ParameterIndex = 0; ParameterIndex < ParametersToUpdate.Num(); ++ParameterIndex)
@@ -59,12 +75,7 @@ void FObjectProxy::Update(float DeltaTime)
 					UFunction* Function = ParameterToUpdate.GetFunction();
 					checkSlow(Function);
 
-					const FObjectProperty* ObjectProperty = CastFieldChecked<FObjectProperty>(PropertyDescs[0].CachedProperty);
-					check(Object->GetClass()->IsChildOf(ObjectProperty->PropertyClass));
-					UObject** ObjectBuffer = ObjectProperty->ContainerPtrToValuePtr<UObject*>(StructData);
-					*ObjectBuffer = NonConstObject;
-
-					FFrame Stack(NonConstObject, Function, ObjectBuffer, nullptr, Function->ChildProperties);
+					FFrame Stack(NonConstObject, Function, &NonConstObject, nullptr, Function->ChildProperties);
 					Function->Invoke(NonConstObject, Stack, ResultBuffer);
 					break;
 				}
@@ -112,6 +123,42 @@ void FObjectProxy::RequestParameterCache(TConstArrayView<FName> InParameterNames
 	LayerHandle = FParamStack::MakeReferenceLayer(RootParameterName, ParameterCache);
 }
 
+void FObjectProxy::RequestParameterCacheAlias(TConstArrayView<TTuple<FName, FName>> InParameterNamePairs)
+{
+	using namespace UE::AnimNext;
+
+	const int32 NumExistingProperties = ParameterCache.GetNumPropertiesInBag();
+
+	TArray<FPropertyBagPropertyDesc> PropertyDescsToAdd;
+	PropertyDescsToAdd.Reserve(InParameterNamePairs.Num());
+
+	for(const TTuple<FName, FName>& ParameterNamePair : InParameterNamePairs)
+	{
+		if(!ParameterNameMap.Contains(ParameterNamePair.Get<0>()))
+		{
+			if(const int32* ParameterIndexPtr = ClassProxy->ParameterNameMap.Find(ParameterNamePair.Get<0>()))
+			{
+				const FClassProxyParameter& ClassProxyParameter = ClassProxy->Parameters[*ParameterIndexPtr];
+
+				FAnimNextObjectProxyParameter& NewParameterToUpdate = ParametersToUpdate.AddDefaulted_GetRef();
+				NewParameterToUpdate.AccessType = ClassProxyParameter.AccessType;
+				NewParameterToUpdate.Function = ClassProxyParameter.Function.Get();
+				NewParameterToUpdate.Property = ClassProxyParameter.Property.Get();
+
+				int32 ValueParamIndex = PropertyDescsToAdd.Emplace(ParameterNamePair.Get<1>(), ClassProxyParameter.Type.GetContainerType(), ClassProxyParameter.Type.GetValueType(), ClassProxyParameter.Type.GetValueTypeObject());
+				NewParameterToUpdate.ValueParamIndex = NumExistingProperties + ValueParamIndex;
+				ParameterNameMap.Add(ParameterNamePair.Get<0>(), NewParameterToUpdate.ValueParamIndex);
+			}
+		}
+	}
+
+	// Update parameter bag struct
+	ParameterCache.AddProperties(PropertyDescsToAdd);
+
+	// Recreate layer handle as layout has changed
+	LayerHandle = FParamStack::MakeReferenceLayer(RootParameterName, ParameterCache);
+}
+
 void FObjectProxy::AddReferencedObjects(FReferenceCollector& Collector)
 {
 	Collector.AddReferencedObject(Object);
diff --git a/Engine/Plugins/Experimental/AnimNext/Source/AnimNext/Private/Param/ObjectProxy.h b/Engine/Plugins/Experimental/AnimNext/Source/AnimNext/Private/Param/ObjectProxy.h
--- a/Engine/Plugins/Experimental/AnimNext/Source/AnimNext/Private/Param/ObjectProxy.h
+++ b/Engine/Plugins/Experimental/AnimNext/Source/AnimNext/Private/Param/ObjectProxy.h
@@ -66,9 +66,12 @@ struct FObjectProxy : public IParameterSource
 {
 	FObjectProxy() = delete;
 
+	FObjectProxy(const UObject* InObject, const TSharedRef<FClassProxy>& InClassProxy);
+
 	FObjectProxy(const UObject* InObject, FStringView InObjectLocatorPath, const TSharedRef<FClassProxy>& InClassProxy);
 
 	// IParameterSource interface
+	virtual FName GetInstanceId() const override;
 	virtual void Update(float DeltaTime) override;
 	virtual const FParamStackLayerHandle& GetLayerHandle() const override { return LayerHandle; }
 	virtual void AddReferencedObjects(FReferenceCollector& Collector) override;
@@ -76,6 +79,9 @@ struct FObjectProxy : public IParameterSource
 	// Adds a set of parameters to cache each time the layer is updated
 	void RequestParameterCache(TConstArrayView<FName> InParameterNames);
 
+	// Adds a set of parameters, aliased to the named provided (tuple of ActualName -> AliasName) to cache each time the layer is updated
+	void RequestParameterCacheAlias(TConstArrayView<TTuple<FName, FName>> InParameterNamePairs);
+
 	// The object that this proxy wraps
 	TObjectPtr<const UObject> Object;
 
diff --git a/Engine/Plugins/Experimental/AnimNext/Source/AnimNext/Private/Param/ParamHelpers.cpp b/Engine/Plugins/Experimental/AnimNext/Source/AnimNext/Private/Param/ParamHelpers.cpp
--- a/Engine/Plugins/Experimental/AnimNext/Source/AnimNext/Private/Param/ParamHelpers.cpp
+++ b/Engine/Plugins/Experimental/AnimNext/Source/AnimNext/Private/Param/ParamHelpers.cpp
@@ -317,7 +317,7 @@ void FParamHelpers::Copy(const FParamTypeHandle& InTypeHandle, TConstArrayView<u
 		break;
 	case FParamTypeHandle::EParamType::Object:
 	case FParamTypeHandle::EParamType::CharacterMovementComponent:
-	case FParamTypeHandle::EParamType::AnimNextMeshComponent:
+	case FParamTypeHandle::EParamType::SkeletalMeshComponent:
 	case FParamTypeHandle::EParamType::AnimSequence:
 		SimpleCopy(sizeof(UObject*), alignof(UObject*));
 		break;
@@ -517,7 +517,7 @@ void FParamHelpers::Destroy(const FParamTypeHandle& InTypeHandle, TArrayView<uin
 	case FParamTypeHandle::EParamType::Transform:
 	case FParamTypeHandle::EParamType::Object:
 	case FParamTypeHandle::EParamType::CharacterMovementComponent:
-	case FParamTypeHandle::EParamType::AnimNextMeshComponent:
+	case FParamTypeHandle::EParamType::SkeletalMeshComponent:
 	case FParamTypeHandle::EParamType::AnimSequence:
 	case FParamTypeHandle::EParamType::AnimNextGraphLODPose:
 		break;
diff --git a/Engine/Plugins/Experimental/AnimNext/Source/AnimNext/Private/Param/ParamType.cpp b/Engine/Plugins/Experimental/AnimNext/Source/AnimNext/Private/Param/ParamType.cpp
--- a/Engine/Plugins/Experimental/AnimNext/Source/AnimNext/Private/Param/ParamType.cpp
+++ b/Engine/Plugins/Experimental/AnimNext/Source/AnimNext/Private/Param/ParamType.cpp
@@ -2,7 +2,6 @@
 
 #include "Param/ParamType.h"
 
-#include "Component/AnimNextMeshComponent.h"
 #include "Graph/AnimNext_LODPose.h"
 #include "Param/ParamTypeHandle.h"
 #include "Misc/StringBuilder.h"
@@ -312,9 +311,9 @@ UE::AnimNext::FParamTypeHandle FAnimNextParamType::GetHandle() const
 					Handle.SetParameterType(FParamTypeHandle::EParamType::CharacterMovementComponent);
 					break;
 				}
-				else if (Class == UAnimNextMeshComponent::StaticClass())
+				else if (Class == USkeletalMeshComponent::StaticClass())
 				{
-					Handle.SetParameterType(FParamTypeHandle::EParamType::AnimNextMeshComponent);
+					Handle.SetParameterType(FParamTypeHandle::EParamType::SkeletalMeshComponent);
 					break;
 				}
 				else if (Class == UAnimSequence::StaticClass())
diff --git a/Engine/Plugins/Experimental/AnimNext/Source/AnimNext/Private/Param/ParamTypeHandle.cpp b/Engine/Plugins/Experimental/AnimNext/Source/AnimNext/Private/Param/ParamTypeHandle.cpp
--- a/Engine/Plugins/Experimental/AnimNext/Source/AnimNext/Private/Param/ParamTypeHandle.cpp
+++ b/Engine/Plugins/Experimental/AnimNext/Source/AnimNext/Private/Param/ParamTypeHandle.cpp
@@ -6,7 +6,6 @@
 #include "UObject/Class.h"
 #include "UObject/TextProperty.h"
 #include "AnimNextStats.h"
-#include "Component/AnimNextMeshComponent.h"
 #include "GameFramework/CharacterMovementComponent.h"
 #include "Components/SkeletalMeshComponent.h"
 #include "Animation/AnimSequence.h"
@@ -234,9 +233,9 @@ FParamTypeHandle FParamTypeHandle::FromPropertyBagPropertyDesc(const FPropertyBa
 						Handle.SetParameterType(EParamType::CharacterMovementComponent);
 						break;
 					}
-					else if (Class == UAnimNextMeshComponent::StaticClass())
+					else if (Class == USkeletalMeshComponent::StaticClass())
 					{
-						Handle.SetParameterType(EParamType::AnimNextMeshComponent);
+						Handle.SetParameterType(EParamType::SkeletalMeshComponent);
 						break;
 					}
 					else if (Class == UAnimSequence::StaticClass())
@@ -405,9 +404,9 @@ FParamTypeHandle FParamTypeHandle::FromProperty(const FProperty* InProperty)
 		{
 			Handle.SetParameterType(EParamType::CharacterMovementComponent);
 		}
-		else if (Class == UAnimNextMeshComponent::StaticClass())
+		else if (Class == USkeletalMeshComponent::StaticClass())
 		{
-			Handle.SetParameterType(EParamType::AnimNextMeshComponent);
+			Handle.SetParameterType(EParamType::SkeletalMeshComponent);
 		}
 	}
 
@@ -576,9 +575,9 @@ FParamTypeHandle FParamTypeHandle::FromObject(const UObject* InObject)
 			{
 				Handle.SetParameterType(EParamType::CharacterMovementComponent);
 			}
-			else if (ObjectClass == UAnimNextMeshComponent::StaticClass())
+			else if (ObjectClass == USkeletalMeshComponent::StaticClass())
 			{
-				Handle.SetParameterType(EParamType::AnimNextMeshComponent);
+				Handle.SetParameterType(EParamType::SkeletalMeshComponent);
 			}
 			else
 			{
@@ -665,8 +664,8 @@ FAnimNextParamType FParamTypeHandle::GetType() const
 		ParameterType.ValueType = FAnimNextParamType::EValueType::Object;
 		ParameterType.ContainerType = FAnimNextParamType::EContainerType::None;
 		break;
-	case EParamType::AnimNextMeshComponent:
-		ParameterType.ValueTypeObject = UAnimNextMeshComponent::StaticClass();
+	case EParamType::SkeletalMeshComponent:
+		ParameterType.ValueTypeObject = USkeletalMeshComponent::StaticClass();
 		ParameterType.ValueType = FAnimNextParamType::EValueType::Object;
 		ParameterType.ContainerType = FAnimNextParamType::EContainerType::None;
 		break;
@@ -734,7 +733,7 @@ size_t FParamTypeHandle::GetSize() const
 		return sizeof(FTransform);
 	case EParamType::Object:
 	case EParamType::CharacterMovementComponent:
-	case EParamType::AnimNextMeshComponent:
+	case EParamType::SkeletalMeshComponent:
 	case EParamType::AnimSequence:
 		return sizeof(UObject*);
 	case EParamType::AnimNextGraphLODPose:
@@ -786,7 +785,7 @@ size_t FParamTypeHandle::GetAlignment() const
 		return alignof(FTransform);
 	case EParamType::Object:
 	case EParamType::CharacterMovementComponent:
-	case EParamType::AnimNextMeshComponent:
+	case EParamType::SkeletalMeshComponent:
 	case EParamType::AnimSequence:
 		return alignof(UObject*);
 	case EParamType::AnimNextGraphLODPose:
diff --git a/Engine/Plugins/Experimental/AnimNext/Source/AnimNext/Private/Param/ParamUtils.cpp b/Engine/Plugins/Experimental/AnimNext/Source/AnimNext/Private/Param/ParamUtils.cpp
--- a/Engine/Plugins/Experimental/AnimNext/Source/AnimNext/Private/Param/ParamUtils.cpp
+++ b/Engine/Plugins/Experimental/AnimNext/Source/AnimNext/Private/Param/ParamUtils.cpp
@@ -3,10 +3,9 @@
 #include "Param/ParamUtils.h"
 
 #include "UniversalObjectLocator.h"
-#include "Component/AnimNextMeshComponent.h"
+#include "Components/SkeletalMeshComponent.h"
 #include "GameFramework/CharacterMovementComponent.h"
 #include "Animation/AnimSequence.h"
-#include "Param/AnimNextClassExtensionLibrary.h"
 #include "Param/ParamType.h"
 #include "Param/ParamTypeHandle.h"
 #include "Param/ParamCompatibility.h"
@@ -215,15 +214,15 @@ FParamCompatibility FParamUtils::GetCompatibility(const FParamTypeHandle& InLHS,
 			}
 		}
 		break;
-	case FParamTypeHandle::EParamType::AnimNextMeshComponent:
+	case FParamTypeHandle::EParamType::SkeletalMeshComponent:
 		switch (InLHS.GetParameterType())
 		{
 		case FParamTypeHandle::EParamType::Object:
 			return EParamCompatibility::Compatible_Cast;
-		case FParamTypeHandle::EParamType::AnimNextMeshComponent:
+		case FParamTypeHandle::EParamType::SkeletalMeshComponent:
 			return EParamCompatibility::Compatible_Equal;
 		case FParamTypeHandle::EParamType::Custom:
-			if(CheckClassCastToCustom(InLHS, UAnimNextMeshComponent::StaticClass()))
+			if(CheckClassCastToCustom(InLHS, USkeletalMeshComponent::StaticClass()))
 			{
 				return EParamCompatibility::Compatible_Cast;
 			}
@@ -272,8 +271,8 @@ FParamCompatibility FParamUtils::GetCompatibility(const FParamTypeHandle& InLHS,
 				return EParamCompatibility::Compatible_Cast;
 			}
 			break;
-		case FParamTypeHandle::EParamType::AnimNextMeshComponent:
-			if(CheckClassCastFromCustom(UAnimNextMeshComponent::StaticClass(), InRHS))
+		case FParamTypeHandle::EParamType::SkeletalMeshComponent:
+			if(CheckClassCastFromCustom(USkeletalMeshComponent::StaticClass(), InRHS))
 			{
 				return EParamCompatibility::Compatible_Cast;
 			}
@@ -409,14 +408,11 @@ FParamCompatibility FParamUtils::GetCompatibility(const FAnimNextParamType& InLH
 	return GetCompatibility(InLHS.GetHandle(), InRHS.GetHandle());
 }
 
-bool FParamUtils::CanUseFunction(const UFunction* InFunction)
+bool FParamUtils::CanUseFunction(const UFunction* InFunction, const UClass* InExpectedClass)
 {
 	UClass* FunctionClass = InFunction->GetOuterUClass();
-	if(FunctionClass->IsChildOf(UAnimNextClassExtensionLibrary::StaticClass()))
+	if(FunctionClass->IsChildOf(UBlueprintFunctionLibrary::StaticClass()))
 	{
-		UAnimNextClassExtensionLibrary* CDO = FunctionClass->GetDefaultObject<UAnimNextClassExtensionLibrary>();
-		const UClass* ExtendedClass = CDO->GetSupportedClass();
-		
 		// Check 'hoisted' functions on BPFLs
 		if(!InFunction->HasAllFunctionFlags(FUNC_BlueprintCallable | FUNC_Static | FUNC_Native | FUNC_Public))
 		{
@@ -439,11 +435,28 @@ bool FParamUtils::CanUseFunction(const UFunction* InFunction)
 				{ 
 					return false;
 				}
-				
-				if(!ObjectProperty->PropertyClass->IsChildOf(ExtendedClass))
+
+				// TODO: Class checks have to be editor only right now until Verse moves to using UHT (and UHT can understand verse classes)
+				// For now we need to use metadata to distinguish types
+#if WITH_EDITORONLY_DATA
+				if(InExpectedClass != nullptr)
 				{
-					return false;
+					if(!InExpectedClass->IsChildOf(ObjectProperty->PropertyClass))
+					{
+						const FString& AllowedClassMeta = ObjectProperty->GetMetaData("AllowedClass");
+						if(AllowedClassMeta.Len() == 0)
+						{
+							return false;
+						}
+
+						const UClass* AllowedClass = FindObject<UClass>(nullptr, *AllowedClassMeta);
+						if(AllowedClass == nullptr || !AllowedClass->IsChildOf(InExpectedClass))
+						{
+							return false;
+						}
+					}
 				}
+#endif
 			}
 
 			// Check return value
diff --git a/Engine/Plugins/Experimental/AnimNext/Source/AnimNext/Private/Param/ParametersProxy.cpp b/Engine/Plugins/Experimental/AnimNext/Source/AnimNext/Private/Param/ParametersProxy.cpp
--- a/Engine/Plugins/Experimental/AnimNext/Source/AnimNext/Private/Param/ParametersProxy.cpp
+++ b/Engine/Plugins/Experimental/AnimNext/Source/AnimNext/Private/Param/ParametersProxy.cpp
@@ -1,34 +1,75 @@
 // Copyright Epic Games, Inc. All Rights Reserved.
 
 #include "Param/ParametersProxy.h"
+
+#include "ParamHelpers.h"
 #include "Graph/AnimNextGraph.h"
 #include "Param/ParamStack.h"
 
 namespace UE::AnimNext
 {
 
-FParametersProxy::FParametersProxy(UAnimNextGraph* InGraph)
+FParametersProxy::FParametersProxy(const UAnimNextGraph* InGraph)
 	: Graph(InGraph)
-	, PropertyBag(InGraph->PropertyBag)
+	, PropertyBag(InGraph->DefaultState.State)
 	, LayerHandle(FParamStack::MakeReferenceLayer(NAME_None, PropertyBag))
 {
 	check(Graph);
+	UpdateCachedExternalParamData();
 }
 
 void FParametersProxy::Update(float DeltaTime)
 {
 #if WITH_EDITOR	// Layout should only be changing in editor
 	const FInstancedPropertyBag* HandlePropertyBag = LayerHandle.As<FInstancedPropertyBag>();
-	if(HandlePropertyBag == nullptr || HandlePropertyBag->GetPropertyBagStruct() != Graph->PropertyBag.GetPropertyBagStruct())
+	if(HandlePropertyBag == nullptr || HandlePropertyBag->GetPropertyBagStruct() != Graph->DefaultState.State.GetPropertyBagStruct())
 	{
-		PropertyBag = Graph->PropertyBag;
+		PropertyBag = Graph->DefaultState.State;
 		LayerHandle = FParamStack::MakeReferenceLayer(NAME_None, PropertyBag);
+		UpdateCachedExternalParamData();
 	}
 #endif
 
+	// First of all we update public state from external sources, if any.
+	// This is to ensure that when running user-defined layer update logic in UpdateLayer()
+	// that the visible state is consistent with the external source.
+	if(ExternalParamData.Num() > 0)
+	{
+		FParamStack& ParamStack = FParamStack::Get();
+		for(const FExternalParamData& ExternalParam : ExternalParamData)
+		{
+			TConstArrayView<uint8> Data;
+			if(ParamStack.GetParamData(ExternalParam.ParamId, ExternalParam.TypeHandle, Data).IsSuccessful())
+			{
+				FParamHelpers::Copy(ExternalParam.TypeHandle, Data, ExternalParam.InternalData);
+			}
+		}
+	}
+
+	// Next we update the layer
 	Graph->UpdateLayer(LayerHandle, DeltaTime);
 }
 
+void FParametersProxy::UpdateCachedExternalParamData()
+{
+	ExternalParamData.Reset();
+
+	if(Graph->DefaultState.PublicParameterStartIndex != INDEX_NONE)
+	{
+		TConstArrayView<FPropertyBagPropertyDesc> Descs = Graph->DefaultState.State.GetPropertyBagStruct()->GetPropertyDescs();
+		const int32 NumProperties = Descs.Num();
+		ExternalParamData.Reserve(NumProperties - Graph->DefaultState.PublicParameterStartIndex);
+		for(int32 PropertyIndex = Graph->DefaultState.PublicParameterStartIndex; PropertyIndex < NumProperties; ++PropertyIndex)
+		{
+			const FPropertyBagPropertyDesc& Desc = Descs[PropertyIndex];
+			check(Desc.CachedProperty);
+			uint8* ValuePtr = Desc.CachedProperty->ContainerPtrToValuePtr<uint8>(PropertyBag.GetMutableValue().GetMemory());
+			TArrayView<uint8> InternalData(ValuePtr, Desc.CachedProperty->GetSize());
+			ExternalParamData.Emplace(FParamId(Desc.Name), FAnimNextParamType(Desc.ValueType, Desc.ContainerTypes.GetFirstContainerType(), Desc.ValueTypeObject).GetHandle(), InternalData);
+		}
+	}
+}
+
 void FParametersProxy::AddReferencedObjects(FReferenceCollector& Collector)
 {
 	Collector.AddReferencedObject(Graph);
diff --git a/Engine/Plugins/Experimental/AnimNext/Source/AnimNext/Private/Param/ParametersProxy.h b/Engine/Plugins/Experimental/AnimNext/Source/AnimNext/Private/Param/ParametersProxy.h
--- a/Engine/Plugins/Experimental/AnimNext/Source/AnimNext/Private/Param/ParametersProxy.h
+++ b/Engine/Plugins/Experimental/AnimNext/Source/AnimNext/Private/Param/ParametersProxy.h
@@ -16,21 +16,48 @@ struct FParametersProxy : public IParameterSource
 {
 	FParametersProxy() = delete;
 
-	FParametersProxy(UAnimNextGraph* InGraph);
+	FParametersProxy(const UAnimNextGraph* InGraph);
 
 	// IParameterSource interface
+	virtual FName GetInstanceId() const override { return NAME_None; }
 	virtual void Update(float DeltaTime) override;
 	virtual const FParamStackLayerHandle& GetLayerHandle() const override { return LayerHandle; }
 	virtual void AddReferencedObjects(FReferenceCollector& Collector) override;
 
+	// Update the external param copy data we use to maintain the visible state of external parameters
+	void UpdateCachedExternalParamData();
+
 	// The object that this proxy wraps
-	TObjectPtr<UAnimNextGraph> Graph;
+	TObjectPtr<const UAnimNextGraph> Graph;
 
 	// Copy of the parameter data
 	FInstancedPropertyBag PropertyBag;
 
 	// Layer handle - must be updated if PropertyBag changes layout
 	FParamStackLayerHandle LayerHandle;
+
+	// Data for an external parameter copy
+	struct FExternalParamData
+	{
+		FExternalParamData() = default;
+		
+		FExternalParamData(const FParamId& InParamId, const FParamTypeHandle& InTypeHandle, TArrayView<uint8> InInternalData)
+			: ParamId(InParamId)
+			, TypeHandle(InTypeHandle)
+			, InternalData(InInternalData)
+		{
+		}
+
+		// ID for the param
+		FParamId ParamId;
+		// Type of the data
+		FParamTypeHandle TypeHandle;
+		// The data held in the PropertyBag for this param
+		TArrayView<uint8> InternalData;
+	};
+
+	// All external params we will be copying pre-update
+	TArray<FExternalParamData> ExternalParamData;
 };
 
 }
\ No newline at end of file
diff --git a/Engine/Plugins/Experimental/AnimNext/Source/AnimNext/Private/Param/PropertyBagProxy.cpp b/Engine/Plugins/Experimental/AnimNext/Source/AnimNext/Private/Param/PropertyBagProxy.cpp
--- a/Engine/Plugins/Experimental/AnimNext/Source/AnimNext/Private/Param/PropertyBagProxy.cpp
+++ b/Engine/Plugins/Experimental/AnimNext/Source/AnimNext/Private/Param/PropertyBagProxy.cpp
@@ -1,4 +1,4 @@
-﻿// Copyright Epic Games, Inc. All Rights Reserved.
+// Copyright Epic Games, Inc. All Rights Reserved.
 
 #include "Param/PropertyBagProxy.h"
 #include "Param/ParamStack.h"
@@ -6,43 +6,32 @@
 namespace UE::AnimNext
 {
 
-FPropertyBagProxy::FPropertyBagProxy()
+FPropertyBagProxy::FPropertyBagProxy(FName InInstanceId)
+	: InstanceId(InInstanceId)
 {
 	LayerHandle = FParamStack::MakeReferenceLayer(NAME_None, PropertyBag);
 }
 
-void FPropertyBagProxy::AddReferencedObjects(FReferenceCollector& Collector)
+FPropertyBagProxy::FPropertyBagProxy(FName InInstanceId, FInstancedPropertyBag&& InPropertyBag)
+	: PropertyBag(MoveTemp(InPropertyBag))
+	, InstanceId(InInstanceId)
 {
-	PropertyBag.AddStructReferencedObjects(Collector);
+	LayerHandle = FParamStack::MakeReferenceLayer(NAME_None, PropertyBag);
 }
 
-void FPropertyBagProxy::AddPropertyAndValue(FName InName, const FProperty* InProperty, const void* InContainerPtr)
+FName FPropertyBagProxy::GetInstanceId() const
 {
-	PropertyBag.AddProperty(InName, InProperty);
-	PropertyBag.SetValue(InName, InProperty, InContainerPtr);
+	return InstanceId;
+}
 
-	// Recreate the layer handle as the bag layout has changed
-	LayerHandle = FParamStack::MakeReferenceLayer(NAME_None, PropertyBag);
+void FPropertyBagProxy::AddReferencedObjects(FReferenceCollector& Collector)
+{
+	PropertyBag.AddStructReferencedObjects(Collector);
 }
 
-void FPropertyBagProxy::AddPropertiesAndValues(TConstArrayView<FPropertyAndValue> InPropertiesAndValues)
+void FPropertyBagProxy::ReplaceAllParameters(TConstArrayView<FPropertyBagPropertyDesc> InDescs, TConstArrayView<TConstArrayView<uint8>> InValues)
 {
-	TArray<FPropertyBagPropertyDesc, TInlineAllocator<16>> NewDescs;
-	NewDescs.Reserve(InPropertiesAndValues.Num());
-
-	// Add all new properties
-	for(const FPropertyAndValue& PropertyAndValue : InPropertiesAndValues)
-	{
-		NewDescs.Emplace(PropertyAndValue.Name, PropertyAndValue.Property);
-	}
-
-	PropertyBag.AddProperties(NewDescs);
-
-	// Set each value - note that SetValue uses a linear search internally so this is quite slow
-	for(const FPropertyAndValue& PropertyAndValue : InPropertiesAndValues)
-	{
-		PropertyBag.SetValue(PropertyAndValue.Name, PropertyAndValue.Property, PropertyAndValue.ContainerPtr);
-	}
+	PropertyBag.ReplaceAllPropertiesAndValues(InDescs, InValues);
 
 	// Recreate the layer handle as the bag layout has changed
 	LayerHandle = FParamStack::MakeReferenceLayer(NAME_None, PropertyBag);
diff --git a/Engine/Plugins/Experimental/AnimNext/Source/AnimNext/Private/Param/PropertyBagProxy.h b/Engine/Plugins/Experimental/AnimNext/Source/AnimNext/Private/Param/PropertyBagProxy.h
--- a/Engine/Plugins/Experimental/AnimNext/Source/AnimNext/Private/Param/PropertyBagProxy.h
+++ b/Engine/Plugins/Experimental/AnimNext/Source/AnimNext/Private/Param/PropertyBagProxy.h
@@ -4,6 +4,7 @@
 
 #include "Param/IParameterSource.h"
 #include "PropertyBag.h"
+#include "IAnimNextModule.h"
 
 namespace UE::AnimNext
 {
@@ -11,26 +12,18 @@ namespace UE::AnimNext
 // Proxy struct used to hold data in a property bag
 struct FPropertyBagProxy : public IParameterSource
 {
-	FPropertyBagProxy();
+	FPropertyBagProxy(FName InInstanceId);
+
+	FPropertyBagProxy(FName InInstanceId, FInstancedPropertyBag&& InPropertyBag);
 
 	// IParameterSource interface
+	virtual FName GetInstanceId() const override;
 	virtual void Update(float DeltaTime) override {};
 	virtual const FParamStackLayerHandle& GetLayerHandle() const override { return LayerHandle; }
 	virtual void AddReferencedObjects(FReferenceCollector& Collector) override;
 
-	// Adds (or replaces) a property with the specified name and sets its value using the specified container
-	void AddPropertyAndValue(FName InName, const FProperty* InProperty, const void* InContainerPtr);
-
-	// Arg to AddPropertiesAndValues
-	struct FPropertyAndValue
-	{
-		FName Name;
-		const FProperty* Property = nullptr;
-		const void* ContainerPtr = nullptr;
-	};
-
-	// Adds (or replaces) a group of properties with the specified names and sets their values using the specified containers
-	void AddPropertiesAndValues(TConstArrayView<FPropertyAndValue> InPropertiesAndValues);
+	// Replaces all properties with the specified names and sets their values using the specified containers
+	void ReplaceAllParameters(TConstArrayView<FPropertyBagPropertyDesc> InDescs, TConstArrayView<TConstArrayView<uint8>> InValues);
 
 private:
 	// Property bag that this wraps
@@ -38,6 +31,9 @@ private:
 
 	// Layer handle - must be updated if PropertyBag is changed
 	FParamStackLayerHandle LayerHandle;
+
+	// Instance ID provided on construction
+	FName InstanceId;
 };
 
 }
\ No newline at end of file
diff --git a/Engine/Plugins/Experimental/AnimNext/Source/AnimNext/Private/Scheduler/AnimNextScheduleGraphTask.cpp b/Engine/Plugins/Experimental/AnimNext/Source/AnimNext/Private/Scheduler/AnimNextScheduleGraphTask.cpp
--- a/Engine/Plugins/Experimental/AnimNext/Source/AnimNext/Private/Scheduler/AnimNextScheduleGraphTask.cpp
+++ b/Engine/Plugins/Experimental/AnimNext/Source/AnimNext/Private/Scheduler/AnimNextScheduleGraphTask.cpp
@@ -9,46 +9,15 @@
 #include "Graph/AnimNext_LODPose.h"
 #include "AnimNextStats.h"
 #include "Component/AnimNextComponent.h"
-#include "Component/AnimNextMeshComponent.h"
+#include "Components/SkeletalMeshComponent.h"
 #include "GameFramework/Character.h"
 #include "Logging/StructuredLog.h"
-#include "Param/AnimNextObjectCastLocatorFragment.h"
-#include "Param/AnimNextObjectFunctionLocatorFragment.h"
 #include "Param/AnimNextParam.h"
 #include "Param/AnimNextParamUniversalObjectLocator.h"
+#include "Param/ParametersProxy.h"
 
 DEFINE_STAT(STAT_AnimNext_Task_Graph);
 
-namespace UE::AnimNext
-{
-	// TODO: Currently we hard-code the ACharacter mesh component as the source of the reference pose and LOD index, but this should be a pin-input in the
-	// final schedule incarnation.
-	TInstancedStruct<FAnimNextParamUniversalObjectLocator> GetCharacterInstanceId()
-	{
-		TInstancedStruct<FAnimNextParamUniversalObjectLocator> Locator = TInstancedStruct<FAnimNextParamUniversalObjectLocator>::Make();
-		Locator.GetMutable().Locator.AddFragment<FAnimNextObjectFunctionLocatorFragment>(UAnimNextComponent::StaticClass()->FindFunctionByName("GetOwner"));
-		Locator.GetMutable().Locator.AddFragment<FAnimNextObjectCastLocatorFragment>(ACharacter::StaticClass());
-		return Locator;
-	};
-
-	FName GetCharacterInstanceIdName()
-	{
-		static FName Name(NAME_None);
-		if(Name.IsNone())
-		{
-			const TInstancedStruct<FAnimNextParamUniversalObjectLocator>& Locator = GetCharacterInstanceId();
-			Name = Locator.Get().ToName();
-		}
-		return Name;
-	};
-
-	FParamId GetMeshComponentParamId()
-	{
-		static const FParamId MeshComponentParamId("/Script/Engine.Character:Mesh", GetCharacterInstanceIdName());
-		return MeshComponentParamId;
-	}
-}
-
 UAnimNextGraph* FAnimNextScheduleGraphTask::GetGraphToRun(UE::AnimNext::FParamStack& ParamStack) const
 {
 	UAnimNextGraph* GraphToRun = Graph;
@@ -142,22 +111,17 @@ void FAnimNextScheduleGraphTask::RunGraph(const UE::AnimNext::FScheduleContext&
 		}
 	}
 
-	const TObjectPtr<USkeletalMeshComponent>* ComponentPtr = ParamStack.GetParamPtr<TObjectPtr<USkeletalMeshComponent>>(GetMeshComponentParamId());
-	if(ComponentPtr == nullptr)
+	const FAnimNextGraphReferencePose* GraphRefPosePtr = ParamStack.GetParamPtr<FAnimNextGraphReferencePose>(ReferencePose.GetParamId());
+	if (GraphRefPosePtr == nullptr || !GraphRefPosePtr->ReferencePose.IsValid())
 	{
 		return;
 	}
 
-	TObjectPtr<UAnimNextMeshComponent> Component = Cast<UAnimNextMeshComponent>(*ComponentPtr);
-	if(Component == nullptr)
-	{
-		return;
-	}
-	
-	const FAnimNextGraphReferencePose GraphReferencePose = Component->GetReferencePose();
-	if(!GraphReferencePose.ReferencePose.IsValid())
+	int32 LODIndex = 0;
+	const int32* LODPtr = ParamStack.GetParamPtr<int32>(LOD.GetParamId());
+	if (LODPtr != nullptr)
 	{
-		return;
+		LODIndex = *LODPtr;
 	}
 
 	// Check and allocate remapped term layer
@@ -184,10 +148,8 @@ void FAnimNextScheduleGraphTask::RunGraph(const UE::AnimNext::FScheduleContext&
 	{
 		return;
 	}
-	
-	const UE::AnimNext::FReferencePose& RefPose = GraphReferencePose.ReferencePose.GetRef<UE::AnimNext::FReferencePose>();
 
-	const int32 LODIndex = Component->GetPredictedLODLevel();
+	const UE::AnimNext::FReferencePose& RefPose = GraphRefPosePtr->ReferencePose.GetRef<UE::AnimNext::FReferencePose>();
 
 	// Create or update our result pose
 	// TODO: Currently forcing additive flag to false here
@@ -209,19 +171,3 @@ void FAnimNextScheduleGraphTask::RunGraph(const UE::AnimNext::FScheduleContext&
 	IAnimNextModule::Get().UpdateGraph(GraphCache.GraphInstanceData, InContext.GetDeltaTime());
 	IAnimNextModule::Get().EvaluateGraph(GraphCache.GraphInstanceData, RefPose, LODIndex, OutputPose->LODPose);
 }
-
-#if WITH_EDITORONLY_DATA 
-
-TArray<FAnimNextEditorParam> FAnimNextScheduleGraphTask::GetRequiredParametersInternal()
-{
-	using namespace UE::AnimNext;
-
-	const TArray<FAnimNextEditorParam> Params =
-	{
-		FAnimNextEditorParam(GetMeshComponentParamId().GetName(), FAnimNextParamType::GetType<TObjectPtr<UAnimNextMeshComponent>>(), GetCharacterInstanceId()),
-	};
-
-	return Params;
-}
-
-#endif
\ No newline at end of file
diff --git a/Engine/Plugins/Experimental/AnimNext/Source/AnimNext/Private/Scheduler/AnimNextScheduleParamScopeTask.cpp b/Engine/Plugins/Experimental/AnimNext/Source/AnimNext/Private/Scheduler/AnimNextScheduleParamScopeTask.cpp
--- a/Engine/Plugins/Experimental/AnimNext/Source/AnimNext/Private/Scheduler/AnimNextScheduleParamScopeTask.cpp
+++ b/Engine/Plugins/Experimental/AnimNext/Source/AnimNext/Private/Scheduler/AnimNextScheduleParamScopeTask.cpp
@@ -24,11 +24,17 @@ void FAnimNextScheduleParamScopeEntryTask::RunParamScopeEntry(const UE::AnimNext
 
 	// Update & push user params (before scope)
 	FScheduleInstanceData::FUserScope* FoundUserScope = InstanceData.UserScopes.Find(Scope.Name);
-	if (FoundUserScope && FoundUserScope->BeforeSource.IsValid())
+	if (FoundUserScope)
 	{
-		IParameterSource& UserParameterSource = *FoundUserScope->BeforeSource.Get();
-		UserParameterSource.Update(DeltaTime);
-		ScopeCache.PushedLayers.Add(ParamStack.PushLayer(UserParameterSource.GetLayerHandle()));
+		for(const TPair<FName, TUniquePtr<IParameterSource>>& SourcePair : FoundUserScope->BeforeSources)
+		{
+			if(SourcePair.Value.IsValid())
+			{
+				IParameterSource& UserParameterSource = *SourcePair.Value.Get();
+				UserParameterSource.Update(DeltaTime);
+				ScopeCache.PushedLayers.Add(ParamStack.PushLayer(UserParameterSource.GetLayerHandle()));
+			}
+		}
 	}
 
 	// Update & push static params
@@ -41,11 +47,17 @@ void FAnimNextScheduleParamScopeEntryTask::RunParamScopeEntry(const UE::AnimNext
 	}
 
 	// Update & push user params (after scope)
-	if (FoundUserScope && FoundUserScope->AfterSource.IsValid())
+	if(FoundUserScope)
 	{
-		IParameterSource& UserParameterSource = *FoundUserScope->AfterSource.Get();
-		UserParameterSource.Update(DeltaTime);
-		ScopeCache.PushedLayers.Add(ParamStack.PushLayer(UserParameterSource.GetLayerHandle()));
+		for(const TPair<FName, TUniquePtr<IParameterSource>>& SourcePair : FoundUserScope->AfterSources)
+		{
+			if(SourcePair.Value.IsValid())
+			{
+				IParameterSource& UserParameterSource = *SourcePair.Value.Get();
+				UserParameterSource.Update(DeltaTime);
+				ScopeCache.PushedLayers.Add(ParamStack.PushLayer(UserParameterSource.GetLayerHandle()));
+			}
+		}
 	}
 }
 
diff --git a/Engine/Plugins/Experimental/AnimNext/Source/AnimNext/Private/Scheduler/AnimNextSchedulerEntry.cpp b/Engine/Plugins/Experimental/AnimNext/Source/AnimNext/Private/Scheduler/AnimNextSchedulerEntry.cpp
--- a/Engine/Plugins/Experimental/AnimNext/Source/AnimNext/Private/Scheduler/AnimNextSchedulerEntry.cpp
+++ b/Engine/Plugins/Experimental/AnimNext/Source/AnimNext/Private/Scheduler/AnimNextSchedulerEntry.cpp
@@ -215,9 +215,9 @@ void FAnimNextSchedulerEntry::Initialize()
 		}
 		else
 		{
-			// In editor preview worlds (ideally just thumbnail scenes) we run a linearized 'initial tick' to ensure we
-			// generate an output pose, as these worlds never tick
-			if(World->WorldType == EWorldType::EditorPreview)
+			// In editor worlds we run a linearized 'initial tick' to ensure we generate an initial output pose, as these worlds dont always tick
+			if( World->WorldType == EWorldType::Editor ||
+				World->WorldType == EWorldType::EditorPreview)
 			{
 				FScheduleTickFunction::RunSchedule(*this);
 			}
@@ -343,14 +343,14 @@ void FAnimNextSchedulerEntry::OnScheduleCompiled()
 	using namespace UE::AnimNext;
 
 	// Store any user-defined scopes, as the instance data will be going away
-	TUniquePtr<IParameterSource> RootUserScope = MoveTemp(Context.InstanceData->RootUserScope);
+	TArray<TUniquePtr<IParameterSource>> RootUserScopes = MoveTemp(Context.InstanceData->RootUserScopes);
 	TMap<FName, FScheduleInstanceData::FUserScope> UserScopes = MoveTemp(Context.InstanceData->UserScopes);
 
 	ResetBindingsAndInstanceData();
 	Initialize();
 
 	// Restore any user scopes to the recreated instance data
-	Context.InstanceData->RootUserScope = MoveTemp(RootUserScope);
+	Context.InstanceData->RootUserScopes = MoveTemp(RootUserScopes);
 	Context.InstanceData->UserScopes = MoveTemp(UserScopes);
 }
 #endif
\ No newline at end of file
diff --git a/Engine/Plugins/Experimental/AnimNext/Source/AnimNext/Private/Scheduler/Ports/AnimNextSchedulePort_AnimNextMeshComponentPose.cpp b/Engine/Plugins/Experimental/AnimNext/Source/AnimNext/Private/Scheduler/Ports/AnimNextSchedulePort_AnimNextMeshComponentPose.cpp
--- a/Engine/Plugins/Experimental/AnimNext/Source/AnimNext/Private/Scheduler/Ports/AnimNextSchedulePort_AnimNextMeshComponentPose.cpp
+++ b/Engine/Plugins/Experimental/AnimNext/Source/AnimNext/Private/Scheduler/Ports/AnimNextSchedulePort_AnimNextMeshComponentPose.cpp
@@ -10,13 +10,45 @@
 #include "Engine/SkeletalMesh.h"
 #include "AnimNextStats.h"
 #include "Component/AnimNextComponent.h"
-#include "Component/AnimNextMeshComponent.h"
 #include "GameFramework/Character.h"
 #include "Scheduler/AnimNextScheduleGraphTask.h"
+#include "Param/AnimNextObjectCastLocatorFragment.h"
+#include "Param/AnimNextObjectFunctionLocatorFragment.h"
 #include "Param/AnimNextParamUniversalObjectLocator.h"
+#include "Component/SkinnedMeshComponentExtensions.h"
 
 DEFINE_STAT(STAT_AnimNext_Port_SkeletalMeshComponent);
 
+namespace UE::AnimNext
+{
+	// TODO: Currently we hard-code the ACharacter mesh component as the source of the reference pose and LOD index, but this should be a pin-input in the
+	// final schedule incarnation.
+	TInstancedStruct<FAnimNextParamUniversalObjectLocator> GetCharacterInstanceId()
+	{
+		TInstancedStruct<FAnimNextParamUniversalObjectLocator> Locator = TInstancedStruct<FAnimNextParamUniversalObjectLocator>::Make();
+		Locator.GetMutable().Locator.AddFragment<FAnimNextObjectFunctionLocatorFragment>(UAnimNextComponent::StaticClass()->FindFunctionByName("GetOwner"));
+		Locator.GetMutable().Locator.AddFragment<FAnimNextObjectCastLocatorFragment>(ACharacter::StaticClass());
+		return Locator;
+	};
+
+	FName GetCharacterInstanceIdName()
+	{
+		static FName Name(NAME_None);
+		if (Name.IsNone())
+		{
+			const TInstancedStruct<FAnimNextParamUniversalObjectLocator>& Locator = GetCharacterInstanceId();
+			Name = Locator.Get().ToName();
+		}
+		return Name;
+	};
+
+	FParamId GetMeshComponentParamId()
+	{
+		static const FParamId MeshComponentParamId("/Script/Engine.Character:Mesh", GetCharacterInstanceIdName());
+		return MeshComponentParamId;
+	}
+}
+
 void UAnimNextSchedulePort_AnimNextMeshComponentPose::Run(const UE::AnimNext::FScheduleTermContext& InContext) const
 {
 	SCOPE_CYCLE_COUNTER(STAT_AnimNext_Port_SkeletalMeshComponent);
@@ -31,7 +63,7 @@ void UAnimNextSchedulePort_AnimNextMeshComponentPose::Run(const UE::AnimNext::FS
 		return;
 	}
 
-	TObjectPtr<UAnimNextMeshComponent> Component = Cast<UAnimNextMeshComponent>(*ComponentPtr);
+	TObjectPtr<USkeletalMeshComponent> Component = Cast<USkeletalMeshComponent>(*ComponentPtr);
 	if(Component == nullptr)
 	{
 		return;
@@ -47,14 +79,15 @@ void UAnimNextSchedulePort_AnimNextMeshComponentPose::Run(const UE::AnimNext::FS
 	{
 		return;
 	}
-	
+
 	USkeletalMesh* SkeletalMesh = Component->GetSkeletalMeshAsset();
 	if(SkeletalMesh == nullptr)
 	{
 		return;
 	}
 
-	const UE::AnimNext::FReferencePose& RefPose = Component->GetReferencePose().ReferencePose.GetRef<UE::AnimNext::FReferencePose>();
+	UE::AnimNext::FDataHandle RefPoseHandle = UE::AnimNext::FDataRegistry::Get()->GetOrGenerateReferencePose(Component);
+	const UE::AnimNext::FReferencePose& RefPose = RefPoseHandle.GetRef<UE::AnimNext::FReferencePose>();
 
 	FMemMark MemMark(FMemStack::Get());
 
@@ -65,7 +98,8 @@ void UAnimNextSchedulePort_AnimNextMeshComponentPose::Run(const UE::AnimNext::FS
 	FGenerationTools::RemapPose(InputPose->LODPose, LocalSpaceTransforms);
 
 	// Convert and dispatch to renderer
-	Component->CompleteAndDispatch(
+	UE::Anim::FSkinnedMeshComponentExtensions::CompleteAndDispatch(
+		Component,
 		RefPose.GetParentIndices(),
 		RefPose.GetLODBoneIndexToMeshBoneIndexMap(InputPose->LODPose.LODLevel),
 		LocalSpaceTransforms);
@@ -91,7 +125,7 @@ TConstArrayView<FAnimNextEditorParam> UAnimNextSchedulePort_AnimNextMeshComponen
 
 	if(RequiredParams.Num() == 0)
 	{
-		RequiredParams.Emplace(GetMeshComponentParamId().GetName(), FAnimNextParamType::GetType<TObjectPtr<UAnimNextMeshComponent>>(), GetCharacterInstanceId());
+		RequiredParams.Emplace(GetMeshComponentParamId().GetName(), FAnimNextParamType::GetType<TObjectPtr<USkeletalMeshComponent>>(), GetCharacterInstanceId());
 	}
 
 	return RequiredParams;
diff --git a/Engine/Plugins/Experimental/AnimNext/Source/AnimNext/Private/Scheduler/Ports/AnimNextSchedulePort_AnimNextMeshComponentPose.h b/Engine/Plugins/Experimental/AnimNext/Source/AnimNext/Private/Scheduler/Ports/AnimNextSchedulePort_AnimNextMeshComponentPose.h
--- a/Engine/Plugins/Experimental/AnimNext/Source/AnimNext/Private/Scheduler/Ports/AnimNextSchedulePort_AnimNextMeshComponentPose.h
+++ b/Engine/Plugins/Experimental/AnimNext/Source/AnimNext/Private/Scheduler/Ports/AnimNextSchedulePort_AnimNextMeshComponentPose.h
@@ -6,7 +6,7 @@
 #include "Param/AnimNextEditorParam.h"
 #include "AnimNextSchedulePort_AnimNextMeshComponentPose.generated.h"
 
-UCLASS(DisplayName = "AnimNext Mesh Component Pose")
+UCLASS(DisplayName = "Skeletal Mesh Component Pose")
 class UAnimNextSchedulePort_AnimNextMeshComponentPose : public UAnimNextSchedulePort
 {
 	GENERATED_BODY()
diff --git a/Engine/Plugins/Experimental/AnimNext/Source/AnimNext/Private/Scheduler/ScheduleInitializationContext.cpp b/Engine/Plugins/Experimental/AnimNext/Source/AnimNext/Private/Scheduler/ScheduleInitializationContext.cpp
--- a/Engine/Plugins/Experimental/AnimNext/Source/AnimNext/Private/Scheduler/ScheduleInitializationContext.cpp
+++ b/Engine/Plugins/Experimental/AnimNext/Source/AnimNext/Private/Scheduler/ScheduleInitializationContext.cpp
@@ -1,6 +1,12 @@
 // Copyright Epic Games, Inc. All Rights Reserved.
 
 #include "Scheduler/ScheduleInitializationContext.h"
+
+#include "AnimNextModule.h"
+#include "Modules/ModuleManager.h"
+#include "Param/ObjectProxy.h"
+#include "Param/ObjectProxyFactory.h"
+#include "Param/PropertyBagProxy.h"
 #include "Scheduler/ScheduleContext.h"
 
 namespace UE::AnimNext
@@ -11,10 +17,37 @@ FScheduleInitializationContext::FScheduleInitializationContext(const FScheduleCo
 {
 }
 
-void FScheduleInitializationContext::ApplyParametersToScope(FName InScope, EParameterScopeOrdering InOrdering, TUniquePtr<IParameterSource>&& InParameters) const
+void FScheduleInitializationContext::ApplyParametersToScope(FName InScope, EParameterScopeOrdering InOrdering, const UObject* InObject, TConstArrayView<TTuple<FName, FName>> InRequiredParameterAliases) const
+{
+	FModule& AnimNextModule = FModuleManager::GetModuleChecked<FModule>("AnimNext");
+	TSharedPtr<FObjectProxyFactory> ObjectProxyFactory = StaticCastSharedPtr<FObjectProxyFactory>(AnimNextModule.FindParameterSourceFactory("ObjectProxy"));
+	if(!ObjectProxyFactory.IsValid())
+	{
+		return;
+	}
+
+	TSharedRef<FClassProxy> ClassProxy = ObjectProxyFactory->FindOrCreateClassProxy(InObject->GetClass());
+	TUniquePtr<FObjectProxy> ObjectProxy = MakeUnique<FObjectProxy>(InObject, ClassProxy);
+	ObjectProxy->RequestParameterCacheAlias(InRequiredParameterAliases);
+	Context.GetInstanceData().ApplyParametersToScope(InScope, InOrdering, MoveTemp(ObjectProxy));
+}
+
+void FScheduleInitializationContext::ApplyParametersToScope(FName InScope, EParameterScopeOrdering InOrdering, const TInstancedStruct<FAnimNextParamInstanceIdentifier>& InInstanceId, const FParameterSourceContext& InContext, TConstArrayView<FName> InRequiredParameters) const
+{
+	FModule& AnimNextModule = FModuleManager::GetModuleChecked<FModule>("AnimNextModule");
+	TUniquePtr<IParameterSource> ParameterSource = AnimNextModule.CreateParameterSource(InContext, InInstanceId, InRequiredParameters);
+	Context.GetInstanceData().ApplyParametersToScope(InScope, InOrdering, MoveTemp(ParameterSource));
+}
+
+void FScheduleInitializationContext::ApplyParametersToScope(FName InScope, EParameterScopeOrdering InOrdering, FName InInstanceId, TConstArrayView<FPropertyBagPropertyDesc> InPropertyDescs, TConstArrayView<TConstArrayView<uint8>> InValues) const
 {
-	FScheduleInstanceData& InstanceData = Context.GetInstanceData();
-	InstanceData.ApplyParametersToScope(InScope, InOrdering, MoveTemp(InParameters));
+	check(InPropertyDescs.Num() == InValues.Num());
+	if(InPropertyDescs.Num() > 0)
+	{
+		TUniquePtr<FPropertyBagProxy> PropertyBagProxy = MakeUnique<FPropertyBagProxy>(InInstanceId);
+		PropertyBagProxy->ReplaceAllParameters(InPropertyDescs, InValues);
+		Context.GetInstanceData().ApplyParametersToScope(InScope, InOrdering, MoveTemp(PropertyBagProxy));
+	}
 }
 
 }
\ No newline at end of file
diff --git a/Engine/Plugins/Experimental/AnimNext/Source/AnimNext/Private/Scheduler/ScheduleInstanceData.cpp b/Engine/Plugins/Experimental/AnimNext/Source/AnimNext/Private/Scheduler/ScheduleInstanceData.cpp
--- a/Engine/Plugins/Experimental/AnimNext/Source/AnimNext/Private/Scheduler/ScheduleInstanceData.cpp
+++ b/Engine/Plugins/Experimental/AnimNext/Source/AnimNext/Private/Scheduler/ScheduleInstanceData.cpp
@@ -5,13 +5,11 @@
 #include "Scheduler/ScheduleHandle.h"
 #include "Scheduler/AnimNextSchedulerWorldSubsystem.h"
 #include "Scheduler/AnimNextSchedulerEntry.h"
-#include "Scheduler/AnimNextSchedulePort.h"
 #include "AnimNextStats.h"
-#include "IAnimNextModule.h"
+#include "AnimNextModule.h"
 #include "Modules/ModuleManager.h"
 #include "Param/IParameterSourceFactory.h"
 #include "Param/ParametersProxy.h"
-#include "Param/PropertyBagProxy.h"
 #include "Scheduler/AnimNextScheduleExternalParamTask.h"
 #include "Scheduler/ScheduleInitializationContext.h"
 
@@ -74,7 +72,7 @@ FScheduleInstanceData::FScheduleInstanceData(const FScheduleContext& InScheduleC
 	}
 
 	// Set up external parameters
-	IAnimNextModule& AnimNextModule = FModuleManager::GetModuleChecked<IAnimNextModule>("AnimNext");
+	UE::AnimNext::FModule& AnimNextModule = FModuleManager::GetModuleChecked<UE::AnimNext::FModule>("AnimNext");
 
 	FParameterSourceContext ParameterSourceContext;
 	ParameterSourceContext.Object = Entry->WeakObject.Get();
@@ -115,13 +113,19 @@ void FScheduleInstanceData::AddReferencedObjects(FReferenceCollector& Collector)
 {
 	for (const TPair<FName, FUserScope>& ParamPair : UserScopes)
 	{
-		if(ParamPair.Value.AfterSource.IsValid())
+		for(const TPair<FName, TUniquePtr<IParameterSource>>& SourcePair : ParamPair.Value.BeforeSources)
 		{
-			ParamPair.Value.AfterSource->AddReferencedObjects(Collector);
+			if(SourcePair.Value.IsValid())
+			{
+				SourcePair.Value->AddReferencedObjects(Collector);
+			}
 		}
-		if(ParamPair.Value.BeforeSource.IsValid())
+		for(const TPair<FName, TUniquePtr<IParameterSource>>& SourcePair : ParamPair.Value.AfterSources)
 		{
-			ParamPair.Value.BeforeSource->AddReferencedObjects(Collector);
+			if(SourcePair.Value.IsValid())
+			{
+				SourcePair.Value->AddReferencedObjects(Collector);
+			}
 		}
 	}
 
@@ -160,7 +164,7 @@ void FScheduleInstanceData::ApplyParametersToScope(FName InScope, EParameterScop
 {
 	if (InScope == NAME_None)
 	{
-		RootUserScope = MoveTemp(InParameters);
+		RootUserScopes.Add(MoveTemp(InParameters));
 	}
 	else // apply to specified scope
 	{
@@ -169,10 +173,10 @@ void FScheduleInstanceData::ApplyParametersToScope(FName InScope, EParameterScop
 		{
 		default:
 		case EParameterScopeOrdering::Before:
-			ScopeSource.BeforeSource = MoveTemp(InParameters);
+			ScopeSource.BeforeSources.Add(InParameters->GetInstanceId(), MoveTemp(InParameters));
 			break;
 		case EParameterScopeOrdering::After:
-			ScopeSource.AfterSource = MoveTemp(InParameters);
+			ScopeSource.AfterSources.Add(InParameters->GetInstanceId(), MoveTemp(InParameters));
 			break;
 		}
 	}
diff --git a/Engine/Plugins/Experimental/AnimNext/Source/AnimNext/Private/Scheduler/ScheduleInstanceData.h b/Engine/Plugins/Experimental/AnimNext/Source/AnimNext/Private/Scheduler/ScheduleInstanceData.h
--- a/Engine/Plugins/Experimental/AnimNext/Source/AnimNext/Private/Scheduler/ScheduleInstanceData.h
+++ b/Engine/Plugins/Experimental/AnimNext/Source/AnimNext/Private/Scheduler/ScheduleInstanceData.h
@@ -46,19 +46,19 @@ struct FScheduleInstanceData : public FGCObject
 	// Schedule entry that owns this instance
 	FAnimNextSchedulerEntry* Entry = nullptr;
 
-	// Scope for user parameters to be applied at the root of the schedule
-	TUniquePtr<IParameterSource> RootUserScope;
+	// Scopes for user parameters to be applied at the root of the schedule
+	TArray<TUniquePtr<IParameterSource>> RootUserScopes;
 
-	// Pushed layer for the root scope
-	FParamStack::FPushedLayerHandle PushedRootUserLayer;
+	// Pushed layers for the root scope
+	TArray<FParamStack::FPushedLayerHandle> PushedRootUserLayers;
 
 	struct FUserScope
 	{
-		// Layer that will be pushed before the scope, allowing the static scope to override the layer
-		TUniquePtr<IParameterSource> BeforeSource;
+		// Layers that will be pushed before the scope, allowing the static scope to override the layer
+		TMap<FName, TUniquePtr<IParameterSource>> BeforeSources;
 
-		// Layer that will be pushed after the scope, overriding the static scope
-		TUniquePtr<IParameterSource> AfterSource;
+		// Layers that will be pushed after the scope, overriding the static scope
+		TMap<FName, TUniquePtr<IParameterSource>> AfterSources;
 	};
 	
 	// Set of dynamic parameter scopes supplied by the user
diff --git a/Engine/Plugins/Experimental/AnimNext/Source/AnimNext/Private/Scheduler/ScheduleTaskContext.cpp b/Engine/Plugins/Experimental/AnimNext/Source/AnimNext/Private/Scheduler/ScheduleTaskContext.cpp
--- a/Engine/Plugins/Experimental/AnimNext/Source/AnimNext/Private/Scheduler/ScheduleTaskContext.cpp
+++ b/Engine/Plugins/Experimental/AnimNext/Source/AnimNext/Private/Scheduler/ScheduleTaskContext.cpp
@@ -1,6 +1,8 @@
 // Copyright Epic Games, Inc. All Rights Reserved.
 
 #include "Scheduler/ScheduleTaskContext.h"
+
+#include "Param/PropertyBagProxy.h"
 #include "Scheduler/ScheduleContext.h"
 
 namespace UE::AnimNext
@@ -11,10 +13,10 @@ FScheduleTaskContext::FScheduleTaskContext(const FScheduleContext& InContext)
 {
 }
 
-void FScheduleTaskContext::ApplyParametersToScope(FName InScope, EParameterScopeOrdering InOrdering, TUniquePtr<IParameterSource>&& InParameters) const
+void FScheduleTaskContext::ApplyParametersToScope(FName InScope, EParameterScopeOrdering InOrdering, FName InInstanceId, FInstancedPropertyBag&& InPropertyBag) const
 {
-	FScheduleInstanceData& InstanceData = Context.GetInstanceData();
-	InstanceData.ApplyParametersToScope(InScope, InOrdering, MoveTemp(InParameters));
+	TUniquePtr<FPropertyBagProxy> PropertyBagProxy = MakeUnique<FPropertyBagProxy>(InInstanceId, MoveTemp(InPropertyBag));
+	Context.GetInstanceData().ApplyParametersToScope(InScope, InOrdering, MoveTemp(PropertyBagProxy));
 }
 
 }
\ No newline at end of file
diff --git a/Engine/Plugins/Experimental/AnimNext/Source/AnimNext/Private/Scheduler/ScheduleTickFunction.cpp b/Engine/Plugins/Experimental/AnimNext/Source/AnimNext/Private/Scheduler/ScheduleTickFunction.cpp
--- a/Engine/Plugins/Experimental/AnimNext/Source/AnimNext/Private/Scheduler/ScheduleTickFunction.cpp
+++ b/Engine/Plugins/Experimental/AnimNext/Source/AnimNext/Private/Scheduler/ScheduleTickFunction.cpp
@@ -30,11 +30,14 @@ void FScheduleBeginTickFunction::Run(float DeltaTime)
 		Function.GetValue()(FScheduleTaskContext(Entry.Context));
 	}
 
-	// Push any user layer we have at the root
+	// Push any user layers we have at the root
 	FScheduleInstanceData& InstanceData = Entry.Context.GetInstanceData();
-	if(InstanceData.RootUserScope.IsValid())
+	for(const TUniquePtr<IParameterSource>& RootUserScope : InstanceData.RootUserScopes)
 	{
-		InstanceData.PushedRootUserLayer = InstanceData.RootParamStack->PushLayer(InstanceData.RootUserScope->GetLayerHandle());
+		if(RootUserScope.IsValid())
+		{
+			InstanceData.PushedRootUserLayers.Add(InstanceData.RootParamStack->PushLayer(RootUserScope->GetLayerHandle()));
+		}
 	}
 
 	Entry.ResolvedObject = Entry.WeakObject.Get();
@@ -53,10 +56,14 @@ void FScheduleEndTickFunction::ExecuteTick(float DeltaTime, ELevelTick TickType,
 
 void FScheduleEndTickFunction::Run()
 {
-	// Pop any user layer we have at the root
+	// Pop any user layers we have at the root
 	FScheduleInstanceData& InstanceData = Entry.Context.GetInstanceData();
-	InstanceData.RootParamStack->PopLayer(InstanceData.PushedRootUserLayer);
-
+	for (const FParamStack::FPushedLayerHandle& PushedRootUserLayer : ReverseIterate(InstanceData.PushedRootUserLayers))
+	{
+		InstanceData.RootParamStack->PopLayer(PushedRootUserLayer);
+	}
+	InstanceData.PushedRootUserLayers.Reset();
+	
 	auto RunTaskOnGameThread = [](TUniqueFunction<void(void)>&& InFunction)
 	{
 		if(IsInGameThread())
diff --git a/Engine/Plugins/Experimental/AnimNext/Source/AnimNext/Private/Scheduler/Scheduler.cpp b/Engine/Plugins/Experimental/AnimNext/Source/AnimNext/Private/Scheduler/Scheduler.cpp
--- a/Engine/Plugins/Experimental/AnimNext/Source/AnimNext/Private/Scheduler/Scheduler.cpp
+++ b/Engine/Plugins/Experimental/AnimNext/Source/AnimNext/Private/Scheduler/Scheduler.cpp
@@ -64,12 +64,14 @@ FScheduleHandle FScheduler::AcquireHandle(UObject* InObject, UAnimNextSchedule*
 	UWorld* World = InObject->GetWorld();
 	if(World == nullptr)
 	{
+		UE_LOG(LogAnimation, Warning, TEXT("FScheduler::AcquireHandle: Invalid world"));
 		return FScheduleHandle();
 	}
 
 	UAnimNextSchedulerWorldSubsystem* Subsystem = World->GetSubsystem<UAnimNextSchedulerWorldSubsystem>();
 	if(Subsystem == nullptr)
 	{
+		UE_LOG(LogAnimation, Warning, TEXT("FScheduler::AcquireHandle: Invalid world subsystem"));
 		return FScheduleHandle();
 	}
 
diff --git a/Engine/Plugins/Experimental/AnimNext/Source/AnimNext/Public/Graph/AnimNextGraphInstance.h b/Engine/Plugins/Experimental/AnimNext/Source/AnimNext/Public/Graph/AnimNextGraphInstance.h
--- a/Engine/Plugins/Experimental/AnimNext/Source/AnimNext/Public/Graph/AnimNextGraphInstance.h
+++ b/Engine/Plugins/Experimental/AnimNext/Source/AnimNext/Public/Graph/AnimNextGraphInstance.h
@@ -4,6 +4,7 @@
 
 #include "CoreMinimal.h"
 #include "HAL/CriticalSection.h"
+#include "Param/ParamStack.h"
 #include "TraitCore/TraitEvent.h"
 #include "TraitCore/TraitEventList.h"
 #include "TraitCore/TraitPtr.h"
@@ -19,6 +20,7 @@ class UAnimNextGraph;
 
 namespace UE::AnimNext
 {
+	class IParameterSource;
 	struct FExecutionContext;
 	struct FGraphInstanceComponent;
 	struct FLatentPropertyHandle;
@@ -36,7 +38,7 @@ struct ANIMNEXT_API FAnimNextGraphInstance
 	GENERATED_BODY()
 
 	// Creates an empty graph instance that doesn't reference anything
-	FAnimNextGraphInstance() = default;
+	FAnimNextGraphInstance();
 
 	// No copying, no moving
 	FAnimNextGraphInstance(const FAnimNextGraphInstance&) = delete;
@@ -129,6 +131,12 @@ private:
 	// During graph compilation, once compilation is done we thaw existing graph instances to reallocate their memory
 	void Thaw();
 
+	// Update and push any graph state that this graph has onto the parameter stack
+	UE::AnimNext::FParamStack::FPushedLayerHandle UpdateAndPushGraphState(float InDeltaTime) const;
+
+	// Pop any graph state that this graph has off the parameter stack
+	void PopGraphState(UE::AnimNext::FParamStack::FPushedLayerHandle InHandle) const;
+
 	// Hard reference to the graph used to create this instance to ensure we can release it safely
 	UPROPERTY()
 	TObjectPtr<const UAnimNextGraph> Graph;
@@ -145,6 +153,9 @@ private:
 	// The root graph instance that owns us and the components
 	FAnimNextGraphInstance* RootGraphInstance = nullptr;
 
+	// Graph state parameter source
+	TUniquePtr<UE::AnimNext::IParameterSource> GraphState;
+
 	// Extended execute context instance for this graph instance, we own it
 	UPROPERTY()
 	FRigVMExtendedExecuteContext ExtendedExecuteContext;
diff --git a/Engine/Plugins/Experimental/AnimNext/Source/AnimNext/Public/Graph/AnimNextGraphInstancePtr.h b/Engine/Plugins/Experimental/AnimNext/Source/AnimNext/Public/Graph/AnimNextGraphInstancePtr.h
--- a/Engine/Plugins/Experimental/AnimNext/Source/AnimNext/Public/Graph/AnimNextGraphInstancePtr.h
+++ b/Engine/Plugins/Experimental/AnimNext/Source/AnimNext/Public/Graph/AnimNextGraphInstancePtr.h
@@ -3,6 +3,7 @@
 #pragma once
 
 #include "CoreMinimal.h"
+#include "Param/ParamStack.h"
 #include "TraitCore/TraitEvent.h"
 #include "TraitCore/TraitPtr.h"
 
@@ -99,6 +100,12 @@ struct ANIMNEXT_API FAnimNextGraphInstancePtr
 	// Called each time the graph updates
 	void Update();
 
+	// Push any graph state that this graph has onto the parameter stack
+	UE::AnimNext::FParamStack::FPushedLayerHandle UpdateAndPushGraphState(float InDeltaTime) const;
+
+	// Pop any graph state that this graph has off the parameter stack
+	void PopGraphState(UE::AnimNext::FParamStack::FPushedLayerHandle InHandle) const;
+
 private:
 	// Returns a pointer to the specified component, or nullptr if not found
 	UE::AnimNext::FGraphInstanceComponent* TryGetComponent(int32 ComponentNameHash, FName ComponentName) const;
diff --git a/Engine/Plugins/Experimental/AnimNext/Source/AnimNext/Public/IAnimNextModule.h b/Engine/Plugins/Experimental/AnimNext/Source/AnimNext/Public/IAnimNextModule.h
--- a/Engine/Plugins/Experimental/AnimNext/Source/AnimNext/Public/IAnimNextModule.h
+++ b/Engine/Plugins/Experimental/AnimNext/Source/AnimNext/Public/IAnimNextModule.h
@@ -32,27 +32,6 @@ public:
 class IAnimNextModule : public IModuleInterface
 {
 public:
-	// Factory method used to create a parameter source to access the specified scope locator, with a set of parameters that are initially required
-	// @param   InContext             Context used to set up the parameter source
-	// @param   InInstanceId          The instance identifier associated with the parameters that are required
-	// @param   InRequiredParameters  Any required parameters that the source should initially supply, can be empty
-	// @return a new parameter source, or nullptr if the scope was not found or could not be handled
-	virtual TUniquePtr<IParameterSource> CreateParameterSource(const FParameterSourceContext& InContext, const TInstancedStruct<FAnimNextParamInstanceIdentifier>& InInstanceId, TConstArrayView<FName> InRequiredParameters) const = 0;
-
-	// Register a factory that can be used to generate parameter sources
-	// @param   InName                Identifier for the factory
-	// @param   InFactory             The factory to register
-	virtual void RegisterParameterSourceFactory(FName InName, TSharedRef<IParameterSourceFactory> InFactory) = 0;
-
-	// Unregister a factory previously passed to RegisterParameterSourceFactory
-	// @param   InName                Identifier for the factory
-	virtual void UnregisterParameterSourceFactory(FName InName) = 0;
-
-	// Find a factory previously passed to RegisterParameterSourceFactory
-	// @param   InName                Identifier for the factory
-	// @return a shared ptr to the factory, if found, otherwise an invalid shared ptr
-	virtual TSharedPtr<IParameterSourceFactory> FindParameterSourceFactory(FName InName) = 0;
-
 	// WARNING
 	// These functions are here because the schedule currently cannot have tasks defined in
 	// external plugins. And as such, we must be able to update and evaluate graphs from
diff --git a/Engine/Plugins/Experimental/AnimNext/Source/AnimNext/Public/Param/AnimNextClassExtensionLibrary.h b/Engine/Plugins/Experimental/AnimNext/Source/AnimNext/Public/Param/AnimNextClassExtensionLibrary.h
deleted file mode 100644
diff --git a/Engine/Plugins/Experimental/AnimNext/Source/AnimNext/Public/Param/IParameterSource.h b/Engine/Plugins/Experimental/AnimNext/Source/AnimNext/Public/Param/IParameterSource.h
deleted file mode 100644
diff --git a/Engine/Plugins/Experimental/AnimNext/Source/AnimNext/Public/Param/IParameterSourceFactory.h b/Engine/Plugins/Experimental/AnimNext/Source/AnimNext/Public/Param/IParameterSourceFactory.h
deleted file mode 100644
diff --git a/Engine/Plugins/Experimental/AnimNext/Source/AnimNext/Public/Param/ParamTypeHandle.h b/Engine/Plugins/Experimental/AnimNext/Source/AnimNext/Public/Param/ParamTypeHandle.h
--- a/Engine/Plugins/Experimental/AnimNext/Source/AnimNext/Public/Param/ParamTypeHandle.h
+++ b/Engine/Plugins/Experimental/AnimNext/Source/AnimNext/Public/Param/ParamTypeHandle.h
@@ -8,7 +8,7 @@
 
 class UAnimSequence;
 class UCharacterMovementComponent;
-class UAnimNextMeshComponent;
+class USkeletalMeshComponent;
 struct FAnimNextGraphLODPose;
 struct FAnimNextGraphReferencePose;
 
@@ -52,7 +52,7 @@ struct ANIMNEXT_API FParamTypeHandle
 		// Common object types
 		Object,
 		CharacterMovementComponent,
-		AnimNextMeshComponent,
+		SkeletalMeshComponent,
 		AnimSequence,
 
 		// Common struct types
@@ -309,9 +309,9 @@ public:
 		{
 			TypeHandle.SetParameterType(EParamType::CharacterMovementComponent);
 		}
-		else if constexpr (std::is_same_v<NonConstType, TObjectPtr<UAnimNextMeshComponent>> || std::is_same_v<NonConstType, UAnimNextMeshComponent*>)
+		else if constexpr (std::is_same_v<NonConstType, TObjectPtr<USkeletalMeshComponent>> || std::is_same_v<NonConstType, USkeletalMeshComponent*>)
 		{
-			TypeHandle.SetParameterType(EParamType::AnimNextMeshComponent);
+			TypeHandle.SetParameterType(EParamType::SkeletalMeshComponent);
 		}
 		else if constexpr (std::is_same_v<NonConstType, TObjectPtr<UAnimSequence>> || std::is_same_v<NonConstType, UAnimSequence*>)
 		{
diff --git a/Engine/Plugins/Experimental/AnimNext/Source/AnimNextEditor/Internal/AnimNextEditorModule.h b/Engine/Plugins/Experimental/AnimNext/Source/AnimNextEditor/Internal/AnimNextEditorModule.h
deleted file mode 100644
diff --git a/Engine/Plugins/Experimental/AnimNext/Source/AnimNextEditor/Internal/IAnimNextEditorModule.h b/Engine/Plugins/Experimental/AnimNext/Source/AnimNextEditor/Internal/IAnimNextEditorModule.h
new file mode 100644
--- /dev/null
+++ b/Engine/Plugins/Experimental/AnimNext/Source/AnimNextEditor/Internal/IAnimNextEditorModule.h
@@ -0,0 +1,31 @@
+﻿// Copyright Epic Games, Inc. All Rights Reserved.
+
+#pragma once
+
+#include "CoreMinimal.h"
+#include "Modules/ModuleManager.h"
+
+class SWidget;
+
+namespace UE::AnimNext::Editor
+{
+
+struct FParameterPickerArgs;
+
+class IModule : public IModuleInterface
+{
+public:
+	// Create a parameter picker
+	// @param  InArgs          Arguments used for configuring the picker
+	virtual TSharedRef<SWidget> CreateParameterPicker(const FParameterPickerArgs& InArgs) = 0;
+
+	// Register a valid fragment type name to be used with parameter UOLs
+	// @param InLocatorFragmentEditorName The name of the locator fragment editor
+	virtual void RegisterLocatorFragmentEditorType(FName InLocatorFragmentEditorName) = 0;
+
+	// Unregister a valid fragment type name to be used with parameter UOLs
+	// @param InLocatorFragmentEditorName The name of the locator fragment editor
+	virtual void UnregisterLocatorFragmentEditorType(FName InLocatorFragmentEditorName) = 0;
+};
+
+}
\ No newline at end of file
diff --git a/Engine/Plugins/Experimental/AnimNext/Source/AnimNextEditor/Private/AnimNextEditorModule.cpp b/Engine/Plugins/Experimental/AnimNext/Source/AnimNextEditor/Private/AnimNextEditorModule.cpp
--- a/Engine/Plugins/Experimental/AnimNext/Source/AnimNextEditor/Private/AnimNextEditorModule.cpp
+++ b/Engine/Plugins/Experimental/AnimNext/Source/AnimNextEditor/Private/AnimNextEditorModule.cpp
@@ -32,6 +32,7 @@
 #include "Editor/RigVMGraphDetailCustomization.h"
 #include "EditorUtils.h"
 #include "IUniversalObjectLocatorEditorModule.h"
+#include "Param/AnimNextComponentLocatorEditor.h"
 #include "Param/AnimNextParam.h"
 #include "Param/AnimNextLocatorContext.h"
 #include "Param/ObjectCastLocatorEditor.h"
@@ -40,466 +41,574 @@
 #include "Param/SAddParametersDialog.h"
 #include "AnimNextRigVMWorkspaceAssetUserData.h"
 #include "Graph/AnimNextGraph_OutlinerItemDetails.h"
+#include "Param/AnimNextActorLocatorEditor.h"
 
 #define LOCTEXT_NAMESPACE "AnimNextEditorModule"
 
 namespace UE::AnimNext::Editor
 {
 
-class FModule : public IModule
+void FModule::StartupModule()
 {
-	virtual void StartupModule() override
-	{
-		// Register settings for user editing
-		ISettingsModule& SettingsModule = FModuleManager::Get().LoadModuleChecked<ISettingsModule>("Settings");
-		SettingsModule.RegisterSettings("Editor", "General", "AnimNext",
-			LOCTEXT("SettingsName", "AnimNext"),
-			LOCTEXT("SettingsDescription", "Customize AnimNext Settings."),
-			GetMutableDefault<UAnimNextConfig>()
-		);
-
-		FPropertyEditorModule& PropertyModule = FModuleManager::LoadModuleChecked<FPropertyEditorModule>("PropertyEditor");
-
-		PropertyModule.RegisterCustomPropertyTypeLayout(
-			"AnimNextParamType",
-			FOnGetPropertyTypeCustomizationInstance::CreateLambda([] { return MakeShared<FParamTypePropertyTypeCustomization>(); }));
-
-		PropertyModule.RegisterCustomPropertyTypeLayout(
-			"AnimNextParam",
-			FOnGetPropertyTypeCustomizationInstance::CreateLambda([] { return MakeShared<FParamPropertyCustomization>(); }));
+	// Register settings for user editing
+	ISettingsModule& SettingsModule = FModuleManager::Get().LoadModuleChecked<ISettingsModule>("Settings");
+	SettingsModule.RegisterSettings("Editor", "General", "AnimNext",
+		LOCTEXT("SettingsName", "AnimNext"),
+		LOCTEXT("SettingsDescription", "Customize AnimNext Settings."),
+		GetMutableDefault<UAnimNextConfig>()
+	);
+
+	FPropertyEditorModule& PropertyModule = FModuleManager::LoadModuleChecked<FPropertyEditorModule>("PropertyEditor");
+
+	PropertyModule.RegisterCustomPropertyTypeLayout(
+		"AnimNextParamType",
+		FOnGetPropertyTypeCustomizationInstance::CreateLambda([] { return MakeShared<FParamTypePropertyTypeCustomization>(); }));
+
+	PropertyModule.RegisterCustomPropertyTypeLayout(
+		"AnimNextParam",
+		FOnGetPropertyTypeCustomizationInstance::CreateLambda([] { return MakeShared<FParamPropertyCustomization>(); }));
+
+	PropertyModule.RegisterCustomPropertyTypeLayout(
+		"AnimNextEditorParam",
+		FOnGetPropertyTypeCustomizationInstance::CreateLambda([] { return MakeShared<FParamPropertyCustomization>(); }));
+	
+	Identifier = MakeShared<FParamNamePropertyTypeIdentifier>();
+	PropertyModule.RegisterCustomPropertyTypeLayout(
+		FNameProperty::StaticClass()->GetFName(),
+		FOnGetPropertyTypeCustomizationInstance::CreateLambda([] { return MakeShared<FParamNamePropertyTypeCustomization>(); }),
+		Identifier);
 
-		PropertyModule.RegisterCustomPropertyTypeLayout(
-			"AnimNextEditorParam",
-			FOnGetPropertyTypeCustomizationInstance::CreateLambda([] { return MakeShared<FParamPropertyCustomization>(); }));
-		
-		Identifier = MakeShared<FParamNamePropertyTypeIdentifier>();
-		PropertyModule.RegisterCustomPropertyTypeLayout(
-			FNameProperty::StaticClass()->GetFName(),
-			FOnGetPropertyTypeCustomizationInstance::CreateLambda([] { return MakeShared<FParamNamePropertyTypeCustomization>(); }),
-			Identifier);
+	PropertyModule.RegisterCustomClassLayout("AnimNextGraph_Parameter", 
+		FOnGetDetailCustomizationInstance::CreateLambda([] { return MakeShared<FParameterCustomization>(); }));
 
-		PropertyModule.RegisterCustomClassLayout("AnimNextGraph_Parameter", 
-			FOnGetDetailCustomizationInstance::CreateLambda([] { return MakeShared<FParameterCustomization>(); }));
+	PropertyModule.RegisterCustomClassLayout("AnimNextGraph_EdGraphNode",
+		FOnGetDetailCustomizationInstance::CreateLambda([] { return MakeShared<FAnimNextGraph_EdGraphNodeCustomization>(); }));
 
-		PropertyModule.RegisterCustomClassLayout("AnimNextGraph_EdGraphNode",
-			FOnGetDetailCustomizationInstance::CreateLambda([] { return MakeShared<FAnimNextGraph_EdGraphNodeCustomization>(); }));
+	AnimNextGraphPanelNodeFactory = MakeShared<FAnimNextGraphPanelNodeFactory>();
+	FEdGraphUtilities::RegisterVisualNodeFactory(AnimNextGraphPanelNodeFactory);
 
-		AnimNextGraphPanelNodeFactory = MakeShared<FAnimNextGraphPanelNodeFactory>();
-		FEdGraphUtilities::RegisterVisualNodeFactory(AnimNextGraphPanelNodeFactory);
+	ParametersGraphPanelPinFactory = MakeShared<FParametersGraphPanelPinFactory>();
+	FEdGraphUtilities::RegisterVisualPinFactory(ParametersGraphPanelPinFactory);
 
-		ParametersGraphPanelPinFactory = MakeShared<FParametersGraphPanelPinFactory>();
-		FEdGraphUtilities::RegisterVisualPinFactory(ParametersGraphPanelPinFactory);
+	Workspace::IWorkspaceEditorModule& WorkspaceEditorModule = FModuleManager::Get().LoadModuleChecked<Workspace::IWorkspaceEditorModule>("WorkspaceEditor");
 
-		Workspace::IWorkspaceEditorModule& WorkspaceEditorModule = FModuleManager::Get().LoadModuleChecked<Workspace::IWorkspaceEditorModule>("WorkspaceEditor");
+	RegisterWorkspaceDocumentTypes(WorkspaceEditorModule);
 
-		RegisterWorkspaceDocumentTypes(WorkspaceEditorModule);
+	WorkspaceEditorModule.OnRegisterWorkspaceDetailsCustomization().AddLambda([](TSharedPtr<IDetailsView>& InDetailsView)
+		{
+			TArray<UScriptStruct*> StructsToCustomize = {
+				TBaseStructure<FVector>::Get(),
+				TBaseStructure<FVector2D>::Get(),
+				TBaseStructure<FVector4>::Get(),
+				TBaseStructure<FRotator>::Get(),
+				TBaseStructure<FQuat>::Get(),
+				TBaseStructure<FTransform>::Get(),
+				TBaseStructure<FEulerTransform>::Get(),
+			};
+
+			for (UScriptStruct* StructToCustomize : StructsToCustomize)
+			{
+				InDetailsView->RegisterInstancedCustomPropertyTypeLayout(StructToCustomize->GetFName(), FOnGetPropertyTypeCustomizationInstance::CreateLambda([]()
+					{
+						return FRigVMGraphMathTypeDetailCustomization::MakeInstance();
+					}));
+			}
+		});
 
-		WorkspaceEditorModule.OnRegisterWorkspaceDetailsCustomization().AddLambda([](TSharedPtr<IDetailsView>& InDetailsView)
+	SRigVMAssetView::RegisterCategoryFactory("Parameters", [](UAnimNextRigVMAssetEditorData* InEditorData)
+	{
+		UAnimNextGraph_EditorData* EditorData = CastChecked<UAnimNextGraph_EditorData>(InEditorData);
+		UAnimNextRigVMAsset* Asset = UncookedOnly::FUtils::GetAsset(InEditorData);
+		return SNew(SSimpleButton)
+			.Text(LOCTEXT("AddParameterButton", "Add Parameter"))
+			.Icon(FAppStyle::Get().GetBrush("Icons.Plus"))
+			.OnClicked_Lambda([EditorData, Asset]()
 			{
-				TArray<UScriptStruct*> StructsToCustomize = {
-					TBaseStructure<FVector>::Get(),
-					TBaseStructure<FVector2D>::Get(),
-					TBaseStructure<FVector4>::Get(),
-					TBaseStructure<FRotator>::Get(),
-					TBaseStructure<FQuat>::Get(),
-					TBaseStructure<FTransform>::Get(),
-					TBaseStructure<FEulerTransform>::Get(),
-				};
-
-				for (UScriptStruct* StructToCustomize : StructsToCustomize)
-				{
-					InDetailsView->RegisterInstancedCustomPropertyTypeLayout(StructToCustomize->GetFName(), FOnGetPropertyTypeCustomizationInstance::CreateLambda([]()
-						{
-							return FRigVMGraphMathTypeDetailCustomization::MakeInstance();
-						}));
-				}
-			});
+				TSharedRef<SAddParametersDialog> AddParametersDialog =
+					SNew(SAddParametersDialog, FAssetData(EditorData));
 
-		SRigVMAssetView::RegisterCategoryFactory("Parameters", [](UAnimNextRigVMAssetEditorData* InEditorData)
-		{
-			UAnimNextGraph_EditorData* EditorData = CastChecked<UAnimNextGraph_EditorData>(InEditorData);
-			UAnimNextRigVMAsset* Asset = UncookedOnly::FUtils::GetAsset(InEditorData);
-			return SNew(SSimpleButton)
-				.Text(LOCTEXT("AddParameterButton", "Add Parameter"))
-				.Icon(FAppStyle::Get().GetBrush("Icons.Plus"))
-				.OnClicked_Lambda([EditorData, Asset]()
+				TArray<FParameterToAdd> ParametersToAdd;
+				if(AddParametersDialog->ShowModal(ParametersToAdd))
 				{
-					TSharedRef<SAddParametersDialog> AddParametersDialog =
-						SNew(SAddParametersDialog, FAssetData(EditorData));
-
-					TArray<FParameterToAdd> ParametersToAdd;
-					if(AddParametersDialog->ShowModal(ParametersToAdd))
+					if(ParametersToAdd.Num() > 0)
 					{
-						if(ParametersToAdd.Num() > 0)
+						FScopedTransaction Transaction(LOCTEXT("AddParameter", "Add parameter"));
+						for (const FParameterToAdd& ParameterToAdd : ParametersToAdd)
 						{
-							FScopedTransaction Transaction(LOCTEXT("AddParameter", "Add parameter"));
-							for (const FParameterToAdd& ParameterToAdd : ParametersToAdd)
-							{
-								check(EditorData->FindEntry(ParameterToAdd.Name) == nullptr);
-								EditorData->AddParameter(ParameterToAdd.Name, ParameterToAdd.Type);
-							}
+							check(EditorData->FindEntry(ParameterToAdd.Name) == nullptr);
+							EditorData->AddParameter(ParameterToAdd.Name, ParameterToAdd.Type);
 						}
 					}
-					return FReply::Handled();
-				});
-		});
+				}
+				return FReply::Handled();
+			});
+	});
 
-		SRigVMAssetView::RegisterCategoryFactory("Event Graphs", [](UAnimNextRigVMAssetEditorData* InEditorData)
-		{
-			UAnimNextGraph_EditorData* EditorData = CastChecked<UAnimNextGraph_EditorData>(InEditorData);
-			return SNew(SSimpleButton)
-				.Text(LOCTEXT("AddEventGraphButton", "Add Event Graph"))
-				.Icon(FAppStyle::Get().GetBrush("Icons.Plus"))
-				.OnClicked_Lambda([EditorData]()
-				{
-					FScopedTransaction Transaction(LOCTEXT("AddEventGraph", "Add Event Graph"));
+	SRigVMAssetView::RegisterCategoryFactory("Event Graphs", [](UAnimNextRigVMAssetEditorData* InEditorData)
+	{
+		UAnimNextGraph_EditorData* EditorData = CastChecked<UAnimNextGraph_EditorData>(InEditorData);
+		return SNew(SSimpleButton)
+			.Text(LOCTEXT("AddEventGraphButton", "Add Event Graph"))
+			.Icon(FAppStyle::Get().GetBrush("Icons.Plus"))
+			.OnClicked_Lambda([EditorData]()
+			{
+				FScopedTransaction Transaction(LOCTEXT("AddEventGraph", "Add Event Graph"));
 
-					// Create a new entry for the graph
-					EditorData->AddEventGraph(TEXT("NewGraph"));
+				// Create a new entry for the graph
+				EditorData->AddEventGraph(TEXT("NewGraph"));
 
-					return FReply::Handled();
-				});
-		});
+				return FReply::Handled();
+			});
+	});
 
-		SRigVMAssetView::RegisterCategoryFactory("Animation Graphs", [](UAnimNextRigVMAssetEditorData* InEditorData)
-		{
-			UAnimNextGraph_EditorData* EditorData = CastChecked<UAnimNextGraph_EditorData>(InEditorData);
-			return SNew(SSimpleButton)
-				.Text(LOCTEXT("AddGraphButton", "Add Animation Graph"))
-				.Icon(FAppStyle::Get().GetBrush("Icons.Plus"))
-				.OnClicked_Lambda([EditorData]()
-				{
-					FScopedTransaction Transaction(LOCTEXT("AddAnimationGraph", "Add Animation Graph"));
+	SRigVMAssetView::RegisterCategoryFactory("Animation Graphs", [](UAnimNextRigVMAssetEditorData* InEditorData)
+	{
+		UAnimNextGraph_EditorData* EditorData = CastChecked<UAnimNextGraph_EditorData>(InEditorData);
+		return SNew(SSimpleButton)
+			.Text(LOCTEXT("AddGraphButton", "Add Animation Graph"))
+			.Icon(FAppStyle::Get().GetBrush("Icons.Plus"))
+			.OnClicked_Lambda([EditorData]()
+			{
+				FScopedTransaction Transaction(LOCTEXT("AddAnimationGraph", "Add Animation Graph"));
 
-					// Create a new entry for the graph
-					EditorData->AddAnimationGraph(TEXT("NewGraph"));
+				// Create a new entry for the graph
+				EditorData->AddAnimationGraph(TEXT("NewGraph"));
 
-					return FReply::Handled();
-				});
-		});
+				return FReply::Handled();
+			});
+	});
+
+	UE::UniversalObjectLocator::IUniversalObjectLocatorEditorModule& UolEditorModule = FModuleManager::LoadModuleChecked<UE::UniversalObjectLocator::IUniversalObjectLocatorEditorModule>("UniversalObjectLocatorEditor");
+	UolEditorModule.RegisterLocatorEditor("AnimNextObjectFunction", MakeShared<FObjectFunctionLocatorEditor>());
+	UolEditorModule.RegisterLocatorEditor("AnimNextObjectProperty", MakeShared<FObjectPropertyLocatorEditor>());
+	UolEditorModule.RegisterLocatorEditor("AnimNextObjectCast", MakeShared<FObjectCastLocatorEditor>());
+	UolEditorModule.RegisterLocatorEditor("AnimNextComponent", MakeShared<FComponentLocatorEditor>());
+	UolEditorModule.RegisterLocatorEditor("AnimNextActor", MakeShared<FActorLocatorEditor>());
+
+	UolEditorModule.RegisterEditorContext("AnimNextContext", MakeShared<FLocatorContext>());
+
+	RegisterLocatorFragmentEditorType("Actor");
+	RegisterLocatorFragmentEditorType("Asset");
+	RegisterLocatorFragmentEditorType("AnimNextScope");
+	RegisterLocatorFragmentEditorType("AnimNextGraph");
+	RegisterLocatorFragmentEditorType("AnimNextObjectFunction");
+	RegisterLocatorFragmentEditorType("AnimNextObjectProperty");
+	RegisterLocatorFragmentEditorType("AnimNextObjectCast");
+	RegisterLocatorFragmentEditorType("AnimNextComponent");
+	RegisterLocatorFragmentEditorType("AnimNextActor");
+
+	Workspace::IWorkspaceEditorModule& WorkspaceModule = FModuleManager::Get().LoadModuleChecked<Workspace::IWorkspaceEditorModule>("WorkspaceEditor");
+	const TSharedPtr<FAnimNextGraphItemDetails> GraphItemDetails = MakeShareable<FAnimNextGraphItemDetails>(new FAnimNextGraphItemDetails());
+	WorkspaceModule.RegisterWorkspaceItemDetails(Workspace::FOutlinerItemDetailsId(FAnimNextGraphOutlinerData::StaticStruct()->GetFName()), StaticCastSharedPtr<UE::Workspace::IWorkspaceOutlinerItemDetails>(GraphItemDetails));
+	FAnimNextGraphItemDetails::RegisterToolMenuExtensions();	
+}
+
+void FModule::ShutdownModule()
+{
+	if(FModuleManager::Get().IsModuleLoaded("PropertyEditor"))
+	{
+		FPropertyEditorModule& PropertyModule = FModuleManager::LoadModuleChecked<FPropertyEditorModule>("PropertyEditor");
+		PropertyModule.UnregisterCustomPropertyTypeLayout("AnimNextParamType");
+		PropertyModule.UnregisterCustomPropertyTypeLayout("AnimNextParam");
+		PropertyModule.UnregisterCustomPropertyTypeLayout("AnimNextEditorParam");
+		PropertyModule.UnregisterCustomPropertyTypeLayout("NameProperty");
+		PropertyModule.UnregisterCustomClassLayout("AnimNextGraph_Parameter");
+		PropertyModule.UnregisterCustomClassLayout("AnimNextGraph_EdGraphNode");
+	}
 
-		UE::UniversalObjectLocator::IUniversalObjectLocatorEditorModule& UolEditorModule = FModuleManager::LoadModuleChecked<UE::UniversalObjectLocator::IUniversalObjectLocatorEditorModule>("UniversalObjectLocatorEditor");
-		UolEditorModule.RegisterLocatorEditor("AnimNextObjectFunction", MakeShared<FObjectFunctionLocatorEditor>());
-		UolEditorModule.RegisterLocatorEditor("AnimNextObjectProperty", MakeShared<FObjectPropertyLocatorEditor>());
-		UolEditorModule.RegisterLocatorEditor("AnimNextObjectCast", MakeShared<FObjectCastLocatorEditor>());
+	FEdGraphUtilities::UnregisterVisualNodeFactory(AnimNextGraphPanelNodeFactory);
 
-		UolEditorModule.RegisterEditorContext("AnimNextContext", MakeShared<FAnimNextLocatorContext>());
+	FEdGraphUtilities::UnregisterVisualPinFactory(ParametersGraphPanelPinFactory);
 
-		Workspace::IWorkspaceEditorModule& WorkspaceModule = FModuleManager::Get().LoadModuleChecked<Workspace::IWorkspaceEditorModule>("WorkspaceEditor");
-		const TSharedPtr<FAnimNextGraphItemDetails> GraphItemDetails = MakeShareable<FAnimNextGraphItemDetails>(new FAnimNextGraphItemDetails());
-		WorkspaceModule.RegisterWorkspaceItemDetails(Workspace::FOutlinerItemDetailsId(FAnimNextGraphOutlinerData::StaticStruct()->GetFName()), StaticCastSharedPtr<UE::Workspace::IWorkspaceOutlinerItemDetails>(GraphItemDetails));
-		FAnimNextGraphItemDetails::RegisterToolMenuExtensions();
-	}
+	UnregisterWorkspaceDocumentTypes();
+
+	SRigVMAssetView::UnregisterCategoryFactory("Parameters");
+	SRigVMAssetView::UnregisterCategoryFactory("Event Graphs");
+	SRigVMAssetView::UnregisterCategoryFactory("Animation Graphs");
 
-	virtual void ShutdownModule() override
+	if(FModuleManager::Get().IsModuleLoaded("UniversalObjectLocatorEditor"))
 	{
-		if(FModuleManager::Get().IsModuleLoaded("PropertyEditor"))
-		{
-			FPropertyEditorModule& PropertyModule = FModuleManager::LoadModuleChecked<FPropertyEditorModule>("PropertyEditor");
-			PropertyModule.UnregisterCustomPropertyTypeLayout("AnimNextParamType");
-			PropertyModule.UnregisterCustomPropertyTypeLayout("AnimNextParam");
-			PropertyModule.UnregisterCustomPropertyTypeLayout("AnimNextEditorParam");
-			PropertyModule.UnregisterCustomPropertyTypeLayout("NameProperty");
-			PropertyModule.UnregisterCustomClassLayout("AnimNextGraph_Parameter");
-			PropertyModule.UnregisterCustomClassLayout("AnimNextGraph_EdGraphNode");
-		}
+		UE::UniversalObjectLocator::IUniversalObjectLocatorEditorModule& UolEditorModule = FModuleManager::GetModuleChecked<UE::UniversalObjectLocator::IUniversalObjectLocatorEditorModule>("UniversalObjectLocatorEditor");
+		UolEditorModule.UnregisterLocatorEditor("AnimNextObjectCast");
+		UolEditorModule.UnregisterLocatorEditor("AnimNextObjectFunction");
+		UolEditorModule.UnregisterLocatorEditor("AnimNextObjectProperty");
+		UolEditorModule.UnregisterLocatorEditor("AnimNextComponent");
+		UolEditorModule.UnregisterLocatorEditor("AnimNextActor");
+
+		UolEditorModule.UnregisterEditorContext("AnimNextContext");
+	}
+	
+	if (UObjectInitialized())
+	{
+		Workspace::IWorkspaceEditorModule& WorkspaceModule = FModuleManager::Get().LoadModuleChecked<Workspace::IWorkspaceEditorModule>("WorkspaceEditor");
+		WorkspaceModule.UnregisterWorkspaceItemDetails(Workspace::FOutlinerItemDetailsId(FAnimNextGraphOutlinerData::StaticStruct()->GetFName()));
+		FAnimNextGraphItemDetails::UnregisterToolMenuExtensions();
+	}
 
-		FEdGraphUtilities::UnregisterVisualNodeFactory(AnimNextGraphPanelNodeFactory);
+	UnregisterLocatorFragmentEditorType("Actor");
+	UnregisterLocatorFragmentEditorType("Asset");
+	UnregisterLocatorFragmentEditorType("AnimNextScope");
+	UnregisterLocatorFragmentEditorType("AnimNextGraph");
+	UnregisterLocatorFragmentEditorType("AnimNextObjectFunction");
+	UnregisterLocatorFragmentEditorType("AnimNextObjectProperty");
+	UnregisterLocatorFragmentEditorType("AnimNextObjectCast");
+	UnregisterLocatorFragmentEditorType("AnimNextComponent");
+	UnregisterLocatorFragmentEditorType("AnimNextActor");
+}
 
-		FEdGraphUtilities::UnregisterVisualPinFactory(ParametersGraphPanelPinFactory);
+TSharedRef<SWidget> FModule::CreateParameterPicker(const FParameterPickerArgs& InArgs)
+{
+	return SNew(SParameterPicker)
+		.Args(InArgs);
+}
 
-		UnregisterWorkspaceDocumentTypes();
+void FModule::RegisterLocatorFragmentEditorType(FName InLocatorFragmentEditorName)
+{
+	LocatorFragmentEditorNames.Add(InLocatorFragmentEditorName);
+}
 
-		SRigVMAssetView::UnregisterCategoryFactory("Parameters");
-		SRigVMAssetView::UnregisterCategoryFactory("Event Graphs");
-		SRigVMAssetView::UnregisterCategoryFactory("Animation Graphs");
+void FModule::UnregisterLocatorFragmentEditorType(FName InLocatorFragmentEditorName)
+{
+	LocatorFragmentEditorNames.Remove(InLocatorFragmentEditorName);
+}
 
-		if(FModuleManager::Get().IsModuleLoaded("UniversalObjectLocatorEditor"))
+void FModule::RegisterWorkspaceDocumentTypes(Workspace::IWorkspaceEditorModule& WorkspaceEditorModule)
+{
+	// --- AnimNextSchedule ---
+	Workspace::FObjectDocumentArgs ScheduleDocumentArgs(
+		Workspace::FOnMakeDocumentWidget::CreateLambda([](const Workspace::FWorkspaceEditorContext& InContext)
 		{
-			UE::UniversalObjectLocator::IUniversalObjectLocatorEditorModule& UolEditorModule = FModuleManager::GetModuleChecked<UE::UniversalObjectLocator::IUniversalObjectLocatorEditorModule>("UniversalObjectLocatorEditor");
-			UolEditorModule.UnregisterLocatorEditor("AnimNextObjectCast");
-			UolEditorModule.UnregisterLocatorEditor("AnimNextObjectFunction");
-			UolEditorModule.UnregisterLocatorEditor("AnimNextObjectProperty");
-
-			UolEditorModule.UnregisterEditorContext("AnimNextContext");
+			UAnimNextSchedule* Schedule = CastChecked<UAnimNextSchedule>(InContext.Object);
+			FPropertyEditorModule& PropertyEditorModule = FModuleManager::GetModuleChecked<FPropertyEditorModule>( "PropertyEditor" );
+			FDetailsViewArgs DetailsViewArgs;
+			DetailsViewArgs.NameAreaSettings = FDetailsViewArgs::HideNameArea;
+			TSharedRef<IDetailsView> DetailsView = PropertyEditorModule.CreateDetailView(DetailsViewArgs);
+			DetailsView->SetObject(Schedule);
+			return DetailsView;
 		}
+	), Workspace::WorkspaceTabs::MiddleDocumentArea);
+	ScheduleDocumentArgs.OnGetTabName = Workspace::FOnGetTabName::CreateLambda([](const Workspace::FWorkspaceEditorContext& InContext)
+	{
+		const UAnimNextSchedule* Schedule = CastChecked<UAnimNextSchedule>(InContext.Object);
+		return FText::FromName(Schedule->GetFName());
+	});
 
-		if (UObjectInitialized())
+	ScheduleDocumentArgs.OnGetDocumentBreadcrumbTrail = Workspace::FOnGetDocumentBreadcrumbTrail::CreateLambda([](const Workspace::FWorkspaceEditorContext& InContext, TArray<TSharedPtr<Workspace::FWorkspaceBreadcrumb>>& OutBreadcrumbs)
+	{
+		if (const UAnimNextSchedule* Schedule = Cast<UAnimNextSchedule>(InContext.Object))
 		{
-			Workspace::IWorkspaceEditorModule& WorkspaceModule = FModuleManager::Get().LoadModuleChecked<Workspace::IWorkspaceEditorModule>("WorkspaceEditor");
-			WorkspaceModule.UnregisterWorkspaceItemDetails(Workspace::FOutlinerItemDetailsId(FAnimNextGraphOutlinerData::StaticStruct()->GetFName()));
-			FAnimNextGraphItemDetails::UnregisterToolMenuExtensions();
+			const TSharedPtr<Workspace::FWorkspaceBreadcrumb>& GraphCrumb = OutBreadcrumbs.Add_GetRef(MakeShared<Workspace::FWorkspaceBreadcrumb>());
+			GraphCrumb->OnGetLabel = Workspace::FWorkspaceBreadcrumb::FOnGetBreadcrumbLabel::CreateLambda([ScheduleName = Schedule->GetFName()]{ return FText::FromName(ScheduleName); });
 		}
-	}
+	});
 
-	virtual TSharedRef<SWidget> CreateParameterPicker(const FParameterPickerArgs& InArgs) override
+	ScheduleDocumentArgs.OnGetTabIcon = Workspace::FOnGetTabIcon::CreateLambda([](const Workspace::FWorkspaceEditorContext& InContext)
 	{
-		return SNew(SParameterPicker)
-			.Args(InArgs);
-	}
+		return FAppStyle::GetBrush(TEXT("ClassIcon.Default"));
+	});
 
-	void RegisterWorkspaceDocumentTypes(Workspace::IWorkspaceEditorModule& WorkspaceEditorModule)
-	{
-		// --- AnimNextSchedule ---
-		Workspace::FObjectDocumentArgs ScheduleDocumentArgs(
-			Workspace::FOnMakeDocumentWidget::CreateLambda([](const Workspace::FWorkspaceEditorContext& InContext)
-			{
-				UAnimNextSchedule* Schedule = CastChecked<UAnimNextSchedule>(InContext.Object);
-				FPropertyEditorModule& PropertyEditorModule = FModuleManager::GetModuleChecked<FPropertyEditorModule>( "PropertyEditor" );
-				FDetailsViewArgs DetailsViewArgs;
-				DetailsViewArgs.NameAreaSettings = FDetailsViewArgs::HideNameArea;
-				TSharedRef<IDetailsView> DetailsView = PropertyEditorModule.CreateDetailView(DetailsViewArgs);
-				DetailsView->SetObject(Schedule);
-				return DetailsView;
-			}
-		), Workspace::WorkspaceTabs::MiddleDocumentArea);
-		ScheduleDocumentArgs.OnGetTabName = Workspace::FOnGetTabName::CreateLambda([](const Workspace::FWorkspaceEditorContext& InContext)
-		{
-			const UAnimNextSchedule* Schedule = CastChecked<UAnimNextSchedule>(InContext.Object);
-			return FText::FromName(Schedule->GetFName());
-		});
+	WorkspaceEditorModule.RegisterObjectDocumentType(FTopLevelAssetPath(TEXT("/Script/AnimNext.AnimNextSchedule")), ScheduleDocumentArgs);
 
-		ScheduleDocumentArgs.OnGetDocumentBreadcrumbTrail = Workspace::FOnGetDocumentBreadcrumbTrail::CreateLambda([](const Workspace::FWorkspaceEditorContext& InContext, TArray<TSharedPtr<Workspace::FWorkspaceBreadcrumb>>& OutBreadcrumbs)
-		{
-			if (const UAnimNextSchedule* Schedule = Cast<UAnimNextSchedule>(InContext.Object))
+	// --- AnimNextGraph ---
+	Workspace::FObjectDocumentArgs AnimNextGraphDocumentArgs(
+			Workspace::FOnMakeDocumentWidget::CreateLambda([](const Workspace::FWorkspaceEditorContext& InContext)
 			{
-				const TSharedPtr<Workspace::FWorkspaceBreadcrumb>& GraphCrumb = OutBreadcrumbs.Add_GetRef(MakeShared<Workspace::FWorkspaceBreadcrumb>());
-				GraphCrumb->OnGetLabel = Workspace::FWorkspaceBreadcrumb::FOnGetBreadcrumbLabel::CreateLambda([ScheduleName = Schedule->GetFName()]{ return FText::FromName(ScheduleName); });
-			}
-		});
-
-		ScheduleDocumentArgs.OnGetTabIcon = Workspace::FOnGetTabIcon::CreateLambda([](const Workspace::FWorkspaceEditorContext& InContext)
-		{
-			return FAppStyle::GetBrush(TEXT("ClassIcon.Default"));
-		});
+				UAnimNextGraph* Graph = CastChecked<UAnimNextGraph>(InContext.Object);
+				UAnimNextGraph_EditorData* EditorData = UncookedOnly::FUtils::GetEditorData(Graph);
 
-		WorkspaceEditorModule.RegisterObjectDocumentType(FTopLevelAssetPath(TEXT("/Script/AnimNext.AnimNextSchedule")), ScheduleDocumentArgs);
+				TWeakPtr<Workspace::IWorkspaceEditor> WeakWorkspaceEditor = InContext.WorkspaceEditor;
 
-		// --- AnimNextGraph ---
-		Workspace::FObjectDocumentArgs AnimNextGraphDocumentArgs(
-				Workspace::FOnMakeDocumentWidget::CreateLambda([](const Workspace::FWorkspaceEditorContext& InContext)
+				EditorData->RigVMGraphModifiedEvent.RemoveAll(&InContext.WorkspaceEditor.Get());
+				EditorData->RigVMGraphModifiedEvent.AddSPLambda(&InContext.WorkspaceEditor.Get(), [WeakWorkspaceEditor](ERigVMGraphNotifType InType, URigVMGraph* InGraph, UObject* InSubject)
 				{
-					UAnimNextGraph* Graph = CastChecked<UAnimNextGraph>(InContext.Object);
-					UAnimNextGraph_EditorData* EditorData = UncookedOnly::FUtils::GetEditorData(Graph);
-
-					TWeakPtr<Workspace::IWorkspaceEditor> WeakWorkspaceEditor = InContext.WorkspaceEditor;
-
-					EditorData->RigVMGraphModifiedEvent.RemoveAll(&InContext.WorkspaceEditor.Get());
-					EditorData->RigVMGraphModifiedEvent.AddSPLambda(&InContext.WorkspaceEditor.Get(), [WeakWorkspaceEditor](ERigVMGraphNotifType InType, URigVMGraph* InGraph, UObject* InSubject)
+					if(TSharedPtr<Workspace::IWorkspaceEditor> WorkspaceEditor = WeakWorkspaceEditor.Pin())
 					{
-						if(TSharedPtr<Workspace::IWorkspaceEditor> WorkspaceEditor = WeakWorkspaceEditor.Pin())
+						if (InType == ERigVMGraphNotifType::InteractionBracketClosed)
 						{
-							if (InType == ERigVMGraphNotifType::InteractionBracketClosed)
-							{
-								WorkspaceEditor->RefreshDetails();
-							}
+							WorkspaceEditor->RefreshDetails();
 						}
-					});
+					}
+				});
 
-					return SNew(SRigVMAssetView, EditorData)
-						.OnSelectionChanged_Lambda([WeakWorkspaceEditor](const TArray<UObject*>& InEntries)
+				return SNew(SRigVMAssetView, EditorData)
+					.OnSelectionChanged_Lambda([WeakWorkspaceEditor](const TArray<UObject*>& InEntries)
+					{
+						if(TSharedPtr<Workspace::IWorkspaceEditor> WorkspaceEditor = WeakWorkspaceEditor.Pin())
 						{
-							if(TSharedPtr<Workspace::IWorkspaceEditor> WorkspaceEditor = WeakWorkspaceEditor.Pin())
-							{
-								WorkspaceEditor->SetDetailsObjects(InEntries);
-							}
-						})
-						.OnOpenGraph_Lambda([WeakWorkspaceEditor](URigVMGraph* InGraph)
+							WorkspaceEditor->SetDetailsObjects(InEntries);
+						}
+					})
+					.OnOpenGraph_Lambda([WeakWorkspaceEditor](URigVMGraph* InGraph)
+					{
+						if(TSharedPtr<Workspace::IWorkspaceEditor> WorkspaceEditor = WeakWorkspaceEditor.Pin())
 						{
-							if(TSharedPtr<Workspace::IWorkspaceEditor> WorkspaceEditor = WeakWorkspaceEditor.Pin())
+							if(IRigVMClientHost* RigVMClientHost = InGraph->GetImplementingOuter<IRigVMClientHost>())
 							{
-								if(IRigVMClientHost* RigVMClientHost = InGraph->GetImplementingOuter<IRigVMClientHost>())
+								if(UObject* EditorObject = RigVMClientHost->GetEditorObjectForRigVMGraph(InGraph))
 								{
-									if(UObject* EditorObject = RigVMClientHost->GetEditorObjectForRigVMGraph(InGraph))
-									{
-										WorkspaceEditor->OpenObjects({EditorObject});
-									}
+									WorkspaceEditor->OpenObjects({EditorObject});
 								}
 							}
-						})
-						.OnDeleteEntries_Lambda([WeakWorkspaceEditor](const TArray<UAnimNextRigVMAssetEntry*>& InEntries)
+						}
+					})
+					.OnDeleteEntries_Lambda([WeakWorkspaceEditor](const TArray<UAnimNextRigVMAssetEntry*>& InEntries)
+					{
+						if(TSharedPtr<Workspace::IWorkspaceEditor> WorkspaceEditor = WeakWorkspaceEditor.Pin())
 						{
-							if(TSharedPtr<Workspace::IWorkspaceEditor> WorkspaceEditor = WeakWorkspaceEditor.Pin())
+							if(InEntries.Num() > 0)
 							{
-								if(InEntries.Num() > 0)
+								TArray<UObject*> EdGraphsToClose;
+								EdGraphsToClose.Reserve(InEntries.Num());
+								for(UAnimNextRigVMAssetEntry* Entry : InEntries)
 								{
-									TArray<UObject*> EdGraphsToClose;
-									EdGraphsToClose.Reserve(InEntries.Num());
-									for(UAnimNextRigVMAssetEntry* Entry : InEntries)
+									if(IAnimNextRigVMGraphInterface* GraphInterface = Cast<IAnimNextRigVMGraphInterface>(Entry))
 									{
-										if(IAnimNextRigVMGraphInterface* GraphInterface = Cast<IAnimNextRigVMGraphInterface>(Entry))
+										if(URigVMEdGraph* EdGraph = GraphInterface->GetEdGraph())
 										{
-											if(URigVMEdGraph* EdGraph = GraphInterface->GetEdGraph())
-											{
-												EdGraphsToClose.Add(EdGraph);
-											}
+											EdGraphsToClose.Add(EdGraph);
 										}
 									}
-
-									WorkspaceEditor->CloseObjects(EdGraphsToClose);
 								}
+
+								WorkspaceEditor->CloseObjects(EdGraphsToClose);
 							}
-						});
-				}),
-				Workspace::WorkspaceTabs::MiddleDocumentArea);
-		AnimNextGraphDocumentArgs.OnGetTabName = Workspace::FOnGetTabName::CreateLambda([](const Workspace::FWorkspaceEditorContext& InContext)
-		{
-			const UAnimNextGraph* Graph = CastChecked<UAnimNextGraph>(InContext.Object);
-			return FText::FromName(Graph->GetFName());
-		});
+						}
+					});
+			}),
+			Workspace::WorkspaceTabs::MiddleDocumentArea);
+	AnimNextGraphDocumentArgs.OnGetTabName = Workspace::FOnGetTabName::CreateLambda([](const Workspace::FWorkspaceEditorContext& InContext)
+	{
+		const UAnimNextGraph* Graph = CastChecked<UAnimNextGraph>(InContext.Object);
+		return FText::FromName(Graph->GetFName());
+	});
 
-		AnimNextGraphDocumentArgs.OnGetDocumentBreadcrumbTrail = Workspace::FOnGetDocumentBreadcrumbTrail::CreateLambda([](const Workspace::FWorkspaceEditorContext& InContext, TArray<TSharedPtr<Workspace::FWorkspaceBreadcrumb>>& OutBreadcrumbs)
+	AnimNextGraphDocumentArgs.OnGetDocumentBreadcrumbTrail = Workspace::FOnGetDocumentBreadcrumbTrail::CreateLambda([](const Workspace::FWorkspaceEditorContext& InContext, TArray<TSharedPtr<Workspace::FWorkspaceBreadcrumb>>& OutBreadcrumbs)
+	{
+		if (const UAnimNextGraph* Graph = Cast<UAnimNextGraph>(InContext.Object))
 		{
-			if (const UAnimNextGraph* Graph = Cast<UAnimNextGraph>(InContext.Object))
-			{
-				const TSharedPtr<Workspace::FWorkspaceBreadcrumb>& GraphCrumb = OutBreadcrumbs.Add_GetRef(MakeShared<Workspace::FWorkspaceBreadcrumb>());
-				GraphCrumb->OnGetLabel = Workspace::FWorkspaceBreadcrumb::FOnGetBreadcrumbLabel::CreateLambda([GraphName = Graph->GetFName()]{ return FText::FromName(GraphName); });
-			}
-		});
+			const TSharedPtr<Workspace::FWorkspaceBreadcrumb>& GraphCrumb = OutBreadcrumbs.Add_GetRef(MakeShared<Workspace::FWorkspaceBreadcrumb>());
+			GraphCrumb->OnGetLabel = Workspace::FWorkspaceBreadcrumb::FOnGetBreadcrumbLabel::CreateLambda([GraphName = Graph->GetFName()]{ return FText::FromName(GraphName); });
+		}
+	});
 
-		AnimNextGraphDocumentArgs.OnGetTabIcon = Workspace::FOnGetTabIcon::CreateLambda([](const Workspace::FWorkspaceEditorContext& InContext)
-		{
-			return FAppStyle::GetBrush(TEXT("ClassIcon.Default"));
-		});
+	AnimNextGraphDocumentArgs.OnGetTabIcon = Workspace::FOnGetTabIcon::CreateLambda([](const Workspace::FWorkspaceEditorContext& InContext)
+	{
+		return FAppStyle::GetBrush(TEXT("ClassIcon.Default"));
+	});
 
-		WorkspaceEditorModule.RegisterObjectDocumentType(FTopLevelAssetPath(TEXT("/Script/AnimNext.AnimNextGraph")), AnimNextGraphDocumentArgs);
+	WorkspaceEditorModule.RegisterObjectDocumentType(FTopLevelAssetPath(TEXT("/Script/AnimNext.AnimNextGraph")), AnimNextGraphDocumentArgs);
 
-		// --- AnimNextGraph_EdGraph ---
-		Workspace::FGraphDocumentWidgetArgs GraphArgs;
-		GraphArgs.SpawnLocation = Workspace::WorkspaceTabs::MiddleDocumentArea;
-		GraphArgs.OnCreateActionMenu = Workspace::FOnCreateActionMenu::CreateLambda([](const Workspace::FWorkspaceEditorContext& InContext, UEdGraph* InGraph, const FVector2D& InNodePosition, const TArray<UEdGraphPin*>& InDraggedPins, bool bAutoExpand, SGraphEditor::FActionMenuClosed InOnMenuClosed)
+	// --- AnimNextGraph_EdGraph ---
+	Workspace::FGraphDocumentWidgetArgs GraphArgs;
+	GraphArgs.SpawnLocation = Workspace::WorkspaceTabs::MiddleDocumentArea;
+	GraphArgs.OnCreateActionMenu = Workspace::FOnCreateActionMenu::CreateLambda([](const Workspace::FWorkspaceEditorContext& InContext, UEdGraph* InGraph, const FVector2D& InNodePosition, const TArray<UEdGraphPin*>& InDraggedPins, bool bAutoExpand, SGraphEditor::FActionMenuClosed InOnMenuClosed)
+	{
+		TSharedRef<SActionMenu> ActionMenu = SNew(SActionMenu, InGraph)
+			.AutoExpandActionMenu(bAutoExpand)
+			.NewNodePosition(InNodePosition)
+			.DraggedFromPins(InDraggedPins)
+			.OnClosedCallback(InOnMenuClosed);
+
+		TSharedPtr<SWidget> FilterTextBox = StaticCastSharedRef<SWidget>(ActionMenu->GetFilterTextBox());
+		return FActionMenuContent(StaticCastSharedRef<SWidget>(ActionMenu), FilterTextBox);
+	});
+	GraphArgs.OnNodeTextCommitted = Workspace::FOnNodeTextCommitted::CreateLambda([](const Workspace::FWorkspaceEditorContext& InContext, const FText& NewText, ETextCommit::Type CommitInfo, UEdGraphNode* NodeBeingChanged)
+	{
+		URigVMEdGraph* RigVMEdGraph = Cast<URigVMEdGraph>(NodeBeingChanged->GetGraph());
+		if (RigVMEdGraph == nullptr)
 		{
-			TSharedRef<SActionMenu> ActionMenu = SNew(SActionMenu, InGraph)
-				.AutoExpandActionMenu(bAutoExpand)
-				.NewNodePosition(InNodePosition)
-				.DraggedFromPins(InDraggedPins)
-				.OnClosedCallback(InOnMenuClosed);
-
-			TSharedPtr<SWidget> FilterTextBox = StaticCastSharedRef<SWidget>(ActionMenu->GetFilterTextBox());
-			return FActionMenuContent(StaticCastSharedRef<SWidget>(ActionMenu), FilterTextBox);
-		});
-		GraphArgs.OnNodeTextCommitted = Workspace::FOnNodeTextCommitted::CreateLambda([](const Workspace::FWorkspaceEditorContext& InContext, const FText& NewText, ETextCommit::Type CommitInfo, UEdGraphNode* NodeBeingChanged)
+			return;
+		}
+
+		UEdGraphNode_Comment* CommentBeingChanged = Cast<UEdGraphNode_Comment>(NodeBeingChanged);
+		if (CommentBeingChanged == nullptr)
 		{
-			URigVMEdGraph* RigVMEdGraph = Cast<URigVMEdGraph>(NodeBeingChanged->GetGraph());
-			if (RigVMEdGraph == nullptr)
-			{
-				return;
-			}
+			return;
+		}
 
-			UEdGraphNode_Comment* CommentBeingChanged = Cast<UEdGraphNode_Comment>(NodeBeingChanged);
-			if (CommentBeingChanged == nullptr)
-			{
-				return;
-			}
+		RigVMEdGraph->GetController()->SetCommentTextByName(CommentBeingChanged->GetFName(), NewText.ToString(), CommentBeingChanged->FontSize, CommentBeingChanged->bCommentBubbleVisible, CommentBeingChanged->bColorCommentBubble, true, true);
+	});
+	GraphArgs.OnCanDeleteSelectedNodes = Workspace::FOnCanPerformActionOnSelectedNodes::CreateLambda([](const Workspace::FWorkspaceEditorContext& InContext, const FGraphPanelSelectionSet& InSelectedNodes)
+	{
+		bool bCanUserDeleteNode = false;
 
-			RigVMEdGraph->GetController()->SetCommentTextByName(CommentBeingChanged->GetFName(), NewText.ToString(), CommentBeingChanged->FontSize, CommentBeingChanged->bCommentBubbleVisible, CommentBeingChanged->bColorCommentBubble, true, true);
-		});
-		GraphArgs.OnCanDeleteSelectedNodes = Workspace::FOnCanPerformActionOnSelectedNodes::CreateLambda([](const Workspace::FWorkspaceEditorContext& InContext, const FGraphPanelSelectionSet& InSelectedNodes)
+		if(InSelectedNodes.Num() > 0)
 		{
-			bool bCanUserDeleteNode = false;
-
-			if(InSelectedNodes.Num() > 0)
+			for(UObject* NodeObject : InSelectedNodes)
 			{
-				for(UObject* NodeObject : InSelectedNodes)
+				// If any nodes allow deleting, then do not disable the delete option
+				const UEdGraphNode* Node = Cast<UEdGraphNode>(NodeObject);
+				if(Node && Node->CanUserDeleteNode())
 				{
-					// If any nodes allow deleting, then do not disable the delete option
-					const UEdGraphNode* Node = Cast<UEdGraphNode>(NodeObject);
-					if(Node && Node->CanUserDeleteNode())
-					{
-						bCanUserDeleteNode = true;
-						break;
-					}
+					bCanUserDeleteNode = true;
+					break;
 				}
 			}
+		}
 
-			return bCanUserDeleteNode;
-		});
-		GraphArgs.OnDeleteSelectedNodes = Workspace::FOnPerformActionOnSelectedNodes::CreateLambda([](const Workspace::FWorkspaceEditorContext& InContext, const FGraphPanelSelectionSet& InSelectedNodes)
+		return bCanUserDeleteNode;
+	});
+	GraphArgs.OnDeleteSelectedNodes = Workspace::FOnPerformActionOnSelectedNodes::CreateLambda([](const Workspace::FWorkspaceEditorContext& InContext, const FGraphPanelSelectionSet& InSelectedNodes)
+	{
+		if(InSelectedNodes.IsEmpty())
 		{
-			if(InSelectedNodes.IsEmpty())
-			{
-				return;
-			}
+			return;
+		}
 
-			URigVMController* Controller = nullptr;
-			
-			bool bRelinkPins = false;
-			TArray<URigVMNode*> NodesToRemove;
+		URigVMController* Controller = nullptr;
+		
+		bool bRelinkPins = false;
+		TArray<URigVMNode*> NodesToRemove;
 
-			for (FGraphPanelSelectionSet::TConstIterator NodeIt(InSelectedNodes); NodeIt; ++NodeIt)
+		for (FGraphPanelSelectionSet::TConstIterator NodeIt(InSelectedNodes); NodeIt; ++NodeIt)
+		{
+			if (UEdGraphNode* Node = Cast<UEdGraphNode>(*NodeIt))
 			{
-				if (UEdGraphNode* Node = Cast<UEdGraphNode>(*NodeIt))
+				URigVMEdGraph* RigVMEdGraph = Cast<URigVMEdGraph>(Node->GetGraph());
+				if (RigVMEdGraph == nullptr)
 				{
-					URigVMEdGraph* RigVMEdGraph = Cast<URigVMEdGraph>(Node->GetGraph());
-					if (RigVMEdGraph == nullptr)
-					{
-						continue;
-					}
-					
-					if (Node->CanUserDeleteNode())
+					continue;
+				}
+				
+				if (Node->CanUserDeleteNode())
+				{
+					if (const URigVMEdGraphNode* RigVMEdGraphNode = Cast<URigVMEdGraphNode>(Node))
 					{
-						if (const URigVMEdGraphNode* RigVMEdGraphNode = Cast<URigVMEdGraphNode>(Node))
+						if(Controller == nullptr)
 						{
-							if(Controller == nullptr)
-							{
-								Controller = RigVMEdGraphNode->GetController();
-							}
+							Controller = RigVMEdGraphNode->GetController();
+						}
 
-							bRelinkPins = bRelinkPins || FSlateApplication::Get().GetModifierKeys().IsShiftDown();
+						bRelinkPins = bRelinkPins || FSlateApplication::Get().GetModifierKeys().IsShiftDown();
 
-							if(URigVMGraph* Model = RigVMEdGraph->GetModel())
+						if(URigVMGraph* Model = RigVMEdGraph->GetModel())
+						{
+							if(URigVMNode* ModelNode = Model->FindNodeByName(*RigVMEdGraphNode->GetModelNodePath()))
 							{
-								if(URigVMNode* ModelNode = Model->FindNodeByName(*RigVMEdGraphNode->GetModelNodePath()))
-								{
-									NodesToRemove.Add(ModelNode);
-								}
+								NodesToRemove.Add(ModelNode);
 							}
 						}
-						else if (const UEdGraphNode_Comment* CommentNode = Cast<UEdGraphNode_Comment>(Node))
+					}
+					else if (const UEdGraphNode_Comment* CommentNode = Cast<UEdGraphNode_Comment>(Node))
+					{
+						if(URigVMGraph* Model = RigVMEdGraph->GetModel())
 						{
-							if(URigVMGraph* Model = RigVMEdGraph->GetModel())
+							if(URigVMNode* ModelNode = Model->FindNodeByName(CommentNode->GetFName()))
 							{
-								if(URigVMNode* ModelNode = Model->FindNodeByName(CommentNode->GetFName()))
-								{
-									NodesToRemove.Add(ModelNode);
-								}
+								NodesToRemove.Add(ModelNode);
 							}
 						}
-						else
-						{
-							Node->GetGraph()->RemoveNode(Node);
-						}
+					}
+					else
+					{
+						Node->GetGraph()->RemoveNode(Node);
 					}
 				}
 			}
+		}
+
+		if(NodesToRemove.IsEmpty() || Controller == nullptr)
+		{
+			return;
+		}
+
+		Controller->OpenUndoBracket(TEXT("Delete selected nodes"));
+		if(bRelinkPins && NodesToRemove.Num() == 1)
+		{
+			Controller->RelinkSourceAndTargetPins(NodesToRemove[0], true);;
+		}
+		Controller->RemoveNodes(NodesToRemove, true);
+		Controller->CloseUndoBracket();
+	});
+	GraphArgs.OnCanCopySelectedNodes = Workspace::FOnCanPerformActionOnSelectedNodes::CreateLambda([](const Workspace::FWorkspaceEditorContext& InContext, const FGraphPanelSelectionSet& InSelectedNodes)
+	{
+		bool bCanUserCopyNode = false;
+
+		if(InSelectedNodes.Num() > 0)
+		{
+			bCanUserCopyNode = true;
+		}
+
+		return bCanUserCopyNode;
+	});
+	GraphArgs.OnCopySelectedNodes = Workspace::FOnPerformActionOnSelectedNodes::CreateLambda([](const Workspace::FWorkspaceEditorContext& InContext, const FGraphPanelSelectionSet& InSelectedNodes)
+	{
+		if(InSelectedNodes.IsEmpty())
+		{
+			return;
+		}
+
+		URigVMEdGraph* RigVMEdGraph = Cast<URigVMEdGraph>(InContext.Object);
+		if (RigVMEdGraph == nullptr)
+		{
+			return;
+		}
+
+		URigVMController* Controller = RigVMEdGraph->GetController();
+
+		FString ExportedText = Controller->ExportSelectedNodesToText();
+		FPlatformApplicationMisc::ClipboardCopy(*ExportedText);
+	});
+	GraphArgs.OnCanPasteNodes = Workspace::FOnCanPasteNodes::CreateLambda([](const Workspace::FWorkspaceEditorContext& InContext, const FString& InImportData)
+	{
+		bool bCanUserImportNodes = false;
+
+		if (!InImportData.IsEmpty())
+		{
+			bCanUserImportNodes = true;
+		}
 
-			if(NodesToRemove.IsEmpty() || Controller == nullptr)
+		return bCanUserImportNodes;
+	});
+	GraphArgs.OnPasteNodes = Workspace::FOnPasteNodes::CreateLambda([](const Workspace::FWorkspaceEditorContext& InContext, const FVector2D& InPasteLocation, const FString& InImportData)
+	{
+		if(InImportData.IsEmpty())
+		{
+			return;
+		}
+
+		URigVMEdGraph* RigVMEdGraph = Cast<URigVMEdGraph>(InContext.Object);
+		if (RigVMEdGraph == nullptr)
+		{
+			return;
+		}
+
+		if (IRigVMClientHost* RigVMClientHost = RigVMEdGraph->GetImplementingOuter<IRigVMClientHost>())
+		{
+			FString TextToImport;
+			FPlatformApplicationMisc::ClipboardPaste(TextToImport);
+			URigVMController* Controller = RigVMEdGraph->GetController();
+
+			Controller->OpenUndoBracket(TEXT("Pasted Nodes."));
+
+			if (UE::RigVM::Editor::Tools::PasteNodes(InPasteLocation, TextToImport, Controller, RigVMEdGraph->GetModel(), RigVMClientHost->GetLocalFunctionLibrary(), RigVMClientHost->GetRigVMGraphFunctionHost()))
 			{
-				return;
+				Controller->CloseUndoBracket();
 			}
-
-			Controller->OpenUndoBracket(TEXT("Delete selected nodes"));
-			if(bRelinkPins && NodesToRemove.Num() == 1)
+			else
 			{
-				Controller->RelinkSourceAndTargetPins(NodesToRemove[0], true);;
+				Controller->CancelUndoBracket();
 			}
-			Controller->RemoveNodes(NodesToRemove, true);
-			Controller->CloseUndoBracket();
-		});
-		GraphArgs.OnCanCopySelectedNodes = Workspace::FOnCanPerformActionOnSelectedNodes::CreateLambda([](const Workspace::FWorkspaceEditorContext& InContext, const FGraphPanelSelectionSet& InSelectedNodes)
+		}
+	});
+	UE::Workspace::FOnCanPerformActionOnSelectedNodes& OnCanCopySelectedNodes = GraphArgs.OnCanCopySelectedNodes;
+	UE::Workspace::FOnCanPerformActionOnSelectedNodes& OnCanDeleteSelectedNodes = GraphArgs.OnCanDeleteSelectedNodes;
+	GraphArgs.OnCanCutSelectedNodes = Workspace::FOnCanPerformActionOnSelectedNodes::CreateLambda([OnCanCopySelectedNodes, OnCanDeleteSelectedNodes](const Workspace::FWorkspaceEditorContext& InContext, const FGraphPanelSelectionSet& InSelectedNodes)
 		{
 			bool bCanUserCopyNode = false;
 
-			if(InSelectedNodes.Num() > 0)
+			if (OnCanCopySelectedNodes.IsBound() && OnCanDeleteSelectedNodes.IsBound())
 			{
-				bCanUserCopyNode = true;
+				bCanUserCopyNode = OnCanCopySelectedNodes.Execute(InContext, InSelectedNodes) && OnCanDeleteSelectedNodes.Execute(InContext, InSelectedNodes);
 			}
 
 			return bCanUserCopyNode;
 		});
-		GraphArgs.OnCopySelectedNodes = Workspace::FOnPerformActionOnSelectedNodes::CreateLambda([](const Workspace::FWorkspaceEditorContext& InContext, const FGraphPanelSelectionSet& InSelectedNodes)
+	UE::Workspace::FOnPerformActionOnSelectedNodes& OnCopySelectedNodes = GraphArgs.OnCopySelectedNodes;
+	UE::Workspace::FOnPerformActionOnSelectedNodes& OnDeleteSelectedNodes = GraphArgs.OnDeleteSelectedNodes;
+	GraphArgs.OnCutSelectedNodes = Workspace::FOnPerformActionOnSelectedNodes::CreateLambda([OnCopySelectedNodes, OnDeleteSelectedNodes](const Workspace::FWorkspaceEditorContext& InContext, const FGraphPanelSelectionSet& InSelectedNodes)
 		{
-			if(InSelectedNodes.IsEmpty())
+			if (InSelectedNodes.IsEmpty())
 			{
 				return;
 			}
@@ -510,25 +619,37 @@ class FModule : public IModule
 				return;
 			}
 
-			URigVMController* Controller = RigVMEdGraph->GetController();
 
-			FString ExportedText = Controller->ExportSelectedNodesToText();
-			FPlatformApplicationMisc::ClipboardCopy(*ExportedText);
+			if (OnCopySelectedNodes.IsBound() && OnDeleteSelectedNodes.IsBound())
+			{
+				URigVMController* Controller = RigVMEdGraph->GetController();
+
+				OnCopySelectedNodes.Execute(InContext, InSelectedNodes);
+
+				Controller->OpenUndoBracket(TEXT("Cut Nodes."));
+				OnDeleteSelectedNodes.Execute(InContext, InSelectedNodes);
+				Controller->CloseUndoBracket();
+			}
 		});
-		GraphArgs.OnCanPasteNodes = Workspace::FOnCanPasteNodes::CreateLambda([](const Workspace::FWorkspaceEditorContext& InContext, const FString& InImportData)
+	UE::Workspace::FOnCanPasteNodes& OnCanPasteNodes = GraphArgs.OnCanPasteNodes;
+	GraphArgs.OnCanDuplicateSelectedNodes = Workspace::FOnCanPerformActionOnSelectedNodes::CreateLambda([OnCanCopySelectedNodes, OnCanPasteNodes](const Workspace::FWorkspaceEditorContext& InContext, const FGraphPanelSelectionSet& InSelectedNodes)
 		{
-			bool bCanUserImportNodes = false;
+			bool bCanUserCopyNode = false;
 
-			if (!InImportData.IsEmpty())
+			if (OnCanCopySelectedNodes.IsBound() && OnCanPasteNodes.IsBound())
 			{
-				bCanUserImportNodes = true;
+				FString TextToImport;
+				FPlatformApplicationMisc::ClipboardPaste(TextToImport);
+
+				bCanUserCopyNode = OnCanCopySelectedNodes.Execute(InContext, InSelectedNodes) && OnCanPasteNodes.Execute(InContext, TextToImport);
 			}
 
-			return bCanUserImportNodes;
+			return bCanUserCopyNode;
 		});
-		GraphArgs.OnPasteNodes = Workspace::FOnPasteNodes::CreateLambda([](const Workspace::FWorkspaceEditorContext& InContext, const FVector2D& InPasteLocation, const FString& InImportData)
+	UE::Workspace::FOnPasteNodes& OnPasteNodes = GraphArgs.OnPasteNodes;
+	GraphArgs.OnDuplicateSelectedNodes = Workspace::FOnDuplicateSelectedNodes::CreateLambda([OnCopySelectedNodes, OnPasteNodes](const Workspace::FWorkspaceEditorContext& InContext, const FVector2D& InPasteLocation, const FGraphPanelSelectionSet& InSelectedNodes)
 		{
-			if(InImportData.IsEmpty())
+			if (InSelectedNodes.IsEmpty())
 			{
 				return;
 			}
@@ -539,198 +660,101 @@ class FModule : public IModule
 				return;
 			}
 
-			if (IRigVMClientHost* RigVMClientHost = RigVMEdGraph->GetImplementingOuter<IRigVMClientHost>())
+			if (OnCopySelectedNodes.IsBound() && OnPasteNodes.IsBound())
 			{
+				OnCopySelectedNodes.Execute(InContext, InSelectedNodes);
+
 				FString TextToImport;
 				FPlatformApplicationMisc::ClipboardPaste(TextToImport);
-				URigVMController* Controller = RigVMEdGraph->GetController();
 
-				Controller->OpenUndoBracket(TEXT("Pasted Nodes."));
+				URigVMController* Controller = RigVMEdGraph->GetController();
 
-				if (UE::RigVM::Editor::Tools::PasteNodes(InPasteLocation, TextToImport, Controller, RigVMEdGraph->GetModel(), RigVMClientHost->GetLocalFunctionLibrary(), RigVMClientHost->GetRigVMGraphFunctionHost()))
-				{
-					Controller->CloseUndoBracket();
-				}
-				else
-				{
-					Controller->CancelUndoBracket();
-				}
+				Controller->OpenUndoBracket(TEXT("Duplicate Nodes."));
+				OnPasteNodes.Execute(InContext, InPasteLocation, TextToImport);
+				Controller->CloseUndoBracket();
 			}
 		});
-		UE::Workspace::FOnCanPerformActionOnSelectedNodes& OnCanCopySelectedNodes = GraphArgs.OnCanCopySelectedNodes;
-		UE::Workspace::FOnCanPerformActionOnSelectedNodes& OnCanDeleteSelectedNodes = GraphArgs.OnCanDeleteSelectedNodes;
-		GraphArgs.OnCanCutSelectedNodes = Workspace::FOnCanPerformActionOnSelectedNodes::CreateLambda([OnCanCopySelectedNodes, OnCanDeleteSelectedNodes](const Workspace::FWorkspaceEditorContext& InContext, const FGraphPanelSelectionSet& InSelectedNodes)
-			{
-				bool bCanUserCopyNode = false;
-
-				if (OnCanCopySelectedNodes.IsBound() && OnCanDeleteSelectedNodes.IsBound())
-				{
-					bCanUserCopyNode = OnCanCopySelectedNodes.Execute(InContext, InSelectedNodes) && OnCanDeleteSelectedNodes.Execute(InContext, InSelectedNodes);
-				}
-
-				return bCanUserCopyNode;
-			});
-		UE::Workspace::FOnPerformActionOnSelectedNodes& OnCopySelectedNodes = GraphArgs.OnCopySelectedNodes;
-		UE::Workspace::FOnPerformActionOnSelectedNodes& OnDeleteSelectedNodes = GraphArgs.OnDeleteSelectedNodes;
-		GraphArgs.OnCutSelectedNodes = Workspace::FOnPerformActionOnSelectedNodes::CreateLambda([OnCopySelectedNodes, OnDeleteSelectedNodes](const Workspace::FWorkspaceEditorContext& InContext, const FGraphPanelSelectionSet& InSelectedNodes)
-			{
-				if (InSelectedNodes.IsEmpty())
-				{
-					return;
-				}
-
-				URigVMEdGraph* RigVMEdGraph = Cast<URigVMEdGraph>(InContext.Object);
-				if (RigVMEdGraph == nullptr)
-				{
-					return;
-				}
-
-
-				if (OnCopySelectedNodes.IsBound() && OnDeleteSelectedNodes.IsBound())
-				{
-					URigVMController* Controller = RigVMEdGraph->GetController();
-
-					OnCopySelectedNodes.Execute(InContext, InSelectedNodes);
-
-					Controller->OpenUndoBracket(TEXT("Cut Nodes."));
-					OnDeleteSelectedNodes.Execute(InContext, InSelectedNodes);
-					Controller->CloseUndoBracket();
-				}
-			});
-		UE::Workspace::FOnCanPasteNodes& OnCanPasteNodes = GraphArgs.OnCanPasteNodes;
-		GraphArgs.OnCanDuplicateSelectedNodes = Workspace::FOnCanPerformActionOnSelectedNodes::CreateLambda([OnCanCopySelectedNodes, OnCanPasteNodes](const Workspace::FWorkspaceEditorContext& InContext, const FGraphPanelSelectionSet& InSelectedNodes)
-			{
-				bool bCanUserCopyNode = false;
-
-				if (OnCanCopySelectedNodes.IsBound() && OnCanPasteNodes.IsBound())
-				{
-					FString TextToImport;
-					FPlatformApplicationMisc::ClipboardPaste(TextToImport);
-
-					bCanUserCopyNode = OnCanCopySelectedNodes.Execute(InContext, InSelectedNodes) && OnCanPasteNodes.Execute(InContext, TextToImport);
-				}
-
-				return bCanUserCopyNode;
-			});
-		UE::Workspace::FOnPasteNodes& OnPasteNodes = GraphArgs.OnPasteNodes;
-		GraphArgs.OnDuplicateSelectedNodes = Workspace::FOnDuplicateSelectedNodes::CreateLambda([OnCopySelectedNodes, OnPasteNodes](const Workspace::FWorkspaceEditorContext& InContext, const FVector2D& InPasteLocation, const FGraphPanelSelectionSet& InSelectedNodes)
-			{
-				if (InSelectedNodes.IsEmpty())
-				{
-					return;
-				}
-
-				URigVMEdGraph* RigVMEdGraph = Cast<URigVMEdGraph>(InContext.Object);
-				if (RigVMEdGraph == nullptr)
-				{
-					return;
-				}
-
-				if (OnCopySelectedNodes.IsBound() && OnPasteNodes.IsBound())
-				{
-					OnCopySelectedNodes.Execute(InContext, InSelectedNodes);
+	GraphArgs.OnGraphSelectionChanged = Workspace::FOnGraphSelectionChanged::CreateLambda([](const Workspace::FWorkspaceEditorContext& InContext, const FGraphPanelSelectionSet& NewSelection)
+	{
+		URigVMEdGraph* RigVMEdGraph = Cast<URigVMEdGraph>(InContext.Object);
+		if (RigVMEdGraph == nullptr)
+		{
+			return;
+		}
 
-					FString TextToImport;
-					FPlatformApplicationMisc::ClipboardPaste(TextToImport);
+		if (RigVMEdGraph->bIsSelecting || GIsTransacting)
+		{
+			return;
+		}
 
-					URigVMController* Controller = RigVMEdGraph->GetController();
+		TGuardValue<bool> SelectGuard(RigVMEdGraph->bIsSelecting, true);
 
-					Controller->OpenUndoBracket(TEXT("Duplicate Nodes."));
-					OnPasteNodes.Execute(InContext, InPasteLocation, TextToImport);
-					Controller->CloseUndoBracket();
-				}
-			});
-		GraphArgs.OnGraphSelectionChanged = Workspace::FOnGraphSelectionChanged::CreateLambda([](const Workspace::FWorkspaceEditorContext& InContext, const FGraphPanelSelectionSet& NewSelection)
+		TArray<FName> NodeNamesToSelect;
+		for (UObject* Object : NewSelection)
 		{
-			URigVMEdGraph* RigVMEdGraph = Cast<URigVMEdGraph>(InContext.Object);
-			if (RigVMEdGraph == nullptr)
+			if (URigVMEdGraphNode* RigVMEdGraphNode = Cast<URigVMEdGraphNode>(Object))
 			{
-				return;
+				NodeNamesToSelect.Add(RigVMEdGraphNode->GetModelNodeName());
 			}
-
-			if (RigVMEdGraph->bIsSelecting || GIsTransacting)
+			else if(UEdGraphNode* Node = Cast<UEdGraphNode>(Object))
 			{
-				return;
+				NodeNamesToSelect.Add(Node->GetFName());
 			}
+		}
+		RigVMEdGraph->GetController()->SetNodeSelection(NodeNamesToSelect, true, true);
 
-			TGuardValue<bool> SelectGuard(RigVMEdGraph->bIsSelecting, true);
-
-			TArray<FName> NodeNamesToSelect;
-			for (UObject* Object : NewSelection)
-			{
-				if (URigVMEdGraphNode* RigVMEdGraphNode = Cast<URigVMEdGraphNode>(Object))
-				{
-					NodeNamesToSelect.Add(RigVMEdGraphNode->GetModelNodeName());
-				}
-				else if(UEdGraphNode* Node = Cast<UEdGraphNode>(Object))
-				{
-					NodeNamesToSelect.Add(Node->GetFName());
-				}
-			}
-			RigVMEdGraph->GetController()->SetNodeSelection(NodeNamesToSelect, true, true);
-
-			InContext.WorkspaceEditor->SetDetailsObjects(NewSelection.Array());
-		});
+		InContext.WorkspaceEditor->SetDetailsObjects(NewSelection.Array());
+	});
 
-		Workspace::FObjectDocumentArgs GraphDocumentArgs = WorkspaceEditorModule.CreateGraphDocumentArgs(GraphArgs);
-		GraphDocumentArgs.OnGetDocumentBreadcrumbTrail = Workspace::FOnGetDocumentBreadcrumbTrail::CreateLambda([](const Workspace::FWorkspaceEditorContext& InContext, TArray<TSharedPtr<Workspace::FWorkspaceBreadcrumb>>& OutBreadcrumbs)
+	Workspace::FObjectDocumentArgs GraphDocumentArgs = WorkspaceEditorModule.CreateGraphDocumentArgs(GraphArgs);
+	GraphDocumentArgs.OnGetDocumentBreadcrumbTrail = Workspace::FOnGetDocumentBreadcrumbTrail::CreateLambda([](const Workspace::FWorkspaceEditorContext& InContext, TArray<TSharedPtr<Workspace::FWorkspaceBreadcrumb>>& OutBreadcrumbs)
+	{
+		if (const URigVMEdGraph* RigVMEdGraph = Cast<URigVMEdGraph>(InContext.Object))
 		{
-			if (const URigVMEdGraph* RigVMEdGraph = Cast<URigVMEdGraph>(InContext.Object))
-			{
-				const TSharedPtr<Workspace::FWorkspaceBreadcrumb>& GraphCrumb = OutBreadcrumbs.Add_GetRef(MakeShared<Workspace::FWorkspaceBreadcrumb>());
-				GraphCrumb->OnGetLabel = Workspace::FWorkspaceBreadcrumb::FOnGetBreadcrumbLabel::CreateLambda([GraphName = RigVMEdGraph->GetFName()]{ return FText::FromName(GraphName); });
+			const TSharedPtr<Workspace::FWorkspaceBreadcrumb>& GraphCrumb = OutBreadcrumbs.Add_GetRef(MakeShared<Workspace::FWorkspaceBreadcrumb>());
+			GraphCrumb->OnGetLabel = Workspace::FWorkspaceBreadcrumb::FOnGetBreadcrumbLabel::CreateLambda([GraphName = RigVMEdGraph->GetFName()]{ return FText::FromName(GraphName); });
 
-				if (const UAnimNextRigVMAssetEntry* OuterAssetEntry = CastChecked<UAnimNextRigVMAssetEntry>(RigVMEdGraph->GetOuter()))
+			if (const UAnimNextRigVMAssetEntry* OuterAssetEntry = CastChecked<UAnimNextRigVMAssetEntry>(RigVMEdGraph->GetOuter()))
+			{
+				if (UAnimNextGraph_EditorData* OuterEditorData = CastChecked<UAnimNextGraph_EditorData>(OuterAssetEntry->GetOuter()))
 				{
-					if (UAnimNextGraph_EditorData* OuterEditorData = CastChecked<UAnimNextGraph_EditorData>(OuterAssetEntry->GetOuter()))
+					if(UAnimNextGraph* OuterGraph = UncookedOnly::FUtils::GetGraph(OuterEditorData))
 					{
-						if(UAnimNextGraph* OuterGraph = UncookedOnly::FUtils::GetGraph(OuterEditorData))
-						{
-							const TSharedPtr<Workspace::FWorkspaceBreadcrumb>& OuterGraphCrumb = OutBreadcrumbs.Add_GetRef(MakeShared<Workspace::FWorkspaceBreadcrumb>());
-							TWeakObjectPtr<UAnimNextGraph> WeakOuterGraph = OuterGraph;
-							TWeakPtr<Workspace::IWorkspaceEditor> WeakWorkspaceEditor = InContext.WorkspaceEditor;
-							OuterGraphCrumb->OnGetLabel = Workspace::FWorkspaceBreadcrumb::FOnGetBreadcrumbLabel::CreateLambda([GraphName = OuterGraph->GetFName()]{ return FText::FromName(GraphName); });
-							OuterGraphCrumb->OnClicked = Workspace::FWorkspaceBreadcrumb::FOnBreadcrumbClicked::CreateLambda(
-								[WeakOuterGraph, WeakWorkspaceEditor]
+						const TSharedPtr<Workspace::FWorkspaceBreadcrumb>& OuterGraphCrumb = OutBreadcrumbs.Add_GetRef(MakeShared<Workspace::FWorkspaceBreadcrumb>());
+						TWeakObjectPtr<UAnimNextGraph> WeakOuterGraph = OuterGraph;
+						TWeakPtr<Workspace::IWorkspaceEditor> WeakWorkspaceEditor = InContext.WorkspaceEditor;
+						OuterGraphCrumb->OnGetLabel = Workspace::FWorkspaceBreadcrumb::FOnGetBreadcrumbLabel::CreateLambda([GraphName = OuterGraph->GetFName()]{ return FText::FromName(GraphName); });
+						OuterGraphCrumb->OnClicked = Workspace::FWorkspaceBreadcrumb::FOnBreadcrumbClicked::CreateLambda(
+							[WeakOuterGraph, WeakWorkspaceEditor]
+							{
+								if (const TSharedPtr<Workspace::IWorkspaceEditor> SharedWorkspaceEditor = WeakWorkspaceEditor.Pin())
 								{
-									if (const TSharedPtr<Workspace::IWorkspaceEditor> SharedWorkspaceEditor = WeakWorkspaceEditor.Pin())
-									{
-										SharedWorkspaceEditor->OpenObjects({WeakOuterGraph.Get()});
-									}
+									SharedWorkspaceEditor->OpenObjects({WeakOuterGraph.Get()});
 								}
-							);
-						}
+							}
+						);
 					}
 				}
 			}
-		});
+		}
+	});
 
-		WorkspaceEditorModule.RegisterObjectDocumentType(FTopLevelAssetPath(TEXT("/Script/AnimNextUncookedOnly.AnimNextGraph_EdGraph")), GraphDocumentArgs);
-	}
+	WorkspaceEditorModule.RegisterObjectDocumentType(FTopLevelAssetPath(TEXT("/Script/AnimNextUncookedOnly.AnimNextGraph_EdGraph")), GraphDocumentArgs);
+}
 
-	void UnregisterWorkspaceDocumentTypes()
+void FModule::UnregisterWorkspaceDocumentTypes()
+{
+	if(FModuleManager::Get().IsModuleLoaded("WorkspaceEditor"))
 	{
-		if(FModuleManager::Get().IsModuleLoaded("WorkspaceEditor"))
-		{
-			Workspace::IWorkspaceEditorModule& WorkspaceEditorModule = FModuleManager::LoadModuleChecked<Workspace::IWorkspaceEditorModule>("PropertyEditor");
-			WorkspaceEditorModule.UnregisterObjectDocumentType(FTopLevelAssetPath(TEXT("/Script/AnimNext.AnimNextSchedule")));
-			WorkspaceEditorModule.UnregisterObjectDocumentType(FTopLevelAssetPath(TEXT("/Script/AnimNext.AnimNextGraph")));
-			WorkspaceEditorModule.UnregisterObjectDocumentType(FTopLevelAssetPath(TEXT("/Script/AnimNextUncookedOnly.AnimNextGraph_EdGraph")));
-		}
+		Workspace::IWorkspaceEditorModule& WorkspaceEditorModule = FModuleManager::LoadModuleChecked<Workspace::IWorkspaceEditorModule>("PropertyEditor");
+		WorkspaceEditorModule.UnregisterObjectDocumentType(FTopLevelAssetPath(TEXT("/Script/AnimNext.AnimNextSchedule")));
+		WorkspaceEditorModule.UnregisterObjectDocumentType(FTopLevelAssetPath(TEXT("/Script/AnimNext.AnimNextGraph")));
+		WorkspaceEditorModule.UnregisterObjectDocumentType(FTopLevelAssetPath(TEXT("/Script/AnimNextUncookedOnly.AnimNextGraph_EdGraph")));
 	}
-	
-	/** Node factory for the AnimNext graph */
-	TSharedPtr<FAnimNextGraphPanelNodeFactory> AnimNextGraphPanelNodeFactory;
-	
-	/** Pin factory for parameters */
-	TSharedPtr<FParametersGraphPanelPinFactory> ParametersGraphPanelPinFactory;
-
-	/** Type identifier for parameter names */
-	TSharedPtr<FParamNamePropertyTypeIdentifier> Identifier;
-};
+}
 
 }
 
 IMPLEMENT_MODULE(UE::AnimNext::Editor::FModule, AnimNextEditor);
 
-#undef LOCTEXT_NAMESPACE
\ No newline at end of file
+#undef LOCTEXT_NAMESPACE
diff --git a/Engine/Plugins/Experimental/AnimNext/Source/AnimNextEditor/Private/AnimNextEditorModule.h b/Engine/Plugins/Experimental/AnimNext/Source/AnimNextEditor/Private/AnimNextEditorModule.h
new file mode 100644
--- /dev/null
+++ b/Engine/Plugins/Experimental/AnimNext/Source/AnimNextEditor/Private/AnimNextEditorModule.h
@@ -0,0 +1,56 @@
+﻿// Copyright Epic Games, Inc. All Rights Reserved.
+
+#pragma once
+
+#include "IAnimNextEditorModule.h"
+
+class FAnimNextGraphPanelNodeFactory;
+
+namespace UE::AnimNext::Editor
+{
+	class FParamNamePropertyTypeIdentifier;
+	class FParametersGraphPanelPinFactory;
+	class FLocatorContext;
+}
+
+namespace UE::Workspace
+{
+	class IWorkspaceEditorModule;
+}
+
+namespace UE::AnimNext::Editor
+{
+
+class FModule : public IModule
+{
+
+private:
+	virtual void StartupModule() override;
+	virtual void ShutdownModule() override;
+
+	// IAnimNextEditorModule interface
+	virtual TSharedRef<SWidget> CreateParameterPicker(const FParameterPickerArgs& InArgs) override;
+	virtual void RegisterLocatorFragmentEditorType(FName InLocatorFragmentEditorName) override;
+	virtual void UnregisterLocatorFragmentEditorType(FName InLocatorFragmentEditorName) override;
+
+private:
+	void RegisterWorkspaceDocumentTypes(Workspace::IWorkspaceEditorModule& WorkspaceEditorModule);
+	void UnregisterWorkspaceDocumentTypes();
+
+private:
+	/** Node factory for the AnimNext graph */
+	TSharedPtr<FAnimNextGraphPanelNodeFactory> AnimNextGraphPanelNodeFactory;
+	
+	/** Pin factory for parameters */
+	TSharedPtr<FParametersGraphPanelPinFactory> ParametersGraphPanelPinFactory;
+
+	/** Type identifier for parameter names */
+	TSharedPtr<FParamNamePropertyTypeIdentifier> Identifier;
+
+	/** Registered names for locator fragments */
+	TSet<FName> LocatorFragmentEditorNames;
+
+	friend class FLocatorContext;
+};
+
+}
diff --git a/Engine/Plugins/Experimental/AnimNext/Source/AnimNextEditor/Private/Common/SRigVMAssetView.cpp b/Engine/Plugins/Experimental/AnimNext/Source/AnimNextEditor/Private/Common/SRigVMAssetView.cpp
--- a/Engine/Plugins/Experimental/AnimNext/Source/AnimNextEditor/Private/Common/SRigVMAssetView.cpp
+++ b/Engine/Plugins/Experimental/AnimNext/Source/AnimNextEditor/Private/Common/SRigVMAssetView.cpp
@@ -11,6 +11,7 @@
 #include "ISourceControlModule.h"
 #include "ISourceControlProvider.h"
 #include "EditorUtils.h"
+#include "IAnimNextRigVMExportInterface.h"
 #include "IAnimNextRigVMParameterInterface.h"
 #include "InstancedPropertyBagStructureDataProvider.h"
 #include "UncookedOnlyUtils.h"
@@ -28,6 +29,7 @@
 #include "ISinglePropertyView.h"
 #include "ObjectEditorUtils.h"
 #include "Misc/NotifyHook.h"
+#include "Widgets/Input/SButton.h"
 
 #define LOCTEXT_NAMESPACE "SRigVMAssetView"
 
@@ -39,10 +41,10 @@ namespace RigVMAssetView
 
 static FName ContextMenuName(TEXT("AnimNext.RigVMAssetView.ContextMenu"));
 static FName Column_RevisionControl(TEXT("RevisionControl"));
-static FName Column_ModifiedStatus(TEXT("ModifiedStatus"));
 static FName Column_Name(TEXT("Name"));
 static FName Column_Type(TEXT("Type"));
 static FName Column_Value(TEXT("Value"));
+static FName Column_AccessSpecifier(TEXT("AccessSpecifier"));
 
 static const FName NAME_Category("Category");
 
@@ -264,7 +266,8 @@ void SRigVMAssetView::Construct(const FArguments& InArgs, UAnimNextRigVMAssetEdi
 			.OnSelectionChanged(this, &SRigVMAssetView::HandleSelectionChanged)
 			.HeaderRow(
 				SNew(SHeaderRow)
-				+SHeaderRow::Column(Column_RevisionControl)
+
+				+SHeaderRow::Column(Column_AccessSpecifier)
 				.DefaultLabel(FText::GetEmpty())
 				.FixedWidth(24.0f)
 				.HeaderContent()
@@ -277,12 +280,27 @@ void SRigVMAssetView::Construct(const FArguments& InArgs, UAnimNextRigVMAssetEdi
 					[
 						SNew(SImage)
 						.ColorAndOpacity(FSlateColor::UseForeground())
-						.Image(FRevisionControlStyleManager::Get().GetBrush("RevisionControl.Icon"))
-						.ToolTipText(LOCTEXT("RevisionControlStatusTooltip", "Revision control status of this parameter"))
+						.Image(FAppStyle::GetBrush("Level.VisibleIcon16x"))
+						.ToolTipText(LOCTEXT("AccessSpecifierTooltip", "Access level of this entry"))
 					]
 				]
 
-				+SHeaderRow::Column(Column_ModifiedStatus)
+				+SHeaderRow::Column(Column_Name)
+				.DefaultLabel(LOCTEXT("NameColumnHeader", "Name"))
+				.ToolTipText(LOCTEXT("NameColumnHeaderTooltip", "The name of the entry"))
+				.FillWidth(20.0f)
+
+				+SHeaderRow::Column(Column_Type)
+				.DefaultLabel(LOCTEXT("TypeColumnHeader", "Type"))
+				.ToolTipText(LOCTEXT("TypeColumnHeaderTooltip", "The type of the entry"))
+				.ManualWidth(145.0f)
+
+				+SHeaderRow::Column(Column_Value)
+				.DefaultLabel(LOCTEXT("ValueColumnHeader", "Value"))
+				.ToolTipText(LOCTEXT("ValueColumnHeaderTooltip", "The value or binding to the entry"))
+				.FillWidth(10.0f)
+				
+				+SHeaderRow::Column(Column_RevisionControl)
 				.DefaultLabel(FText::GetEmpty())
 				.FixedWidth(24.0f)
 				.HeaderContent()
@@ -295,25 +313,10 @@ void SRigVMAssetView::Construct(const FArguments& InArgs, UAnimNextRigVMAssetEdi
 					[
 						SNew(SImage)
 						.ColorAndOpacity(FSlateColor::UseForeground())
-						.Image(FAppStyle::GetBrush("ContentBrowser.ContentDirty"))
-						.ToolTipText(LOCTEXT("ModifiedStatusTooltip", "Modified status of this parameter"))
+						.Image(FRevisionControlStyleManager::Get().GetBrush("RevisionControl.Icon"))
+						.ToolTipText(LOCTEXT("RevisionControlStatusTooltip", "Revision control status of this entry"))
 					]
 				]
-
-				+SHeaderRow::Column(Column_Name)
-				.DefaultLabel(LOCTEXT("NameColumnHeader", "Name"))
-				.ToolTipText(LOCTEXT("NameColumnHeaderTooltip", "The name of the parameter"))
-				.FillWidth(20.0f)
-
-				+SHeaderRow::Column(Column_Type)
-				.DefaultLabel(LOCTEXT("TypeColumnHeader", "Type"))
-				.ToolTipText(LOCTEXT("TypeColumnHeaderTooltip", "The type of the parameter"))
-				.ManualWidth(145.0f)
-
-				+SHeaderRow::Column(Column_Value)
-				.DefaultLabel(LOCTEXT("ValueColumnHeader", "Value"))
-				.ToolTipText(LOCTEXT("ValueColumnHeaderTooltip", "The value or binding to the parameter"))
-				.FillWidth(10.0f)
 			)
 		]
 	];
@@ -397,12 +400,40 @@ static bool CompareGraphActionNode(TSharedRef<FRigVMAssetViewEntry> A, TSharedRe
 	return true;
 }
 
+void RestoreSelectionState(TSharedPtr<STreeView<TSharedRef<FRigVMAssetViewEntry>>> InTree, const TArray<TSharedRef<FRigVMAssetViewEntry>>& ItemSource, const TArray<TSharedRef<FRigVMAssetViewEntry>>& OldSelectionState)
+{
+	// Iterate over new tree items
+	TArray<TSharedRef<FRigVMAssetViewEntry>> NewSelection;
+	for (int32 ItemIdx = 0; ItemIdx < ItemSource.Num(); ItemIdx++)
+	{
+		const TSharedRef<FRigVMAssetViewEntry>& NewItem = ItemSource[ItemIdx];
+
+		// Look through old selection state
+		for (const TSharedRef<FRigVMAssetViewEntry>& OldSelection : OldSelectionState)
+		{
+			if( OldSelection->WeakEntry.IsValid() && NewItem->WeakEntry.IsValid() &&
+				OldSelection->WeakEntry->GetEntryName() == NewItem->WeakEntry->GetEntryName())
+			{
+				NewSelection.Add(NewItem);
+			}
+		}
+	}
+
+	if(NewSelection.Num() > 0)
+	{
+		InTree->SetItemSelection(NewSelection, true);
+	}
+}
+
 void SRigVMAssetView::RefreshEntries()
 {
 	// First, save off current expansion state
 	TSet< TSharedRef<FRigVMAssetViewEntry> > OldExpansionState;
 	EntriesList->GetExpandedItems(OldExpansionState);
 
+	TArray< TSharedRef<FRigVMAssetViewEntry> > OldSelectionState;
+	EntriesList->GetSelectedItems(OldSelectionState);
+
 	Entries.Reset();
 
 	TArray<TSharedRef<FRigVMAssetViewEntry>> EntriesToSelect;
@@ -428,14 +459,16 @@ void SRigVMAssetView::RefreshEntries()
 	}
 	RestoreExpansionState< TSharedRef<FRigVMAssetViewEntry> >(EntriesList, AllEntries, OldExpansionState, CompareGraphActionNode);
 
-	PendingSelection.Empty();
-
 	RefreshFilter();
 
 	if(EntriesToSelect.Num() > 0)
 	{
 		EntriesList->SetItemSelection(EntriesToSelect, true);
 	}
+	else
+	{
+		RestoreSelectionState(EntriesList, AllEntries, OldSelectionState);
+	}
 }
 
 void SRigVMAssetView::RefreshFilter()
@@ -606,7 +639,7 @@ class SRigVMAssetViewRow : public SMultiColumnTableRow<TSharedRef<FRigVMAssetVie
 		{
 			if (UAnimNextRigVMAssetEditorData* EditorData = View->EditorData)
 			{
-				if (UAnimNextRigVMAssetEntry* AssetEntry = EditorData->FindEntry(PropertyAboutToChange->GetFName()))
+				if (UAnimNextRigVMAssetEntry* AssetEntry = EditorData->FindEntry(UncookedOnly::FUtils::GetParameterNameFromQualifiedName(PropertyAboutToChange->GetFName())))
 				{
 					UAnimNextRigVMAsset* Asset = UE::AnimNext::UncookedOnly::FUtils::GetAsset(EditorData);
 
@@ -624,11 +657,14 @@ class SRigVMAssetViewRow : public SMultiColumnTableRow<TSharedRef<FRigVMAssetVie
 		{
 			if (UAnimNextRigVMAssetEditorData* EditorData = View->EditorData)
 			{
-				if (UAnimNextRigVMAssetEntry* AssetEntry = EditorData->FindEntry(PropertyChangedEvent.GetMemberPropertyName()))
+				if (UAnimNextRigVMAssetEntry* AssetEntry = EditorData->FindEntry(UncookedOnly::FUtils::GetParameterNameFromQualifiedName(PropertyChangedEvent.GetMemberPropertyName())))
 				{
 					// Needed to show the changed sign a the table when we modify the PropertyBag
 					// TODO: remove this by moving defaults into entries
-					AssetEntry->MarkPackageDirty(); 
+					AssetEntry->MarkPackageDirty();
+
+					// Ensure that default values get picked up and forwarded to compiler
+					EditorData->BroadcastModified();
 				}
 			}
 		}
@@ -638,84 +674,93 @@ class SRigVMAssetViewRow : public SMultiColumnTableRow<TSharedRef<FRigVMAssetVie
 	virtual TSharedRef<SWidget> GenerateWidgetForColumn(const FName& InColumnName) override
 	{
 		using namespace RigVMAssetView;
-		
-		if(InColumnName == Column_RevisionControl)
+
+		if(InColumnName == Column_AccessSpecifier)
 		{
 			return
 				SNew(SBox)
-				.HAlign(HAlign_Center)
+				.WidthOverride(16.0f)
+				.HeightOverride(16.0f)
 				.VAlign(VAlign_Center)
-				.HeightOverride(20.0f)
+				.HAlign(HAlign_Center)
 				[
-					SNew(SImage)
-					.ToolTipText_Lambda([this]()
+					SNew(SButton)
+					.Visibility_Lambda([this]()
 					{
-						if(UAnimNextRigVMAssetEntry* AssetEntry = Entry->WeakEntry.Get())
+						if(IAnimNextRigVMExportInterface* Export = Cast<IAnimNextRigVMExportInterface>(Entry->WeakEntry.Get()))
 						{
-							ISourceControlProvider& SourceControlProvider = ISourceControlModule::Get().GetProvider();
-							UPackage* Package = AssetEntry->GetExternalPackage();
-							if(FSourceControlStatePtr State = SourceControlProvider.GetState(Package, EStateCacheUsage::Use))
-							{
-								return FText::Format(LOCTEXT("RevisionControlStatusFormat", "File: {0}\nStatus: {1}"), FText::FromName(Package->GetFName()),  State->GetDisplayTooltip());
-							}
+							return Export->GetExportAccessSpecifier() == EAnimNextExportAccessSpecifier::Public || IsHovered() ? EVisibility::Visible : EVisibility::Collapsed;
 						}
-
-						return LOCTEXT("RevisionControlStatus", "Revision control status of this parameter");
+						return EVisibility::Collapsed;
 					})
-					.Image_Lambda([this]() -> const FSlateBrush*
+					.ButtonStyle(FAppStyle::Get(), "HoverHintOnly")
+					.OnClicked_Lambda([this]()
 					{
-						if(UAnimNextRigVMAssetEntry* AssetEntry = Entry->WeakEntry.Get())
+						if(IAnimNextRigVMExportInterface* Export = Cast<IAnimNextRigVMExportInterface>(Entry->WeakEntry.Get()))
 						{
-							ISourceControlProvider& SourceControlProvider = ISourceControlModule::Get().GetProvider();
-							if(FSourceControlStatePtr State = SourceControlProvider.GetState(AssetEntry->GetExternalPackage(), EStateCacheUsage::Use))
-							{
-								return State->GetIcon().GetSmallIcon();
-							}
+							Export->SetExportAccessSpecifier(Export->GetExportAccessSpecifier() == EAnimNextExportAccessSpecifier::Public ? EAnimNextExportAccessSpecifier::Private : EAnimNextExportAccessSpecifier::Public);  
 						}
-						return nullptr;
+						return FReply::Unhandled();	// Fall through so we dont deselect our item
 					})
+					.Content()
+					[
+						SNew(SImage)
+						.ColorAndOpacity(FSlateColor::UseForeground())
+						.Image_Lambda([this]() -> const FSlateBrush*
+						{
+							if(IAnimNextRigVMExportInterface* Export = Cast<IAnimNextRigVMExportInterface>(Entry->WeakEntry.Get()))
+							{
+								return Export->GetExportAccessSpecifier() == EAnimNextExportAccessSpecifier::Public ? FAppStyle::GetBrush("Level.VisibleIcon16x") : FAppStyle::GetBrush("Level.NotVisibleHighlightIcon16x");
+							}
+							return nullptr;
+						})
+						.ToolTipText_Lambda([this]()
+						{
+							if(IAnimNextRigVMExportInterface* Export = Cast<IAnimNextRigVMExportInterface>(Entry->WeakEntry.Get()))
+							{
+								FText AccessSpecifier = Export->GetExportAccessSpecifier() == EAnimNextExportAccessSpecifier::Public ? LOCTEXT("PublicSpecifier", "public") : LOCTEXT("PrivateSpecifier", "private");
+								FText AccessSpecifierDesc = Export->GetExportAccessSpecifier() == EAnimNextExportAccessSpecifier::Public ?
+									LOCTEXT("PublicSpecifierDesc", "This means that the entry is usable from gameplay and from other AnimNext assets") :
+									LOCTEXT("PrivateSpecifierDesc", "This means that the entry is only usable inside this asset");
+								return FText::Format(LOCTEXT("AccessSpecifierTooltip", "This entry is {0}.\n{1}"), AccessSpecifier, AccessSpecifierDesc);
+							}
+							return FText::GetEmpty();
+						})
+					]
 				];
 		}
-		else if(InColumnName == Column_ModifiedStatus)
+		else if(InColumnName == Column_RevisionControl)
 		{
 			return
 				SNew(SBox)
+				.ToolTipText_Lambda([this]()
+				{
+					if(UAnimNextRigVMAssetEntry* AssetEntry = Entry->WeakEntry.Get())
+					{
+						ISourceControlProvider& SourceControlProvider = ISourceControlModule::Get().GetProvider();
+						UPackage* Package = AssetEntry->GetExternalPackage();
+						if(FSourceControlStatePtr State = SourceControlProvider.GetState(Package, EStateCacheUsage::Use))
+						{
+							return FText::Format(LOCTEXT("RevisionControlStatusFormat", "File: {0}\nStatus: {1}"), FText::FromName(Package->GetFName()),  State->GetDisplayTooltip());
+						}
+					}
+
+					return LOCTEXT("RevisionControlStatus", "Revision control status of this parameter");
+				})
 				.HAlign(HAlign_Center)
 				.VAlign(VAlign_Center)
 				.HeightOverride(20.0f)
 				[
 					SNew(SImage)
-					.ToolTipText_Lambda([this]()
-					{
-						FTextBuilder TextBuilder;
-						TextBuilder.AppendLine(LOCTEXT("ModifiedTooltip", "Modified Status"));
-
-						if(UAnimNextRigVMAssetEntry* AssetEntry = Entry->WeakEntry.Get())
-						{
-							const UPackage* ExternalPackage = AssetEntry->GetExternalPackage();
-							check(ExternalPackage);
-							if(ExternalPackage->IsDirty())
-							{
-								TextBuilder.AppendLine(FText::FromName(ExternalPackage->GetFName()));
-							}
-						}
-
-						return TextBuilder.ToText();
-					})
 					.Image_Lambda([this]() -> const FSlateBrush*
 					{
 						if(UAnimNextRigVMAssetEntry* AssetEntry = Entry->WeakEntry.Get())
 						{
-							bool bIsDirty = false;
-							const UPackage* ExternalPackage = AssetEntry->GetExternalPackage();
-							check(ExternalPackage);
-							if(ExternalPackage->IsDirty())
+							ISourceControlProvider& SourceControlProvider = ISourceControlModule::Get().GetProvider();
+							if(FSourceControlStatePtr State = SourceControlProvider.GetState(AssetEntry->GetExternalPackage(), EStateCacheUsage::Use))
 							{
-								bIsDirty = true;
+								return State->GetIcon().GetSmallIcon();
 							}
-
-
-							return bIsDirty ? FAppStyle::GetBrush("ContentBrowser.ContentDirty") : nullptr;
 						}
 						return nullptr;
 					})
@@ -775,63 +820,109 @@ class SRigVMAssetViewRow : public SMultiColumnTableRow<TSharedRef<FRigVMAssetVie
 				.VAlign(VAlign_Center)
 				.AutoWidth()
 				[
-					SAssignNew(Entry->NameWidget, SInlineEditableTextBlock)
-					.Font(IDetailLayoutBuilder::GetDetailFont())
-					.IsSelected(this, &SRigVMAssetViewRow::IsSelectedExclusively)
-					.IsReadOnly_Lambda([this]()
-					{
-						if (Cast<IAnimNextRigVMParameterInterface>(Entry->WeakEntry.Get()) 
-							|| Cast<IAnimNextRigVMGraphInterface>(Entry->WeakEntry.Get()))
+					SNew(SHorizontalBox)
+					+SHorizontalBox::Slot()
+					.VAlign(VAlign_Center)
+					.AutoWidth()
+					[
+						SAssignNew(Entry->NameWidget, SInlineEditableTextBlock)
+						.Font(IDetailLayoutBuilder::GetDetailFont())
+						.IsSelected(this, &SRigVMAssetViewRow::IsSelectedExclusively)
+						.IsReadOnly_Lambda([this]()
 						{
-							return false;
-						}
-						return true;
-					})
-					.OnTextCommitted_Lambda([this](const FText& InNewText, ETextCommit::Type InCommitType)
-					{
-						if(InCommitType == ETextCommit::OnEnter)
+							if (Cast<IAnimNextRigVMParameterInterface>(Entry->WeakEntry.Get()) 
+								|| Cast<IAnimNextRigVMGraphInterface>(Entry->WeakEntry.Get()))
+							{
+								return false;
+							}
+							return true;
+						})
+						.OnTextCommitted_Lambda([this](const FText& InNewText, ETextCommit::Type InCommitType)
 						{
-							if(UAnimNextRigVMAssetEntry* AssetEntry = Cast<UAnimNextRigVMAssetEntry>(Entry->WeakEntry.Get()))
+							if(InCommitType == ETextCommit::OnEnter)
 							{
-								FScopedTransaction Transaction(LOCTEXT("SetParameterName", "Set Entry name"));
-								AssetEntry->SetEntryName(*InNewText.ToString());
+								if(UAnimNextRigVMAssetEntry* AssetEntry = Cast<UAnimNextRigVMAssetEntry>(Entry->WeakEntry.Get()))
+								{
+									FScopedTransaction Transaction(LOCTEXT("SetParameterName", "Set Entry name"));
+									AssetEntry->SetEntryName(*InNewText.ToString());
+								}
 							}
-						}
-					})
-					.OnVerifyTextChanged_Lambda([this](const FText& InNewText, FText& OutErrorText)
-					{
-						const FString NewString = InNewText.ToString();
-						if (!FUtils::IsValidParameterNameString(NewString, OutErrorText))
+						})
+						.OnVerifyTextChanged_Lambda([this](const FText& InNewText, FText& OutErrorText)
 						{
-							return false;
-						}
+							const FString NewString = InNewText.ToString();
+							if (!FUtils::IsValidParameterNameString(NewString, OutErrorText))
+							{
+								return false;
+							}
 
-						FName Name(*NewString);
-						if (TSharedPtr<SRigVMAssetView> View = WeakView.Pin())
+							FName Name(*NewString);
+							if (TSharedPtr<SRigVMAssetView> View = WeakView.Pin())
+							{
+								if(UAnimNextRigVMAssetEntry* ExistingEntry = View->EditorData->FindEntry(Name))
+								{
+									OutErrorText = LOCTEXT("Error_NameExists", "This name already exists in this asset");
+									return false;
+								}
+							}
+							return true;
+						})
+						.Text_Lambda([this]()
 						{
-							if(UAnimNextRigVMAssetEntry* ExistingEntry = View->EditorData->FindEntry(Name))
+							if(UAnimNextRigVMAssetEntry* AssetEntry = Cast<UAnimNextRigVMAssetEntry>(Entry->WeakEntry.Get()))
 							{
-								OutErrorText = LOCTEXT("Error_NameExists", "This name already exists in this asset");
-								return false;
+								return FText::FromName(AssetEntry->GetEntryName());
 							}
-						}
-						return true;
-					})
-					.Text_Lambda([this]()
+							return FText::GetEmpty();
+						})
+						.ToolTipText_Lambda([this]()
+						{
+							if(UAnimNextRigVMAssetEntry* AssetEntry = Entry->WeakEntry.Get())
+							{
+								return AssetEntry->GetDisplayNameTooltip();
+							}
+							return FText::GetEmpty();
+						})
+					]
+				]
+				+SHorizontalBox::Slot()
+				.VAlign(VAlign_Center)
+				.Padding(2.0f, 0.0f, 2.0f, 3.0f)
+				.AutoWidth()
+				[
+					SNew(SImage)
+					.ToolTipText_Lambda([this]()
 					{
-						if(UAnimNextRigVMAssetEntry* AssetEntry = Cast<UAnimNextRigVMAssetEntry>(Entry->WeakEntry.Get()))
+						FTextBuilder TextBuilder;
+						TextBuilder.AppendLine(LOCTEXT("ModifiedTooltip", "Modified"));
+
+						if(UAnimNextRigVMAssetEntry* AssetEntry = Entry->WeakEntry.Get())
 						{
-							return FText::FromName(AssetEntry->GetEntryName());
+							const UPackage* ExternalPackage = AssetEntry->GetExternalPackage();
+							check(ExternalPackage);
+							if(ExternalPackage->IsDirty())
+							{
+								TextBuilder.AppendLine(FText::FromName(ExternalPackage->GetFName()));
+							}
 						}
-						return FText::GetEmpty();
+
+						return TextBuilder.ToText();
 					})
-					.ToolTipText_Lambda([this]()
+					.Image_Lambda([this]() -> const FSlateBrush*
 					{
 						if(UAnimNextRigVMAssetEntry* AssetEntry = Entry->WeakEntry.Get())
 						{
-							return AssetEntry->GetDisplayNameTooltip();
+							bool bIsDirty = false;
+							const UPackage* ExternalPackage = AssetEntry->GetExternalPackage();
+							check(ExternalPackage);
+							if(ExternalPackage->IsDirty())
+							{
+								bIsDirty = true;
+							}
+
+							return bIsDirty ? FAppStyle::GetBrush("Icons.DirtyBadge") : nullptr;
 						}
-						return FText::GetEmpty();
+						return nullptr;
 					})
 				];
 		}
diff --git a/Engine/Plugins/Experimental/AnimNext/Source/AnimNextEditor/Private/Param/AnimNextActorLocatorEditor.cpp b/Engine/Plugins/Experimental/AnimNext/Source/AnimNextEditor/Private/Param/AnimNextActorLocatorEditor.cpp
new file mode 100644
--- /dev/null
+++ b/Engine/Plugins/Experimental/AnimNext/Source/AnimNextEditor/Private/Param/AnimNextActorLocatorEditor.cpp
@@ -0,0 +1,70 @@
+﻿// Copyright Epic Games, Inc. All Rights Reserved.
+
+#include "AnimNextActorLocatorEditor.h"
+
+#include "UniversalObjectLocatorFragmentTypeHandle.h"
+#include "UniversalObjectLocator.h"
+#include "Textures/SlateIcon.h"
+#include "GameFramework/Actor.h"
+#include "Param/AnimNextActorLocatorFragment.h"
+#include "Styling/AppStyle.h"
+
+#define LOCTEXT_NAMESPACE "AnimNextActorLocatorEditor"
+
+namespace UE::AnimNext::Editor
+{
+
+UE::UniversalObjectLocator::ELocatorFragmentEditorType FActorLocatorEditor::GetLocatorFragmentEditorType() const
+{
+	return UE::UniversalObjectLocator::ELocatorFragmentEditorType::Absolute;
+}
+
+bool FActorLocatorEditor::IsAllowedInContext(FName InContextName) const
+{
+	return InContextName == "AnimNextContext";
+}
+
+bool FActorLocatorEditor::IsDragSupported(TSharedPtr<FDragDropOperation> DragOperation, UObject* Context) const
+{
+	return false;
+}
+
+UObject* FActorLocatorEditor::ResolveDragOperation(TSharedPtr<FDragDropOperation> DragOperation, UObject* Context) const
+{
+	return nullptr;
+}
+
+TSharedPtr<SWidget> FActorLocatorEditor::MakeEditUI(const FEditUIParameters& InParameters)
+{
+	return nullptr;
+}
+
+FText FActorLocatorEditor::GetDisplayText(const FUniversalObjectLocatorFragment* InFragment) const
+{
+	return LOCTEXT("AnimNextActorLocatorName", "Current Actor");
+}
+
+FText FActorLocatorEditor::GetDisplayTooltip(const FUniversalObjectLocatorFragment* InFragment) const
+{
+	return LOCTEXT("AnimNextActorLocatorTooltip", "The current actor");
+}
+
+FSlateIcon FActorLocatorEditor::GetDisplayIcon(const FUniversalObjectLocatorFragment* InFragment) const
+{
+	return FSlateIcon(FAppStyle::GetAppStyleSetName(), "Icons.Actor");
+}
+
+UClass* FActorLocatorEditor::ResolveClass(const FUniversalObjectLocatorFragment& InFragment, UObject* InContext) const
+{
+	return AActor::StaticClass(); 
+}
+
+FUniversalObjectLocatorFragment FActorLocatorEditor::MakeDefaultLocatorFragment() const
+{
+	FUniversalObjectLocatorFragment NewFragment(FAnimNextActorLocatorFragment::FragmentType);
+	return NewFragment;
+}
+
+}
+
+#undef LOCTEXT_NAMESPACE
diff --git a/Engine/Plugins/Experimental/AnimNext/Source/AnimNextEditor/Private/Param/AnimNextActorLocatorEditor.h b/Engine/Plugins/Experimental/AnimNext/Source/AnimNextEditor/Private/Param/AnimNextActorLocatorEditor.h
new file mode 100644
--- /dev/null
+++ b/Engine/Plugins/Experimental/AnimNext/Source/AnimNextEditor/Private/Param/AnimNextActorLocatorEditor.h
@@ -0,0 +1,26 @@
+﻿// Copyright Epic Games, Inc. All Rights Reserved.
+
+#pragma once
+
+#include "CoreTypes.h"
+#include "UniversalObjectLocatorEditor.h"
+
+namespace UE::AnimNext::Editor
+{
+
+class FActorLocatorEditor : public UE::UniversalObjectLocator::ILocatorFragmentEditor
+{
+	// ILocatorFragmentEditor interface
+	virtual UE::UniversalObjectLocator::ELocatorFragmentEditorType GetLocatorFragmentEditorType() const override;
+	virtual bool IsAllowedInContext(FName InContextName) const override;
+	virtual bool IsDragSupported(TSharedPtr<FDragDropOperation> DragOperation, UObject* Context) const override;
+	virtual UObject* ResolveDragOperation(TSharedPtr<FDragDropOperation> DragOperation, UObject* Context) const override;
+	virtual TSharedPtr<SWidget> MakeEditUI(const FEditUIParameters& InParameters) override;
+	virtual FText GetDisplayText(const FUniversalObjectLocatorFragment* InFragment = nullptr) const override;
+	virtual FText GetDisplayTooltip(const FUniversalObjectLocatorFragment* InFragment = nullptr) const override;
+	virtual FSlateIcon GetDisplayIcon(const FUniversalObjectLocatorFragment* InFragment = nullptr) const override;
+	virtual UClass* ResolveClass(const FUniversalObjectLocatorFragment& InFragment, UObject* InContext) const override; 
+	virtual FUniversalObjectLocatorFragment MakeDefaultLocatorFragment() const override;
+};
+
+}
diff --git a/Engine/Plugins/Experimental/AnimNext/Source/AnimNextEditor/Private/Param/AnimNextComponentLocatorEditor.cpp b/Engine/Plugins/Experimental/AnimNext/Source/AnimNextEditor/Private/Param/AnimNextComponentLocatorEditor.cpp
new file mode 100644
--- /dev/null
+++ b/Engine/Plugins/Experimental/AnimNext/Source/AnimNextEditor/Private/Param/AnimNextComponentLocatorEditor.cpp
@@ -0,0 +1,182 @@
+﻿// Copyright Epic Games, Inc. All Rights Reserved.
+
+#include "AnimNextComponentLocatorEditor.h"
+
+#include "UniversalObjectLocatorFragmentTypeHandle.h"
+#include "UniversalObjectLocator.h"
+#include "Textures/SlateIcon.h"
+#include "Modules/ModuleManager.h"
+#include "Param/AnimNextComponentLocatorFragment.h"
+#include "ClassViewerFilter.h"
+#include "ClassViewerModule.h"
+#include "IUniversalObjectLocatorCustomization.h"
+#include "Components/ActorComponent.h"
+#include "Styling/AppStyle.h"
+#include "Widgets/SCompoundWidget.h"
+#include "Widgets/Layout/SBox.h"
+
+#define LOCTEXT_NAMESPACE "EntityComponentLocatorEditor"
+
+namespace UE::AnimNext::Editor
+{
+
+class SComponentLocatorEditor : public SCompoundWidget
+{
+	SLATE_BEGIN_ARGS(SComponentLocatorEditor){}
+	SLATE_END_ARGS()
+
+	class FComponentFilter : public IClassViewerFilter
+	{
+	public:
+		virtual bool IsClassAllowed(const FClassViewerInitializationOptions& InInitOptions, const UClass* InClass, TSharedRef<FClassViewerFilterFuncs> InFilterFuncs)
+		{
+			return InClass->IsChildOf(UActorComponent::StaticClass());
+		}
+
+		virtual bool IsUnloadedClassAllowed(const FClassViewerInitializationOptions& InInitOptions, const TSharedRef<const IUnloadedBlueprintData> InUnloadedClassData, TSharedRef<FClassViewerFilterFuncs> InFilterFuncs)
+		{
+			return false;
+		}
+	};
+
+	void Construct(const FArguments& InArgs, TSharedPtr<UE::UniversalObjectLocator::IFragmentEditorHandle> InHandle)
+	{
+		WeakHandle = InHandle;
+		CurrentClass = InHandle->GetResolvedClass() ? InHandle->GetResolvedClass() : InHandle->GetContextClass();
+
+		FClassViewerInitializationOptions Options;
+		Options.ClassFilters.Add(MakeShared<FComponentFilter>());
+		Options.InitiallySelectedClass = const_cast<UClass*>(CurrentClass);
+		FClassViewerModule& ClassViewerModule = FModuleManager::LoadModuleChecked<FClassViewerModule>("ClassViewer");
+
+		ChildSlot
+		[
+			SNew(SBox)
+			.WidthOverride(400.0f)
+			.HeightOverride(400.0f)
+			[
+				ClassViewerModule.CreateClassViewer(Options, FOnClassPicked::CreateRaw(this, &SComponentLocatorEditor::HandleClassPicked))
+			]
+		];
+	}
+
+	void HandleClassPicked(UClass* InClass)
+	{
+		if(InClass != CurrentClass)
+		{
+			if(TSharedPtr<UE::UniversalObjectLocator::IFragmentEditorHandle> Handle = WeakHandle.Pin())
+			{
+				FUniversalObjectLocatorFragment NewFragment(FAnimNextComponentLocatorFragment::FragmentType);
+				FAnimNextComponentLocatorFragment* Payload = NewFragment.GetPayloadAs(FAnimNextComponentLocatorFragment::FragmentType);
+				Payload->Path = InClass;
+				Handle->SetValue(NewFragment);
+			}
+			
+			CurrentClass = InClass;
+		}
+	}
+
+private:
+	const UClass* CurrentClass = nullptr;
+	TWeakPtr<UE::UniversalObjectLocator::IFragmentEditorHandle> WeakHandle;
+};
+
+
+UE::UniversalObjectLocator::ELocatorFragmentEditorType FComponentLocatorEditor::GetLocatorFragmentEditorType() const
+{
+	return UE::UniversalObjectLocator::ELocatorFragmentEditorType::Relative;
+}
+
+bool FComponentLocatorEditor::IsAllowedInContext(FName InContextName) const
+{
+	return InContextName == "AnimNextContext";
+}
+
+bool FComponentLocatorEditor::IsDragSupported(TSharedPtr<FDragDropOperation> DragOperation, UObject* Context) const
+{
+	return false;
+}
+
+UObject* FComponentLocatorEditor::ResolveDragOperation(TSharedPtr<FDragDropOperation> DragOperation, UObject* Context) const
+{
+	return nullptr;
+}
+
+TSharedPtr<SWidget> FComponentLocatorEditor::MakeEditUI(const FEditUIParameters& InParameters)
+{
+	return SNew(SComponentLocatorEditor, InParameters.Handle);
+}
+
+FText FComponentLocatorEditor::GetDisplayText(const FUniversalObjectLocatorFragment* InFragment) const
+{
+	if(InFragment != nullptr)
+	{
+		ensure(InFragment->GetFragmentTypeHandle() == FAnimNextComponentLocatorFragment::FragmentType);
+		const FAnimNextComponentLocatorFragment* Payload = InFragment->GetPayloadAs(FAnimNextComponentLocatorFragment::FragmentType);
+		if(Payload && Payload->Path.IsValid())
+		{
+			static const FTextFormat Format(LOCTEXT("ComponentLabelFormat", "Get {0}"));
+			return FText::Format(Format, FText::FromString(Payload->Path.GetAssetName()));
+		}
+	}
+	
+	return LOCTEXT("AnimationEntityComponentLocatorName", "Component");
+}
+
+FText FComponentLocatorEditor::GetDisplayTooltip(const FUniversalObjectLocatorFragment* InFragment) const
+{
+	if(InFragment != nullptr)
+	{
+		ensure(InFragment->GetFragmentTypeHandle() == FAnimNextComponentLocatorFragment::FragmentType);
+		const FAnimNextComponentLocatorFragment* Payload = InFragment->GetPayloadAs(FAnimNextComponentLocatorFragment::FragmentType);
+		if(Payload && Payload->Path.IsValid())
+		{
+			static const FTextFormat Format(LOCTEXT("ComponentTooltipFormat", "Get the first component of type '{0}'"));
+			return FText::Format(Format, FText::FromString(Payload->Path.GetAssetName()));
+		}
+	}
+
+	return LOCTEXT("AnimationEntityComponentLocatorTooltip", "An actor component of a selected class");
+}
+
+FSlateIcon FComponentLocatorEditor::GetDisplayIcon(const FUniversalObjectLocatorFragment* InFragment) const
+{
+	static const FName ActorComponentClassIconName("ClassIcon." + UActorComponent::StaticClass()->GetName());
+	if(InFragment != nullptr)
+	{
+		ensure(InFragment->GetFragmentTypeHandle() == FAnimNextComponentLocatorFragment::FragmentType);
+		const FAnimNextComponentLocatorFragment* Payload = InFragment->GetPayloadAs(FAnimNextComponentLocatorFragment::FragmentType);
+		if(Payload && Payload->Path.IsValid())
+		{
+			if(UClass* Class = Cast<UClass>(Payload->Path.ResolveObject()))
+			{
+				const FName ComponentClassIconName("ClassIcon." + Class->GetName());
+				FSlateIcon(FAppStyle::GetAppStyleSetName(), ComponentClassIconName);
+			}
+		}
+	}
+	
+	return FSlateIcon(FAppStyle::GetAppStyleSetName(), ActorComponentClassIconName);
+}
+
+UClass* FComponentLocatorEditor::ResolveClass(const FUniversalObjectLocatorFragment& InFragment, UObject* InContext) const
+{
+	ensure(InFragment.GetFragmentTypeHandle() == FAnimNextComponentLocatorFragment::FragmentType);
+	const FAnimNextComponentLocatorFragment* Payload = InFragment.GetPayloadAs(FAnimNextComponentLocatorFragment::FragmentType);
+	if(Payload && Payload->Path.IsValid())
+	{
+		return Cast<UClass>(Payload->Path.ResolveObject());
+	}
+
+	return UActorComponent::StaticClass(); 
+}
+
+FUniversalObjectLocatorFragment FComponentLocatorEditor::MakeDefaultLocatorFragment() const
+{
+	FUniversalObjectLocatorFragment NewFragment(FAnimNextComponentLocatorFragment::FragmentType);
+	return NewFragment;
+}
+
+}
+
+#undef LOCTEXT_NAMESPACE
diff --git a/Engine/Plugins/Experimental/AnimNext/Source/AnimNextEditor/Private/Param/AnimNextComponentLocatorEditor.h b/Engine/Plugins/Experimental/AnimNext/Source/AnimNextEditor/Private/Param/AnimNextComponentLocatorEditor.h
new file mode 100644
--- /dev/null
+++ b/Engine/Plugins/Experimental/AnimNext/Source/AnimNextEditor/Private/Param/AnimNextComponentLocatorEditor.h
@@ -0,0 +1,26 @@
+﻿// Copyright Epic Games, Inc. All Rights Reserved.
+
+#pragma once
+
+#include "CoreTypes.h"
+#include "UniversalObjectLocatorEditor.h"
+
+namespace UE::AnimNext::Editor
+{
+
+class FComponentLocatorEditor : public UE::UniversalObjectLocator::ILocatorFragmentEditor
+{
+	// ILocatorFragmentEditor interface
+	virtual UE::UniversalObjectLocator::ELocatorFragmentEditorType GetLocatorFragmentEditorType() const override;
+	virtual bool IsAllowedInContext(FName InContextName) const override;
+	virtual bool IsDragSupported(TSharedPtr<FDragDropOperation> DragOperation, UObject* Context) const override;
+	virtual UObject* ResolveDragOperation(TSharedPtr<FDragDropOperation> DragOperation, UObject* Context) const override;
+	virtual TSharedPtr<SWidget> MakeEditUI(const FEditUIParameters& InParameters) override;
+	virtual FText GetDisplayText(const FUniversalObjectLocatorFragment* InFragment = nullptr) const override;
+	virtual FText GetDisplayTooltip(const FUniversalObjectLocatorFragment* InFragment = nullptr) const override;
+	virtual FSlateIcon GetDisplayIcon(const FUniversalObjectLocatorFragment* InFragment = nullptr) const override;
+	virtual UClass* ResolveClass(const FUniversalObjectLocatorFragment& InFragment, UObject* InContext) const override; 
+	virtual FUniversalObjectLocatorFragment MakeDefaultLocatorFragment() const override;
+};
+
+}
diff --git a/Engine/Plugins/Experimental/AnimNext/Source/AnimNextEditor/Private/Param/AnimNextLocatorContext.h b/Engine/Plugins/Experimental/AnimNext/Source/AnimNextEditor/Private/Param/AnimNextLocatorContext.h
--- a/Engine/Plugins/Experimental/AnimNext/Source/AnimNextEditor/Private/Param/AnimNextLocatorContext.h
+++ b/Engine/Plugins/Experimental/AnimNext/Source/AnimNextEditor/Private/Param/AnimNextLocatorContext.h
@@ -4,11 +4,13 @@
 
 #include "UniversalObjectLocatorEditorContext.h"
 #include "Component/AnimNextComponent.h"
+#include "Modules/ModuleManager.h"
+#include "AnimNextEditorModule.h"
 
-namespace UE::AnimNext
+namespace UE::AnimNext::Editor
 {
 
-class FAnimNextLocatorContext : public UE::UniversalObjectLocator::ILocatorFragmentEditorContext
+class FLocatorContext : public UE::UniversalObjectLocator::ILocatorFragmentEditorContext
 {
 	// ILocatorFragmentEditorContext interface
 	virtual UObject* GetContext(const IPropertyHandle& InPropertyHandle) const override
@@ -19,13 +21,8 @@ class FAnimNextLocatorContext : public UE::UniversalObjectLocator::ILocatorFragm
 
 	virtual bool IsFragmentAllowed(FName InFragmentName) const override
 	{
-		return (InFragmentName == "Actor" ||
-				InFragmentName == "Asset" ||
-				InFragmentName == "AnimNextScope" ||
-				InFragmentName == "AnimNextGraph" ||
-				InFragmentName == "AnimNextObjectFunction" ||
-				InFragmentName == "AnimNextObjectProperty" ||
-				InFragmentName == "AnimNextObjectCast");
+		FModule& EditorModule = FModuleManager::GetModuleChecked<FModule>("AnimNextEditor");
+		return EditorModule.LocatorFragmentEditorNames.Contains(InFragmentName);
 	}
 };
 
diff --git a/Engine/Plugins/Experimental/AnimNext/Source/AnimNextEditor/Private/Param/ObjectFunctionLocatorEditor.cpp b/Engine/Plugins/Experimental/AnimNext/Source/AnimNextEditor/Private/Param/ObjectFunctionLocatorEditor.cpp
--- a/Engine/Plugins/Experimental/AnimNext/Source/AnimNextEditor/Private/Param/ObjectFunctionLocatorEditor.cpp
+++ b/Engine/Plugins/Experimental/AnimNext/Source/AnimNextEditor/Private/Param/ObjectFunctionLocatorEditor.cpp
@@ -1,4 +1,4 @@
-﻿// Copyright Epic Games, Inc. All Rights Reserved.
+// Copyright Epic Games, Inc. All Rights Reserved.
 
 #include "ObjectFunctionLocatorEditor.h"
 
@@ -16,6 +16,7 @@
 #include "Modules/ModuleManager.h"
 #include "Framework/PropertyViewer/IFieldIterator.h"
 #include "Widgets/PropertyViewer/SPropertyViewer.h"
+#include "Param/ParamUtils.h"
 
 #define LOCTEXT_NAMESPACE "ObjectFunctionLocatorEditor"
 
@@ -94,6 +95,22 @@ class SObjectFunctionLocatorEditor : public SCompoundWidget
 		];
 
 		PropertyViewer->AddContainer(CurrentClass);
+
+		// Add any BP function libraries that can potentially hoist this class
+		TArray<UClass*> Classes;
+		GetDerivedClasses(UBlueprintFunctionLibrary::StaticClass(), Classes, true);
+		for (const UClass* Class : Classes)
+		{
+			for (TFieldIterator<UFunction> FieldIt(Class, EFieldIteratorFlags::IncludeSuper); FieldIt; ++FieldIt)
+			{
+				UFunction* Function = *FieldIt;
+				if (FParamUtils::CanUseFunction(Function, CurrentClass))
+				{
+					PropertyViewer->AddContainer(Class);
+					break;
+				}
+			}
+		}
 	}
 
 	void HandleClassPicked(UClass* InClass)
diff --git a/Engine/Plugins/Experimental/AnimNext/Source/AnimNextEditor/Private/Param/ObjectPropertyLocatorEditor.cpp b/Engine/Plugins/Experimental/AnimNext/Source/AnimNextEditor/Private/Param/ObjectPropertyLocatorEditor.cpp
--- a/Engine/Plugins/Experimental/AnimNext/Source/AnimNextEditor/Private/Param/ObjectPropertyLocatorEditor.cpp
+++ b/Engine/Plugins/Experimental/AnimNext/Source/AnimNextEditor/Private/Param/ObjectPropertyLocatorEditor.cpp
@@ -37,7 +37,7 @@ class SObjectPropertyLocatorEditor : public SCompoundWidget
 				for (TFieldIterator<FProperty> PropertyIt(InPropertyStruct, EFieldIteratorFlags::IncludeSuper); PropertyIt; ++PropertyIt)
 				{
 					FProperty* Property = *PropertyIt;
-					if (Property && Property->HasAnyPropertyFlags(CPF_BlueprintVisible | CPF_EditConst))
+					if (Property && Property->HasAnyPropertyFlags(CPF_BlueprintVisible | CPF_Edit | CPF_EditConst))
 					{
 						if(FObjectProperty* ObjectProperty = CastField<FObjectProperty>(Property))
 						{
@@ -57,7 +57,7 @@ class SObjectPropertyLocatorEditor : public SCompoundWidget
 			for (TFieldIterator<FProperty> PropertyIt(InStruct, EFieldIteratorFlags::IncludeSuper); PropertyIt; ++PropertyIt)
 			{
 				FProperty* Property = *PropertyIt;
-				if (Property && Property->HasAnyPropertyFlags(CPF_BlueprintVisible | CPF_EditConst))
+				if (Property && Property->HasAnyPropertyFlags(CPF_BlueprintVisible | CPF_Edit | CPF_EditConst))
 				{
 					if(FObjectProperty* ObjectProperty = CastField<FObjectProperty>(Property))
 					{
diff --git a/Engine/Plugins/Experimental/AnimNext/Source/AnimNextEditor/Private/Param/ParameterCustomization.cpp b/Engine/Plugins/Experimental/AnimNext/Source/AnimNextEditor/Private/Param/ParameterCustomization.cpp
--- a/Engine/Plugins/Experimental/AnimNext/Source/AnimNextEditor/Private/Param/ParameterCustomization.cpp
+++ b/Engine/Plugins/Experimental/AnimNext/Source/AnimNextEditor/Private/Param/ParameterCustomization.cpp
@@ -37,19 +37,21 @@ void FParameterCustomization::CustomizeDetails(IDetailLayoutBuilder& DetailBuild
 
 		if (UAnimNextGraph_EditorData* EditorData = Cast<UAnimNextGraph_EditorData>(Parameter->GetOuter()))
 		{
-			const FName ParameterName = Parameter->GetEntryName();
-			if (UAnimNextRigVMAssetEntry* AssetEntry = EditorData->FindEntry(ParameterName)) 
+			const FName EntryName = Parameter->GetEntryName();
+			if (UAnimNextRigVMAssetEntry* AssetEntry = EditorData->FindEntry(EntryName)) 
 			{
 				if (UAnimNextGraph* ReferencedGraph = UE::AnimNext::UncookedOnly::FUtils::GetGraph(EditorData))
 				{
 					FAddPropertyParams AddPropertyParams;
 					TArray<IDetailPropertyRow*> DetailPropertyRows;
 
-					if (const FPropertyBagPropertyDesc* PropertyDesc = ReferencedGraph->PropertyBag.FindPropertyDescByName(ParameterName))
+					if (const FPropertyBagPropertyDesc* PropertyDesc = ReferencedGraph->DefaultState.State.FindPropertyDescByName(Parameter->GetParamName()))
 					{
-						IDetailPropertyRow* DetailPropertyRow = DefaultValueCategory.AddExternalStructureProperty(MakeShared<FInstancePropertyBagStructureDataProvider>(ReferencedGraph->PropertyBag), ParameterName, EPropertyLocation::Default, AddPropertyParams);
+						IDetailPropertyRow* DetailPropertyRow = DefaultValueCategory.AddExternalStructureProperty(MakeShared<FInstancePropertyBagStructureDataProvider>(ReferencedGraph->DefaultState.State), Parameter->GetParamName(), EPropertyLocation::Default, AddPropertyParams);
 						if (TSharedPtr<IPropertyHandle> Handle = DetailPropertyRow->GetPropertyHandle(); Handle.IsValid())
 						{
+							Handle->SetPropertyDisplayName(FText::FromName(EntryName));
+							
 							const TWeakObjectPtr<UAnimNextGraph> ReferencedGraphWeak = ReferencedGraph;
 
 							const auto OnPropertyValuePreChange = [ReferencedGraphWeak]()
@@ -65,9 +67,10 @@ void FParameterCustomization::CustomizeDetails(IDetailLayoutBuilder& DetailBuild
 									{
 										if (UAnimNextGraph_EditorData* EditorData = Cast<UAnimNextGraph_EditorData>(ReferencedGraphWeak->EditorData))
 										{
-											if (UAnimNextRigVMAssetEntry* AssetEntry = EditorData->FindEntry(InEvent.GetPropertyName()))
+											if (UAnimNextRigVMAssetEntry* AssetEntry = EditorData->FindEntry(UncookedOnly::FUtils::GetParameterNameFromQualifiedName(InEvent.GetPropertyName())))
 											{
 												AssetEntry->MarkPackageDirty();
+												AssetEntry->BroadcastModified();
 											}
 										}
 									}
diff --git a/Engine/Plugins/Experimental/AnimNext/Source/AnimNextEditor/Private/Param/SParameterPicker.cpp b/Engine/Plugins/Experimental/AnimNext/Source/AnimNextEditor/Private/Param/SParameterPicker.cpp
--- a/Engine/Plugins/Experimental/AnimNext/Source/AnimNextEditor/Private/Param/SParameterPicker.cpp
+++ b/Engine/Plugins/Experimental/AnimNext/Source/AnimNextEditor/Private/Param/SParameterPicker.cpp
@@ -2,6 +2,7 @@
 
 #include "SParameterPicker.h"
 
+#include "AssetDefinitionDefault.h"
 #include "UncookedOnlyUtils.h"
 #include "AssetRegistry/AssetRegistryModule.h"
 #include "Kismet2/BlueprintEditorUtils.h"
@@ -27,7 +28,6 @@
 #include "Component/AnimNextComponent.h"
 #include "Editor/PropertyEditor/Private/SSingleProperty.h"
 #include "Graph/AnimNextGraph.h"
-#include "Param/AnimNextClassExtensionLibrary.h"
 #include "Param/AnimNextParam.h"
 #include "Param/IParameterSourceType.h"
 #include "Param/ParamUtils.h"
@@ -224,7 +224,7 @@ void SParameterPicker::RefreshEntries()
 					PropertyDescs.Reserve(ExportPair.Value.Parameters.Num());
 					for(const FAnimNextParameterAssetRegistryExportEntry& ParameterEntry : ExportPair.Value.Parameters)
 					{
-						if(EnumHasAnyFlags(ParameterEntry.GetFlags(), EAnimNextParameterFlags::Bound) && ParameterEntry.Name != NAME_None)
+						if(EnumHasAllFlags(ParameterEntry.GetFlags(), EAnimNextParameterFlags::Declared | EAnimNextParameterFlags::Public) && ParameterEntry.Name != NAME_None)
 						{
 							FParameterBindingReference ParameterBinding;
 							ParameterBinding.Type = ParameterEntry.Type;
@@ -256,6 +256,7 @@ void SParameterPicker::RefreshEntries()
 	}
 	else if(Struct != nullptr)
 	{
+		FieldIterator->CurrentStruct = Struct;
 		CachedContainers.Emplace_GetRef(Struct->GetDisplayNameText(), Struct->GetToolTipText(), Struct);
 		if(const UScriptStruct* ScriptStruct = Cast<UScriptStruct>(Struct))
 		{
@@ -264,32 +265,41 @@ void SParameterPicker::RefreshEntries()
 		}
 		else if(const UClass* Class = Cast<UClass>(Struct))
 		{
-			UE::PropertyViewer::SPropertyViewer::FHandle Handle = PropertyViewer->AddContainer(Class);
-			ContainerMap.Add(Handle, CachedContainers.Num() - 1);
-
-			// Find any UAnimNextClassExtensionLibrary classes to extend this class
-			for(TObjectIterator<UClass> It; It; ++It)
+			// Find any native UBlueprintFunctionLibrary classes to extend this class
+			TArray<UClass*> Classes;
+			GetDerivedClasses(UBlueprintFunctionLibrary::StaticClass(), Classes, true);
+			Classes.Add(const_cast<UClass*>(Class));
+			for (const UClass* LibraryClass : Classes)
 			{
-				UClass* LibraryClass = *It;
-				if(LibraryClass->HasAnyClassFlags(CLASS_Abstract) || !LibraryClass->HasAnyClassFlags(CLASS_Native) || !LibraryClass->IsChildOf(UAnimNextClassExtensionLibrary::StaticClass()))
+				if(LibraryClass->HasAnyClassFlags(CLASS_Abstract) || !LibraryClass->HasAnyClassFlags(CLASS_Native))
 				{
 					continue;
 				}
 
-				UClass* ExtendedClass = LibraryClass->GetDefaultObject<UAnimNextClassExtensionLibrary>()->GetSupportedClass();
-				if(ExtendedClass == nullptr)
+				auto PassesFilterChecks = [this](const FProperty* InProperty)
 				{
-					continue;
-				}
+					if(InProperty && Args.OnFilterParameterType.IsBound())
+					{
+						FAnimNextParamType Type = FParamTypeHandle::FromProperty(InProperty).GetType();
+						return Args.OnFilterParameterType.Execute(Type) == EFilterParameterResult::Include;
+					}
+
+					return false;
+				};
 
-				if(!Class->IsChildOf(ExtendedClass))
+				for (TFieldIterator<UFunction> FieldIt(LibraryClass); FieldIt; ++FieldIt)
 				{
-					continue;
+					if (FParamUtils::CanUseFunction(*FieldIt, Class))
+					{
+						if(PassesFilterChecks(FieldIt->GetReturnProperty()))
+						{
+							CachedContainers.Emplace_GetRef(LibraryClass->GetDisplayNameText(), LibraryClass->GetToolTipText(), LibraryClass);
+							UE::PropertyViewer::SPropertyViewer::FHandle Handle = PropertyViewer->AddContainer(LibraryClass);
+							ContainerMap.Add(Handle, CachedContainers.Num() - 1);
+							break;
+						}
+					}
 				}
-
-				CachedContainers.Emplace_GetRef(LibraryClass->GetDisplayNameText(), LibraryClass->GetToolTipText(), LibraryClass);
-				Handle = PropertyViewer->AddContainer(LibraryClass);
-				ContainerMap.Add(Handle, CachedContainers.Num() - 1);
 			}
 		}
 	}
@@ -470,7 +480,7 @@ TArray<FFieldVariant> SParameterPicker::FFieldIterator::GetFields(const UStruct*
 	for (TFieldIterator<UFunction> FunctionIt(InStruct, EFieldIteratorFlags::IncludeSuper); FunctionIt; ++FunctionIt)
 	{
 		UFunction* Function = *FunctionIt;
-		if (FParamUtils::CanUseFunction(Function))
+		if (FParamUtils::CanUseFunction(Function, CastChecked<UClass>(CurrentStruct)))
 		{
 			if(PassesFilterChecks(Function->GetReturnProperty()))
 			{
diff --git a/Engine/Plugins/Experimental/AnimNext/Source/AnimNextEditor/Private/Param/SParameterPicker.h b/Engine/Plugins/Experimental/AnimNext/Source/AnimNextEditor/Private/Param/SParameterPicker.h
--- a/Engine/Plugins/Experimental/AnimNext/Source/AnimNextEditor/Private/Param/SParameterPicker.h
+++ b/Engine/Plugins/Experimental/AnimNext/Source/AnimNextEditor/Private/Param/SParameterPicker.h
@@ -73,6 +73,7 @@ private:
 		virtual TArray<FFieldVariant> GetFields(const UStruct* InStruct) const override;
 
 		FOnFilterParameterType OnFilterParameterType;
+		const UStruct* CurrentStruct = nullptr;
 	};
 
 	struct FFieldExpander : UE::PropertyViewer::IFieldExpander
diff --git a/Engine/Plugins/Experimental/AnimNext/Source/AnimNextUncookedOnly/Internal/AnimNextRigVMAssetEntry.h b/Engine/Plugins/Experimental/AnimNext/Source/AnimNextUncookedOnly/Internal/AnimNextRigVMAssetEntry.h
--- a/Engine/Plugins/Experimental/AnimNext/Source/AnimNextUncookedOnly/Internal/AnimNextRigVMAssetEntry.h
+++ b/Engine/Plugins/Experimental/AnimNext/Source/AnimNextUncookedOnly/Internal/AnimNextRigVMAssetEntry.h
@@ -13,6 +13,7 @@ namespace UE::AnimNext::Editor
 {
 	class SRigVMAssetViewRow;
 	struct FRigVMAssetViewEntry;
+	class FParameterCustomization;
 }
 
 /** Base class that defines an entry in a graph, e.g. a parameter binding */
@@ -41,7 +42,12 @@ public:
 
 	// UObject interface
 	virtual bool IsAsset() const override;
+#if WITH_EDITOR
+	virtual void PostEditChangeProperty(FPropertyChangedEvent& PropertyChangedEvent) override;
+#endif
 
 protected:
-	void BroadcastModified();
+	ANIMNEXTUNCOOKEDONLY_API void BroadcastModified();
+
+	friend class UE::AnimNext::Editor::FParameterCustomization;
 };
\ No newline at end of file
diff --git a/Engine/Plugins/Experimental/AnimNext/Source/AnimNextUncookedOnly/Internal/Graph/AnimNextGraph_AnimationGraph.h b/Engine/Plugins/Experimental/AnimNext/Source/AnimNextUncookedOnly/Internal/Graph/AnimNextGraph_AnimationGraph.h
--- a/Engine/Plugins/Experimental/AnimNext/Source/AnimNextUncookedOnly/Internal/Graph/AnimNextGraph_AnimationGraph.h
+++ b/Engine/Plugins/Experimental/AnimNext/Source/AnimNextUncookedOnly/Internal/Graph/AnimNextGraph_AnimationGraph.h
@@ -29,6 +29,8 @@ class UAnimNextGraph_AnimationGraph : public UAnimNextRigVMAssetEntry, public IA
 	// IAnimNextRigVMExportInterface interface
 	virtual FAnimNextParamType GetExportType() const override;
 	virtual FName GetExportName() const override;
+	virtual EAnimNextExportAccessSpecifier GetExportAccessSpecifier() const override;
+	virtual void SetExportAccessSpecifier(EAnimNextExportAccessSpecifier InAccessSpecifier, bool bSetupUndoRedo = true) override;
 
 	// UAnimNextRigVMAssetEntry interface
 	virtual FName GetEntryName() const override;
@@ -41,6 +43,10 @@ class UAnimNextGraph_AnimationGraph : public UAnimNextRigVMAssetEntry, public IA
 	virtual void SetEdGraph(URigVMEdGraph* InGraph) override;
 
 protected:
+	/** Access specifier - whether the graph's entry point is visible external to this asset */
+	UPROPERTY(EditAnywhere, Category = AnimationGraph)
+	EAnimNextExportAccessSpecifier Access = EAnimNextExportAccessSpecifier::Private;
+
 	/** The name of the graph */
 	UPROPERTY(VisibleAnywhere, Category = AnimationGraph)
 	FName GraphName;
diff --git a/Engine/Plugins/Experimental/AnimNext/Source/AnimNextUncookedOnly/Internal/Graph/AnimNextGraph_EventGraph.h b/Engine/Plugins/Experimental/AnimNext/Source/AnimNextUncookedOnly/Internal/Graph/AnimNextGraph_EventGraph.h
--- a/Engine/Plugins/Experimental/AnimNext/Source/AnimNextUncookedOnly/Internal/Graph/AnimNextGraph_EventGraph.h
+++ b/Engine/Plugins/Experimental/AnimNext/Source/AnimNextUncookedOnly/Internal/Graph/AnimNextGraph_EventGraph.h
@@ -18,17 +18,13 @@ namespace UE::AnimNext::Tests
 }
 
 UCLASS(MinimalAPI, Category = "Event Graphs")
-class UAnimNextGraph_EventGraph : public UAnimNextRigVMAssetEntry, public IAnimNextRigVMGraphInterface, public IAnimNextRigVMExportInterface
+class UAnimNextGraph_EventGraph : public UAnimNextRigVMAssetEntry, public IAnimNextRigVMGraphInterface
 {
 	GENERATED_BODY()
 
 	friend class UAnimNextGraph_EditorData;
 	friend class UE::AnimNext::Tests::FEditor_Parameters;
 
-	// IAnimNextRigVMExportInterface interface
-	virtual FAnimNextParamType GetExportType() const override;
-	virtual FName GetExportName() const override;
-
 	// UAnimNextRigVMAssetEntry interface
 	virtual FName GetEntryName() const override { return GraphName; }
 	virtual void SetEntryName(FName InName, bool bSetupUndoRedo = true) override;
@@ -42,7 +38,7 @@ class UAnimNextGraph_EventGraph : public UAnimNextRigVMAssetEntry, public IAnimN
 	virtual void SetEdGraph(URigVMEdGraph* InGraph) override;
 
 	/** The name of the graph */
-	UPROPERTY(VisibleAnywhere, Category = EventGraph)
+	UPROPERTY(VisibleAnywhere, Category = "Event Graph")
 	FName GraphName;
 
 	/** Graph */
diff --git a/Engine/Plugins/Experimental/AnimNext/Source/AnimNextUncookedOnly/Internal/Graph/AnimNextGraph_Parameter.h b/Engine/Plugins/Experimental/AnimNext/Source/AnimNextUncookedOnly/Internal/Graph/AnimNextGraph_Parameter.h
--- a/Engine/Plugins/Experimental/AnimNext/Source/AnimNextUncookedOnly/Internal/Graph/AnimNextGraph_Parameter.h
+++ b/Engine/Plugins/Experimental/AnimNext/Source/AnimNextUncookedOnly/Internal/Graph/AnimNextGraph_Parameter.h
@@ -33,6 +33,8 @@ class UAnimNextGraph_Parameter : public UAnimNextRigVMAssetEntry, public IAnimNe
 	// IAnimNextRigVMExportInterface interface
 	virtual FAnimNextParamType GetExportType() const override;
 	virtual FName GetExportName() const override;
+	virtual EAnimNextExportAccessSpecifier GetExportAccessSpecifier() const override;
+	virtual void SetExportAccessSpecifier(EAnimNextExportAccessSpecifier InAccessSpecifier, bool bSetupUndoRedo = true) override;
 
 	// UAnimNextRigVMAssetEntry interface
 	virtual FName GetEntryName() const override;
@@ -46,13 +48,17 @@ class UAnimNextGraph_Parameter : public UAnimNextRigVMAssetEntry, public IAnimNe
 	virtual FName GetParamName() const override;
 	virtual void SetParamName(FName InName, bool bSetupUndoRedo = true) override;
 	virtual FInstancedPropertyBag& GetPropertyBag() const override;
-	
+
+	/** Access specifier - whether the parameter is visible external to this asset */
+	UPROPERTY(EditAnywhere, Category = "Parameter")
+	EAnimNextExportAccessSpecifier Access = EAnimNextExportAccessSpecifier::Private;
+
 	/** Parameter name we reference */
 	UPROPERTY(VisibleAnywhere, Category = Parameter)
 	FName ParameterName;
 
 	/** The parameter's type */
-	UPROPERTY(EditAnywhere, Category = "Parameter", AssetRegistrySearchable)
+	UPROPERTY(EditAnywhere, Category = "Parameter")
 	FAnimNextParamType Type = FAnimNextParamType::GetType<bool>();
 	
 	/** Comment to display in editor */
diff --git a/Engine/Plugins/Experimental/AnimNext/Source/AnimNextUncookedOnly/Internal/IAnimNextRigVMExportInterface.h b/Engine/Plugins/Experimental/AnimNext/Source/AnimNextUncookedOnly/Internal/IAnimNextRigVMExportInterface.h
--- a/Engine/Plugins/Experimental/AnimNext/Source/AnimNextUncookedOnly/Internal/IAnimNextRigVMExportInterface.h
+++ b/Engine/Plugins/Experimental/AnimNext/Source/AnimNextUncookedOnly/Internal/IAnimNextRigVMExportInterface.h
@@ -7,6 +7,16 @@
 
 struct FAnimNextParamType;
 
+UENUM()
+enum class EAnimNextExportAccessSpecifier : int32
+{
+	// Export can only be used/referenced in the graph it is declared in
+	Private,
+
+	// Export can be used/referenced in other graphs/contexts (e.g. BP, Verse)
+	Public,
+};
+
 UINTERFACE(meta=(CannotImplementInterfaceInBlueprint))
 class ANIMNEXTUNCOOKEDONLY_API UAnimNextRigVMExportInterface : public UInterface
 {
@@ -23,4 +33,10 @@ public:
 
 	// Get the export name (e.g. asset path + entry name)
 	virtual FName GetExportName() const = 0;
+
+	// Get the export access specifier
+	virtual EAnimNextExportAccessSpecifier GetExportAccessSpecifier() const = 0;
+
+	// Set the export access specifier
+	virtual void SetExportAccessSpecifier(EAnimNextExportAccessSpecifier InAccessSpecifier, bool bSetupUndoRedo = true) = 0;
 };
\ No newline at end of file
diff --git a/Engine/Plugins/Experimental/AnimNext/Source/AnimNextUncookedOnly/Internal/UncookedOnlyUtils.h b/Engine/Plugins/Experimental/AnimNext/Source/AnimNextUncookedOnly/Internal/UncookedOnlyUtils.h
--- a/Engine/Plugins/Experimental/AnimNext/Source/AnimNextUncookedOnly/Internal/UncookedOnlyUtils.h
+++ b/Engine/Plugins/Experimental/AnimNext/Source/AnimNextUncookedOnly/Internal/UncookedOnlyUtils.h
@@ -38,10 +38,10 @@ UENUM()
 enum class EAnimNextParameterFlags : uint32
 {
 	NoFlags = 0x0,
-	Private = 0x1,
+	Public = 0x1,
 	Read = 0x02,
 	Write = 0x04,
-	Bound = 0x08,
+	Declared = 0x08,
 	Max
 };
 
@@ -53,7 +53,7 @@ struct FAnimNextParameterAssetRegistryExportEntry
 	GENERATED_BODY()
 
 	FAnimNextParameterAssetRegistryExportEntry() = default;
-	
+
 	FAnimNextParameterAssetRegistryExportEntry(FName InName, const TInstancedStruct<FAnimNextParamInstanceIdentifier>& InInstanceId, const FAnimNextParamType& InType, EAnimNextParameterFlags InFlags = EAnimNextParameterFlags::NoFlags)
 		: Name(InName)
 		, InstanceId(InInstanceId)
diff --git a/Engine/Plugins/Experimental/AnimNext/Source/AnimNextUncookedOnly/Private/AnimNextRigVMAssetEntry.cpp b/Engine/Plugins/Experimental/AnimNext/Source/AnimNextUncookedOnly/Private/AnimNextRigVMAssetEntry.cpp
--- a/Engine/Plugins/Experimental/AnimNext/Source/AnimNextUncookedOnly/Private/AnimNextRigVMAssetEntry.cpp
+++ b/Engine/Plugins/Experimental/AnimNext/Source/AnimNextUncookedOnly/Private/AnimNextRigVMAssetEntry.cpp
@@ -16,6 +16,17 @@ bool UAnimNextRigVMAssetEntry::IsAsset() const
 	return IsPackageExternal() && !GetPackage()->HasAnyFlags(RF_Transient) && !HasAnyFlags(RF_Transient | RF_ClassDefaultObject);
 }
 
+#if WITH_EDITOR
+
+void UAnimNextRigVMAssetEntry::PostEditChangeProperty(FPropertyChangedEvent& PropertyChangedEvent)
+{
+	Super::PostEditChangeProperty(PropertyChangedEvent);
+
+	BroadcastModified();
+}
+
+#endif
+
 void UAnimNextRigVMAssetEntry::BroadcastModified()
 {
 	if(UAnimNextRigVMAssetEditorData* EditorData = Cast<UAnimNextRigVMAssetEditorData>(GetOuter()))
diff --git a/Engine/Plugins/Experimental/AnimNext/Source/AnimNextUncookedOnly/Private/Graph/AnimGraph/AnimGraphNodeBinding_AnimNextParameters.cpp b/Engine/Plugins/Experimental/AnimNext/Source/AnimNextUncookedOnly/Private/Graph/AnimGraph/AnimGraphNodeBinding_AnimNextParameters.cpp
--- a/Engine/Plugins/Experimental/AnimNext/Source/AnimNextUncookedOnly/Private/Graph/AnimGraph/AnimGraphNodeBinding_AnimNextParameters.cpp
+++ b/Engine/Plugins/Experimental/AnimNext/Source/AnimNextUncookedOnly/Private/Graph/AnimGraph/AnimGraphNodeBinding_AnimNextParameters.cpp
@@ -8,7 +8,7 @@
 #include "AnimBlueprintExtension_Base.h"
 
 #if WITH_EDITOR
-#include "AnimNextEditorModule.h"
+#include "IAnimNextEditorModule.h"
 #include "Param/ParameterPickerArgs.h"
 #include "Widgets/Input/SComboButton.h"
 #include "Widgets/Layout/SBox.h"
diff --git a/Engine/Plugins/Experimental/AnimNext/Source/AnimNextUncookedOnly/Private/Graph/AnimNextGraph_AnimationGraph.cpp b/Engine/Plugins/Experimental/AnimNext/Source/AnimNextUncookedOnly/Private/Graph/AnimNextGraph_AnimationGraph.cpp
--- a/Engine/Plugins/Experimental/AnimNext/Source/AnimNextUncookedOnly/Private/Graph/AnimNextGraph_AnimationGraph.cpp
+++ b/Engine/Plugins/Experimental/AnimNext/Source/AnimNextUncookedOnly/Private/Graph/AnimNextGraph_AnimationGraph.cpp
@@ -23,6 +23,23 @@ FName UAnimNextGraph_AnimationGraph::GetExportName() const
 	return GraphName;
 }
 
+EAnimNextExportAccessSpecifier UAnimNextGraph_AnimationGraph::GetExportAccessSpecifier() const
+{
+	return Access;
+}
+
+void UAnimNextGraph_AnimationGraph::SetExportAccessSpecifier(EAnimNextExportAccessSpecifier InAccessSpecifier, bool bSetupUndoRedo)
+{
+	if(bSetupUndoRedo)
+	{
+		Modify();
+	};
+
+	Access = InAccessSpecifier;
+
+	BroadcastModified();
+}
+
 FName UAnimNextGraph_AnimationGraph::GetEntryName() const
 {
 	return GraphName;
@@ -75,4 +92,4 @@ void UAnimNextGraph_AnimationGraph::SetRigVMGraph(URigVMGraph* InGraph)
 void UAnimNextGraph_AnimationGraph::SetEdGraph(URigVMEdGraph* InGraph)
 {
 	EdGraph = CastChecked<UAnimNextGraph_EdGraph>(InGraph);
-}
\ No newline at end of file
+}
diff --git a/Engine/Plugins/Experimental/AnimNext/Source/AnimNextUncookedOnly/Private/Graph/AnimNextGraph_Controller.cpp b/Engine/Plugins/Experimental/AnimNext/Source/AnimNextUncookedOnly/Private/Graph/AnimNextGraph_Controller.cpp
--- a/Engine/Plugins/Experimental/AnimNext/Source/AnimNextUncookedOnly/Private/Graph/AnimNextGraph_Controller.cpp
+++ b/Engine/Plugins/Experimental/AnimNext/Source/AnimNextUncookedOnly/Private/Graph/AnimNextGraph_Controller.cpp
@@ -166,9 +166,9 @@ bool UAnimNextGraph_Controller::SetAnimNextParameterNode(URigVMNode* ParameterNo
 			{
 				if(bSetupUndoRedo)
 				{
-					CancelUndoBracket();
+					CloseUndoBracket();
 				}
-				return false;
+				return true;
 			}
 		}
 
@@ -178,9 +178,9 @@ bool UAnimNextGraph_Controller::SetAnimNextParameterNode(URigVMNode* ParameterNo
 			{
 				if(bSetupUndoRedo)
 				{
-					CancelUndoBracket();
+					CloseUndoBracket();
 				}
-				return false;
+				return true;
 			}
 		}
 
@@ -191,9 +191,9 @@ bool UAnimNextGraph_Controller::SetAnimNextParameterNode(URigVMNode* ParameterNo
 			{
 				if(bSetupUndoRedo)
 				{
-					CancelUndoBracket();
+					CloseUndoBracket();
 				}
-				return false;
+				return true;
 			}
 		}
 	}
diff --git a/Engine/Plugins/Experimental/AnimNext/Source/AnimNextUncookedOnly/Private/Graph/AnimNextGraph_EditorData.cpp b/Engine/Plugins/Experimental/AnimNext/Source/AnimNextUncookedOnly/Private/Graph/AnimNextGraph_EditorData.cpp
--- a/Engine/Plugins/Experimental/AnimNext/Source/AnimNextUncookedOnly/Private/Graph/AnimNextGraph_EditorData.cpp
+++ b/Engine/Plugins/Experimental/AnimNext/Source/AnimNextUncookedOnly/Private/Graph/AnimNextGraph_EditorData.cpp
@@ -127,6 +127,28 @@ void UAnimNextGraph_EditorData::PostLoad()
 
 		RecompileVM();
 	}
+
+	if(GetLinkerCustomVersion(FFortniteMainBranchObjectVersion::GUID) < FFortniteMainBranchObjectVersion::AnimNextGraphAccessSpecifiers)
+	{
+		// Must preload entries so their data is populated as we will be modifying them
+		for(UAnimNextRigVMAssetEntry* Entry : Entries) 
+		{
+			Entry->GetLinker()->Preload(Entry);
+		}
+
+		// Force older assets to all have public symbols so they work as-is. Newer assets need user intervention as entries default to private
+		for(UAnimNextRigVMAssetEntry* Entry : Entries)
+		{
+			if(UAnimNextGraph_AnimationGraph* AnimationGraphEntry = Cast<UAnimNextGraph_AnimationGraph>(Entry))
+			{
+				AnimationGraphEntry->Access = EAnimNextExportAccessSpecifier::Public;
+			}
+			else if(UAnimNextGraph_Parameter* ParameterEntry = Cast<UAnimNextGraph_Parameter>(Entry))
+			{
+				ParameterEntry->Access = EAnimNextExportAccessSpecifier::Public;
+			}
+		}
+	}
 }
 
 void UAnimNextGraph_EditorData::RecompileVM()
diff --git a/Engine/Plugins/Experimental/AnimNext/Source/AnimNextUncookedOnly/Private/Graph/AnimNextGraph_EventGraph.cpp b/Engine/Plugins/Experimental/AnimNext/Source/AnimNextUncookedOnly/Private/Graph/AnimNextGraph_EventGraph.cpp
--- a/Engine/Plugins/Experimental/AnimNext/Source/AnimNextUncookedOnly/Private/Graph/AnimNextGraph_EventGraph.cpp
+++ b/Engine/Plugins/Experimental/AnimNext/Source/AnimNextUncookedOnly/Private/Graph/AnimNextGraph_EventGraph.cpp
@@ -8,20 +8,6 @@
 #include "Param/AnimNextTag.h"
 #include "Param/ParamType.h"
 
-FAnimNextParamType UAnimNextGraph_EventGraph::GetExportType() const
-{
-	return FAnimNextParamType::GetType<FAnimNextEntryPoint>();
-}
-
-FName UAnimNextGraph_EventGraph::GetExportName() const
-{
-	if(UAnimNextRigVMAsset* OuterAsset = GetTypedOuter<UAnimNextRigVMAsset>())
-	{
-		return UE::AnimNext::UncookedOnly::FUtils::GetQualifiedName(OuterAsset, GraphName);
-	}
-	return GraphName;
-}
-
 FText UAnimNextGraph_EventGraph::GetDisplayName() const
 {
 	return FText::FromName(GraphName);
diff --git a/Engine/Plugins/Experimental/AnimNext/Source/AnimNextUncookedOnly/Private/Graph/AnimNextGraph_Parameter.cpp b/Engine/Plugins/Experimental/AnimNext/Source/AnimNextUncookedOnly/Private/Graph/AnimNextGraph_Parameter.cpp
--- a/Engine/Plugins/Experimental/AnimNext/Source/AnimNextUncookedOnly/Private/Graph/AnimNextGraph_Parameter.cpp
+++ b/Engine/Plugins/Experimental/AnimNext/Source/AnimNextUncookedOnly/Private/Graph/AnimNextGraph_Parameter.cpp
@@ -16,6 +16,23 @@ FName UAnimNextGraph_Parameter::GetExportName() const
 	return GetParamName();
 }
 
+EAnimNextExportAccessSpecifier UAnimNextGraph_Parameter::GetExportAccessSpecifier() const
+{
+	return Access;
+}
+
+void UAnimNextGraph_Parameter::SetExportAccessSpecifier(EAnimNextExportAccessSpecifier InAccessSpecifier, bool bSetupUndoRedo)
+{
+	if(bSetupUndoRedo)
+	{
+		Modify();
+	};
+
+	Access = InAccessSpecifier;
+
+	BroadcastModified();
+}
+
 FAnimNextParamType UAnimNextGraph_Parameter::GetParamType() const
 {
 	return Type;
@@ -59,7 +76,7 @@ FInstancedPropertyBag& UAnimNextGraph_Parameter::GetPropertyBag() const
 	// TODO: move property bag for defaults onto this entry!
 	UAnimNextGraph* Asset = GetTypedOuter<UAnimNextGraph>();
 	check(Asset);
-	return Asset->PropertyBag;
+	return Asset->DefaultState.State;
 }
 
 void UAnimNextGraph_Parameter::SetEntryName(FName InName, bool bSetupUndoRedo)
diff --git a/Engine/Plugins/Experimental/AnimNext/Source/AnimNextUncookedOnly/Private/Param/ObjectProxyType.cpp b/Engine/Plugins/Experimental/AnimNext/Source/AnimNextUncookedOnly/Private/Param/ObjectProxyType.cpp
--- a/Engine/Plugins/Experimental/AnimNext/Source/AnimNextUncookedOnly/Private/Param/ObjectProxyType.cpp
+++ b/Engine/Plugins/Experimental/AnimNext/Source/AnimNextUncookedOnly/Private/Param/ObjectProxyType.cpp
@@ -2,7 +2,7 @@
 
 #include "ObjectProxyType.h"
 
-#include "IAnimNextModule.h"
+#include "AnimNextModule.h"
 #include "IUniversalObjectLocatorEditorModule.h"
 #include "UniversalObjectLocatorEditor.h"
 #include "UniversalObjectLocatorFragmentType.h"
@@ -135,7 +135,7 @@ bool FObjectProxyType::FindParameterInfo(const TInstancedStruct<FAnimNextParamIn
 		return false;
 	}
 
-	IAnimNextModule& AnimNextModule = FModuleManager::GetModuleChecked<IAnimNextModule>("AnimNext");
+	UE::AnimNext::FModule& AnimNextModule = FModuleManager::GetModuleChecked<UE::AnimNext::FModule>("AnimNext");
 	TSharedPtr<FObjectProxyFactory> ObjectProxyFactory = StaticCastSharedPtr<FObjectProxyFactory>(AnimNextModule.FindParameterSourceFactory("ObjectProxy"));
 	if(!ObjectProxyFactory.IsValid())
 	{
@@ -179,7 +179,7 @@ void FObjectProxyType::ForEachParameter(const TInstancedStruct<FAnimNextParamIns
 		return;
 	}
 
-	IAnimNextModule& AnimNextModule = FModuleManager::GetModuleChecked<IAnimNextModule>("AnimNext");
+	UE::AnimNext::FModule& AnimNextModule = FModuleManager::GetModuleChecked<UE::AnimNext::FModule>("AnimNext");
 	TSharedPtr<FObjectProxyFactory> ObjectProxyFactory = StaticCastSharedPtr<FObjectProxyFactory>(AnimNextModule.FindParameterSourceFactory("ObjectProxy"));
 	if(!ObjectProxyFactory.IsValid())
 	{
diff --git a/Engine/Plugins/Experimental/AnimNext/Source/AnimNextUncookedOnly/Private/UncookedOnlyUtils.cpp b/Engine/Plugins/Experimental/AnimNext/Source/AnimNextUncookedOnly/Private/UncookedOnlyUtils.cpp
--- a/Engine/Plugins/Experimental/AnimNext/Source/AnimNextUncookedOnly/Private/UncookedOnlyUtils.cpp
+++ b/Engine/Plugins/Experimental/AnimNext/Source/AnimNextUncookedOnly/Private/UncookedOnlyUtils.cpp
@@ -538,6 +538,7 @@ void FUtils::CompileVM(UAnimNextGraph* InGraph)
 	InGraph->VMRuntimeSettings = EditorData->VMRuntimeSettings;
 	InGraph->EntryPoints.Empty();
 	InGraph->ResolvedRootTraitHandles.Empty();
+	InGraph->ResolvedEntryPoints.Empty();
 	InGraph->ExecuteDefinition = FAnimNextGraphEvaluatorExecuteDefinition();
 	InGraph->SharedDataBuffer.Empty();
 	InGraph->GraphReferencedObjects.Empty();
@@ -697,8 +698,8 @@ void FUtils::CompileVM(UAnimNextGraph* InGraph)
 
 	for(FAnimNextParameterAssetRegistryExportEntry& Entry : Exports.Parameters)
 	{
-		// Required parameters are those that are read in this asset but not bound in this asset as state
-		if(EnumHasAnyFlags(Entry.GetFlags(), EAnimNextParameterFlags::Read) && !EnumHasAnyFlags(Entry.GetFlags(), EAnimNextParameterFlags::Bound))
+		// Required parameters are those that are read in this asset but not declared in this asset as state
+		if(EnumHasAllFlags(Entry.GetFlags(), EAnimNextParameterFlags::Read) && !EnumHasAnyFlags(Entry.GetFlags(), EAnimNextParameterFlags::Declared))
 		{
 			InGraph->RequiredParameters.Emplace(Entry.Name, Entry.Type, Entry.InstanceId);
 		}
@@ -851,30 +852,72 @@ void FUtils::CompileStruct(UAnimNextGraph* InGraph)
 
 	TGuardValue<bool> CompilingGuard(EditorData->bIsCompiling, true);
 
-	TArray<FPropertyBagPropertyDesc> PropertyDescs;
-	PropertyDescs.Reserve(EditorData->Entries.Num());
-	
+	struct FStructEntryInfo
+	{
+		FName Name;
+		FAnimNextParamType Type;
+		EAnimNextExportAccessSpecifier AccessSpecifier;
+	};
+
+	TArray<FStructEntryInfo> StructEntryInfos;
+	StructEntryInfos.Reserve(EditorData->Entries.Num());
+
 	// Gather all parameters in this asset
 	for(const UAnimNextRigVMAssetEntry* Entry : EditorData->Entries)
 	{
-		if(const IAnimNextRigVMParameterInterface* Parameter = Cast<IAnimNextRigVMParameterInterface>(Entry))
+		const IAnimNextRigVMExportInterface* Export = Cast<IAnimNextRigVMExportInterface>(Entry);
+		const IAnimNextRigVMParameterInterface* Parameter = Cast<IAnimNextRigVMParameterInterface>(Entry);
+		if(Export && Parameter)
 		{
-			const FAnimNextParamType& Type = Parameter->GetParamType();
+			const FAnimNextParamType& Type = Export->GetExportType();
 			ensure(Type.IsValid());
-			const FName Name = Parameter->GetParamName();
+			const FName Name = Export->GetExportName();
+			const EAnimNextExportAccessSpecifier AccessSpecifier = Export->GetExportAccessSpecifier();
 
-			PropertyDescs.Emplace(Name, Type.GetContainerType(), Type.GetValueType(), Type.GetValueTypeObject());
+			StructEntryInfos.Add( { Name, FAnimNextParamType(Type.GetValueType(),  Type.GetContainerType(), Type.GetValueTypeObject()), AccessSpecifier } );
 		}
 	}
 
-	if(PropertyDescs.Num() > 0)
+	// Sort private entries first & then by size, largest first, for better packing
+	static_assert(EAnimNextExportAccessSpecifier::Private < EAnimNextExportAccessSpecifier::Public, "Private must be less than Public as parameters are sorted internally according to this assumption");
+	StructEntryInfos.Sort([](const FStructEntryInfo& InLHS, const FStructEntryInfo& InRHS)
+	{
+		if(InLHS.AccessSpecifier < InRHS.AccessSpecifier)
+		{
+			return true;
+		}
+		else
+		{
+			return InLHS.Type.GetSize() > InRHS.Type.GetSize();
+		}
+	});
+
+	if(StructEntryInfos.Num() > 0)
 	{
-		// find any existing IDs for old properties with name-matching
+		// Build PropertyDescs to batch-create the property bag
+		TArray<FPropertyBagPropertyDesc> PropertyDescs;
+		PropertyDescs.Reserve(StructEntryInfos.Num());
+
+		InGraph->DefaultState.PublicParameterStartIndex = INDEX_NONE;
+
+		for(int32 EntryIndex = 0; EntryIndex < StructEntryInfos.Num(); ++EntryIndex)
+		{
+			const FStructEntryInfo& StructEntryInfo = StructEntryInfos[EntryIndex];
+			// Find the first parameter that is public and record it
+			if(StructEntryInfo.AccessSpecifier == EAnimNextExportAccessSpecifier::Public)
+			{
+				InGraph->DefaultState.PublicParameterStartIndex = EntryIndex;
+			}
+			PropertyDescs.Emplace(StructEntryInfo.Name, StructEntryInfo.Type.ContainerType, StructEntryInfo.Type.ValueType, StructEntryInfo.Type.ValueTypeObject);
+		}
+
+		// Find any existing IDs for old properties with name-matching
+		// TODO: linear search - we could cache the name->GUID lookup in editor to accelerate this.
 		for(FPropertyBagPropertyDesc& NewDesc : PropertyDescs)
 		{
-			if(InGraph->PropertyBag.GetPropertyBagStruct())
+			if(InGraph->DefaultState.State.GetPropertyBagStruct())
 			{
-				for(const FPropertyBagPropertyDesc& ExistingDesc : InGraph->PropertyBag.GetPropertyBagStruct()->GetPropertyDescs())
+				for(const FPropertyBagPropertyDesc& ExistingDesc : InGraph->DefaultState.State.GetPropertyBagStruct()->GetPropertyDescs())
 				{
 					if(ExistingDesc.Name == NewDesc.Name)
 					{
@@ -887,11 +930,11 @@ void FUtils::CompileStruct(UAnimNextGraph* InGraph)
 
 		// Create new property bag and migrate
 		const UPropertyBag* NewBagStruct = UPropertyBag::GetOrCreateFromDescs(PropertyDescs);
-		InGraph->PropertyBag.MigrateToNewBagStruct(NewBagStruct);
+		InGraph->DefaultState.State.MigrateToNewBagStruct(NewBagStruct);
 	}
 	else
 	{
-		InGraph->PropertyBag.Reset();
+		InGraph->DefaultState.Reset();
 	}
 }
 
@@ -909,7 +952,7 @@ UAnimNextRigVMAssetEditorData* FUtils::GetEditorData(UAnimNextRigVMAsset* InAsse
 
 FInstancedPropertyBag* FUtils::GetPropertyBag(UAnimNextGraph* InAnimNextGraph)
 {
-	FInstancedPropertyBag* InstancedPropertyBag = &InAnimNextGraph->PropertyBag;
+	FInstancedPropertyBag* InstancedPropertyBag = &InAnimNextGraph->DefaultState.State;
 
 	return InstancedPropertyBag;
 }
@@ -1364,9 +1407,13 @@ void FUtils::GetAssetParameters(const UAnimNextRigVMAssetEditorData* EditorData,
 	{
 		if(const IAnimNextRigVMExportInterface* ExportInterface = Cast<IAnimNextRigVMExportInterface>(Entry))
 		{
-			// TODO: Public/private symbols would influence whether this would be exposed to the asset registry here
-			FAnimNextParameterAssetRegistryExportEntry NewParam(ExportInterface->GetExportName(), TInstancedStruct<FAnimNextParamInstanceIdentifier>(), ExportInterface->GetExportType(), EAnimNextParameterFlags::Bound);
-			AddParamToSet(NewParam, OutExports);
+			EAnimNextParameterFlags Flags = EAnimNextParameterFlags::Declared;
+			if(ExportInterface->GetExportAccessSpecifier() == EAnimNextExportAccessSpecifier::Public)
+			{
+				Flags |= EAnimNextParameterFlags::Public;
+				FAnimNextParameterAssetRegistryExportEntry NewParam(ExportInterface->GetExportName(), TInstancedStruct<FAnimNextParamInstanceIdentifier>(), ExportInterface->GetExportType(), Flags);
+				AddParamToSet(NewParam, OutExports);
+			}
 		}
 		if(const IAnimNextRigVMGraphInterface* GraphInterface = Cast<IAnimNextRigVMGraphInterface>(Entry))
 		{
@@ -1851,6 +1898,13 @@ void FUtils::CompileSchedule(UAnimNextSchedule* InSchedule)
 					}
 				}
 
+				// We must have a reference pose
+				if (!GraphEntry->ReferencePose.IsValid())
+				{
+					UE_LOG(LogAnimation, Error, TEXT("AnimNext: Invalid reference pose supplied to graph task"));
+					bValid = false;
+				}
+
 				// Validate terms and check against priors
 				for(int32 TermIndex = 0; TermIndex < GraphEntry->Terms.Num(); ++TermIndex)
 				{
@@ -1891,6 +1945,8 @@ void FUtils::CompileSchedule(UAnimNextSchedule* InSchedule)
 					GraphTask.EntryPoint = FAnimNextParam(GraphEntry->EntryPoint);
 					GraphTask.Graph = GraphEntry->Graph;
 					GraphTask.DynamicGraph = FAnimNextParam(GraphEntry->DynamicGraph);
+					GraphTask.ReferencePose = FAnimNextParam(GraphEntry->ReferencePose);
+					GraphTask.LOD = FAnimNextParam(GraphEntry->LOD);
 					if(GraphEntry->Graph == nullptr && GraphEntry->DynamicGraph.IsValid())
 					{
 						Algo::Transform(GraphEntry->RequiredParameters, GraphTask.SuppliedParameters, [](const FAnimNextEditorParam& InParam){ return FAnimNextParam(InParam); });
@@ -2079,8 +2135,7 @@ void FUtils::CompileSchedule(UAnimNextSchedule* InSchedule)
 
 					TrackExternalParameters(GraphEntry->RequiredParameters, InDistance, Entry, &InEntries, ArrayIndex, true);
 
-					// TODO: need to not require defaults here - use static graph params. In fact, these params should probably be defined at the schedule level.
-					TrackExternalParameters(FAnimNextScheduleGraphTask::GetRequiredParametersInternal(), InDistance, Entry, &InEntries, ArrayIndex, true);
+					TrackExternalParameters({ GraphEntry->DynamicGraph, GraphEntry->EntryPoint, GraphEntry->ReferencePose, GraphEntry->LOD }, InDistance, Entry, &InEntries, ArrayIndex, true);
 				}
 				else if (UAnimNextScheduleEntry_Port* PortEntry = Cast<UAnimNextScheduleEntry_Port>(Entry))
 				{
diff --git a/Engine/Plugins/Experimental/AnimNextAnimGraph/Source/AnimNextAnimGraph/Private/TraitInterfaces/IUpdate.cpp b/Engine/Plugins/Experimental/AnimNextAnimGraph/Source/AnimNextAnimGraph/Private/TraitInterfaces/IUpdate.cpp
--- a/Engine/Plugins/Experimental/AnimNextAnimGraph/Source/AnimNextAnimGraph/Private/TraitInterfaces/IUpdate.cpp
+++ b/Engine/Plugins/Experimental/AnimNextAnimGraph/Source/AnimNextAnimGraph/Private/TraitInterfaces/IUpdate.cpp
@@ -5,6 +5,7 @@
 #include "TraitInterfaces/IHierarchy.h"
 #include "Graph/GraphInstanceComponent.h"
 #include "AnimNextAnimGraphStats.h"
+#include "Param/ParamStack.h"
 
 DEFINE_STAT(STAT_AnimNext_UpdateGraph);
 
@@ -490,6 +491,9 @@ namespace UE::AnimNext
 			It.Value()->PreUpdate(TraversalContext);
 		}
 
+		// Push any parameters this graph owns
+		FParamStack::FPushedLayerHandle PushedLayerHandle = GraphInstance.UpdateAndPushGraphState(DeltaTime);
+
 		// Add the graph root to start the update process
 		Private::FUpdateEntry RootEntry(GraphInstance.GetGraphRootPtr(), RootState);
 		TraversalContext.PushUpdateEntry(&RootEntry);
@@ -603,6 +607,9 @@ namespace UE::AnimNext
 			}
 		}
 
+		// Pop any parameters we pushed above
+		GraphInstance.PopGraphState(PushedLayerHandle);
+
 		// Clear our executing entry
 		TraversalContext.ExecutingEntry = nullptr;
 
diff --git a/Engine/Source/Runtime/Core/Public/UObject/FortniteMainBranchObjectVersion.h b/Engine/Source/Runtime/Core/Public/UObject/FortniteMainBranchObjectVersion.h
--- a/Engine/Source/Runtime/Core/Public/UObject/FortniteMainBranchObjectVersion.h
+++ b/Engine/Source/Runtime/Core/Public/UObject/FortniteMainBranchObjectVersion.h
@@ -468,6 +468,9 @@ struct FFortniteMainBranchObjectVersion
 		// PCGGridDescriptor
 		PCGGridDescriptor,
 
+		// AnimNext graphs now have public/private state
+		AnimNextGraphAccessSpecifiers,
+
 		// -----<new versions can be added above this line>-------------------------------------------------
 		VersionPlusOne,
 		LatestVersion = VersionPlusOne - 1
