Fix for Horde errors Mac ARM and EditorTestsCooked
#rnx

diff --git a/Engine/Plugins/Experimental/AnimNext/Source/AnimNext/Public/TraitCore/TraitInterfaceUID.h b/Engine/Plugins/Experimental/AnimNext/Source/AnimNext/Public/TraitCore/TraitInterfaceUID.h
--- a/Engine/Plugins/Experimental/AnimNext/Source/AnimNext/Public/TraitCore/TraitInterfaceUID.h
+++ b/Engine/Plugins/Experimental/AnimNext/Source/AnimNext/Public/TraitCore/TraitInterfaceUID.h
@@ -39,7 +39,7 @@ namespace UE::AnimNext
 		{
 		}
 
-#if !UE_BUILD_SHIPPING
+#if !UE_BUILD_SHIPPING || WITH_EDITOR
 		// Returns a literal string to the interface name
 		constexpr const TCHAR* GetInterfaceName() const { return InterfaceName; }
 		constexpr const TCHAR* GetInterfaceShortName() const { return InterfaceShortName; }
@@ -56,7 +56,7 @@ namespace UE::AnimNext
 
 		FTraitInterfaceUIDRaw		UID;
 
-#if !UE_BUILD_SHIPPING
+#if !UE_BUILD_SHIPPING || WITH_EDITOR
 		const TCHAR*				InterfaceName = nullptr;
 		const TCHAR*				InterfaceShortName = nullptr;
 #endif
diff --git a/Engine/Plugins/Experimental/AnimNext/Source/AnimNextEditor/Private/Graph/AnimNextGraph_EdGraphNodeCustomization.cpp b/Engine/Plugins/Experimental/AnimNext/Source/AnimNextEditor/Private/Graph/AnimNextGraph_EdGraphNodeCustomization.cpp
--- a/Engine/Plugins/Experimental/AnimNext/Source/AnimNextEditor/Private/Graph/AnimNextGraph_EdGraphNodeCustomization.cpp
+++ b/Engine/Plugins/Experimental/AnimNext/Source/AnimNextEditor/Private/Graph/AnimNextGraph_EdGraphNodeCustomization.cpp
@@ -34,7 +34,7 @@ void FAnimNextGraph_EdGraphNodeCustomization::PendingDelete()
 		{
 			if (TSharedPtr<STraitEditorView> TraitEditorView = StaticCastSharedPtr<STraitEditorView>(DockTab->GetContent().ToSharedPtr()))
 			{
-				TraitEditorView->SetTraitData(FTraitStackData(nullptr));
+				TraitEditorView->SetTraitData(FTraitStackData());
 			}
 		}
 	}
@@ -87,7 +87,7 @@ void FAnimNextGraph_EdGraphNodeCustomization::CustomizeObjects(IDetailLayoutBuil
 				{
 					if (TSharedPtr<STraitEditorView> TraitEditorView = StaticCastSharedPtr<STraitEditorView>(DockTab->GetContent().ToSharedPtr()))
 					{
-						TraitEditorView->SetTraitData(EdGraphNode->IsTraitStack() ? FTraitStackData(EdGraphNode) : FTraitStackData(nullptr));
+						TraitEditorView->SetTraitData(EdGraphNode->IsTraitStack() ? FTraitStackData(EdGraphNode) : FTraitStackData());
 					}
 				}
 			}
diff --git a/Engine/Plugins/Experimental/AnimNext/Source/AnimNextEditor/Private/Graph/STraitEditorView.cpp b/Engine/Plugins/Experimental/AnimNext/Source/AnimNextEditor/Private/Graph/STraitEditorView.cpp
--- a/Engine/Plugins/Experimental/AnimNext/Source/AnimNextEditor/Private/Graph/STraitEditorView.cpp
+++ b/Engine/Plugins/Experimental/AnimNext/Source/AnimNextEditor/Private/Graph/STraitEditorView.cpp
@@ -610,7 +610,8 @@ int32 STraitEditorView::GetTraitPinIndex(UAnimNextGraph_EdGraphNode* InEdGraphNo
 									{
 										if (!FTraitEditorUtils::IsInternalTraitInteface(TraitInterface))
 										{
-											if (int32 Index = StackUsedInterfaces.Find(TraitInterface) != INDEX_NONE)
+											int32 Index = StackUsedInterfaces.Find(TraitInterface);
+											if (Index != INDEX_NONE)
 											{
 												StackUsedInterfacesNumHits[Index]++;
 											}
@@ -626,7 +627,8 @@ int32 STraitEditorView::GetTraitPinIndex(UAnimNextGraph_EdGraphNode* InEdGraphNo
 									{
 										if (!FTraitEditorUtils::IsInternalTraitInteface(TraitInterface))
 										{
-											if (int32 Index = StackUsedInterfaces.Find(TraitInterface) != INDEX_NONE)
+											int32 Index = StackUsedInterfaces.Find(TraitInterface);
+											if (Index != INDEX_NONE)
 											{
 												StackUsedInterfacesNumHits[Index]++;
 											}
diff --git a/Engine/Plugins/Experimental/AnimNext/Source/AnimNextEditor/Private/Graph/STraitEditorView.h b/Engine/Plugins/Experimental/AnimNext/Source/AnimNextEditor/Private/Graph/STraitEditorView.h
--- a/Engine/Plugins/Experimental/AnimNext/Source/AnimNextEditor/Private/Graph/STraitEditorView.h
+++ b/Engine/Plugins/Experimental/AnimNext/Source/AnimNextEditor/Private/Graph/STraitEditorView.h
@@ -36,6 +36,11 @@ namespace UE::AnimNext::Editor
 
 struct FTraitStackData
 {
+	FTraitStackData() = default;
+	explicit FTraitStackData(const TWeakObjectPtr<UAnimNextGraph_EdGraphNode>& InEdGraphNodeWeak)
+		: EdGraphNodeWeak(InEdGraphNodeWeak)
+	{}
+
 	TWeakObjectPtr<UAnimNextGraph_EdGraphNode> EdGraphNodeWeak = nullptr;
 };
 
