AnimNext: Fixed crash and inconsistency when selecting parameter types from the details panel

diff --git a/Engine/Plugins/Experimental/AnimNext/Source/AnimNextEditor/Private/Param/ParamPropertyCustomization.cpp b/Engine/Plugins/Experimental/AnimNext/Source/AnimNextEditor/Private/Param/ParamPropertyCustomization.cpp
--- a/Engine/Plugins/Experimental/AnimNext/Source/AnimNextEditor/Private/Param/ParamPropertyCustomization.cpp
+++ b/Engine/Plugins/Experimental/AnimNext/Source/AnimNextEditor/Private/Param/ParamPropertyCustomization.cpp
@@ -14,6 +14,7 @@
 #include "Param/ParamCompatibility.h"
 #include "Param/ParamUtils.h"
 #include "HAL/PlatformApplicationMisc.h"
+#include "RigVMCore/RigVMRegistry.h"
 
 #define LOCTEXT_NAMESPACE "ParamPropertyCustomization"
 
@@ -46,13 +47,26 @@ void FParamPropertyCustomization::CustomizeHeader(TSharedRef<IPropertyHandle> In
 
 	FParameterPickerArgs PickerArgs;
 	PickerArgs.bMultiSelect = false;
-	PickerArgs.OnFilterParameterType = FOnFilterParameterType::CreateLambda([FilterType](const FAnimNextParamType& InParameterType)
+	PickerArgs.OnFilterParameterType = FOnFilterParameterType::CreateLambda([FilterType](const FAnimNextParamType& InParamType)
 	{
-		if(!FilterType.IsValid() || FParamUtils::GetCompatibility(FilterType, InParameterType).IsCompatible())
+		if(FilterType.IsValid())
 		{
-			return EFilterParameterResult::Include;
+			if(!FParamUtils::GetCompatibility(FilterType, InParamType).IsCompatible())
+			{
+				return EFilterParameterResult::Exclude;
+			}
 		}
-		return EFilterParameterResult::Exclude;
+
+		if(InParamType.IsValid())
+		{
+			const FRigVMTemplateArgumentType RigVMType = InParamType.ToRigVMTemplateArgument();
+			if(!RigVMType.IsValid() || FRigVMRegistry::Get().GetTypeIndex(RigVMType) == INDEX_NONE)
+			{
+				return EFilterParameterResult::Exclude;
+			}
+		}
+
+		return EFilterParameterResult::Include;
 	});
 	PickerArgs.NewParameterType = FilterType;
 	PickerArgs.OnInstanceIdChanged = FOnInstanceIdChanged::CreateSPLambda(this, [this](const TInstancedStruct<FAnimNextParamInstanceIdentifier>& InInstanceId)
@@ -100,6 +114,7 @@ void FParamPropertyCustomization::CustomizeHeader(TSharedRef<IPropertyHandle> In
 					Param = ParamValue;
 					return true;
 				});
+				CachedParam = ParamValue;
 			}
 			else if(ParamStruct == FAnimNextParam::StaticStruct())
 			{
@@ -110,10 +125,10 @@ void FParamPropertyCustomization::CustomizeHeader(TSharedRef<IPropertyHandle> In
 					Param = ParamValue;
 					return true;
 				});
+				CachedParam = FAnimNextEditorParam(ParamValue);
 			}
 			PropertyHandle->NotifyPostChange(EPropertyChangeType::ValueSet);
 			PropertyHandle->NotifyFinishedChangingProperties();
-			Refresh();
 		}
 	});
 
