AnimNext: Dont display deprecated units & dispatch factories in context menu

diff --git a/Engine/Plugins/Experimental/AnimNext/Source/AnimNextEditor/Private/Common/SActionMenu.cpp b/Engine/Plugins/Experimental/AnimNext/Source/AnimNextEditor/Private/Common/SActionMenu.cpp
--- a/Engine/Plugins/Experimental/AnimNext/Source/AnimNextEditor/Private/Common/SActionMenu.cpp
+++ b/Engine/Plugins/Experimental/AnimNext/Source/AnimNextEditor/Private/Common/SActionMenu.cpp
@@ -39,6 +39,12 @@ void SActionMenu::CollectAllAnimNextGraphActions(FGraphContextMenuBuilder& MenuB
 			continue;
 		}
 
+		// skip deprecated units
+		if(Function.Struct->HasMetaData(FRigVMStruct::DeprecatedMetaName))
+		{
+			continue;
+		}
+
 		FString CategoryMetadata, DisplayNameMetadata, MenuDescSuffixMetadata;
 		Struct->GetStringMetaDataHierarchical(FRigVMStruct::CategoryMetaName, &CategoryMetadata);
 		Struct->GetStringMetaDataHierarchical(FRigVMStruct::DisplayNameMetaName, &DisplayNameMetadata);
@@ -72,6 +78,12 @@ void SActionMenu::CollectAllAnimNextGraphActions(FGraphContextMenuBuilder& MenuB
 			continue;
 		}
 
+		// skip deprecated factories
+		if(Factory->GetScriptStruct()->HasMetaData(FRigVMStruct::DeprecatedMetaName))
+		{
+			continue;
+		}
+
 		FText NodeCategory = FText::FromString(Factory->GetCategory());
 		FText MenuDesc = FText::FromString(Factory->GetNodeTitle(FRigVMTemplateTypeMap()));
 		FText ToolTip = Factory->GetNodeTooltip(FRigVMTemplateTypeMap());
