AnimNext: Added compiler results redirection

Redirected trait editor results, graph compilation results and schedule compilation results there

diff --git a/Engine/Plugins/Experimental/AnimNext/Source/AnimNext/Internal/Scheduler/AnimNextSchedule.h b/Engine/Plugins/Experimental/AnimNext/Source/AnimNext/Internal/Scheduler/AnimNextSchedule.h
--- a/Engine/Plugins/Experimental/AnimNext/Source/AnimNext/Internal/Scheduler/AnimNextSchedule.h
+++ b/Engine/Plugins/Experimental/AnimNext/Source/AnimNext/Internal/Scheduler/AnimNextSchedule.h
@@ -37,6 +37,11 @@ namespace UE::AnimNext::UncookedOnly
 	struct FUtils;
 }
 
+namespace UE::AnimNext::Editor
+{
+	class FModule;
+}
+
 USTRUCT()
 struct FAnimNextScheduleEntryTerm
 {
@@ -243,6 +248,7 @@ private:
 	friend struct FAnimNextSchedulerEntry;
 	friend class UAnimNextComponent;
 	friend class UAnimNextSchedulerWorldSubsystem;
+	friend class UE::AnimNext::Editor::FModule;
 	friend class UE::AnimNext::UncookedOnly::FModule;
 	friend struct UE::AnimNext::UncookedOnly::FUtils;
 
@@ -268,7 +274,10 @@ private:
 
 	// Function hook used to expose to asset registry
 	static TUniqueFunction<void(const UAnimNextSchedule*, FAssetRegistryTagsContext)> GetAssetRegistryTagsFunction;
-	
+
+	// Delegate fired when compiled
+	FSimpleMulticastDelegate CompiledEvent;
+
 	// Editor only
 	// TODO: move this into an editor only subobject
 	// TODO: this is currently only a linear list, we want it to be a graph
diff --git a/Engine/Plugins/Experimental/AnimNext/Source/AnimNextEditor/Private/AnimNextEditorModule.cpp b/Engine/Plugins/Experimental/AnimNext/Source/AnimNextEditor/Private/AnimNextEditorModule.cpp
--- a/Engine/Plugins/Experimental/AnimNext/Source/AnimNextEditor/Private/AnimNextEditorModule.cpp
+++ b/Engine/Plugins/Experimental/AnimNext/Source/AnimNextEditor/Private/AnimNextEditorModule.cpp
@@ -226,7 +226,7 @@ void FModule::StartupModule()
 	Workspace::IWorkspaceEditorModule& WorkspaceModule = FModuleManager::Get().LoadModuleChecked<Workspace::IWorkspaceEditorModule>("WorkspaceEditor");
 	const TSharedPtr<FAnimNextGraphItemDetails> GraphItemDetails = MakeShareable<FAnimNextGraphItemDetails>(new FAnimNextGraphItemDetails());
 	WorkspaceModule.RegisterWorkspaceItemDetails(Workspace::FOutlinerItemDetailsId(FAnimNextGraphOutlinerData::StaticStruct()->GetFName()), StaticCastSharedPtr<UE::Workspace::IWorkspaceOutlinerItemDetails>(GraphItemDetails));
-	FAnimNextGraphItemDetails::RegisterToolMenuExtensions();	
+	FAnimNextGraphItemDetails::RegisterToolMenuExtensions();
 }
 
 void FModule::ShutdownModule()
@@ -310,6 +310,21 @@ void FModule::RegisterWorkspaceDocumentTypes(Workspace::IWorkspaceEditorModule&
 			DetailsViewArgs.NameAreaSettings = FDetailsViewArgs::HideNameArea;
 			TSharedRef<IDetailsView> DetailsView = PropertyEditorModule.CreateDetailView(DetailsViewArgs);
 			DetailsView->SetObject(Schedule);
+
+			TWeakPtr<Workspace::IWorkspaceEditor> WeakWorkspaceEditor = InContext.WorkspaceEditor;
+			Schedule->CompiledEvent.RemoveAll(&InContext.WorkspaceEditor.Get());
+			Schedule->CompiledEvent.AddSPLambda(&InContext.WorkspaceEditor.Get(), [WeakWorkspaceEditor]()
+				{
+					if (TSharedPtr<Workspace::IWorkspaceEditor> WorkspaceEditor = WeakWorkspaceEditor.Pin())
+					{
+						int32 NumEntries = FMessageLog("AnimNextCompilerResults").NumMessages(EMessageSeverity::Warning);
+						if(NumEntries > 0)
+						{
+							WorkspaceEditor->GetTabManager()->TryInvokeTab(CompilerResultsTabName);
+						}
+					}
+				});
+			
 			return DetailsView;
 		}
 	), Workspace::WorkspaceTabs::TopMiddleDocumentArea);
@@ -739,6 +754,19 @@ void FModule::RegisterWorkspaceDocumentTypes(Workspace::IWorkspaceEditorModule&
 						WorkspaceEditor->RefreshDetails();
 					}
 				});
+
+			EditorData->RigVMCompiledEvent.RemoveAll(&InContext.WorkspaceEditor.Get());
+			EditorData->RigVMCompiledEvent.AddSPLambda(&InContext.WorkspaceEditor.Get(), [WeakWorkspaceEditor](UObject*, URigVM*, FRigVMExtendedExecuteContext&)
+				{
+					if (TSharedPtr<Workspace::IWorkspaceEditor> WorkspaceEditor = WeakWorkspaceEditor.Pin())
+					{
+						int32 NumEntries = FMessageLog("AnimNextCompilerResults").NumMessages(EMessageSeverity::Warning);
+						if(NumEntries > 0)
+						{
+							WorkspaceEditor->GetTabManager()->TryInvokeTab(CompilerResultsTabName);
+						}
+					}
+				});
 		}
 
 		if (WorkspaceMakeDocumentWidgetDelegate.IsBound())
diff --git a/Engine/Plugins/Experimental/AnimNext/Source/AnimNextEditor/Private/Graph/AnimNextCompilerResultsTabSummoner.cpp b/Engine/Plugins/Experimental/AnimNext/Source/AnimNextEditor/Private/Graph/AnimNextCompilerResultsTabSummoner.cpp
--- a/Engine/Plugins/Experimental/AnimNext/Source/AnimNextEditor/Private/Graph/AnimNextCompilerResultsTabSummoner.cpp
+++ b/Engine/Plugins/Experimental/AnimNext/Source/AnimNextEditor/Private/Graph/AnimNextCompilerResultsTabSummoner.cpp
@@ -35,19 +35,9 @@ void SAnimNextCompilerResultsWidget::CreateMessageLog(const TWeakPtr<UE::Workspa
 	if (const TSharedPtr<UE::Workspace::IWorkspaceEditor> WorkspaceEditorShared = InWorkspaceEditorWeak.Pin())
 	{
 		FMessageLogModule& MessageLogModule = FModuleManager::LoadModuleChecked<FMessageLogModule>("MessageLog");
-		const FName LogName = *FString::Printf(TEXT("%s_CompilerResultsLog"), WorkspaceEditorShared.IsValid() ? *WorkspaceEditorShared->GetEditorName().ToString() : TEXT("TraitEditor"));
-		// Reuse any existing log, or create a new one (that is not held onto bey the message log system)
-		if (MessageLogModule.IsRegisteredLogListing(LogName))
-		{
-			CompilerResultsListing = MessageLogModule.GetLogListing(LogName);
-		}
-		else
-		{
-			FMessageLogInitializationOptions LogInitOptions;
-			LogInitOptions.bShowInLogWindow = false;
-			CompilerResultsListing = MessageLogModule.CreateLogListing(LogName, LogInitOptions);
-		}
-
+		const FName LogName("AnimNextCompilerResults");
+		check(MessageLogModule.IsRegisteredLogListing(LogName));
+		CompilerResultsListing = MessageLogModule.GetLogListing(LogName);
 		CompilerResults = MessageLogModule.CreateLogListingWidget(CompilerResultsListing.ToSharedRef());
 	}
 }
diff --git a/Engine/Plugins/Experimental/AnimNext/Source/AnimNextEditor/Private/Graph/STraitEditorView.cpp b/Engine/Plugins/Experimental/AnimNext/Source/AnimNextEditor/Private/Graph/STraitEditorView.cpp
--- a/Engine/Plugins/Experimental/AnimNext/Source/AnimNextEditor/Private/Graph/STraitEditorView.cpp
+++ b/Engine/Plugins/Experimental/AnimNext/Source/AnimNextEditor/Private/Graph/STraitEditorView.cpp
@@ -17,6 +17,7 @@
 #include "IMessageLogListing.h"
 #include "Graph/AnimNextGraph_Controller.h"
 #include "Graph/AnimNextCompilerResultsTabSummoner.h"
+#include "Logging/MessageLog.h"
 #include "Widgets/Docking/SDockTab.h"
 #include "Widgets/Input/SComboButton.h"
 #include "Widgets/Input/SSearchBox.h"
@@ -264,6 +265,9 @@ void STraitEditorView::RefreshTraitStackTraitsStatus()
 
 			TArray<TSharedPtr<FTraitDataEditorDef>>& CurrentTraitsData = *(TraitEditorSharedData->CurrentTraitsDataShared.Get());
 
+			FMessageLog Log("AnimNextCompilerResults");
+			Log.NewPage(LOCTEXT("TraitStackCompileResults", "Trait Stack Compilation"));
+
 			const int32 NumTraits = CurrentTraitsData.Num();
 			for (int32 i = NumTraits - 1; i >= 0; i--)
 			{
@@ -283,16 +287,17 @@ void STraitEditorView::RefreshTraitStackTraitsStatus()
 				{
 					TraitEditorSharedData->bStackContainsErrors = true;
 
-					if (TSharedPtr<IMessageLogListing> CompilerResultsListing = GetMessageLogListing().Pin())
+					for (const FTraitStackTraitStatus::FStatusMessage& StatusMessage : TraitData->StackStatus.StatusMessages)
 					{
-						for (const FTraitStackTraitStatus::FStatusMessage& StatusMessage : TraitData->StackStatus.StatusMessages)
-						{
-							const EMessageSeverity::Type Severity = (StatusMessage.Status == FTraitStackTraitStatus::EStackStatus::Warning) ? EMessageSeverity::Warning : EMessageSeverity::Error;
-							const FText& MessageText = StatusMessage.MessageText;
+						const EMessageSeverity::Type Severity = (StatusMessage.Status == FTraitStackTraitStatus::EStackStatus::Warning) ? EMessageSeverity::Warning : EMessageSeverity::Error;
+						const FText& MessageText = StatusMessage.MessageText;
+						Log.AddMessage(FTokenizedMessage::Create(Severity, FText::Format(LOCTEXT("TraitEditorLogTraitNameErrorFormat", "{0}: {1}"), TraitData->TraitDisplayName, MessageText)));
+					}
 
-							TSharedRef<FTokenizedMessage> Message = FTokenizedMessage::Create(Severity, FText::Format(LOCTEXT("TraitEditorLogTraitNameErrorFormat", "{0} : {1}"), TraitData->TraitDisplayName, MessageText));
-							CompilerResultsListing->AddMessage(Message, false);
-						}
+					// Open tab to display errors
+					if (TSharedPtr<UE::Workspace::IWorkspaceEditor> WorkspaceEditor = WorkspaceEditorWeak.Pin())
+					{
+						WorkspaceEditor->GetTabManager()->TryInvokeTab(CompilerResultsTabName);
 					}
 				}
 			}
@@ -474,31 +479,6 @@ TSharedRef<SWidget> STraitEditorView::GetOptionsMenuWidget()
 		];
 }
 
-TWeakPtr<IMessageLogListing>& STraitEditorView::GetMessageLogListing()
-{
-	if (!CompilerResultsListingWeak.IsValid())
-	{
-		if (TSharedPtr<UE::Workspace::IWorkspaceEditor> WorkspaceEditor = WorkspaceEditorWeak.Pin())
-		{
-			TSharedPtr<SDockTab> DockTab = WorkspaceEditor->GetTabManager()->FindExistingLiveTab(UE::AnimNext::Editor::CompilerResultsTabName);
-			if (!DockTab.IsValid())
-			{
-				DockTab = WorkspaceEditor->GetTabManager()->TryInvokeTab(UE::AnimNext::Editor::CompilerResultsTabName);
-			}
-
-			if (DockTab.IsValid())
-			{
-				if (TSharedPtr<SAnimNextCompilerResultsWidget> AnimNextCompilerResultsWidget = StaticCastSharedPtr<SAnimNextCompilerResultsWidget>(DockTab->GetContent().ToSharedPtr()))
-				{
-					CompilerResultsListingWeak = AnimNextCompilerResultsWidget->GetCompilerResultsListing();
-				}
-			}
-		}
-	}
-
-	return CompilerResultsListingWeak;
-}
-
 int32 STraitEditorView::GetTraitPinIndex(UAnimNextGraph_EdGraphNode* InEdGraphNode, const TSharedPtr<FTraitDataEditorDef>& InTraitData, int32 InTraitIndex)
 {
 	if (InEdGraphNode == nullptr || !InTraitData.IsValid())
diff --git a/Engine/Plugins/Experimental/AnimNext/Source/AnimNextEditor/Private/Graph/STraitEditorView.h b/Engine/Plugins/Experimental/AnimNext/Source/AnimNextEditor/Private/Graph/STraitEditorView.h
--- a/Engine/Plugins/Experimental/AnimNext/Source/AnimNextEditor/Private/Graph/STraitEditorView.h
+++ b/Engine/Plugins/Experimental/AnimNext/Source/AnimNextEditor/Private/Graph/STraitEditorView.h
@@ -75,8 +75,6 @@ private:
 
 	TSharedRef<SWidget> GetOptionsMenuWidget();
 
-	TWeakPtr<IMessageLogListing>& GetMessageLogListing();
-
 	int32 GetTraitPinIndex(UAnimNextGraph_EdGraphNode* InEdGraphNode, const TSharedPtr<FTraitDataEditorDef>& InTraitData, int32 TraitIndex = INDEX_NONE);
 
 	static void GenerateTraitStackData(const TWeakObjectPtr<UAnimNextGraph_EdGraphNode>& EdGraphNodeWeak, TSharedPtr<FTraitEditorSharedData>& InTraitEditorSharedData);
diff --git a/Engine/Plugins/Experimental/AnimNext/Source/AnimNextUncookedOnly/AnimNextUncookedOnly.Build.cs b/Engine/Plugins/Experimental/AnimNext/Source/AnimNextUncookedOnly/AnimNextUncookedOnly.Build.cs
--- a/Engine/Plugins/Experimental/AnimNext/Source/AnimNextUncookedOnly/AnimNextUncookedOnly.Build.cs
+++ b/Engine/Plugins/Experimental/AnimNext/Source/AnimNextUncookedOnly/AnimNextUncookedOnly.Build.cs
@@ -34,7 +34,8 @@ namespace UnrealBuildTool.Rules
 					"ToolMenus",
 					"UniversalObjectLocator", 
 					"UniversalObjectLocatorEditor",
-					"WorkspaceEditor"
+					"WorkspaceEditor",
+					"MessageLog"
 				}
 			);
 
diff --git a/Engine/Plugins/Experimental/AnimNext/Source/AnimNextUncookedOnly/Internal/AnimNextRigVMAssetEditorData.h b/Engine/Plugins/Experimental/AnimNext/Source/AnimNextUncookedOnly/Internal/AnimNextRigVMAssetEditorData.h
--- a/Engine/Plugins/Experimental/AnimNext/Source/AnimNextUncookedOnly/Internal/AnimNextRigVMAssetEditorData.h
+++ b/Engine/Plugins/Experimental/AnimNext/Source/AnimNextUncookedOnly/Internal/AnimNextRigVMAssetEditorData.h
@@ -265,8 +265,6 @@ protected:
 	UPROPERTY(transient, DuplicateTransient)
 	bool bIsCompiling = false;
 
-	FCompilerResultsLog CompileLog;
-
 	FOnRigVMCompiledEvent RigVMCompiledEvent;
 
 	FRigVMGraphModifiedEvent RigVMGraphModifiedEvent;
diff --git a/Engine/Plugins/Experimental/AnimNext/Source/AnimNextUncookedOnly/Private/AnimNextUncookedOnlyModule.cpp b/Engine/Plugins/Experimental/AnimNext/Source/AnimNextUncookedOnly/Private/AnimNextUncookedOnlyModule.cpp
--- a/Engine/Plugins/Experimental/AnimNext/Source/AnimNextUncookedOnly/Private/AnimNextUncookedOnlyModule.cpp
+++ b/Engine/Plugins/Experimental/AnimNext/Source/AnimNextUncookedOnly/Private/AnimNextUncookedOnlyModule.cpp
@@ -11,6 +11,9 @@
 #include "Param/ObjectProxyType.h"
 #include "Scheduler/AnimNextSchedule.h"
 #include "UObject/AssetRegistryTagsContext.h"
+#include "MessageLogModule.h"
+
+#define LOCTEXT_NAMESPACE "AnimNextUncookedOnlyModule"
 
 namespace UE::AnimNext::UncookedOnly
 {
@@ -69,10 +72,22 @@ void FModule::StartupModule()
 	});
 
 	RegisterParameterSourceType(FAnimNextParamUniversalObjectLocator::StaticStruct(), MakeShared<FObjectProxyType>());
+
+	// Register the compilation log (hidden from the main log set, it is displayed in the workspace editor)
+	FMessageLogModule& MessageLogModule = FModuleManager::LoadModuleChecked<FMessageLogModule>("MessageLog");
+	FMessageLogInitializationOptions LogInitOptions;
+	LogInitOptions.bShowInLogWindow = false;
+	LogInitOptions.MaxPageCount = 10;
+	MessageLogModule.RegisterLogListing("AnimNextCompilerResults", LOCTEXT("CompilerResults", "AnimNext Compiler Results"), LogInitOptions);
 }
 
 void FModule::ShutdownModule()
 {
+	if(FMessageLogModule* MessageLogModule = FModuleManager::GetModulePtr<FMessageLogModule>("MessageLog"))
+	{
+		MessageLogModule->UnregisterLogListing("AnimNextCompilerResults");
+	}
+
 	UnregisterParameterSourceType(FAnimNextParamUniversalObjectLocator::StaticStruct());
 
 	UAnimNextSchedule::GetAssetRegistryTagsFunction = nullptr;
@@ -98,4 +113,6 @@ TSharedPtr<IParameterSourceType> FModule::FindParameterSourceType(const UScriptS
 
 }
 
+#undef LOCTEXT_NAMESPACE 
+
 IMPLEMENT_MODULE(UE::AnimNext::UncookedOnly::FModule, AnimNextUncookedOnly);
diff --git a/Engine/Plugins/Experimental/AnimNext/Source/AnimNextUncookedOnly/Private/UncookedOnlyUtils.cpp b/Engine/Plugins/Experimental/AnimNext/Source/AnimNextUncookedOnly/Private/UncookedOnlyUtils.cpp
--- a/Engine/Plugins/Experimental/AnimNext/Source/AnimNextUncookedOnly/Private/UncookedOnlyUtils.cpp
+++ b/Engine/Plugins/Experimental/AnimNext/Source/AnimNextUncookedOnly/Private/UncookedOnlyUtils.cpp
@@ -504,6 +504,8 @@ void FUtils::Compile(UAnimNextGraph* InGraph)
 {
 	check(InGraph);
 
+	FMessageLog("AnimNextCompilerResults").NewPage(FText::FromName(InGraph->GetFName()));
+
 	CompileStruct(InGraph);
 	CompileVM(InGraph);
 }
@@ -545,9 +547,6 @@ void FUtils::CompileVM(UAnimNextGraph* InGraph)
 	InGraph->RequiredParametersHash = 0;
 	InGraph->RequiredParameters.Empty();
 
-	EditorData->CompileLog.Messages.Reset();
-	EditorData->CompileLog.NumErrors = EditorData->CompileLog.NumWarnings = 0;
-
 	FRigVMClient* VMClient = EditorData->GetRigVMClient();
 	URigVMGraph* VMRootGraph = VMClient->GetDefaultModel();
 
@@ -661,7 +660,12 @@ void FUtils::CompileVM(UAnimNextGraph* InGraph)
 
 	URigVMCompiler* Compiler = URigVMCompiler::StaticClass()->GetDefaultObject<URigVMCompiler>();
 	EditorData->VMCompileSettings.SetExecuteContextStruct(FAnimNextExecuteContext::StaticStruct());
-	const FRigVMCompileSettings Settings = (EditorData->bCompileInDebugMode) ? FRigVMCompileSettings::Fast(EditorData->VMCompileSettings.GetExecuteContextStruct()) : EditorData->VMCompileSettings;
+	FRigVMCompileSettings Settings = (EditorData->bCompileInDebugMode) ? FRigVMCompileSettings::Fast(EditorData->VMCompileSettings.GetExecuteContextStruct()) : EditorData->VMCompileSettings;
+	Settings.ASTSettings.ReportDelegate.BindLambda([InGraph](EMessageSeverity::Type InType, UObject* InObject, const FString& InString)
+	{
+		FMessageLog("AnimNextCompilerResults").Message(InType, FText::FromString(InString));
+	});
+
 	Compiler->Compile(Settings, TempGraphs, TempController, InGraph->VM, InGraph->ExtendedExecuteContext, TArray<FRigVMExternalVariable>(), & EditorData->PinToOperandMap);
 
 	// Initialize right away, in packaged builds we initialize during PostLoad
@@ -1695,6 +1699,9 @@ void FUtils::CompileSchedule(UAnimNextSchedule* InSchedule)
 {
 	using namespace UE::AnimNext;
 
+	FMessageLog Log("AnimNextCompilerResults");
+	Log.NewPage(FText::FromName(InSchedule->GetFName()));
+
 	InSchedule->Instructions.Empty();
 	InSchedule->GraphTasks.Empty();
 	InSchedule->Ports.Empty();
@@ -1757,19 +1764,19 @@ void FUtils::CompileSchedule(UAnimNextSchedule* InSchedule)
 	TFunction<void(TArrayView<TObjectPtr<UAnimNextScheduleEntry>>)> EmitEntries;
 
 	// Iterate over all entries, recursing into scopes
-	EmitEntries = [InSchedule, &EmitEntries, &Emit, &EmitPrerequisite, &ParentScopeIndex, &IntermediateTerms, &IntermediateMap](TArrayView<TObjectPtr<UAnimNextScheduleEntry>> InEntries)
+	EmitEntries = [InSchedule, &EmitEntries, &Emit, &EmitPrerequisite, &ParentScopeIndex, &IntermediateTerms, &IntermediateMap, &Log](TArrayView<TObjectPtr<UAnimNextScheduleEntry>> InEntries)
 	{
 		for (int32 EntryIndex = 0; EntryIndex < InEntries.Num(); ++EntryIndex)
 		{
 			UAnimNextScheduleEntry* Entry = InEntries[EntryIndex];
 
-			auto CheckTermDirectionCompatibility = [](FName InName, EScheduleTermDirection InExistingDirection, EScheduleTermDirection InNewDirection)
+			auto CheckTermDirectionCompatibility = [&Log](FName InName, EScheduleTermDirection InExistingDirection, EScheduleTermDirection InNewDirection)
 			{
 				switch(InExistingDirection)
 				{
 				case EScheduleTermDirection::Input:
 					// Input before output: error
-					UE_LOG(LogAnimation, Error, TEXT("Term '%s' was used as an input before it was output"), *InName.ToString());
+					Log.Error(FText::Format(LOCTEXT("TermInputError", "Term '{0}' was used as an input before it was output"), FText::FromName(InName)));
 					return false;
 				case EScheduleTermDirection::Output:
 					return true;
@@ -1784,7 +1791,7 @@ void FUtils::CompileSchedule(UAnimNextSchedule* InSchedule)
 
 				if(PortEntry->Port == nullptr)
 				{
-					UE_LOG(LogAnimation, Error, TEXT("AnimNext: Invalid port class found"));
+					Log.Error(LOCTEXT("InvalidPortError", "Invalid port class found"));
 					bValid = false;
 				}
 				else
@@ -1795,7 +1802,7 @@ void FUtils::CompileSchedule(UAnimNextSchedule* InSchedule)
 					TConstArrayView<FScheduleTerm> Terms = CDO->GetTerms();
 					if(PortEntry->Terms.Num() != Terms.Num())
 					{
-						UE_LOG(LogAnimation, Error, TEXT("AnimNext: Incorrect term count for port: %d"), PortEntry->Terms.Num());
+						Log.Error(FText::Format(LOCTEXT("PortIncorrectTermCountError", "Incorrect term count for port: {0}"), FText::AsNumber(PortEntry->Terms.Num())));
 						bValid = false;
 					}
 
@@ -1804,7 +1811,7 @@ void FUtils::CompileSchedule(UAnimNextSchedule* InSchedule)
 						FName TermName = PortEntry->Terms[TermIndex].Name;
 						if(!PortEntry->Terms[TermIndex].Type.IsValid())
 						{
-							UE_LOG(LogAnimation, Error, TEXT("AnimNext: Invalid type when processing port term, ignored: '%s'"), *TermName.ToString());
+							Log.Error(FText::Format(LOCTEXT("PortIncorrectTermTypeError", "Invalid type when processing port term, ignored: '{0}'"), FText::FromName(TermName)));
 							bValid = false;
 						}
 						else
@@ -1815,7 +1822,7 @@ void FUtils::CompileSchedule(UAnimNextSchedule* InSchedule)
 								const FAnimNextScheduleEntryTerm& IntermediateTerm = IntermediateTerms[*ExistingIntermediateIndexPtr];
 								if(IntermediateTerm.Type != Terms[TermIndex].GetType())
 								{
-									UE_LOG(LogAnimation, Error, TEXT("AnimNext: Mismatched types when processing port term, ignored: '%s'"), *TermName.ToString());
+									Log.Error(FText::Format(LOCTEXT("PortMismatchedTermTypeError", "Mismatched types when processing port term, ignored: '{0}'"), FText::FromName(TermName)));
 									bValid = false;
 								}
 
@@ -1866,7 +1873,7 @@ void FUtils::CompileSchedule(UAnimNextSchedule* InSchedule)
 
 				if(GraphEntry->Graph == nullptr && !GraphEntry->DynamicGraph.IsValid())
 				{
-					UE_LOG(LogAnimation, Error, TEXT("AnimNext: Invalid graph or invalid parameter supplied"));
+					Log.Error(LOCTEXT("InvalidGraphOrParameterError", "Invalid graph or invalid parameter supplied in graph task"));
 					bValid = false;
 				}
 				else if(GraphEntry->Graph != nullptr)
@@ -1874,7 +1881,7 @@ void FUtils::CompileSchedule(UAnimNextSchedule* InSchedule)
 					TConstArrayView<FScheduleTerm> Terms = GraphEntry->Graph->GetTerms();
 					if(GraphEntry->Terms.Num() != Terms.Num())
 					{
-						UE_LOG(LogAnimation, Error, TEXT("AnimNext: Incorrect term count for graph: %d"), GraphEntry->Terms.Num());
+						Log.Error(FText::Format(LOCTEXT("GraphIncorrectTermCountError", "Incorrect term count for graph: {0}"), FText::AsNumber(GraphEntry->Terms.Num())));
 						bValid = false;
 					}
 					else
@@ -1885,13 +1892,13 @@ void FUtils::CompileSchedule(UAnimNextSchedule* InSchedule)
 							FName TermName = GraphEntry->Terms[TermIndex].Name;
 							if(Terms[TermIndex].Direction != GraphEntry->Terms[TermIndex].Direction)
 							{
-								UE_LOG(LogAnimation, Error, TEXT("AnimNext: Mismatched direction when processing graph term, ignored: '%s'"), *TermName.ToString());
+								Log.Error(FText::Format(LOCTEXT("MismatchedDirectionInGraphTermError", "Mismatched direction when processing graph term, ignored: '{0}'"), FText::FromName(TermName)));
 								bValid = false;
 							}
 							
 							if(Terms[TermIndex].GetType() != GraphEntry->Terms[TermIndex].Type)
 							{
-								UE_LOG(LogAnimation, Error, TEXT("AnimNext: Mismatched types when processing graph term, ignored: '%s'"), *TermName.ToString());
+								Log.Error(FText::Format(LOCTEXT("MismatchedTypesInGraphTermError", "Mismatched types when processing graph term, ignored: '{0}'"), FText::FromName(TermName)));
 								bValid = false;
 							}
 						}
@@ -1901,7 +1908,7 @@ void FUtils::CompileSchedule(UAnimNextSchedule* InSchedule)
 				// We must have a reference pose
 				if (!GraphEntry->ReferencePose.IsValid())
 				{
-					UE_LOG(LogAnimation, Error, TEXT("AnimNext: Invalid reference pose supplied to graph task"));
+					Log.Error(LOCTEXT("InvalidRefPoseError", "Invalid reference pose supplied to graph task"));
 					bValid = false;
 				}
 
@@ -1911,7 +1918,7 @@ void FUtils::CompileSchedule(UAnimNextSchedule* InSchedule)
 					FName TermName = GraphEntry->Terms[TermIndex].Name;
 					if(!GraphEntry->Terms[TermIndex].Type.IsValid())
 					{
-						UE_LOG(LogAnimation, Error, TEXT("AnimNext: Invalid type when processing graph term, ignored: '%s'"), *TermName.ToString());
+						Log.Error(FText::Format(LOCTEXT("InvalidTypeInGraphTermError", "Invalid type when processing graph term, ignored: '{0}'"), FText::FromName(TermName)));
 						bValid = false;
 					}
 					else
@@ -1922,7 +1929,7 @@ void FUtils::CompileSchedule(UAnimNextSchedule* InSchedule)
 							const FAnimNextScheduleEntryTerm& IntermediateTerm = IntermediateTerms[*ExistingIntermediateIndexPtr];
 							if(IntermediateTerm.Type != GraphEntry->Terms[TermIndex].Type)
 							{
-								UE_LOG(LogAnimation, Error, TEXT("AnimNext: Mismatched types when processing graph term, ignored: '%s'"), *TermName.ToString());
+								Log.Error(FText::Format(LOCTEXT("MismatchedTypeInGraphTermError", "Mismatched types when processing graph term, ignored: '{0}'"), FText::FromName(TermName)));
 								bValid = false;
 							}
 
@@ -2268,6 +2275,8 @@ void FUtils::CompileSchedule(UAnimNextSchedule* InSchedule)
 	}
 
 	FScheduler::OnScheduleCompiled(InSchedule);
+
+	InSchedule->CompiledEvent.Broadcast();
 }
 
 uint64 FUtils::SortAndHashParameters(TArray<FAnimNextParam>& InParameters)
