Fix for Horde LOCTEXT errors
#rnx

diff --git a/Engine/Plugins/Experimental/AnimNext/Source/AnimNextEditor/Private/Common/SRigVMAssetView.cpp b/Engine/Plugins/Experimental/AnimNext/Source/AnimNextEditor/Private/Common/SRigVMAssetView.cpp
--- a/Engine/Plugins/Experimental/AnimNext/Source/AnimNextEditor/Private/Common/SRigVMAssetView.cpp
+++ b/Engine/Plugins/Experimental/AnimNext/Source/AnimNextEditor/Private/Common/SRigVMAssetView.cpp
@@ -195,7 +195,7 @@ void SRigVMAssetView::Construct(const FArguments& InArgs, UAnimNextRigVMAssetEdi
 						SNew(SImage)
 						.ColorAndOpacity(FSlateColor::UseForeground())
 						.Image(FAppStyle::GetBrush("Level.VisibleIcon16x"))
-						.ToolTipText(LOCTEXT("AccessSpecifierTooltip", "Access level of this entry"))
+						.ToolTipText(LOCTEXT("AccessSpecifierAccessLevelTooltip", "Access level of this entry"))
 					]
 				]
 
@@ -636,7 +636,7 @@ class SRigVMAssetViewRow : public SMultiColumnTableRow<TSharedRef<FRigVMAssetVie
 								FText AccessSpecifierDesc = Export->GetExportAccessSpecifier() == EAnimNextExportAccessSpecifier::Public ?
 									LOCTEXT("PublicSpecifierDesc", "This means that the entry is usable from gameplay and from other AnimNext assets") :
 									LOCTEXT("PrivateSpecifierDesc", "This means that the entry is only usable inside this asset");
-								return FText::Format(LOCTEXT("AccessSpecifierTooltip", "This entry is {0}.\n{1}"), AccessSpecifier, AccessSpecifierDesc);
+								return FText::Format(LOCTEXT("AccessSpecifierEntryTooltip", "This entry is {0}.\n{1}"), AccessSpecifier, AccessSpecifierDesc);
 							}
 							return FText::GetEmpty();
 						})
diff --git a/Engine/Plugins/Experimental/AnimNext/Source/AnimNextEditor/Private/Graph/STraitEditorView.cpp b/Engine/Plugins/Experimental/AnimNext/Source/AnimNextEditor/Private/Graph/STraitEditorView.cpp
--- a/Engine/Plugins/Experimental/AnimNext/Source/AnimNextEditor/Private/Graph/STraitEditorView.cpp
+++ b/Engine/Plugins/Experimental/AnimNext/Source/AnimNextEditor/Private/Graph/STraitEditorView.cpp
@@ -432,8 +432,8 @@ TSharedRef<SWidget> STraitEditorView::GetOptionsMenuWidget()
 		EUserInterfaceActionType::Check
 	);
 	DetailViewOptions.AddMenuEntry(
-		LOCTEXT("TraitEditor_ShowTraitInterfaces", "Show Trait Interfaces If Warnings / Errors"),
-		LOCTEXT("TraitEditor_ShowTraitInterfaces_ToolTip", "Displays Trait Implemented and Required interfaces if is there any Warning or Error on the Stack"),
+		LOCTEXT("TraitEditor_ShowTraitInterfaces_Errors", "Show Trait Interfaces If Warnings / Errors"),
+		LOCTEXT("TraitEditor_ShowTraitInterfaces_Errors_ToolTip", "Displays Trait Implemented and Required interfaces if is there any Warning or Error on the Stack"),
 		FSlateIcon(),
 		FUIAction(
 			FExecuteAction::CreateLambda([this, TraitEditorSharedDataLocal]()
