Make sure to remove the trait when we unregister it

diff --git a/Engine/Plugins/Experimental/AnimNext/Source/AnimNext/Private/TraitCore/TraitRegistry.cpp b/Engine/Plugins/Experimental/AnimNext/Source/AnimNext/Private/TraitCore/TraitRegistry.cpp
--- a/Engine/Plugins/Experimental/AnimNext/Source/AnimNext/Private/TraitCore/TraitRegistry.cpp
+++ b/Engine/Plugins/Experimental/AnimNext/Source/AnimNext/Private/TraitCore/TraitRegistry.cpp
@@ -204,7 +204,6 @@ namespace UE::AnimNext
 				}
 
 				TraitUIDToEntryMap.Remove(It.Key());
-				
 
 				break;
 			}
@@ -338,6 +337,9 @@ namespace UE::AnimNext
 		{
 			check(Entry->TraitHandle.IsValid());
 
+			// Remove name from map before we destroy the Trait
+			TraitNameToUIDMap.Remove(*Entry->Trait->GetTraitName());
+
 			if (Entry->TraitHandle.IsDynamic())
 			{
 				// It has been dynamically registered, free it
