Allow coercing AnimNext weak pointers between types

diff --git a/Engine/Plugins/Experimental/AnimNext/Source/AnimNext/Private/TraitCore/TraitPtr.cpp b/Engine/Plugins/Experimental/AnimNext/Source/AnimNext/Private/TraitCore/TraitPtr.cpp
--- a/Engine/Plugins/Experimental/AnimNext/Source/AnimNext/Private/TraitCore/TraitPtr.cpp
+++ b/Engine/Plugins/Experimental/AnimNext/Source/AnimNext/Private/TraitCore/TraitPtr.cpp
@@ -35,6 +35,19 @@ namespace UE::AnimNext
 		}
 	}
 
+	FTraitPtr::FTraitPtr(const FWeakTraitPtr& TraitPtr)
+		: PackedPointerAndFlags(reinterpret_cast<uintptr_t>(TraitPtr.GetNodeInstance()))
+		, TraitIndex(TraitPtr.GetTraitIndex())
+	{
+		check((reinterpret_cast<uintptr_t>(TraitPtr.GetNodeInstance()) & FLAGS_MASK) == 0);	// Make sure we have enough alignment
+		check(TraitPtr.GetTraitIndex() <= MAX_uint8);	// Make sure we don't truncate
+
+		if (TraitPtr.IsValid())
+		{
+			PackedPointerAndFlags |= IS_WEAK_BIT;
+		}
+	}
+
 	FTraitPtr::FTraitPtr(FTraitPtr&& TraitPtr) noexcept
 		: PackedPointerAndFlags(TraitPtr.PackedPointerAndFlags)
 		, TraitIndex(TraitPtr.TraitIndex)
@@ -84,6 +97,21 @@ namespace UE::AnimNext
 		return *this;
 	}
 
+	FTraitPtr& FTraitPtr::operator=(const FWeakTraitPtr& TraitPtr)
+	{
+		Reset();
+
+		PackedPointerAndFlags = reinterpret_cast<uintptr_t>(TraitPtr.GetNodeInstance());
+		TraitIndex = TraitPtr.GetTraitIndex();
+
+		if (TraitPtr.IsValid())
+		{
+			PackedPointerAndFlags |= IS_WEAK_BIT;
+		}
+
+		return *this;
+	}
+
 	FTraitPtr& FTraitPtr::operator=(FTraitPtr&& TraitPtr) noexcept
 	{
 		Swap(PackedPointerAndFlags, TraitPtr.PackedPointerAndFlags);
diff --git a/Engine/Plugins/Experimental/AnimNext/Source/AnimNext/Public/TraitCore/TraitPtr.h b/Engine/Plugins/Experimental/AnimNext/Source/AnimNext/Public/TraitCore/TraitPtr.h
--- a/Engine/Plugins/Experimental/AnimNext/Source/AnimNext/Public/TraitCore/TraitPtr.h
+++ b/Engine/Plugins/Experimental/AnimNext/Source/AnimNext/Public/TraitCore/TraitPtr.h
@@ -31,9 +31,11 @@ namespace UE::AnimNext
 		FTraitPtr(FNodeInstance* InNodeInstance, uint32 InTraitIndex);
 
 		FTraitPtr(const FTraitPtr& TraitPtr);
+		explicit FTraitPtr(const FWeakTraitPtr& TraitPtr);
 		FTraitPtr(FTraitPtr&& TraitPtr) noexcept;
 		~FTraitPtr();
 
+		FTraitPtr& operator=(const FWeakTraitPtr& TraitPtr);
 		FTraitPtr& operator=(const FTraitPtr& TraitPtr);
 		FTraitPtr& operator=(FTraitPtr&& TraitPtr) noexcept;
 
