Expose AnimNext parameter property bag proxy publicly

diff --git a/Engine/Plugins/Experimental/AnimNext/Source/AnimNext/Internal/Param/PropertyBagProxy.h b/Engine/Plugins/Experimental/AnimNext/Source/AnimNext/Internal/Param/PropertyBagProxy.h
new file mode 100644
--- /dev/null
+++ b/Engine/Plugins/Experimental/AnimNext/Source/AnimNext/Internal/Param/PropertyBagProxy.h
@@ -0,0 +1,39 @@
+// Copyright Epic Games, Inc. All Rights Reserved.
+
+#pragma once
+
+#include "Param/IParameterSource.h"
+#include "PropertyBag.h"
+#include "IAnimNextModule.h"
+
+namespace UE::AnimNext
+{
+
+// Proxy struct used to hold data in a property bag
+struct ANIMNEXT_API FPropertyBagProxy : public IParameterSource
+{
+	FPropertyBagProxy(FName InInstanceId);
+
+	FPropertyBagProxy(FName InInstanceId, FInstancedPropertyBag&& InPropertyBag);
+
+	// IParameterSource interface
+	virtual FName GetInstanceId() const override;
+	virtual void Update(float DeltaTime) override {};
+	virtual const FParamStackLayerHandle& GetLayerHandle() const override { return LayerHandle; }
+	virtual void AddReferencedObjects(FReferenceCollector& Collector) override;
+
+	// Replaces all properties with the specified names and sets their values using the specified containers
+	void ReplaceAllParameters(TConstArrayView<FPropertyBagPropertyDesc> InDescs, TConstArrayView<TConstArrayView<uint8>> InValues);
+
+private:
+	// Property bag that this wraps
+	FInstancedPropertyBag PropertyBag;
+
+	// Layer handle - must be updated if PropertyBag is changed
+	FParamStackLayerHandle LayerHandle;
+
+	// Instance ID provided on construction
+	FName InstanceId;
+};
+
+}
\ No newline at end of file
diff --git a/Engine/Plugins/Experimental/AnimNext/Source/AnimNext/Private/Param/PropertyBagProxy.h b/Engine/Plugins/Experimental/AnimNext/Source/AnimNext/Private/Param/PropertyBagProxy.h
deleted file mode 100644
