Add pointer to owning graph instance to each component

diff --git a/Engine/Plugins/Experimental/AnimNext/Source/AnimNext/Private/Graph/GC_GraphInstanceComponent.cpp b/Engine/Plugins/Experimental/AnimNext/Source/AnimNext/Private/Graph/GC_GraphInstanceComponent.cpp
--- a/Engine/Plugins/Experimental/AnimNext/Source/AnimNext/Private/Graph/GC_GraphInstanceComponent.cpp
+++ b/Engine/Plugins/Experimental/AnimNext/Source/AnimNext/Private/Graph/GC_GraphInstanceComponent.cpp
@@ -7,6 +7,11 @@
 
 namespace UE::AnimNext
 {
+	FGCGraphInstanceComponent::FGCGraphInstanceComponent(FAnimNextGraphInstance& InOwnerInstance)
+		: FGraphInstanceComponent(InOwnerInstance)
+	{
+	}
+
 	void FGCGraphInstanceComponent::Register(const FWeakTraitPtr& InTraitPtr)
 	{
 		TraitsWithReferences.Add(InTraitPtr);
diff --git a/Engine/Plugins/Experimental/AnimNext/Source/AnimNext/Private/Graph/GC_GraphInstanceComponent.h b/Engine/Plugins/Experimental/AnimNext/Source/AnimNext/Private/Graph/GC_GraphInstanceComponent.h
--- a/Engine/Plugins/Experimental/AnimNext/Source/AnimNext/Private/Graph/GC_GraphInstanceComponent.h
+++ b/Engine/Plugins/Experimental/AnimNext/Source/AnimNext/Private/Graph/GC_GraphInstanceComponent.h
@@ -19,6 +19,8 @@ namespace UE::AnimNext
 	{
 		DECLARE_ANIM_GRAPH_INSTANCE_COMPONENT(FGCGraphInstanceComponent)
 
+		explicit FGCGraphInstanceComponent(FAnimNextGraphInstance& InOwnerInstance);
+
 		// Registers the provided trait with the GC system
 		// Once registered, IGarbageCollection::AddReferencedObjects will be called on it during GC
 		void Register(const FWeakTraitPtr& InTraitPtr);
diff --git a/Engine/Plugins/Experimental/AnimNext/Source/AnimNext/Public/Graph/AnimNextGraphInstance.h b/Engine/Plugins/Experimental/AnimNext/Source/AnimNext/Public/Graph/AnimNextGraphInstance.h
--- a/Engine/Plugins/Experimental/AnimNext/Source/AnimNext/Public/Graph/AnimNextGraphInstance.h
+++ b/Engine/Plugins/Experimental/AnimNext/Source/AnimNext/Public/Graph/AnimNextGraphInstance.h
@@ -202,7 +202,7 @@ ComponentType& FAnimNextGraphInstance::GetComponent()
 		return *static_cast<ComponentType*>(Component);
 	}
 
-	return static_cast<ComponentType&>(AddComponent(ComponentNameHash, ComponentName, MakeShared<ComponentType>()));
+	return static_cast<ComponentType&>(AddComponent(ComponentNameHash, ComponentName, MakeShared<ComponentType>(*this)));
 }
 
 template<class ComponentType>
diff --git a/Engine/Plugins/Experimental/AnimNext/Source/AnimNext/Public/Graph/GraphInstanceComponent.h b/Engine/Plugins/Experimental/AnimNext/Source/AnimNext/Public/Graph/GraphInstanceComponent.h
--- a/Engine/Plugins/Experimental/AnimNext/Source/AnimNext/Public/Graph/GraphInstanceComponent.h
+++ b/Engine/Plugins/Experimental/AnimNext/Source/AnimNext/Public/Graph/GraphInstanceComponent.h
@@ -5,6 +5,7 @@
 #include "CoreMinimal.h"
 
 struct FAnimNextTraitEvent;
+struct FAnimNextGraphInstance;
 
 // Declares a graph instance component and implements the necessary boilerplate
 #define DECLARE_ANIM_GRAPH_INSTANCE_COMPONENT(ComponentType) \
@@ -23,11 +24,18 @@ namespace UE::AnimNext
 	 */
 	struct FGraphInstanceComponent
 	{
+		explicit FGraphInstanceComponent(FAnimNextGraphInstance& InOwnerInstance) : OwnerInstance(InOwnerInstance) {}
 		virtual ~FGraphInstanceComponent() {}
 
 		static FName StaticComponentName() { return FName(TEXT("FGraphInstanceComponent")); }
 		virtual FName GetComponentName() const { return StaticComponentName(); }
 
+		// Returns the owning graph instance this component lives on
+		FAnimNextGraphInstance& GetGraphInstance() { return OwnerInstance; }
+
+		// Returns the owning graph instance this component lives on
+		const FAnimNextGraphInstance& GetGraphInstance() const { return OwnerInstance; }
+
 		// Called before the update traversal begins, before any node has been visited
 		// Note that PreUpdate won't be called if a component is created during the update traversal until the next update
 		// The execution context provided is bound to the graph root and can be bound to anything the component wishes
@@ -39,5 +47,9 @@ namespace UE::AnimNext
 
 		// Called before PreUpdate with input events and before PostUpdate with output events
 		virtual void OnTraitEvent(FExecutionContext& Context, FAnimNextTraitEvent& Event) {}
+
+	private:
+		// The owning graph instance this component lives on
+		FAnimNextGraphInstance& OwnerInstance;
 	};
 }
diff --git a/Engine/Plugins/Experimental/AnimNext/Source/AnimNext/Public/TraitCore/ExecutionContext.h b/Engine/Plugins/Experimental/AnimNext/Source/AnimNext/Public/TraitCore/ExecutionContext.h
--- a/Engine/Plugins/Experimental/AnimNext/Source/AnimNext/Public/TraitCore/ExecutionContext.h
+++ b/Engine/Plugins/Experimental/AnimNext/Source/AnimNext/Public/TraitCore/ExecutionContext.h
@@ -328,7 +328,7 @@ namespace UE::AnimNext
 			return *static_cast<ComponentType*>(Component);
 		}
 
-		return static_cast<ComponentType&>(AddComponent(ComponentNameHash, ComponentName, MakeShared<ComponentType>()));
+		return static_cast<ComponentType&>(AddComponent(ComponentNameHash, ComponentName, MakeShared<ComponentType>(*RootGraphInstance)));
 	}
 
 	template<class ComponentType>
diff --git a/Engine/Plugins/Experimental/AnimNextAnimGraph/Source/AnimNextAnimGraph/Private/Graph/SyncGroup_GraphInstanceComponent.cpp b/Engine/Plugins/Experimental/AnimNextAnimGraph/Source/AnimNextAnimGraph/Private/Graph/SyncGroup_GraphInstanceComponent.cpp
--- a/Engine/Plugins/Experimental/AnimNextAnimGraph/Source/AnimNextAnimGraph/Private/Graph/SyncGroup_GraphInstanceComponent.cpp
+++ b/Engine/Plugins/Experimental/AnimNextAnimGraph/Source/AnimNextAnimGraph/Private/Graph/SyncGroup_GraphInstanceComponent.cpp
@@ -5,6 +5,11 @@
 
 namespace UE::AnimNext
 {
+	FSyncGroupGraphInstanceComponent::FSyncGroupGraphInstanceComponent(FAnimNextGraphInstance& InOwnerInstance)
+		: FGraphInstanceComponent(InOwnerInstance)
+	{
+	}
+
 	void FSyncGroupGraphInstanceComponent::RegisterWithGroup(FName GroupName, EAnimGroupRole::Type GroupRole, const FWeakTraitPtr& TraitPtr, const FTraitUpdateState& TraitState)
 	{
 		FSyncGroupState& GroupState = SyncGroupMap.FindOrAdd(GroupName);
diff --git a/Engine/Plugins/Experimental/AnimNextAnimGraph/Source/AnimNextAnimGraph/Private/Graph/SyncGroup_GraphInstanceComponent.h b/Engine/Plugins/Experimental/AnimNextAnimGraph/Source/AnimNextAnimGraph/Private/Graph/SyncGroup_GraphInstanceComponent.h
--- a/Engine/Plugins/Experimental/AnimNextAnimGraph/Source/AnimNextAnimGraph/Private/Graph/SyncGroup_GraphInstanceComponent.h
+++ b/Engine/Plugins/Experimental/AnimNextAnimGraph/Source/AnimNextAnimGraph/Private/Graph/SyncGroup_GraphInstanceComponent.h
@@ -19,6 +19,8 @@ namespace UE::AnimNext
 	{
 		DECLARE_ANIM_GRAPH_INSTANCE_COMPONENT(FSyncGroupGraphInstanceComponent)
 
+		explicit FSyncGroupGraphInstanceComponent(FAnimNextGraphInstance& InOwnerInstance);
+
 		void RegisterWithGroup(FName GroupName, EAnimGroupRole::Type GroupRole, const FWeakTraitPtr& TraitPtr, const FTraitUpdateState& TraitState);
 
 		// FGraphInstanceComponent impl
