AnimNext:  Fixed an issue where the reference pose was not available, blocking schedule creation

diff --git a/Engine/Plugins/Experimental/AnimNext/Source/AnimNext/Private/AnimNextModule.cpp b/Engine/Plugins/Experimental/AnimNext/Source/AnimNext/Private/AnimNextModule.cpp
--- a/Engine/Plugins/Experimental/AnimNext/Source/AnimNext/Private/AnimNextModule.cpp
+++ b/Engine/Plugins/Experimental/AnimNext/Source/AnimNext/Private/AnimNextModule.cpp
@@ -66,6 +66,7 @@ namespace UE::AnimNext
 			FAnimNextEntryPoint::StaticStruct(),
 			FUniversalObjectLocator::StaticStruct(),
 			FAnimNextTickFunctionBinding::StaticStruct(),
+			FAnimNextGraphReferencePose::StaticStruct()
 		};
 
 		RigVMRegistry.RegisterStructTypes(AllowedStructTypes);
