Move AnimNext event ownership to schedule instance instead of graph instance

diff --git a/Engine/Plugins/Experimental/AnimNext/Source/AnimNext/Internal/AnimNextModule.h b/Engine/Plugins/Experimental/AnimNext/Source/AnimNext/Internal/AnimNextModule.h
--- a/Engine/Plugins/Experimental/AnimNext/Source/AnimNext/Internal/AnimNextModule.h
+++ b/Engine/Plugins/Experimental/AnimNext/Source/AnimNext/Internal/AnimNextModule.h
@@ -1,4 +1,4 @@
-﻿// Copyright Epic Games, Inc. All Rights Reserved.
+// Copyright Epic Games, Inc. All Rights Reserved.
 
 #pragma once
 
@@ -31,7 +31,7 @@ public:
 	// IAnimNextModule interface
 	virtual void RegisterAnimNextAnimGraph(const IAnimNextAnimGraph& InAnimGraphImpl) override;
 	virtual void UnregisterAnimNextAnimGraph() override;
-	virtual void UpdateGraph(FAnimNextGraphInstancePtr& GraphInstance, float DeltaTime) override;
+	virtual void UpdateGraph(FAnimNextGraphInstancePtr& GraphInstance, float DeltaTime, FTraitEventList& InputEventList, FTraitEventList& OutputEventList) override;
 	virtual void EvaluateGraph(FAnimNextGraphInstancePtr& GraphInstance, const FReferencePose& RefPose, int32 GraphLODLevel, FLODPoseHeap& OutputPose) const override;
 
 	// Factory method used to create a parameter source to access the specified instance ID, with a set of parameters that are initially required
diff --git a/Engine/Plugins/Experimental/AnimNext/Source/AnimNext/Internal/Component/AnimNextComponent.h b/Engine/Plugins/Experimental/AnimNext/Source/AnimNext/Internal/Component/AnimNextComponent.h
--- a/Engine/Plugins/Experimental/AnimNext/Source/AnimNext/Internal/Component/AnimNextComponent.h
+++ b/Engine/Plugins/Experimental/AnimNext/Source/AnimNext/Internal/Component/AnimNextComponent.h
@@ -6,6 +6,8 @@
 #include "Components/ActorComponent.h"
 #include "Scheduler/AnimNextSchedule.h"
 #include "Scheduler/ScheduleHandle.h"
+#include "TraitCore/TraitEvent.h"
+
 #include "AnimNextComponent.generated.h"
 
 class UAnimNextSchedule;
@@ -42,11 +44,15 @@ public:
 	// @param    Name     The name of the parameter to apply
 	// @param    Value    The value to set the parameter to
 	UFUNCTION(BlueprintCallable, Category = "AnimNext", CustomThunk, meta = (CustomStructureParam = Value, UnsafeDuringActorConstruction))
-	void SetParameterInScope(UPARAM(meta = (CustomWidget = "ParamName", AllowedParamType = "FAnimNextScope")) FName Scope, EAnimNextParameterScopeOrdering Ordering, UPARAM(meta = (CustomWidget = "ParamName")) FName Name, int32 Value);
+	ANIMNEXT_API void SetParameterInScope(UPARAM(meta = (CustomWidget = "ParamName", AllowedParamType = "FAnimNextScope")) FName Scope, EAnimNextParameterScopeOrdering Ordering, UPARAM(meta = (CustomWidget = "ParamName")) FName Name, int32 Value);
 
 	// Enable or disable this component's update
 	UFUNCTION(BlueprintCallable, Category = "AnimNext")
-	void Enable(bool bEnabled);
+	ANIMNEXT_API void Enable(bool bEnabled);
+
+	// Queues an input trait event
+	// Input events will be processed in the next graph update after they are queued
+	ANIMNEXT_API void QueueInputTraitEvent(FAnimNextTraitEventPtr Event);
 	
 private:
 	DECLARE_FUNCTION(execSetParameterInScope);
diff --git a/Engine/Plugins/Experimental/AnimNext/Source/AnimNext/Internal/Scheduler/ScheduleTaskContext.h b/Engine/Plugins/Experimental/AnimNext/Source/AnimNext/Internal/Scheduler/ScheduleTaskContext.h
--- a/Engine/Plugins/Experimental/AnimNext/Source/AnimNext/Internal/Scheduler/ScheduleTaskContext.h
+++ b/Engine/Plugins/Experimental/AnimNext/Source/AnimNext/Internal/Scheduler/ScheduleTaskContext.h
@@ -4,6 +4,8 @@
 
 #include "CoreMinimal.h"
 
+#include "TraitCore/TraitEvent.h"
+
 struct FInstancedPropertyBag;
 
 namespace UE::AnimNext
@@ -24,6 +26,10 @@ public:
 	// Apply the supplied parameter source to the specified scope, evicting any source that was there previously
 	ANIMNEXT_API void ApplyParametersToScope(FName InScope, EParameterScopeOrdering InOrdering, FName InInstanceId, FInstancedPropertyBag&& InPropertyBag) const;
 
+	// Queues an input trait event
+	// Input events will be processed in the next graph update after they are queued
+	ANIMNEXT_API void QueueInputTraitEvent(FAnimNextTraitEventPtr Event) const;
+
 private:
 	FScheduleTaskContext(const FScheduleContext& InContext);
 
diff --git a/Engine/Plugins/Experimental/AnimNext/Source/AnimNext/Private/AnimNextModule.cpp b/Engine/Plugins/Experimental/AnimNext/Source/AnimNext/Private/AnimNextModule.cpp
--- a/Engine/Plugins/Experimental/AnimNext/Source/AnimNext/Private/AnimNextModule.cpp
+++ b/Engine/Plugins/Experimental/AnimNext/Source/AnimNext/Private/AnimNextModule.cpp
@@ -157,11 +157,11 @@ namespace UE::AnimNext
 		AnimGraphImpl = nullptr;
 	}
 
-	void FModule::UpdateGraph(FAnimNextGraphInstancePtr& GraphInstance, float DeltaTime)
+	void FModule::UpdateGraph(FAnimNextGraphInstancePtr& GraphInstance, float DeltaTime, FTraitEventList& InputEventList, FTraitEventList& OutputEventList)
 	{
 		if (AnimGraphImpl != nullptr)
 		{
-			AnimGraphImpl->UpdateGraph(GraphInstance, DeltaTime);
+			AnimGraphImpl->UpdateGraph(GraphInstance, DeltaTime, InputEventList, OutputEventList);
 		}
 	}
 
diff --git a/Engine/Plugins/Experimental/AnimNext/Source/AnimNext/Private/Component/AnimNextComponent.cpp b/Engine/Plugins/Experimental/AnimNext/Source/AnimNext/Private/Component/AnimNextComponent.cpp
--- a/Engine/Plugins/Experimental/AnimNext/Source/AnimNext/Private/Component/AnimNextComponent.cpp
+++ b/Engine/Plugins/Experimental/AnimNext/Source/AnimNext/Private/Component/AnimNextComponent.cpp
@@ -142,3 +142,13 @@ void UAnimNextComponent::Enable(bool bEnabled)
 {
 	UE::AnimNext::FScheduler::EnableHandle(this, SchedulerHandle, bEnabled);
 }
+
+void UAnimNextComponent::QueueInputTraitEvent(FAnimNextTraitEventPtr Event)
+{
+	using namespace UE::AnimNext;
+
+	FScheduler::QueueTask(this, SchedulerHandle, NAME_None, [Event = MoveTemp(Event)](const FScheduleTaskContext& InContext)
+		{
+			InContext.QueueInputTraitEvent(Event);
+		});
+}
diff --git a/Engine/Plugins/Experimental/AnimNext/Source/AnimNext/Private/Graph/AnimNextGraphInstance.cpp b/Engine/Plugins/Experimental/AnimNext/Source/AnimNext/Private/Graph/AnimNextGraphInstance.cpp
--- a/Engine/Plugins/Experimental/AnimNext/Source/AnimNext/Private/Graph/AnimNextGraphInstance.cpp
+++ b/Engine/Plugins/Experimental/AnimNext/Source/AnimNext/Private/Graph/AnimNextGraphInstance.cpp
@@ -33,6 +33,7 @@ void FAnimNextGraphInstance::Release()
 	ExtendedExecuteContext.Reset();
 	Components.Empty();
 	Graph = nullptr;
+	GraphState = nullptr;
 }
 
 bool FAnimNextGraphInstance::IsValid() const
@@ -132,38 +133,6 @@ GraphInstanceComponentMapType::TConstIterator FAnimNextGraphInstance::GetCompone
 	return RootGraphInstance->Components.CreateConstIterator();
 }
 
-void FAnimNextGraphInstance::QueueInputTraitEvent(FAnimNextTraitEventPtr Event)
-{
-	if (!Event || !Event->IsValid())
-	{
-		return;
-	}
-
-	FWriteScopeLock WriteLock(RootGraphInstance->InputEventListLock);
-	RootGraphInstance->InputEventList.Push(Event);
-}
-
-void FAnimNextGraphInstance::QueueInputTraitEvents(const UE::AnimNext::FTraitEventList& Events)
-{
-	FWriteScopeLock WriteLock(RootGraphInstance->InputEventListLock);
-
-	for (const FAnimNextTraitEventPtr& Event : Events)
-	{
-		if (Event->IsValid())
-		{
-			RootGraphInstance->InputEventList.Push(Event);
-		}
-	}
-}
-
-void FAnimNextGraphInstance::CollectInputTraitEvents(UE::AnimNext::FTraitEventList& OutInputEvents)
-{
-	OutInputEvents.Reset();
-
-	FWriteScopeLock WriteLock(RootGraphInstance->InputEventListLock);
-	Swap(OutInputEvents, RootGraphInstance->InputEventList);
-}
-
 void FAnimNextGraphInstance::Update()
 {
 	bHasUpdatedOnce = true;
diff --git a/Engine/Plugins/Experimental/AnimNext/Source/AnimNext/Private/Graph/AnimNextGraphInstancePtr.cpp b/Engine/Plugins/Experimental/AnimNext/Source/AnimNext/Private/Graph/AnimNextGraphInstancePtr.cpp
--- a/Engine/Plugins/Experimental/AnimNext/Source/AnimNext/Private/Graph/AnimNextGraphInstancePtr.cpp
+++ b/Engine/Plugins/Experimental/AnimNext/Source/AnimNext/Private/Graph/AnimNextGraphInstancePtr.cpp
@@ -95,24 +95,6 @@ GraphInstanceComponentMapType::TConstIterator FAnimNextGraphInstancePtr::GetComp
 	return Impl->GetComponentIterator();
 }
 
-void FAnimNextGraphInstancePtr::QueueInputTraitEvent(FAnimNextTraitEventPtr Event)
-{
-	check(Impl);
-	Impl->QueueInputTraitEvent(Event);
-}
-
-void FAnimNextGraphInstancePtr::QueueInputTraitEvents(const UE::AnimNext::FTraitEventList& Events)
-{
-	check(Impl);
-	Impl->QueueInputTraitEvents(Events);
-}
-
-void FAnimNextGraphInstancePtr::CollectInputTraitEvents(UE::AnimNext::FTraitEventList& OutInputEvents)
-{
-	check(Impl);
-	Impl->CollectInputTraitEvents(OutInputEvents);
-}
-
 void FAnimNextGraphInstancePtr::Update()
 {
 	check(Impl);
diff --git a/Engine/Plugins/Experimental/AnimNext/Source/AnimNext/Private/Scheduler/AnimNextScheduleGraphTask.cpp b/Engine/Plugins/Experimental/AnimNext/Source/AnimNext/Private/Scheduler/AnimNextScheduleGraphTask.cpp
--- a/Engine/Plugins/Experimental/AnimNext/Source/AnimNext/Private/Scheduler/AnimNextScheduleGraphTask.cpp
+++ b/Engine/Plugins/Experimental/AnimNext/Source/AnimNext/Private/Scheduler/AnimNextScheduleGraphTask.cpp
@@ -160,6 +160,19 @@ void FAnimNextScheduleGraphTask::RunGraph(const UE::AnimNext::FScheduleContext&
 
 	check(OutputPose->LODPose.LODLevel == LODIndex);
 
+	// Every graph in a schedule will see the same input events (if they were queued before the schedule started)
+	FTraitEventList InputEventList;
+	FTraitEventList OutputEventList;
+
+	// A schedule can contain multiple graphs, we copy the input event list since it might be appended to during our update
+	{
+		FReadScopeLock ReadLock(InstanceData.EventListLock);
+		InputEventList = InstanceData.InputEventList;
+	}
+
+	// Track how many input events we started with, we'll append the new ones
+	const int32 NumOriginalInputEvents = InputEventList.Num();
+
 	// Internally we use memstack allocation, so we need a mark here
 	FMemStack& MemStack = FMemStack::Get();
 	FMemMark MemMark(MemStack);
@@ -168,6 +181,27 @@ void FAnimNextScheduleGraphTask::RunGraph(const UE::AnimNext::FScheduleContext&
 	// This reduces churn internally by avoiding a chunk to be repeatedly allocated and freed as we push/pop marks
 	MemStack.Alloc(size_t(FPageAllocator::SmallPageSize) + 1, 16);
 
-	IAnimNextModule::Get().UpdateGraph(GraphCache.GraphInstanceData, InContext.GetDeltaTime());
-	IAnimNextModule::Get().EvaluateGraph(GraphCache.GraphInstanceData, RefPose, LODIndex, OutputPose->LODPose);
+	IAnimNextModule& AnimNextModule = IAnimNextModule::Get();
+	AnimNextModule.UpdateGraph(GraphCache.GraphInstanceData, InContext.GetDeltaTime(), InputEventList, OutputEventList);
+	AnimNextModule.EvaluateGraph(GraphCache.GraphInstanceData, RefPose, LODIndex, OutputPose->LODPose);
+
+	// We might have appended new input/output events, append them
+	{
+		const int32 NumInputEvents = InputEventList.Num();
+
+		FWriteScopeLock WriteLock(InstanceData.EventListLock);
+
+		// Append the new input events
+		for (int32 EventIndex = NumOriginalInputEvents; EventIndex < NumInputEvents; ++EventIndex)
+		{
+			FAnimNextTraitEventPtr& Event = InputEventList[EventIndex];
+			if (Event->IsValid())
+			{
+				InstanceData.InputEventList.Push(MoveTemp(Event));
+			}
+		}
+
+		// Append our output events
+		InstanceData.OutputEventList.Append(OutputEventList);
+	}
 }
diff --git a/Engine/Plugins/Experimental/AnimNext/Source/AnimNext/Private/Scheduler/ScheduleInstanceData.cpp b/Engine/Plugins/Experimental/AnimNext/Source/AnimNext/Private/Scheduler/ScheduleInstanceData.cpp
--- a/Engine/Plugins/Experimental/AnimNext/Source/AnimNext/Private/Scheduler/ScheduleInstanceData.cpp
+++ b/Engine/Plugins/Experimental/AnimNext/Source/AnimNext/Private/Scheduler/ScheduleInstanceData.cpp
@@ -182,4 +182,9 @@ void FScheduleInstanceData::ApplyParametersToScope(FName InScope, EParameterScop
 	}
 }
 
+void FScheduleInstanceData::QueueInputTraitEvent(FAnimNextTraitEventPtr Event)
+{
+	InputEventList.Push(MoveTemp(Event));
+}
+
 }
\ No newline at end of file
diff --git a/Engine/Plugins/Experimental/AnimNext/Source/AnimNext/Private/Scheduler/ScheduleInstanceData.h b/Engine/Plugins/Experimental/AnimNext/Source/AnimNext/Private/Scheduler/ScheduleInstanceData.h
--- a/Engine/Plugins/Experimental/AnimNext/Source/AnimNext/Private/Scheduler/ScheduleInstanceData.h
+++ b/Engine/Plugins/Experimental/AnimNext/Source/AnimNext/Private/Scheduler/ScheduleInstanceData.h
@@ -7,6 +7,8 @@
 #include "Graph/AnimNextGraph.h"
 #include "Graph/AnimNextGraphInstancePtr.h"
 #include "Param/ParamStack.h"
+#include "TraitCore/TraitEvent.h"
+#include "TraitCore/TraitEventList.h"
 
 class UAnimNextSchedule;
 struct FAnimNextSchedulerEntry;
@@ -37,6 +39,10 @@ struct FScheduleInstanceData : public FGCObject
 	// Apply the supplied parameter source to the specified scope, evicting any source that was there previously
 	void ApplyParametersToScope(FName InScope, EParameterScopeOrdering InOrdering, TUniquePtr<IParameterSource>&& InParameters);
 
+	// Queues an input trait event
+	// Input events will be processed in the next graph update after they are queued
+	void QueueInputTraitEvent(FAnimNextTraitEventPtr Event);
+
 	// Get the appropriate params stack given the ID
 	TSharedPtr<FParamStack> GetParamStack(uint32 InIndex) const;
 	
@@ -112,6 +118,15 @@ struct FScheduleInstanceData : public FGCObject
 	};
 
 	TArray<FExternalParamCache> ExternalParamCaches;
+
+	// Input event list to be processed on the next update
+	FTraitEventList InputEventList;
+
+	// Output event list to be processed at the end of the schedule tick
+	FTraitEventList OutputEventList;
+
+	// Lock to ensure event list actions are thread safe
+	FRWLock EventListLock;
 };
 
 }
\ No newline at end of file
diff --git a/Engine/Plugins/Experimental/AnimNext/Source/AnimNext/Private/Scheduler/ScheduleTaskContext.cpp b/Engine/Plugins/Experimental/AnimNext/Source/AnimNext/Private/Scheduler/ScheduleTaskContext.cpp
--- a/Engine/Plugins/Experimental/AnimNext/Source/AnimNext/Private/Scheduler/ScheduleTaskContext.cpp
+++ b/Engine/Plugins/Experimental/AnimNext/Source/AnimNext/Private/Scheduler/ScheduleTaskContext.cpp
@@ -19,4 +19,10 @@ void FScheduleTaskContext::ApplyParametersToScope(FName InScope, EParameterScope
 	Context.GetInstanceData().ApplyParametersToScope(InScope, InOrdering, MoveTemp(PropertyBagProxy));
 }
 
+
+void FScheduleTaskContext::QueueInputTraitEvent(FAnimNextTraitEventPtr Event) const
+{
+	FScheduleInstanceData& InstanceData = Context.GetInstanceData();
+	InstanceData.QueueInputTraitEvent(MoveTemp(Event));
+}
 }
\ No newline at end of file
diff --git a/Engine/Plugins/Experimental/AnimNext/Source/AnimNext/Private/Scheduler/ScheduleTickFunction.cpp b/Engine/Plugins/Experimental/AnimNext/Source/AnimNext/Private/Scheduler/ScheduleTickFunction.cpp
--- a/Engine/Plugins/Experimental/AnimNext/Source/AnimNext/Private/Scheduler/ScheduleTickFunction.cpp
+++ b/Engine/Plugins/Experimental/AnimNext/Source/AnimNext/Private/Scheduler/ScheduleTickFunction.cpp
@@ -12,6 +12,7 @@
 #include "Scheduler/ScheduleInstanceData.h"
 #include "Scheduler/AnimNextSchedulerWorldSubsystem.h"
 #include "Scheduler/ScheduleTaskContext.h"
+#include "Scheduler/ScheduleEvents.h"
 
 namespace UE::AnimNext
 {
@@ -63,6 +64,39 @@ void FScheduleEndTickFunction::Run()
 		InstanceData.RootParamStack->PopLayer(PushedRootUserLayer);
 	}
 	InstanceData.PushedRootUserLayers.Reset();
+
+	// Decrement the remaining lifetime of the input events we processed and queue up any remaining events
+	UE::AnimNext::DecrementLifetimeAndPurgeExpired(InstanceData.InputEventList, InstanceData.OutputEventList);
+
+	// Filter out our schedule action events, we'll hand them off to the main thread to execute
+	FTraitEventList MainThreadActionEventList;
+	if (!InstanceData.OutputEventList.IsEmpty())
+	{
+		for (FAnimNextTraitEventPtr& Event : InstanceData.OutputEventList)
+		{
+			if (!Event->IsValid())
+			{
+				continue;
+			}
+
+			if (FAnimNextSchedule_ActionEvent* ActionEvent = Event->AsType<FAnimNextSchedule_ActionEvent>())
+			{
+				if (ActionEvent->IsThreadSafe())
+				{
+					// Execute this action now
+					ActionEvent->Execute();
+				}
+				else
+				{
+					// Defer this action and execute it on the main thread
+					MainThreadActionEventList.Push(Event);
+				}
+			}
+		}
+
+		// Reset our list of output events, we don't retain any
+		InstanceData.OutputEventList.Reset();
+	}
 	
 	auto RunTaskOnGameThread = [](TUniqueFunction<void(void)>&& InFunction)
 	{
@@ -101,6 +135,19 @@ void FScheduleEndTickFunction::Run()
 		});
 	}
 
+	if (!MainThreadActionEventList.IsEmpty())
+	{
+		RunTaskOnGameThread([MainThreadActionEventList = MoveTemp(MainThreadActionEventList)]()
+			{
+				check(IsInGameThread());
+				for (const FAnimNextTraitEventPtr& Event : MainThreadActionEventList)
+				{
+					FAnimNextSchedule_ActionEvent* ActionEvent = Event->AsType<FAnimNextSchedule_ActionEvent>();
+					ActionEvent->Execute();
+				}
+			});
+	}
+
 	Entry.ResolvedObject = nullptr;
 }
 
diff --git a/Engine/Plugins/Experimental/AnimNext/Source/AnimNext/Private/TraitCore/TraitEvent.cpp b/Engine/Plugins/Experimental/AnimNext/Source/AnimNext/Private/TraitCore/TraitEvent.cpp
--- a/Engine/Plugins/Experimental/AnimNext/Source/AnimNext/Private/TraitCore/TraitEvent.cpp
+++ b/Engine/Plugins/Experimental/AnimNext/Source/AnimNext/Private/TraitCore/TraitEvent.cpp
@@ -6,6 +6,17 @@
 #include "TraitCore/TraitStackBinding.h"
 #include "Graph/GraphInstanceComponent.h"
 
+bool FAnimNextTraitEvent::DecrementLifetime(UE::AnimNext::FTraitEventList& OutputEventList)
+{
+	const bool bExpired = Lifetime.Decrement();
+	if (bExpired)
+	{
+		OnExpired(OutputEventList);
+	}
+
+	return bExpired;
+}
+
 namespace UE::AnimNext
 {
 	void RaiseTraitEvent(FExecutionContext& Context, const FTraitStackBinding& StackBinding, FAnimNextTraitEvent& Event)
diff --git a/Engine/Plugins/Experimental/AnimNext/Source/AnimNext/Private/TraitCore/TraitEventList.cpp b/Engine/Plugins/Experimental/AnimNext/Source/AnimNext/Private/TraitCore/TraitEventList.cpp
new file mode 100644
--- /dev/null
+++ b/Engine/Plugins/Experimental/AnimNext/Source/AnimNext/Private/TraitCore/TraitEventList.cpp
@@ -0,0 +1,25 @@
+// Copyright Epic Games, Inc. All Rights Reserved.
+
+#include "TraitCore/TraitEventList.h"
+
+namespace UE::AnimNext
+{
+	void DecrementLifetimeAndPurgeExpired(FTraitEventList& EventList, FTraitEventList& OutputEventList)
+	{
+		// We remove expired events and modify the array in place to avoid heap churn
+		const int32 NumInputEvents = EventList.Num();
+		int32 InputEventWriteIndex = 0;
+
+		for (int32 InputEventReadIndex = 0; InputEventReadIndex < NumInputEvents; ++InputEventReadIndex)
+		{
+			FAnimNextTraitEventPtr& Event = EventList[InputEventReadIndex];
+			if (!Event->IsConsumed() && !Event->DecrementLifetime(OutputEventList))
+			{
+				// Event has not expired, swap it to keep it around
+				Swap(EventList[InputEventWriteIndex++], Event);
+			}
+		}
+
+		EventList.SetNum(InputEventWriteIndex, EAllowShrinking::Yes);
+	}
+}
diff --git a/Engine/Plugins/Experimental/AnimNext/Source/AnimNext/Public/Graph/AnimNextGraphInstance.h b/Engine/Plugins/Experimental/AnimNext/Source/AnimNext/Public/Graph/AnimNextGraphInstance.h
--- a/Engine/Plugins/Experimental/AnimNext/Source/AnimNext/Public/Graph/AnimNextGraphInstance.h
+++ b/Engine/Plugins/Experimental/AnimNext/Source/AnimNext/Public/Graph/AnimNextGraphInstance.h
@@ -98,18 +98,6 @@ struct ANIMNEXT_API FAnimNextGraphInstance
 	// Returns const iterators to the graph instance component container
 	GraphInstanceComponentMapType::TConstIterator GetComponentIterator() const;
 
-	// Queues an input trait event (thread safe)
-	// Input events will be processed in the next graph update after they are queued
-	void QueueInputTraitEvent(FAnimNextTraitEventPtr Event);
-
-	// Queues a list of input trait events (thread safe)
-	// Input events will be processed in the next graph update after they are queued
-	// Expired or invalid events will be skipped
-	void QueueInputTraitEvents(const UE::AnimNext::FTraitEventList& Events);
-
-	// Collects all currently queued input trait events (thread safe)
-	void CollectInputTraitEvents(UE::AnimNext::FTraitEventList& OutInputEvents);
-
 	// Called each time the graph updates
 	void Update();
 
@@ -163,12 +151,6 @@ private:
 	// Graph instance components that persist from update to update
 	GraphInstanceComponentMapType Components;
 
-	// Input event list to be processed on the next update
-	UE::AnimNext::FTraitEventList InputEventList;
-
-	// Lock to ensure input event list actions are thread safe
-	FRWLock InputEventListLock;
-
 	// Whether or not this graph has updated once
 	bool bHasUpdatedOnce = false;
 
diff --git a/Engine/Plugins/Experimental/AnimNext/Source/AnimNext/Public/Graph/AnimNextGraphInstancePtr.h b/Engine/Plugins/Experimental/AnimNext/Source/AnimNext/Public/Graph/AnimNextGraphInstancePtr.h
--- a/Engine/Plugins/Experimental/AnimNext/Source/AnimNext/Public/Graph/AnimNextGraphInstancePtr.h
+++ b/Engine/Plugins/Experimental/AnimNext/Source/AnimNext/Public/Graph/AnimNextGraphInstancePtr.h
@@ -85,18 +85,6 @@ struct ANIMNEXT_API FAnimNextGraphInstancePtr
 	// Returns const iterators to the graph instance component container
 	GraphInstanceComponentMapType::TConstIterator GetComponentIterator() const;
 
-	// Queues an input trait event (thread safe)
-	// Input events will be processed in the next graph update after they are queued
-	void QueueInputTraitEvent(FAnimNextTraitEventPtr Event);
-
-	// Queues a list of input trait events (thread safe)
-	// Input events will be processed in the next graph update after they are queued
-	// Expired or invalid events will be skipped
-	void QueueInputTraitEvents(const UE::AnimNext::FTraitEventList& Events);
-
-	// Collects all currently queued input trait events (thread safe)
-	void CollectInputTraitEvents(UE::AnimNext::FTraitEventList& OutInputEvents);
-
 	// Called each time the graph updates
 	void Update();
 
diff --git a/Engine/Plugins/Experimental/AnimNext/Source/AnimNext/Public/IAnimNextModule.h b/Engine/Plugins/Experimental/AnimNext/Source/AnimNext/Public/IAnimNextModule.h
--- a/Engine/Plugins/Experimental/AnimNext/Source/AnimNext/Public/IAnimNextModule.h
+++ b/Engine/Plugins/Experimental/AnimNext/Source/AnimNext/Public/IAnimNextModule.h
@@ -15,6 +15,7 @@ namespace UE::AnimNext
 	class IParameterSource;
 	class IParameterSourceFactory;
 	struct FParameterSourceContext;
+	struct FTraitEventList;
 }
 
 namespace UE::AnimNext
@@ -25,8 +26,8 @@ class IAnimNextAnimGraph
 public:
 	virtual ~IAnimNextAnimGraph() = default;
 
-	virtual void UpdateGraph(FAnimNextGraphInstancePtr& GraphInstance, float DeltaTime) const = 0;
-	virtual void EvaluateGraph(FAnimNextGraphInstancePtr& GraphInstance, const UE::AnimNext::FReferencePose& RefPose, int32 GraphLODLevel, FLODPoseHeap& OutputPose) const = 0;
+	virtual void UpdateGraph(FAnimNextGraphInstancePtr& GraphInstance, float DeltaTime, FTraitEventList& InputEventList, FTraitEventList& OutputEventList) const = 0;
+	virtual void EvaluateGraph(FAnimNextGraphInstancePtr& GraphInstance, const FReferencePose& RefPose, int32 GraphLODLevel, FLODPoseHeap& OutputPose) const = 0;
 };
 
 class IAnimNextModule : public IModuleInterface
@@ -42,7 +43,7 @@ public:
 	static ANIMNEXT_API IAnimNextModule& Get();
 	virtual void RegisterAnimNextAnimGraph(const IAnimNextAnimGraph& InAnimGraphImpl) = 0;
 	virtual void UnregisterAnimNextAnimGraph() = 0;
-	virtual void UpdateGraph(FAnimNextGraphInstancePtr& GraphInstance, float DeltaTime) = 0;
+	virtual void UpdateGraph(FAnimNextGraphInstancePtr& GraphInstance, float DeltaTime, FTraitEventList& InputEventList, FTraitEventList& OutputEventList) = 0;
 	virtual void EvaluateGraph(FAnimNextGraphInstancePtr& GraphInstance, const FReferencePose& RefPose, int32 GraphLODLevel, FLODPoseHeap& OutputPose) const = 0;
 };
 
diff --git a/Engine/Plugins/Experimental/AnimNext/Source/AnimNext/Public/Scheduler/ScheduleEvents.h b/Engine/Plugins/Experimental/AnimNext/Source/AnimNext/Public/Scheduler/ScheduleEvents.h
new file mode 100644
--- /dev/null
+++ b/Engine/Plugins/Experimental/AnimNext/Source/AnimNext/Public/Scheduler/ScheduleEvents.h
@@ -0,0 +1,27 @@
+﻿// Copyright Epic Games, Inc. All Rights Reserved.
+
+#pragma once
+
+#include "TraitCore/TraitEvent.h"
+
+/**
+ * Schedule Action Event
+ * 
+ * Schedule action events are trait events that are processed at the end of a schedule.
+ * If the event is not thread-safe, it is dispatched to execute on the main thread.
+ */
+struct FAnimNextSchedule_ActionEvent : public FAnimNextTraitEvent
+{
+	DECLARE_ANIM_TRAIT_EVENT(FAnimNextSchedule_ActionEvent, FAnimNextTraitEvent)
+
+	// Whether or not this event is thread-safe and can execute on any thread
+	// Events that are not thread-safe will execute on the main thread
+	virtual bool IsThreadSafe() const { return false; }
+
+	// Executes the schedule action
+	// Derived types can override this and implement whatever they wish instead of using a lambda
+	virtual void Execute() const { ActionFunction(); }
+
+	// The optional action to execute
+	TUniqueFunction<void(void)> ActionFunction;
+};
diff --git a/Engine/Plugins/Experimental/AnimNext/Source/AnimNext/Public/TraitCore/TraitEvent.h b/Engine/Plugins/Experimental/AnimNext/Source/AnimNext/Public/TraitCore/TraitEvent.h
--- a/Engine/Plugins/Experimental/AnimNext/Source/AnimNext/Public/TraitCore/TraitEvent.h
+++ b/Engine/Plugins/Experimental/AnimNext/Source/AnimNext/Public/TraitCore/TraitEvent.h
@@ -11,16 +11,19 @@
 
 namespace UE::AnimNext
 {
+	struct FTraitEventList;
+
 	// Create an alias to improve readability
 	using FTraitEventUID = UE::AnimNext::FConstExprUID;
 }
 
 // Helper macros
 // In a trait event struct declaration, this macro declares the necessary boilerplate we require
-#define DECLARE_ANIM_TRAIT_EVENT(EventName) \
+#define DECLARE_ANIM_TRAIT_EVENT(EventName, SuperEventName) \
 	/* FAnimNextTraitEvent impl */ \
 	static constexpr UE::AnimNext::FTraitEventUID TypeUID = UE::AnimNext::FTraitEventUID::MakeFromString(TEXT(#EventName)); \
-	virtual UE::AnimNext::FTraitEventUID GetTypeUID() const override { return TypeUID; }
+	virtual UE::AnimNext::FTraitEventUID GetTypeUID() const override { return TypeUID; } \
+	virtual bool IsA(UE::AnimNext::FTraitEventUID InTypeUID) const override { return InTypeUID == TypeUID ? true : SuperEventName::IsA(InTypeUID); } \
 
 
 /**
@@ -44,6 +47,13 @@ struct FAnimNextTraitEvent
 	// Returns the event type UID
 	virtual UE::AnimNext::FTraitEventUID GetTypeUID() const { return UE::AnimNext::FTraitEventUID(); }
 
+	// Returns whether or not this event derives from the specified type
+	template<class EventType>
+	bool IsA() const { return IsA(EventType::TypeUID); }
+
+	// Returns whether or not this event derives from the specified type
+	virtual bool IsA(UE::AnimNext::FTraitEventUID InTypeUID) const { return false; }
+
 	// Returns whether or not this event is valid
 	// An event is valid if it hasn't been consumed and if it isn't expired
 	bool IsValid() const { return !IsConsumed() && !IsExpired(); }
@@ -56,17 +66,19 @@ struct FAnimNextTraitEvent
 
 	// Marks this event as being handled.
 	// Subsequent handlers can use this to determine if they should take action or not.
-	void MarkHandled() { bIsHandled = true; }
+	// Returns whether or not this thread succeeded in marking this event as handled
+	bool MarkHandled() { return FPlatformAtomics::InterlockedExchange(&bIsHandled, 1) == 0; }
 
 	// Returns whether or not this event has been handled by a previous handler
-	bool IsHandled() const { return bIsHandled; }
+	bool IsHandled() const { return FPlatformAtomics::AtomicRead_Relaxed(&bIsHandled) != 0; }
 
 	// Marks this event as being consumed.
 	// This event won't be forward to any other handlers.
-	void MarkConsumed() { bIsConsumed = true; }
+	// Returns whether or not this thread succeeded in marking this event as consumed
+	bool MarkConsumed() { return FPlatformAtomics::InterlockedExchange(&bIsConsumed, 1) == 0; }
 
 	// Returns whether or not this event has been consumed by a previous handler
-	bool IsConsumed() const { return bIsConsumed; }
+	bool IsConsumed() const { return FPlatformAtomics::AtomicRead_Relaxed(&bIsConsumed) != 0; }
 
 	// Returns whether or not this event has expired
 	bool IsExpired() const { return Lifetime.IsExpired(); }
@@ -78,15 +90,16 @@ struct FAnimNextTraitEvent
 	bool IsTransient() const { return Lifetime.IsTransient(); }
 
 	// Decrements the lifetime count and returns whether or not this event has expired in the process
-	bool DecrementLifetime() { return Lifetime.Decrement(); }
+	// OnExpired will be called on the derived type if we expire and new output events can be appended
+	ANIMNEXT_API bool DecrementLifetime(UE::AnimNext::FTraitEventList& OutputEventList);
 
 	// Returns a pointer to this event cast to the specified type if the TypeUIDs match, nullptr otherwise
 	template<class EventType>
-	EventType* AsType() { return EventType::TypeUID == GetTypeUID() ? static_cast<EventType*>(this) : nullptr; }
+	EventType* AsType() { return IsA<EventType>() ? static_cast<EventType*>(this) : nullptr; }
 
 	// Returns a pointer to this event cast to the specified type if the TypeUIDs match, nullptr otherwise
 	template<class EventType>
-	const EventType* AsType() const { return EventType::TypeUID == GetTypeUID() ? static_cast<const EventType*>(this) : nullptr; }
+	const EventType* AsType() const { return IsA<EventType>() ? static_cast<const EventType*>(this) : nullptr; }
 
 protected:
 	// Creates an event of the specified type and associated lifetime
@@ -95,22 +108,25 @@ protected:
 	{
 	}
 
+	// Called when an events' lifetime expires
+	virtual void OnExpired(UE::AnimNext::FTraitEventList& OutputEventList) {}
+
 private:
 	// The specified lifetime of this event
 	UE::AnimNext::FTraitEventLifetime Lifetime = UE::AnimNext::FTraitEventLifetime::MakeTransient();
 
 	// Signals whether someone has opted to handle this event
-	bool bIsHandled = false;
+	int8 bIsHandled = 0;
 
 	// Signals whether someone consumed this event and it thus shouldn't be forwarded to other handlers
-	bool bIsConsumed = false;
+	int8 bIsConsumed = 0;
 };
 
 // Create a shared pointer alias for trait events
 using FAnimNextTraitEventPtr = TSharedPtr<FAnimNextTraitEvent, ESPMode::ThreadSafe>;
 
 template<class EventType, typename... TArgs>
-TSharedPtr<EventType, ESPMode::ThreadSafe> MakeTraitEvent(TArgs&&... Args)
+inline TSharedPtr<EventType, ESPMode::ThreadSafe> MakeTraitEvent(TArgs&&... Args)
 {
 	return MakeShared<EventType>(Forward<TArgs>(Args)...);
 }
diff --git a/Engine/Plugins/Experimental/AnimNext/Source/AnimNext/Public/TraitCore/TraitEventList.h b/Engine/Plugins/Experimental/AnimNext/Source/AnimNext/Public/TraitCore/TraitEventList.h
--- a/Engine/Plugins/Experimental/AnimNext/Source/AnimNext/Public/TraitCore/TraitEventList.h
+++ b/Engine/Plugins/Experimental/AnimNext/Source/AnimNext/Public/TraitCore/TraitEventList.h
@@ -23,6 +23,11 @@ namespace UE::AnimNext
 		void Push(FAnimNextTraitEventPtr Event) { Events.Add(MoveTemp(Event)); }
 		void Reset() { Events.Reset(); }
 		int32 Num() const { return Events.Num(); }
+		bool IsEmpty() const { return Events.IsEmpty(); }
+
+		void SetNum(int32 NewNum, EAllowShrinking AllowShrinking = EAllowShrinking::Yes) { Events.SetNum(NewNum, AllowShrinking); }
+
+		void Append(const FTraitEventList& Source) { Events.Append(Source.Events); }
 
 		FAnimNextTraitEventPtr& operator[](int32 EventIndex) { return Events[EventIndex]; }
 		const FAnimNextTraitEventPtr& operator[](int32 EventIndex) const { return Events[EventIndex]; }
@@ -32,17 +37,12 @@ namespace UE::AnimNext
 		IteratorType end() { return Events.end(); }
 		ConstIteratorType end() const { return Events.end(); }
 
-		// Decrements the lifetime of every event within the list
-		void DecrementLifetime()
-		{
-			for (FAnimNextTraitEventPtr& Event : Events)
-			{
-				Event->DecrementLifetime();
-			}
-		}
-
 	private:
 		// A list of events
 		TArray<FAnimNextTraitEventPtr> Events;
 	};
+
+	// Decrements and purges expired entries from the specified event list
+	// Expired events can generate new output events if they wish
+	ANIMNEXT_API void DecrementLifetimeAndPurgeExpired(FTraitEventList& EventList, FTraitEventList& OutputEventList);
 }
diff --git a/Engine/Plugins/Experimental/AnimNext/Source/AnimNextTestSuite/Private/AnimNextTraitEventTest.cpp b/Engine/Plugins/Experimental/AnimNext/Source/AnimNextTestSuite/Private/AnimNextTraitEventTest.cpp
--- a/Engine/Plugins/Experimental/AnimNext/Source/AnimNextTestSuite/Private/AnimNextTraitEventTest.cpp
+++ b/Engine/Plugins/Experimental/AnimNext/Source/AnimNextTestSuite/Private/AnimNextTraitEventTest.cpp
@@ -196,7 +196,8 @@ bool FAnimationAnimNextRuntimeTest_TraitEvent::RunTest(const FString& InParamete
 		AddErrorIfFalse(!Event.IsInfinite(), "FAnimationAnimNextRuntimeTest_TraitEvent -> Default constructed event should be transient");
 		AddErrorIfFalse(Event.IsTransient(), "FAnimationAnimNextRuntimeTest_TraitEvent -> Default constructed event should be transient");
 
-		bool bIsExpired = Event.DecrementLifetime();
+		FTraitEventList OutputEventList;
+		bool bIsExpired = Event.DecrementLifetime(OutputEventList);
 		AddErrorIfFalse(bIsExpired, "FAnimationAnimNextRuntimeTest_TraitEvent -> Event should be expired");
 		AddErrorIfFalse(!Event.IsValid(), "FAnimationAnimNextRuntimeTest_TraitEvent -> Event should be expired");
 		AddErrorIfFalse(!Event.IsHandled(), "FAnimationAnimNextRuntimeTest_TraitEvent -> Event should not be handled");
diff --git a/Engine/Plugins/Experimental/AnimNext/Source/AnimNextTestSuite/Private/AnimNextTraitEventTest.h b/Engine/Plugins/Experimental/AnimNext/Source/AnimNextTestSuite/Private/AnimNextTraitEventTest.h
--- a/Engine/Plugins/Experimental/AnimNext/Source/AnimNextTestSuite/Private/AnimNextTraitEventTest.h
+++ b/Engine/Plugins/Experimental/AnimNext/Source/AnimNextTestSuite/Private/AnimNextTraitEventTest.h
@@ -7,24 +7,18 @@
 #include "TraitCore/TraitUID.h"
 #include "TraitCore/TraitEvent.h"
 
-#include "AnimNextTraitEventTest.generated.h"
-
-USTRUCT()
 struct FTraitCoreTest_EventA : public FAnimNextTraitEvent
 {
-	GENERATED_BODY()
-	DECLARE_ANIM_TRAIT_EVENT(FTraitCoreTest_EventA)
+	DECLARE_ANIM_TRAIT_EVENT(FTraitCoreTest_EventA, FAnimNextTraitEvent)
 
 	bool bAlwaysForwardToBase = true;
 
 	TArray<UE::AnimNext::FTraitUID> VisitedTraits;
 };
 
-USTRUCT()
 struct FTraitCoreTest_EventB : public FAnimNextTraitEvent
 {
-	GENERATED_BODY()
-	DECLARE_ANIM_TRAIT_EVENT(FTraitCoreTest_EventB)
+	DECLARE_ANIM_TRAIT_EVENT(FTraitCoreTest_EventB, FAnimNextTraitEvent)
 
 	TArray<UE::AnimNext::FTraitUID> VisitedTraits;
 };
diff --git a/Engine/Plugins/Experimental/AnimNextAnimGraph/Source/AnimNextAnimGraph/Private/Graph/AnimGraph/AnimNode_AnimNextGraph.cpp b/Engine/Plugins/Experimental/AnimNextAnimGraph/Source/AnimNextAnimGraph/Private/Graph/AnimGraph/AnimNode_AnimNextGraph.cpp
--- a/Engine/Plugins/Experimental/AnimNextAnimGraph/Source/AnimNextAnimGraph/Private/Graph/AnimGraph/AnimNode_AnimNextGraph.cpp
+++ b/Engine/Plugins/Experimental/AnimNextAnimGraph/Source/AnimNextAnimGraph/Private/Graph/AnimGraph/AnimNode_AnimNextGraph.cpp
@@ -61,7 +61,9 @@ void FAnimNode_AnimNextGraph::Update_AnyThread(const FAnimationUpdateContext& Co
 
 		PropagateInputProperties(Context.AnimInstanceProxy->GetAnimInstanceObject());
 
-		UE::AnimNext::UpdateGraph(GraphInstance, Context.GetDeltaTime());
+		UE::AnimNext::FTraitEventList InputEventList;
+		UE::AnimNext::FTraitEventList OutputEventList;
+		UE::AnimNext::UpdateGraph(GraphInstance, Context.GetDeltaTime(), InputEventList, OutputEventList);
 	}
 
 	FAnimNode_CustomProperty::Update_AnyThread(Context);
diff --git a/Engine/Plugins/Experimental/AnimNextAnimGraph/Source/AnimNextAnimGraph/Private/Module.cpp b/Engine/Plugins/Experimental/AnimNextAnimGraph/Source/AnimNextAnimGraph/Private/Module.cpp
--- a/Engine/Plugins/Experimental/AnimNextAnimGraph/Source/AnimNextAnimGraph/Private/Module.cpp
+++ b/Engine/Plugins/Experimental/AnimNextAnimGraph/Source/AnimNextAnimGraph/Private/Module.cpp
@@ -41,9 +41,9 @@ public:
 		IAnimNextModule::Get().UnregisterAnimNextAnimGraph();
 	}
 
-	virtual void UpdateGraph(FAnimNextGraphInstancePtr& GraphInstance, float DeltaTime) const override
+	virtual void UpdateGraph(FAnimNextGraphInstancePtr& GraphInstance, float DeltaTime, UE::AnimNext::FTraitEventList& InputEventList, FTraitEventList& OutputEventList) const override
 	{
-		UE::AnimNext::UpdateGraph(GraphInstance, DeltaTime);
+		UE::AnimNext::UpdateGraph(GraphInstance, DeltaTime, InputEventList, OutputEventList);
 	}
 
 	virtual void EvaluateGraph(FAnimNextGraphInstancePtr& GraphInstance, const UE::AnimNext::FReferencePose& RefPose, int32 GraphLODLevel, FLODPoseHeap& OutputPose) const override
diff --git a/Engine/Plugins/Experimental/AnimNextAnimGraph/Source/AnimNextAnimGraph/Private/TraitInterfaces/IUpdate.cpp b/Engine/Plugins/Experimental/AnimNextAnimGraph/Source/AnimNextAnimGraph/Private/TraitInterfaces/IUpdate.cpp
--- a/Engine/Plugins/Experimental/AnimNextAnimGraph/Source/AnimNextAnimGraph/Private/TraitInterfaces/IUpdate.cpp
+++ b/Engine/Plugins/Experimental/AnimNextAnimGraph/Source/AnimNextAnimGraph/Private/TraitInterfaces/IUpdate.cpp
@@ -215,7 +215,7 @@ namespace UE::AnimNext
 			Private::QueueBookkeepingEntry(ExecutingEntry->EventBookkeepingList, BookkeepingEntry);
 		}
 
-		InputEventList.Push(Event);
+		InputEventList->Push(MoveTemp(Event));
 	}
 
 	void FUpdateTraversalContext::RaiseOutputTraitEvent(FAnimNextTraitEventPtr Event)
@@ -242,7 +242,7 @@ namespace UE::AnimNext
 		{
 			// We aren't executing a trait stack or we are the root stack, just queue the output
 			// We might be in a component pre/post-update
-			OutputEventList.Push(Event);
+			OutputEventList->Push(MoveTemp(Event));
 		}
 	}
 
@@ -260,7 +260,7 @@ namespace UE::AnimNext
 			switch (BookkeepingEntry->Action)
 			{
 			case Private::FUpdateEventBookkeepingAction::PushOutput:
-				OutputEventList.Push(BookkeepingEntry->Event);
+				OutputEventList->Push(MoveTemp(BookkeepingEntry->Event));
 				break;
 			case Private::FUpdateEventBookkeepingAction::Consume:
 				BookkeepingEntry->Event->MarkConsumed();
@@ -443,7 +443,7 @@ namespace UE::AnimNext
 	// perform as much useful work as possible while waiting for memory, hiding its slow latency by fully
 	// leveraging out-of-order CPU execution.
 
-	void UpdateGraph(FAnimNextGraphInstancePtr& GraphInstance, float DeltaTime)
+	void UpdateGraph(FAnimNextGraphInstancePtr& GraphInstance, float DeltaTime, FTraitEventList& InputEventList, FTraitEventList& OutputEventList)
 	{
 		SCOPE_CYCLE_COUNTER(STAT_AnimNext_UpdateGraph);
 		
@@ -458,6 +458,8 @@ namespace UE::AnimNext
 		}
 
 		FUpdateTraversalContext TraversalContext;
+		TraversalContext.InputEventList = &InputEventList;
+		TraversalContext.OutputEventList = &OutputEventList;
 
 		FMemStack& MemStack = TraversalContext.GetMemStack();
 		FMemMark Mark(MemStack);
@@ -479,14 +481,11 @@ namespace UE::AnimNext
 		// Update the graph instance itself
 		GraphInstance.Update();
 
-		// Grab our input events, we'll propagate them down the graph as we traverse
-		GraphInstance.CollectInputTraitEvents(TraversalContext.InputEventList);
-
 		// Before we start the traversal, we give the graph instance components the chance to do some work
 		TraversalContext.BindTo(GraphInstance);
 		for (auto It = TraversalContext.GetComponentIterator(); It; ++It)
 		{
-			RaiseTraitEvents(TraversalContext, *It.Value(), TraversalContext.InputEventList);
+			RaiseTraitEvents(TraversalContext, *It.Value(), *TraversalContext.InputEventList);
 
 			It.Value()->PreUpdate(TraversalContext);
 		}
@@ -528,7 +527,7 @@ namespace UE::AnimNext
 				}
 
 				// Raise our input events
-				Private::RaiseTraitEvents(TraversalContext, Entry, TraversalContext.InputEventList);
+				Private::RaiseTraitEvents(TraversalContext, Entry, *TraversalContext.InputEventList);
 
 				// Main update before our children
 				if (bImplementsIUpdate)
@@ -593,7 +592,7 @@ namespace UE::AnimNext
 				TraversalContext.ExecuteBookkeepingActions(Entry->EventBookkeepingList);
 
 				// Raise our output events
-				Private::RaiseTraitEvents(TraversalContext, Entry, TraversalContext.OutputEventList);
+				Private::RaiseTraitEvents(TraversalContext, Entry, *TraversalContext.OutputEventList);
 
 				// We've already visited this node once, time to PostUpdate
 				if (Entry->UpdateTrait.IsValid())
@@ -621,21 +620,13 @@ namespace UE::AnimNext
 		// After we finish the traversal, we give the graph instance components the chance to do some work
 		for (auto It = TraversalContext.GetComponentIterator(); It; ++It)
 		{
-			RaiseTraitEvents(TraversalContext, *It.Value(), TraversalContext.OutputEventList);
+			RaiseTraitEvents(TraversalContext, *It.Value(),* TraversalContext.OutputEventList);
 
 			It.Value()->PostUpdate(TraversalContext);
 		}
 
-		// Decrement the remaining lifetime of the input events we processed and queue up any remaining events
-		TraversalContext.InputEventList.DecrementLifetime();
-		GraphInstance.QueueInputTraitEvents(TraversalContext.InputEventList);
-
 		// At this point, we shouldn't have any remaining scoped interfaces
 		// If this fails, it means we failed to pop them due to a push/pop mismatch
 		ensure(!TraversalContext.HasScopedInterfaces());
-
-		// TODO: Figure out what to do with the output events, we probably want to return them here and let the caller
-		// determine whether they should be ignored, filtered, or handled (e.g. post-physics vs pre-physics)
-		// The caller can determine what events can fire because it knows where it is running (e.g. worker vs main thread)
 	}
 }
diff --git a/Engine/Plugins/Experimental/AnimNextAnimGraph/Source/AnimNextAnimGraph/Public/TraitInterfaces/IUpdate.h b/Engine/Plugins/Experimental/AnimNextAnimGraph/Source/AnimNextAnimGraph/Public/TraitInterfaces/IUpdate.h
--- a/Engine/Plugins/Experimental/AnimNextAnimGraph/Source/AnimNextAnimGraph/Public/TraitInterfaces/IUpdate.h
+++ b/Engine/Plugins/Experimental/AnimNextAnimGraph/Source/AnimNextAnimGraph/Public/TraitInterfaces/IUpdate.h
@@ -161,8 +161,8 @@ namespace UE::AnimNext
 		void ExecuteBookkeepingActions(Private::FUpdateEventBookkeepingList& BookkeepingList);
 
 		// The input and output event lists
-		UE::AnimNext::FTraitEventList InputEventList;
-		UE::AnimNext::FTraitEventList OutputEventList;
+		UE::AnimNext::FTraitEventList* InputEventList = nullptr;
+		UE::AnimNext::FTraitEventList* OutputEventList = nullptr;
 
 		// The currently executing entry
 		Private::FUpdateEntry* ExecutingEntry = nullptr;
@@ -185,7 +185,7 @@ namespace UE::AnimNext
 		// The root node doesn't have a parent but we need a bookkeeping list regardless
 		Private::FUpdateEventBookkeepingList* RootParentBookkeepingEntryList = nullptr;
 
-		friend ANIMNEXTANIMGRAPH_API void UpdateGraph(FAnimNextGraphInstancePtr& GraphInstance, float DeltaTime);
+		friend ANIMNEXTANIMGRAPH_API void UpdateGraph(FAnimNextGraphInstancePtr& GraphInstance, float DeltaTime, FTraitEventList& InputEventList, FTraitEventList& OutputEventList);
 		friend FUpdateTraversalQueue;
 	};
 
@@ -213,7 +213,7 @@ namespace UE::AnimNext
 		// emptied and pushed onto the update stack
 		Private::FUpdateEntry* QueuedUpdateStackHead = nullptr;
 
-		friend ANIMNEXTANIMGRAPH_API void UpdateGraph(FAnimNextGraphInstancePtr& GraphInstance, float DeltaTime);
+		friend ANIMNEXTANIMGRAPH_API void UpdateGraph(FAnimNextGraphInstancePtr& GraphInstance, float DeltaTime, FTraitEventList& InputEventList, FTraitEventList& OutputEventList);
 		friend FUpdateTraversalContext;
 	};
 
@@ -312,8 +312,10 @@ namespace UE::AnimNext
 	 *     - We call PreUpdate on all its traits
 	 *     - We update all children
 	 *     - We call PostUpdate on all its traits
+	 * 
+	 * During our update, we can append new input/output events.
 	 *
 	 * @see IUpdate::PreUpdate, IUpdate::PostUpdate, IHierarchy::GetChildren
 	 */
-	ANIMNEXTANIMGRAPH_API void UpdateGraph(FAnimNextGraphInstancePtr& GraphInstance, float DeltaTime);
+	ANIMNEXTANIMGRAPH_API void UpdateGraph(FAnimNextGraphInstancePtr& GraphInstance, float DeltaTime, FTraitEventList& InputEventList, FTraitEventList& OutputEventList);
 }
diff --git a/Engine/Plugins/Experimental/AnimNextAnimGraph/Source/AnimNextAnimGraph/Public/Traits/Inertialization.h b/Engine/Plugins/Experimental/AnimNextAnimGraph/Source/AnimNextAnimGraph/Public/Traits/Inertialization.h
--- a/Engine/Plugins/Experimental/AnimNextAnimGraph/Source/AnimNextAnimGraph/Public/Traits/Inertialization.h
+++ b/Engine/Plugins/Experimental/AnimNextAnimGraph/Source/AnimNextAnimGraph/Public/Traits/Inertialization.h
@@ -74,7 +74,7 @@ USTRUCT()
 struct ANIMNEXTANIMGRAPH_API FAnimNextInertializationRequestEvent : public FAnimNextTraitEvent
 {
 	GENERATED_BODY()
-	DECLARE_ANIM_TRAIT_EVENT(FAnimNextInertializationRequestEvent)
+	DECLARE_ANIM_TRAIT_EVENT(FAnimNextInertializationRequestEvent, FAnimNextTraitEvent)
 
 	UE::AnimNext::FInertializationRequest Request;
 };
\ No newline at end of file
diff --git a/Engine/Plugins/Experimental/AnimNextAnimGraph/Source/AnimNextAnimGraphTestSuite/Private/AnimNextAnimGraphTraitEventTest.cpp b/Engine/Plugins/Experimental/AnimNextAnimGraph/Source/AnimNextAnimGraphTestSuite/Private/AnimNextAnimGraphTraitEventTest.cpp
--- a/Engine/Plugins/Experimental/AnimNextAnimGraph/Source/AnimNextAnimGraphTestSuite/Private/AnimNextAnimGraphTraitEventTest.cpp
+++ b/Engine/Plugins/Experimental/AnimNextAnimGraph/Source/AnimNextAnimGraphTestSuite/Private/AnimNextAnimGraphTraitEventTest.cpp
@@ -307,10 +307,12 @@ bool FAnimationAnimNextRuntimeTest_GraphTraitEvent::RunTest(const FString& InPar
 				auto EventA = MakeTraitEvent<FTraitAnimGraphTest_EventA>();
 				auto EventB = MakeTraitEvent<FTraitAnimGraphTest_EventB>();
 
-				GraphInstance.QueueInputTraitEvent(EventA);
-				GraphInstance.QueueInputTraitEvent(EventB);
+				FTraitEventList InputEventList;
+				InputEventList.Push(EventA);
+				InputEventList.Push(EventB);
 
-				UpdateGraph(GraphInstance, 0.0333f);
+				FTraitEventList OutputEventList;
+				UpdateGraph(GraphInstance, 0.0333f, InputEventList, OutputEventList);
 
 				UE_RETURN_ON_ERROR(EventA->VisitedTraits.Num() == 4, "FAnimationAnimNextRuntimeTest_GraphTraitEvent -> Unexpected number of traits visited");
 				AddErrorIfFalse(EventA->VisitedTraits[0] == NodeTemplateTraitListB[0], "FAnimationAnimNextRuntimeTest_GraphTraitEvent -> Unexpected trait visited");
@@ -330,9 +332,11 @@ bool FAnimationAnimNextRuntimeTest_GraphTraitEvent::RunTest(const FString& InPar
 				auto EventA = MakeTraitEvent<FTraitAnimGraphTest_EventA>();
 				EventA->bTestFlag = true;
 
-				GraphInstance.QueueInputTraitEvent(EventA);
+				FTraitEventList InputEventList;
+				InputEventList.Push(EventA);
 
-				UpdateGraph(GraphInstance, 0.0333f);
+				FTraitEventList OutputEventList;
+				UpdateGraph(GraphInstance, 0.0333f, InputEventList, OutputEventList);
 
 				UE_RETURN_ON_ERROR(EventA->VisitedTraits.Num() == 2, "FAnimationAnimNextRuntimeTest_GraphTraitEvent -> Unexpected number of traits visited");
 				AddErrorIfFalse(EventA->VisitedTraits[0] == NodeTemplateTraitListB[0], "FAnimationAnimNextRuntimeTest_GraphTraitEvent -> Unexpected trait visited");
@@ -345,10 +349,12 @@ bool FAnimationAnimNextRuntimeTest_GraphTraitEvent::RunTest(const FString& InPar
 				EventA->bTestFlag = true;
 				auto EventB = MakeTraitEvent<FTraitAnimGraphTest_EventB>();
 
-				GraphInstance.QueueInputTraitEvent(EventA);
-				GraphInstance.QueueInputTraitEvent(EventB);
+				FTraitEventList InputEventList;
+				InputEventList.Push(EventA);
+				InputEventList.Push(EventB);
 
-				UpdateGraph(GraphInstance, 0.0333f);
+				FTraitEventList OutputEventList;
+				UpdateGraph(GraphInstance, 0.0333f, InputEventList, OutputEventList);
 
 				UE_RETURN_ON_ERROR(EventA->VisitedTraits.Num() == 2, "FAnimationAnimNextRuntimeTest_GraphTraitEvent -> Unexpected number of traits visited");
 				AddErrorIfFalse(EventA->VisitedTraits[0] == NodeTemplateTraitListB[0], "FAnimationAnimNextRuntimeTest_GraphTraitEvent -> Unexpected trait visited");
@@ -366,9 +372,11 @@ bool FAnimationAnimNextRuntimeTest_GraphTraitEvent::RunTest(const FString& InPar
 				auto EventB = MakeTraitEvent<FTraitAnimGraphTest_EventB>();
 				EventB->bTestFlag0 = true;
 
-				GraphInstance.QueueInputTraitEvent(EventB);
+				FTraitEventList InputEventList;
+				InputEventList.Push(EventB);
 
-				UpdateGraph(GraphInstance, 0.0333f);
+				FTraitEventList OutputEventList;
+				UpdateGraph(GraphInstance, 0.0333f, InputEventList, OutputEventList);
 
 				UE_RETURN_ON_ERROR(EventB->VisitedTraits.Num() == 4, "FAnimationAnimNextRuntimeTest_GraphTraitEvent -> Unexpected number of traits visited");
 				AddErrorIfFalse(EventB->VisitedTraits[0] == NodeTemplateTraitListB[0], "FAnimationAnimNextRuntimeTest_GraphTraitEvent -> Unexpected trait visited");
@@ -389,9 +397,11 @@ bool FAnimationAnimNextRuntimeTest_GraphTraitEvent::RunTest(const FString& InPar
 				auto EventB = MakeTraitEvent<FTraitAnimGraphTest_EventB>();
 				EventB->bTestFlag1 = true;
 
-				GraphInstance.QueueInputTraitEvent(EventB);
+				FTraitEventList InputEventList;
+				InputEventList.Push(EventB);
 
-				UpdateGraph(GraphInstance, 0.0333f);
+				FTraitEventList OutputEventList;
+				UpdateGraph(GraphInstance, 0.0333f, InputEventList, OutputEventList);
 
 				UE_RETURN_ON_ERROR(EventB->VisitedTraits.Num() == 4, "FAnimationAnimNextRuntimeTest_GraphTraitEvent -> Unexpected number of traits visited");
 				AddErrorIfFalse(EventB->VisitedTraits[0] == NodeTemplateTraitListB[0], "FAnimationAnimNextRuntimeTest_GraphTraitEvent -> Unexpected trait visited");
diff --git a/Engine/Plugins/Experimental/AnimNextAnimGraph/Source/AnimNextAnimGraphTestSuite/Private/AnimNextAnimGraphTraitEventTest.h b/Engine/Plugins/Experimental/AnimNextAnimGraph/Source/AnimNextAnimGraphTestSuite/Private/AnimNextAnimGraphTraitEventTest.h
--- a/Engine/Plugins/Experimental/AnimNextAnimGraph/Source/AnimNextAnimGraphTestSuite/Private/AnimNextAnimGraphTraitEventTest.h
+++ b/Engine/Plugins/Experimental/AnimNextAnimGraph/Source/AnimNextAnimGraphTestSuite/Private/AnimNextAnimGraphTraitEventTest.h
@@ -7,24 +7,18 @@
 #include "TraitCore/TraitUID.h"
 #include "TraitCore/TraitEvent.h"
 
-#include "AnimNextAnimGraphTraitEventTest.generated.h"
-
-USTRUCT()
 struct FTraitAnimGraphTest_EventA : public FAnimNextTraitEvent
 {
-	GENERATED_BODY()
-	DECLARE_ANIM_TRAIT_EVENT(FTraitAnimGraphTest_EventA)
+	DECLARE_ANIM_TRAIT_EVENT(FTraitAnimGraphTest_EventA, FAnimNextTraitEvent)
 
 	bool bTestFlag = false;
 
 	TArray<UE::AnimNext::FTraitUID> VisitedTraits;
 };
 
-USTRUCT()
 struct FTraitAnimGraphTest_EventB : public FAnimNextTraitEvent
 {
-	GENERATED_BODY()
-	DECLARE_ANIM_TRAIT_EVENT(FTraitAnimGraphTest_EventB)
+	DECLARE_ANIM_TRAIT_EVENT(FTraitAnimGraphTest_EventB, FAnimNextTraitEvent)
 
 	bool bTestFlag0 = false;
 	bool bTestFlag1 = false;
diff --git a/Engine/Plugins/Experimental/AnimNextAnimGraph/Source/AnimNextAnimGraphTestSuite/Private/AnimNextAnimGraphTraitGraphTest.cpp b/Engine/Plugins/Experimental/AnimNextAnimGraph/Source/AnimNextAnimGraphTestSuite/Private/AnimNextAnimGraphTraitGraphTest.cpp
--- a/Engine/Plugins/Experimental/AnimNextAnimGraph/Source/AnimNextAnimGraphTestSuite/Private/AnimNextAnimGraphTraitGraphTest.cpp
+++ b/Engine/Plugins/Experimental/AnimNextAnimGraph/Source/AnimNextAnimGraphTestSuite/Private/AnimNextAnimGraphTraitGraphTest.cpp
@@ -455,7 +455,9 @@ bool FAnimationAnimNextRuntimeTest_GraphExecute::RunTest(const FString& InParame
 		);
 
 		{
-			UE::AnimNext::UpdateGraph(GraphInstance, 1.0f / 30.0f);
+			UE::AnimNext::FTraitEventList InputEventList;
+			UE::AnimNext::FTraitEventList OutputEventList;
+			UE::AnimNext::UpdateGraph(GraphInstance, 1.0f / 30.0f, InputEventList, OutputEventList);
 			(void)UE::AnimNext::EvaluateGraph(GraphInstance);
 		}
 
@@ -598,7 +600,9 @@ bool FAnimationAnimNextRuntimeTest_GraphExecuteLatent::RunTest(const FString& In
 		);
 
 		{
-			UE::AnimNext::UpdateGraph(GraphInstance, 1.0f / 30.0f);
+			UE::AnimNext::FTraitEventList InputEventList;
+			UE::AnimNext::FTraitEventList OutputEventList;
+			UE::AnimNext::UpdateGraph(GraphInstance, 1.0f / 30.0f, InputEventList, OutputEventList);
 			(void)UE::AnimNext::EvaluateGraph(GraphInstance);
 		}
 
diff --git a/Engine/Plugins/Experimental/AnimNextAnimGraph/Source/AnimNextAnimGraphTestSuite/Private/AnimNextAnimGraphTraitInterfacesTest.cpp b/Engine/Plugins/Experimental/AnimNextAnimGraph/Source/AnimNextAnimGraphTestSuite/Private/AnimNextAnimGraphTraitInterfacesTest.cpp
--- a/Engine/Plugins/Experimental/AnimNextAnimGraph/Source/AnimNextAnimGraphTestSuite/Private/AnimNextAnimGraphTraitInterfacesTest.cpp
+++ b/Engine/Plugins/Experimental/AnimNextAnimGraph/Source/AnimNextAnimGraphTestSuite/Private/AnimNextAnimGraphTraitInterfacesTest.cpp
@@ -687,7 +687,9 @@ bool FAnimationAnimNextRuntimeTest_IUpdate::RunTest(const FString& InParameters)
 			Private::UpdatedTraits = &UpdatedTraits;
 
 			// Call pre/post update on our graph
-			UpdateGraph(GraphInstance, 0.0333f);
+			FTraitEventList InputEventList;
+			FTraitEventList OutputEventList;
+			UpdateGraph(GraphInstance, 0.0333f, InputEventList, OutputEventList);
 
 			AddErrorIfFalse(UpdatedTraits.Num() == 6, "FAnimationAnimNextRuntimeTest_IUpdate -> Expected 6 nodes to have been visited during the update traversal");
 			AddErrorIfFalse(UpdatedTraits[0] == FTraitWithChildren::TraitUID, "FAnimationAnimNextRuntimeTest_IUpdate -> Unexpected update order");		// NodeC
@@ -961,7 +963,9 @@ bool FAnimationAnimNextRuntimeTest_IScopedInterface::RunTest(const FString& InPa
 
 			// Call pre/post update on our graph
 			Private::AutoPopTag = true;
-			UpdateGraph(GraphInstance, 0.0333f);
+			FTraitEventList InputEventList;
+			FTraitEventList OutputEventList;
+			UpdateGraph(GraphInstance, 0.0333f, InputEventList, OutputEventList);
 
 			AddErrorIfFalse(IsTagInScope.Num() == 8, "FAnimationAnimNextRuntimeTest_IScopedInterface -> Unexpected number of entries");
 			AddErrorIfFalse(IsTagInScope[0] == false, "FAnimationAnimNextRuntimeTest_IScopedInterface -> Unexpected scoped tag state");		// NodeC::PreUpdate (template 0)
@@ -977,7 +981,7 @@ bool FAnimationAnimNextRuntimeTest_IScopedInterface::RunTest(const FString& InPa
 			Private::AutoPopTag = false;
 
 			// Call pre/post update on our graph
-			UpdateGraph(GraphInstance, 0.0333f);
+			UpdateGraph(GraphInstance, 0.0333f, InputEventList, OutputEventList);
 
 			AddErrorIfFalse(IsTagInScope.Num() == 8, "FAnimationAnimNextRuntimeTest_IScopedInterface -> Unexpected number of entries");
 			AddErrorIfFalse(IsTagInScope[0] == false, "FAnimationAnimNextRuntimeTest_IScopedInterface -> Unexpected scoped tag state");		// NodeC::PreUpdate (template 0)
