AnimNext : Created Trait Interface Registry
- Added editor only functions at ITraitInterface : IsInternal, GetDisplayName, GetDisplayShortName
- Added IsHidden to Traits, and an Advanced videw option to the Traits Editor, so we can show / hide the ones we flag as hidden
- Fixed base interface implemented interfaces going into the required interfaces list
- Traits Editor : Fixed available Traits list only displaying the Interfaces that are currently used in the stack (it has to display all the interfaces of the available traits)
- Added "Internal" flag to Trait Interfaces, hiding the interfaces flagged as such in the Traits editor

#rnx

diff --git a/Engine/Plugins/Experimental/AnimNext/Source/AnimNext/Internal/TraitCore/TraitInterfaceRegistry.h b/Engine/Plugins/Experimental/AnimNext/Source/AnimNext/Internal/TraitCore/TraitInterfaceRegistry.h
new file mode 100644
--- /dev/null
+++ b/Engine/Plugins/Experimental/AnimNext/Source/AnimNext/Internal/TraitCore/TraitInterfaceRegistry.h
@@ -0,0 +1,66 @@
+// Copyright Epic Games, Inc. All Rights Reserved.
+
+#pragma once
+
+#include "CoreMinimal.h"
+#include "TraitCore/ITraitInterface.h"
+#include "TraitCore/TraitInterfaceUID.h"
+
+namespace UE::AnimNext
+{
+	/**
+	 * FTraitInterfaceRegistry
+	 * 
+	 * A global registry of all existing trait interfaces that can be used in animation graph traits.
+	 * 
+	 * @see FTraitInterface
+	 */
+	struct ANIMNEXT_API FTraitInterfaceRegistry final
+	{
+		// Access the global registry
+		static FTraitInterfaceRegistry& Get();
+
+		// Finds and returns the trait interface associated with the provided trait interface UID.
+		// If the trait interface is not registered, nullptr is returned.
+		const ITraitInterface* Find(FTraitInterfaceUID InterfaceUID) const;
+
+		// Registers a trait interface dynamically
+		void Register(const TSharedPtr<ITraitInterface>& TraitInterface);
+
+		// Unregisters a trait intrerface dynamically
+		void Unregister(const TSharedPtr<ITraitInterface>& TraitInterface);
+
+		// Returns a list of all registered trait interfaces
+		TArray<const ITraitInterface*> GetTraitInterfaces() const;
+
+		// Returns the number of registered trait interfaces
+		uint32 GetNum() const;
+
+	private:
+		FTraitInterfaceRegistry() = default;
+		FTraitInterfaceRegistry(const FTraitInterfaceRegistry&) = delete;
+		FTraitInterfaceRegistry(FTraitInterfaceRegistry&&) = default;
+		FTraitInterfaceRegistry& operator=(const FTraitInterfaceRegistry&) = delete;
+		FTraitInterfaceRegistry& operator=(FTraitInterfaceRegistry&&) = default;
+
+		// Static init lifetime functions
+		static void StaticRegister(const TSharedPtr<ITraitInterface>& TraitInterface);
+		static void StaticUnregister(const TSharedPtr<ITraitInterface>& TraitInterface);
+
+		// Module lifetime functions
+		static void Init();
+		static void Destroy();
+
+		// Holds information for each registered trait
+		struct FRegistryEntry
+		{
+			// A pointer to the trait
+			TSharedPtr<ITraitInterface> TraitInterface = nullptr;
+		};
+
+		TMap<FTraitInterfaceUIDRaw, FRegistryEntry>	TraitInterfaceUIDToEntryMap;
+
+		friend class FModule;
+		friend struct FTraitInterfaceStaticInitHook;
+	};
+}
diff --git a/Engine/Plugins/Experimental/AnimNext/Source/AnimNext/Internal/TraitCore/TraitRegistry.h b/Engine/Plugins/Experimental/AnimNext/Source/AnimNext/Internal/TraitCore/TraitRegistry.h
--- a/Engine/Plugins/Experimental/AnimNext/Source/AnimNext/Internal/TraitCore/TraitRegistry.h
+++ b/Engine/Plugins/Experimental/AnimNext/Source/AnimNext/Internal/TraitCore/TraitRegistry.h
@@ -37,7 +37,7 @@ namespace UE::AnimNext
 		// If the matching trait is not registered, nullptr is returned.
 		const FTrait* Find(const UScriptStruct* TraitSharedDataStruct) const;
 
-		// Finds and returns the trait UID associated with the provided trait name.
+		// Finds and returns the trait associated with the provided trait name.
 		// If the matching trait is not registered, nullptr is returned.
 		const FTrait* Find(FName TraitTypeName) const;
 
diff --git a/Engine/Plugins/Experimental/AnimNext/Source/AnimNext/Private/AnimNextModule.cpp b/Engine/Plugins/Experimental/AnimNext/Source/AnimNext/Private/AnimNextModule.cpp
--- a/Engine/Plugins/Experimental/AnimNext/Source/AnimNext/Private/AnimNextModule.cpp
+++ b/Engine/Plugins/Experimental/AnimNext/Source/AnimNext/Private/AnimNextModule.cpp
@@ -26,6 +26,7 @@
 #include "Scheduler/Scheduler.h"
 #include "TraitCore/NodeTemplateRegistry.h"
 #include "TraitCore/TraitRegistry.h"
+#include "TraitCore/TraitInterfaceRegistry.h"
 
 #if WITH_ANIMNEXT_CONSOLE_COMMANDS
 #include "HAL/IConsoleManager.h"
@@ -74,6 +75,7 @@ namespace UE::AnimNext
 		RegisterParameterSourceFactory("ObjectProxy", MakeShared<FObjectProxyFactory>());
 		FDataRegistry::Init();
 		FTraitRegistry::Init();
+		FTraitInterfaceRegistry::Init();
 		FNodeTemplateRegistry::Init();
 		FScheduler::Init();
 		FRigVMRuntimeDataRegistry::Init();
@@ -132,6 +134,7 @@ namespace UE::AnimNext
 		FRigVMRuntimeDataRegistry::Destroy();
 		FScheduler::Destroy();
 		FNodeTemplateRegistry::Destroy();
+		FTraitInterfaceRegistry::Destroy();
 		FTraitRegistry::Destroy();
 		FDataRegistry::Destroy();
 		UnregisterParameterSourceFactory("ObjectProxy");
diff --git a/Engine/Plugins/Experimental/AnimNext/Source/AnimNext/Private/TraitCore/TraiInterfacetRegistry.cpp b/Engine/Plugins/Experimental/AnimNext/Source/AnimNext/Private/TraitCore/TraiInterfacetRegistry.cpp
new file mode 100644
--- /dev/null
+++ b/Engine/Plugins/Experimental/AnimNext/Source/AnimNext/Private/TraitCore/TraiInterfacetRegistry.cpp
@@ -0,0 +1,168 @@
+// Copyright Epic Games, Inc. All Rights Reserved.
+
+#include "TraitCore/TraitInterfaceRegistry.h"
+
+#include "TraitCore/ITraitInterface.h"
+
+namespace UE::AnimNext
+{
+
+namespace Private
+{
+	static FTraitInterfaceRegistry* GTraitInterfaceRegistry = nullptr;
+
+	static TArray<TSharedPtr<ITraitInterface>> GPendingInterfaceRegistrationQueue;
+}
+
+FTraitInterfaceRegistry& FTraitInterfaceRegistry::Get()
+{
+	checkf(Private::GTraitInterfaceRegistry, TEXT("Trait Registry is not instanced. It is only valid to access this while the engine module is loaded."));
+	return *Private::GTraitInterfaceRegistry;
+}
+
+void FTraitInterfaceRegistry::Init()
+{
+	if (ensure(Private::GTraitInterfaceRegistry == nullptr))
+	{
+		Private::GTraitInterfaceRegistry = new FTraitInterfaceRegistry();
+
+		// Register all our pending static init traits
+		for (TSharedPtr<ITraitInterface>& TraitInterface : Private::GPendingInterfaceRegistrationQueue)
+		{
+			Private::GTraitInterfaceRegistry->Register(TraitInterface);
+		}
+
+		// Reset the registration queue, it won't be used anymore now that the registry is up and ready
+		Private::GPendingInterfaceRegistrationQueue.Empty(0);
+	}
+}
+
+void FTraitInterfaceRegistry::Destroy()
+{
+	if (ensure(Private::GTraitInterfaceRegistry != nullptr))
+	{
+		TArray<FRegistryEntry> Entries;
+		Private::GTraitInterfaceRegistry->TraitInterfaceUIDToEntryMap.GenerateValueArray(Entries);
+
+		for (const FRegistryEntry& Entry : Entries)
+		{
+			Private::GTraitInterfaceRegistry->Unregister(Entry.TraitInterface);
+		}
+
+		delete Private::GTraitInterfaceRegistry;
+		Private::GTraitInterfaceRegistry = nullptr;
+	}
+}
+
+void FTraitInterfaceRegistry::StaticRegister(const TSharedPtr<ITraitInterface>& TraitInterface)
+{
+	if (Private::GTraitInterfaceRegistry != nullptr)
+	{
+		// Registry is already up and running, use it
+		Private::GTraitInterfaceRegistry->Register(TraitInterface);
+	}
+	else
+	{
+		// Registry isn't ready yet, queue up our trait
+		// Once Init() is called, our queue will be processed
+		Private::GPendingInterfaceRegistrationQueue.Add(TraitInterface);
+	}
+}
+
+void FTraitInterfaceRegistry::StaticUnregister(const TSharedPtr<ITraitInterface>& TraitInterface)
+{
+	if (Private::GTraitInterfaceRegistry != nullptr)
+	{
+		// Registry is already up and running, use it
+		Private::GTraitInterfaceRegistry->Unregister(TraitInterface);
+	}
+	else
+	{
+		// Registry isn't ready yet or it got destroyed before the traits are unregistering
+		const int32 TraitInterfaceIndex = Private::GPendingInterfaceRegistrationQueue.IndexOfByKey(TraitInterface);
+		if (TraitInterfaceIndex != INDEX_NONE)
+		{
+			Private::GPendingInterfaceRegistrationQueue.RemoveAtSwap(TraitInterfaceIndex);
+		}
+	}
+}
+
+const ITraitInterface* FTraitInterfaceRegistry::Find(FTraitInterfaceUID TraitInterfaceUID) const
+{
+	if (!TraitInterfaceUID.IsValid())
+	{
+		return nullptr;
+	}
+
+	if (const FRegistryEntry* Entry = TraitInterfaceUIDToEntryMap.Find(TraitInterfaceUID.GetUID()))
+	{
+		return Entry->TraitInterface.Get();
+	}
+
+	// Trait UID not found
+	return nullptr;
+}
+
+void FTraitInterfaceRegistry::Register(const TSharedPtr<ITraitInterface>& TraitInterface)
+{
+	if (TraitInterface == nullptr)
+	{
+		return;
+	}
+
+	const FTraitInterfaceUID TraitInterfaceUID = TraitInterface->GetInterfaceUID();
+
+	if (ensure(!TraitInterfaceUIDToEntryMap.Contains(TraitInterfaceUID.GetUID())))
+	{
+		TraitInterfaceUIDToEntryMap.Add(TraitInterfaceUID.GetUID(), FRegistryEntry{ TraitInterface });
+	}
+}
+
+void FTraitInterfaceRegistry::Unregister(const TSharedPtr<ITraitInterface>& TraitInterface)
+{
+	if (TraitInterface == nullptr)
+	{
+		return;
+	}
+
+	const FTraitInterfaceUID TraitInterfaceUID = TraitInterface->GetInterfaceUID();
+
+	if (FRegistryEntry* Entry = TraitInterfaceUIDToEntryMap.Find(TraitInterfaceUID.GetUID()))
+	{
+		check(Entry->TraitInterface != nullptr);
+		TraitInterfaceUIDToEntryMap.Remove(TraitInterfaceUID.GetUID());
+	}
+}
+
+TArray<const ITraitInterface*> FTraitInterfaceRegistry::GetTraitInterfaces() const
+{
+	TArray<const ITraitInterface*> TraitInterfaces;
+	TraitInterfaces.Reserve(TraitInterfaceUIDToEntryMap.Num());
+	
+	for (const auto& It : TraitInterfaceUIDToEntryMap)
+	{
+		TraitInterfaces.Add(It.Value.TraitInterface.Get());
+	}
+
+	return TraitInterfaces;
+}
+
+uint32 FTraitInterfaceRegistry::GetNum() const
+{
+	return TraitInterfaceUIDToEntryMap.Num();
+}
+
+// --- FTraitInterfaceStaticInitHook ---
+
+FTraitInterfaceStaticInitHook::FTraitInterfaceStaticInitHook(const TSharedPtr<ITraitInterface>& InTraitInterface)
+	: TraitInterface(InTraitInterface)
+{
+	FTraitInterfaceRegistry::StaticRegister(TraitInterface);
+}
+
+FTraitInterfaceStaticInitHook::~FTraitInterfaceStaticInitHook()
+{
+	FTraitInterfaceRegistry::StaticUnregister(TraitInterface);
+}
+
+} // end namespace UE::AnimNext
diff --git a/Engine/Plugins/Experimental/AnimNext/Source/AnimNext/Private/TraitInterfaces/IGarbageColllection.cpp b/Engine/Plugins/Experimental/AnimNext/Source/AnimNext/Private/TraitInterfaces/IGarbageColllection.cpp
--- a/Engine/Plugins/Experimental/AnimNext/Source/AnimNext/Private/TraitInterfaces/IGarbageColllection.cpp
+++ b/Engine/Plugins/Experimental/AnimNext/Source/AnimNext/Private/TraitInterfaces/IGarbageColllection.cpp
@@ -7,6 +7,21 @@
 
 namespace UE::AnimNext
 {
+	AUTO_REGISTER_ANIM_TRAIT_INTERFACE(IGarbageCollection)
+
+#if WITH_EDITOR
+	const FText& IGarbageCollection::GetDisplayName() const
+	{
+		static FText InterfaceName = NSLOCTEXT("TraitInterfaces", "TraitInterface_IGarbageCollection_Name", "Garbage Collection");
+		return InterfaceName;
+	}
+	const FText& IGarbageCollection::GetDisplayShortName() const
+	{
+		static FText InterfaceShortName = NSLOCTEXT("TraitInterfaces", "TraitInterface_IGarbageCollection_ShortName", "GC");
+		return InterfaceShortName;
+	}
+#endif // WITH_EDITOR
+
 	void IGarbageCollection::RegisterWithGC(const FExecutionContext& Context, const FTraitBinding& Binding)
 	{
 		FGCGraphInstanceComponent& Component = Context.GetComponent<FGCGraphInstanceComponent>();
diff --git a/Engine/Plugins/Experimental/AnimNext/Source/AnimNext/Private/TraitInterfaces/IHierarchy.cpp b/Engine/Plugins/Experimental/AnimNext/Source/AnimNext/Private/TraitInterfaces/IHierarchy.cpp
--- a/Engine/Plugins/Experimental/AnimNext/Source/AnimNext/Private/TraitInterfaces/IHierarchy.cpp
+++ b/Engine/Plugins/Experimental/AnimNext/Source/AnimNext/Private/TraitInterfaces/IHierarchy.cpp
@@ -6,6 +6,21 @@
 
 namespace UE::AnimNext
 {
+	AUTO_REGISTER_ANIM_TRAIT_INTERFACE(IHierarchy)
+
+#if WITH_EDITOR
+	const FText& IHierarchy::GetDisplayName() const
+	{
+		static FText InterfaceName = NSLOCTEXT("TraitInterfaces", "TraitInterface_IHierarchy_Name", "Hierarchy");
+		return InterfaceName;
+	}
+	const FText& IHierarchy::GetDisplayShortName() const
+	{
+		static FText InterfaceShortName = NSLOCTEXT("TraitInterfaces", "TraitInterface_IHierarchy_ShortName", "HIE");
+		return InterfaceShortName;
+	}
+#endif // WITH_EDITOR
+
 	uint32 IHierarchy::GetNumChildren(const FExecutionContext& Context, const TTraitBinding<IHierarchy>& Binding) const
 	{
 		// We only wish to count children of the queried trait
diff --git a/Engine/Plugins/Experimental/AnimNext/Source/AnimNext/Public/TraitCore/IScopedTraitInterface.h b/Engine/Plugins/Experimental/AnimNext/Source/AnimNext/Public/TraitCore/IScopedTraitInterface.h
--- a/Engine/Plugins/Experimental/AnimNext/Source/AnimNext/Public/TraitCore/IScopedTraitInterface.h
+++ b/Engine/Plugins/Experimental/AnimNext/Source/AnimNext/Public/TraitCore/IScopedTraitInterface.h
@@ -14,6 +14,6 @@ namespace UE::AnimNext
 	 */
 	struct ANIMNEXT_API IScopedTraitInterface : ITraitInterface
 	{
-		DECLARE_ANIM_TRAIT_INTERFACE(IScopedTraitInterface, STI, 0x621b01eb)
+		DECLARE_ANIM_TRAIT_INTERFACE(IScopedTraitInterface, 0x621b01eb)
 	};
 }
diff --git a/Engine/Plugins/Experimental/AnimNext/Source/AnimNext/Public/TraitCore/ITraitInterface.h b/Engine/Plugins/Experimental/AnimNext/Source/AnimNext/Public/TraitCore/ITraitInterface.h
--- a/Engine/Plugins/Experimental/AnimNext/Source/AnimNext/Public/TraitCore/ITraitInterface.h
+++ b/Engine/Plugins/Experimental/AnimNext/Source/AnimNext/Public/TraitCore/ITraitInterface.h
@@ -5,12 +5,25 @@
 #include "CoreMinimal.h"
 #include "TraitCore/TraitInterfaceUID.h"
 
+namespace UE::AnimNext
+{
+	struct ITraitInterface;
+}
+
 // Helper macros
-#define DECLARE_ANIM_TRAIT_INTERFACE(InterfaceName, InterfaceShortName, InterfaceNameHash) \
+
+// In a trait interface struct declaration, this macro declares the necessary boilerplate we require
+#define DECLARE_ANIM_TRAIT_INTERFACE(InterfaceName, InterfaceNameHash) \
 	/* Globally unique UID for this interface */ \
-	static constexpr UE::AnimNext::FTraitInterfaceUID InterfaceUID = UE::AnimNext::FTraitInterfaceUID(InterfaceNameHash, TEXT(#InterfaceName), TEXT(#InterfaceShortName)); \
+	static constexpr UE::AnimNext::FTraitInterfaceUID InterfaceUID = UE::AnimNext::FTraitInterfaceUID(InterfaceNameHash, TEXT(#InterfaceName)); \
 	virtual UE::AnimNext::FTraitInterfaceUID GetInterfaceUID() const override { return InterfaceUID; }
 
+// Allows a trait interface to auto-register and unregister within the current execution scope
+// The trait interface must be found in the current scope without a namespace qualification
+#define AUTO_REGISTER_ANIM_TRAIT_INTERFACE(Interface) \
+		UE::AnimNext::FTraitInterfaceStaticInitHook Interface##Hook(MakeShared<Interface>());
+
+
 namespace UE::AnimNext
 {
 	struct FExecutionContext;	// Derived types will have functions that accept the execution context
@@ -31,8 +44,19 @@ namespace UE::AnimNext
 
 		// Returns the globally unique UID for this interface
 		virtual FTraitInterfaceUID GetInterfaceUID() const { return InterfaceUID; };
+
+#if WITH_EDITOR
+		// Internal interfaces are only displayed in the Traits Editor in the Advanced View
+		virtual bool IsInternal() const { return false; }
+
+		// Human readable interface names, in long and short format
+		virtual const FText& GetDisplayName() const { ensure(false); return FText::GetEmpty(); }
+		// Human readable interface names, in short format (ideally 3 or 4 letters, due to space restrictions in the editor)
+		virtual const FText& GetDisplayShortName() const { ensure(false); return FText::GetEmpty(); }
+#endif // WITH_EDITOR
 	};
 
+
 	/**
 	 * Trait Stack Propagation
 	 * 
@@ -47,4 +71,20 @@ namespace UE::AnimNext
 		// Do not forward the call to our parent, execution stops
 		Stop,
 	};
+
+	/**
+	* FTraitInterfaceStaticInitHook
+	*
+	* Allows trait interfaces to automatically register/unregister within the current scope.
+	* This can be used during static init.
+	* See AUTO_REGISTER_ANIM_TRAIT_INTERFACE
+	*/
+	struct ANIMNEXT_API FTraitInterfaceStaticInitHook final
+	{
+		explicit FTraitInterfaceStaticInitHook(const TSharedPtr<UE::AnimNext::ITraitInterface>& InTraitInterface);
+		~FTraitInterfaceStaticInitHook();
+
+	private:
+		TSharedPtr<ITraitInterface> TraitInterface = nullptr;
+	};
 }
diff --git a/Engine/Plugins/Experimental/AnimNext/Source/AnimNext/Public/TraitCore/Trait.h b/Engine/Plugins/Experimental/AnimNext/Source/AnimNext/Public/TraitCore/Trait.h
--- a/Engine/Plugins/Experimental/AnimNext/Source/AnimNext/Public/TraitCore/Trait.h
+++ b/Engine/Plugins/Experimental/AnimNext/Source/AnimNext/Public/TraitCore/Trait.h
@@ -124,7 +124,7 @@ class FArchive;
 	{ \
 		/* Thread safe cache initialization */ \
 		static TArray<UE::AnimNext::FTraitInterfaceUID> CachedInterfaceList = FTrait::BuildTraitInterfaceList( \
-			TraitSuper::GetTraitInterfaces(), \
+			TraitSuper::GetTraitRequiredInterfaces(), \
 			{ \
 				InterfaceEnumeratorMacro(ANIM_NEXT_IMPL_GET_INTERFACES_IMPL_FOR_INTERFACE) \
 			}); \
@@ -388,6 +388,9 @@ namespace UE::AnimNext
 		TArray<FLatentPropertyMetadata> GetLatentPropertyHandles(
 			bool bFilterEditorOnly,
 			const TFunction<uint16(FName PropertyName)>& GetTraitLatentPropertyIndex) const;
+
+		// Makes the Trait Editor only display this Trait in Advanced view
+		virtual bool IsHidden() const { return false; }
 #endif
 
 	protected:
diff --git a/Engine/Plugins/Experimental/AnimNext/Source/AnimNext/Public/TraitCore/TraitInterfaceUID.h b/Engine/Plugins/Experimental/AnimNext/Source/AnimNext/Public/TraitCore/TraitInterfaceUID.h
--- a/Engine/Plugins/Experimental/AnimNext/Source/AnimNext/Public/TraitCore/TraitInterfaceUID.h
+++ b/Engine/Plugins/Experimental/AnimNext/Source/AnimNext/Public/TraitCore/TraitInterfaceUID.h
@@ -25,24 +25,21 @@ namespace UE::AnimNext
 			: UID(INVALID_UID)
 #if !UE_BUILD_SHIPPING
 			, InterfaceName(TEXT("<Invalid Interface UID>"))
-			, InterfaceShortName(TEXT("???"))
 #endif
 		{}
 
 		// Constructs an interface UID
-		explicit constexpr FTraitInterfaceUID(FTraitInterfaceUIDRaw InUID, const TCHAR* InInterfaceName = TEXT("<Unknown Trait Interface Name>"), const TCHAR* InInterfaceShortName = TEXT("???"))
+		explicit constexpr FTraitInterfaceUID(FTraitInterfaceUIDRaw InUID, const TCHAR* InInterfaceName = TEXT("<Unknown Trait Interface Name>"))
 			: UID(InUID)
 #if !UE_BUILD_SHIPPING
 			, InterfaceName(InInterfaceName)
-			, InterfaceShortName(InInterfaceShortName)
 #endif
 		{
 		}
 
-#if !UE_BUILD_SHIPPING || WITH_EDITOR
+#if !UE_BUILD_SHIPPING
 		// Returns a literal string to the interface name
 		constexpr const TCHAR* GetInterfaceName() const { return InterfaceName; }
-		constexpr const TCHAR* GetInterfaceShortName() const { return InterfaceShortName; }
 #endif
 
 		// Returns the interface global UID
@@ -58,7 +55,6 @@ namespace UE::AnimNext
 
 #if !UE_BUILD_SHIPPING || WITH_EDITOR
 		const TCHAR*				InterfaceName = nullptr;
-		const TCHAR*				InterfaceShortName = nullptr;
 #endif
 	};
 
diff --git a/Engine/Plugins/Experimental/AnimNext/Source/AnimNext/Public/TraitInterfaces/IGarbageCollection.h b/Engine/Plugins/Experimental/AnimNext/Source/AnimNext/Public/TraitInterfaces/IGarbageCollection.h
--- a/Engine/Plugins/Experimental/AnimNext/Source/AnimNext/Public/TraitInterfaces/IGarbageCollection.h
+++ b/Engine/Plugins/Experimental/AnimNext/Source/AnimNext/Public/TraitInterfaces/IGarbageCollection.h
@@ -17,7 +17,8 @@ namespace UE::AnimNext
 	 */
 	struct ANIMNEXT_API IGarbageCollection : ITraitInterface
 	{
-		DECLARE_ANIM_TRAIT_INTERFACE(IGarbageCollection, GC, 0x231a2017)
+		DECLARE_ANIM_TRAIT_INTERFACE(IGarbageCollection, 0x231a2017)
+
 
 		// Registers the provided binding for GC callback
 		// Once registered, AddReferencedObjects is called during GC to collect references
@@ -29,6 +30,13 @@ namespace UE::AnimNext
 		// Called when garbage collection requests hard/strong object references
 		// @see UObject::AddReferencedObjects
 		virtual void AddReferencedObjects(const FExecutionContext& Context, const TTraitBinding<IGarbageCollection>& Binding, FReferenceCollector& Collector) const;
+
+#if WITH_EDITOR
+		virtual bool IsInternal() const override { return true; }
+
+		virtual const FText& GetDisplayName() const override;
+		virtual const FText& GetDisplayShortName() const override;
+#endif // WITH_EDITOR
 	};
 
 	/**
diff --git a/Engine/Plugins/Experimental/AnimNext/Source/AnimNext/Public/TraitInterfaces/IHierarchy.h b/Engine/Plugins/Experimental/AnimNext/Source/AnimNext/Public/TraitInterfaces/IHierarchy.h
--- a/Engine/Plugins/Experimental/AnimNext/Source/AnimNext/Public/TraitInterfaces/IHierarchy.h
+++ b/Engine/Plugins/Experimental/AnimNext/Source/AnimNext/Public/TraitInterfaces/IHierarchy.h
@@ -20,7 +20,7 @@ namespace UE::AnimNext
 	 */
 	struct ANIMNEXT_API IHierarchy : ITraitInterface
 	{
-		DECLARE_ANIM_TRAIT_INTERFACE(IHierarchy, HIE, 0x846d8a37)
+		DECLARE_ANIM_TRAIT_INTERFACE(IHierarchy, 0x846d8a37)
 
 		// Returns the number of children of the trait implementation (not the whole stack)
 		// Includes inactive children
@@ -42,6 +42,11 @@ namespace UE::AnimNext
 
 		// Queries the trait stack of the specified binding and calls GetNumChildren for each trait, accumulating the result.
 		static uint32 GetNumStackChildren(const FExecutionContext& Context, const FTraitBinding& Binding);
+
+#if WITH_EDITOR
+		virtual const FText& GetDisplayName() const override;
+		virtual const FText& GetDisplayShortName() const override;
+#endif // WITH_EDITOR
 	};
 
 	/**
diff --git a/Engine/Plugins/Experimental/AnimNext/Source/AnimNextEditor/Private/Graph/STraitEditorView.cpp b/Engine/Plugins/Experimental/AnimNext/Source/AnimNextEditor/Private/Graph/STraitEditorView.cpp
--- a/Engine/Plugins/Experimental/AnimNext/Source/AnimNextEditor/Private/Graph/STraitEditorView.cpp
+++ b/Engine/Plugins/Experimental/AnimNext/Source/AnimNextEditor/Private/Graph/STraitEditorView.cpp
@@ -7,6 +7,7 @@
 #include "ScopedTransaction.h"
 #include "WorkspaceSchema.h"
 #include "Framework/Commands/GenericCommands.h"
+#include "TraitCore/TraitInterfaceRegistry.h"
 #include "TraitCore/TraitRegistry.h"
 #include "ObjectEditorUtils.h"
 #include "Graph/AnimNextGraph_EdGraphNode.h"
@@ -384,9 +385,12 @@ void STraitEditorView::UpdateTraitStatusInStack(const TArray<TSharedPtr<FTraitDa
 
 		for (const FTraitInterfaceUID& MissingInterface : TraitData->StackStatus.MissingInterfaces)
 		{
-			const FString InterfaceName(MissingInterface.GetInterfaceName());
-			const FText MissingError = FText::Format(LOCTEXT("TraitStatusInStack_MissingInterface", "Trait {0} requires a parent implementing interface {1}"), TraitData->TraitDisplayName, FText::FromString(InterfaceName));
-			TraitData->StackStatus.StatusMessages.Add(FTraitStackTraitStatus::FStatusMessage(FTraitStackTraitStatus::EStackStatus::Warning, MissingError));
+			if (const ITraitInterface* TraitInterface = FTraitInterfaceRegistry::Get().Find(MissingInterface))
+			{
+				const FText InterfaceName(TraitInterface->GetDisplayName());
+				const FText MissingError = FText::Format(LOCTEXT("TraitStatusInStack_MissingInterface", "Trait {0} requires a parent implementing interface {1}"), TraitData->TraitDisplayName, InterfaceName);
+				TraitData->StackStatus.StatusMessages.Add(FTraitStackTraitStatus::FStatusMessage(FTraitStackTraitStatus::EStackStatus::Warning, MissingError));
+			}
 		}
 	}
 
@@ -455,6 +459,29 @@ TSharedRef<SWidget> STraitEditorView::GetOptionsMenuWidget()
 		NAME_None,
 		EUserInterfaceActionType::Check
 	);
+	DetailViewOptions.AddMenuEntry(
+		LOCTEXT("TraitEditor_AdvancedView", "Advanced View"),
+		LOCTEXT("TraitEditor_AdvancedView_ToolTip", "Displays all Traits, including hidden ones"),
+		FSlateIcon(),
+		FUIAction(
+			FExecuteAction::CreateLambda([this, TraitEditorSharedDataLocal]()
+				{
+					if (TraitEditorSharedDataLocal.IsValid())
+					{
+						// TODO zzz : Store bShowTraitInterfacesIfWarningsOrErrors in a serialized settings class somewhere
+						TraitEditorSharedDataLocal->bAdvancedView = !TraitEditorSharedDataLocal->bAdvancedView;
+						RefreshWidgets(); //RefreshTraitStack();
+					}
+				}),
+			FCanExecuteAction(),
+			FIsActionChecked::CreateLambda([TraitEditorSharedDataLocal]()
+				{
+					return (TraitEditorSharedDataLocal.IsValid()) ? TraitEditorSharedDataLocal->bAdvancedView : false;
+				})
+		),
+		NAME_None,
+		EUserInterfaceActionType::Check
+	);
 
 
 	return SNew(SHorizontalBox)
@@ -588,7 +615,7 @@ int32 STraitEditorView::GetTraitPinIndex(UAnimNextGraph_EdGraphNode* InEdGraphNo
 
 									for (const FTraitInterfaceUID& TraitInterface : TraitImplementedInterfaces)
 									{
-										if (!FTraitEditorUtils::IsInternalTraitInteface(TraitInterface))
+										if (!FTraitEditorUtils::IsInternal(TraitInterface))
 										{
 											int32 Index = StackUsedInterfaces.Find(TraitInterface);
 											if (Index != INDEX_NONE)
@@ -605,7 +632,7 @@ int32 STraitEditorView::GetTraitPinIndex(UAnimNextGraph_EdGraphNode* InEdGraphNo
 
 									for (const FTraitInterfaceUID& TraitInterface : TraitRequiredInterfaces)
 									{
-										if (!FTraitEditorUtils::IsInternalTraitInteface(TraitInterface))
+										if (!FTraitEditorUtils::IsInternal(TraitInterface))
 										{
 											int32 Index = StackUsedInterfaces.Find(TraitInterface);
 											if (Index != INDEX_NONE)
diff --git a/Engine/Plugins/Experimental/AnimNext/Source/AnimNextEditor/Private/Graph/STraitListView.cpp b/Engine/Plugins/Experimental/AnimNext/Source/AnimNextEditor/Private/Graph/STraitListView.cpp
--- a/Engine/Plugins/Experimental/AnimNext/Source/AnimNextEditor/Private/Graph/STraitListView.cpp
+++ b/Engine/Plugins/Experimental/AnimNext/Source/AnimNextEditor/Private/Graph/STraitListView.cpp
@@ -11,6 +11,7 @@
 #include "ObjectEditorUtils.h"
 #include "TraitCore/TraitMode.h"
 #include "TraitCore/TraitUID.h"
+#include "TraitCore/TraitInterfaceRegistry.h"
 #include "TraitCore/TraitRegistry.h"
 #include "Common/SCategoryTableRow.h"
 #include "Widgets/Input/SButton.h"
@@ -700,11 +701,14 @@ TSharedRef<ITableRow> STraitListView::HandleGenerateRow(TSharedRef<FTraitListEnt
 															TraitInfoString.Append(LOCTEXT("TraitInfoImplementedInterfaces", "Implements :").ToString())
 																.AppendChar(TEXT('\n'));
 
-															for (const auto& ImplementedInterface : TraitData->ImplementedInterfaces)
+															for (const auto& ImplementedInterfaceUID : TraitData->ImplementedInterfaces)
 															{
-																TraitInfoString.Append(TEXT("- "))
-																	.Append(ImplementedInterface.GetInterfaceName())
-																	.AppendChar(TEXT('\n'));
+																if (const ITraitInterface* ImplementedInterface = FTraitInterfaceRegistry::Get().Find(ImplementedInterfaceUID))
+																{
+																	TraitInfoString.Append(TEXT("- "))
+																		.Append(ImplementedInterface->GetDisplayName().ToString())
+																		.AppendChar(TEXT('\n'));
+																}
 															}
 														}
 
@@ -714,11 +718,14 @@ TSharedRef<ITableRow> STraitListView::HandleGenerateRow(TSharedRef<FTraitListEnt
 															TraitInfoString.Append(LOCTEXT("TraitInfoRequiredInterfaces", "Requires :").ToString())
 															.AppendChar(TEXT('\n'));
 
-															for (const auto& RequiredInterface : TraitData->RequiredInterfaces)
+															for (const auto& RequiredInterfaceUID : TraitData->RequiredInterfaces)
 															{
-																TraitInfoString.Append(TEXT("- "))
-																	.Append(RequiredInterface.GetInterfaceName())
-																	.AppendChar(TEXT('\n'));
+																if (const ITraitInterface* RequiredInterface = FTraitInterfaceRegistry::Get().Find(RequiredInterfaceUID))
+																{
+																	TraitInfoString.Append(TEXT("- "))
+																		.Append(RequiredInterface->GetDisplayName().ToString())
+																		.AppendChar(TEXT('\n'));
+																}
 															}
 														}
 
@@ -796,6 +803,11 @@ void STraitListView::GenerateTraitList()
 	const TArray<const FTrait*> Traits = TraitRegistry.GetTraits();
 	for (const FTrait* Trait : Traits)
 	{
+		if (Trait->IsHidden() && !TraitEditorSharedData->bAdvancedView)
+		{
+			continue;
+		}
+
 		FName Category = NAME_None;
 		FText CategoryText;
 
diff --git a/Engine/Plugins/Experimental/AnimNext/Source/AnimNextEditor/Private/Graph/TraitEditorDefs.cpp b/Engine/Plugins/Experimental/AnimNext/Source/AnimNextEditor/Private/Graph/TraitEditorDefs.cpp
--- a/Engine/Plugins/Experimental/AnimNext/Source/AnimNextEditor/Private/Graph/TraitEditorDefs.cpp
+++ b/Engine/Plugins/Experimental/AnimNext/Source/AnimNextEditor/Private/Graph/TraitEditorDefs.cpp
@@ -3,6 +3,7 @@
 #include "Graph/TraitEditorDefs.h"
 #include "TraitInterfaces/IGarbageCollection.h"
 #include "Widgets/Layout/SBox.h"
+#include "TraitCore/TraitInterfaceRegistry.h"
 
 #define LOCTEXT_NAMESPACE "TraitListEditor"
 
@@ -85,15 +86,58 @@ namespace UE::AnimNext::Editor
 
 	TSharedRef<SHorizontalBox> InterfaceWidgetsHorizontalBox = SNew(SHorizontalBox);
 
-	const int32 NumSlots = InTraitEditorSharedDataShared->StackUsedInterfaces.Num();
-	for (int32 SlotIndex = 0; SlotIndex < NumSlots; SlotIndex++)
+	switch (InterfaceDisplayType)
 	{
-		InterfaceWidgetsHorizontalBox->AddSlot()
-		.AutoWidth()
-		.Padding(FMargin(4.f, 1.f))
-		[
-			GetInterfaceWidget(InterfaceDisplayType, SlotIndex, InTraitDataShared, InTraitEditorSharedDataShared)
-		];
+		case EInterfaceDisplayType::ListImplemented:
+		{
+			for(const FTraitInterfaceUID& InterfaceUID : InTraitDataShared->ImplementedInterfaces)
+			{
+				// For the trait list just display the interfaces with no spacing, one after another
+				if( !IsInternal(InterfaceUID))
+				{
+					InterfaceWidgetsHorizontalBox->AddSlot()
+						.AutoWidth()
+						.Padding(FMargin(4.f, 1.f))
+						[
+							GetInterfaceWidget(InterfaceDisplayType, InterfaceUID, InTraitDataShared, InTraitEditorSharedDataShared)
+						];
+				}
+			}
+			break;
+		}
+
+		case EInterfaceDisplayType::StackRequired:
+		case EInterfaceDisplayType::StackImplemented:
+		{
+			// For the trait stack display the interfaces aligning them by type, so the same interface is always in the same position
+			const int32 NumSlots = InTraitEditorSharedDataShared->StackUsedInterfaces.Num();
+			for (int32 SlotIndex = 0; SlotIndex < NumSlots; SlotIndex++)
+			{
+				if (SlotIndex < InTraitEditorSharedDataShared->StackUsedInterfaces.Num())
+				{
+					// If the trait contains the slot, display the interface
+					const FTraitInterfaceUID& InterfaceUID = InTraitEditorSharedDataShared->StackUsedInterfaces[SlotIndex];
+
+					InterfaceWidgetsHorizontalBox->AddSlot()
+					.AutoWidth()
+					.Padding(FMargin(4.f, 1.f))
+					[
+						GetInterfaceWidget(InterfaceDisplayType, InterfaceUID, InTraitDataShared, InTraitEditorSharedDataShared)
+					];
+				}
+				else
+				{
+					// else, display an empty space (so we keep the alignment)
+					InterfaceWidgetsHorizontalBox->AddSlot()
+						.AutoWidth()
+						.Padding(FMargin(4.f, 1.f))
+						[
+							SNullWidget::NullWidget
+						];
+				}
+			}
+			break;
+		}
 	}
 
 	return SNew(SBorder)
@@ -108,20 +152,20 @@ namespace UE::AnimNext::Editor
 		];
 }
 
-/*static*/ TSharedRef<SWidget> FTraitEditorUtils::GetInterfaceWidget(EInterfaceDisplayType InterfaceDisplayType, int32 SlotIndex, const TSharedPtr<FTraitDataEditorDef>& InTraitDataShared, const TSharedPtr<FTraitEditorSharedData>& InTraitEditorSharedDataShared)
+/*static*/ TSharedRef<SWidget> FTraitEditorUtils::GetInterfaceWidget(EInterfaceDisplayType InterfaceDisplayType, FTraitInterfaceUID InterfaceUID, const TSharedPtr<FTraitDataEditorDef>& InTraitDataShared, const TSharedPtr<FTraitEditorSharedData>& InTraitEditorSharedDataShared)
 {
 	return SNew(SBorder)
-		.BorderImage_Lambda([InterfaceDisplayType, InTraitDataShared, InTraitEditorSharedDataShared, SlotIndex]()
+		.BorderImage_Lambda([InterfaceDisplayType, InTraitDataShared, InTraitEditorSharedDataShared, InterfaceUID]()
 			{
 				switch (InterfaceDisplayType)
 				{
 					case EInterfaceDisplayType::StackRequired:
 					{
-						if (InTraitEditorSharedDataShared.IsValid() && SlotIndex < InTraitEditorSharedDataShared->StackUsedInterfaces.Num())
+						if (InTraitEditorSharedDataShared.IsValid() && InterfaceUID.IsValid())
 						{
 							if (InTraitDataShared.IsValid())
 							{
-								if (InTraitDataShared->StackStatus.MissingInterfaces.Find(InTraitEditorSharedDataShared->StackUsedInterfaces[SlotIndex]) != INDEX_NONE)
+								if (InTraitDataShared->StackStatus.MissingInterfaces.Find(InterfaceUID) != INDEX_NONE)
 								{
 									return FAppStyle::Get().GetBrush("Brushes.Error");
 								}
@@ -131,11 +175,11 @@ namespace UE::AnimNext::Editor
 					}
 					case EInterfaceDisplayType::ListImplemented:
 					{
-						if (InTraitEditorSharedDataShared.IsValid() && SlotIndex < InTraitEditorSharedDataShared->StackUsedInterfaces.Num())
+						if (InTraitEditorSharedDataShared.IsValid() && InterfaceUID.IsValid())
 						{
 							if (InTraitDataShared.IsValid())
 							{
-								if (InTraitEditorSharedDataShared->StackMissingInterfaces.Find(InTraitEditorSharedDataShared->StackUsedInterfaces[SlotIndex]) != INDEX_NONE)
+								if (InTraitEditorSharedDataShared->StackMissingInterfaces.Find(InterfaceUID) != INDEX_NONE)
 								{
 									return FAppStyle::Get().GetBrush("Brushes.Select");
 								}
@@ -151,16 +195,16 @@ namespace UE::AnimNext::Editor
 
 				return FAppStyle::Get().GetBrush("Brushes.Background");
 			})
-		.Visibility_Lambda([InterfaceDisplayType, InTraitDataShared, InTraitEditorSharedDataShared, SlotIndex]()
+		.Visibility_Lambda([InterfaceDisplayType, InTraitDataShared, InTraitEditorSharedDataShared, InterfaceUID]()
 			{
-				if (InTraitDataShared.IsValid())
+				if (InTraitDataShared.IsValid() && InterfaceUID.IsValid())
 				{
 					switch (InterfaceDisplayType)
 					{
 						case EInterfaceDisplayType::ListImplemented:
 						case EInterfaceDisplayType::StackImplemented:
 						{
-							if (InTraitDataShared->ImplementedInterfacesStackListIndexes.Find(SlotIndex) != INDEX_NONE)
+							if (InTraitDataShared->ImplementedInterfaces.Find(InterfaceUID) != INDEX_NONE)
 							{
 								return EVisibility::Visible;
 							}
@@ -168,7 +212,7 @@ namespace UE::AnimNext::Editor
 						}
 						case EInterfaceDisplayType::StackRequired:
 						{
-							if (InTraitDataShared->RequiredInterfacesStackListIndexes.Find(SlotIndex) != INDEX_NONE)
+							if (InTraitDataShared->RequiredInterfaces.Find(InterfaceUID) != INDEX_NONE)
 							{
 								return EVisibility::Visible;
 							}
@@ -186,29 +230,27 @@ namespace UE::AnimNext::Editor
 			SNew(STextBlock)
 			.Font(FAppStyle::GetFontStyle("PropertyWindow.BoldFont"))
 			.ColorAndOpacity(FSlateColor(FColor::White))
-			.Text_Lambda([InTraitEditorSharedDataShared, SlotIndex]()->FText
+			.Text_Lambda([InTraitEditorSharedDataShared, InterfaceUID]()->FText
 			{
-				if (InTraitEditorSharedDataShared.IsValid())
+				if (InTraitEditorSharedDataShared.IsValid() && InterfaceUID.IsValid())
 				{
-					// TODO zzz : Should the interface names be LOCTEXT in origin?
-					const FString InterfaceShortName(InTraitEditorSharedDataShared->StackUsedInterfaces.IsValidIndex(SlotIndex)
-						? InTraitEditorSharedDataShared->StackUsedInterfaces[SlotIndex].GetInterfaceShortName()
-						: FText(LOCTEXT("MissingInterfaceShortName", "???")).ToString());
-					return FText::FromString(InterfaceShortName);
+					if (const ITraitInterface* TraitInterface = FTraitInterfaceRegistry::Get().Find(InterfaceUID))
+					{
+						return TraitInterface->GetDisplayShortName();
+					}
 				}
-				return FText();
+				return FText(LOCTEXT("MissingInterfaceShortName", "???"));
 			})
-			.ToolTipText_Lambda([InTraitEditorSharedDataShared, SlotIndex]()->FText
+			.ToolTipText_Lambda([InTraitEditorSharedDataShared, InterfaceUID]()->FText
 			{
-				if (InTraitEditorSharedDataShared.IsValid())
+				if (InTraitEditorSharedDataShared.IsValid() && InterfaceUID.IsValid())
 				{
-					// TODO zzz : Should the interface names be LOCTEXT in origin?
-					const FString InterfaceName(InTraitEditorSharedDataShared->StackUsedInterfaces.IsValidIndex(SlotIndex)
-						? InTraitEditorSharedDataShared->StackUsedInterfaces[SlotIndex].GetInterfaceName()
-						: FText(LOCTEXT("MissingInterfaceName", "Invalid or Missing Interface")).ToString());
-					return FText::FromString(InterfaceName);
+					if (const ITraitInterface* TraitInterface = FTraitInterfaceRegistry::Get().Find(InterfaceUID))
+					{
+						return TraitInterface->GetDisplayName();
+					}
 				}
-				return FText();
+				return FText(LOCTEXT("MissingInterfaceName", "Invalid or Missing Interface"));
 			})
 		];
 }
@@ -259,13 +301,14 @@ namespace UE::AnimNext::Editor
 	return nullptr;
 }
 
-/*static*/ bool FTraitEditorUtils::IsInternalTraitInteface(const FTraitInterfaceUID& InTraitInterfaceUID) 
+/*static*/ bool FTraitEditorUtils::IsInternal(const FTraitInterfaceUID& InTraitInterfaceUID) 
 {
-	switch (InTraitInterfaceUID.GetUID())
+	if (const ITraitInterface* ImplementedInterface = FTraitInterfaceRegistry::Get().Find(InTraitInterfaceUID))
 	{
-		case IGarbageCollection::InterfaceUID.GetUID():		return true;
-		default:											return false;
+		return ImplementedInterface->IsInternal();
 	}
+
+	return false;
 }
 
 
diff --git a/Engine/Plugins/Experimental/AnimNext/Source/AnimNextEditor/Private/Graph/TraitEditorDefs.h b/Engine/Plugins/Experimental/AnimNext/Source/AnimNextEditor/Private/Graph/TraitEditorDefs.h
--- a/Engine/Plugins/Experimental/AnimNext/Source/AnimNextEditor/Private/Graph/TraitEditorDefs.h
+++ b/Engine/Plugins/Experimental/AnimNext/Source/AnimNextEditor/Private/Graph/TraitEditorDefs.h
@@ -134,6 +134,7 @@ struct FTraitEditorSharedData
 	bool bStackContainsErrors = false;
 	bool bShowTraitInterfaces = false;
 	bool bShowTraitInterfacesIfWarningsOrErrors = false;
+	bool bAdvancedView = false;
 };
 
 struct FTraitEditorUtils
@@ -149,14 +150,13 @@ struct FTraitEditorUtils
 		StackRequired
 	};
 	static TSharedRef<SWidget> GetInterfaceListWidget(EInterfaceDisplayType InterfaceDisplayType, const TSharedPtr<FTraitDataEditorDef>& InTraitDataShared, const TSharedPtr<FTraitEditorSharedData>& InTraitEditorSharedDataShared);
-	static TSharedRef<SWidget> GetInterfaceWidget(EInterfaceDisplayType InterfaceDisplayType, int32 SlotIndex, const TSharedPtr<FTraitDataEditorDef>& InTraitDataShared, const TSharedPtr<FTraitEditorSharedData>& InTraitEditorSharedDataShared);
+	static TSharedRef<SWidget> GetInterfaceWidget(EInterfaceDisplayType InterfaceDisplayType, FTraitInterfaceUID InterfaceUID, const TSharedPtr<FTraitDataEditorDef>& InTraitDataShared, const TSharedPtr<FTraitEditorSharedData>& InTraitEditorSharedDataShared);
 	static void GenerateStackInterfacesUsedIndexes(TSharedPtr<FTraitDataEditorDef>& TraitData, const TSharedPtr<FTraitEditorSharedData>& InTraitEditorSharedData);
 
 	static TSharedPtr<FTraitDataEditorDef> FindTraitInCurrentStackData(const FTraitUID InTraitUID, TSharedPtr<TArray<TSharedPtr<FTraitDataEditorDef>>> CurrentTraitsDataShared, int32* OutIndex = nullptr);
 
 	// Internal Interfaces are not shown in the Traits Editor
-	// TODO zzz : This should be in a TraitInterfaceRegistry
-	static bool IsInternalTraitInteface(const FTraitInterfaceUID& InterfaceUID);
+	static bool IsInternal(const FTraitInterfaceUID& InterfaceUID);
 };
 
 // --- FTraitListDragDropBase ---
diff --git a/Engine/Plugins/Experimental/AnimNext/Source/AnimNextTestSuite/Private/AnimNextTraitBaseTest.cpp b/Engine/Plugins/Experimental/AnimNext/Source/AnimNextTestSuite/Private/AnimNextTraitBaseTest.cpp
--- a/Engine/Plugins/Experimental/AnimNext/Source/AnimNextTestSuite/Private/AnimNextTraitBaseTest.cpp
+++ b/Engine/Plugins/Experimental/AnimNext/Source/AnimNextTestSuite/Private/AnimNextTraitBaseTest.cpp
@@ -36,7 +36,7 @@ namespace UE::AnimNext
 
 	struct IInterfaceA : ITraitInterface
 	{
-		DECLARE_ANIM_TRAIT_INTERFACE(IInterfaceA, ITA, 0x34cb8e62)
+		DECLARE_ANIM_TRAIT_INTERFACE(IInterfaceA, 0x34cb8e62)
 
 		virtual void FuncA(const FExecutionContext& Context, const TTraitBinding<IInterfaceA>& Binding) const;
 	};
@@ -66,7 +66,7 @@ namespace UE::AnimNext
 
 	struct IInterfaceB : ITraitInterface
 	{
-		DECLARE_ANIM_TRAIT_INTERFACE(IInterfaceB, ITB, 0x33cb8ccf)
+		DECLARE_ANIM_TRAIT_INTERFACE(IInterfaceB, 0x33cb8ccf)
 
 		virtual void FuncB(const FExecutionContext& Context, const TTraitBinding<IInterfaceB>& Binding) const;
 	};
@@ -96,7 +96,7 @@ namespace UE::AnimNext
 
 	struct IInterfaceC : ITraitInterface
 	{
-		DECLARE_ANIM_TRAIT_INTERFACE(IInterfaceC, ITC, 0x32cb8b3c)
+		DECLARE_ANIM_TRAIT_INTERFACE(IInterfaceC, 0x32cb8b3c)
 
 		virtual void FuncC(const FExecutionContext& Context, const TTraitBinding<IInterfaceC>& Binding) const;
 	};
diff --git a/Engine/Plugins/Experimental/AnimNextAnimGraph/Source/AnimNextAnimGraph/Private/TraitInterfaces/IContinuousBlend.cpp b/Engine/Plugins/Experimental/AnimNextAnimGraph/Source/AnimNextAnimGraph/Private/TraitInterfaces/IContinuousBlend.cpp
--- a/Engine/Plugins/Experimental/AnimNextAnimGraph/Source/AnimNextAnimGraph/Private/TraitInterfaces/IContinuousBlend.cpp
+++ b/Engine/Plugins/Experimental/AnimNextAnimGraph/Source/AnimNextAnimGraph/Private/TraitInterfaces/IContinuousBlend.cpp
@@ -6,6 +6,21 @@
 
 namespace UE::AnimNext
 {
+	AUTO_REGISTER_ANIM_TRAIT_INTERFACE(IContinuousBlend)
+
+#if WITH_EDITOR
+	const FText& IContinuousBlend::GetDisplayName() const
+	{
+		static FText InterfaceName = NSLOCTEXT("TraitInterfaces", "TraitInterface_ContinuousBlend_Name", "Continuous Blend");
+		return InterfaceName;
+	}  
+	const FText& IContinuousBlend::GetDisplayShortName() const
+	{
+		static FText InterfaceShortName = NSLOCTEXT("TraitInterfaces", "TraitInterface_ContinuousBlend_ShortName", "CNB");
+		return InterfaceShortName;
+	}
+#endif // WITH_EDITOR
+
 	float IContinuousBlend::GetBlendWeight(const FExecutionContext& Context, const TTraitBinding<IContinuousBlend>& Binding, int32 ChildIndex) const
 	{
 		TTraitBinding<IContinuousBlend> SuperBinding;
diff --git a/Engine/Plugins/Experimental/AnimNextAnimGraph/Source/AnimNextAnimGraph/Private/TraitInterfaces/IDiscreteBlend.cpp b/Engine/Plugins/Experimental/AnimNextAnimGraph/Source/AnimNextAnimGraph/Private/TraitInterfaces/IDiscreteBlend.cpp
--- a/Engine/Plugins/Experimental/AnimNextAnimGraph/Source/AnimNextAnimGraph/Private/TraitInterfaces/IDiscreteBlend.cpp
+++ b/Engine/Plugins/Experimental/AnimNextAnimGraph/Source/AnimNextAnimGraph/Private/TraitInterfaces/IDiscreteBlend.cpp
@@ -6,6 +6,21 @@
 
 namespace UE::AnimNext
 {
+	AUTO_REGISTER_ANIM_TRAIT_INTERFACE(IDiscreteBlend)
+
+#if WITH_EDITOR
+	const FText& IDiscreteBlend::GetDisplayName() const
+	{
+		static FText InterfaceName = NSLOCTEXT("TraitInterfaces", "TraitInterface_DiscreteBlend_Name", "Discrete Blend");
+		return InterfaceName;
+	}
+	const FText& IDiscreteBlend::GetDisplayShortName() const
+	{
+		static FText InterfaceShortName = NSLOCTEXT("TraitInterfaces", "TraitInterface_DiscreteBlend_ShortName", "DSB");
+		return InterfaceShortName;
+	}
+#endif // WITH_EDITOR
+
 	float IDiscreteBlend::GetBlendWeight(FExecutionContext& Context, const TTraitBinding<IDiscreteBlend>& Binding, int32 ChildIndex) const
 	{
 		TTraitBinding<IDiscreteBlend> SuperBinding;
diff --git a/Engine/Plugins/Experimental/AnimNextAnimGraph/Source/AnimNextAnimGraph/Private/TraitInterfaces/IEvaluate.cpp b/Engine/Plugins/Experimental/AnimNextAnimGraph/Source/AnimNextAnimGraph/Private/TraitInterfaces/IEvaluate.cpp
--- a/Engine/Plugins/Experimental/AnimNextAnimGraph/Source/AnimNextAnimGraph/Private/TraitInterfaces/IEvaluate.cpp
+++ b/Engine/Plugins/Experimental/AnimNextAnimGraph/Source/AnimNextAnimGraph/Private/TraitInterfaces/IEvaluate.cpp
@@ -15,6 +15,21 @@ namespace UE::AnimNext
 	{
 	}
 
+	AUTO_REGISTER_ANIM_TRAIT_INTERFACE(IEvaluate)
+
+#if WITH_EDITOR
+	const FText& IEvaluate::GetDisplayName() const
+	{
+		static FText InterfaceName = NSLOCTEXT("TraitInterfaces", "TraitInterface_Evaluate_Name", "Evaluate");
+		return InterfaceName;
+	}
+	const FText& IEvaluate::GetDisplayShortName() const
+	{
+		static FText InterfaceShortName = NSLOCTEXT("TraitInterfaces", "TraitInterface_Evaluate_ShortName", "EVA");
+		return InterfaceShortName;
+	}
+#endif // WITH_EDITOR
+
 	void IEvaluate::PreEvaluate(FEvaluateTraversalContext& Context, const TTraitBinding<IEvaluate>& Binding) const
 	{
 		TTraitBinding<IEvaluate> SuperBinding;
diff --git a/Engine/Plugins/Experimental/AnimNextAnimGraph/Source/AnimNextAnimGraph/Private/TraitInterfaces/IGroupSynchronization.cpp b/Engine/Plugins/Experimental/AnimNextAnimGraph/Source/AnimNextAnimGraph/Private/TraitInterfaces/IGroupSynchronization.cpp
--- a/Engine/Plugins/Experimental/AnimNextAnimGraph/Source/AnimNextAnimGraph/Private/TraitInterfaces/IGroupSynchronization.cpp
+++ b/Engine/Plugins/Experimental/AnimNextAnimGraph/Source/AnimNextAnimGraph/Private/TraitInterfaces/IGroupSynchronization.cpp
@@ -6,6 +6,21 @@
 
 namespace UE::AnimNext
 {
+	AUTO_REGISTER_ANIM_TRAIT_INTERFACE(IGroupSynchronization)
+
+#if WITH_EDITOR
+	const FText& IGroupSynchronization::GetDisplayName() const
+	{
+		static FText InterfaceName = NSLOCTEXT("TraitInterfaces", "TraitInterface_IGroupSynchronization_Name", "Group Synchronization");
+		return InterfaceName;
+	}
+	const FText& IGroupSynchronization::GetDisplayShortName() const
+	{
+		static FText InterfaceShortName = NSLOCTEXT("TraitInterfaces", "TraitInterface_IGroupSynchronization_ShortName", "GRS");
+		return InterfaceShortName;
+	}
+#endif // WITH_EDITOR
+
 	FName IGroupSynchronization::GetGroupName(FExecutionContext& Context, const TTraitBinding<IGroupSynchronization>& Binding) const
 	{
 		TTraitBinding<IGroupSynchronization> SuperBinding;
diff --git a/Engine/Plugins/Experimental/AnimNextAnimGraph/Source/AnimNextAnimGraph/Private/TraitInterfaces/IInertializerBlend.cpp b/Engine/Plugins/Experimental/AnimNextAnimGraph/Source/AnimNextAnimGraph/Private/TraitInterfaces/IInertializerBlend.cpp
--- a/Engine/Plugins/Experimental/AnimNextAnimGraph/Source/AnimNextAnimGraph/Private/TraitInterfaces/IInertializerBlend.cpp
+++ b/Engine/Plugins/Experimental/AnimNextAnimGraph/Source/AnimNextAnimGraph/Private/TraitInterfaces/IInertializerBlend.cpp
@@ -6,6 +6,21 @@
 
 namespace UE::AnimNext
 {
+	AUTO_REGISTER_ANIM_TRAIT_INTERFACE(IInertializerBlend)
+
+#if WITH_EDITOR
+	const FText& IInertializerBlend::GetDisplayName() const
+	{
+		static FText InterfaceName = NSLOCTEXT("TraitInterfaces", "TraitInterface_IInertializerBlend_Name", "Inertializer Blend");
+		return InterfaceName;
+	}
+	const FText& IInertializerBlend::GetDisplayShortName() const
+	{
+		static FText InterfaceShortName = NSLOCTEXT("TraitInterfaces", "TraitInterface_IInertializerBlend_ShortName", "IB");
+		return InterfaceShortName;
+	}
+#endif // WITH_EDITOR
+
 	float IInertializerBlend::GetBlendTime(FExecutionContext& Context, const TTraitBinding<IInertializerBlend>& Binding, int32 ChildIndex) const
 	{
 		TTraitBinding<IInertializerBlend> SuperBinding;
diff --git a/Engine/Plugins/Experimental/AnimNextAnimGraph/Source/AnimNextAnimGraph/Private/TraitInterfaces/ISmoothBlend.cpp b/Engine/Plugins/Experimental/AnimNextAnimGraph/Source/AnimNextAnimGraph/Private/TraitInterfaces/ISmoothBlend.cpp
--- a/Engine/Plugins/Experimental/AnimNextAnimGraph/Source/AnimNextAnimGraph/Private/TraitInterfaces/ISmoothBlend.cpp
+++ b/Engine/Plugins/Experimental/AnimNextAnimGraph/Source/AnimNextAnimGraph/Private/TraitInterfaces/ISmoothBlend.cpp
@@ -6,6 +6,21 @@
 
 namespace UE::AnimNext
 {
+	AUTO_REGISTER_ANIM_TRAIT_INTERFACE(ISmoothBlend)
+
+#if WITH_EDITOR
+	const FText& ISmoothBlend::GetDisplayName() const
+	{
+		static FText InterfaceName = NSLOCTEXT("TraitInterfaces", "TraitInterface_ISmoothBlend_Name", "Group Synchronization");
+		return InterfaceName;
+	}
+	const FText& ISmoothBlend::GetDisplayShortName() const
+	{
+		static FText InterfaceShortName = NSLOCTEXT("TraitInterfaces", "TraitInterface_ISmoothBlend_ShortName", "SMB");
+		return InterfaceShortName;
+	}
+#endif // WITH_EDITOR
+
 	float ISmoothBlend::GetBlendTime(FExecutionContext& Context, const TTraitBinding<ISmoothBlend>& Binding, int32 ChildIndex) const
 	{
 		TTraitBinding<ISmoothBlend> SuperBinding;
diff --git a/Engine/Plugins/Experimental/AnimNextAnimGraph/Source/AnimNextAnimGraph/Private/TraitInterfaces/ITimeline.cpp b/Engine/Plugins/Experimental/AnimNextAnimGraph/Source/AnimNextAnimGraph/Private/TraitInterfaces/ITimeline.cpp
--- a/Engine/Plugins/Experimental/AnimNextAnimGraph/Source/AnimNextAnimGraph/Private/TraitInterfaces/ITimeline.cpp
+++ b/Engine/Plugins/Experimental/AnimNextAnimGraph/Source/AnimNextAnimGraph/Private/TraitInterfaces/ITimeline.cpp
@@ -6,6 +6,21 @@
 
 namespace UE::AnimNext
 {
+	AUTO_REGISTER_ANIM_TRAIT_INTERFACE(ITimeline)
+
+#if WITH_EDITOR
+	const FText& ITimeline::GetDisplayName() const
+	{
+		static FText InterfaceName = NSLOCTEXT("TraitInterfaces", "TraitInterface_ITimeline_Name", "Timeline");
+		return InterfaceName;
+	}
+	const FText& ITimeline::GetDisplayShortName() const
+	{
+		static FText InterfaceShortName = NSLOCTEXT("TraitInterfaces", "TraitInterface_ITimeline_ShortName", "TIM");
+		return InterfaceShortName;
+	}
+#endif // WITH_EDITOR
+
 	float ITimeline::GetPlayRate(FExecutionContext& Context, const TTraitBinding<ITimeline>& Binding) const
 	{
 		TTraitBinding<ITimeline> SuperBinding;
diff --git a/Engine/Plugins/Experimental/AnimNextAnimGraph/Source/AnimNextAnimGraph/Private/TraitInterfaces/IUpdate.cpp b/Engine/Plugins/Experimental/AnimNextAnimGraph/Source/AnimNextAnimGraph/Private/TraitInterfaces/IUpdate.cpp
--- a/Engine/Plugins/Experimental/AnimNextAnimGraph/Source/AnimNextAnimGraph/Private/TraitInterfaces/IUpdate.cpp
+++ b/Engine/Plugins/Experimental/AnimNextAnimGraph/Source/AnimNextAnimGraph/Private/TraitInterfaces/IUpdate.cpp
@@ -157,6 +157,21 @@ namespace UE::AnimNext
 		}
 	}
 
+	AUTO_REGISTER_ANIM_TRAIT_INTERFACE(IUpdate)
+
+#if WITH_EDITOR
+	const FText& IUpdate::GetDisplayName() const
+	{
+		static FText InterfaceName = NSLOCTEXT("TraitInterfaces", "TraitInterface_IUpdate_Name", "Update");
+		return InterfaceName;
+	}
+	const FText& IUpdate::GetDisplayShortName() const
+	{
+		static FText InterfaceShortName = NSLOCTEXT("TraitInterfaces", "TraitInterface_IUpdate_ShortName", "UPD");
+		return InterfaceShortName;
+	}
+#endif // WITH_EDITOR
+
 	void IUpdate::OnBecomeRelevant(FUpdateTraversalContext& Context, const TTraitBinding<IUpdate>& Binding, const FTraitUpdateState& TraitState) const
 	{
 		TTraitBinding<IUpdate> SuperBinding;
@@ -184,6 +199,21 @@ namespace UE::AnimNext
 		}
 	}
 
+	AUTO_REGISTER_ANIM_TRAIT_INTERFACE(IUpdateTraversal)
+
+#if WITH_EDITOR
+	const FText& IUpdateTraversal::GetDisplayName() const
+	{
+		static FText InterfaceName = NSLOCTEXT("TraitInterfaces", "TraitInterface_IUpdateTraversal_Name", "Update Traversal");
+		return InterfaceName;
+	}
+	const FText& IUpdateTraversal::GetDisplayShortName() const
+	{
+		static FText InterfaceShortName = NSLOCTEXT("TraitInterfaces", "TraitInterface_IUpdateTraversal_ShortName", "TRA");
+		return InterfaceShortName;
+	}
+#endif // WITH_EDITOR
+
 	void IUpdateTraversal::QueueChildrenForTraversal(FUpdateTraversalContext& Context, const TTraitBinding<IUpdateTraversal>& Binding, const FTraitUpdateState& TraitState, FUpdateTraversalQueue& TraversalQueue) const
 	{
 		// Nothing to do
diff --git a/Engine/Plugins/Experimental/AnimNextAnimGraph/Source/AnimNextAnimGraph/Private/Traits/BlendInertializer.h b/Engine/Plugins/Experimental/AnimNextAnimGraph/Source/AnimNextAnimGraph/Private/Traits/BlendInertializer.h
--- a/Engine/Plugins/Experimental/AnimNextAnimGraph/Source/AnimNextAnimGraph/Private/Traits/BlendInertializer.h
+++ b/Engine/Plugins/Experimental/AnimNextAnimGraph/Source/AnimNextAnimGraph/Private/Traits/BlendInertializer.h
@@ -48,6 +48,10 @@ namespace UE::AnimNext
 
 		// ISmoothBlend impl
 		virtual float GetBlendTime(FExecutionContext& Context, const TTraitBinding<ISmoothBlend>& Binding, int32 ChildIndex) const override;
+
+#if WITH_EDITOR
+		virtual bool IsHidden() const override { return true; }
+#endif
 	};
 
 	/**
@@ -67,5 +71,9 @@ namespace UE::AnimNext
 
 		// IInertializerBlend impl
 		virtual float GetBlendTime(FExecutionContext& Context, const TTraitBinding<IInertializerBlend>& Binding, int32 ChildIndex) const override;
+
+#if WITH_EDITOR
+		virtual bool IsHidden() const override { return false; } // overrdie the base, as it is hidden
+#endif
 	};
 }
diff --git a/Engine/Plugins/Experimental/AnimNextAnimGraph/Source/AnimNextAnimGraph/Private/Traits/BlendSmoother.h b/Engine/Plugins/Experimental/AnimNextAnimGraph/Source/AnimNextAnimGraph/Private/Traits/BlendSmoother.h
--- a/Engine/Plugins/Experimental/AnimNextAnimGraph/Source/AnimNextAnimGraph/Private/Traits/BlendSmoother.h
+++ b/Engine/Plugins/Experimental/AnimNextAnimGraph/Source/AnimNextAnimGraph/Private/Traits/BlendSmoother.h
@@ -87,6 +87,10 @@ namespace UE::AnimNext
 
 		// Internal impl
 		static void InitializeInstanceData(FExecutionContext& Context, const FTraitBinding& Binding, const FSharedData* SharedData, FInstanceData* InstanceData);
+
+#if WITH_EDITOR
+		virtual bool IsHidden() const override { return true; }
+#endif
 	};
 
 	/**
@@ -105,5 +109,9 @@ namespace UE::AnimNext
 		virtual float GetBlendTime(FExecutionContext& Context, const TTraitBinding<ISmoothBlend>& Binding, int32 ChildIndex) const override;
 		virtual EAlphaBlendOption GetBlendType(FExecutionContext& Context, const TTraitBinding<ISmoothBlend>& Binding, int32 ChildIndex) const override;
 		virtual UCurveFloat* GetCustomBlendCurve(FExecutionContext& Context, const TTraitBinding<ISmoothBlend>& Binding, int32 ChildIndex) const override;
+
+#if WITH_EDITOR
+		virtual bool IsHidden() const override { return false; } // overrdie the base, as it is hidden
+#endif
 	};
 }
diff --git a/Engine/Plugins/Experimental/AnimNextAnimGraph/Source/AnimNextAnimGraph/Public/TraitInterfaces/IContinuousBlend.h b/Engine/Plugins/Experimental/AnimNextAnimGraph/Source/AnimNextAnimGraph/Public/TraitInterfaces/IContinuousBlend.h
--- a/Engine/Plugins/Experimental/AnimNextAnimGraph/Source/AnimNextAnimGraph/Public/TraitInterfaces/IContinuousBlend.h
+++ b/Engine/Plugins/Experimental/AnimNextAnimGraph/Source/AnimNextAnimGraph/Public/TraitInterfaces/IContinuousBlend.h
@@ -15,12 +15,17 @@ namespace UE::AnimNext
 	 */
 	struct ANIMNEXTANIMGRAPH_API IContinuousBlend : ITraitInterface
 	{
-		DECLARE_ANIM_TRAIT_INTERFACE(IContinuousBlend, CNB, 0xe7d79186)
+		DECLARE_ANIM_TRAIT_INTERFACE(IContinuousBlend, 0xe7d79186)
 
 		// Returns the blend weight for the specified child
 		// Multiple children can have non-zero weight but their sum must be 1.0
 		// Returns -1.0 if the child index is invalid
 		virtual float GetBlendWeight(const FExecutionContext& Context, const TTraitBinding<IContinuousBlend>& Binding, int32 ChildIndex) const;
+
+#if WITH_EDITOR
+		virtual const FText& GetDisplayName() const override;
+		virtual const FText& GetDisplayShortName() const override;
+#endif // WITH_EDITOR
 	};
 
 	/**
diff --git a/Engine/Plugins/Experimental/AnimNextAnimGraph/Source/AnimNextAnimGraph/Public/TraitInterfaces/IDiscreteBlend.h b/Engine/Plugins/Experimental/AnimNextAnimGraph/Source/AnimNextAnimGraph/Public/TraitInterfaces/IDiscreteBlend.h
--- a/Engine/Plugins/Experimental/AnimNextAnimGraph/Source/AnimNextAnimGraph/Public/TraitInterfaces/IDiscreteBlend.h
+++ b/Engine/Plugins/Experimental/AnimNextAnimGraph/Source/AnimNextAnimGraph/Public/TraitInterfaces/IDiscreteBlend.h
@@ -17,7 +17,7 @@ namespace UE::AnimNext
 	 */
 	struct ANIMNEXTANIMGRAPH_API IDiscreteBlend : ITraitInterface
 	{
-		DECLARE_ANIM_TRAIT_INTERFACE(IDiscreteBlend, DSB, 0x2d395d56)
+		DECLARE_ANIM_TRAIT_INTERFACE(IDiscreteBlend, 0x2d395d56)
 
 		// Returns the blend weight for the specified child
 		// Multiple children can have non-zero weight but their sum must be 1.0
@@ -45,6 +45,11 @@ namespace UE::AnimNext
 
 		// Called when the blend for specified child terminates
 		virtual void OnBlendTerminated(FExecutionContext& Context, const TTraitBinding<IDiscreteBlend>& Binding, int32 ChildIndex) const;
+
+#if WITH_EDITOR
+		virtual const FText& GetDisplayName() const override;
+		virtual const FText& GetDisplayShortName() const override;
+#endif // WITH_EDITOR
 	};
 
 	/**
diff --git a/Engine/Plugins/Experimental/AnimNextAnimGraph/Source/AnimNextAnimGraph/Public/TraitInterfaces/IEvaluate.h b/Engine/Plugins/Experimental/AnimNextAnimGraph/Source/AnimNextAnimGraph/Public/TraitInterfaces/IEvaluate.h
--- a/Engine/Plugins/Experimental/AnimNextAnimGraph/Source/AnimNextAnimGraph/Public/TraitInterfaces/IEvaluate.h
+++ b/Engine/Plugins/Experimental/AnimNextAnimGraph/Source/AnimNextAnimGraph/Public/TraitInterfaces/IEvaluate.h
@@ -49,13 +49,18 @@ namespace UE::AnimNext
 	 */
 	struct ANIMNEXTANIMGRAPH_API IEvaluate : ITraitInterface
 	{
-		DECLARE_ANIM_TRAIT_INTERFACE(IEvaluate, EVA, 0xa303e9e7)
+		DECLARE_ANIM_TRAIT_INTERFACE(IEvaluate, 0xa303e9e7)
 
 		// Called before a traits children are evaluated
 		virtual void PreEvaluate(FEvaluateTraversalContext& Context, const TTraitBinding<IEvaluate>& Binding) const;
 
 		// Called after a traits children have been evaluated
 		virtual void PostEvaluate(FEvaluateTraversalContext& Context, const TTraitBinding<IEvaluate>& Binding) const;
+
+#if WITH_EDITOR
+		virtual const FText& GetDisplayName() const override;
+		virtual const FText& GetDisplayShortName() const override;
+#endif // WITH_EDITOR
 	};
 
 	/**
diff --git a/Engine/Plugins/Experimental/AnimNextAnimGraph/Source/AnimNextAnimGraph/Public/TraitInterfaces/IGroupSynchronization.h b/Engine/Plugins/Experimental/AnimNextAnimGraph/Source/AnimNextAnimGraph/Public/TraitInterfaces/IGroupSynchronization.h
--- a/Engine/Plugins/Experimental/AnimNextAnimGraph/Source/AnimNextAnimGraph/Public/TraitInterfaces/IGroupSynchronization.h
+++ b/Engine/Plugins/Experimental/AnimNextAnimGraph/Source/AnimNextAnimGraph/Public/TraitInterfaces/IGroupSynchronization.h
@@ -17,7 +17,7 @@ namespace UE::AnimNext
 	 */
 	struct ANIMNEXTANIMGRAPH_API IGroupSynchronization : ITraitInterface
 	{
-		DECLARE_ANIM_TRAIT_INTERFACE(IGroupSynchronization, GRS, 0xf607d0fd)
+		DECLARE_ANIM_TRAIT_INTERFACE(IGroupSynchronization, 0xf607d0fd)
 
 		// Returns the group name used for synchronization
 		virtual FName GetGroupName(FExecutionContext& Context, const TTraitBinding<IGroupSynchronization>& Binding) const;
@@ -32,6 +32,11 @@ namespace UE::AnimNext
 		// Called by the sync group graph instance component once a group has been synchronized to advance time on each follower
 		// Progress ratio must be between [0.0, 1.0]
 		virtual void AdvanceToRatio(FExecutionContext& Context, const TTraitBinding<IGroupSynchronization>& Binding, float ProgressRatio) const;
+	
+#if WITH_EDITOR
+		virtual const FText& GetDisplayName() const override;
+		virtual const FText& GetDisplayShortName() const override;
+#endif // WITH_EDITOR
 	};
 
 	/**
diff --git a/Engine/Plugins/Experimental/AnimNextAnimGraph/Source/AnimNextAnimGraph/Public/TraitInterfaces/IInertializerBlend.h b/Engine/Plugins/Experimental/AnimNextAnimGraph/Source/AnimNextAnimGraph/Public/TraitInterfaces/IInertializerBlend.h
--- a/Engine/Plugins/Experimental/AnimNextAnimGraph/Source/AnimNextAnimGraph/Public/TraitInterfaces/IInertializerBlend.h
+++ b/Engine/Plugins/Experimental/AnimNextAnimGraph/Source/AnimNextAnimGraph/Public/TraitInterfaces/IInertializerBlend.h
@@ -15,10 +15,15 @@ namespace UE::AnimNext
 	 */
 	struct ANIMNEXTANIMGRAPH_API IInertializerBlend : ITraitInterface
 	{
-		DECLARE_ANIM_TRAIT_INTERFACE(IInertializerBlend, IB, 0x3856b8e9)
+		DECLARE_ANIM_TRAIT_INTERFACE(IInertializerBlend, 0x3856b8e9)
 
 		// Returns the desired blend time for the specified child
 		virtual float GetBlendTime(FExecutionContext& Context, const TTraitBinding<IInertializerBlend>& Binding, int32 ChildIndex) const;
+
+#if WITH_EDITOR
+		virtual const FText& GetDisplayName() const override;
+		virtual const FText& GetDisplayShortName() const override;
+#endif // WITH_EDITOR
 	};
 
 	/**
diff --git a/Engine/Plugins/Experimental/AnimNextAnimGraph/Source/AnimNextAnimGraph/Public/TraitInterfaces/ISmoothBlend.h b/Engine/Plugins/Experimental/AnimNextAnimGraph/Source/AnimNextAnimGraph/Public/TraitInterfaces/ISmoothBlend.h
--- a/Engine/Plugins/Experimental/AnimNextAnimGraph/Source/AnimNextAnimGraph/Public/TraitInterfaces/ISmoothBlend.h
+++ b/Engine/Plugins/Experimental/AnimNextAnimGraph/Source/AnimNextAnimGraph/Public/TraitInterfaces/ISmoothBlend.h
@@ -18,7 +18,7 @@ namespace UE::AnimNext
 	 */
 	struct ANIMNEXTANIMGRAPH_API ISmoothBlend : ITraitInterface
 	{
-		DECLARE_ANIM_TRAIT_INTERFACE(ISmoothBlend, SMB, 0x1c2c1739)
+		DECLARE_ANIM_TRAIT_INTERFACE(ISmoothBlend, 0x1c2c1739)
 
 		// Returns the desired blend time for the specified child
 		virtual float GetBlendTime(FExecutionContext& Context, const TTraitBinding<ISmoothBlend>& Binding, int32 ChildIndex) const;
@@ -28,6 +28,11 @@ namespace UE::AnimNext
 
 		// Returns the desired blend curve for the specified child
 		virtual UCurveFloat* GetCustomBlendCurve(FExecutionContext& Context, const TTraitBinding<ISmoothBlend>& Binding, int32 ChildIndex) const;
+
+#if WITH_EDITOR
+		virtual const FText& GetDisplayName() const override;
+		virtual const FText& GetDisplayShortName() const override;
+#endif // WITH_EDITOR
 	};
 
 	/**
diff --git a/Engine/Plugins/Experimental/AnimNextAnimGraph/Source/AnimNextAnimGraph/Public/TraitInterfaces/ITimeline.h b/Engine/Plugins/Experimental/AnimNextAnimGraph/Source/AnimNextAnimGraph/Public/TraitInterfaces/ITimeline.h
--- a/Engine/Plugins/Experimental/AnimNextAnimGraph/Source/AnimNextAnimGraph/Public/TraitInterfaces/ITimeline.h
+++ b/Engine/Plugins/Experimental/AnimNextAnimGraph/Source/AnimNextAnimGraph/Public/TraitInterfaces/ITimeline.h
@@ -58,7 +58,7 @@ namespace UE::AnimNext
 	 */
 	struct ANIMNEXTANIMGRAPH_API ITimeline : ITraitInterface
 	{
-		DECLARE_ANIM_TRAIT_INTERFACE(ITimeline, TIM, 0x53760727)
+		DECLARE_ANIM_TRAIT_INTERFACE(ITimeline, 0x53760727)
 
 		// Returns the play rate of this timeline
 		virtual float GetPlayRate(FExecutionContext& Context, const TTraitBinding<ITimeline>& Binding) const;
@@ -77,6 +77,11 @@ namespace UE::AnimNext
 		// Advances time to the specified progress ratio on this timeline
 		// Progress ratio must be between [0.0, 1.0]
 		virtual void AdvanceToRatio(FExecutionContext& Context, const TTraitBinding<ITimeline>& Binding, float ProgressRatio) const;
+
+#if WITH_EDITOR
+		virtual const FText& GetDisplayName() const override;
+		virtual const FText& GetDisplayShortName() const override;
+#endif // WITH_EDITOR
 	};
 
 	/**
diff --git a/Engine/Plugins/Experimental/AnimNextAnimGraph/Source/AnimNextAnimGraph/Public/TraitInterfaces/IUpdate.h b/Engine/Plugins/Experimental/AnimNextAnimGraph/Source/AnimNextAnimGraph/Public/TraitInterfaces/IUpdate.h
--- a/Engine/Plugins/Experimental/AnimNextAnimGraph/Source/AnimNextAnimGraph/Public/TraitInterfaces/IUpdate.h
+++ b/Engine/Plugins/Experimental/AnimNextAnimGraph/Source/AnimNextAnimGraph/Public/TraitInterfaces/IUpdate.h
@@ -229,7 +229,7 @@ namespace UE::AnimNext
 	 */
 	struct ANIMNEXTANIMGRAPH_API IUpdate : ITraitInterface
 	{
-		DECLARE_ANIM_TRAIT_INTERFACE(IUpdate, UPD, 0x59d24dc5)
+		DECLARE_ANIM_TRAIT_INTERFACE(IUpdate, 0x59d24dc5)
 
 		// Called before the first update when a trait stack becomes relevant
 		virtual void OnBecomeRelevant(FUpdateTraversalContext& Context, const TTraitBinding<IUpdate>& Binding, const FTraitUpdateState& TraitState) const;
@@ -239,6 +239,11 @@ namespace UE::AnimNext
 
 		// Called after a traits children have been updated
 		virtual void PostUpdate(FUpdateTraversalContext& Context, const TTraitBinding<IUpdate>& Binding, const FTraitUpdateState& TraitState) const;
+	
+#if WITH_EDITOR
+		virtual const FText& GetDisplayName() const override;
+		virtual const FText& GetDisplayShortName() const override;
+#endif // WITH_EDITOR
 	};
 
 	/**
@@ -255,10 +260,15 @@ namespace UE::AnimNext
 	 */
 	struct ANIMNEXTANIMGRAPH_API IUpdateTraversal : ITraitInterface
 	{
-		DECLARE_ANIM_TRAIT_INTERFACE(IUpdateTraversal, TRA, 0x256c21b1)
+		DECLARE_ANIM_TRAIT_INTERFACE(IUpdateTraversal, 0x256c21b1)
 
 		// Called after PreUpdate to request that children be queued with the provided context
 		virtual void QueueChildrenForTraversal(FUpdateTraversalContext& Context, const TTraitBinding<IUpdateTraversal>& Binding, const FTraitUpdateState& TraitState, FUpdateTraversalQueue& TraversalQueue) const;
+
+#if WITH_EDITOR
+		virtual const FText& GetDisplayName() const override;
+		virtual const FText& GetDisplayShortName() const override;
+#endif // WITH_EDITOR
 	};
 
 	/**
diff --git a/Engine/Plugins/Experimental/AnimNextAnimGraph/Source/AnimNextAnimGraphTestSuite/Private/AnimNextAnimGraphTraitInterfacesTest.cpp b/Engine/Plugins/Experimental/AnimNextAnimGraph/Source/AnimNextAnimGraphTestSuite/Private/AnimNextAnimGraphTraitInterfacesTest.cpp
--- a/Engine/Plugins/Experimental/AnimNextAnimGraph/Source/AnimNextAnimGraphTestSuite/Private/AnimNextAnimGraphTraitInterfacesTest.cpp
+++ b/Engine/Plugins/Experimental/AnimNextAnimGraph/Source/AnimNextAnimGraphTestSuite/Private/AnimNextAnimGraphTraitInterfacesTest.cpp
@@ -13,6 +13,7 @@
 #include "TraitCore/Trait.h"
 #include "TraitCore/TraitReader.h"
 #include "TraitCore/TraitRegistry.h"
+#include "TraitCore/TraitInterfaceRegistry.h"
 #include "TraitCore/TraitWriter.h"
 #include "TraitCore/ExecutionContext.h"
 #include "TraitCore/ITraitInterface.h"
@@ -230,7 +231,7 @@ namespace UE::AnimNext
 
 	struct IScopedTagInterface : IScopedTraitInterface
 	{
-		DECLARE_ANIM_TRAIT_INTERFACE(IScopedTagInterface, STG, 0xb8f4b397)
+		DECLARE_ANIM_TRAIT_INTERFACE(IScopedTagInterface, 0xb8f4b397)
 
 		virtual FName GetTag(const FExecutionContext& Context, const TTraitBinding<IScopedTagInterface>& Binding) const;
 
@@ -351,8 +352,120 @@ namespace UE::AnimNext
 
 	GENERATE_ANIM_TRAIT_IMPLEMENTATION(FTestScopedTagTrait, TRAIT_INTERFACE_ENUMERATOR, NULL_ANIM_TRAIT_INTERFACE_ENUMERATOR, NULL_ANIM_TRAIT_EVENT_ENUMERATOR)
 	#undef TRAIT_INTERFACE_ENUMERATOR
+
+
+	static const FText GInterfaceTestAName = FText::FromString(TEXT("Interface Test A"));
+	static const FText GInterfaceTestAShortName = FText::FromString(TEXT("ITA"));
+
+	struct ITraitInterfaceTestA : ITraitInterface
+	{
+		DECLARE_ANIM_TRAIT_INTERFACE(ITraitInterfaceTestA, 0xc6d4bde9)
+
+	#if WITH_EDITOR
+		virtual const FText& GetDisplayName() const override
+		{
+			return GInterfaceTestAName;
+		}
+
+		virtual const FText& GetDisplayShortName() const override
+		{
+			return GInterfaceTestAShortName;
+		}
+	#endif // WITH_EDITOR
+	};
+
+	static const FText GInterfaceTestBName = FText::FromString(TEXT("Interface Test B"));
+	static const FText GInterfaceTestBShortName = FText::FromString(TEXT("ITB"));
+
+	struct ITraitInterfaceTestB : ITraitInterface
+	{
+		DECLARE_ANIM_TRAIT_INTERFACE(ITraitInterfaceTestB, 0xc0eadb31)
+
+	#if WITH_EDITOR
+		virtual const FText& GetDisplayName() const override
+		{
+			return GInterfaceTestBName;
+		}
+
+		virtual const FText& GetDisplayShortName() const override
+		{
+			return GInterfaceTestBShortName;
+		}
+
+		virtual bool IsInternal() const override
+		{
+			return true;
+		}
+	#endif // WITH_EDITOR
+	};
+
+} // end namespace UE::AnimNext
+
+// --- Runtime Test Trait Interface Registry ---
+
+IMPLEMENT_SIMPLE_AUTOMATION_TEST(FAnimationAnimNextRuntimeTest_TraitInterfaceRegistry, "Animation.AnimNext.Runtime.TraitInterfaceRegistry", EAutomationTestFlags::EditorContext | EAutomationTestFlags::EngineFilter)
+
+bool FAnimationAnimNextRuntimeTest_TraitInterfaceRegistry::RunTest(const FString& InParameters)
+{
+	using namespace UE::AnimNext;
+
+	{
+		FTraitInterfaceRegistry& Registry = FTraitInterfaceRegistry::Get();
+
+		// Some traits already exist in the engine, keep track of them
+		const uint32 NumAutoRegisteredTraitInterfaces = Registry.GetNum();
+
+		AddErrorIfFalse(Registry.Find(ITraitInterfaceTestA::InterfaceUID) == nullptr, "FAnimationAnimNextRuntimeTest_TraitInterfaceRegistry -> Registry should not contain the Test Interface A");
+		AddErrorIfFalse(Registry.Find(ITraitInterfaceTestB::InterfaceUID) == nullptr, "FAnimationAnimNextRuntimeTest_TraitInterfaceRegistry -> Registry should not contain the Test Interface B");
+
+		{
+			AUTO_REGISTER_ANIM_TRAIT_INTERFACE(ITraitInterfaceTestA)
+
+			AddErrorIfFalse(Registry.GetNum() == NumAutoRegisteredTraitInterfaces + 1, "FAnimationAnimNextRuntimeTest_TraitInterfaceRegistry -> Registry should contain 1 new trait interface");
+
+			const ITraitInterface* TraitInterfaceA = Registry.Find(ITraitInterfaceTestA::InterfaceUID);
+			AddErrorIfFalse(TraitInterfaceA->GetInterfaceUID() == ITraitInterfaceTestA::InterfaceUID, "FAnimationAnimNextRuntimeTest_TraitInterfaceRegistry -> Incorrect InterfaceUID for ITraitInterfaceTestA");
+
+#if WITH_EDITOR
+			AddErrorIfFalse(TraitInterfaceA->GetDisplayName().EqualTo(GInterfaceTestAName), "FAnimationAnimNextRuntimeTest_TraitInterfaceRegistry -> Incorrect Interface Display Name for ITraitInterfaceTestA");
+			AddErrorIfFalse(TraitInterfaceA->GetDisplayShortName().EqualTo(GInterfaceTestAShortName), "FAnimationAnimNextRuntimeTest_TraitInterfaceRegistry -> Incorrect Interface Display Short Name for ITraitInterfaceTestA");
+
+			AddErrorIfFalse(TraitInterfaceA->IsInternal() == false, "FAnimationAnimNextRuntimeTest_TraitInterfaceRegistry -> Incorrect Interface Internal flag for ITraitInterfaceTestA");
+#endif // WITH_EDITOR
+
+			{
+				AUTO_REGISTER_ANIM_TRAIT_INTERFACE(ITraitInterfaceTestB)
+
+				AddErrorIfFalse(Registry.GetNum() == NumAutoRegisteredTraitInterfaces + 2, "FAnimationAnimNextRuntimeTest_TraitInterfaceRegistry -> Registry should contain 2 new trait interfaces");
+
+				const ITraitInterface* TraitInterfaceB = Registry.Find(ITraitInterfaceTestB::InterfaceUID);
+				AddErrorIfFalse(TraitInterfaceB->GetInterfaceUID() == ITraitInterfaceTestB::InterfaceUID, "FAnimationAnimNextRuntimeTest_TraitInterfaceRegistry -> Incorrect InterfaceUID for ITraitInterfaceTestB");
+
+#if WITH_EDITOR
+				AddErrorIfFalse(TraitInterfaceB->GetDisplayName().EqualTo(GInterfaceTestBName), "FAnimationAnimNextRuntimeTest_TraitInterfaceRegistry -> Incorrect Interface Display Name for ITraitInterfaceTestB");
+				AddErrorIfFalse(TraitInterfaceB->GetDisplayShortName().EqualTo(GInterfaceTestBShortName), "FAnimationAnimNextRuntimeTest_TraitInterfaceRegistry -> Incorrect Interface Display Short Name for ITraitInterfaceTestB");
+
+				AddErrorIfFalse(TraitInterfaceB->IsInternal() == true, "FAnimationAnimNextRuntimeTest_TraitInterfaceRegistry -> Incorrect Interface Internal flag for ITraitInterfaceTestB");
+#endif // WITH_EDITOR
+			}
+
+			AddErrorIfFalse(Registry.Find(ITraitInterfaceTestB::InterfaceUID) == nullptr, "FAnimationAnimNextRuntimeTest_TraitInterfaceRegistry -> Registry should not contain the Test Interface B");
+
+			AddErrorIfFalse(Registry.GetNum() == NumAutoRegisteredTraitInterfaces + 1, "FAnimationAnimNextRuntimeTest_TraitInterfaceRegistry -> Registry should contain 1 new trait interface");
+		}
+
+		AddErrorIfFalse(Registry.Find(ITraitInterfaceTestA::InterfaceUID) == nullptr, "FAnimationAnimNextRuntimeTest_TraitInterfaceRegistry -> Registry should not contain the Test Interface A");
+
+		AddErrorIfFalse(Registry.GetNum() == NumAutoRegisteredTraitInterfaces, "FAnimationAnimNextRuntimeTest_TraitInterfaceRegistry -> Registry should contain 0 new trait interfaces");
+	}
+
+	Tests::FUtils::CleanupAfterTests();
+
+	return true;
 }
 
+// --- Trait Interfaces IHierarchy Test ---
+
 IMPLEMENT_SIMPLE_AUTOMATION_TEST(FAnimationAnimNextRuntimeTest_IHierarchy, "Animation.AnimNext.Runtime.IHierarchy", EAutomationTestFlags::EditorContext | EAutomationTestFlags::EngineFilter)
 
 bool FAnimationAnimNextRuntimeTest_IHierarchy::RunTest(const FString& InParameters)
@@ -573,6 +686,8 @@ bool FAnimationAnimNextRuntimeTest_IHierarchy::RunTest(const FString& InParamete
 	return true;
 }
 
+// --- Trait Interfaces IUpdate Test ---
+
 IMPLEMENT_SIMPLE_AUTOMATION_TEST(FAnimationAnimNextRuntimeTest_IUpdate, "Animation.AnimNext.Runtime.IUpdate", EAutomationTestFlags::EditorContext | EAutomationTestFlags::EngineFilter)
 
 bool FAnimationAnimNextRuntimeTest_IUpdate::RunTest(const FString& InParameters)
@@ -713,6 +828,8 @@ bool FAnimationAnimNextRuntimeTest_IUpdate::RunTest(const FString& InParameters)
 	return true;
 }
 
+// --- Trait Interfaces IEvaluate Test ---
+
 IMPLEMENT_SIMPLE_AUTOMATION_TEST(FAnimationAnimNextRuntimeTest_IEvaluate, "Animation.AnimNext.Runtime.IEvaluate", EAutomationTestFlags::EditorContext | EAutomationTestFlags::EngineFilter)
 
 bool FAnimationAnimNextRuntimeTest_IEvaluate::RunTest(const FString& InParameters)
@@ -849,6 +966,9 @@ bool FAnimationAnimNextRuntimeTest_IEvaluate::RunTest(const FString& InParameter
 
 	return true;
 }
+
+// --- Trait Interfaces IScopedInterface Test ---
+
 IMPLEMENT_SIMPLE_AUTOMATION_TEST(FAnimationAnimNextRuntimeTest_IScopedInterface, "Animation.AnimNext.Runtime.IScopedInterface", EAutomationTestFlags::EditorContext | EAutomationTestFlags::EngineFilter)
 
 bool FAnimationAnimNextRuntimeTest_IScopedInterface::RunTest(const FString& InParameters)
