AnimNext: Fixes for scene graph

Split CanUseProperty/CanUseFunciton up to allow flag checks to be sepearted from type checks (allows adapter types to be used)
Removed alias overload of ApplyParametersToScope, it is not going to be used.
Tweaked the meaning of CreateParameterSource APIs - now if the array of required parameters is empty all parameters will be opted-into
Added call to Update() root parameter scopes

diff --git a/Engine/Plugins/Experimental/AnimNext/Source/AnimNext/Internal/AnimNextModule.h b/Engine/Plugins/Experimental/AnimNext/Source/AnimNext/Internal/AnimNextModule.h
--- a/Engine/Plugins/Experimental/AnimNext/Source/AnimNext/Internal/AnimNextModule.h
+++ b/Engine/Plugins/Experimental/AnimNext/Source/AnimNext/Internal/AnimNextModule.h
@@ -37,9 +37,9 @@ public:
 	// Factory method used to create a parameter source to access the specified instance ID, with a set of parameters that are initially required
 	// @param   InContext             Context used to set up the parameter source
 	// @param   InInstanceId          The instance identifier associated with the parameters that are required
-	// @param   InRequiredParameters  Any required parameters that the source should initially supply, can be empty
+	// @param   InRequiredParameters  Any required parameters that the source should initially supply, can be empty, in which ase all parameters are created
 	// @return a new parameter source, or nullptr if the instance ID could not be handled
-	ANIMNEXT_API TUniquePtr<IParameterSource> CreateParameterSource(const FParameterSourceContext& InContext, const TInstancedStruct<FAnimNextParamInstanceIdentifier>& InInstanceId, TConstArrayView<FName> InRequiredParameters) const;
+	ANIMNEXT_API TUniquePtr<IParameterSource> CreateParameterSource(const FParameterSourceContext& InContext, const TInstancedStruct<FAnimNextParamInstanceIdentifier>& InInstanceId, TConstArrayView<FName> InRequiredParameters = TConstArrayView<FName>()) const;
 
 	// Register a factory that can be used to generate parameter sources
 	// @param   InName                Identifier for the factory
diff --git a/Engine/Plugins/Experimental/AnimNext/Source/AnimNext/Internal/Param/IParameterSourceFactory.h b/Engine/Plugins/Experimental/AnimNext/Source/AnimNext/Internal/Param/IParameterSourceFactory.h
--- a/Engine/Plugins/Experimental/AnimNext/Source/AnimNext/Internal/Param/IParameterSourceFactory.h
+++ b/Engine/Plugins/Experimental/AnimNext/Source/AnimNext/Internal/Param/IParameterSourceFactory.h
@@ -38,9 +38,9 @@ public:
 	// Factory method used to create a parameter source of the specified name, with a set of parameters that are initially required
 	// @param  InContext              Context used to resolve the scope
 	// @param  InInstanceId           The instance identifier associated with the parameters that are required
-	// @param  InRequiredParameters   Any required parameters that the source should initially supply, can be empty
+	// @param  InRequiredParameters   Any required parameters that the source should initially supply, can be empty, in which case all parameters are created
 	// @return a new parameter source, or nullptr if the source was not found
-	virtual TUniquePtr<IParameterSource> CreateParameterSource(const FParameterSourceContext& InContext, const TInstancedStruct<FAnimNextParamInstanceIdentifier>& InInstanceId, TConstArrayView<FName> InRequiredParameters) const = 0;
+	virtual TUniquePtr<IParameterSource> CreateParameterSource(const FParameterSourceContext& InContext, const TInstancedStruct<FAnimNextParamInstanceIdentifier>& InInstanceId, TConstArrayView<FName> InRequiredParameters = TConstArrayView<FName>()) const = 0;
 };
 
 }
\ No newline at end of file
diff --git a/Engine/Plugins/Experimental/AnimNext/Source/AnimNext/Internal/Param/ParamUtils.h b/Engine/Plugins/Experimental/AnimNext/Source/AnimNext/Internal/Param/ParamUtils.h
--- a/Engine/Plugins/Experimental/AnimNext/Source/AnimNext/Internal/Param/ParamUtils.h
+++ b/Engine/Plugins/Experimental/AnimNext/Source/AnimNext/Internal/Param/ParamUtils.h
@@ -25,14 +25,23 @@ struct FParamUtils
 	// occurs, but IsCompatible((int32)InLHS, (int64)InRHS) is not as B could be truncated.
 	static ANIMNEXT_API FParamCompatibility GetCompatibility(const FAnimNextParamType& InLHS, const FAnimNextParamType& InRHS);
 
-	// Check whether the supplied function can be used to access/map parameters
+	// Check whether the supplied function can be used to access/map parameters. Does not check type for validity.
 	// @param	InFunction	The function to check
 	// @param	InClass		The expected class to check - used to determine if hoisted functions are valid in this context.
 	static ANIMNEXT_API bool CanUseFunction(const UFunction* InFunction, const UClass* InExpectedClass);
 
-	// Check whether the supplied property can be used to access/map parameters
+	// Check whether the supplied function can be used to access/map parameters. Checks type and returns it in OutTypeHandle
+	// @param	InFunction		The function to check
+	// @param	InClass			The expected class to check - used to determine if hoisted functions are valid in this context.
+	// @param	OutTypeHandle	Type handle of the function's return value
+	static ANIMNEXT_API bool CanUseFunction(const UFunction* InFunction, const UClass* InExpectedClass, FParamTypeHandle& OutTypeHandle);
+	
+	// Check (via flags only) the supplied property can be used to access/map parameters
 	static ANIMNEXT_API bool CanUseProperty(const FProperty* InProperty);
 
+	// Check whether the supplied property can be used to access/map parameters. Checks type and returns it in OutTypeHandle
+	static ANIMNEXT_API bool CanUseProperty(const FProperty* InProperty, FParamTypeHandle& OutTypeHandle);
+
 	// Convert a UOL to an FName
 	static ANIMNEXT_API FName LocatorToName(const FUniversalObjectLocator& InLocator);
 };
diff --git a/Engine/Plugins/Experimental/AnimNext/Source/AnimNext/Internal/Scheduler/ScheduleInitializationContext.h b/Engine/Plugins/Experimental/AnimNext/Source/AnimNext/Internal/Scheduler/ScheduleInitializationContext.h
--- a/Engine/Plugins/Experimental/AnimNext/Source/AnimNext/Internal/Scheduler/ScheduleInitializationContext.h
+++ b/Engine/Plugins/Experimental/AnimNext/Source/AnimNext/Internal/Scheduler/ScheduleInitializationContext.h
@@ -32,8 +32,7 @@ struct FScheduleInitializationContext
 {
 public:
 	// Apply the supplied parameter source to the specified scope, evicting any source that was there previously
-	ANIMNEXT_API void ApplyParametersToScope(FName InScope, EParameterScopeOrdering InOrdering, const UObject* InObject, TConstArrayView<TTuple<FName, FName>> InRequiredParameterAliases) const;
-	ANIMNEXT_API void ApplyParametersToScope(FName InScope, EParameterScopeOrdering InOrdering, const TInstancedStruct<FAnimNextParamInstanceIdentifier>& InInstanceId, const FParameterSourceContext& InContext, TConstArrayView<FName> InRequiredParameters) const;
+	ANIMNEXT_API void ApplyParametersToScope(FName InScope, EParameterScopeOrdering InOrdering, const TInstancedStruct<FAnimNextParamInstanceIdentifier>& InInstanceId, const FParameterSourceContext& InContext, TConstArrayView<FName> InRequiredParameters = TConstArrayView<FName>()) const;
 	ANIMNEXT_API void ApplyParametersToScope(FName InScope, EParameterScopeOrdering InOrdering, FName InId, TConstArrayView<FPropertyBagPropertyDesc> InPropertyDescs, TConstArrayView<TConstArrayView<uint8>> InValues) const;
 
 private:
diff --git a/Engine/Plugins/Experimental/AnimNext/Source/AnimNext/Private/Param/ObjectProxyFactory.cpp b/Engine/Plugins/Experimental/AnimNext/Source/AnimNext/Private/Param/ObjectProxyFactory.cpp
--- a/Engine/Plugins/Experimental/AnimNext/Source/AnimNext/Private/Param/ObjectProxyFactory.cpp
+++ b/Engine/Plugins/Experimental/AnimNext/Source/AnimNext/Private/Param/ObjectProxyFactory.cpp
@@ -4,6 +4,8 @@
 
 #include "Param/ClassProxy.h"
 #include "ObjectProxy.h"
+#include "Engine/Blueprint.h"
+#include "Engine/BlueprintGeneratedClass.h"
 #include "Param/AnimNextParamUniversalObjectLocator.h"
 
 namespace UE::AnimNext
@@ -57,7 +59,8 @@ TUniquePtr<IParameterSource> FObjectProxyFactory::CreateParameterSource(const FP
 		// We shouldn't be loading as part of this call - if this hits we need to consider loading objects up front somehow
 		check(!(ResultData.Flags.bWasLoaded || ResultData.Flags.bWasLoadedIndirectly));	
 
-		if(ResultData.Object)
+		// We only support UClass/UBlueprintGeneratedClass here, not other subclasses (e.g. Verse classes)
+		if(ResultData.Object != nullptr && (ExactCast<UClass>(ResultData.Object->GetClass()) != nullptr || ExactCast<UBlueprintGeneratedClass>(ResultData.Object->GetClass()) != nullptr))
 		{
 			TStringBuilder<256> ScopeAsString;
 			Locator->Locator.ToString(ScopeAsString);
diff --git a/Engine/Plugins/Experimental/AnimNext/Source/AnimNext/Private/Param/ParamUtils.cpp b/Engine/Plugins/Experimental/AnimNext/Source/AnimNext/Private/Param/ParamUtils.cpp
--- a/Engine/Plugins/Experimental/AnimNext/Source/AnimNext/Private/Param/ParamUtils.cpp
+++ b/Engine/Plugins/Experimental/AnimNext/Source/AnimNext/Private/Param/ParamUtils.cpp
@@ -408,7 +408,7 @@ FParamCompatibility FParamUtils::GetCompatibility(const FAnimNextParamType& InLH
 	return GetCompatibility(InLHS.GetHandle(), InRHS.GetHandle());
 }
 
-bool FParamUtils::CanUseFunction(const UFunction* InFunction, const UClass* InExpectedClass)
+static bool CanUseFunctionInternal(const UFunction* InFunction, const UClass* InExpectedClass, FProperty*& OutReturnProperty)
 {
 	UClass* FunctionClass = InFunction->GetOuterUClass();
 	if(FunctionClass->IsChildOf(UBlueprintFunctionLibrary::StaticClass()))
@@ -441,7 +441,8 @@ bool FParamUtils::CanUseFunction(const UFunction* InFunction, const UClass* InEx
 #if WITH_EDITORONLY_DATA
 				if(InExpectedClass != nullptr)
 				{
-					if(!InExpectedClass->IsChildOf(ObjectProperty->PropertyClass))
+					// It its just a UObject, check the metadata
+					if(ObjectProperty->PropertyClass == UObject::StaticClass())
 					{
 						const FString& AllowedClassMeta = ObjectProperty->GetMetaData("AllowedClass");
 						if(AllowedClassMeta.Len() == 0)
@@ -450,11 +451,15 @@ bool FParamUtils::CanUseFunction(const UFunction* InFunction, const UClass* InEx
 						}
 
 						const UClass* AllowedClass = FindObject<UClass>(nullptr, *AllowedClassMeta);
-						if(AllowedClass == nullptr || !AllowedClass->IsChildOf(InExpectedClass))
+						if(AllowedClass == nullptr || !InExpectedClass->IsChildOf(AllowedClass))
 						{
 							return false;
 						}
 					}
+					else if(!InExpectedClass->IsChildOf(ObjectProperty->PropertyClass))
+					{
+						return false;
+					}
 				}
 #endif
 			}
@@ -464,22 +469,42 @@ bool FParamUtils::CanUseFunction(const UFunction* InFunction, const UClass* InEx
 			{
 				return false;
 			}
+
+			OutReturnProperty = *It;
 		}
 	}
 	else
 	{
 		// We add only 'accessor' functions (no params apart from the return value) that have valid return types
-		const FProperty* ReturnProperty = InFunction->GetReturnProperty();
-		if(ReturnProperty == nullptr || InFunction->NumParms != 1 || !InFunction->HasAnyFunctionFlags(FUNC_BlueprintCallable))
+		OutReturnProperty = InFunction->GetReturnProperty();
+		if(OutReturnProperty == nullptr || InFunction->NumParms != 1 || !InFunction->HasAnyFunctionFlags(FUNC_BlueprintCallable))
 		{
 			return false;
 		}
+	}
 
-		const FParamTypeHandle TypeHandle = FParamTypeHandle::FromProperty(ReturnProperty);
-		if(!TypeHandle.IsValid())
-		{
-			return false;
-		}
+	return true;
+}
+
+bool FParamUtils::CanUseFunction(const UFunction* InFunction, const UClass* InExpectedClass)
+{
+	FProperty* ReturnProperty = nullptr;
+	return CanUseFunctionInternal(InFunction, InExpectedClass, ReturnProperty);
+}
+
+bool FParamUtils::CanUseFunction(const UFunction* InFunction, const UClass* InExpectedClass, FParamTypeHandle& OutTypeHandle)
+{
+	FProperty* ReturnProperty = nullptr;
+	if(!CanUseFunctionInternal(InFunction, InExpectedClass, ReturnProperty))
+	{
+		return false;
+	}
+
+	check(ReturnProperty);
+	OutTypeHandle = FParamTypeHandle::FromProperty(ReturnProperty);
+	if(!OutTypeHandle.IsValid())
+	{
+		return false;
 	}
 
 	return true;
@@ -491,9 +516,18 @@ bool FParamUtils::CanUseProperty(const FProperty* InProperty)
 	{
 		return false;
 	}
+	return true;
+}
+
+bool FParamUtils::CanUseProperty(const FProperty* InProperty, FParamTypeHandle& OutTypeHandle)
+{
+	if(!CanUseProperty(InProperty))
+	{
+		return false;
+	}
 
-	const FParamTypeHandle TypeHandle = FParamTypeHandle::FromProperty(InProperty);
-	if(!TypeHandle.IsValid())
+	OutTypeHandle = FParamTypeHandle::FromProperty(InProperty);
+	if(!OutTypeHandle.IsValid())
 	{
 		return false;
 	}
diff --git a/Engine/Plugins/Experimental/AnimNext/Source/AnimNext/Private/Scheduler/ScheduleInitializationContext.cpp b/Engine/Plugins/Experimental/AnimNext/Source/AnimNext/Private/Scheduler/ScheduleInitializationContext.cpp
--- a/Engine/Plugins/Experimental/AnimNext/Source/AnimNext/Private/Scheduler/ScheduleInitializationContext.cpp
+++ b/Engine/Plugins/Experimental/AnimNext/Source/AnimNext/Private/Scheduler/ScheduleInitializationContext.cpp
@@ -17,24 +17,9 @@ FScheduleInitializationContext::FScheduleInitializationContext(const FScheduleCo
 {
 }
 
-void FScheduleInitializationContext::ApplyParametersToScope(FName InScope, EParameterScopeOrdering InOrdering, const UObject* InObject, TConstArrayView<TTuple<FName, FName>> InRequiredParameterAliases) const
-{
-	FModule& AnimNextModule = FModuleManager::GetModuleChecked<FModule>("AnimNext");
-	TSharedPtr<FObjectProxyFactory> ObjectProxyFactory = StaticCastSharedPtr<FObjectProxyFactory>(AnimNextModule.FindParameterSourceFactory("ObjectProxy"));
-	if(!ObjectProxyFactory.IsValid())
-	{
-		return;
-	}
-
-	TSharedRef<FClassProxy> ClassProxy = ObjectProxyFactory->FindOrCreateClassProxy(InObject->GetClass());
-	TUniquePtr<FObjectProxy> ObjectProxy = MakeUnique<FObjectProxy>(InObject, ClassProxy);
-	ObjectProxy->RequestParameterCacheAlias(InRequiredParameterAliases);
-	Context.GetInstanceData().ApplyParametersToScope(InScope, InOrdering, MoveTemp(ObjectProxy));
-}
-
 void FScheduleInitializationContext::ApplyParametersToScope(FName InScope, EParameterScopeOrdering InOrdering, const TInstancedStruct<FAnimNextParamInstanceIdentifier>& InInstanceId, const FParameterSourceContext& InContext, TConstArrayView<FName> InRequiredParameters) const
 {
-	FModule& AnimNextModule = FModuleManager::GetModuleChecked<FModule>("AnimNextModule");
+	FModule& AnimNextModule = FModuleManager::GetModuleChecked<FModule>("AnimNext");
 	TUniquePtr<IParameterSource> ParameterSource = AnimNextModule.CreateParameterSource(InContext, InInstanceId, InRequiredParameters);
 	Context.GetInstanceData().ApplyParametersToScope(InScope, InOrdering, MoveTemp(ParameterSource));
 }
diff --git a/Engine/Plugins/Experimental/AnimNext/Source/AnimNext/Private/Scheduler/ScheduleTickFunction.cpp b/Engine/Plugins/Experimental/AnimNext/Source/AnimNext/Private/Scheduler/ScheduleTickFunction.cpp
--- a/Engine/Plugins/Experimental/AnimNext/Source/AnimNext/Private/Scheduler/ScheduleTickFunction.cpp
+++ b/Engine/Plugins/Experimental/AnimNext/Source/AnimNext/Private/Scheduler/ScheduleTickFunction.cpp
@@ -32,11 +32,13 @@ void FScheduleBeginTickFunction::Run(float DeltaTime)
 	}
 
 	// Push any user layers we have at the root
+	// TODO: This is potentially thread-unsafe - need to add a GT/WT split here to split unsafe/safe parameter sources
 	FScheduleInstanceData& InstanceData = Entry.Context.GetInstanceData();
 	for(const TUniquePtr<IParameterSource>& RootUserScope : InstanceData.RootUserScopes)
 	{
 		if(RootUserScope.IsValid())
 		{
+			RootUserScope->Update(DeltaTime);
 			InstanceData.PushedRootUserLayers.Add(InstanceData.RootParamStack->PushLayer(RootUserScope->GetLayerHandle()));
 		}
 	}
