Workspace Editor global selection sync system:
- register and forget system through use of TWeakPtr<SWidget>
- last known selection will be cleared when new selection is set
Misc:
- Added SetOnSelectionChanged to SRigVMAssetView to enable capturing weak-ptr

diff --git a/Engine/Plugins/Experimental/AnimNext/Source/AnimNextEditor/Private/AnimNextEditorModule.cpp b/Engine/Plugins/Experimental/AnimNext/Source/AnimNextEditor/Private/AnimNextEditorModule.cpp
--- a/Engine/Plugins/Experimental/AnimNext/Source/AnimNextEditor/Private/AnimNextEditorModule.cpp
+++ b/Engine/Plugins/Experimental/AnimNext/Source/AnimNextEditor/Private/AnimNextEditorModule.cpp
@@ -359,14 +359,7 @@ void FModule::RegisterWorkspaceDocumentTypes(Workspace::IWorkspaceEditorModule&
 
 				TWeakPtr<Workspace::IWorkspaceEditor> WeakWorkspaceEditor = InContext.WorkspaceEditor;
 
-				return SNew(SRigVMAssetView, EditorData)
-					.OnSelectionChanged_Lambda([WeakWorkspaceEditor](const TArray<UObject*>& InEntries)
-					{
-						if(TSharedPtr<Workspace::IWorkspaceEditor> WorkspaceEditor = WeakWorkspaceEditor.Pin())
-						{
-							WorkspaceEditor->SetDetailsObjects(InEntries);
-						}
-					})
+				TSharedPtr<SRigVMAssetView> SharedAssetView = SNew(SRigVMAssetView, EditorData)
 					.OnOpenGraph_Lambda([WeakWorkspaceEditor](URigVMGraph* InGraph)
 					{
 						if(TSharedPtr<Workspace::IWorkspaceEditor> WorkspaceEditor = WeakWorkspaceEditor.Pin())
@@ -403,6 +396,21 @@ void FModule::RegisterWorkspaceDocumentTypes(Workspace::IWorkspaceEditorModule&
 							}
 						}
 					});
+				
+				TWeakPtr<SRigVMAssetView> WeakAssetView = SharedAssetView;
+				SharedAssetView->SetOnSelectionChanged(SRigVMAssetView::FOnSelectionChanged::CreateLambda([WeakWorkspaceEditor, WeakAssetView](const TArray<UObject*>& InEntries)
+				{
+					if(TSharedPtr<SRigVMAssetView> SharedAssetView = WeakAssetView.Pin())
+					{
+						if(TSharedPtr<Workspace::IWorkspaceEditor> WorkspaceEditor = WeakWorkspaceEditor.Pin())
+						{	
+							WorkspaceEditor->SetGlobalSelection(SharedAssetView, UE::Workspace::FOnClearGlobalSelection::CreateLambda([WeakAssetView](){ if (const TSharedPtr<SRigVMAssetView> SharedAssetView = WeakAssetView.Pin()) { SharedAssetView->ClearSelection(); } }));
+							WorkspaceEditor->SetDetailsObjects(InEntries);
+						}
+				}}));
+
+
+				return SharedAssetView.ToSharedRef();
 			}),
 			Workspace::WorkspaceTabs::TopMiddleDocumentArea);
 	AnimNextGraphDocumentArgs.OnGetTabName = Workspace::FOnGetTabName::CreateLambda([](const Workspace::FWorkspaceEditorContext& InContext)
@@ -827,4 +835,4 @@ void FModule::UnregisterWorkspaceDocumentTypes()
 
 IMPLEMENT_MODULE(UE::AnimNext::Editor::FModule, AnimNextEditor);
 
-#undef LOCTEXT_NAMESPACE
+#undef LOCTEXT_NAMESPACE
\ No newline at end of file
diff --git a/Engine/Plugins/Experimental/AnimNext/Source/AnimNextEditor/Private/Common/SRigVMAssetView.cpp b/Engine/Plugins/Experimental/AnimNext/Source/AnimNextEditor/Private/Common/SRigVMAssetView.cpp
--- a/Engine/Plugins/Experimental/AnimNext/Source/AnimNextEditor/Private/Common/SRigVMAssetView.cpp
+++ b/Engine/Plugins/Experimental/AnimNext/Source/AnimNextEditor/Private/Common/SRigVMAssetView.cpp
@@ -118,6 +118,16 @@ void SRigVMAssetView::UnregisterCategoryFactory(FName InCategory)
 	CategoryFactories.Remove(InCategory);
 }
 
+void SRigVMAssetView::ClearSelection()
+{
+	EntriesList->ClearSelection();
+}
+
+void SRigVMAssetView::SetOnSelectionChanged(SRigVMAssetView::FOnSelectionChanged InDelegate)
+{
+	OnSelectionChangedDelegate = InDelegate;
+}
+
 void SRigVMAssetView::Construct(const FArguments& InArgs, UAnimNextRigVMAssetEditorData* InEditorData)
 {
 	using namespace RigVMAssetView;
diff --git a/Engine/Plugins/Experimental/AnimNext/Source/AnimNextEditor/Private/Common/SRigVMAssetView.h b/Engine/Plugins/Experimental/AnimNext/Source/AnimNextEditor/Private/Common/SRigVMAssetView.h
--- a/Engine/Plugins/Experimental/AnimNext/Source/AnimNextEditor/Private/Common/SRigVMAssetView.h
+++ b/Engine/Plugins/Experimental/AnimNext/Source/AnimNextEditor/Private/Common/SRigVMAssetView.h
@@ -44,6 +44,8 @@ public:
 	static void RegisterCategoryFactory(FName InCategory, FCategoryWidgetFactoryFunction&& InFunction);
 	static void UnregisterCategoryFactory(FName InCategory);
 
+	void ClearSelection();
+	void SetOnSelectionChanged(SRigVMAssetView::FOnSelectionChanged InDelegate);
 private:
 	virtual FReply OnKeyDown(const FGeometry& MyGeometry, const FKeyEvent& InKeyEvent) override;
 
diff --git a/Engine/Plugins/Experimental/Workspace/Source/WorkspaceEditor/Private/Outliner/WorkspaceOutlinerMode.cpp b/Engine/Plugins/Experimental/Workspace/Source/WorkspaceEditor/Private/Outliner/WorkspaceOutlinerMode.cpp
--- a/Engine/Plugins/Experimental/Workspace/Source/WorkspaceEditor/Private/Outliner/WorkspaceOutlinerMode.cpp
+++ b/Engine/Plugins/Experimental/Workspace/Source/WorkspaceEditor/Private/Outliner/WorkspaceOutlinerMode.cpp
@@ -215,6 +215,11 @@ void FWorkspaceOutlinerMode::HandleItemSelection(const FSceneOutlinerItemSelecti
 void FWorkspaceOutlinerMode::OnItemSelectionChanged(FSceneOutlinerTreeItemPtr Item, ESelectInfo::Type SelectionType, const FSceneOutlinerItemSelection& Selection)
 {
 	HandleItemSelection(Selection);
+
+	if (TSharedPtr<IWorkspaceEditor> SharedWorkspaceEditor = WeakWorkspaceEditor.Pin())
+	{
+		SharedWorkspaceEditor->SetGlobalSelection(SceneOutliner->AsShared(), FOnClearGlobalSelection::CreateRaw(this, &FWorkspaceOutlinerMode::ResetOutlinerSelection));
+	}
 }
 
 TUniquePtr<ISceneOutlinerHierarchy> FWorkspaceOutlinerMode::CreateHierarchy()
@@ -228,6 +233,11 @@ void FWorkspaceOutlinerMode::OnWorkspaceModified(UWorkspace* InWorkspace)
 	SceneOutliner->FullRefresh();
 }
 
+void FWorkspaceOutlinerMode::ResetOutlinerSelection()
+{
+	SceneOutliner->ClearSelection();
+}
+
 void FWorkspaceOutlinerMode::OpenItems(TArrayView<const FSceneOutlinerTreeItemPtr> Items) const
 {
 	for (const FSceneOutlinerTreeItemPtr& Item : Items)
diff --git a/Engine/Plugins/Experimental/Workspace/Source/WorkspaceEditor/Private/Outliner/WorkspaceOutlinerMode.h b/Engine/Plugins/Experimental/Workspace/Source/WorkspaceEditor/Private/Outliner/WorkspaceOutlinerMode.h
--- a/Engine/Plugins/Experimental/Workspace/Source/WorkspaceEditor/Private/Outliner/WorkspaceOutlinerMode.h
+++ b/Engine/Plugins/Experimental/Workspace/Source/WorkspaceEditor/Private/Outliner/WorkspaceOutlinerMode.h
@@ -30,6 +30,7 @@ namespace UE::Workspace
 		// End ISceneOutlinerMode overrides
 
 		void OnWorkspaceModified(UWorkspace* InWorkspace);
+		void ResetOutlinerSelection();
 
 		void OpenItems(TArrayView<const FSceneOutlinerTreeItemPtr> Items) const;
 		void DeleteItems(TArrayView<const FSceneOutlinerTreeItemPtr> Items) const;
diff --git a/Engine/Plugins/Experimental/Workspace/Source/WorkspaceEditor/Private/SGraphDocument.cpp b/Engine/Plugins/Experimental/Workspace/Source/WorkspaceEditor/Private/SGraphDocument.cpp
--- a/Engine/Plugins/Experimental/Workspace/Source/WorkspaceEditor/Private/SGraphDocument.cpp
+++ b/Engine/Plugins/Experimental/Workspace/Source/WorkspaceEditor/Private/SGraphDocument.cpp
@@ -40,6 +40,11 @@ void SGraphDocument::Construct(const FArguments& InArgs, TSharedRef<FWorkspaceEd
 	Events.OnSelectionChanged = SGraphEditor::FOnSelectionChanged::CreateLambda([this, OnGraphSelectionChanged = InArgs._OnGraphSelectionChanged](const TSet<UObject*>& NewSelection)
 	{
 		OnGraphSelectionChanged.ExecuteIfBound(FWorkspaceEditorContext(HostingAppPtr.Pin().ToSharedRef(), EdGraph), NewSelection);
+
+		if (const TSharedPtr<FWorkspaceEditor> SharedWorkspaceEditor = HostingAppPtr.Pin())
+		{
+			SharedWorkspaceEditor->SetGlobalSelection(AsShared(), FOnClearGlobalSelection::CreateRaw(this, &SGraphDocument::OnResetSelection));
+		}
 	});
 	Events.OnTextCommitted = ::FOnNodeTextCommitted::CreateLambda([this, OnNodeTextCommitted = InArgs._OnNodeTextCommitted](const FText& NewText, ETextCommit::Type CommitInfo, UEdGraphNode* NodeBeingChanged)
 	{
@@ -180,4 +185,8 @@ bool SGraphDocument::IsEditable(UEdGraph* InGraph) const
 	return InGraph && HostingAppPtr.Pin()->InEditingMode() && InGraph->bEditable;
 }
 
+void SGraphDocument::OnResetSelection()
+{
+	GraphEditor->ClearSelectionSet();
+}
 };
\ No newline at end of file
diff --git a/Engine/Plugins/Experimental/Workspace/Source/WorkspaceEditor/Private/SGraphDocument.h b/Engine/Plugins/Experimental/Workspace/Source/WorkspaceEditor/Private/SGraphDocument.h
--- a/Engine/Plugins/Experimental/Workspace/Source/WorkspaceEditor/Private/SGraphDocument.h
+++ b/Engine/Plugins/Experimental/Workspace/Source/WorkspaceEditor/Private/SGraphDocument.h
@@ -71,6 +71,7 @@ class SGraphDocument : public SCompoundWidget
 
 
 	bool IsEditable(UEdGraph* InGraph) const;
+	void OnResetSelection();
 
 	// The graph we are editing
 	UEdGraph* EdGraph = nullptr;
diff --git a/Engine/Plugins/Experimental/Workspace/Source/WorkspaceEditor/Private/WorkspaceEditor.cpp b/Engine/Plugins/Experimental/Workspace/Source/WorkspaceEditor/Private/WorkspaceEditor.cpp
--- a/Engine/Plugins/Experimental/Workspace/Source/WorkspaceEditor/Private/WorkspaceEditor.cpp
+++ b/Engine/Plugins/Experimental/Workspace/Source/WorkspaceEditor/Private/WorkspaceEditor.cpp
@@ -283,6 +283,18 @@ UWorkspaceSchema* FWorkspaceEditor::GetSchema() const
 	return Workspace.Get() != nullptr ? Workspace->GetSchema() : nullptr;
 }
 
+void FWorkspaceEditor::SetGlobalSelection(FGlobalSelectionId SelectionId, FOnClearGlobalSelection OnClearSelectionDelegate)
+{
+	// Only execute if widget is still valid, and it is not the same as the previous call 
+	if (LastGlobalSelectionId.IsValid() && SelectionId != LastGlobalSelectionId)
+	{		
+		LastOnClearSelectionDelegate.ExecuteIfBound();
+	}
+
+	LastGlobalSelectionId = SelectionId;
+	LastOnClearSelectionDelegate = OnClearSelectionDelegate;
+}
+
 void FWorkspaceEditor::BindCommands()
 {	
 	FWorkspaceAssetEditorCommands::Register();
diff --git a/Engine/Plugins/Experimental/Workspace/Source/WorkspaceEditor/Private/WorkspaceEditor.h b/Engine/Plugins/Experimental/Workspace/Source/WorkspaceEditor/Private/WorkspaceEditor.h
--- a/Engine/Plugins/Experimental/Workspace/Source/WorkspaceEditor/Private/WorkspaceEditor.h
+++ b/Engine/Plugins/Experimental/Workspace/Source/WorkspaceEditor/Private/WorkspaceEditor.h
@@ -90,6 +90,8 @@ private:
 	virtual void RefreshDetails() override;
 	virtual UWorkspaceSchema* GetSchema() const override;
 
+	virtual void SetGlobalSelection(FGlobalSelectionId SelectionId, FOnClearGlobalSelection OnClearSelectionDelegate) override;
+
 	void BindCommands();
 
 	void ExtendMenu();
@@ -124,6 +126,9 @@ private:
 
 	bool bSavingWorkspaceOnly = false;
 	bool bClosingDown = false;
+
+	FGlobalSelectionId LastGlobalSelectionId = nullptr;
+	FOnClearGlobalSelection LastOnClearSelectionDelegate = nullptr;
 };
 
-}
+}
\ No newline at end of file
diff --git a/Engine/Plugins/Experimental/Workspace/Source/WorkspaceEditor/Public/IWorkspaceEditor.h b/Engine/Plugins/Experimental/Workspace/Source/WorkspaceEditor/Public/IWorkspaceEditor.h
--- a/Engine/Plugins/Experimental/Workspace/Source/WorkspaceEditor/Public/IWorkspaceEditor.h
+++ b/Engine/Plugins/Experimental/Workspace/Source/WorkspaceEditor/Public/IWorkspaceEditor.h
@@ -1,4 +1,4 @@
-// Copyright Epic Games, Inc. All Rights Reserved.
+﻿// Copyright Epic Games, Inc. All Rights Reserved.
 
 #pragma once
 
@@ -10,6 +10,9 @@ class UWorkspaceSchema;
 namespace UE::Workspace
 {
 
+typedef TWeakPtr<SWidget> FGlobalSelectionId;
+using FOnClearGlobalSelection = FSimpleDelegate;
+
 class IWorkspaceEditor : public FBaseAssetToolkit
 {
 public:
@@ -32,6 +35,9 @@ public:
 
 	// Exposes the editor WorkspaceSchema
 	virtual UWorkspaceSchema* GetSchema() const = 0;
+
+	// Set the _current_ global selection (last SWidget with selection set) with delegate to clear it selection on next SetGlobalSelection()
+	virtual void SetGlobalSelection(FGlobalSelectionId SelectionId, FOnClearGlobalSelection OnClearSelectionDelegate) = 0;
 };
 
 }
\ No newline at end of file
