Rename AnimNextGraph asset to AnimNextModule

Something of a name conflict with the 'module' (IModuleInterface). I've elected to add Interface/Impl suffixes to those files rather than rename the whole module (e.g. to some thing like 'AnimNextCore') as that is much more complex.

#virtualized

diff --git a/Engine/Plugins/Experimental/AnimNext/AnimNext.uplugin b/Engine/Plugins/Experimental/AnimNext/AnimNext.uplugin
--- a/Engine/Plugins/Experimental/AnimNext/AnimNext.uplugin
+++ b/Engine/Plugins/Experimental/AnimNext/AnimNext.uplugin
@@ -36,6 +36,11 @@
 		{
 			"Name": "Workspace",
 			"Enabled": true
+		},
+		{
+			"Name": "EditorScriptingUtilities",
+			"Enabled": true,
+			"TargetAllowList": [ "Editor" ]
 		}
 	],
 	"Modules" :
diff --git a/Engine/Plugins/Experimental/AnimNext/Config/BaseAnimNext.ini b/Engine/Plugins/Experimental/AnimNext/Config/BaseAnimNext.ini
--- a/Engine/Plugins/Experimental/AnimNext/Config/BaseAnimNext.ini
+++ b/Engine/Plugins/Experimental/AnimNext/Config/BaseAnimNext.ini
@@ -21,4 +21,27 @@
 +ClassRedirects=(OldName="/Script/AnimNextEditor.AnimNextWorkspace", NewName="/Script/WorkspaceEditor.Workspace")
 
 ; Scene graph prototype compatibility
-+ClassRedirects=(OldName="/Script/AnimNext.AnimNextMeshComponent", NewName="/Script/Engine.SkeletalMeshComponent")
\ No newline at end of file
++ClassRedirects=(OldName="/Script/AnimNext.AnimNextMeshComponent", NewName="/Script/Engine.SkeletalMeshComponent")
+
+; AnimNextGraph -> AnimNextModule
++ClassRedirects=(OldName="/Script/AnimNext.AnimNextGraph", NewName="/Script/AnimNext.AnimNextModule")
++ClassRedirects=(OldName="/Script/AnimNextUncookedOnly.AnimNextGraph_AnimationGraph", NewName="/Script/AnimNextUncookedOnly.AnimNextModule_AnimationGraph")
++ClassRedirects=(OldName="/Script/AnimNextUncookedOnly.AnimNextGraph_Controller", NewName="/Script/AnimNextUncookedOnly.AnimNextModule_Controller")
++ClassRedirects=(OldName="/Script/AnimNextUncookedOnly.AnimNextGraph_EditorData", NewName="/Script/AnimNextUncookedOnly.AnimNextModule_EditorData")
++ClassRedirects=(OldName="/Script/AnimNextUncookedOnly.AnimNextGraph_EventGraph", NewName="/Script/AnimNextUncookedOnly.AnimNextModule_EventGraph")
++ClassRedirects=(OldName="/Script/AnimNextUncookedOnly.AnimNextGraph_Parameter", NewName="/Script/AnimNextUncookedOnly.AnimNextModule_Parameter")
++ClassRedirects=(OldName="/Script/AnimNextUncookedOnly.AnimNextGraph_AnimationGraphSchema", NewName="/Script/AnimNextUncookedOnly.AnimNextAnimationGraphSchema")
++ClassRedirects=(OldName="/Script/AnimNextUncookedOnly.AnimNextGraph_EdGraph", NewName="/Script/AnimNextUncookedOnly.AnimNextEdGraph")
++ClassRedirects=(OldName="/Script/AnimNextUncookedOnly.AnimNextGraph_EdGraphNode", NewName="/Script/AnimNextUncookedOnly.AnimNextEdGraphNode")
++ClassRedirects=(OldName="/Script/AnimNextUncookedOnly.AnimNextGraph_EdGraphSchema", NewName="/Script/AnimNextUncookedOnly.AnimNextEdGraphSchema")
++ClassRedirects=(OldName="/Script/AnimNextUncookedOnly.AnimNextGraph_UnitNode", NewName="/Script/AnimNextUncookedOnly.AnimNextUnitNode")
++ClassRedirects=(OldName="/Script/AnimNextUncookedOnly.AnimNextGraphLibrary", NewName="/Script/AnimNextUncookedOnly.AnimNextModuleLibrary")
++ClassRedirects=(OldName="/Script/AnimNextUncookedOnly.AnimNextRigVMWorkspaceAssetUserData", NewName="/Script/AnimNextUncookedOnly.AnimNextModuleWorkspaceAssetUserData")
+
++StructRedirects=(OldName="/Script/AnimNext.RigUnit_AnimNextGraphSelf", NewName="/Script/AnimNext.RigUnit_AnimNextModuleSelf")
+
++PropertyRedirects=(OldName="/Script/AnimNext.AnimNextGraphInstance.Graph", NewName="Module")
++PropertyRedirects=(OldName="/Script/AnimNext.AnimNextScheduleEntry_AnimNextGraph.Graph", NewName="Module")
++PropertyRedirects=(OldName="/Script/AnimNext.AnimNextSubGraphHostTraitSharedData.SubGraph", NewName="Module")
++PropertyRedirects=(OldName="/Script/AnimNext.AnimNextScheduleGraphTask.Graph", NewName="Module")
++PropertyRedirects=(OldName="/Script/AnimNext.AnimNextScheduleGraphTask.DynamicGraph", NewName="DynamicModule")
diff --git a/Engine/Plugins/Experimental/AnimNext/Source/AnimNext/Internal/AnimNextExecuteContext.h b/Engine/Plugins/Experimental/AnimNext/Source/AnimNext/Internal/AnimNextExecuteContext.h
--- a/Engine/Plugins/Experimental/AnimNext/Source/AnimNext/Internal/AnimNextExecuteContext.h
+++ b/Engine/Plugins/Experimental/AnimNext/Source/AnimNext/Internal/AnimNextExecuteContext.h
@@ -57,6 +57,6 @@ protected:
 	TVariant<FAnimNextGraphContextData, FAnimNextParamContextData> ContextData;
 
 	friend struct FAnimNextGraphInstance;
-	friend class UAnimNextGraph;
+	friend class UAnimNextModule;
 };
 
diff --git a/Engine/Plugins/Experimental/AnimNext/Source/AnimNext/Internal/AnimNextModuleImpl.h b/Engine/Plugins/Experimental/AnimNext/Source/AnimNext/Internal/AnimNextModuleImpl.h
new file mode 100644
--- /dev/null
+++ b/Engine/Plugins/Experimental/AnimNext/Source/AnimNext/Internal/AnimNextModuleImpl.h
@@ -0,0 +1,69 @@
+// Copyright Epic Games, Inc. All Rights Reserved.
+
+#pragma once
+
+#include "CoreMinimal.h"
+#include "IAnimNextModuleInterface.h"
+
+struct FAnimNextGraphInstancePtr;
+struct FAnimNextParamInstanceIdentifier;
+
+namespace UE::AnimNext
+{
+	class IParameterSource;
+	class IParameterSourceFactory;
+	struct FParameterSourceContext;
+}
+
+// Enable console commands only in development builds when logging is enabled
+#define WITH_ANIMNEXT_CONSOLE_COMMANDS (!UE_BUILD_SHIPPING && !NO_LOGGING)
+
+namespace UE::AnimNext
+{
+
+class FAnimNextModuleImpl : public IAnimNextModuleInterface
+{
+public:
+	// IModuleInterface interface
+	virtual void StartupModule() override;
+	virtual void ShutdownModule() override;
+
+	// IAnimNextModuleInterface interface
+	virtual void RegisterAnimNextAnimGraph(const IAnimNextAnimGraph& InAnimGraphImpl) override;
+	virtual void UnregisterAnimNextAnimGraph() override;
+	virtual void UpdateGraph(FAnimNextGraphInstancePtr& GraphInstance, float DeltaTime, FTraitEventList& InputEventList, FTraitEventList& OutputEventList) override;
+	virtual void EvaluateGraph(FAnimNextGraphInstancePtr& GraphInstance, const FReferencePose& RefPose, int32 GraphLODLevel, FLODPoseHeap& OutputPose) const override;
+
+	// Factory method used to create a parameter source to access the specified instance ID, with a set of parameters that are initially required
+	// @param   InContext             Context used to set up the parameter source
+	// @param   InInstanceId          The instance identifier associated with the parameters that are required
+	// @param   InRequiredParameters  Any required parameters that the source should initially supply, can be empty, in which ase all parameters are created
+	// @return a new parameter source, or nullptr if the instance ID could not be handled
+	ANIMNEXT_API TUniquePtr<IParameterSource> CreateParameterSource(const FParameterSourceContext& InContext, const TInstancedStruct<FAnimNextParamInstanceIdentifier>& InInstanceId, TConstArrayView<FName> InRequiredParameters = TConstArrayView<FName>()) const;
+
+	// Register a factory that can be used to generate parameter sources
+	// @param   InName                Identifier for the factory
+	// @param   InFactory             The factory to register
+	ANIMNEXT_API void RegisterParameterSourceFactory(FName InName, TSharedRef<IParameterSourceFactory> InFactory);
+
+	// Unregister a factory previously passed to RegisterParameterSourceFactory
+	// @param   InName                Identifier for the factory
+	ANIMNEXT_API void UnregisterParameterSourceFactory(FName InName);
+
+	// Find a factory previously passed to RegisterParameterSourceFactory
+	// @param   InName                Identifier for the factory
+	// @return a shared ptr to the factory, if found, otherwise an invalid shared ptr
+	ANIMNEXT_API TSharedPtr<IParameterSourceFactory> FindParameterSourceFactory(FName InName);
+
+#if WITH_ANIMNEXT_CONSOLE_COMMANDS
+	void ListNodeTemplates(const TArray<FString>& Args);
+	void ListModules(const TArray<FString>& Args);
+
+	TArray<IConsoleObject*> ConsoleCommands;
+#endif
+
+	// All known factories
+	TMap<FName, TSharedRef<IParameterSourceFactory>> ParameterSourceFactories;
+};
+
+}
diff --git a/Engine/Plugins/Experimental/AnimNext/Source/AnimNext/Internal/Component/AnimNextComponentParameter.h b/Engine/Plugins/Experimental/AnimNext/Source/AnimNext/Internal/Component/AnimNextComponentParameter.h
--- a/Engine/Plugins/Experimental/AnimNext/Source/AnimNext/Internal/Component/AnimNextComponentParameter.h
+++ b/Engine/Plugins/Experimental/AnimNext/Source/AnimNext/Internal/Component/AnimNextComponentParameter.h
@@ -7,7 +7,7 @@
 #include "Param/ParamTypeHandle.h"
 #include "AnimNextComponentParameter.generated.h"
 
-class UAnimNextGraph;
+class UAnimNextModule;
 
 // Boilerplate macro that all derived types of UAnimNextComponentParameter should use.
 // Usage is: IMPLEMENT_ANIMNEXT_COMPONENT_PARAMETER(UAnimNextComponentParameter_DerivedType)
@@ -68,10 +68,10 @@ class UAnimNextComponentParameter_AnimNextGraph : public UAnimNextComponentParam
 	IMPLEMENT_ANIMNEXT_COMPONENT_PARAMETER(UAnimNextComponentParameter_AnimNextGraph);
 
 	// The parameter to set the value to
-	UPROPERTY(EditAnywhere, Category = "Parameter", meta = (CustomWidget = "ParamName", AllowedParamType = "TObjectPtr<UAnimNextGraph>"))
+	UPROPERTY(EditAnywhere, Category = "Parameter", meta = (CustomWidget = "ParamName", AllowedParamType = "TObjectPtr<UAnimNextModule>"))
 	FName Parameter;
 
 	// The value to set
 	UPROPERTY(EditAnywhere, Category = "Parameter")
-	TObjectPtr<UAnimNextGraph> Value;
+	TObjectPtr<UAnimNextModule> Value;
 };
diff --git a/Engine/Plugins/Experimental/AnimNext/Source/AnimNext/Internal/Graph/AnimGraph/AnimNode_AnimNextParameters.h b/Engine/Plugins/Experimental/AnimNext/Source/AnimNext/Internal/Graph/AnimGraph/AnimNode_AnimNextParameters.h
--- a/Engine/Plugins/Experimental/AnimNext/Source/AnimNext/Internal/Graph/AnimGraph/AnimNode_AnimNextParameters.h
+++ b/Engine/Plugins/Experimental/AnimNext/Source/AnimNext/Internal/Graph/AnimGraph/AnimNode_AnimNextParameters.h
@@ -6,7 +6,7 @@
 #include "Animation/AnimNodeBase.h"
 #include "AnimNode_AnimNextParameters.generated.h"
 
-class UAnimNextGraph;
+class UAnimNextModule;
 
 namespace UE::AnimNext
 {
@@ -32,10 +32,10 @@ private:
 	FPoseLink Source;
 
 	UPROPERTY(EditAnywhere, Category = Settings)
-	TObjectPtr<UAnimNextGraph> Parameters;
+	TObjectPtr<UAnimNextModule> Parameters;
 
 	// Cache previous parameters so we know when it changes via pin
-	TObjectPtr<UAnimNextGraph> PreviousParameters;
+	TObjectPtr<UAnimNextModule> PreviousParameters;
 
 	// Cached proxy
 	TUniquePtr<UE::AnimNext::FParametersProxy> ParametersProxy;
diff --git a/Engine/Plugins/Experimental/AnimNext/Source/AnimNext/Internal/Module/AnimNextModule.h b/Engine/Plugins/Experimental/AnimNext/Source/AnimNext/Internal/Module/AnimNextModule.h
new file mode 100644
--- /dev/null
+++ b/Engine/Plugins/Experimental/AnimNext/Source/AnimNext/Internal/Module/AnimNextModule.h
@@ -0,0 +1,194 @@
+// Copyright Epic Games, Inc. All Rights Reserved.
+
+#pragma once
+
+#include "CoreMinimal.h"
+#include "Graph/AnimNextGraphEntryPoint.h"
+#include "AnimNextRigVMAsset.h"
+#include "Graph/RigUnit_AnimNextGraphRoot.h"
+#include "RigVMCore/RigVM.h"
+#include "TraitCore/TraitPtr.h"
+#include "TraitCore/TraitHandle.h"
+#include "TraitCore/EntryPointHandle.h"
+#include "Graph/GraphInstanceComponent.h"
+#include "Graph/RigUnit_AnimNextGraphEvaluator.h"
+#include "Param/ParamId.h"
+#include "Scheduler/IAnimNextScheduleTermInterface.h"
+#include "RigVMHost.h"
+
+#include "AnimNextModule.generated.h"
+
+class UEdGraph;
+class UAnimNextModule;
+class UAnimGraphNode_AnimNextGraph;
+struct FAnimNode_AnimNextGraph;
+struct FRigUnit_AnimNextGraphEvaluator;
+struct FAnimNextGraphInstancePtr;
+struct FAnimNextGraphInstance;
+class UAnimNextSchedule;
+struct FAnimNextScheduleGraphTask;
+struct FAnimNextEditorParam;
+struct FAnimNextParam;
+
+namespace UE::AnimNext
+{
+	struct FContext;
+	struct FExecutionContext;
+	class FAnimNextModuleImpl;
+	struct FTestUtils;
+	struct FParametersProxy;
+}
+
+namespace UE::AnimNext::UncookedOnly
+{
+	struct FUtils;
+}
+
+namespace UE::AnimNext::Editor
+{
+	class FModuleEditor;
+	class FParameterCustomization;
+}
+
+namespace UE::AnimNext::Graph
+{
+	extern ANIMNEXT_API const FName EntryPointName;
+	extern ANIMNEXT_API const FName ResultName;
+}
+
+// A user-created collection of animation logic & data
+UCLASS(BlueprintType)
+class ANIMNEXT_API UAnimNextModule : public UAnimNextRigVMAsset, public IAnimNextScheduleTermInterface
+{
+	GENERATED_BODY()
+
+public:
+	UAnimNextModule(const FObjectInitializer& ObjectInitializer);
+
+	// UObject interface
+	virtual void Serialize(FArchive& Ar) override;
+	virtual void PostLoad() override;
+#if WITH_EDITOR
+	virtual void PostEditChangeProperty(FPropertyChangedEvent& PropertyChangedEvent) override;
+#endif
+
+	// IAnimNextScheduleTermInterface interface
+	virtual TConstArrayView<UE::AnimNext::FScheduleTerm> GetTerms() const override;
+
+	// Allocates an instance of the graph
+	// @param	OutInstance		The instance to allocate data for
+	// @param	InEntryPoint	The entry point to use. If this is NAME_None then the default entry point for this graph is used
+	void AllocateInstance(FAnimNextGraphInstancePtr& OutInstance, FName InEntryPoint = NAME_None) const;
+
+	// Allocates an instance of the graph with the specified parent graph instance
+	// @param	InOutParentGraphInstance	The parent graph instance to use
+	// @param	OutInstance					The instance to allocate data for
+	// @param	InEntryPoint				The entry point to use. If this is NAME_None then the default entry point for this graph is used
+	void AllocateInstance(FAnimNextGraphInstance& InOutParentGraphInstance, FAnimNextGraphInstancePtr& OutInstance, FName InEntryPoint = NAME_None) const;
+
+	// Update the parameter layer, if any
+	void UpdateLayer(UE::AnimNext::FParamStackLayerHandle& InHandle, float InDeltaTime) const;
+
+	// Get the fully-qualified default entry point name (/Path/To/Asset.Asset:EntryPoint)
+	FName GetDefaultEntryPoint() const;
+
+protected:
+
+	// Loads the graph data from the provided archive buffer and returns true on success, false otherwise
+	bool LoadFromArchiveBuffer(const TArray<uint8>& SharedDataArchiveBuffer);
+
+	// Allocates an instance of the graph with an optional parent graph instance
+	void AllocateInstanceImpl(FAnimNextGraphInstance* InOutParentGraphInstance, FAnimNextGraphInstancePtr& OutInstance, FName InEntryPoint) const;
+
+#if WITH_EDITORONLY_DATA
+	// During graph compilation, if we have existing graph instances, we freeze them by releasing their memory before thawing them
+	// Freezing is a partial release of resources that retains the necessary information to re-create things safely
+	void FreezeGraphInstances();
+
+	// During graph compilation, once compilation is done we thaw existing graph instances to reallocate their memory
+	void ThawGraphInstances();
+#endif
+
+	// Set the default entry point name (unqualified)
+	void SetDefaultEntryPoint(FName InEntryPoint);
+
+	// Cache the default entry point
+	void CacheDefaultEntryPoint() const;
+
+	friend class UAnimNextModuleFactory;
+	friend class UAnimNextModule_EditorData;
+	friend class UAnimNextModule_Parameter;
+	friend struct UE::AnimNext::UncookedOnly::FUtils;
+	friend class UE::AnimNext::Editor::FModuleEditor;
+	friend struct UE::AnimNext::FTestUtils;
+	friend FAnimNextGraphInstancePtr;
+	friend FAnimNextGraphInstance;
+	friend class UAnimGraphNode_AnimNextGraph;
+	friend UE::AnimNext::FExecutionContext;
+	friend class UAnimNextSchedule;
+	friend struct FAnimNextScheduleGraphTask;
+	friend UE::AnimNext::FAnimNextModuleImpl;
+	friend class UE::AnimNext::Editor::FParameterCustomization;
+	friend struct UE::AnimNext::FParametersProxy;
+
+#if WITH_EDITORONLY_DATA
+	mutable FCriticalSection GraphInstancesLock;
+
+	// This is a list of live graph instances that have been allocated, used in the editor to reset instances when we re-compile/live edit
+	mutable TSet<FAnimNextGraphInstance*> GraphInstances;
+#endif
+
+	// This is the execute method definition used by a graph to evaluate latent pins
+	UPROPERTY()
+	FAnimNextGraphEvaluatorExecuteDefinition ExecuteDefinition;
+
+	// Data for each entry point in this graph
+	UPROPERTY()
+	TArray<FAnimNextGraphEntryPoint> EntryPoints;
+
+	// This is a resolved handle to the root trait in our graph, for each entry point 
+	TMap<FName, FAnimNextTraitHandle> ResolvedRootTraitHandles;
+
+	// This is an index into EntryPoints, for each entry point
+	TMap<FName, int32> ResolvedEntryPoints;
+
+	// This is the graph shared data used by the trait system, the output of FTraitReader
+	// We de-serialize manually into this buffer from the archive buffer, this is never saved on disk
+	TArray<uint8> SharedDataBuffer;
+
+	// This is a list of all referenced UObjects in the graph shared data
+	// We collect all the references here to make it quick and easy for the GC to query them
+	// It means that object references in the graph shared data are not visited at runtime by the GC (they are immutable)
+	// The shared data serialization archive stores indices to these to perform UObject serialization
+	UPROPERTY()
+	TArray<TObjectPtr<UObject>> GraphReferencedObjects;
+
+	// The entry point that this graph defaults to using (unqualified by asset). Use GetDefaultEntryPoint to get the fully-qualified name.
+	UPROPERTY(EditAnywhere, Category = "Graph")
+	FName DefaultEntryPoint = FRigUnit_AnimNextGraphRoot::DefaultEntryPoint;
+
+	// Cached fully-qualified entry point, set via by SetDefaultEntryPoint/GetDefaultEntryPoint.
+	UPROPERTY(Transient)
+	mutable FName CachedDefaultEntryPoint = NAME_None;
+
+	// Hash of required parameters
+	UPROPERTY()
+	uint64 RequiredParametersHash = 0;
+
+	// All the parameters that are required for this graph to run
+	UPROPERTY()
+	TArray<FAnimNextParam> RequiredParameters;
+
+	// Default state for this graph
+	UPROPERTY()
+	FAnimNextGraphState DefaultState;
+
+#if WITH_EDITORONLY_DATA
+	UPROPERTY()
+	FInstancedPropertyBag PropertyBag_DEPRECATED;
+
+	// This buffer holds the output of the FTraitWriter post compilation
+	// We serialize it manually and it is discarded at runtime
+	TArray<uint8> SharedDataArchiveBuffer;
+#endif
+};
diff --git a/Engine/Plugins/Experimental/AnimNext/Source/AnimNext/Internal/Module/RigUnit_AnimNextModuleSelf.h b/Engine/Plugins/Experimental/AnimNext/Source/AnimNext/Internal/Module/RigUnit_AnimNextModuleSelf.h
new file mode 100644
--- /dev/null
+++ b/Engine/Plugins/Experimental/AnimNext/Source/AnimNext/Internal/Module/RigUnit_AnimNextModuleSelf.h
@@ -0,0 +1,22 @@
+﻿// Copyright Epic Games, Inc. All Rights Reserved.
+
+#pragma once
+
+#include "Graph/RigUnit_AnimNextBase.h"
+#include "RigUnit_AnimNextModuleSelf.generated.h"
+
+class UAnimNextModule;
+
+/** Get a reference to the currently executing module */
+USTRUCT(meta=(DisplayName="Self", Category="Module", NodeColor="0, 0, 1", Keywords="Current,This"))
+struct ANIMNEXT_API FRigUnit_AnimNextModuleSelf : public FRigUnit_AnimNextBase
+{
+	GENERATED_BODY()
+
+	RIGVM_METHOD()
+	void Execute();
+
+	// The currently-executing module
+	UPROPERTY(meta=(Output))
+	TObjectPtr<UAnimNextModule> Self;
+};
diff --git a/Engine/Plugins/Experimental/AnimNext/Source/AnimNext/Internal/Param/AnimNextParamContextData.h b/Engine/Plugins/Experimental/AnimNext/Source/AnimNext/Internal/Param/AnimNextParamContextData.h
--- a/Engine/Plugins/Experimental/AnimNext/Source/AnimNext/Internal/Param/AnimNextParamContextData.h
+++ b/Engine/Plugins/Experimental/AnimNext/Source/AnimNext/Internal/Param/AnimNextParamContextData.h
@@ -36,6 +36,6 @@ private:
 	// Parameter scope context
 	UE::AnimNext::FParamStackLayerHandle* LayerHandle = nullptr;
 
-	friend class UAnimNextGraph;
+	friend class UAnimNextModule;
 	friend struct FAnimNextExecuteContext;
 };
diff --git a/Engine/Plugins/Experimental/AnimNext/Source/AnimNext/Internal/Param/PropertyBagProxy.h b/Engine/Plugins/Experimental/AnimNext/Source/AnimNext/Internal/Param/PropertyBagProxy.h
--- a/Engine/Plugins/Experimental/AnimNext/Source/AnimNext/Internal/Param/PropertyBagProxy.h
+++ b/Engine/Plugins/Experimental/AnimNext/Source/AnimNext/Internal/Param/PropertyBagProxy.h
@@ -4,7 +4,7 @@
 
 #include "Param/IParameterSource.h"
 #include "PropertyBag.h"
-#include "IAnimNextModule.h"
+#include "IAnimNextModuleInterface.h"
 
 namespace UE::AnimNext
 {
diff --git a/Engine/Plugins/Experimental/AnimNext/Source/AnimNext/Internal/RigVMRuntimeDataRegistry.h b/Engine/Plugins/Experimental/AnimNext/Source/AnimNext/Internal/RigVMRuntimeDataRegistry.h
--- a/Engine/Plugins/Experimental/AnimNext/Source/AnimNext/Internal/RigVMRuntimeDataRegistry.h
+++ b/Engine/Plugins/Experimental/AnimNext/Source/AnimNext/Internal/RigVMRuntimeDataRegistry.h
@@ -78,7 +78,7 @@ private:
 	static void Init();
 	static void Destroy();
 
-	friend class FModule;
+	friend class FAnimNextModuleImpl;
 };
 
 } // end namespace UE::AnimNext
diff --git a/Engine/Plugins/Experimental/AnimNext/Source/AnimNext/Internal/Scheduler/AnimNextSchedule.h b/Engine/Plugins/Experimental/AnimNext/Source/AnimNext/Internal/Scheduler/AnimNextSchedule.h
--- a/Engine/Plugins/Experimental/AnimNext/Source/AnimNext/Internal/Scheduler/AnimNextSchedule.h
+++ b/Engine/Plugins/Experimental/AnimNext/Source/AnimNext/Internal/Scheduler/AnimNextSchedule.h
@@ -14,7 +14,7 @@
 #include "Scheduler/IAnimNextScheduleTermInterface.h"
 #include "AnimNextSchedule.generated.h"
 
-class UAnimNextGraph;
+class UAnimNextModule;
 class UAnimNextSchedule;
 class UAnimNextSchedulerWorldSubsystem;
 class UAnimNextComponent;
@@ -82,12 +82,12 @@ private:
 	// UObject interface
 	virtual void GetPreloadDependencies(TArray<UObject*>& OutDeps) override;
 
-	// The graph to run by default
+	// The module to run by default
 	UPROPERTY(EditAnywhere, Category = "Graph")
-	TObjectPtr<UAnimNextGraph> Graph = nullptr;
+	TObjectPtr<UAnimNextModule> Module = nullptr;
 
 	// Parameter to get the graph from dynamically
-	UPROPERTY(EditAnywhere, Category = "Graph", meta = (AllowedParamType = "TObjectPtr<UAnimNextGraph>"))
+	UPROPERTY(EditAnywhere, Category = "Graph", meta = (AllowedParamType = "TObjectPtr<UAnimNextModule>"))
 	FAnimNextEditorParam DynamicGraph;
 
 	// An optional entry point to use when running the supplied graph
@@ -158,7 +158,7 @@ private:
 
 	// Parameters to apply in this scope
 	UPROPERTY(EditAnywhere, Category = "Parameters")
-	TArray<TObjectPtr<UAnimNextGraph>> Parameters;
+	TArray<TObjectPtr<UAnimNextModule>> Parameters;
 
 	// Entries that are part of this scope
 	UPROPERTY(EditAnywhere, Category = "Parameters", Instanced)
diff --git a/Engine/Plugins/Experimental/AnimNext/Source/AnimNext/Internal/Scheduler/AnimNextScheduleGraphTask.h b/Engine/Plugins/Experimental/AnimNext/Source/AnimNext/Internal/Scheduler/AnimNextScheduleGraphTask.h
--- a/Engine/Plugins/Experimental/AnimNext/Source/AnimNext/Internal/Scheduler/AnimNextScheduleGraphTask.h
+++ b/Engine/Plugins/Experimental/AnimNext/Source/AnimNext/Internal/Scheduler/AnimNextScheduleGraphTask.h
@@ -8,7 +8,7 @@
 #include "Param/ParamId.h"
 #include "AnimNextScheduleGraphTask.generated.h"
 
-class UAnimNextGraph;
+class UAnimNextModule;
 struct FAnimNextEditorParam;
 struct FAnimNextParamUniversalObjectLocator;
 template<typename T> struct TInstancedStruct;
@@ -39,12 +39,12 @@ private:
 	friend struct UE::AnimNext::FScheduleInstanceData;
 	friend struct UE::AnimNext::FScheduleTickFunction;
 
-	void RunGraph(const UE::AnimNext::FScheduleContext& InContext) const;
+	void RunModule(const UE::AnimNext::FScheduleContext& InContext) const;
 
-	UAnimNextGraph* GetGraphToRun(UE::AnimNext::FParamStack& ParamStack) const;
+	UAnimNextModule* GetModuleToRun(UE::AnimNext::FParamStack& ParamStack) const;
 
 	// Verify graph's required parameters are satisfied by this task's supplied parameters
-	void VerifyRequiredParameters(UAnimNextGraph* InGraphToRun) const;
+	void VerifyRequiredParameters(UAnimNextModule* InModuleToRun) const;
 
 private:
 	UPROPERTY()
@@ -60,10 +60,10 @@ private:
 	FAnimNextParam EntryPoint;
 
 	UPROPERTY()
-	TObjectPtr<UAnimNextGraph> Graph = nullptr;
+	TObjectPtr<UAnimNextModule> Module = nullptr;
 
 	UPROPERTY()
-	FAnimNextParam DynamicGraph;
+	FAnimNextParam DynamicModule;
 
 	UPROPERTY()
 	FAnimNextParam ReferencePose;
diff --git a/Engine/Plugins/Experimental/AnimNext/Source/AnimNext/Internal/Scheduler/AnimNextScheduleParamScopeTask.h b/Engine/Plugins/Experimental/AnimNext/Source/AnimNext/Internal/Scheduler/AnimNextScheduleParamScopeTask.h
--- a/Engine/Plugins/Experimental/AnimNext/Source/AnimNext/Internal/Scheduler/AnimNextScheduleParamScopeTask.h
+++ b/Engine/Plugins/Experimental/AnimNext/Source/AnimNext/Internal/Scheduler/AnimNextScheduleParamScopeTask.h
@@ -7,7 +7,7 @@
 #include "AnimNextScheduleParamScopeTask.generated.h"
 
 class UAnimNextSchedule;
-class UAnimNextGraph;
+class UAnimNextModule;
 class UAnimNextSchedulerWorldSubsystem;
 class UAnimNextComponent;
 
@@ -58,7 +58,7 @@ private:
 
 	// Parameters to apply in this scopes
 	UPROPERTY()
-	TArray<TObjectPtr<UAnimNextGraph>> Parameters;
+	TArray<TObjectPtr<UAnimNextModule>> Parameters;
 };
 
 USTRUCT()
diff --git a/Engine/Plugins/Experimental/AnimNext/Source/AnimNext/Internal/Scheduler/Scheduler.h b/Engine/Plugins/Experimental/AnimNext/Source/AnimNext/Internal/Scheduler/Scheduler.h
--- a/Engine/Plugins/Experimental/AnimNext/Source/AnimNext/Internal/Scheduler/Scheduler.h
+++ b/Engine/Plugins/Experimental/AnimNext/Source/AnimNext/Internal/Scheduler/Scheduler.h
@@ -6,7 +6,7 @@
 #include "Scheduler/ScheduleHandle.h"
 
 enum class EAnimNextScheduleInitMethod : uint8;
-class UAnimNextGraph;
+class UAnimNextModule;
 class UAnimNextSchedule;
 
 namespace UE::AnimNext
diff --git a/Engine/Plugins/Experimental/AnimNext/Source/AnimNext/Internal/TraitCore/NodeTemplateRegistry.h b/Engine/Plugins/Experimental/AnimNext/Source/AnimNext/Internal/TraitCore/NodeTemplateRegistry.h
--- a/Engine/Plugins/Experimental/AnimNext/Source/AnimNext/Internal/TraitCore/NodeTemplateRegistry.h
+++ b/Engine/Plugins/Experimental/AnimNext/Source/AnimNext/Internal/TraitCore/NodeTemplateRegistry.h
@@ -60,7 +60,7 @@ namespace UE::AnimNext
 		TArray<uint8>								TemplateBuffer;
 		TMap<uint32, FNodeTemplateRegistryHandle>	TemplateUIDToHandleMap;
 
-		friend class FModule;
+		friend class FAnimNextModuleImpl;
 		friend class FTraitWriter;
 		friend struct FScopedClearNodeTemplateRegistry;
 	};
diff --git a/Engine/Plugins/Experimental/AnimNext/Source/AnimNext/Private/AnimNextModuleImpl.cpp b/Engine/Plugins/Experimental/AnimNext/Source/AnimNext/Private/AnimNextModuleImpl.cpp
new file mode 100644
--- /dev/null
+++ b/Engine/Plugins/Experimental/AnimNext/Source/AnimNext/Private/AnimNextModuleImpl.cpp
@@ -0,0 +1,386 @@
+// Copyright Epic Games, Inc. All Rights Reserved.
+
+#include "AnimNextModuleImpl.h"
+#include "AnimNextConfig.h"
+#include "DataRegistry.h"
+#include "IUniversalObjectLocatorModule.h"
+#include "RigVMRuntimeDataRegistry.h"
+#include "Animation/AnimSequence.h"
+#include "Animation/BlendProfile.h"
+#include "Component/AnimNextComponent.h"
+#include "Curves/CurveFloat.h"
+#include "Module/AnimNextModule.h"
+#include "Modules/ModuleManager.h"
+#include "Param/AnimNextActorLocatorFragment.h"
+#include "Param/AnimNextComponentLocatorFragment.h"
+#include "Param/AnimNextEditorParam.h"
+#include "Param/AnimNextObjectCastLocatorFragment.h"
+#include "Param/AnimNextObjectFunctionLocatorFragment.h"
+#include "Param/AnimNextObjectPropertyLocatorFragment.h"
+#include "Param/AnimNextParam.h"
+#include "Param/AnimNextTag.h"
+#include "Param/ObjectProxyFactory.h"
+#include "Param/PropertyBagProxy.h"
+#include "RigVMCore/RigVMRegistry.h"
+#include "Scheduler/AnimNextTickFunctionBinding.h"
+#include "Scheduler/Scheduler.h"
+#include "TraitCore/NodeTemplateRegistry.h"
+#include "TraitCore/TraitRegistry.h"
+#include "TraitCore/TraitInterfaceRegistry.h"
+
+#if WITH_ANIMNEXT_CONSOLE_COMMANDS
+#include "HAL/IConsoleManager.h"
+#include "UObject/UObjectIterator.h"
+
+#include "TraitCore/NodeDescription.h"
+#include "TraitCore/NodeTemplate.h"
+#include "TraitCore/TraitTemplate.h"
+#endif
+
+#define LOCTEXT_NAMESPACE "AnimNextModule"
+
+namespace UE::AnimNext
+{
+
+	void FAnimNextModuleImpl::StartupModule()
+	{
+		GetMutableDefault<UAnimNextConfig>()->LoadConfig();
+
+		FRigVMRegistry& RigVMRegistry = FRigVMRegistry::Get();
+		static TPair<UClass*, FRigVMRegistry::ERegisterObjectOperation> const AllowedObjectTypes[] =
+		{
+			{ UAnimSequence::StaticClass(), FRigVMRegistry::ERegisterObjectOperation::Class },
+			{ UScriptStruct::StaticClass(), FRigVMRegistry::ERegisterObjectOperation::Class },
+			{ UBlendProfile::StaticClass(), FRigVMRegistry::ERegisterObjectOperation::Class },
+			{ UCurveFloat::StaticClass(), FRigVMRegistry::ERegisterObjectOperation::Class },
+			{ UAnimNextModule::StaticClass(), FRigVMRegistry::ERegisterObjectOperation::Class },
+			{ UAnimNextComponent::StaticClass(), FRigVMRegistry::ERegisterObjectOperation::Class },
+		};
+
+		RigVMRegistry.RegisterObjectTypes(AllowedObjectTypes);
+
+		static UScriptStruct* const AllowedStructTypes[] =
+		{
+			FAnimNextEditorParam::StaticStruct(),
+			FAnimNextParam::StaticStruct(),
+			FAnimNextScope::StaticStruct(),
+			FAnimNextEntryPoint::StaticStruct(),
+			FUniversalObjectLocator::StaticStruct(),
+			FAnimNextTickFunctionBinding::StaticStruct(),
+			FAnimNextGraphReferencePose::StaticStruct()
+		};
+
+		RigVMRegistry.RegisterStructTypes(AllowedStructTypes);
+
+		RegisterParameterSourceFactory("ObjectProxy", MakeShared<FObjectProxyFactory>());
+		FDataRegistry::Init();
+		FTraitRegistry::Init();
+		FTraitInterfaceRegistry::Init();
+		FNodeTemplateRegistry::Init();
+		FScheduler::Init();
+		FRigVMRuntimeDataRegistry::Init();
+
+		UE::UniversalObjectLocator::IUniversalObjectLocatorModule& UolModule = FModuleManager::Get().LoadModuleChecked<UE::UniversalObjectLocator::IUniversalObjectLocatorModule>("UniversalObjectLocator");
+		FDelayedAutoRegisterHelper(EDelayedRegisterRunPhase::ObjectSystemReady,
+			[&UolModule]
+			{
+				{
+					UE::UniversalObjectLocator::FFragmentTypeParameters FragmentTypeParams("animobjfunc", LOCTEXT("AnimNextObjectFunctionFragment", "Function"));
+					FragmentTypeParams.PrimaryEditorType = "AnimNextObjectFunction";
+					FAnimNextObjectFunctionLocatorFragment::FragmentType = UolModule.RegisterFragmentType<FAnimNextObjectFunctionLocatorFragment>(FragmentTypeParams);
+				}
+				{
+					UE::UniversalObjectLocator::FFragmentTypeParameters FragmentTypeParams("animobjprop", LOCTEXT("AnimNextObjectPropertyFragment", "Property"));
+					FragmentTypeParams.PrimaryEditorType = "AnimNextObjectProperty";
+					FAnimNextObjectPropertyLocatorFragment::FragmentType = UolModule.RegisterFragmentType<FAnimNextObjectPropertyLocatorFragment>(FragmentTypeParams);
+				}
+				{
+					UE::UniversalObjectLocator::FFragmentTypeParameters FragmentTypeParams("animobjcast", LOCTEXT("AnimNextCastFragment", "Cast"));
+					FragmentTypeParams.PrimaryEditorType = "AnimNextObjectCast";
+					FAnimNextObjectCastLocatorFragment::FragmentType = UolModule.RegisterFragmentType<FAnimNextObjectCastLocatorFragment>(FragmentTypeParams);
+				}
+				{
+					UE::UniversalObjectLocator::FFragmentTypeParameters FragmentTypeParams("animcomp", LOCTEXT("AnimNextComponentFragment", "AnimNextComponent"));
+					FragmentTypeParams.PrimaryEditorType = "AnimNextComponent";
+					FAnimNextComponentLocatorFragment::FragmentType = UolModule.RegisterFragmentType<FAnimNextComponentLocatorFragment>(FragmentTypeParams);
+				}
+				{
+					UE::UniversalObjectLocator::FFragmentTypeParameters FragmentTypeParams("animactor", LOCTEXT("AnimNextActorFragment", "AnimNextActor"));
+					FragmentTypeParams.PrimaryEditorType = "AnimNextActor";
+					FAnimNextActorLocatorFragment::FragmentType = UolModule.RegisterFragmentType<FAnimNextActorLocatorFragment>(FragmentTypeParams);
+				}
+			});
+#if WITH_ANIMNEXT_CONSOLE_COMMANDS
+		if (!IsRunningCommandlet())
+		{
+			ConsoleCommands.Add(IConsoleManager::Get().RegisterConsoleCommand(
+				TEXT("AnimNext.ListNodeTemplates"),
+				TEXT("Dumps statistics about node templates to the log."),
+				FConsoleCommandWithArgsDelegate::CreateRaw(this, &FAnimNextModuleImpl::ListNodeTemplates),
+				ECVF_Default
+			));
+			ConsoleCommands.Add(IConsoleManager::Get().RegisterConsoleCommand(
+				TEXT("AnimNext.Module"),
+				TEXT("Dumps statistics about modules to the log."),
+				FConsoleCommandWithArgsDelegate::CreateRaw(this, &FAnimNextModuleImpl::ListModules),
+				ECVF_Default
+			));
+		}
+#endif
+	}
+
+	void FAnimNextModuleImpl::ShutdownModule()
+	{
+		FRigVMRuntimeDataRegistry::Destroy();
+		FScheduler::Destroy();
+		FNodeTemplateRegistry::Destroy();
+		FTraitInterfaceRegistry::Destroy();
+		FTraitRegistry::Destroy();
+		FDataRegistry::Destroy();
+		UnregisterParameterSourceFactory("ObjectProxy");
+
+#if WITH_ANIMNEXT_CONSOLE_COMMANDS
+		for (IConsoleObject* Cmd : ConsoleCommands)
+		{
+			IConsoleManager::Get().UnregisterConsoleObject(Cmd);
+		}
+		ConsoleCommands.Empty();
+#endif
+	}
+
+	const IAnimNextAnimGraph* AnimGraphImpl = nullptr;
+
+	void FAnimNextModuleImpl::RegisterAnimNextAnimGraph(const IAnimNextAnimGraph& InAnimGraphImpl)
+	{
+		AnimGraphImpl = &InAnimGraphImpl;
+	}
+
+	void FAnimNextModuleImpl::UnregisterAnimNextAnimGraph()
+	{
+		AnimGraphImpl = nullptr;
+	}
+
+	void FAnimNextModuleImpl::UpdateGraph(FAnimNextGraphInstancePtr& GraphInstance, float DeltaTime, FTraitEventList& InputEventList, FTraitEventList& OutputEventList)
+	{
+		if (AnimGraphImpl != nullptr)
+		{
+			AnimGraphImpl->UpdateGraph(GraphInstance, DeltaTime, InputEventList, OutputEventList);
+		}
+	}
+
+	void FAnimNextModuleImpl::EvaluateGraph(FAnimNextGraphInstancePtr& GraphInstance, const UE::AnimNext::FReferencePose& RefPose, int32 GraphLODLevel, FLODPoseHeap& OutputPose) const
+	{
+		if (AnimGraphImpl != nullptr)
+		{
+			AnimGraphImpl->EvaluateGraph(GraphInstance, RefPose, GraphLODLevel, OutputPose);
+		}
+	}
+
+	TUniquePtr<IParameterSource> FAnimNextModuleImpl::CreateParameterSource(const FParameterSourceContext& InContext, const TInstancedStruct<FAnimNextParamInstanceIdentifier>& InInstanceId, TConstArrayView<FName> InRequiredParameters) const
+	{
+		for(const TPair<FName, TSharedRef<IParameterSourceFactory>>& FactoryPair : ParameterSourceFactories)
+		{
+			// TODO: in future we could introduce a priority system or more complex logic here. For now it is earliest wins as are sources are
+			// all mutually exclusive at the moment.
+			TUniquePtr<IParameterSource> ParameterSource = FactoryPair.Value->CreateParameterSource(InContext, InInstanceId, InRequiredParameters);
+			if(ParameterSource.IsValid())
+			{
+				return ParameterSource;
+			}
+		}
+
+		return nullptr;
+	}
+
+	void FAnimNextModuleImpl::RegisterParameterSourceFactory(FName InName, TSharedRef<IParameterSourceFactory> InFactory)
+	{
+		ParameterSourceFactories.Add(InName, InFactory);
+	}
+
+	void FAnimNextModuleImpl::UnregisterParameterSourceFactory(FName InName)
+	{
+		ParameterSourceFactories.Remove(InName);
+	}
+
+	TSharedPtr<IParameterSourceFactory> FAnimNextModuleImpl::FindParameterSourceFactory(FName InName)
+	{
+		if(TSharedRef<IParameterSourceFactory>* FoundFactory = ParameterSourceFactories.Find(InName))
+		{
+			return *FoundFactory;
+		}
+
+		return nullptr;
+	}
+
+#if WITH_ANIMNEXT_CONSOLE_COMMANDS
+	void FAnimNextModuleImpl::ListNodeTemplates(const TArray<FString>& Args)
+	{
+		// Turn off log times to make diff-ing easier
+		TGuardValue<ELogTimes::Type> DisableLogTimes(GPrintLogTimes, ELogTimes::None);
+
+		// Make sure to log everything
+		const ELogVerbosity::Type OldVerbosity = LogAnimation.GetVerbosity();
+		LogAnimation.SetVerbosity(ELogVerbosity::All);
+
+		const FNodeTemplateRegistry& NodeTemplateRegistry = FNodeTemplateRegistry::Get();
+		const FTraitRegistry& TraitRegistry = FTraitRegistry::Get();
+
+		UE_LOG(LogAnimation, Log, TEXT("===== AnimNext Node Templates ====="));
+		UE_LOG(LogAnimation, Log, TEXT("Template Buffer Size: %u bytes"), NodeTemplateRegistry.TemplateBuffer.GetAllocatedSize());
+
+		for (auto It = NodeTemplateRegistry.TemplateUIDToHandleMap.CreateConstIterator(); It; ++It)
+		{
+			const FNodeTemplateRegistryHandle Handle = It.Value();
+			const FNodeTemplate* NodeTemplate = NodeTemplateRegistry.Find(Handle);
+
+			const uint32 NumTraits = NodeTemplate->GetNumTraits();
+
+			UE_LOG(LogAnimation, Log, TEXT("[%x] has %u traits ..."), NodeTemplate->GetUID(), NumTraits);
+			UE_LOG(LogAnimation, Log, TEXT("    Template Size: %u bytes"), NodeTemplate->GetNodeTemplateSize());
+			UE_LOG(LogAnimation, Log, TEXT("    Shared Data Size: %u bytes"), NodeTemplate->GetNodeSharedDataSize());
+			UE_LOG(LogAnimation, Log, TEXT("    Instance Data Size: %u bytes"), NodeTemplate->GetNodeInstanceDataSize());
+			UE_LOG(LogAnimation, Log, TEXT("    Traits ..."));
+
+			const FTraitTemplate* TraitTemplates = NodeTemplate->GetTraits();
+			for (uint32 TraitIndex = 0; TraitIndex < NumTraits; ++TraitIndex)
+			{
+				const FTraitTemplate* TraitTemplate = TraitTemplates + TraitIndex;
+				const FTrait* Trait = TraitRegistry.Find(TraitTemplate->GetRegistryHandle());
+				const FString TraitName = Trait != nullptr ? Trait->GetTraitName() : TEXT("<Unknown>");
+
+				const uint32 NextTraitIndex = TraitIndex + 1;
+				const uint32 EndOfNextTraitSharedData = NextTraitIndex < NumTraits ? TraitTemplates[NextTraitIndex].GetNodeSharedOffset() : NodeTemplate->GetNodeSharedDataSize();
+				const uint32 TraitSharedDataSize = EndOfNextTraitSharedData - TraitTemplate->GetNodeSharedOffset();
+
+				const uint32 EndOfNextTraitInstanceData = NextTraitIndex < NumTraits ? TraitTemplates[NextTraitIndex].GetNodeInstanceOffset() : NodeTemplate->GetNodeInstanceDataSize();
+				const uint32 TraitInstanceDataSize = EndOfNextTraitInstanceData - TraitTemplate->GetNodeInstanceOffset();
+
+				UE_LOG(LogAnimation, Log, TEXT("            %u: [%x] %s (%s)"), TraitIndex, TraitTemplate->GetUID().GetUID(), *TraitName, TraitTemplate->GetMode() == ETraitMode::Base ? TEXT("Base") : TEXT("Additive"));
+				UE_LOG(LogAnimation, Log, TEXT("                Shared Data: [Offset: %u bytes, Size: %u bytes]"), TraitTemplate->GetNodeSharedOffset(), TraitSharedDataSize);
+				if (TraitTemplate->HasLatentProperties() && Trait != nullptr)
+				{
+					UE_LOG(LogAnimation, Log, TEXT("                Shared Data Latent Property Handles: [Offset: %u bytes, Count: %u]"), TraitTemplate->GetNodeSharedLatentPropertyHandlesOffset(), Trait->GetNumLatentTraitProperties());
+				}
+				UE_LOG(LogAnimation, Log, TEXT("                Instance Data: [Offset: %u bytes, Size: %u bytes]"), TraitTemplate->GetNodeInstanceOffset(), TraitInstanceDataSize);
+			}
+		}
+
+		LogAnimation.SetVerbosity(OldVerbosity);
+	}
+
+	void FAnimNextModuleImpl::ListModules(const TArray<FString>& Args)
+	{
+		// Turn off log times to make diff-ing easier
+		TGuardValue<ELogTimes::Type> DisableLogTimes(GPrintLogTimes, ELogTimes::None);
+
+		// Make sure to log everything
+		const ELogVerbosity::Type OldVerbosity = LogAnimation.GetVerbosity();
+		LogAnimation.SetVerbosity(ELogVerbosity::All);
+
+		TArray<const UAnimNextModule*> Modules;
+
+		for (TObjectIterator<UAnimNextModule> It; It; ++It)
+		{
+			Modules.Add(*It);
+		}
+
+		struct FCompareObjectNames
+		{
+			FORCEINLINE bool operator()(const UAnimNextModule& Lhs, const UAnimNextModule& Rhs) const
+			{
+				return Lhs.GetPathName().Compare(Rhs.GetPathName()) < 0;
+			}
+		};
+		Modules.Sort(FCompareObjectNames());
+
+		const FNodeTemplateRegistry& NodeTemplateRegistry = FNodeTemplateRegistry::Get();
+		const FTraitRegistry& TraitRegistry = FTraitRegistry::Get();
+		const bool bDetailedOutput = true;
+
+		UE_LOG(LogAnimation, Log, TEXT("===== AnimNext Modules ====="));
+		UE_LOG(LogAnimation, Log, TEXT("Num Graphs: %u"), Modules.Num());
+
+		for (const UAnimNextModule* Module : Modules)
+		{
+			uint32 TotalInstanceSize = 0;
+			uint32 NumNodes = 0;
+			{
+				// We always have a node at offset 0
+				int32 NodeOffset = 0;
+
+				while (NodeOffset < Module->SharedDataBuffer.Num())
+				{
+					const FNodeDescription* NodeDesc = reinterpret_cast<const FNodeDescription*>(&Module->SharedDataBuffer[NodeOffset]);
+
+					TotalInstanceSize += NodeDesc->GetNodeInstanceDataSize();
+					NumNodes++;
+
+					const FNodeTemplate* NodeTemplate = NodeTemplateRegistry.Find(NodeDesc->GetTemplateHandle());
+					NodeOffset += NodeTemplate->GetNodeSharedDataSize();
+				}
+			}
+
+			UE_LOG(LogAnimation, Log, TEXT("    %s ..."), *Module->GetPathName());
+			UE_LOG(LogAnimation, Log, TEXT("        Shared Data Size: %.2f KB"), double(Module->SharedDataBuffer.Num()) / 1024.0);
+			UE_LOG(LogAnimation, Log, TEXT("        Max Instance Data Size: %.2f KB"), double(TotalInstanceSize) / 1024.0);
+			UE_LOG(LogAnimation, Log, TEXT("        Num Nodes: %u"), NumNodes);
+
+			if (bDetailedOutput)
+			{
+				// We always have a node at offset 0
+				int32 NodeOffset = 0;
+
+				while (NodeOffset < Module->SharedDataBuffer.Num())
+				{
+					const FNodeDescription* NodeDesc = reinterpret_cast<const FNodeDescription*>(&Module->SharedDataBuffer[NodeOffset]);
+					const FNodeTemplate* NodeTemplate = NodeTemplateRegistry.Find(NodeDesc->GetTemplateHandle());
+
+					const uint32 NumTraits = NodeTemplate->GetNumTraits();
+
+					UE_LOG(LogAnimation, Log, TEXT("        Node %u: [Template %x with %u traits]"), NodeDesc->GetUID().GetNodeIndex(), NodeTemplate->GetUID(), NumTraits);
+					UE_LOG(LogAnimation, Log, TEXT("            Shared Data: [Offset: %u bytes, Size: %u bytes]"), NodeOffset, NodeTemplate->GetNodeSharedDataSize());
+					UE_LOG(LogAnimation, Log, TEXT("            Instance Data Size: %u bytes"), NodeDesc->GetNodeInstanceDataSize());
+					UE_LOG(LogAnimation, Log, TEXT("            Traits ..."));
+
+					const FTraitTemplate* TraitTemplates = NodeTemplate->GetTraits();
+					for (uint32 TraitIndex = 0; TraitIndex < NumTraits; ++TraitIndex)
+					{
+						const FTraitTemplate* TraitTemplate = TraitTemplates + TraitIndex;
+						const FTrait* Trait = TraitRegistry.Find(TraitTemplate->GetRegistryHandle());
+						const FString TraitName = Trait != nullptr ? Trait->GetTraitName() : TEXT("<Unknown>");
+
+						const uint32 NextTraitIndex = TraitIndex + 1;
+						const uint32 EndOfNextTraitSharedData = NextTraitIndex < NumTraits ? TraitTemplates[NextTraitIndex].GetNodeSharedOffset() : NodeTemplate->GetNodeSharedDataSize();
+						const uint32 TraitSharedDataSize = EndOfNextTraitSharedData - TraitTemplate->GetNodeSharedOffset();
+
+						const uint32 EndOfNextTraitInstanceData = NextTraitIndex < NumTraits ? TraitTemplates[NextTraitIndex].GetNodeInstanceOffset() : NodeTemplate->GetNodeInstanceDataSize();
+						const uint32 TraitInstanceDataSize = EndOfNextTraitInstanceData - TraitTemplate->GetNodeInstanceOffset();
+
+						UE_LOG(LogAnimation, Log, TEXT("                    %u: [%x] %s (%s)"), TraitIndex, TraitTemplate->GetUID().GetUID(), *TraitName, TraitTemplate->GetMode() == ETraitMode::Base ? TEXT("Base") : TEXT("Additive"));
+						UE_LOG(LogAnimation, Log, TEXT("                        Shared Data: [Offset: %u bytes, Size: %u bytes]"), TraitTemplate->GetNodeSharedOffset(), TraitSharedDataSize);
+						if (TraitTemplate->HasLatentProperties() && Trait != nullptr)
+						{
+							UE_LOG(LogAnimation, Log, TEXT("                        Shared Data Latent Property Handles: [Offset: %u bytes, Count: %u]"), TraitTemplate->GetNodeSharedLatentPropertyHandlesOffset(), Trait->GetNumLatentTraitProperties());
+						}
+						UE_LOG(LogAnimation, Log, TEXT("                        Instance Data: [Offset: %u bytes, Size: %u bytes]"), TraitTemplate->GetNodeInstanceOffset(), TraitInstanceDataSize);
+					}
+
+					NodeOffset += NodeTemplate->GetNodeSharedDataSize();
+				}
+			}
+		}
+
+		LogAnimation.SetVerbosity(OldVerbosity);
+	}
+#endif
+
+	IAnimNextModuleInterface& IAnimNextModuleInterface::Get()
+	{
+		return FModuleManager::LoadModuleChecked<IAnimNextModuleInterface>(TEXT("AnimNext"));
+	}
+
+}
+
+#undef LOCTEXT_NAMESPACE
+
+IMPLEMENT_MODULE(UE::AnimNext::FAnimNextModuleImpl, AnimNext)
diff --git a/Engine/Plugins/Experimental/AnimNext/Source/AnimNext/Private/Component/AnimNextComponent.cpp b/Engine/Plugins/Experimental/AnimNext/Source/AnimNext/Private/Component/AnimNextComponent.cpp
--- a/Engine/Plugins/Experimental/AnimNext/Source/AnimNext/Private/Component/AnimNextComponent.cpp
+++ b/Engine/Plugins/Experimental/AnimNext/Source/AnimNext/Private/Component/AnimNextComponent.cpp
@@ -8,7 +8,7 @@
 #include "Param/PropertyBagProxy.h"
 #include "Scheduler/ScheduleTaskContext.h"
 #include "Scheduler/ScheduleInitializationContext.h"
-#include "IAnimNextModule.h"
+#include "IAnimNextModuleInterface.h"
 
 namespace UE::AnimNext::Private
 {
diff --git a/Engine/Plugins/Experimental/AnimNext/Source/AnimNext/Private/Graph/AnimGraph/AnimNode_AnimNextParameters.cpp b/Engine/Plugins/Experimental/AnimNext/Source/AnimNext/Private/Graph/AnimGraph/AnimNode_AnimNextParameters.cpp
--- a/Engine/Plugins/Experimental/AnimNext/Source/AnimNext/Private/Graph/AnimGraph/AnimNode_AnimNextParameters.cpp
+++ b/Engine/Plugins/Experimental/AnimNext/Source/AnimNext/Private/Graph/AnimGraph/AnimNode_AnimNextParameters.cpp
@@ -2,7 +2,7 @@
 
 #include "Graph/AnimGraph/AnimNode_AnimNextParameters.h"
 #include "Param/ParamStack.h"
-#include "Graph/AnimNextGraph.h"
+#include "Module/AnimNextModule.h"
 #include "Graph/AnimGraphParamStackScope.h"
 #include "Param/ParametersProxy.h"
 
@@ -57,7 +57,7 @@ void FAnimNode_AnimNextParameters::Update_AnyThread(const FAnimationUpdateContex
 
 	GetEvaluateGraphExposedInputs().Execute(Context);
 
-	UAnimNextGraph* CurrentParameters = Parameters;
+	UAnimNextModule* CurrentParameters = Parameters;
 
 	// Reconstruct parameter's cached layer if required
 	if (CurrentParameters != PreviousParameters || !ParametersProxy.IsValid())
diff --git a/Engine/Plugins/Experimental/AnimNext/Source/AnimNext/Private/Graph/AnimNextGraphInstance.cpp b/Engine/Plugins/Experimental/AnimNext/Source/AnimNext/Private/Graph/AnimNextGraphInstance.cpp
--- a/Engine/Plugins/Experimental/AnimNext/Source/AnimNext/Private/Graph/AnimNextGraphInstance.cpp
+++ b/Engine/Plugins/Experimental/AnimNext/Source/AnimNext/Private/Graph/AnimNextGraphInstance.cpp
@@ -4,7 +4,7 @@
 
 #include "AnimNextStats.h"
 #include "TraitCore/ExecutionContext.h"
-#include "Graph/AnimNextGraph.h"
+#include "Module/AnimNextModule.h"
 #include "Graph/GC_GraphInstanceComponent.h"
 #include "Graph/RigUnit_AnimNextShimRoot.h"
 #include "Misc/ScopeRWLock.h"
@@ -32,7 +32,7 @@ void FAnimNextGraphInstance::Release()
 	RootGraphInstance = nullptr;
 	ExtendedExecuteContext.Reset();
 	Components.Empty();
-	Graph = nullptr;
+	Module = nullptr;
 	GraphState = nullptr;
 }
 
@@ -41,9 +41,9 @@ bool FAnimNextGraphInstance::IsValid() const
 	return GraphInstancePtr.IsValid();
 }
 
-const UAnimNextGraph* FAnimNextGraphInstance::GetGraph() const
+const UAnimNextModule* FAnimNextGraphInstance::GetModule() const
 {
-	return Graph;
+	return Module;
 }
 
 FName FAnimNextGraphInstance::GetEntryPoint() const
@@ -66,18 +66,18 @@ FAnimNextGraphInstance* FAnimNextGraphInstance::GetRootGraphInstance() const
 	return RootGraphInstance;
 }
 
-bool FAnimNextGraphInstance::UsesGraph(const UAnimNextGraph* InGraph) const
+bool FAnimNextGraphInstance::UsesModule(const UAnimNextModule* InModule) const
 {
-	return Graph == InGraph;
+	return Module == InModule;
 }
 
 bool FAnimNextGraphInstance::UsesEntryPoint(FName InEntryPoint) const
 {
-	if(Graph != nullptr)
+	if(Module != nullptr)
 	{
 		if(InEntryPoint == NAME_None)
 		{
-			return EntryPoint == Graph->DefaultEntryPoint;
+			return EntryPoint == Module->DefaultEntryPoint;
 		}
 
 		return InEntryPoint == EntryPoint;
@@ -147,7 +147,7 @@ void FAnimNextGraphInstance::ExecuteLatentPins(const TConstArrayView<UE::AnimNex
 		return;
 	}
 
-	if (URigVM* VM = Graph->VM)
+	if (URigVM* VM = Module->VM)
 	{
 		FAnimNextExecuteContext& AnimNextContext = ExtendedExecuteContext.GetPublicDataSafe<FAnimNextExecuteContext>();
 		AnimNextContext.SetContextData<FAnimNextGraphContextData>(this, LatentHandles, DestinationBasePtr, bIsFrozen);
@@ -175,16 +175,16 @@ void FAnimNextGraphInstance::Freeze()
 
 void FAnimNextGraphInstance::Thaw()
 {
-	if (const UAnimNextGraph* GraphPtr = Graph)
+	if (const UAnimNextModule* ModulePtr = Module)
 	{
-		GraphState = MakeUnique<UE::AnimNext::FParametersProxy>(Graph);
+		GraphState = MakeUnique<UE::AnimNext::FParametersProxy>(Module);
 
-		ExtendedExecuteContext.CopyMemoryStorage(GraphPtr->ExtendedExecuteContext);
-		GraphPtr->VM->InitializeInstance(ExtendedExecuteContext);
+		ExtendedExecuteContext.CopyMemoryStorage(ModulePtr->ExtendedExecuteContext);
+		ModulePtr->VM->InitializeInstance(ExtendedExecuteContext);
 
 		{
 			UE::AnimNext::FExecutionContext Context(*this);
-			if(const FAnimNextTraitHandle* FoundHandle = GraphPtr->ResolvedRootTraitHandles.Find(EntryPoint))
+			if(const FAnimNextTraitHandle* FoundHandle = ModulePtr->ResolvedRootTraitHandles.Find(EntryPoint))
 			{
 				GraphInstancePtr = Context.AllocateNodeInstance(*this, *FoundHandle);
 			}
diff --git a/Engine/Plugins/Experimental/AnimNext/Source/AnimNext/Private/Graph/AnimNextGraphInstancePtr.cpp b/Engine/Plugins/Experimental/AnimNext/Source/AnimNext/Private/Graph/AnimNextGraphInstancePtr.cpp
--- a/Engine/Plugins/Experimental/AnimNext/Source/AnimNext/Private/Graph/AnimNextGraphInstancePtr.cpp
+++ b/Engine/Plugins/Experimental/AnimNext/Source/AnimNext/Private/Graph/AnimNextGraphInstancePtr.cpp
@@ -2,7 +2,7 @@
 
 #include "Graph/AnimNextGraphInstancePtr.h"
 
-#include "Graph/AnimNextGraph.h"
+#include "Module/AnimNextModule.h"
 #include "Graph/AnimNextGraphInstance.h"
 #include "Param/IParameterSource.h"
 
@@ -21,9 +21,9 @@ void FAnimNextGraphInstancePtr::Release()
 	{
 #if WITH_EDITORONLY_DATA
 		{
-			const UAnimNextGraph* Graph = Impl->GetGraph();
-			FScopeLock Lock(&Graph->GraphInstancesLock);
-			Graph->GraphInstances.Remove(Impl.Get());
+			const UAnimNextModule* Module = Impl->GetModule();
+			FScopeLock Lock(&Module->GraphInstancesLock);
+			Module->GraphInstances.Remove(Impl.Get());
 		}
 #endif
 
@@ -40,9 +40,9 @@ bool FAnimNextGraphInstancePtr::IsValid() const
 	return Impl && Impl->IsValid();
 }
 
-const UAnimNextGraph* FAnimNextGraphInstancePtr::GetGraph() const
+const UAnimNextModule* FAnimNextGraphInstancePtr::GetModule() const
 {
-	return Impl ? Impl->GetGraph() : nullptr;
+	return Impl ? Impl->GetModule() : nullptr;
 }
 
 UE::AnimNext::FWeakTraitPtr FAnimNextGraphInstancePtr::GetGraphRootPtr() const
@@ -55,9 +55,9 @@ FAnimNextGraphInstance* FAnimNextGraphInstancePtr::GetImpl() const
 	return Impl.Get();
 }
 
-bool FAnimNextGraphInstancePtr::UsesGraph(const UAnimNextGraph* InGraph) const
+bool FAnimNextGraphInstancePtr::UsesModule(const UAnimNextModule* InModule) const
 {
-	return Impl ? Impl->UsesGraph(InGraph) : false;
+	return Impl ? Impl->UsesModule(InModule) : false;
 }
 
 bool FAnimNextGraphInstancePtr::IsRoot() const
diff --git a/Engine/Plugins/Experimental/AnimNext/Source/AnimNext/Private/Module/AnimNextModule.cpp b/Engine/Plugins/Experimental/AnimNext/Source/AnimNext/Private/Module/AnimNextModule.cpp
new file mode 100644
--- /dev/null
+++ b/Engine/Plugins/Experimental/AnimNext/Source/AnimNext/Private/Module/AnimNextModule.cpp
@@ -0,0 +1,308 @@
+// Copyright Epic Games, Inc. All Rights Reserved.
+
+#include "Module/AnimNextModule.h"
+#include "RigVMCore/RigVMMemoryStorage.h"
+#include "RigVMCore/RigVMExecuteContext.h"
+#include "Graph/RigUnit_AnimNextShimRoot.h"
+#include "UObject/AssetRegistryTagsContext.h"
+#include "UObject/ObjectSaveContext.h"
+#include "UObject/Package.h"
+#include "Param/ParamTypeHandle.h"
+#include "TraitCore/TraitReader.h"
+#include "TraitCore/ExecutionContext.h"
+#include "Graph/AnimNext_LODPose.h"
+#include "Graph/GC_GraphInstanceComponent.h"
+#include "Graph/AnimNextGraphInstance.h"
+#include "Serialization/MemoryReader.h"
+#include "AnimNextStats.h"
+#include "RigVMRuntimeDataRegistry.h"
+#include "Graph/AnimNextGraphEntryPoint.h"
+#include "Graph/RigUnit_AnimNextBeginExecution.h"
+#include "Graph/RigUnit_AnimNextGraphRoot.h"
+#include "Param/ParametersProxy.h"
+#include "UObject/LinkerLoad.h"
+#include "UObject/ObjectResource.h"
+
+#if WITH_EDITOR	
+#include "Engine/ExternalAssetDependencyGatherer.h"
+REGISTER_ASSETDEPENDENCY_GATHERER(FExternalAssetDependencyGatherer, UAnimNextModule);
+#endif // WITH_EDITOR
+
+DEFINE_STAT(STAT_AnimNext_Graph_AllocateInstance);
+DEFINE_STAT(STAT_AnimNext_Graph_UpdateParamLayer);
+
+#include UE_INLINE_GENERATED_CPP_BY_NAME(AnimNextModule)
+
+UAnimNextModule::UAnimNextModule(const FObjectInitializer& ObjectInitializer)
+	: Super(ObjectInitializer)
+{
+	ExtendedExecuteContext.SetContextPublicDataStruct(FAnimNextExecuteContext::StaticStruct());
+}
+
+void UAnimNextModule::AllocateInstance(FAnimNextGraphInstancePtr& Instance, FName InEntryPoint) const
+{
+	AllocateInstanceImpl(nullptr, Instance, InEntryPoint);
+}
+
+void UAnimNextModule::AllocateInstance(FAnimNextGraphInstance& ParentGraphInstance, FAnimNextGraphInstancePtr& Instance, FName InEntryPoint) const
+{
+	AllocateInstanceImpl(&ParentGraphInstance, Instance, InEntryPoint);
+}
+
+void UAnimNextModule::AllocateInstanceImpl(FAnimNextGraphInstance* ParentGraphInstance, FAnimNextGraphInstancePtr& Instance, FName InEntryPoint) const
+{
+	SCOPE_CYCLE_COUNTER(STAT_AnimNext_Graph_AllocateInstance);
+
+	Instance.Release();
+
+	const FName EntryPoint = (InEntryPoint == NAME_None) ? GetDefaultEntryPoint() : InEntryPoint;
+	const FAnimNextTraitHandle ResolvedRootTraitHandle = ResolvedRootTraitHandles.FindRef(EntryPoint);
+	if (!ResolvedRootTraitHandle.IsValid())
+	{
+		return;
+	}
+
+	{
+		TUniquePtr<FAnimNextGraphInstance> InstanceImpl = MakeUnique<FAnimNextGraphInstance>();
+
+		InstanceImpl->Module = this;
+		InstanceImpl->ParentGraphInstance = ParentGraphInstance;
+		InstanceImpl->EntryPoint = EntryPoint;
+		InstanceImpl->GraphState = MakeUnique<UE::AnimNext::FParametersProxy>(this);
+
+		// If we have a parent graph, use its root since we share the same root, otherwise if we have no parent, we are the root
+		InstanceImpl->RootGraphInstance = ParentGraphInstance != nullptr ? ParentGraphInstance->GetRootGraphInstance() : InstanceImpl.Get();
+
+		InstanceImpl->ExtendedExecuteContext.CopyMemoryStorage(ExtendedExecuteContext);
+		VM->InitializeInstance(InstanceImpl->ExtendedExecuteContext);
+
+		// Move our implementation so that we can use the instance to allocate our root node
+		Instance.Impl = MoveTemp(InstanceImpl);
+	}
+
+	{
+		UE::AnimNext::FExecutionContext Context(Instance);
+		Instance.Impl->GraphInstancePtr = Context.AllocateNodeInstance(*Instance.Impl, ResolvedRootTraitHandle);
+	}
+
+	if (!Instance.IsValid())
+	{
+		// We failed to allocate our instance, clear everything
+		Instance.Release();
+	}
+
+#if WITH_EDITORONLY_DATA
+	if (Instance.IsValid())
+	{
+		FScopeLock Lock(&GraphInstancesLock);
+		check(!GraphInstances.Contains(Instance.Impl.Get()));
+		GraphInstances.Add(Instance.Impl.Get());
+	}
+#endif
+}
+
+void UAnimNextModule::UpdateLayer(UE::AnimNext::FParamStackLayerHandle& InHandle, float InDeltaTime) const
+{
+	SCOPE_CYCLE_COUNTER(STAT_AnimNext_Graph_UpdateParamLayer);
+	
+	if (VM)
+	{
+		if (TSharedPtr<UE::AnimNext::FRigVMRuntimeData> RuntimeData = UE::AnimNext::FRigVMRuntimeDataRegistry::FindOrAddLocalRuntimeData(VM, GetRigVMExtendedExecuteContext()).Pin())
+		{
+			FRigVMExtendedExecuteContext& Context = RuntimeData->Context;
+
+			check(Context.VMHash == VM->GetVMHash());
+
+			FAnimNextExecuteContext& AnimNextContext = Context.GetPublicDataSafe<FAnimNextExecuteContext>();
+
+			// Parameter setup
+			AnimNextContext.SetContextData<FAnimNextParamContextData>(InHandle);
+
+			// RigVM setup
+			AnimNextContext.SetDeltaTime(InDeltaTime);
+
+			VM->ExecuteVM(Context, FRigUnit_AnimNextBeginExecution::EventName);
+
+			// Reset the context to avoid issues if we forget to reset it the next time we use it
+			AnimNextContext.DebugReset<FAnimNextParamContextData>();
+		}
+	}
+}
+
+void UAnimNextModule::Serialize(FArchive& Ar)
+{
+	Super::Serialize(Ar);
+
+	if (Ar.IsLoading())
+	{
+		if(Ar.CustomVer(FFortniteMainBranchObjectVersion::GUID) < FFortniteMainBranchObjectVersion::AnimNextCombineParameterBlocksAndGraphs)
+		{
+			// Skip over shared archive buffer if we are loading from an older version
+			if (const FLinkerLoad* Linker = GetLinker())
+			{
+				const int32 LinkerIndex = GetLinkerIndex();
+				const FObjectExport& Export = Linker->ExportMap[LinkerIndex];
+				Ar.Seek(Export.SerialOffset + Export.SerialSize);
+			}
+		}
+		else
+		{
+			int32 SharedDataArchiveBufferSize = 0;
+			Ar << SharedDataArchiveBufferSize;
+
+#if !WITH_EDITORONLY_DATA
+			// When editor data isn't present, we don't persist the archive buffer as it is only needed on load
+			// to populate the graph shared data
+			TArray<uint8> SharedDataArchiveBuffer;
+#endif
+
+			SharedDataArchiveBuffer.SetNumUninitialized(SharedDataArchiveBufferSize);
+			Ar.Serialize(SharedDataArchiveBuffer.GetData(), SharedDataArchiveBufferSize);
+
+			if (Ar.IsLoadingFromCookedPackage())
+			{
+				// If we are cooked, we populate our graph shared data otherwise in the editor we'll compile on load
+				// and re-populate everything then to account for changes in code/content
+				LoadFromArchiveBuffer(SharedDataArchiveBuffer);
+			}
+		}
+
+		CacheDefaultEntryPoint();
+	}
+	else if (Ar.IsSaving())
+	{
+#if WITH_EDITORONLY_DATA
+		// We only save the archive buffer, if code changes we'll be able to de-serialize from it when
+		// building the runtime buffer
+		// This allows us to have editor only/non-shipping only properties that are stripped out on load
+		int32 SharedDataArchiveBufferSize = SharedDataArchiveBuffer.Num();
+		Ar << SharedDataArchiveBufferSize;
+		Ar.Serialize(SharedDataArchiveBuffer.GetData(), SharedDataArchiveBufferSize);
+#endif
+	}
+	else
+	{
+		// Counting, etc
+		Ar << SharedDataBuffer;
+
+#if WITH_EDITORONLY_DATA
+		Ar << SharedDataArchiveBuffer;
+#endif
+	}
+}
+
+void UAnimNextModule::PostLoad()
+{
+	Super::PostLoad();
+
+#if WITH_EDITORONLY_DATA
+	if(GetLinkerCustomVersion(FFortniteMainBranchObjectVersion::GUID) < FFortniteMainBranchObjectVersion::AnimNextGraphAccessSpecifiers)
+	{
+		DefaultState.State = PropertyBag_DEPRECATED;
+	}
+#endif
+}
+
+#if WITH_EDITOR
+void UAnimNextModule::PostEditChangeProperty(FPropertyChangedEvent& PropertyChangedEvent)
+{
+	if(PropertyChangedEvent.Property && PropertyChangedEvent.Property->GetFName() == GET_MEMBER_NAME_CHECKED(UAnimNextModule, DefaultEntryPoint))
+	{
+		CacheDefaultEntryPoint(); 
+	}
+}
+#endif
+
+TConstArrayView<UE::AnimNext::FScheduleTerm> UAnimNextModule::GetTerms() const
+{
+	using namespace UE::AnimNext;
+
+	static const FParamId OutputPoseId("UE_Internal_Graph_OutputPose");
+
+	static const FScheduleTerm Terms[] =
+	{
+		FScheduleTerm(OutputPoseId, FAnimNextParamType::GetType<FAnimNextGraphLODPose>(), EScheduleTermDirection::Output)
+	};
+
+	return Terms;
+}
+
+bool UAnimNextModule::LoadFromArchiveBuffer(const TArray<uint8>& InSharedDataArchiveBuffer)
+{
+	using namespace UE::AnimNext;
+
+	// Reconstruct our graph shared data
+	FMemoryReader GraphSharedDataArchive(InSharedDataArchiveBuffer);
+	FTraitReader TraitReader(GraphReferencedObjects, GraphSharedDataArchive);
+
+	const FTraitReader::EErrorState ErrorState = TraitReader.ReadGraph(SharedDataBuffer);
+	if (ErrorState == FTraitReader::EErrorState::None)
+	{
+		for(int32 EntryPointIndex = 0; EntryPointIndex < EntryPoints.Num(); ++EntryPointIndex)
+		{
+			const FAnimNextGraphEntryPoint& EntryPoint = EntryPoints[EntryPointIndex];
+			ResolvedRootTraitHandles.Add(EntryPoint.EntryPointName, TraitReader.ResolveEntryPointHandle(EntryPoint.RootTraitHandle));
+			ResolvedEntryPoints.Add(EntryPoint.EntryPointName, EntryPointIndex);
+		}
+
+		// Make sure our execute method is registered
+		FRigUnit_AnimNextGraphEvaluator::RegisterExecuteMethod(ExecuteDefinition);
+		return true;
+	}
+	else
+	{
+		SharedDataBuffer.Empty(0);
+		ResolvedRootTraitHandles.Add(GetDefaultEntryPoint(), FAnimNextTraitHandle());
+		return false;
+	}
+}
+
+#if WITH_EDITORONLY_DATA
+void UAnimNextModule::FreezeGraphInstances()
+{
+	FScopeLock Lock(&GraphInstancesLock);
+
+	TSet<FAnimNextGraphInstance*> GraphInstancesCopy = GraphInstances;
+	for (FAnimNextGraphInstance* GraphInstance : GraphInstancesCopy)
+	{
+		GraphInstance->Freeze();
+	}
+}
+
+void UAnimNextModule::ThawGraphInstances()
+{
+	FScopeLock Lock(&GraphInstancesLock);
+
+	TSet<FAnimNextGraphInstance*> GraphInstancesCopy = GraphInstances;
+	for (FAnimNextGraphInstance* GraphInstance : GraphInstancesCopy)
+	{
+		GraphInstance->Thaw();
+	}
+}
+#endif
+
+void UAnimNextModule::SetDefaultEntryPoint(FName InEntryPoint)
+{
+	DefaultEntryPoint = InEntryPoint;
+
+	CacheDefaultEntryPoint();
+}
+
+FName UAnimNextModule::GetDefaultEntryPoint() const
+{
+	if(CachedDefaultEntryPoint == NAME_None)
+	{
+		CacheDefaultEntryPoint();
+	}
+
+	return CachedDefaultEntryPoint;
+}
+
+void UAnimNextModule::CacheDefaultEntryPoint() const
+{
+	TStringBuilder<256> StringBuilder;
+	StringBuilder.Append(GetPathName());
+	StringBuilder.Append(TEXT(":"));
+	DefaultEntryPoint.AppendString(StringBuilder);
+	CachedDefaultEntryPoint = FName(StringBuilder.ToView());
+}
\ No newline at end of file
diff --git a/Engine/Plugins/Experimental/AnimNext/Source/AnimNext/Private/Module/RigUnit_AnimNextModuleSelf.cpp b/Engine/Plugins/Experimental/AnimNext/Source/AnimNext/Private/Module/RigUnit_AnimNextModuleSelf.cpp
new file mode 100644
--- /dev/null
+++ b/Engine/Plugins/Experimental/AnimNext/Source/AnimNext/Private/Module/RigUnit_AnimNextModuleSelf.cpp
@@ -0,0 +1,9 @@
+﻿// Copyright Epic Games, Inc. All Rights Reserved.
+
+#include "Module/RigUnit_AnimNextModuleSelf.h"
+#include "Graph/AnimNextGraphInstance.h"
+
+FRigUnit_AnimNextModuleSelf_Execute()
+{
+	Self = const_cast<UAnimNextModule*>(ExecuteContext.GetContextData<FAnimNextGraphContextData>().GetGraphInstance().GetModule());
+}
diff --git a/Engine/Plugins/Experimental/AnimNext/Source/AnimNext/Private/Param/ParametersProxy.cpp b/Engine/Plugins/Experimental/AnimNext/Source/AnimNext/Private/Param/ParametersProxy.cpp
--- a/Engine/Plugins/Experimental/AnimNext/Source/AnimNext/Private/Param/ParametersProxy.cpp
+++ b/Engine/Plugins/Experimental/AnimNext/Source/AnimNext/Private/Param/ParametersProxy.cpp
@@ -3,18 +3,19 @@
 #include "Param/ParametersProxy.h"
 
 #include "ParamHelpers.h"
-#include "Graph/AnimNextGraph.h"
+#include "Graph/AnimNextGraphState.h"
+#include "Module/AnimNextModule.h"
 #include "Param/ParamStack.h"
 
 namespace UE::AnimNext
 {
 
-FParametersProxy::FParametersProxy(const UAnimNextGraph* InGraph)
-	: Graph(InGraph)
-	, PropertyBag(InGraph->DefaultState.State)
+FParametersProxy::FParametersProxy(const UAnimNextModule* InModule)
+	: Module(InModule)
+	, PropertyBag(InModule->DefaultState.State)
 	, LayerHandle(FParamStack::MakeReferenceLayer(NAME_None, PropertyBag))
 {
-	check(Graph);
+	check(Module);
 	UpdateCachedExternalParamData();
 }
 
@@ -22,9 +23,9 @@ void FParametersProxy::Update(float DeltaTime)
 {
 #if WITH_EDITOR	// Layout should only be changing in editor
 	const FInstancedPropertyBag* HandlePropertyBag = LayerHandle.As<FInstancedPropertyBag>();
-	if(HandlePropertyBag == nullptr || HandlePropertyBag->GetPropertyBagStruct() != Graph->DefaultState.State.GetPropertyBagStruct())
+	if(HandlePropertyBag == nullptr || HandlePropertyBag->GetPropertyBagStruct() != Module->DefaultState.State.GetPropertyBagStruct())
 	{
-		PropertyBag = Graph->DefaultState.State;
+		PropertyBag = Module->DefaultState.State;
 		LayerHandle = FParamStack::MakeReferenceLayer(NAME_None, PropertyBag);
 		UpdateCachedExternalParamData();
 	}
@@ -47,19 +48,19 @@ void FParametersProxy::Update(float DeltaTime)
 	}
 
 	// Next we update the layer
-	Graph->UpdateLayer(LayerHandle, DeltaTime);
+	Module->UpdateLayer(LayerHandle, DeltaTime);
 }
 
 void FParametersProxy::UpdateCachedExternalParamData()
 {
 	ExternalParamData.Reset();
 
-	if(Graph->DefaultState.PublicParameterStartIndex != INDEX_NONE)
+	if(Module->DefaultState.PublicParameterStartIndex != INDEX_NONE)
 	{
-		TConstArrayView<FPropertyBagPropertyDesc> Descs = Graph->DefaultState.State.GetPropertyBagStruct()->GetPropertyDescs();
+		TConstArrayView<FPropertyBagPropertyDesc> Descs = Module->DefaultState.State.GetPropertyBagStruct()->GetPropertyDescs();
 		const int32 NumProperties = Descs.Num();
-		ExternalParamData.Reserve(NumProperties - Graph->DefaultState.PublicParameterStartIndex);
-		for(int32 PropertyIndex = Graph->DefaultState.PublicParameterStartIndex; PropertyIndex < NumProperties; ++PropertyIndex)
+		ExternalParamData.Reserve(NumProperties - Module->DefaultState.PublicParameterStartIndex);
+		for(int32 PropertyIndex = Module->DefaultState.PublicParameterStartIndex; PropertyIndex < NumProperties; ++PropertyIndex)
 		{
 			const FPropertyBagPropertyDesc& Desc = Descs[PropertyIndex];
 			check(Desc.CachedProperty);
@@ -72,7 +73,7 @@ void FParametersProxy::UpdateCachedExternalParamData()
 
 void FParametersProxy::AddReferencedObjects(FReferenceCollector& Collector)
 {
-	Collector.AddReferencedObject(Graph);
+	Collector.AddReferencedObject(Module);
 	PropertyBag.AddStructReferencedObjects(Collector);
 }
 
diff --git a/Engine/Plugins/Experimental/AnimNext/Source/AnimNext/Private/Param/ParametersProxy.h b/Engine/Plugins/Experimental/AnimNext/Source/AnimNext/Private/Param/ParametersProxy.h
--- a/Engine/Plugins/Experimental/AnimNext/Source/AnimNext/Private/Param/ParametersProxy.h
+++ b/Engine/Plugins/Experimental/AnimNext/Source/AnimNext/Private/Param/ParametersProxy.h
@@ -6,7 +6,7 @@
 #include "PropertyBag.h"
 #include "UObject/StrongObjectPtr.h"
 
-class UAnimNextGraph;
+class UAnimNextModule;
 
 namespace UE::AnimNext
 {
@@ -16,7 +16,7 @@ struct FParametersProxy : public IParameterSource
 {
 	FParametersProxy() = delete;
 
-	FParametersProxy(const UAnimNextGraph* InGraph);
+	FParametersProxy(const UAnimNextModule* InModule);
 
 	// IParameterSource interface
 	virtual FName GetInstanceId() const override { return NAME_None; }
@@ -28,7 +28,7 @@ struct FParametersProxy : public IParameterSource
 	void UpdateCachedExternalParamData();
 
 	// The object that this proxy wraps
-	TObjectPtr<const UAnimNextGraph> Graph;
+	TObjectPtr<const UAnimNextModule> Module;
 
 	// Copy of the parameter data
 	FInstancedPropertyBag PropertyBag;
diff --git a/Engine/Plugins/Experimental/AnimNext/Source/AnimNext/Private/Scheduler/AnimNextSchedule.cpp b/Engine/Plugins/Experimental/AnimNext/Source/AnimNext/Private/Scheduler/AnimNextSchedule.cpp
--- a/Engine/Plugins/Experimental/AnimNext/Source/AnimNext/Private/Scheduler/AnimNextSchedule.cpp
+++ b/Engine/Plugins/Experimental/AnimNext/Source/AnimNext/Private/Scheduler/AnimNextSchedule.cpp
@@ -5,7 +5,7 @@
 #include "Async/TaskGraphInterfaces.h"
 #include "EngineLogs.h"
 #include "UObject/AssetRegistryTagsContext.h"
-#include "Graph/AnimNextGraph.h"
+#include "Module/AnimNextModule.h"
 
 #if WITH_EDITOR
 TUniqueFunction<void(UAnimNextSchedule*)> UAnimNextSchedule::CompileFunction;
@@ -16,16 +16,16 @@ void UAnimNextScheduleEntry_AnimNextGraph::GetPreloadDependencies(TArray<UObject
 {
 	Super::GetPreloadDependencies(OutDeps);
 
-	OutDeps.Add(Graph);
+	OutDeps.Add(Module);
 }
 
 void UAnimNextScheduleEntry_ParamScope::GetPreloadDependencies(TArray<UObject*>& OutDeps)
 {
 	Super::GetPreloadDependencies(OutDeps);
 
-	for(UAnimNextGraph* Graph : Parameters)
+	for(UAnimNextModule* Module : Parameters)
 	{
-		OutDeps.Add(Graph);
+		OutDeps.Add(Module);
 	}
 
 	for(UAnimNextScheduleEntry* SubEntry : SubEntries)
diff --git a/Engine/Plugins/Experimental/AnimNext/Source/AnimNext/Private/Scheduler/AnimNextScheduleGraphTask.cpp b/Engine/Plugins/Experimental/AnimNext/Source/AnimNext/Private/Scheduler/AnimNextScheduleGraphTask.cpp
--- a/Engine/Plugins/Experimental/AnimNext/Source/AnimNext/Private/Scheduler/AnimNextScheduleGraphTask.cpp
+++ b/Engine/Plugins/Experimental/AnimNext/Source/AnimNext/Private/Scheduler/AnimNextScheduleGraphTask.cpp
@@ -2,9 +2,9 @@
 
 #include "Scheduler/AnimNextScheduleGraphTask.h"
 
-#include "IAnimNextModule.h"
+#include "IAnimNextModuleInterface.h"
 #include "Scheduler/ScheduleContext.h"
-#include "Graph/AnimNextGraph.h"
+#include "Module/AnimNextModule.h"
 #include "Context.h"
 #include "Graph/AnimNext_LODPose.h"
 #include "AnimNextStats.h"
@@ -18,27 +18,27 @@
 
 DEFINE_STAT(STAT_AnimNext_Task_Graph);
 
-UAnimNextGraph* FAnimNextScheduleGraphTask::GetGraphToRun(UE::AnimNext::FParamStack& ParamStack) const
+UAnimNextModule* FAnimNextScheduleGraphTask::GetModuleToRun(UE::AnimNext::FParamStack& ParamStack) const
 {
-	UAnimNextGraph* GraphToRun = Graph;
-	if (GraphToRun == nullptr && DynamicGraph.IsValid())
+	UAnimNextModule* ModuleToRun = Module;
+	if (ModuleToRun == nullptr && DynamicModule.IsValid())
 	{
-		if(const TObjectPtr<UAnimNextGraph>* FoundGraph = ParamStack.GetParamPtr<TObjectPtr<UAnimNextGraph>>(DynamicGraph.GetParamId()))
+		if(const TObjectPtr<UAnimNextModule>* FoundModule = ParamStack.GetParamPtr<TObjectPtr<UAnimNextModule>>(DynamicModule.GetParamId()))
 		{
-			GraphToRun = *FoundGraph;
+			ModuleToRun = *FoundModule;
 		}
 	}
 
-	return GraphToRun;
+	return ModuleToRun;
 }
 
-void FAnimNextScheduleGraphTask::VerifyRequiredParameters(UAnimNextGraph* InGraphToRun) const
+void FAnimNextScheduleGraphTask::VerifyRequiredParameters(UAnimNextModule* InModuleToRun) const
 {
-	if(SuppliedParametersHash != InGraphToRun->RequiredParametersHash)
+	if(SuppliedParametersHash != InModuleToRun->RequiredParametersHash)
 	{
 		bool bWarningOutput = false;
 
-		for(const FAnimNextParam& RequiredParameter : InGraphToRun->RequiredParameters)
+		for(const FAnimNextParam& RequiredParameter : InModuleToRun->RequiredParameters)
 		{
 			bool bFound = false;
 			bool bFoundCorrectType = true;
@@ -59,7 +59,7 @@ void FAnimNextScheduleGraphTask::VerifyRequiredParameters(UAnimNextGraph* InGrap
 
 			if(!bWarningOutput && (!bFound || !bFoundCorrectType))
 			{
-				UE_LOGFMT(LogAnimation, Warning, "AnimNext: Graph {GraphToRun} has different required parameters, it may not run correctly.", InGraphToRun->GetFName());
+				UE_LOGFMT(LogAnimation, Warning, "AnimNext: Graph {ModuleToRun} has different required parameters, it may not run correctly.", InModuleToRun->GetFName());
 				bWarningOutput = true;
 			}
 			
@@ -75,7 +75,7 @@ void FAnimNextScheduleGraphTask::VerifyRequiredParameters(UAnimNextGraph* InGrap
 	}
 }
 
-void FAnimNextScheduleGraphTask::RunGraph(const UE::AnimNext::FScheduleContext& InContext) const
+void FAnimNextScheduleGraphTask::RunModule(const UE::AnimNext::FScheduleContext& InContext) const
 {
 	SCOPE_CYCLE_COUNTER(STAT_AnimNext_Task_Graph);
 
@@ -83,8 +83,8 @@ void FAnimNextScheduleGraphTask::RunGraph(const UE::AnimNext::FScheduleContext&
 
 	FParamStack& ParamStack = FParamStack::Get();
 
-	UAnimNextGraph* GraphToRun = GetGraphToRun(ParamStack);
-	if(GraphToRun == nullptr)
+	UAnimNextModule* ModuleToRun = GetModuleToRun(ParamStack);
+	if(ModuleToRun == nullptr)
 	{
 		return;
 	}
@@ -93,7 +93,7 @@ void FAnimNextScheduleGraphTask::RunGraph(const UE::AnimNext::FScheduleContext&
 	FScheduleInstanceData::FGraphCache& GraphCache = InstanceData.GraphCaches[TaskIndex];
 
 	// Check if we are running the correct graph and release (and any term mapping layers) it if not
-	if(GraphCache.GraphInstanceData.IsValid() && !GraphCache.GraphInstanceData.UsesGraph(GraphToRun))
+	if(GraphCache.GraphInstanceData.IsValid() && !GraphCache.GraphInstanceData.UsesModule(ModuleToRun))
 	{
 		GraphCache.GraphInstanceData.Release();
 		GraphCache.GraphTermLayer.Invalidate();
@@ -102,12 +102,12 @@ void FAnimNextScheduleGraphTask::RunGraph(const UE::AnimNext::FScheduleContext&
 	// Allocate our graph instance data
 	if (!GraphCache.GraphInstanceData.IsValid())
 	{
-		GraphToRun->AllocateInstance(GraphCache.GraphInstanceData, EntryPoint.Name);
+		ModuleToRun->AllocateInstance(GraphCache.GraphInstanceData, EntryPoint.Name);
 
 		// Only do dynamic verification for dynamic graphs. Static graphs get verified at compile time. 
-		if (Graph == nullptr && DynamicGraph.IsValid())
+		if (Module == nullptr && DynamicModule.IsValid())
 		{
-			VerifyRequiredParameters(GraphToRun);
+			VerifyRequiredParameters(ModuleToRun);
 		}
 	}
 
@@ -127,7 +127,7 @@ void FAnimNextScheduleGraphTask::RunGraph(const UE::AnimNext::FScheduleContext&
 	// Check and allocate remapped term layer
 	if(!GraphCache.GraphTermLayer.IsValid())
 	{
-		TConstArrayView<FScheduleTerm> GraphTerms = GraphToRun->GetTerms();
+		TConstArrayView<FScheduleTerm> GraphTerms = ModuleToRun->GetTerms();
 		check(Terms.Num() == GraphTerms.Num());
 
 		TMap<FName, FName> Mapping;
@@ -143,7 +143,7 @@ void FAnimNextScheduleGraphTask::RunGraph(const UE::AnimNext::FScheduleContext&
 	}
 
 	// TODO: This should not be fixed at arg 0, we should define this in the graph asset
-	FAnimNextGraphLODPose* OutputPose = GraphCache.GraphTermLayer.GetMutableParamPtr<FAnimNextGraphLODPose>(GraphToRun->GetTerms()[0].GetId());
+	FAnimNextGraphLODPose* OutputPose = GraphCache.GraphTermLayer.GetMutableParamPtr<FAnimNextGraphLODPose>(ModuleToRun->GetTerms()[0].GetId());
 	if(OutputPose == nullptr)
 	{
 		return;
@@ -181,7 +181,7 @@ void FAnimNextScheduleGraphTask::RunGraph(const UE::AnimNext::FScheduleContext&
 	// This reduces churn internally by avoiding a chunk to be repeatedly allocated and freed as we push/pop marks
 	MemStack.Alloc(size_t(FPageAllocator::SmallPageSize) + 1, 16);
 
-	IAnimNextModule& AnimNextModule = IAnimNextModule::Get();
+	IAnimNextModuleInterface& AnimNextModule = IAnimNextModuleInterface::Get();
 	AnimNextModule.UpdateGraph(GraphCache.GraphInstanceData, InContext.GetDeltaTime(), InputEventList, OutputEventList);
 	AnimNextModule.EvaluateGraph(GraphCache.GraphInstanceData, RefPose, LODIndex, OutputPose->LODPose);
 
diff --git a/Engine/Plugins/Experimental/AnimNext/Source/AnimNext/Private/Scheduler/ScheduleInitializationContext.cpp b/Engine/Plugins/Experimental/AnimNext/Source/AnimNext/Private/Scheduler/ScheduleInitializationContext.cpp
--- a/Engine/Plugins/Experimental/AnimNext/Source/AnimNext/Private/Scheduler/ScheduleInitializationContext.cpp
+++ b/Engine/Plugins/Experimental/AnimNext/Source/AnimNext/Private/Scheduler/ScheduleInitializationContext.cpp
@@ -2,7 +2,7 @@
 
 #include "Scheduler/ScheduleInitializationContext.h"
 
-#include "AnimNextModule.h"
+#include "AnimNextModuleImpl.h"
 #include "Modules/ModuleManager.h"
 #include "Param/ObjectProxy.h"
 #include "Param/ObjectProxyFactory.h"
@@ -19,7 +19,7 @@ FScheduleInitializationContext::FScheduleInitializationContext(const FScheduleCo
 
 void FScheduleInitializationContext::ApplyParametersToScope(FName InScope, EParameterScopeOrdering InOrdering, const TInstancedStruct<FAnimNextParamInstanceIdentifier>& InInstanceId, const FParameterSourceContext& InContext, TConstArrayView<FName> InRequiredParameters) const
 {
-	FModule& AnimNextModule = FModuleManager::GetModuleChecked<FModule>("AnimNext");
+	FAnimNextModuleImpl& AnimNextModule = FModuleManager::GetModuleChecked<FAnimNextModuleImpl>("AnimNext");
 	TUniquePtr<IParameterSource> ParameterSource = AnimNextModule.CreateParameterSource(InContext, InInstanceId, InRequiredParameters);
 	Context.GetInstanceData().ApplyParametersToScope(InScope, InOrdering, MoveTemp(ParameterSource));
 }
diff --git a/Engine/Plugins/Experimental/AnimNext/Source/AnimNext/Private/Scheduler/ScheduleInstanceData.cpp b/Engine/Plugins/Experimental/AnimNext/Source/AnimNext/Private/Scheduler/ScheduleInstanceData.cpp
--- a/Engine/Plugins/Experimental/AnimNext/Source/AnimNext/Private/Scheduler/ScheduleInstanceData.cpp
+++ b/Engine/Plugins/Experimental/AnimNext/Source/AnimNext/Private/Scheduler/ScheduleInstanceData.cpp
@@ -6,7 +6,7 @@
 #include "Scheduler/AnimNextSchedulerWorldSubsystem.h"
 #include "Scheduler/AnimNextSchedulerEntry.h"
 #include "AnimNextStats.h"
-#include "AnimNextModule.h"
+#include "AnimNextModuleImpl.h"
 #include "Modules/ModuleManager.h"
 #include "Param/IParameterSourceFactory.h"
 #include "Param/ParametersProxy.h"
@@ -34,11 +34,11 @@ FScheduleInstanceData::FScheduleInstanceData(const FScheduleContext& InScheduleC
 		const FAnimNextScheduleParamScopeEntryTask& EntryTask = InSchedule->ParamScopeEntryTasks[EntryIndex];
 		FScopeCache& ScopeCache = ScopeCaches[EntryTask.ParamScopeIndex];
 		ScopeCache.ParameterSources.Reserve(EntryTask.Parameters.Num());
-		for(UAnimNextGraph* Graph : EntryTask.Parameters)
+		for(UAnimNextModule* Module : EntryTask.Parameters)
 		{
-			if(Graph)
+			if(Module)
 			{
-				ScopeCache.ParameterSources.Emplace(MakeUnique<FParametersProxy>(Graph));
+				ScopeCache.ParameterSources.Emplace(MakeUnique<FParametersProxy>(Module));
 			}
 		}
 
@@ -72,7 +72,7 @@ FScheduleInstanceData::FScheduleInstanceData(const FScheduleContext& InScheduleC
 	}
 
 	// Set up external parameters
-	UE::AnimNext::FModule& AnimNextModule = FModuleManager::GetModuleChecked<UE::AnimNext::FModule>("AnimNext");
+	UE::AnimNext::FAnimNextModuleImpl& AnimNextModule = FModuleManager::GetModuleChecked<UE::AnimNext::FAnimNextModuleImpl>("AnimNext");
 
 	FParameterSourceContext ParameterSourceContext;
 	ParameterSourceContext.Object = Entry->WeakObject.Get();
diff --git a/Engine/Plugins/Experimental/AnimNext/Source/AnimNext/Private/Scheduler/ScheduleInstanceData.h b/Engine/Plugins/Experimental/AnimNext/Source/AnimNext/Private/Scheduler/ScheduleInstanceData.h
--- a/Engine/Plugins/Experimental/AnimNext/Source/AnimNext/Private/Scheduler/ScheduleInstanceData.h
+++ b/Engine/Plugins/Experimental/AnimNext/Source/AnimNext/Private/Scheduler/ScheduleInstanceData.h
@@ -4,7 +4,7 @@
 
 #include "CoreMinimal.h"
 #include "Scheduler/ScheduleHandle.h"
-#include "Graph/AnimNextGraph.h"
+#include "Module/AnimNextModule.h"
 #include "Graph/AnimNextGraphInstancePtr.h"
 #include "Param/ParamStack.h"
 #include "TraitCore/TraitEvent.h"
diff --git a/Engine/Plugins/Experimental/AnimNext/Source/AnimNext/Private/Scheduler/ScheduleTickFunction.cpp b/Engine/Plugins/Experimental/AnimNext/Source/AnimNext/Private/Scheduler/ScheduleTickFunction.cpp
--- a/Engine/Plugins/Experimental/AnimNext/Source/AnimNext/Private/Scheduler/ScheduleTickFunction.cpp
+++ b/Engine/Plugins/Experimental/AnimNext/Source/AnimNext/Private/Scheduler/ScheduleTickFunction.cpp
@@ -212,7 +212,7 @@ void FScheduleTickFunction::RunScheduleHelper(const FScheduleContext& InSchedule
 				FScheduleInstanceData& InstanceData = InScheduleContext.GetInstanceData();
 				FParamStack::AttachToCurrentThread(InstanceData.GetParamStack(Schedule->GraphTasks[GraphTaskIndex].ParamScopeIndex), FParamStack::ECoalesce::Coalesce);
 
-				Schedule->GraphTasks[GraphTaskIndex].RunGraph(InScheduleContext);
+				Schedule->GraphTasks[GraphTaskIndex].RunModule(InScheduleContext);
 
 				FParamStack::DetachFromCurrentThread(FParamStack::EDecoalesce::Decoalesce);
 				break;
diff --git a/Engine/Plugins/Experimental/AnimNext/Source/AnimNext/Private/TraitCore/ExecutionContext.cpp b/Engine/Plugins/Experimental/AnimNext/Source/AnimNext/Private/TraitCore/ExecutionContext.cpp
--- a/Engine/Plugins/Experimental/AnimNext/Source/AnimNext/Private/TraitCore/ExecutionContext.cpp
+++ b/Engine/Plugins/Experimental/AnimNext/Source/AnimNext/Private/TraitCore/ExecutionContext.cpp
@@ -9,7 +9,7 @@
 #include "TraitCore/NodeInstance.h"
 #include "TraitCore/NodeTemplate.h"
 #include "TraitCore/NodeTemplateRegistry.h"
-#include "Graph/AnimNextGraph.h"
+#include "Module/AnimNextModule.h"
 #include "Graph/AnimNextGraphInstance.h"
 
 namespace UE::AnimNext
@@ -153,7 +153,7 @@ namespace UE::AnimNext
 			return FTraitPtr();	// The execution context must be bound to a valid graph instance
 		}
 
-		if (!ensure(GraphInstance.GetGraph() != nullptr))
+		if (!ensure(GraphInstance.GetModule() != nullptr))
 		{
 			return FTraitPtr();	// We need a valid graph instance to allocate into
 		}
@@ -565,8 +565,8 @@ namespace UE::AnimNext
 	const FNodeDescription& FExecutionContext::GetNodeDescription(const FAnimNextGraphInstance& GraphInstance, FNodeHandle NodeHandle) const
 	{
 		// Grab the node description from the specified graph
-		const UAnimNextGraph* Graph = GraphInstance.GetGraph();
-		return *reinterpret_cast<const FNodeDescription*>(&Graph->SharedDataBuffer[NodeHandle.GetSharedOffset()]);
+		const UAnimNextModule* Module = GraphInstance.GetModule();
+		return *reinterpret_cast<const FNodeDescription*>(&Module->SharedDataBuffer[NodeHandle.GetSharedOffset()]);
 	}
 
 	const FNodeDescription& FExecutionContext::GetNodeDescription(const FNodeInstance& NodeInstance) const
diff --git a/Engine/Plugins/Experimental/AnimNext/Source/AnimNext/Public/DataRegistry.h b/Engine/Plugins/Experimental/AnimNext/Source/AnimNext/Public/DataRegistry.h
--- a/Engine/Plugins/Experimental/AnimNext/Source/AnimNext/Public/DataRegistry.h
+++ b/Engine/Plugins/Experimental/AnimNext/Source/AnimNext/Public/DataRegistry.h
@@ -12,7 +12,7 @@ struct FReferenceSkeleton;
 namespace UE::AnimNext
 {
 
-class FModule;
+class FAnimNextModuleImpl;
 
 struct FDataHandle;
 
@@ -222,7 +222,7 @@ private:
 
 // --- ---
 private:
-	friend class FModule;
+	friend class FAnimNextModuleImpl;
 	friend struct FDataHandle;
 
 	// Initialize the global registry
diff --git a/Engine/Plugins/Experimental/AnimNext/Source/AnimNext/Public/Graph/AnimNextGraphInstance.h b/Engine/Plugins/Experimental/AnimNext/Source/AnimNext/Public/Graph/AnimNextGraphInstance.h
--- a/Engine/Plugins/Experimental/AnimNext/Source/AnimNext/Public/Graph/AnimNextGraphInstance.h
+++ b/Engine/Plugins/Experimental/AnimNext/Source/AnimNext/Public/Graph/AnimNextGraphInstance.h
@@ -16,7 +16,7 @@ class FReferenceCollector;
 
 struct FAnimNextGraphInstancePtr;
 struct FRigUnit_AnimNextGraphEvaluator;
-class UAnimNextGraph;
+class UAnimNextModule;
 
 namespace UE::AnimNext
 {
@@ -53,8 +53,8 @@ struct ANIMNEXT_API FAnimNextGraphInstance
 	// Returns true if we have a live graph instance, false otherwise
 	bool IsValid() const;
 
-	// Returns the graph used by this instance or nullptr if the instance is invalid
-	const UAnimNextGraph* GetGraph() const;
+	// Returns the module used by this instance or nullptr if the instance is invalid
+	const UAnimNextModule* GetModule() const;
 
 	// Returns the entry point in Graph that this instance corresponds to 
 	FName GetEntryPoint() const;
@@ -68,8 +68,8 @@ struct ANIMNEXT_API FAnimNextGraphInstance
 	// Returns the root graph instance that owns us and the components or nullptr if we are invalid
 	FAnimNextGraphInstance* GetRootGraphInstance() const;
 
-	// Check to see if this instance data matches the provided graph
-	bool UsesGraph(const UAnimNextGraph* InGraph) const;
+	// Check to see if this instance data matches the provided module
+	bool UsesModule(const UAnimNextModule* InModule) const;
 
 	// Check to see if this instance data matches the provided graph entry point
 	bool UsesEntryPoint(FName InEntryPoint) const;
@@ -127,7 +127,7 @@ private:
 
 	// Hard reference to the graph used to create this instance to ensure we can release it safely
 	UPROPERTY()
-	TObjectPtr<const UAnimNextGraph> Graph;
+	TObjectPtr<const UAnimNextModule> Module;
 
 	// The entry point in Graph that this instance corresponds to 
 	FName EntryPoint;
@@ -154,7 +154,7 @@ private:
 	// Whether or not this graph has updated once
 	bool bHasUpdatedOnce = false;
 
-	friend UAnimNextGraph;					// The graph is the one that allocates instances
+	friend UAnimNextModule;					// The module is the one that allocates instances
 	friend FRigUnit_AnimNextGraphEvaluator;	// We evaluate the instance
 	friend UE::AnimNext::FExecutionContext;
 	friend FAnimNextGraphInstancePtr;
diff --git a/Engine/Plugins/Experimental/AnimNext/Source/AnimNext/Public/Graph/AnimNextGraphInstancePtr.h b/Engine/Plugins/Experimental/AnimNext/Source/AnimNext/Public/Graph/AnimNextGraphInstancePtr.h
--- a/Engine/Plugins/Experimental/AnimNext/Source/AnimNext/Public/Graph/AnimNextGraphInstancePtr.h
+++ b/Engine/Plugins/Experimental/AnimNext/Source/AnimNext/Public/Graph/AnimNextGraphInstancePtr.h
@@ -13,7 +13,7 @@ class FReferenceCollector;
 
 struct FAnimNextGraphInstance;
 struct FRigUnit_AnimNextGraphEvaluator;
-class UAnimNextGraph;
+class UAnimNextModule;
 
 namespace UE::AnimNext
 {
@@ -49,8 +49,8 @@ struct ANIMNEXT_API FAnimNextGraphInstancePtr
 	// Returns true if we have a live graph instance, false otherwise
 	bool IsValid() const;
 
-	// Returns the graph used by this instance or nullptr if the instance is invalid
-	const UAnimNextGraph* GetGraph() const;
+	// Returns the module used by this instance or nullptr if the instance is invalid
+	const UAnimNextModule* GetModule() const;
 
 	// Returns a weak handle to the root trait instance
 	UE::AnimNext::FWeakTraitPtr GetGraphRootPtr() const;
@@ -58,8 +58,8 @@ struct ANIMNEXT_API FAnimNextGraphInstancePtr
 	// Returns the graph instance implementation
 	FAnimNextGraphInstance* GetImpl() const;
 
-	// Check to see if this instance data matches the provided graph
-	bool UsesGraph(const UAnimNextGraph* InGraph) const;
+	// Check to see if this instance data matches the provided module
+	bool UsesModule(const UAnimNextModule* InModule) const;
 
 	// Returns whether or not this graph instance is the root graph instance or false otherwise
 	bool IsRoot() const;
@@ -104,7 +104,7 @@ private:
 	// Indirection to hide implementation details and to fix the graph instance into a single memory location
 	TUniquePtr<FAnimNextGraphInstance> Impl;
 
-	friend UAnimNextGraph;					// The graph is the one that allocates instances
+	friend UAnimNextModule;					// The graph is the one that allocates instances
 	friend FRigUnit_AnimNextGraphEvaluator;	// We evaluate the instance
 	friend UE::AnimNext::FExecutionContext;
 };
diff --git a/Engine/Plugins/Experimental/AnimNext/Source/AnimNext/Public/IAnimNextModuleInterface.h b/Engine/Plugins/Experimental/AnimNext/Source/AnimNext/Public/IAnimNextModuleInterface.h
new file mode 100644
--- /dev/null
+++ b/Engine/Plugins/Experimental/AnimNext/Source/AnimNext/Public/IAnimNextModuleInterface.h
@@ -0,0 +1,50 @@
+// Copyright Epic Games, Inc. All Rights Reserved.
+
+#pragma once
+
+#include "CoreMinimal.h"
+#include "Modules/ModuleInterface.h"
+#include "LODPose.h"
+#include "ReferencePose.h"
+
+struct FAnimNextGraphInstancePtr;
+struct FAnimNextParamInstanceIdentifier;
+
+namespace UE::AnimNext
+{
+	class IParameterSource;
+	class IParameterSourceFactory;
+	struct FParameterSourceContext;
+	struct FTraitEventList;
+}
+
+namespace UE::AnimNext
+{
+
+class IAnimNextAnimGraph
+{
+public:
+	virtual ~IAnimNextAnimGraph() = default;
+
+	virtual void UpdateGraph(FAnimNextGraphInstancePtr& GraphInstance, float DeltaTime, FTraitEventList& InputEventList, FTraitEventList& OutputEventList) const = 0;
+	virtual void EvaluateGraph(FAnimNextGraphInstancePtr& GraphInstance, const FReferencePose& RefPose, int32 GraphLODLevel, FLODPoseHeap& OutputPose) const = 0;
+};
+
+class IAnimNextModuleInterface : public IModuleInterface
+{
+public:
+	// WARNING
+	// These functions are here because the schedule currently cannot have tasks defined in
+	// external plugins. And as such, we must be able to update and evaluate graphs from
+	// AnimNext even though they are implementation details of the AnimNextAnimGraph.
+	// We thus introduce this interface.
+	// When AnimNextAnimGraph's module will load, it will register callbacks that
+	// AnimNext will call when we need to update/evaluate graphs.
+	static ANIMNEXT_API IAnimNextModuleInterface& Get();
+	virtual void RegisterAnimNextAnimGraph(const IAnimNextAnimGraph& InAnimGraphImpl) = 0;
+	virtual void UnregisterAnimNextAnimGraph() = 0;
+	virtual void UpdateGraph(FAnimNextGraphInstancePtr& GraphInstance, float DeltaTime, FTraitEventList& InputEventList, FTraitEventList& OutputEventList) = 0;
+	virtual void EvaluateGraph(FAnimNextGraphInstancePtr& GraphInstance, const FReferencePose& RefPose, int32 GraphLODLevel, FLODPoseHeap& OutputPose) const = 0;
+};
+
+}
diff --git a/Engine/Plugins/Experimental/AnimNext/Source/AnimNextEditor/Internal/Module/ModuleFactory.h b/Engine/Plugins/Experimental/AnimNext/Source/AnimNextEditor/Internal/Module/ModuleFactory.h
new file mode 100644
--- /dev/null
+++ b/Engine/Plugins/Experimental/AnimNext/Source/AnimNextEditor/Internal/Module/ModuleFactory.h
@@ -0,0 +1,22 @@
+﻿// Copyright Epic Games, Inc. All Rights Reserved.
+
+#pragma once
+
+#include "CoreMinimal.h"
+#include "Factories/Factory.h"
+#include "ModuleFactory.generated.h"
+
+UCLASS(MinimalAPI)
+class UAnimNextModuleFactory : public UFactory
+{
+	GENERATED_BODY()
+
+	UAnimNextModuleFactory();
+	
+	virtual bool ConfigureProperties() override;
+	virtual UObject* FactoryCreateNew(UClass* Class, UObject* InParent, FName Name, EObjectFlags Flags, UObject* Context, FFeedbackContext* Warn, FName CallingContext) override;
+	virtual UObject* FactoryCreateNew(UClass* Class, UObject* InParent, FName Name, EObjectFlags Flags, UObject* Context, FFeedbackContext* Warn) override
+	{
+		return FactoryCreateNew(Class, InParent, Name, Flags, Context, Warn, NAME_None);
+	}
+};
\ No newline at end of file
