Rename AnimNextGraph asset to AnimNextModule

diff --git a/Engine/Plugins/Experimental/AnimNext/Source/AnimNextEditor/Private/AnimNextEditorModule.cpp b/Engine/Plugins/Experimental/AnimNext/Source/AnimNextEditor/Private/AnimNextEditorModule.cpp
--- a/Engine/Plugins/Experimental/AnimNext/Source/AnimNextEditor/Private/AnimNextEditorModule.cpp
+++ b/Engine/Plugins/Experimental/AnimNext/Source/AnimNextEditor/Private/AnimNextEditorModule.cpp
@@ -13,11 +13,11 @@
 #include "Common/SRigVMAssetView.h"
 #include "Editor/RigVMEditorTools.h"
 #include "Framework/Application/SlateApplication.h"
-#include "Graph/AnimNextGraph.h"
+#include "Module/AnimNextModule.h"
 #include "Graph/AnimNextGraphPanelNodeFactory.h"
-#include "Graph/AnimNextGraph_EdGraphNodeCustomization.h"
-#include "Graph/AnimNextGraph_EditorData.h"
-#include "Graph/AnimNextGraph_EdGraphNode.h"
+#include "Graph/AnimNextEdGraphNodeCustomization.h"
+#include "Module/AnimNextModule_EditorData.h"
+#include "AnimNextEdGraphNode.h"
 #include "Graph/TraitEditorTabSummoner.h"
 #include "Graph/AnimNextCompilerResultsTabSummoner.h"
 #include "HAL/PlatformApplicationMisc.h"
@@ -42,8 +42,8 @@
 #include "Param/ObjectFunctionLocatorEditor.h"
 #include "Param/ObjectPropertyLocatorEditor.h"
 #include "Param/SAddParametersDialog.h"
-#include "AnimNextRigVMWorkspaceAssetUserData.h"
-#include "Graph/AnimNextGraph_OutlinerItemDetails.h"
+#include "Module/AnimNextModuleWorkspaceAssetUserData.h"
+#include "Graph/AnimNextGraphItemDetails.h"
 #include "Param/AnimNextActorLocatorEditor.h"
 #include "IWorkspaceEditor.h"
 #include "Framework/Docking/LayoutExtender.h"
@@ -83,7 +83,7 @@ void FModule::StartupModule()
 		FOnGetPropertyTypeCustomizationInstance::CreateLambda([] { return MakeShared<FParamNamePropertyTypeCustomization>(); }),
 		Identifier);
 
-	PropertyModule.RegisterCustomClassLayout("AnimNextGraph_Parameter", 
+	PropertyModule.RegisterCustomClassLayout("AnimNextModule_Parameter", 
 		FOnGetDetailCustomizationInstance::CreateLambda([] { return MakeShared<FParameterCustomization>(); }));
 
 	AnimNextGraphPanelNodeFactory = MakeShared<FAnimNextGraphPanelNodeFactory>();
@@ -118,9 +118,9 @@ void FModule::StartupModule()
 
 	WorkspaceEditorModule.OnRegisterWorkspaceDetailsCustomization().AddLambda([](const TWeakPtr<Workspace::IWorkspaceEditor>& InWorkspaceEditor, TSharedPtr<IDetailsView>& InDetailsView)
 		{
-			InDetailsView->RegisterInstancedCustomPropertyLayout(UAnimNextGraph_EdGraphNode::StaticClass(), FOnGetDetailCustomizationInstance::CreateLambda([InWorkspaceEditor]()
+			InDetailsView->RegisterInstancedCustomPropertyLayout(UAnimNextEdGraphNode::StaticClass(), FOnGetDetailCustomizationInstance::CreateLambda([InWorkspaceEditor]()
 				{
-					return MakeShared<FAnimNextGraph_EdGraphNodeCustomization>(InWorkspaceEditor);
+					return MakeShared<FAnimNextEdGraphNodeCustomization>(InWorkspaceEditor);
 				}));
 
 			TArray<UScriptStruct*> StructsToCustomize = {
@@ -143,7 +143,7 @@ void FModule::StartupModule()
 
 	SRigVMAssetView::RegisterCategoryFactory("Parameters", [](UAnimNextRigVMAssetEditorData* InEditorData)
 	{
-		UAnimNextGraph_EditorData* EditorData = CastChecked<UAnimNextGraph_EditorData>(InEditorData);
+		UAnimNextModule_EditorData* EditorData = CastChecked<UAnimNextModule_EditorData>(InEditorData);
 		UAnimNextRigVMAsset* Asset = UncookedOnly::FUtils::GetAsset(InEditorData);
 		return SNew(SSimpleButton)
 			.Text(LOCTEXT("AddParameterButton", "Add Parameter"))
@@ -172,7 +172,7 @@ void FModule::StartupModule()
 
 	SRigVMAssetView::RegisterCategoryFactory("Event Graphs", [](UAnimNextRigVMAssetEditorData* InEditorData)
 	{
-		UAnimNextGraph_EditorData* EditorData = CastChecked<UAnimNextGraph_EditorData>(InEditorData);
+		UAnimNextModule_EditorData* EditorData = CastChecked<UAnimNextModule_EditorData>(InEditorData);
 		return SNew(SSimpleButton)
 			.Text(LOCTEXT("AddEventGraphButton", "Add Event Graph"))
 			.Icon(FAppStyle::Get().GetBrush("Icons.Plus"))
@@ -189,7 +189,7 @@ void FModule::StartupModule()
 
 	SRigVMAssetView::RegisterCategoryFactory("Animation Graphs", [](UAnimNextRigVMAssetEditorData* InEditorData)
 	{
-		UAnimNextGraph_EditorData* EditorData = CastChecked<UAnimNextGraph_EditorData>(InEditorData);
+		UAnimNextModule_EditorData* EditorData = CastChecked<UAnimNextModule_EditorData>(InEditorData);
 		return SNew(SSimpleButton)
 			.Text(LOCTEXT("AddGraphButton", "Add Animation Graph"))
 			.Icon(FAppStyle::Get().GetBrush("Icons.Plus"))
@@ -238,8 +238,7 @@ void FModule::ShutdownModule()
 		PropertyModule.UnregisterCustomPropertyTypeLayout("AnimNextParam");
 		PropertyModule.UnregisterCustomPropertyTypeLayout("AnimNextEditorParam");
 		PropertyModule.UnregisterCustomPropertyTypeLayout("NameProperty");
-		PropertyModule.UnregisterCustomClassLayout("AnimNextGraph_Parameter");
-		PropertyModule.UnregisterCustomClassLayout("AnimNextGraph_EdGraphNode");
+		PropertyModule.UnregisterCustomClassLayout("AnimNextModule_Parameter");
 	}
 
 	FEdGraphUtilities::UnregisterVisualNodeFactory(AnimNextGraphPanelNodeFactory);
@@ -354,8 +353,8 @@ void FModule::RegisterWorkspaceDocumentTypes(Workspace::IWorkspaceEditorModule&
 	Workspace::FObjectDocumentArgs AnimNextGraphDocumentArgs(
 			Workspace::FOnMakeDocumentWidget::CreateLambda([](const Workspace::FWorkspaceEditorContext& InContext)
 			{
-				UAnimNextGraph* Graph = CastChecked<UAnimNextGraph>(InContext.Object);
-				UAnimNextGraph_EditorData* EditorData = UncookedOnly::FUtils::GetEditorData(Graph);
+				UAnimNextModule* Module = CastChecked<UAnimNextModule>(InContext.Object);
+				UAnimNextModule_EditorData* EditorData = UncookedOnly::FUtils::GetEditorData(Module);
 
 				TWeakPtr<Workspace::IWorkspaceEditor> WeakWorkspaceEditor = InContext.WorkspaceEditor;
 
@@ -415,16 +414,16 @@ void FModule::RegisterWorkspaceDocumentTypes(Workspace::IWorkspaceEditorModule&
 			Workspace::WorkspaceTabs::TopMiddleDocumentArea);
 	AnimNextGraphDocumentArgs.OnGetTabName = Workspace::FOnGetTabName::CreateLambda([](const Workspace::FWorkspaceEditorContext& InContext)
 	{
-		const UAnimNextGraph* Graph = CastChecked<UAnimNextGraph>(InContext.Object);
+		const UAnimNextModule* Graph = CastChecked<UAnimNextModule>(InContext.Object);
 		return FText::FromName(Graph->GetFName());
 	});
 
 	AnimNextGraphDocumentArgs.OnGetDocumentBreadcrumbTrail = Workspace::FOnGetDocumentBreadcrumbTrail::CreateLambda([](const Workspace::FWorkspaceEditorContext& InContext, TArray<TSharedPtr<Workspace::FWorkspaceBreadcrumb>>& OutBreadcrumbs)
 	{
-		if (const UAnimNextGraph* Graph = Cast<UAnimNextGraph>(InContext.Object))
+		if (const UAnimNextModule* Module = Cast<UAnimNextModule>(InContext.Object))
 		{
 			const TSharedPtr<Workspace::FWorkspaceBreadcrumb>& GraphCrumb = OutBreadcrumbs.Add_GetRef(MakeShared<Workspace::FWorkspaceBreadcrumb>());
-			GraphCrumb->OnGetLabel = Workspace::FWorkspaceBreadcrumb::FOnGetBreadcrumbLabel::CreateLambda([GraphName = Graph->GetFName()]{ return FText::FromName(GraphName); });
+			GraphCrumb->OnGetLabel = Workspace::FWorkspaceBreadcrumb::FOnGetBreadcrumbLabel::CreateLambda([ModuleName = Module->GetFName()]{ return FText::FromName(ModuleName); });
 		}
 	});
 
@@ -433,9 +432,9 @@ void FModule::RegisterWorkspaceDocumentTypes(Workspace::IWorkspaceEditorModule&
 		return FAppStyle::GetBrush(TEXT("ClassIcon.Default"));
 	});
 
-	WorkspaceEditorModule.RegisterObjectDocumentType(FTopLevelAssetPath(TEXT("/Script/AnimNext.AnimNextGraph")), AnimNextGraphDocumentArgs);
+	WorkspaceEditorModule.RegisterObjectDocumentType(FTopLevelAssetPath(TEXT("/Script/AnimNext.AnimNextModule")), AnimNextGraphDocumentArgs);
 
-	// --- AnimNextGraph_EdGraph ---
+	// --- AnimNextEdGraph ---
 	Workspace::FGraphDocumentWidgetArgs GraphArgs;
 	GraphArgs.SpawnLocation = Workspace::WorkspaceTabs::TopMiddleDocumentArea;
 	GraphArgs.OnCreateActionMenu = Workspace::FOnCreateActionMenu::CreateLambda([](const Workspace::FWorkspaceEditorContext& InContext, UEdGraph* InGraph, const FVector2D& InNodePosition, const TArray<UEdGraphPin*>& InDraggedPins, bool bAutoExpand, SGraphEditor::FActionMenuClosed InOnMenuClosed)
@@ -749,7 +748,7 @@ void FModule::RegisterWorkspaceDocumentTypes(Workspace::IWorkspaceEditorModule&
 	{
 		TWeakPtr<Workspace::IWorkspaceEditor> WeakWorkspaceEditor = InContext.WorkspaceEditor;
 
-		if (UAnimNextGraph_EdGraph* EdGraph = Cast<UAnimNextGraph_EdGraph>(InContext.Object))
+		if (UAnimNextEdGraph* EdGraph = Cast<UAnimNextEdGraph>(InContext.Object))
 		{
 			UAnimNextRigVMAssetEditorData* EditorData = EdGraph->GetTypedOuter<UAnimNextRigVMAssetEditorData>();
 			check(EditorData);
@@ -794,12 +793,12 @@ void FModule::RegisterWorkspaceDocumentTypes(Workspace::IWorkspaceEditorModule&
 
 			if (const UAnimNextRigVMAssetEntry* OuterAssetEntry = CastChecked<UAnimNextRigVMAssetEntry>(RigVMEdGraph->GetOuter()))
 			{
-				if (UAnimNextGraph_EditorData* OuterEditorData = CastChecked<UAnimNextGraph_EditorData>(OuterAssetEntry->GetOuter()))
+				if (UAnimNextModule_EditorData* OuterEditorData = CastChecked<UAnimNextModule_EditorData>(OuterAssetEntry->GetOuter()))
 				{
-					if(UAnimNextGraph* OuterGraph = UncookedOnly::FUtils::GetGraph(OuterEditorData))
+					if(UAnimNextModule* OuterGraph = UncookedOnly::FUtils::GetGraph(OuterEditorData))
 					{
 						const TSharedPtr<Workspace::FWorkspaceBreadcrumb>& OuterGraphCrumb = OutBreadcrumbs.Add_GetRef(MakeShared<Workspace::FWorkspaceBreadcrumb>());
-						TWeakObjectPtr<UAnimNextGraph> WeakOuterGraph = OuterGraph;
+						TWeakObjectPtr<UAnimNextModule> WeakOuterGraph = OuterGraph;
 						TWeakPtr<Workspace::IWorkspaceEditor> WeakWorkspaceEditor = InContext.WorkspaceEditor;
 						OuterGraphCrumb->OnGetLabel = Workspace::FWorkspaceBreadcrumb::FOnGetBreadcrumbLabel::CreateLambda([GraphName = OuterGraph->GetFName()]{ return FText::FromName(GraphName); });
 						OuterGraphCrumb->OnClicked = Workspace::FWorkspaceBreadcrumb::FOnBreadcrumbClicked::CreateLambda(
@@ -817,7 +816,7 @@ void FModule::RegisterWorkspaceDocumentTypes(Workspace::IWorkspaceEditorModule&
 		}
 	});
 
-	WorkspaceEditorModule.RegisterObjectDocumentType(FTopLevelAssetPath(TEXT("/Script/AnimNextUncookedOnly.AnimNextGraph_EdGraph")), GraphDocumentArgs);
+	WorkspaceEditorModule.RegisterObjectDocumentType(FTopLevelAssetPath(TEXT("/Script/AnimNextUncookedOnly.AnimNextEdGraph")), GraphDocumentArgs);
 }
 
 void FModule::UnregisterWorkspaceDocumentTypes()
@@ -826,8 +825,8 @@ void FModule::UnregisterWorkspaceDocumentTypes()
 	{
 		Workspace::IWorkspaceEditorModule& WorkspaceEditorModule = FModuleManager::LoadModuleChecked<Workspace::IWorkspaceEditorModule>("PropertyEditor");
 		WorkspaceEditorModule.UnregisterObjectDocumentType(FTopLevelAssetPath(TEXT("/Script/AnimNext.AnimNextSchedule")));
-		WorkspaceEditorModule.UnregisterObjectDocumentType(FTopLevelAssetPath(TEXT("/Script/AnimNext.AnimNextGraph")));
-		WorkspaceEditorModule.UnregisterObjectDocumentType(FTopLevelAssetPath(TEXT("/Script/AnimNextUncookedOnly.AnimNextGraph_EdGraph")));
+		WorkspaceEditorModule.UnregisterObjectDocumentType(FTopLevelAssetPath(TEXT("/Script/AnimNext.AnimNextModule")));
+		WorkspaceEditorModule.UnregisterObjectDocumentType(FTopLevelAssetPath(TEXT("/Script/AnimNextUncookedOnly.AnimNextEdGraph")));
 	}
 }
 
diff --git a/Engine/Plugins/Experimental/AnimNext/Source/AnimNextEditor/Private/Common/GraphEditorSchemaActions.cpp b/Engine/Plugins/Experimental/AnimNext/Source/AnimNextEditor/Private/Common/GraphEditorSchemaActions.cpp
--- a/Engine/Plugins/Experimental/AnimNext/Source/AnimNextEditor/Private/Common/GraphEditorSchemaActions.cpp
+++ b/Engine/Plugins/Experimental/AnimNext/Source/AnimNextEditor/Private/Common/GraphEditorSchemaActions.cpp
@@ -2,9 +2,9 @@
 
 #include "GraphEditorSchemaActions.h"
 #include "EditorUtils.h"
-#include "Graph/AnimNextGraph_EdGraph.h"
-#include "Graph/AnimNextGraph_EdGraphNode.h"
-#include "Graph/AnimNextGraph_EditorData.h"
+#include "AnimNextEdGraph.h"
+#include "AnimNextEdGraphNode.h"
+#include "Module/AnimNextModule_EditorData.h"
 #include "RigVMModel/Nodes/RigVMUnitNode.h"
 #include "Settings/ControlRigSettings.h"
 #include "Units/RigUnit.h"
diff --git a/Engine/Plugins/Experimental/AnimNext/Source/AnimNextEditor/Private/Common/SActionMenu.cpp b/Engine/Plugins/Experimental/AnimNext/Source/AnimNextEditor/Private/Common/SActionMenu.cpp
--- a/Engine/Plugins/Experimental/AnimNext/Source/AnimNextEditor/Private/Common/SActionMenu.cpp
+++ b/Engine/Plugins/Experimental/AnimNext/Source/AnimNextEditor/Private/Common/SActionMenu.cpp
@@ -6,7 +6,7 @@
 #include "Framework/Application/SlateApplication.h"
 #include "IDocumentation.h"
 #include "SSubobjectEditor.h"
-#include "Graph/AnimNextGraph_EdGraphSchema.h"
+#include "AnimNextEdGraphSchema.h"
 #include "Widgets/Images/SImage.h"
 #include "Widgets/Layout/SBox.h"
 #include "Widgets/SBoxPanel.h"
diff --git a/Engine/Plugins/Experimental/AnimNext/Source/AnimNextEditor/Private/EditorUtils.cpp b/Engine/Plugins/Experimental/AnimNext/Source/AnimNextEditor/Private/EditorUtils.cpp
--- a/Engine/Plugins/Experimental/AnimNext/Source/AnimNextEditor/Private/EditorUtils.cpp
+++ b/Engine/Plugins/Experimental/AnimNext/Source/AnimNextEditor/Private/EditorUtils.cpp
@@ -3,10 +3,10 @@
 #include "EditorUtils.h"
 
 #include "AnimNextRigVMAssetEntry.h"
-#include "Graph/AnimNextGraph_EditorData.h"
+#include "Module/AnimNextModule_EditorData.h"
 #include "Param/ParamType.h"
 #include "Kismet2/Kismet2NameValidators.h"
-#include "Graph/AnimNextGraph.h"
+#include "Module/AnimNextModule.h"
 #include "PropertyBagDetails.h"
 #include "UncookedOnlyUtils.h"
 #include "Param/RigVMDispatch_GetParameter.h"
diff --git a/Engine/Plugins/Experimental/AnimNext/Source/AnimNextEditor/Private/EditorUtils.h b/Engine/Plugins/Experimental/AnimNext/Source/AnimNextEditor/Private/EditorUtils.h
--- a/Engine/Plugins/Experimental/AnimNext/Source/AnimNextEditor/Private/EditorUtils.h
+++ b/Engine/Plugins/Experimental/AnimNext/Source/AnimNextEditor/Private/EditorUtils.h
@@ -6,7 +6,7 @@
 #include "Param/ParamType.h"
 #include "EdGraphSchema_K2.h"
 
-class UAnimNextGraph;
+class UAnimNextModule;
 class UAnimNextRigVMAssetEditorData;
 struct FAnimNextParamType;
 class URigVMController;
diff --git a/Engine/Plugins/Experimental/AnimNext/Source/AnimNextEditor/Private/Graph/AnimNextGraphItemDetails.cpp b/Engine/Plugins/Experimental/AnimNext/Source/AnimNextEditor/Private/Graph/AnimNextGraphItemDetails.cpp
new file mode 100644
--- /dev/null
+++ b/Engine/Plugins/Experimental/AnimNext/Source/AnimNextEditor/Private/Graph/AnimNextGraphItemDetails.cpp
@@ -0,0 +1,129 @@
+﻿// Copyright Epic Games, Inc. All Rights Reserved.
+
+#include "AnimNextGraphItemDetails.h"
+
+#include "Toolkits/AssetEditorToolkitMenuContext.h"
+#include "Graph/AnimNextModule_AnimationGraph.h"
+#include "Module/AnimNextModule_EventGraph.h"
+#include "InstancedStruct.h"
+#include "Module/AnimNextModuleWorkspaceAssetUserData.h"
+#include "RigVMModel/RigVMGraph.h"
+#include "WorkspaceItemMenuContext.h"
+#include "IWorkspaceEditor.h"
+#include "RigVMModel/RigVMClient.h"
+#include "ToolMenus.h"
+
+#define LOCTEXT_NAMESPACE "FAnimNextGraphItemDetails"
+
+namespace UE::AnimNext::Editor
+{
+void FAnimNextGraphItemDetails::HandleDoubleClick(const FToolMenuContext& ToolMenuContext) const
+{
+	const UWorkspaceItemMenuContext* WorkspaceItemContext = ToolMenuContext.FindContext<UWorkspaceItemMenuContext>();
+	const UAssetEditorToolkitMenuContext* AssetEditorContext = ToolMenuContext.FindContext<UAssetEditorToolkitMenuContext>(); 
+	if (WorkspaceItemContext && AssetEditorContext)
+	{
+		if(const TSharedPtr<UE::Workspace::IWorkspaceEditor> WorkspaceEditor = StaticCastSharedPtr<UE::Workspace::IWorkspaceEditor>(AssetEditorContext->Toolkit.Pin()))
+		{
+			const TInstancedStruct<FWorkspaceOutlinerItemData>& Data = WorkspaceItemContext->SelectedExports[0].Data;
+			if (Data.IsValid() && Data.GetScriptStruct() == FAnimNextGraphOutlinerData::StaticStruct())
+			{
+				const FAnimNextGraphOutlinerData& GraphData = Data.Get<FAnimNextGraphOutlinerData>();							
+				if (GraphData.GraphInterface)
+				{
+					if (URigVMGraph* RigVMGraph = GraphData.GraphInterface->GetRigVMGraph())
+					{
+						if(const IRigVMClientHost* RigVMClientHost = RigVMGraph->GetImplementingOuter<IRigVMClientHost>())
+						{
+							if(UObject* EditorObject = RigVMClientHost->GetEditorObjectForRigVMGraph(RigVMGraph))
+							{
+								WorkspaceEditor->OpenObjects({EditorObject});
+							}
+						}
+					}
+				}
+			}
+		}
+	}
+}
+
+UPackage* FAnimNextGraphItemDetails::GetPackage(const FWorkspaceOutlinerItemExport& Export) const 
+{
+	const TInstancedStruct<FWorkspaceOutlinerItemData>& Data = Export.Data;
+	if (Data.IsValid() && Data.GetScriptStruct() == FAnimNextGraphOutlinerData::StaticStruct())
+	{
+		const FAnimNextGraphOutlinerData& GraphData = Data.Get<FAnimNextGraphOutlinerData>();
+		if (GraphData.GraphInterface)
+		{
+			return GraphData.GraphInterface.GetObject()->GetExternalPackage();
+		}
+	}
+	return nullptr;
+}
+
+const FSlateBrush* FAnimNextGraphItemDetails::GetItemIcon() const
+{
+	return FAppStyle::GetBrush(TEXT("GraphEditor.EventGraph_24x"));
+}
+
+void FAnimNextGraphItemDetails::RegisterToolMenuExtensions()
+{
+	FToolMenuOwnerScoped OwnerScoped(TEXT("FAnimNextGraphItemDetails"));
+	if (UToolMenu* Menu = UToolMenus::Get()->ExtendMenu("WorkspaceOutliner.ItemContextMenu"))
+	{
+		Menu->AddDynamicSection(TEXT("AnimNextGraphItem"), FNewToolMenuDelegate::CreateLambda([](UToolMenu* InMenu)
+		{
+			UWorkspaceItemMenuContext* WorkspaceItemContext = InMenu->FindContext<UWorkspaceItemMenuContext>();
+			const UAssetEditorToolkitMenuContext* AssetEditorContext = InMenu->FindContext<UAssetEditorToolkitMenuContext>();			
+			if (WorkspaceItemContext && AssetEditorContext)
+			{
+				FToolMenuSection& Section = InMenu->AddSection("WorkspaceOutliner.ItemContextMenu.RootAsset", FText::FromString(TEXT("Animation Next")));
+				if(const TSharedPtr<UE::Workspace::IWorkspaceEditor> WorkspaceEditor = StaticCastSharedPtr<UE::Workspace::IWorkspaceEditor>(AssetEditorContext->Toolkit.Pin()))
+				{
+					if (WorkspaceItemContext->SelectedExports.Num() == 1)
+					{
+						TInstancedStruct<FWorkspaceOutlinerItemData>& Data = WorkspaceItemContext->SelectedExports[0].Data;
+						if (Data.IsValid() && Data.GetScriptStruct() == FAnimNextGraphOutlinerData::StaticStruct())
+						{							
+							Section.AddMenuEntry(
+							TEXT("Open"),
+							LOCTEXT("GoHereActionLabel", "Open Graph"),
+							LOCTEXT("GoHereActionTooltip", "Open selected Animation Graph"),
+							FSlateIcon(FAppStyle::GetAppStyleSetName(), "GraphEditor.EventGraph_16x"),
+							FUIAction(
+								FExecuteAction::CreateWeakLambda(WorkspaceItemContext, [&Data, WorkspaceEditor]()
+								{										
+									const FAnimNextGraphOutlinerData& GraphData = Data.Get<FAnimNextGraphOutlinerData>();
+									if (GraphData.GraphInterface)
+									{
+										if (URigVMGraph* RigVMGraph = GraphData.GraphInterface->GetRigVMGraph())
+										{
+											if(const IRigVMClientHost* RigVMClientHost = RigVMGraph->GetImplementingOuter<IRigVMClientHost>())
+											{
+												if(UObject* EditorObject = RigVMClientHost->GetEditorObjectForRigVMGraph(RigVMGraph))
+												{
+													WorkspaceEditor->OpenObjects({EditorObject});
+												}
+											}
+										}
+									}
+								})
+							));
+						}
+					}
+				}
+			}
+		}));
+	}
+}
+
+void FAnimNextGraphItemDetails::UnregisterToolMenuExtensions()
+{
+	if(UToolMenus* ToolMenus = UToolMenus::Get())
+	{
+		ToolMenus->UnregisterOwnerByName("FAnimNextGraphItemDetails");
+	}
+}
+}
+
+#undef LOCTEXT_NAMESPACE // "FAnimNextGraphItemDetails"
\ No newline at end of file
diff --git a/Engine/Plugins/Experimental/AnimNext/Source/AnimNextEditor/Private/Graph/AnimNextGraphItemDetails.h b/Engine/Plugins/Experimental/AnimNext/Source/AnimNextEditor/Private/Graph/AnimNextGraphItemDetails.h
new file mode 100644
--- /dev/null
+++ b/Engine/Plugins/Experimental/AnimNext/Source/AnimNextEditor/Private/Graph/AnimNextGraphItemDetails.h
@@ -0,0 +1,21 @@
+﻿// Copyright Epic Games, Inc. All Rights Reserved.
+
+#pragma once
+
+#include "IWorkspaceOutlinerItemDetails.h"
+
+namespace UE::AnimNext::Editor
+{
+class FAnimNextGraphItemDetails : public UE::Workspace::IWorkspaceOutlinerItemDetails
+{
+public:
+	FAnimNextGraphItemDetails() = default;
+	virtual ~FAnimNextGraphItemDetails() override = default;
+	virtual void HandleDoubleClick(const FToolMenuContext& ToolMenuContext) const override;	
+	virtual UPackage* GetPackage(const FWorkspaceOutlinerItemExport& Export) const override;
+	virtual const FSlateBrush* GetItemIcon() const override;
+
+	static void RegisterToolMenuExtensions();
+	static void UnregisterToolMenuExtensions();
+};
+}
diff --git a/Engine/Plugins/Experimental/AnimNext/Source/AnimNextEditor/Private/Graph/AnimNextGraphPanelNodeFactory.cpp b/Engine/Plugins/Experimental/AnimNext/Source/AnimNextEditor/Private/Graph/AnimNextGraphPanelNodeFactory.cpp
--- a/Engine/Plugins/Experimental/AnimNext/Source/AnimNextEditor/Private/Graph/AnimNextGraphPanelNodeFactory.cpp
+++ b/Engine/Plugins/Experimental/AnimNext/Source/AnimNextEditor/Private/Graph/AnimNextGraphPanelNodeFactory.cpp
@@ -1,12 +1,12 @@
 // Copyright Epic Games, Inc. All Rights Reserved.
 
 #include "Graph/AnimNextGraphPanelNodeFactory.h"
-#include "Graph/AnimNextGraph_EdGraphNode.h"
+#include "AnimNextEdGraphNode.h"
 #include "Graph/SAnimNextGraphNode.h"
 
 TSharedPtr<SGraphNode> FAnimNextGraphPanelNodeFactory::CreateNode(UEdGraphNode* Node) const
 {
-	if (UAnimNextGraph_EdGraphNode* AnimNextGraphNode = Cast<UAnimNextGraph_EdGraphNode>(Node))
+	if (UAnimNextEdGraphNode* AnimNextGraphNode = Cast<UAnimNextEdGraphNode>(Node))
 	{
 		if(AnimNextGraphNode->IsTraitStack())
 		{
diff --git a/Engine/Plugins/Experimental/AnimNext/Source/AnimNextEditor/Private/Graph/AnimNextModuleAssetDefinition.cpp b/Engine/Plugins/Experimental/AnimNext/Source/AnimNextEditor/Private/Graph/AnimNextModuleAssetDefinition.cpp
new file mode 100644
--- /dev/null
+++ b/Engine/Plugins/Experimental/AnimNext/Source/AnimNextEditor/Private/Graph/AnimNextModuleAssetDefinition.cpp
@@ -0,0 +1,104 @@
+﻿// Copyright Epic Games, Inc. All Rights Reserved.
+
+#include "AnimNextModuleAssetDefinition.h"
+#include "AnimNextModuleEditor.h"
+#include "ContentBrowserMenuContexts.h"
+#include "EditorCVars.h"
+#include "FileHelpers.h"
+#include "IWorkspaceEditorModule.h"
+#include "ToolMenus.h"
+#include "Workspace/AnimNextWorkspaceFactory.h"
+
+#define LOCTEXT_NAMESPACE "AssetDefinition_AnimNextGraph"
+
+EAssetCommandResult UAssetDefinition_AnimNextModule::OpenAssets(const FAssetOpenArgs& OpenArgs) const
+{
+	using namespace UE::AnimNext::Editor;
+	using namespace UE::Workspace;
+
+	for (UAnimNextModule* Asset : OpenArgs.LoadObjects<UAnimNextModule>())
+	{
+		if(CVars::GUseWorkspaceEditor.GetValueOnGameThread())
+		{
+			IWorkspaceEditorModule& WorkspaceEditorModule = FModuleManager::Get().LoadModuleChecked<IWorkspaceEditorModule>("WorkspaceEditor");
+			WorkspaceEditorModule.OpenWorkspaceForObject(Asset, EOpenWorkspaceMethod::Default, UAnimNextWorkspaceFactory::StaticClass());
+		}
+		else
+		{
+			TSharedRef<FModuleEditor> ModuleEditor = MakeShared<FModuleEditor>();
+			ModuleEditor->InitEditor(OpenArgs.GetToolkitMode(), OpenArgs.ToolkitHost, Asset);
+		}
+	}
+
+	return EAssetCommandResult::Handled;
+}
+
+
+FText UAssetDefinition_AnimNextParameter::GetObjectDisplayNameText(UObject* Object) const
+{
+	UAnimNextRigVMAssetEntry* Parameter = CastChecked<UAnimNextRigVMAssetEntry>(Object);
+	return Parameter->GetDisplayName();
+}
+
+FText UAssetDefinition_AnimNextAnimationGraph::GetObjectDisplayNameText(UObject* Object) const
+{
+	UAnimNextRigVMAssetEntry* Parameter = CastChecked<UAnimNextRigVMAssetEntry>(Object);
+	return Parameter->GetDisplayName();
+}
+
+FText UAssetDefinition_AnimNextEventGraph::GetObjectDisplayNameText(UObject* Object) const
+{
+	UAnimNextRigVMAssetEntry* Parameter = CastChecked<UAnimNextRigVMAssetEntry>(Object);
+	return Parameter->GetDisplayName();
+}
+
+namespace UE::AnimNext::Editor
+{
+
+static FDelayedAutoRegisterHelper AutoRegisterAssetMenuItems(EDelayedRegisterRunPhase::EndOfEngineInit, []
+{
+	UToolMenus::RegisterStartupCallback(FSimpleMulticastDelegate::FDelegate::CreateLambda([]()
+	{
+		FToolMenuOwnerScoped OwnerScoped(UE_MODULE_NAME);
+		UToolMenu* Menu = UE::ContentBrowser::ExtendToolMenu_AssetContextMenu(UAnimNextModule::StaticClass());
+		 
+		FToolMenuSection& Section = Menu->FindOrAddSection("GetAssetActions");
+		Section.AddDynamicEntry(NAME_None, FNewToolMenuSectionDelegate::CreateLambda([](FToolMenuSection& InSection)
+		{
+			InSection.AddMenuEntry(
+				"ForceSave",
+				LOCTEXT("ForceSaveLabel", "Force Save"),
+				LOCTEXT("ForceSaveTooltip", "Force the save of this item and all its subobjects."),
+				FSlateIcon(FAppStyle::Get().GetStyleSetName(), "Icons.Save"),
+				FToolUIAction(FToolMenuExecuteAction::CreateLambda([](const FToolMenuContext& InContext)
+				{
+					if(UContentBrowserAssetContextMenuContext* AssetContext = InContext.FindContext<UContentBrowserAssetContextMenuContext>())
+					{
+						TArray<UPackage*> PackagesToSave;
+						for (const FAssetData& AssetData : AssetContext->SelectedAssets)
+						{
+							if (UPackage* Package = AssetData.GetPackage())
+							{
+								PackagesToSave.Add(Package);
+								PackagesToSave.Append(Package->GetExternalPackages());
+							}
+						}
+
+						if(PackagesToSave.Num() > 0)
+						{
+							FEditorFileUtils::FPromptForCheckoutAndSaveParams SaveParams;
+							SaveParams.bCheckDirty = false;
+							SaveParams.bPromptToSave = false;
+							SaveParams.bIsExplicitSave = true;
+
+							FEditorFileUtils::PromptForCheckoutAndSave(PackagesToSave, SaveParams);
+						}
+					}
+				})));
+		}));
+	}));
+});
+
+}
+
+#undef LOCTEXT_NAMESPACE
\ No newline at end of file
diff --git a/Engine/Plugins/Experimental/AnimNext/Source/AnimNextEditor/Private/Graph/AnimNextModuleAssetDefinition.h b/Engine/Plugins/Experimental/AnimNext/Source/AnimNextEditor/Private/Graph/AnimNextModuleAssetDefinition.h
new file mode 100644
--- /dev/null
+++ b/Engine/Plugins/Experimental/AnimNext/Source/AnimNextEditor/Private/Graph/AnimNextModuleAssetDefinition.h
@@ -0,0 +1,72 @@
+﻿// Copyright Epic Games, Inc. All Rights Reserved.
+
+#pragma once
+
+#include "Module/AnimNextModule.h"
+#include "AssetDefinitionDefault.h"
+#include "Module/AnimNextModule_EventGraph.h"
+#include "Module/AnimNextModule_Parameter.h"
+#include "Graph/AnimNextModule_AnimationGraph.h"
+#include "AnimNextModuleAssetDefinition.generated.h"
+
+#define LOCTEXT_NAMESPACE "AnimNextAssetDefinitions"
+
+UCLASS()
+class UAssetDefinition_AnimNextModule : public UAssetDefinitionDefault
+{
+	GENERATED_BODY()
+
+public:
+	// UAssetDefinition interface
+	virtual FText GetAssetDisplayName() const override { return LOCTEXT("AnimNextModule", "AnimNext Module"); }
+	virtual FLinearColor GetAssetColor() const override { return FLinearColor(FColor(128,96,48)); }
+	virtual TSoftClassPtr<UObject> GetAssetClass() const override { return UAnimNextModule::StaticClass(); }
+	virtual EAssetCommandResult OpenAssets(const FAssetOpenArgs& OpenArgs) const override;
+	virtual TConstArrayView<FAssetCategoryPath> GetAssetCategories() const override
+	{
+		static const auto Categories = { FAssetCategoryPath(EAssetCategoryPaths::Animation, LOCTEXT("AnimNextSubMenu", "AnimNext")) };
+		return Categories;
+	}
+	virtual bool ShouldSaveExternalPackages() const override { return true; }
+};
+
+UCLASS()
+class UAssetDefinition_AnimNextParameter : public UAssetDefinitionDefault
+{
+	GENERATED_BODY()
+
+public:
+	// UAssetDefinition interface
+	virtual FText GetAssetDisplayName() const override { return LOCTEXT("AnimNextParameter", "Parameter"); }
+	virtual FLinearColor GetAssetColor() const override { return FLinearColor(FColor(128,64,64)); }
+	virtual TSoftClassPtr<UObject> GetAssetClass() const override { return UAnimNextModule_Parameter::StaticClass(); }
+	virtual FText GetObjectDisplayNameText(UObject* Object) const override;
+};
+
+UCLASS()
+class UAssetDefinition_AnimNextAnimationGraph : public UAssetDefinitionDefault
+{
+	GENERATED_BODY()
+
+public:
+	// UAssetDefinition interface
+	virtual FText GetAssetDisplayName() const override { return LOCTEXT("AnimNextAnimationGraph", "Animation Graph"); }
+	virtual FLinearColor GetAssetColor() const override { return FLinearColor(FColor(128,64,64)); }
+	virtual TSoftClassPtr<UObject> GetAssetClass() const override { return UAnimNextModule_AnimationGraph::StaticClass(); }
+	virtual FText GetObjectDisplayNameText(UObject* Object) const override;
+};
+
+UCLASS()
+class UAssetDefinition_AnimNextEventGraph : public UAssetDefinitionDefault
+{
+	GENERATED_BODY()
+
+public:
+	// UAssetDefinition interface
+	virtual FText GetAssetDisplayName() const override { return LOCTEXT("AnimNextEventGraph", "Event Graph"); }
+	virtual FLinearColor GetAssetColor() const override { return FLinearColor(FColor(128,64,64)); }
+	virtual TSoftClassPtr<UObject> GetAssetClass() const override { return UAnimNextModule_EventGraph::StaticClass(); }
+	virtual FText GetObjectDisplayNameText(UObject* Object) const override;
+};
+
+#undef LOCTEXT_NAMESPACE
\ No newline at end of file
diff --git a/Engine/Plugins/Experimental/AnimNext/Source/AnimNextEditor/Private/Graph/AnimNextModuleEditor.cpp b/Engine/Plugins/Experimental/AnimNext/Source/AnimNextEditor/Private/Graph/AnimNextModuleEditor.cpp
new file mode 100644
--- /dev/null
+++ b/Engine/Plugins/Experimental/AnimNext/Source/AnimNextEditor/Private/Graph/AnimNextModuleEditor.cpp
@@ -0,0 +1,326 @@
+﻿// Copyright Epic Games, Inc. All Rights Reserved.
+
+#include "AnimNextModuleEditor.h"
+#include "ModuleEditorMode.h"
+#include "Module/AnimNextModule.h"
+#include "AnimNextEdGraphNode.h"
+#include "Module/AnimNextModule_EditorData.h"
+#include "EdGraphNode_Comment.h"
+#include "Common/SActionMenu.h"
+#include "UncookedOnlyUtils.h"
+#include "Framework/Commands/GenericCommands.h"
+#include "RigVMModel/RigVMController.h"
+#include "Kismet2/BlueprintEditorUtils.h"
+#include "ScopedTransaction.h"
+
+#define LOCTEXT_NAMESPACE "AnimNextGraphEditor"
+
+namespace UE::AnimNext::Editor
+{
+
+namespace GraphModes
+{
+	const FName GraphEditor("AnimNextGraphEditorMode");
+}
+
+namespace GraphTabs
+{
+	const FName Details("DetailsTab");
+	const FName Document("Document");
+}
+
+const FName GraphAppIdentifier("AnimNextGraphEditor");
+
+void FModuleEditor::InitEditor(const EToolkitMode::Type InMode, const TSharedPtr<IToolkitHost>& InInitToolkitHost, UAnimNextModule* InAnimNextModule)
+{
+	AnimNextModule = InAnimNextModule;
+
+	AnimNextModule_EditorData = UE::AnimNext::UncookedOnly::FUtils::GetEditorData(AnimNextModule);
+	AnimNextModule_EditorData->Initialize(false);
+	
+	DocumentManager = MakeShared<FDocumentTracker>();
+	DocumentManager->Initialize(SharedThis(this));
+
+	TSharedRef<FGraphEditorSummoner> GraphEditorSummoner = MakeShared<FGraphEditorSummoner>(SharedThis(this));
+	GraphEditorSummoner->OnCreateGraphEditorWidget().BindSP(this, &FModuleEditor::CreateGraphEditorWidget);
+	GraphEditorSummoner->OnGraphEditorFocused().BindSP(this, &FModuleEditor::OnGraphEditorFocused);
+	GraphEditorSummoner->OnGraphEditorBackgrounded().BindSP(this, &FModuleEditor::OnGraphEditorBackgrounded);
+	GraphEditorTabFactoryPtr = GraphEditorSummoner;
+	DocumentManager->RegisterDocumentFactory(GraphEditorSummoner);
+
+	const bool bCreateDefaultStandaloneMenu = true;
+	const bool bCreateDefaultToolbar = true;
+	FAssetEditorToolkit::InitAssetEditor(InMode, InInitToolkitHost, GraphAppIdentifier, FTabManager::FLayout::NullLayout, bCreateDefaultStandaloneMenu, bCreateDefaultToolbar, InAnimNextModule);
+
+	BindCommands();
+
+	AddApplicationMode(GraphModes::GraphEditor, MakeShared<FModuleEditorMode>(SharedThis(this)));
+	SetCurrentMode(GraphModes::GraphEditor);
+	
+	ExtendMenu();
+	ExtendToolbar();
+	RegenerateMenusAndToolbars();
+}
+
+void FModuleEditor::BindCommands()
+{
+}
+
+void FModuleEditor::ExtendMenu()
+{
+	
+}
+
+void FModuleEditor::ExtendToolbar()
+{
+	
+}
+
+void FModuleEditor::RegisterTabSpawners(const TSharedRef<FTabManager>& InTabManager)
+{
+	DocumentManager->SetTabManager(InTabManager);
+
+	FWorkflowCentricApplication::RegisterTabSpawners(InTabManager);
+}
+
+void FModuleEditor::UnregisterTabSpawners(const TSharedRef<FTabManager>& InTabManager)
+{
+	
+}
+
+FName FModuleEditor::GetToolkitFName() const
+{
+	return FName("AnimNextGraphEditor");
+}
+
+FText FModuleEditor::GetBaseToolkitName() const
+{
+	return LOCTEXT("AppLabel", "AnimNextGraphEditor");
+}
+
+FString FModuleEditor::GetWorldCentricTabPrefix() const
+{
+	return LOCTEXT("WorldCentricTabPrefix", "AnimNextGraphEditor ").ToString();
+}
+
+FLinearColor FModuleEditor::GetWorldCentricTabColorScale() const
+{
+	return FLinearColor(0.3f, 0.2f, 0.5f, 0.5f);
+}
+
+void FModuleEditor::InitToolMenuContext(FToolMenuContext& InMenuContext)
+{
+}
+
+void FModuleEditor::OnGraphEditorFocused(TSharedRef<SGraphEditor> InGraphEditor)
+{
+	// Update the graph editor that is currently focused
+	FocusedGraphEdPtr = InGraphEditor;
+}
+
+void FModuleEditor::OnGraphEditorBackgrounded(TSharedRef<SGraphEditor> InGraphEditor)
+{
+	// Update the graph editor that is currently focused
+	FocusedGraphEdPtr = nullptr;
+}
+
+UEdGraph* FModuleEditor::GetFocusedGraph() const
+{
+	if (FocusedGraphEdPtr.IsValid())
+	{
+		if (UEdGraph* Graph = FocusedGraphEdPtr.Pin()->GetCurrentGraph())
+		{
+			return Graph;
+		}
+	}
+	return nullptr;
+}
+
+URigVMGraph* FModuleEditor::GetFocusedVMGraph() const
+{
+	UAnimNextEdGraph* EdGraph = Cast<UAnimNextEdGraph>(GetFocusedGraph());
+	return AnimNextModule_EditorData->GetRigVMGraphForEditorObject(EdGraph);
+}
+
+URigVMController* FModuleEditor::GetFocusedVMController() const
+{
+	return AnimNextModule_EditorData->GetRigVMClient()->GetController(GetFocusedVMGraph());
+}
+
+TSharedRef<SGraphEditor> FModuleEditor::CreateGraphEditorWidget(TSharedRef<FTabInfo> InTabInfo, UEdGraph* InGraph)
+{
+	if(!GraphEditorCommands.IsValid())
+	{
+		GraphEditorCommands = MakeShared<FUICommandList>();
+	
+		GraphEditorCommands->MapAction(FGenericCommands::Get().Delete,
+			FExecuteAction::CreateSP(this, &FModuleEditor::DeleteSelectedNodes),
+			FCanExecuteAction::CreateSP(this, &FModuleEditor::CanDeleteSelectedNodes));
+	}
+	
+	SGraphEditor::FGraphEditorEvents Events;
+	Events.OnCreateActionMenu = SGraphEditor::FOnCreateActionMenu::CreateSP(this, &FModuleEditor::OnCreateGraphActionMenu);
+	
+	return SNew(SGraphEditor)
+		.AdditionalCommands(GraphEditorCommands)
+		.IsEditable(this, &FModuleEditor::IsEditable, InGraph)
+		.GraphToEdit(InGraph)
+		.GraphEvents(Events)
+		.AssetEditorToolkit(AsShared());
+}
+
+FActionMenuContent FModuleEditor::OnCreateGraphActionMenu(UEdGraph* InGraph, const FVector2D& InNodePosition, const TArray<UEdGraphPin*>& InDraggedPins, bool bAutoExpand, SGraphEditor::FActionMenuClosed InOnMenuClosed)
+{
+	TSharedRef<SActionMenu> ActionMenu = SNew(SActionMenu, InGraph)
+		.AutoExpandActionMenu(bAutoExpand)
+		.NewNodePosition(InNodePosition)
+		.DraggedFromPins(InDraggedPins)
+		.OnClosedCallback(InOnMenuClosed);
+
+	TSharedPtr<SWidget> FilterTextBox = StaticCastSharedRef<SWidget>(ActionMenu->GetFilterTextBox());
+	return FActionMenuContent(StaticCastSharedRef<SWidget>(ActionMenu), FilterTextBox);
+}
+
+void FModuleEditor::OnNodeTitleCommitted(const FText& NewText, ETextCommit::Type CommitInfo, UEdGraphNode* NodeBeingChanged)
+{
+	if (UEdGraphNode_Comment* CommentBeingChanged = Cast<UEdGraphNode_Comment>(NodeBeingChanged))
+	{
+		GetFocusedVMController()->SetCommentTextByName(CommentBeingChanged->GetFName(), NewText.ToString(), CommentBeingChanged->FontSize, CommentBeingChanged->bCommentBubbleVisible, CommentBeingChanged->bColorCommentBubble, true, true);
+	}
+}
+
+void FModuleEditor::CloseDocumentTab(const UObject* DocumentID)
+{
+	TSharedRef<FTabPayload_UObject> Payload = FTabPayload_UObject::Make(DocumentID);
+	DocumentManager->CloseTab(Payload);
+}
+
+bool FModuleEditor::InEditingMode() const
+{
+	// @TODO: disallow editing when debugging when implemented
+
+	return true;
+}
+
+bool FModuleEditor::IsEditable(UEdGraph* InGraph) const
+{
+	return InGraph && InEditingMode() && InGraph->bEditable;
+}
+
+FGraphPanelSelectionSet FModuleEditor::GetSelectedNodes() const
+{
+	FGraphPanelSelectionSet CurrentSelection;
+	TSharedPtr<SGraphEditor> FocusedGraphEd = FocusedGraphEdPtr.Pin();
+	if (FocusedGraphEd.IsValid())
+	{
+		CurrentSelection = FocusedGraphEd->GetSelectedNodes();
+	}
+	return CurrentSelection;
+}
+
+void FModuleEditor::DeleteSelectedNodes()
+{
+	TSharedPtr<SGraphEditor> FocusedGraphEd = FocusedGraphEdPtr.Pin();
+	if (!FocusedGraphEd.IsValid())
+	{
+		return;
+	}
+
+	const FScopedTransaction Transaction(FGenericCommands::Get().Delete->GetDescription());
+	FocusedGraphEd->GetCurrentGraph()->Modify();
+	
+	const FGraphPanelSelectionSet SelectedNodes = GetSelectedNodes();
+
+//	SetUISelectionState(NAME_None);
+
+	if(FocusedGraphEd)
+	{
+		FocusedGraphEd->ClearSelectionSet();
+	}
+
+	TMap<UAnimNextEdGraphNode*, URigVMNode*> NodesToRemove;
+
+	// Some nodes have sub-objects that are represented as other tabs.
+	// Close them here as a pre-pass before we remove their nodes. If the documents are left open they
+	// may reference dangling data and function incorrectly in cases such as FindBlueprintforNodeChecked
+	for (FGraphPanelSelectionSet::TConstIterator NodeIt( SelectedNodes ); NodeIt; ++NodeIt)
+	{
+		if (UEdGraphNode* Node = Cast<UEdGraphNode>(*NodeIt))
+		{
+			if (Node->CanUserDeleteNode())
+			{
+				auto CloseAllDocumentsTab = [this](const UEdGraphNode* InNode)
+				{
+					TArray<UObject*> NodesToClose;
+					GetObjectsWithOuter(InNode, NodesToClose);
+					for (UObject* Node : NodesToClose)
+					{
+						UEdGraph* NodeGraph = Cast<UEdGraph>(Node);
+						if (NodeGraph)
+						{
+							CloseDocumentTab(NodeGraph);
+						}
+					}
+				};
+				
+				if (Node->GetSubGraphs().Num() > 0)
+				{
+					CloseAllDocumentsTab(Node);
+				}
+
+				if (UAnimNextEdGraphNode* AnimNextNode = Cast<UAnimNextEdGraphNode>(Node))
+				{
+					if (URigVMNode* ModelNode = GetFocusedVMController()->GetGraph()->FindNodeByName(*AnimNextNode->GetModelNodePath()))
+					{
+						NodesToRemove.Add(AnimNextNode, ModelNode);
+					}
+				}
+			}
+		}
+	}
+
+	if (NodesToRemove.IsEmpty())
+	{
+		return;
+	}
+
+	GetFocusedVMController()->OpenUndoBracket(TEXT("Delete selected nodes"));
+
+	TArray<URigVMNode*> ModelNodesToRemove;
+	for (auto& It : NodesToRemove)
+	{
+		URigVMNode* ModelNode = It.Value;
+
+		ModelNodesToRemove.Add(ModelNode);
+	}
+
+	GetFocusedVMController()->RemoveNodes(ModelNodesToRemove, true);
+	GetFocusedVMController()->CloseUndoBracket();
+}
+
+bool FModuleEditor::CanDeleteSelectedNodes()
+{
+	const FGraphPanelSelectionSet SelectedNodes = GetSelectedNodes();
+
+	bool bCanUserDeleteNode = false;
+
+	if(IsEditable(GetFocusedGraph()) && SelectedNodes.Num() > 0)
+	{
+		for(UObject* NodeObject : SelectedNodes)
+		{
+			// If any nodes allow deleting, then do not disable the delete option
+			UEdGraphNode* Node = Cast<UEdGraphNode>(NodeObject);
+			if(Node->CanUserDeleteNode())
+			{
+				bCanUserDeleteNode = true;
+				break;
+			}
+		}
+	}
+
+	return bCanUserDeleteNode;
+}
+
+}
+
+#undef LOCTEXT_NAMESPACE
\ No newline at end of file
diff --git a/Engine/Plugins/Experimental/AnimNext/Source/AnimNextEditor/Private/Graph/AnimNextModuleEditor.h b/Engine/Plugins/Experimental/AnimNext/Source/AnimNextEditor/Private/Graph/AnimNextModuleEditor.h
new file mode 100644
--- /dev/null
+++ b/Engine/Plugins/Experimental/AnimNext/Source/AnimNextEditor/Private/Graph/AnimNextModuleEditor.h
@@ -0,0 +1,123 @@
+﻿// Copyright Epic Games, Inc. All Rights Reserved.
+
+#pragma once
+
+#include "CoreMinimal.h"
+#include "WorkflowOrientedApp/WorkflowCentricApplication.h"
+#include "GraphEditor.h"
+
+enum class ERigVMGraphNotifType : uint8;
+class UAnimNextModule;
+class UAnimNextModule_EditorData;
+class FDocumentTracker;
+class FDocumentTabFactory;
+class FTabInfo;
+class FTabManager;
+class IToolkitHost;
+class URigVMGraph;
+class URigVMController;
+class UDetailsViewWrapperObject;
+
+namespace UE::AnimNext::Editor
+{
+
+namespace GraphModes
+{
+	extern const FName GraphEditor;
+}
+
+namespace GraphTabs
+{
+	extern const FName Details;
+	extern const FName Document;
+}
+
+class FModuleEditor : public FWorkflowCentricApplication
+{
+public:
+	/** Edits the specified asset */
+	void InitEditor(const EToolkitMode::Type InMode, const TSharedPtr<IToolkitHost>& InInitToolkitHost, UAnimNextModule* InAnimNextModule);
+
+private:
+	friend class FModuleEditorMode;
+	
+	// IToolkit interface
+	virtual void RegisterTabSpawners(const TSharedRef<FTabManager>& InTabManager) override;
+	virtual void UnregisterTabSpawners(const TSharedRef<FTabManager>& InTabManager) override;
+	virtual FName GetToolkitFName() const override;
+	virtual FText GetBaseToolkitName() const override;
+	virtual FString GetWorldCentricTabPrefix() const override;
+	virtual FLinearColor GetWorldCentricTabColorScale() const override;
+	virtual void InitToolMenuContext(FToolMenuContext& InMenuContext) override;
+
+	void BindCommands();
+	
+	void ExtendMenu();
+
+	void ExtendToolbar();
+
+	void OnGraphEditorFocused(TSharedRef<SGraphEditor> InGraphEditor);
+
+	void OnGraphEditorBackgrounded(TSharedRef<SGraphEditor> InGraphEditor);
+	
+	UEdGraph* GetFocusedGraph() const;
+
+	URigVMGraph* GetFocusedVMGraph() const;
+
+	URigVMController* GetFocusedVMController() const;
+	
+	// Create new graph editor widget for the supplied document container
+	TSharedRef<SGraphEditor> CreateGraphEditorWidget(TSharedRef<FTabInfo> InTabInfo, UEdGraph* InGraph);
+
+	// Create the action menu for the specified graph
+	FActionMenuContent OnCreateGraphActionMenu(UEdGraph* InGraph, const FVector2D& InNodePosition, const TArray<UEdGraphPin*>& InDraggedPins, bool bAutoExpand, SGraphEditor::FActionMenuClosed InOnMenuClosed);
+
+	// Handle node titles being changed
+	void OnNodeTitleCommitted(const FText& NewText, ETextCommit::Type CommitInfo, UEdGraphNode* NodeBeingChanged);
+
+	void CloseDocumentTab(const UObject* DocumentID);
+
+	bool InEditingMode() const;
+
+	bool IsEditable(UEdGraph* InGraph) const;
+	
+	FGraphPanelSelectionSet GetSelectedNodes() const;
+	
+	void DeleteSelectedNodes();
+	
+	bool CanDeleteSelectedNodes();
+
+	TSharedPtr<SGraphEditor> GetGraphEditor(UEdGraph* InEdGraph) const;
+	
+	void HandleModifiedEvent(ERigVMGraphNotifType InNotifType, URigVMGraph* InGraph, UObject* InSubject);
+
+	void HandleDetailsViewCreated(TSharedRef<IDetailsView> InDetailsView)
+	{
+		DetailsView = InDetailsView;
+	}
+	
+	// The asset we are editing
+	UAnimNextModule* AnimNextModule = nullptr;
+
+	// Cached editor data ptr
+	UAnimNextModule_EditorData* AnimNextModule_EditorData = nullptr;
+	
+	// Currently focused graph editor
+	TWeakPtr<SGraphEditor> FocusedGraphEdPtr;
+
+	// Document tracker
+	TSharedPtr<FDocumentTracker> DocumentManager;
+
+	// Factory for document tabs
+	TWeakPtr<FDocumentTabFactory> GraphEditorTabFactoryPtr;
+
+	// Command list for the graph editor
+	TSharedPtr<FUICommandList> GraphEditorCommands;
+
+	// Our details panel
+	TSharedPtr<IDetailsView> DetailsView;
+
+	TArray<TStrongObjectPtr<UDetailsViewWrapperObject>> WrapperObjects;
+};
+
+}
\ No newline at end of file
diff --git a/Engine/Plugins/Experimental/AnimNext/Source/AnimNextEditor/Private/Graph/ModuleEditorMode.cpp b/Engine/Plugins/Experimental/AnimNext/Source/AnimNextEditor/Private/Graph/ModuleEditorMode.cpp
new file mode 100644
--- /dev/null
+++ b/Engine/Plugins/Experimental/AnimNext/Source/AnimNextEditor/Private/Graph/ModuleEditorMode.cpp
@@ -0,0 +1,273 @@
+﻿// Copyright Epic Games, Inc. All Rights Reserved.
+
+#include "ModuleEditorMode.h"
+#include "AnimNextModuleEditor.h"
+#include "PropertyEditorModule.h"
+#include "Widgets/Docking/SDockTab.h"
+#include "EdGraph/EdGraph.h"
+#include "EdGraph/EdGraphSchema.h"
+#include "Widgets/Layout/SSpacer.h"
+#include "Modules/ModuleManager.h"
+
+#define LOCTEXT_NAMESPACE "GraphEditorMode"
+
+namespace UE::AnimNext::Editor
+{
+
+struct FGraphTabHistory : public FGenericTabHistory
+{
+public:
+	FGraphTabHistory(TSharedPtr<FDocumentTabFactory> InFactory, TSharedPtr<FTabPayload> InPayload)
+		: FGenericTabHistory(InFactory, InPayload)
+		, SavedLocation(FVector2D::ZeroVector)
+		, SavedZoomAmount(INDEX_NONE)
+	{
+	}
+
+private:
+	// FGenericTabHistory interface
+	virtual void EvokeHistory(TSharedPtr<::FTabInfo> InTabInfo, bool bPrevTabMatches) override
+	{
+		FWorkflowTabSpawnInfo SpawnInfo;
+		SpawnInfo.Payload = Payload;
+		SpawnInfo.TabInfo = InTabInfo;
+
+		if(bPrevTabMatches)
+		{
+			TSharedPtr<SDockTab> DockTab = InTabInfo->GetTab().Pin();
+			GraphEditor = StaticCastSharedRef<SGraphEditor>(DockTab->GetContent());
+		}
+		else
+		{
+			TSharedRef<SGraphEditor> GraphEditorRef = StaticCastSharedRef<SGraphEditor>(FactoryPtr.Pin()->CreateTabBody(SpawnInfo));
+			GraphEditor = GraphEditorRef;
+			FactoryPtr.Pin()->UpdateTab(InTabInfo->GetTab().Pin(), SpawnInfo, GraphEditorRef);
+		}
+	}
+	
+	virtual void SaveHistory() override
+	{
+		if (IsHistoryValid())
+		{
+			check(GraphEditor.IsValid());
+			GraphEditor.Pin()->GetViewLocation(SavedLocation, SavedZoomAmount);
+			GraphEditor.Pin()->GetViewBookmark(SavedBookmarkId);
+		}
+	}
+	
+	virtual void RestoreHistory() override
+	{
+		if (IsHistoryValid())
+		{
+			check(GraphEditor.IsValid());
+			GraphEditor.Pin()->SetViewLocation(SavedLocation, SavedZoomAmount, SavedBookmarkId);
+		}
+	}
+	
+	/** The graph editor represented by this history node. While this node is inactive, the graph editor is invalid */
+	TWeakPtr<SGraphEditor> GraphEditor;
+	
+	/** Saved location the graph editor was at when this history node was last visited */
+	FVector2D SavedLocation;
+	
+	/** Saved zoom the graph editor was at when this history node was last visited */
+	float SavedZoomAmount;
+	
+	/** Saved bookmark ID the graph editor was at when this history node was last visited */
+	FGuid SavedBookmarkId;
+};
+
+FGraphEditorSummoner::FGraphEditorSummoner(TSharedPtr<FWorkflowCentricApplication> InHostingApp)
+	: FDocumentTabFactoryForObjects<UEdGraph>(GraphTabs::Document, InHostingApp)
+{
+}
+
+void FGraphEditorSummoner::OnTabActivated(TSharedPtr<SDockTab> Tab) const
+{
+	TSharedRef<SGraphEditor> GraphEditor = StaticCastSharedRef<SGraphEditor>(Tab->GetContent());
+	OnGraphEditorFocusedDelegate.ExecuteIfBound(GraphEditor);
+}
+
+void FGraphEditorSummoner::OnTabBackgrounded(TSharedPtr<SDockTab> Tab) const
+{
+	TSharedRef<SGraphEditor> GraphEditor = StaticCastSharedRef<SGraphEditor>(Tab->GetContent());
+	OnGraphEditorBackgroundedDelegate.ExecuteIfBound(GraphEditor);
+}
+
+void FGraphEditorSummoner::OnTabRefreshed(TSharedPtr<SDockTab> Tab) const
+{
+	TSharedRef<SGraphEditor> GraphEditor = StaticCastSharedRef<SGraphEditor>(Tab->GetContent());
+	GraphEditor->NotifyGraphChanged();
+}
+
+void FGraphEditorSummoner::SaveState(TSharedPtr<SDockTab> Tab, TSharedPtr<FTabPayload> Payload) const
+{
+	TSharedRef<SGraphEditor> GraphEditor = StaticCastSharedRef<SGraphEditor>(Tab->GetContent());
+
+	FVector2D ViewLocation;
+	float ZoomAmount;
+	GraphEditor->GetViewLocation(ViewLocation, ZoomAmount);
+
+	UEdGraph* Graph = Payload->IsValid() ? FTabPayload_UObject::CastChecked<UEdGraph>(Payload) : nullptr;
+
+	OnSaveGraphStateDelegate.ExecuteIfBound(Graph, ViewLocation, ZoomAmount);
+}
+
+TAttribute<FText> FGraphEditorSummoner::ConstructTabNameForObject(UEdGraph* DocumentID) const
+{
+	if (DocumentID)
+	{
+		if (const UEdGraphSchema* Schema = DocumentID->GetSchema())
+		{
+			FGraphDisplayInfo Info;
+			Schema->GetGraphDisplayInformation(*DocumentID, /*out*/ Info);
+			return Info.DisplayName;
+		}
+		else
+		{
+			// if we don't have a schema, we're dealing with a malformed (or incomplete graph)...
+			// possibly in the midst of some transaction - here we return the object's outer path 
+			// so we can at least get some context as to which graph we're referring
+			return FText::FromString(DocumentID->GetPathName());
+		}
+	}
+
+	return LOCTEXT("UnknownGraphName", "Unknown");
+}
+
+TSharedRef<SWidget> FGraphEditorSummoner::CreateTabBodyForObject(const FWorkflowTabSpawnInfo& Info, UEdGraph* DocumentID) const
+{
+	check(Info.TabInfo.IsValid());
+	
+	if(OnCreateGraphEditorWidgetDelegate.IsBound())
+	{
+		return OnCreateGraphEditorWidgetDelegate.Execute(Info.TabInfo.ToSharedRef(), DocumentID);
+	}
+	else
+	{
+		return SNew(SSpacer);
+	}
+}
+
+const FSlateBrush* FGraphEditorSummoner::GetTabIconForObject(const FWorkflowTabSpawnInfo& Info, UEdGraph* DocumentID) const
+{
+	if(OnGetTabIconDelegate.IsBound())
+	{
+		return OnGetTabIconDelegate.Execute(DocumentID);
+	}
+	else
+	{
+		return FAppStyle::Get().GetBrush(TEXT("GraphEditor.EventGraph_16x"));
+	}
+}
+
+TSharedRef<FGenericTabHistory> FGraphEditorSummoner::CreateTabHistoryNode(TSharedPtr<FTabPayload> Payload)
+{
+	return MakeShared<FGraphTabHistory>(SharedThis(this), Payload);
+}
+
+class FGraphDetailsTabSummoner : public FWorkflowTabFactory
+{
+public:
+	FGraphDetailsTabSummoner(TSharedPtr<FAssetEditorToolkit> InHostingApp, FOnDetailsViewCreated InOnDetailsViewCreated)
+		: FWorkflowTabFactory(GraphTabs::Details, InHostingApp)
+	{
+		TabLabel = LOCTEXT("DetailsTabLabel", "Details");
+		TabIcon = FSlateIcon("EditorStyle", "LevelEditor.Tabs.Details");
+		ViewMenuDescription = LOCTEXT("DetailsTabMenuDescription", "Details");
+		ViewMenuTooltip = LOCTEXT("DetailsTabToolTip", "Shows the details tab for selected objects.");
+		bIsSingleton = true;
+
+		FDetailsViewArgs DetailsViewArgs;
+		DetailsViewArgs.NameAreaSettings = FDetailsViewArgs::HideNameArea;
+		DetailsViewArgs.bHideSelectionTip = true;
+
+		FPropertyEditorModule& PropertyEditorModule = FModuleManager::LoadModuleChecked<FPropertyEditorModule>("PropertyEditor");
+		DetailsView = PropertyEditorModule.CreateDetailView(DetailsViewArgs);
+		
+		InOnDetailsViewCreated.ExecuteIfBound(DetailsView.ToSharedRef());
+	}
+
+	TSharedPtr<IDetailsView> GetDetailsView() const
+	{
+		return DetailsView;
+	}
+
+private:
+	virtual TSharedRef<SWidget> CreateTabBody(const FWorkflowTabSpawnInfo& Info) const override
+	{
+		return DetailsView.ToSharedRef();
+	}
+	
+	virtual FText GetTabToolTipText(const FWorkflowTabSpawnInfo& Info) const override
+	{
+		return ViewMenuTooltip;
+	}
+	
+	TSharedPtr<IDetailsView> DetailsView;
+};
+
+FModuleEditorMode::FModuleEditorMode(TSharedRef<FWorkflowCentricApplication> InHostingApp)
+	: FApplicationMode(GraphModes::GraphEditor)
+	, HostingAppPtr(InHostingApp)
+{
+	HostingAppPtr = InHostingApp;
+
+	TSharedRef<FModuleEditor> ModuleEditor = StaticCastSharedRef<FModuleEditor>(InHostingApp);
+	
+	TabFactories.RegisterFactory(MakeShared<FGraphDetailsTabSummoner>(ModuleEditor, FOnDetailsViewCreated::CreateSP(&ModuleEditor.Get(), &FModuleEditor::HandleDetailsViewCreated)));
+
+	TabLayout = FTabManager::NewLayout("Standalone_AnimNextGraphEditor_Layout_v1.0")
+		->AddArea
+		(
+			FTabManager::NewPrimaryArea()
+			->SetOrientation(Orient_Vertical)
+			->Split
+			(	
+				FTabManager::NewSplitter()
+				->SetSizeCoefficient(1.0f)
+				->SetOrientation(Orient_Horizontal)
+				->Split
+				(
+					FTabManager::NewStack()
+					->SetSizeCoefficient(0.8f)
+					->SetHideTabWell(false)
+					->AddTab(GraphTabs::Document, ETabState::ClosedTab)
+				)
+				->Split
+				(
+					FTabManager::NewStack()
+					->SetSizeCoefficient(0.2f)
+					->SetHideTabWell(false)
+					->AddTab(GraphTabs::Details, ETabState::OpenedTab)
+				)
+			)
+		);
+
+	ModuleEditor->RegisterModeToolbarIfUnregistered(GetModeName());
+}
+
+void FModuleEditorMode::RegisterTabFactories(TSharedPtr<FTabManager> InTabManager)
+{
+	TSharedPtr<FWorkflowCentricApplication> HostingApp = HostingAppPtr.Pin();
+	HostingApp->PushTabFactories(TabFactories);
+
+	FApplicationMode::RegisterTabFactories(InTabManager);
+}
+
+void FModuleEditorMode::AddTabFactory(FCreateWorkflowTabFactory FactoryCreator)
+{
+	if (FactoryCreator.IsBound())
+	{
+		TabFactories.RegisterFactory(FactoryCreator.Execute(HostingAppPtr.Pin()));
+	}
+}
+
+void FModuleEditorMode::RemoveTabFactory(FName TabFactoryID)
+{
+	TabFactories.UnregisterFactory(TabFactoryID);
+}
+
+}
+
+#undef LOCTEXT_NAMESPACE
\ No newline at end of file
diff --git a/Engine/Plugins/Experimental/AnimNext/Source/AnimNextEditor/Private/Graph/ModuleEditorMode.h b/Engine/Plugins/Experimental/AnimNext/Source/AnimNextEditor/Private/Graph/ModuleEditorMode.h
new file mode 100644
--- /dev/null
+++ b/Engine/Plugins/Experimental/AnimNext/Source/AnimNextEditor/Private/Graph/ModuleEditorMode.h
@@ -0,0 +1,96 @@
+﻿// Copyright Epic Games, Inc. All Rights Reserved.
+
+#pragma once
+
+#include "CoreMinimal.h"
+#include "EdGraph/EdGraph.h"
+#include "WorkflowOrientedApp/ApplicationMode.h"
+#include "WorkflowOrientedApp/WorkflowTabManager.h"
+#include "WorkflowOrientedApp/WorkflowUObjectDocuments.h"
+
+class FWorkflowCentricApplication;
+class SGraphEditor;
+
+namespace UE::AnimNext::Editor
+{
+
+DECLARE_DELEGATE_OneParam(FOnDetailsViewCreated, TSharedRef<IDetailsView>);
+
+struct FGraphEditorSummoner : public FDocumentTabFactoryForObjects<UEdGraph>
+{
+public:
+	// Delegate called when a graph editor widget is created
+	DECLARE_DELEGATE_RetVal_TwoParams(TSharedRef<SGraphEditor>, FOnCreateGraphEditorWidget, TSharedRef<FTabInfo>, UEdGraph*);
+
+	// Delegate called when a graph editor is focused
+	DECLARE_DELEGATE_OneParam(FOnGraphEditorFocused, TSharedRef<SGraphEditor>);
+
+	// Delegate called when a graph editor is backgrounded
+	DECLARE_DELEGATE_OneParam(FOnGraphEditorBackgrounded, TSharedRef<SGraphEditor>);
+
+	// Delegate called to save the state of a graph
+	DECLARE_DELEGATE_ThreeParams(FOnSaveGraphState, UEdGraph*, FVector2D, float);
+
+	// Delegate called to save the state of a graph
+	DECLARE_DELEGATE_RetVal_OneParam(const FSlateBrush*, FOnGetTabIcon, UEdGraph*);
+	
+	FGraphEditorSummoner(TSharedPtr<FWorkflowCentricApplication> InHostingApp);
+
+	FOnCreateGraphEditorWidget& OnCreateGraphEditorWidget() { return OnCreateGraphEditorWidgetDelegate; }
+	FOnGraphEditorFocused& OnGraphEditorFocused() { return OnGraphEditorFocusedDelegate; }
+	FOnGraphEditorBackgrounded& OnGraphEditorBackgrounded() { return OnGraphEditorBackgroundedDelegate; }
+	FOnSaveGraphState& OnSaveGraphState() { return OnSaveGraphStateDelegate; }
+	FOnGetTabIcon& OnGetTabIcon() { return OnGetTabIconDelegate; }
+
+private:
+	// FWorkflowTabFactory interface
+	virtual void OnTabActivated(TSharedPtr<SDockTab> Tab) const override;
+	virtual void OnTabBackgrounded(TSharedPtr<SDockTab> Tab) const override;
+	virtual void OnTabRefreshed(TSharedPtr<SDockTab> Tab) const override;
+	virtual void SaveState(TSharedPtr<SDockTab> Tab, TSharedPtr<FTabPayload> Payload) const override;
+	
+	// FDocumentTabFactoryForObjects interface
+	virtual TAttribute<FText> ConstructTabNameForObject(UEdGraph* DocumentID) const override;
+	virtual TSharedRef<SWidget> CreateTabBodyForObject(const FWorkflowTabSpawnInfo& Info, UEdGraph* DocumentID) const override;
+	virtual const FSlateBrush* GetTabIconForObject(const FWorkflowTabSpawnInfo& Info, UEdGraph* DocumentID) const override;
+	virtual TSharedRef<FGenericTabHistory> CreateTabHistoryNode(TSharedPtr<FTabPayload> Payload) override;
+
+	// The hosting app
+	TWeakPtr<FWorkflowCentricApplication> HostingAppPtr;
+
+	// Delegate called when a graph editor widget is created
+	FOnCreateGraphEditorWidget OnCreateGraphEditorWidgetDelegate;
+
+	// Delegate called when a graph editor is focused
+	FOnGraphEditorFocused OnGraphEditorFocusedDelegate;
+
+	// Delegate called when a graph editor is backgrounded
+	FOnGraphEditorBackgrounded OnGraphEditorBackgroundedDelegate;
+
+	// Delegate called to save the state of a graph
+	FOnSaveGraphState OnSaveGraphStateDelegate;
+
+	// Delegate called to get a tab icon for a graph
+	FOnGetTabIcon OnGetTabIconDelegate;
+};
+
+class FModuleEditorMode : public FApplicationMode
+{
+public:
+	FModuleEditorMode(TSharedRef<FWorkflowCentricApplication> InHostingApp);
+
+private:
+	// FApplicationMode interface
+	virtual void RegisterTabFactories(TSharedPtr<FTabManager> InTabManager) override;
+	virtual void AddTabFactory(FCreateWorkflowTabFactory FactoryCreator) override;
+	virtual void RemoveTabFactory(FName TabFactoryID) override;
+	
+	// The hosting app
+	TWeakPtr<FWorkflowCentricApplication> HostingAppPtr;
+
+	// The tab factories we support
+	FWorkflowAllowedTabSet TabFactories;
+};
+
+}
+
diff --git a/Engine/Plugins/Experimental/AnimNext/Source/AnimNextEditor/Private/Graph/SAnimNextGraphNode.cpp b/Engine/Plugins/Experimental/AnimNext/Source/AnimNextEditor/Private/Graph/SAnimNextGraphNode.cpp
--- a/Engine/Plugins/Experimental/AnimNext/Source/AnimNextEditor/Private/Graph/SAnimNextGraphNode.cpp
+++ b/Engine/Plugins/Experimental/AnimNext/Source/AnimNextEditor/Private/Graph/SAnimNextGraphNode.cpp
@@ -1,7 +1,7 @@
 // Copyright Epic Games, Inc. All Rights Reserved.
 
 #include "Graph/SAnimNextGraphNode.h"
-#include "Graph/AnimNextGraph_EdGraphNode.h"
+#include "AnimNextEdGraphNode.h"
 
 #define LOCTEXT_NAMESPACE "SAnimNextGraphNode"
 
diff --git a/Engine/Plugins/Experimental/AnimNext/Source/AnimNextEditor/Private/Graph/SAnimNextGraphNode.h b/Engine/Plugins/Experimental/AnimNext/Source/AnimNextEditor/Private/Graph/SAnimNextGraphNode.h
--- a/Engine/Plugins/Experimental/AnimNext/Source/AnimNextEditor/Private/Graph/SAnimNextGraphNode.h
+++ b/Engine/Plugins/Experimental/AnimNext/Source/AnimNextEditor/Private/Graph/SAnimNextGraphNode.h
@@ -5,7 +5,7 @@
 #include "CoreMinimal.h"
 #include "Widgets/SRigVMGraphNode.h"
 
-class UAnimNextGraph_EdGraphNode;
+class UAnimNextEdGraphNode;
 
 class SAnimNextGraphNode : public SRigVMGraphNode
 {
@@ -14,7 +14,7 @@ public:
 		: _GraphNodeObj(nullptr)
 		{}
 
-	SLATE_ARGUMENT(UAnimNextGraph_EdGraphNode*, GraphNodeObj)
+	SLATE_ARGUMENT(UAnimNextEdGraphNode*, GraphNodeObj)
 
 	SLATE_END_ARGS()
 
diff --git a/Engine/Plugins/Experimental/AnimNext/Source/AnimNextEditor/Private/Graph/STraitEditorView.cpp b/Engine/Plugins/Experimental/AnimNext/Source/AnimNextEditor/Private/Graph/STraitEditorView.cpp
--- a/Engine/Plugins/Experimental/AnimNext/Source/AnimNextEditor/Private/Graph/STraitEditorView.cpp
+++ b/Engine/Plugins/Experimental/AnimNext/Source/AnimNextEditor/Private/Graph/STraitEditorView.cpp
@@ -10,13 +10,13 @@
 #include "TraitCore/TraitInterfaceRegistry.h"
 #include "TraitCore/TraitRegistry.h"
 #include "ObjectEditorUtils.h"
-#include "Graph/AnimNextGraph_EdGraphNode.h"
+#include "AnimNextEdGraphNode.h"
 #include "STraitListView.h"
 #include "STraitStackView.h"
 #include "MessageLogModule.h"
 #include "IWorkspaceEditor.h"
 #include "IMessageLogListing.h"
-#include "Graph/AnimNextGraph_Controller.h"
+#include "Module/AnimNextModule_Controller.h"
 #include "Graph/AnimNextCompilerResultsTabSummoner.h"
 #include "Logging/MessageLog.h"
 #include "Widgets/Docking/SDockTab.h"
@@ -39,7 +39,7 @@ void STraitEditorView::Construct(const FArguments& InArgs, TWeakPtr<UE::Workspac
 {
 	WorkspaceEditorWeak = MoveTemp(InWorkspaceEditorWeak);
 
-	TWeakObjectPtr<UAnimNextGraph_EdGraphNode> EdGraphNodeWeak = TraitEditorSharedData->EdGraphNodeWeak;
+	TWeakObjectPtr<UAnimNextEdGraphNode> EdGraphNodeWeak = TraitEditorSharedData->EdGraphNodeWeak;
 
 	ChildSlot
 	[
@@ -100,9 +100,9 @@ FReply STraitEditorView::OnTraitClicked(const FTraitUID InClickedTraitUID)
 	FTraitRegistry& TraitRegistry = FTraitRegistry::Get();
 	if (const FTrait* Trait = TraitRegistry.Find(InClickedTraitUID))
 	{
-		if (UAnimNextGraph_EdGraphNode* EdGraphNode = Cast<UAnimNextGraph_EdGraphNode>(TraitEditorSharedData->EdGraphNodeWeak))
+		if (UAnimNextEdGraphNode* EdGraphNode = Cast<UAnimNextEdGraphNode>(TraitEditorSharedData->EdGraphNodeWeak))
 		{
-			if (UAnimNextGraph_Controller* Controller = Cast<UAnimNextGraph_Controller>(EdGraphNode->GetController()))
+			if (UAnimNextModule_Controller* Controller = Cast<UAnimNextModule_Controller>(EdGraphNode->GetController()))
 			{
 				int32 TraitIndex = INDEX_NONE;
 				TSharedPtr<FTraitDataEditorDef> SwapTraitData = FTraitEditorUtils::FindTraitInCurrentStackData(SelectedTraitUID, TraitEditorSharedData->CurrentTraitsDataShared, &TraitIndex);
@@ -145,9 +145,9 @@ FReply STraitEditorView::OnTraitDeleteRequest(const FTraitUID InTraitUIDToDelete
 			{
 				if (TraitEditorSharedData->EdGraphNodeWeak.IsValid())
 				{
-					if (UAnimNextGraph_EdGraphNode* EdGraphNode = Cast<UAnimNextGraph_EdGraphNode>(TraitEditorSharedData->EdGraphNodeWeak.Get()))
+					if (UAnimNextEdGraphNode* EdGraphNode = Cast<UAnimNextEdGraphNode>(TraitEditorSharedData->EdGraphNodeWeak.Get()))
 					{
-						if (UAnimNextGraph_Controller* Controller = Cast<UAnimNextGraph_Controller>(EdGraphNode->GetController()))
+						if (UAnimNextModule_Controller* Controller = Cast<UAnimNextModule_Controller>(EdGraphNode->GetController()))
 						{
 							Controller->RemoveTraitByName(EdGraphNode->GetFName(), TraitDataEditorDef->TraitName);
 						}
@@ -180,9 +180,9 @@ void STraitEditorView::ExecuteTraitDrag(const FTraitUID DraggedTraitUID, const F
 	
 	if (DraggedTrait != nullptr)
 	{
-		if (UAnimNextGraph_EdGraphNode* EdGraphNode = Cast<UAnimNextGraph_EdGraphNode>(TraitEditorSharedData->EdGraphNodeWeak))
+		if (UAnimNextEdGraphNode* EdGraphNode = Cast<UAnimNextEdGraphNode>(TraitEditorSharedData->EdGraphNodeWeak))
 		{
-			if (UAnimNextGraph_Controller* Controller = Cast<UAnimNextGraph_Controller>(EdGraphNode->GetController()))
+			if (UAnimNextModule_Controller* Controller = Cast<UAnimNextModule_Controller>(EdGraphNode->GetController()))
 			{
 				int32 TargetTraitIndex = INDEX_NONE;
 				TSharedPtr<FTraitDataEditorDef> SwapTraitData = FTraitEditorUtils::FindTraitInCurrentStackData(TargetTraitUID, TraitEditorSharedData->CurrentTraitsDataShared, &TargetTraitIndex);
@@ -507,7 +507,7 @@ TSharedRef<SWidget> STraitEditorView::GetOptionsMenuWidget()
 		];
 }
 
-int32 STraitEditorView::GetTraitPinIndex(UAnimNextGraph_EdGraphNode* InEdGraphNode, const TSharedPtr<FTraitDataEditorDef>& InTraitData, int32 InTraitIndex)
+int32 STraitEditorView::GetTraitPinIndex(UAnimNextEdGraphNode* InEdGraphNode, const TSharedPtr<FTraitDataEditorDef>& InTraitData, int32 InTraitIndex)
 {
 	if (InEdGraphNode == nullptr || !InTraitData.IsValid())
 	{
@@ -556,7 +556,7 @@ int32 STraitEditorView::GetTraitPinIndex(UAnimNextGraph_EdGraphNode* InEdGraphNo
 	return INDEX_NONE;
 }
 
-/*static*/ void STraitEditorView::GenerateTraitStackData(const TWeakObjectPtr<UAnimNextGraph_EdGraphNode>& EdGraphNodeWeak, TSharedPtr<FTraitEditorSharedData>& InTraitEditorSharedData)
+/*static*/ void STraitEditorView::GenerateTraitStackData(const TWeakObjectPtr<UAnimNextEdGraphNode>& EdGraphNodeWeak, TSharedPtr<FTraitEditorSharedData>& InTraitEditorSharedData)
 {
 	if (!ensure(InTraitEditorSharedData.IsValid() && InTraitEditorSharedData->CurrentTraitsDataShared.IsValid()))
 	{
@@ -571,7 +571,7 @@ int32 STraitEditorView::GetTraitPinIndex(UAnimNextGraph_EdGraphNode* InEdGraphNo
 	StackUsedInterfaces.Reset();
 	StackMissingInterfaces.Reset();
 
-	if (UAnimNextGraph_EdGraphNode* EdGraphNode = Cast<UAnimNextGraph_EdGraphNode>(EdGraphNodeWeak.Get()))
+	if (UAnimNextEdGraphNode* EdGraphNode = Cast<UAnimNextEdGraphNode>(EdGraphNodeWeak.Get()))
 	{
 		if (URigVMNode* ModelNode = EdGraphNode->GetModelNode())
 		{
diff --git a/Engine/Plugins/Experimental/AnimNext/Source/AnimNextEditor/Private/Graph/STraitEditorView.h b/Engine/Plugins/Experimental/AnimNext/Source/AnimNextEditor/Private/Graph/STraitEditorView.h
--- a/Engine/Plugins/Experimental/AnimNext/Source/AnimNextEditor/Private/Graph/STraitEditorView.h
+++ b/Engine/Plugins/Experimental/AnimNext/Source/AnimNextEditor/Private/Graph/STraitEditorView.h
@@ -9,8 +9,8 @@
 #include "Graph/TraitEditorDefs.h"
 
 
-class UAnimNextGraph_Controller;
-class UAnimNextGraph_EdGraphNode;
+class UAnimNextModule_Controller;
+class UAnimNextEdGraphNode;
 class IMessageLogListing;
 
 namespace UE::Workspace
@@ -37,11 +37,11 @@ namespace UE::AnimNext::Editor
 struct FTraitStackData
 {
 	FTraitStackData() = default;
-	explicit FTraitStackData(const TWeakObjectPtr<UAnimNextGraph_EdGraphNode>& InEdGraphNodeWeak)
+	explicit FTraitStackData(const TWeakObjectPtr<UAnimNextEdGraphNode>& InEdGraphNodeWeak)
 		: EdGraphNodeWeak(InEdGraphNodeWeak)
 	{}
 
-	TWeakObjectPtr<UAnimNextGraph_EdGraphNode> EdGraphNodeWeak = nullptr;
+	TWeakObjectPtr<UAnimNextEdGraphNode> EdGraphNodeWeak = nullptr;
 };
 
 class STraitEditorView : public SCompoundWidget
@@ -75,9 +75,9 @@ private:
 
 	TSharedRef<SWidget> GetOptionsMenuWidget();
 
-	int32 GetTraitPinIndex(UAnimNextGraph_EdGraphNode* InEdGraphNode, const TSharedPtr<FTraitDataEditorDef>& InTraitData, int32 TraitIndex = INDEX_NONE);
+	int32 GetTraitPinIndex(UAnimNextEdGraphNode* InEdGraphNode, const TSharedPtr<FTraitDataEditorDef>& InTraitData, int32 TraitIndex = INDEX_NONE);
 
-	static void GenerateTraitStackData(const TWeakObjectPtr<UAnimNextGraph_EdGraphNode>& EdGraphNodeWeak, TSharedPtr<FTraitEditorSharedData>& InTraitEditorSharedData);
+	static void GenerateTraitStackData(const TWeakObjectPtr<UAnimNextEdGraphNode>& EdGraphNodeWeak, TSharedPtr<FTraitEditorSharedData>& InTraitEditorSharedData);
 
 	TWeakPtr<UE::Workspace::IWorkspaceEditor> WorkspaceEditorWeak;
 	TSharedPtr<FTraitEditorSharedData> TraitEditorSharedData;
diff --git a/Engine/Plugins/Experimental/AnimNext/Source/AnimNextEditor/Private/Graph/STraitListView.h b/Engine/Plugins/Experimental/AnimNext/Source/AnimNextEditor/Private/Graph/STraitListView.h
--- a/Engine/Plugins/Experimental/AnimNext/Source/AnimNextEditor/Private/Graph/STraitListView.h
+++ b/Engine/Plugins/Experimental/AnimNext/Source/AnimNextEditor/Private/Graph/STraitListView.h
@@ -9,7 +9,7 @@
 #include "Widgets/Views/ITableRow.h"
 #include "Graph/TraitEditorDefs.h"
 
-class UAnimNextGraph_EdGraphNode;
+class UAnimNextEdGraphNode;
 
 namespace UE::Workspace
 {
diff --git a/Engine/Plugins/Experimental/AnimNext/Source/AnimNextEditor/Private/Graph/STraitStackView.cpp b/Engine/Plugins/Experimental/AnimNext/Source/AnimNextEditor/Private/Graph/STraitStackView.cpp
--- a/Engine/Plugins/Experimental/AnimNext/Source/AnimNextEditor/Private/Graph/STraitStackView.cpp
+++ b/Engine/Plugins/Experimental/AnimNext/Source/AnimNextEditor/Private/Graph/STraitStackView.cpp
@@ -3,7 +3,7 @@
 #include "STraitStackView.h"
 #include "ContentBrowserModule.h"
 #include "Framework/Commands/GenericCommands.h"
-#include "Graph/AnimNextGraph_EdGraphNode.h"
+#include "AnimNextEdGraphNode.h"
 #include "IContentBrowserSingleton.h"
 #include "RigVMModel/RigVMNode.h"
 #include "SAssetDropTarget.h"
@@ -76,7 +76,7 @@ void STraitStackView::Construct(const FArguments& InArgs, TSharedPtr<FTraitEdito
 		FCanExecuteAction::CreateSP(this, &STraitStackView::HasValidSelection));
 
 	const TSharedPtr<FTraitEditorSharedData>& TraitEditorSharedDataLocal = TraitEditorSharedData;
-	const TWeakObjectPtr<UAnimNextGraph_EdGraphNode> EdGraphNodeWeak = TraitEditorSharedData->EdGraphNodeWeak;
+	const TWeakObjectPtr<UAnimNextEdGraphNode> EdGraphNodeWeak = TraitEditorSharedData->EdGraphNodeWeak;
 
 	ChildSlot
 	[
@@ -114,7 +114,7 @@ void STraitStackView::Construct(const FArguments& InArgs, TSharedPtr<FTraitEdito
 							.IsReadOnly_Lambda([EdGraphNodeWeak]()->bool
 							{
 								bool bIsReadOnly = false;
-								if (UAnimNextGraph_EdGraphNode* EdGraphNode = Cast<UAnimNextGraph_EdGraphNode>(EdGraphNodeWeak.Get()))
+								if (UAnimNextEdGraphNode* EdGraphNode = Cast<UAnimNextEdGraphNode>(EdGraphNodeWeak.Get()))
 								{
 									bIsReadOnly = EdGraphNode->IsDeprecated() || EdGraphNode->IsOutDated();
 								}
@@ -124,7 +124,7 @@ void STraitStackView::Construct(const FArguments& InArgs, TSharedPtr<FTraitEdito
 							{
 								if(InCommitType == ETextCommit::OnEnter)
 								{
-									if(UAnimNextGraph_EdGraphNode* EdGraphNode = Cast<UAnimNextGraph_EdGraphNode>(EdGraphNodeWeak.Get()))
+									if(UAnimNextEdGraphNode* EdGraphNode = Cast<UAnimNextEdGraphNode>(EdGraphNodeWeak.Get()))
 									{
 										// TODO zzz : Find the correct way to change the DisplayName / Title
 										//FScopedTransaction Transaction(LOCTEXT("SetNodeName", "Set Node title"));
diff --git a/Engine/Plugins/Experimental/AnimNext/Source/AnimNextEditor/Private/Graph/STraitStackView.h b/Engine/Plugins/Experimental/AnimNext/Source/AnimNextEditor/Private/Graph/STraitStackView.h
--- a/Engine/Plugins/Experimental/AnimNext/Source/AnimNextEditor/Private/Graph/STraitStackView.h
+++ b/Engine/Plugins/Experimental/AnimNext/Source/AnimNextEditor/Private/Graph/STraitStackView.h
@@ -10,7 +10,7 @@
 #include "TraitCore/TraitMode.h"
 #include "Graph/TraitEditorDefs.h"
 
-class UAnimNextGraph_EdGraphNode;
+class UAnimNextEdGraphNode;
 
 namespace UE::Workspace
 {
diff --git a/Engine/Plugins/Experimental/AnimNext/Source/AnimNextEditor/Private/Graph/TraitEditorDefs.h b/Engine/Plugins/Experimental/AnimNext/Source/AnimNextEditor/Private/Graph/TraitEditorDefs.h
--- a/Engine/Plugins/Experimental/AnimNext/Source/AnimNextEditor/Private/Graph/TraitEditorDefs.h
+++ b/Engine/Plugins/Experimental/AnimNext/Source/AnimNextEditor/Private/Graph/TraitEditorDefs.h
@@ -9,7 +9,7 @@
 #include "DragAndDrop/DecoratedDragDropOp.h"
 #include "DragAndDrop/AssetDragDropOp.h"
 
-class UAnimNextGraph_EdGraphNode;
+class UAnimNextEdGraphNode;
 struct FSlateColor;
 
 namespace UE::Workspace
@@ -124,7 +124,7 @@ struct FTraitCategoryData
 struct FTraitEditorSharedData
 {
 	TWeakPtr<UE::Workspace::IWorkspaceEditor> WorkspaceEditorWeak;
-	TWeakObjectPtr<UAnimNextGraph_EdGraphNode> EdGraphNodeWeak = nullptr;
+	TWeakObjectPtr<UAnimNextEdGraphNode> EdGraphNodeWeak = nullptr;
 	TSharedPtr<TArray<TSharedPtr<FTraitDataEditorDef>>> CurrentTraitsDataShared;
 
 	TArray<FTraitInterfaceUID> StackUsedInterfaces;
diff --git a/Engine/Plugins/Experimental/AnimNext/Source/AnimNextEditor/Private/Module/ModuleFactory.cpp b/Engine/Plugins/Experimental/AnimNext/Source/AnimNextEditor/Private/Module/ModuleFactory.cpp
new file mode 100644
--- /dev/null
+++ b/Engine/Plugins/Experimental/AnimNext/Source/AnimNextEditor/Private/Module/ModuleFactory.cpp
@@ -0,0 +1,40 @@
+﻿// Copyright Epic Games, Inc. All Rights Reserved.
+
+#include "Module/ModuleFactory.h"
+#include "Module/AnimNextModule.h"
+#include "Module/AnimNextModule_EditorData.h"
+#include "UncookedOnlyUtils.h"
+
+UAnimNextModuleFactory::UAnimNextModuleFactory()
+{
+	bCreateNew = true;
+	bEditAfterNew = true;
+	SupportedClass = UAnimNextModule::StaticClass();
+}
+
+bool UAnimNextModuleFactory::ConfigureProperties()
+{
+	return true;
+}
+
+UObject* UAnimNextModuleFactory::FactoryCreateNew(UClass* Class, UObject* InParent, FName Name, EObjectFlags Flags, UObject* Context, FFeedbackContext* Warn, FName CallingContext)
+{
+	EObjectFlags FlagsToUse = Flags | RF_Public | RF_Standalone | RF_Transactional | RF_LoadCompleted;
+	if(InParent == GetTransientPackage())
+	{
+		FlagsToUse &= ~RF_Standalone;
+	}
+
+	UAnimNextModule* NewModule = NewObject<UAnimNextModule>(InParent, Class, Name, FlagsToUse);
+
+	// Create internal editor data
+	UAnimNextModule_EditorData* EditorData = NewObject<UAnimNextModule_EditorData>(NewModule, TEXT("EditorData"), RF_Transactional);
+	NewModule->EditorData = EditorData;
+	EditorData->Initialize(/*bRecompileVM*/false);
+
+	// Compile the initial skeleton
+	UE::AnimNext::UncookedOnly::FUtils::Compile(NewModule);
+	check(!EditorData->bErrorsDuringCompilation);
+
+	return NewModule;
+}
\ No newline at end of file
diff --git a/Engine/Plugins/Experimental/AnimNext/Source/AnimNextEditor/Private/Param/ParameterCustomization.cpp b/Engine/Plugins/Experimental/AnimNext/Source/AnimNextEditor/Private/Param/ParameterCustomization.cpp
--- a/Engine/Plugins/Experimental/AnimNext/Source/AnimNextEditor/Private/Param/ParameterCustomization.cpp
+++ b/Engine/Plugins/Experimental/AnimNext/Source/AnimNextEditor/Private/Param/ParameterCustomization.cpp
@@ -8,9 +8,9 @@
 #include "InstancedPropertyBagStructureDataProvider.h"
 #include "UncookedOnlyUtils.h"
 #include "PropertyHandle.h"
-#include "Graph/AnimNextGraph.h"
-#include "Graph/AnimNextGraph_Parameter.h"
-#include "Graph/AnimNextGraph_EditorData.h"
+#include "Module/AnimNextModule.h"
+#include "Module/AnimNextModule_Parameter.h"
+#include "Module/AnimNextModule_EditorData.h"
 
 #define LOCTEXT_NAMESPACE "ParamTypePropertyCustomization"
 
@@ -27,7 +27,7 @@ void FParameterCustomization::CustomizeDetails(IDetailLayoutBuilder& DetailBuild
 		return;
 	}
 
-	if (UAnimNextGraph_Parameter* Parameter = Cast<UAnimNextGraph_Parameter>(Objects[0].Get()))
+	if (UAnimNextModule_Parameter* Parameter = Cast<UAnimNextModule_Parameter>(Objects[0].Get()))
 	{
 		IDetailCategoryBuilder& ParameterCategory = DetailBuilder.EditCategory(TEXT("Parameter"), FText::GetEmpty(), ECategoryPriority::Important);
 
@@ -35,37 +35,37 @@ void FParameterCustomization::CustomizeDetails(IDetailLayoutBuilder& DetailBuild
 
 		TSharedRef< SWidget > ColumnWidget = SNullWidget::NullWidget;
 
-		if (UAnimNextGraph_EditorData* EditorData = Cast<UAnimNextGraph_EditorData>(Parameter->GetOuter()))
+		if (UAnimNextModule_EditorData* EditorData = Cast<UAnimNextModule_EditorData>(Parameter->GetOuter()))
 		{
 			const FName EntryName = Parameter->GetEntryName();
 			if (UAnimNextRigVMAssetEntry* AssetEntry = EditorData->FindEntry(EntryName)) 
 			{
-				if (UAnimNextGraph* ReferencedGraph = UE::AnimNext::UncookedOnly::FUtils::GetGraph(EditorData))
+				if (UAnimNextModule* ReferencedModule = UE::AnimNext::UncookedOnly::FUtils::GetGraph(EditorData))
 				{
 					FAddPropertyParams AddPropertyParams;
 					TArray<IDetailPropertyRow*> DetailPropertyRows;
 
-					if (const FPropertyBagPropertyDesc* PropertyDesc = ReferencedGraph->DefaultState.State.FindPropertyDescByName(Parameter->GetParamName()))
+					if (const FPropertyBagPropertyDesc* PropertyDesc = ReferencedModule->DefaultState.State.FindPropertyDescByName(Parameter->GetParamName()))
 					{
-						IDetailPropertyRow* DetailPropertyRow = DefaultValueCategory.AddExternalStructureProperty(MakeShared<FInstancePropertyBagStructureDataProvider>(ReferencedGraph->DefaultState.State), Parameter->GetParamName(), EPropertyLocation::Default, AddPropertyParams);
+						IDetailPropertyRow* DetailPropertyRow = DefaultValueCategory.AddExternalStructureProperty(MakeShared<FInstancePropertyBagStructureDataProvider>(ReferencedModule->DefaultState.State), Parameter->GetParamName(), EPropertyLocation::Default, AddPropertyParams);
 						if (TSharedPtr<IPropertyHandle> Handle = DetailPropertyRow->GetPropertyHandle(); Handle.IsValid())
 						{
 							Handle->SetPropertyDisplayName(FText::FromName(EntryName));
 							
-							const TWeakObjectPtr<UAnimNextGraph> ReferencedGraphWeak = ReferencedGraph;
+							const TWeakObjectPtr<UAnimNextModule> ReferencedModuleWeak = ReferencedModule;
 
-							const auto OnPropertyValuePreChange = [ReferencedGraphWeak]()
+							const auto OnPropertyValuePreChange = [ReferencedModuleWeak]()
 								{
-									if (ReferencedGraphWeak.IsValid())
+									if (ReferencedModuleWeak.IsValid())
 									{
-										ReferencedGraphWeak->Modify(); // needed to enable the transaction when we modify the PropertyBag
+										ReferencedModuleWeak->Modify(); // needed to enable the transaction when we modify the PropertyBag
 									}
 								};
-							const auto OnPropertyValueChange = [ReferencedGraphWeak](const FPropertyChangedEvent& InEvent)
+							const auto OnPropertyValueChange = [ReferencedModuleWeak](const FPropertyChangedEvent& InEvent)
 								{
-									if (ReferencedGraphWeak.IsValid())
+									if (ReferencedModuleWeak.IsValid())
 									{
-										if (UAnimNextGraph_EditorData* EditorData = Cast<UAnimNextGraph_EditorData>(ReferencedGraphWeak->EditorData))
+										if (UAnimNextModule_EditorData* EditorData = Cast<UAnimNextModule_EditorData>(ReferencedModuleWeak->EditorData))
 										{
 											if (UAnimNextRigVMAssetEntry* AssetEntry = EditorData->FindEntry(UncookedOnly::FUtils::GetParameterNameFromQualifiedName(InEvent.GetPropertyName())))
 											{
diff --git a/Engine/Plugins/Experimental/AnimNext/Source/AnimNextEditor/Private/Param/SAddParametersDialog.h b/Engine/Plugins/Experimental/AnimNext/Source/AnimNextEditor/Private/Param/SAddParametersDialog.h
--- a/Engine/Plugins/Experimental/AnimNext/Source/AnimNextEditor/Private/Param/SAddParametersDialog.h
+++ b/Engine/Plugins/Experimental/AnimNext/Source/AnimNextEditor/Private/Param/SAddParametersDialog.h
@@ -10,7 +10,7 @@
 
 class SWrapBox;
 class UAnimNextParameterLibrary;
-class UAnimNextGraph_EditorData;
+class UAnimNextModule_EditorData;
 
 namespace UE::AnimNext::Editor
 {
diff --git a/Engine/Plugins/Experimental/AnimNext/Source/AnimNextEditor/Private/Param/SGraphPinParam.cpp b/Engine/Plugins/Experimental/AnimNext/Source/AnimNextEditor/Private/Param/SGraphPinParam.cpp
--- a/Engine/Plugins/Experimental/AnimNext/Source/AnimNextEditor/Private/Param/SGraphPinParam.cpp
+++ b/Engine/Plugins/Experimental/AnimNext/Source/AnimNextEditor/Private/Param/SGraphPinParam.cpp
@@ -6,8 +6,8 @@
 #include "ScopedTransaction.h"
 #include "UncookedOnlyUtils.h"
 #include "EdGraph/RigVMEdGraphNode.h"
-#include "Graph/AnimNextGraph.h"
-#include "Graph/AnimNextGraph_Controller.h"
+#include "Module/AnimNextModule.h"
+#include "Module/AnimNextModule_Controller.h"
 #include "Param/ParamCompatibility.h"
 #include "Param/ParamUtils.h"
 #include "Param/RigVMDispatch_GetLayerParameter.h"
@@ -55,7 +55,7 @@ TSharedRef<SWidget> SGraphPinParam::GetDefaultValueWidget()
 
 			if(Node && ModelPin)
 			{
-				UAnimNextGraph_Controller* Controller = CastChecked<UAnimNextGraph_Controller>(Node->GetController());
+				UAnimNextModule_Controller* Controller = CastChecked<UAnimNextModule_Controller>(Node->GetController());
 				Controller->SetAnimNextParameterNode(Node->GetModelNode(), InParameterBinding.Parameter, InParameterBinding.Type.GetValueType(), InParameterBinding.Type.GetContainerType(), InParameterBinding.Type.GetValueTypeObject(), InParameterBinding.InstanceId);
 			}
 			else
@@ -119,7 +119,7 @@ TSharedRef<SWidget> SGraphPinParam::GetDefaultValueWidget()
 			
 			if(Node && ModelPin)
 			{
-				UAnimNextGraph_Controller* Controller = CastChecked<UAnimNextGraph_Controller>(Node->GetController());
+				UAnimNextModule_Controller* Controller = CastChecked<UAnimNextModule_Controller>(Node->GetController());
 				Controller->SetAnimNextParameterNode(Node->GetModelNode(), NAME_None, EPropertyBagPropertyType::None, EPropertyBagContainerType::None, nullptr, InInstanceId);
 			}
 			else
diff --git a/Engine/Plugins/Experimental/AnimNext/Source/AnimNextEditor/Private/Param/SGraphPinParamName.cpp b/Engine/Plugins/Experimental/AnimNext/Source/AnimNextEditor/Private/Param/SGraphPinParamName.cpp
--- a/Engine/Plugins/Experimental/AnimNext/Source/AnimNextEditor/Private/Param/SGraphPinParamName.cpp
+++ b/Engine/Plugins/Experimental/AnimNext/Source/AnimNextEditor/Private/Param/SGraphPinParamName.cpp
@@ -16,7 +16,7 @@
 #include "RigVMModel/RigVMClient.h"
 #include "AnimNextRigVMAsset.h"
 #include "EdGraph/RigVMEdGraphNode.h"
-#include "Graph/AnimNextGraph_Controller.h"
+#include "Module/AnimNextModule_Controller.h"
 #include "UniversalObjectLocators/AssetLocatorFragment.h"
 
 #define LOCTEXT_NAMESPACE "SGraphPinParamName"
@@ -73,7 +73,7 @@ TSharedRef<SWidget> SGraphPinParamName::GetDefaultValueWidget()
 
 		if(ModelPin && Node)
 		{
-			UAnimNextGraph_Controller* Controller = CastChecked<UAnimNextGraph_Controller>(Node->GetController());
+			UAnimNextModule_Controller* Controller = CastChecked<UAnimNextModule_Controller>(Node->GetController());
 			Controller->SetAnimNextParameterNode(Node->GetModelNode(), InParameterBinding.Parameter, InParameterBinding.Type, InstanceId);
 		}
 		else
diff --git a/Engine/Plugins/Experimental/AnimNext/Source/AnimNextEditor/Private/Param/SParameterPicker.cpp b/Engine/Plugins/Experimental/AnimNext/Source/AnimNextEditor/Private/Param/SParameterPicker.cpp
--- a/Engine/Plugins/Experimental/AnimNext/Source/AnimNextEditor/Private/Param/SParameterPicker.cpp
+++ b/Engine/Plugins/Experimental/AnimNext/Source/AnimNextEditor/Private/Param/SParameterPicker.cpp
@@ -9,7 +9,7 @@
 #include "Param/ParamType.h"
 #include "DetailLayoutBuilder.h"
 #include "EditorUtils.h"
-#include "IAnimNextModule.h"
+#include "IAnimNextModuleInterface.h"
 #include "IDetailTreeNode.h"
 #include "IPropertyRowGenerator.h"
 #include "IStructureDataProvider.h"
@@ -27,7 +27,7 @@
 #include "AnimNextUncookedOnly/Private/AnimNextUncookedOnlyModule.h"
 #include "Component/AnimNextComponent.h"
 #include "Editor/PropertyEditor/Private/SSingleProperty.h"
-#include "Graph/AnimNextGraph.h"
+#include "Module/AnimNextModule.h"
 #include "Param/AnimNextParam.h"
 #include "Param/IParameterSourceType.h"
 #include "Param/ParamUtils.h"
@@ -209,7 +209,7 @@ void SParameterPicker::RefreshEntries()
 		}
 	}
 
-	if(Struct == nullptr || Struct == UAnimNextGraph::StaticClass())
+	if(Struct == nullptr || Struct == UAnimNextModule::StaticClass())
 	{
 		// For AnimNext graphs, we add the structs that are exposed via parameters
 		TMap<FAssetData, FAnimNextParameterProviderAssetRegistryExports> Exports;
@@ -401,7 +401,7 @@ void SParameterPicker::HandleFieldPicked(UE::PropertyViewer::SPropertyViewer::FH
 					if(SourceType.IsValid())
 					{
 						const UStruct* Struct = SourceType->GetStruct(InstanceId);
-						if(Struct && Struct == UAnimNextGraph::StaticClass())
+						if(Struct && Struct == UAnimNextModule::StaticClass())
 						{
 							// Invalidate the instance ID if this is an AnimNext graph, as they dont have instances
 							InstanceId.Reset();
diff --git a/Engine/Plugins/Experimental/AnimNext/Source/AnimNextEditor/Private/Workspace/AnimNextWorkspaceSchema.cpp b/Engine/Plugins/Experimental/AnimNext/Source/AnimNextEditor/Private/Workspace/AnimNextWorkspaceSchema.cpp
--- a/Engine/Plugins/Experimental/AnimNext/Source/AnimNextEditor/Private/Workspace/AnimNextWorkspaceSchema.cpp
+++ b/Engine/Plugins/Experimental/AnimNext/Source/AnimNextEditor/Private/Workspace/AnimNextWorkspaceSchema.cpp
@@ -1,7 +1,7 @@
 ﻿// Copyright Epic Games, Inc. All Rights Reserved.
 
 #include "AnimNextWorkspaceSchema.h"
-#include "Graph/AnimNextGraph.h"
+#include "Module/AnimNextModule.h"
 #include "Scheduler/AnimNextSchedule.h"
 
 #define LOCTEXT_NAMESPACE "AnimNextWorkspaceSchema"
@@ -16,7 +16,7 @@ TConstArrayView<FTopLevelAssetPath> UAnimNextWorkspaceSchema::GetSupportedAssetC
 	static const FTopLevelAssetPath Assets[] =
 	{
 		UAnimNextSchedule::StaticClass()->GetClassPathName(),
-		UAnimNextGraph::StaticClass()->GetClassPathName(),
+		UAnimNextModule::StaticClass()->GetClassPathName(),
 	};
 	
 	return Assets;
diff --git a/Engine/Plugins/Experimental/AnimNext/Source/AnimNextTestSuite/Private/AnimNextEditorTest.cpp b/Engine/Plugins/Experimental/AnimNext/Source/AnimNextTestSuite/Private/AnimNextEditorTest.cpp
--- a/Engine/Plugins/Experimental/AnimNext/Source/AnimNextTestSuite/Private/AnimNextEditorTest.cpp
+++ b/Engine/Plugins/Experimental/AnimNext/Source/AnimNextTestSuite/Private/AnimNextEditorTest.cpp
@@ -5,12 +5,12 @@
 #include "UncookedOnlyUtils.h"
 #include "Misc/AutomationTest.h"
 #include "Animation/AnimSequence.h"
-#include "Graph/AnimNextGraph.h"
-#include "Graph/AnimNextGraph_Parameter.h"
-#include "Graph/AnimNextGraph_EditorData.h"
-#include "Graph/AnimNextGraph_AnimationGraph.h"
-#include "Graph/AnimNextGraph_EventGraph.h"
-#include "Graph/GraphFactory.h"
+#include "Module/AnimNextModule.h"
+#include "Module/AnimNextModule_Parameter.h"
+#include "Module/AnimNextModule_EditorData.h"
+#include "Graph/AnimNextModule_AnimationGraph.h"
+#include "Module/AnimNextModule_EventGraph.h"
+#include "Module/ModuleFactory.h"
 #include "Param/RigVMDispatch_GetParameter.h"
 #if WITH_EDITOR
 #include "ScopedTransaction.h"
@@ -33,25 +33,25 @@ bool FEditor_Parameters::RunTest(const FString& InParameters)
 
 	ON_SCOPE_EXIT{ FUtils::CleanupAfterTests(); };
 
-	const TStrongObjectPtr<UFactory> GraphFactory(NewObject<UAnimNextGraphFactory>());
-	UAnimNextGraph* Graph = Cast<UAnimNextGraph>(GraphFactory->FactoryCreateNew(UAnimNextGraph::StaticClass(), GetTransientPackage(), TEXT("TestAnimNextGraph"), RF_Transient, nullptr, nullptr, NAME_None));
-	if(Graph == nullptr)
+	const TStrongObjectPtr<UFactory> GraphFactory(NewObject<UAnimNextModuleFactory>());
+	UAnimNextModule* Module = Cast<UAnimNextModule>(GraphFactory->FactoryCreateNew(UAnimNextModule::StaticClass(), GetTransientPackage(), TEXT("TestAnimNextGraph"), RF_Transient, nullptr, nullptr, NAME_None));
+	if(Module == nullptr)
 	{
-		AddError(TEXT("Could not create graph."));
+		AddError(TEXT("Could not create module."));
 		return false;
 	}
 
-	UAnimNextGraph_EditorData* EditorData = UncookedOnly::FUtils::GetEditorData(Graph);
+	UAnimNextModule_EditorData* EditorData = UncookedOnly::FUtils::GetEditorData(Module);
 	if(EditorData == nullptr)
 	{
-		AddError(TEXT("Graph has no editor data."));
+		AddError(TEXT("Module has no editor data."));
 		return false;
 	}
 
 	static FName TestParameterName = TEXT("TestParam");
 	
 	// AddParameter
-	UAnimNextGraph_Parameter* Parameter = nullptr;
+	UAnimNextModule_Parameter* Parameter = nullptr;
 	{
 		FScopedTransaction Transaction(FText::GetEmpty());
 		Parameter = EditorData->AddParameter(TestParameterName, FAnimNextParamType::GetType<bool>());
@@ -62,7 +62,7 @@ bool FEditor_Parameters::RunTest(const FString& InParameters)
 		}
 	}
 
-	AddExpectedError(TEXT("UAnimNextGraph_EditorData::AddParameter: A parameter already exists for the supplied parameter name."));
+	AddExpectedError(TEXT("UAnimNextModule_EditorData::AddParameter: A parameter already exists for the supplied parameter name."));
 	AddErrorIfFalse(EditorData->AddParameter(TestParameterName, FAnimNextParamType::GetType<bool>()) == nullptr, TEXT("Expected duplicate parameter name argument to fail"));
 
 	GEditor->UndoTransaction();
@@ -72,12 +72,12 @@ bool FEditor_Parameters::RunTest(const FString& InParameters)
 	AddErrorIfFalse(EditorData->Entries.Num() == 1, FString::Printf(TEXT("Unexpected entry count found in graph (Have %d, expected 1)."), EditorData->Entries.Num()));
 
 	// Failure cases
-	AddExpectedError(TEXT("UAnimNextGraph_EditorData::AddParameter: Invalid parameter name supplied."));
+	AddExpectedError(TEXT("UAnimNextModule_EditorData::AddParameter: Invalid parameter name supplied."));
 	AddErrorIfFalse(EditorData->AddParameter(NAME_None, FAnimNextParamType::GetType<bool>()) == nullptr, TEXT("Expected invalid argument to fail"));
 
 	auto TestParameterType = [this, EditorData](FAnimNextParamType InType, bool bInRemove = true)
 	{
-		UAnimNextGraph_Parameter* TypedParameter = EditorData->AddParameter(TEXT("TestParam0"), InType);
+		UAnimNextModule_Parameter* TypedParameter = EditorData->AddParameter(TEXT("TestParam0"), InType);
 		const bool bValidParameter = TypedParameter != nullptr;
 		if (bValidParameter && AddErrorIfFalse(bValidParameter, FString::Printf(TEXT("Could not create new parameter of type %s in graph."), *InType.ToString())))
 		{
@@ -130,7 +130,7 @@ bool FEditor_Parameters::RunTest(const FString& InParameters)
 	// Add graph
 	{
 		FScopedTransaction Transaction(FText::GetEmpty());
-		UAnimNextGraph_EventGraph* EventGraph = EditorData->AddEventGraph(TEXT("TestGraph"));
+		UAnimNextModule_EventGraph* EventGraph = EditorData->AddEventGraph(TEXT("TestGraph"));
 		AddErrorIfFalse(EventGraph != nullptr, TEXT("Could not create new event graph in graph."));
 	}
 
@@ -148,7 +148,7 @@ bool FEditor_Parameters::RunTest(const FString& InParameters)
 		IAnimNextRigVMParameterInterface* ParameterEntry = CastChecked<IAnimNextRigVMParameterInterface>(EditorData->FindEntry("TestParam0"));
 		UE_RETURN_ON_ERROR(ParameterEntry != nullptr, TEXT("Could not find new parameter entry."));
 
-		UAnimNextGraph_EventGraph* EventGraph = EditorData->AddEventGraph(TEXT("TestGraph1"));
+		UAnimNextModule_EventGraph* EventGraph = EditorData->AddEventGraph(TEXT("TestGraph1"));
 		UE_RETURN_ON_ERROR(EventGraph != nullptr, TEXT("Could not create new event graph in graph."));
 
 		URigVMGraph* RigVMGraph = EventGraph->GetRigVMGraph();
@@ -159,7 +159,7 @@ bool FEditor_Parameters::RunTest(const FString& InParameters)
 		URigVMPin* ExecutePin = EventNode->FindPin("ExecuteContext");
 		UE_RETURN_ON_ERROR(ExecutePin != nullptr, TEXT("Could find initial execute pin."));
 
-		UAnimNextGraph_Controller* Controller = Cast<UAnimNextGraph_Controller>(EditorData->GetController(EventGraph->GetRigVMGraph()));
+		UAnimNextModule_Controller* Controller = Cast<UAnimNextModule_Controller>(EditorData->GetController(EventGraph->GetRigVMGraph()));
 		URigVMNode* GetParameterNode = Controller->AddGetAnimNextGraphParameterNode(FVector2D::ZeroVector, ParameterEntry->GetParamName(), FAnimNextParamType::GetType<bool>());
 		UE_RETURN_ON_ERROR(GetParameterNode != nullptr, TEXT("Could not add GetParameter node."));
 
@@ -183,11 +183,11 @@ bool FEditor_Parameters_Python::RunTest(const FString& InParameters)
 
 	const TCHAR* Script = TEXT(
 		"asset_tools = unreal.AssetToolsHelpers.get_asset_tools()\n"
-		"graph = unreal.AssetTools.create_asset(asset_tools, asset_name = \"TestGraph\", package_path = \"/Game/\", asset_class = unreal.AnimNextGraph, factory = unreal.AnimNextGraphFactory())\n"
-		"graph.add_parameter(name = \"TestParam\", value_type = unreal.PropertyBagPropertyType.BOOL, container_type = unreal.PropertyBagContainerType.NONE)\n"
-		"graph.add_event_graph(name = \"TestEventGraph\")\n"
-		"graph.add_animation_graph(name = \"TestAnimationGraph\")\n"
-		"unreal.EditorAssetLibrary.delete_loaded_asset(graph)\n"
+		"module = unreal.AssetTools.create_asset(asset_tools, asset_name = \"TestModule\", package_path = \"/Game/\", asset_class = unreal.AnimNextModule, factory = unreal.AnimNextModuleFactory())\n"
+		"module.add_parameter(name = \"TestParam\", value_type = unreal.PropertyBagPropertyType.BOOL, container_type = unreal.PropertyBagContainerType.NONE)\n"
+		"module.add_event_graph(name = \"TestEventGraph\")\n"
+		"module.add_animation_graph(name = \"TestAnimationGraph\")\n"
+		"unreal.EditorAssetLibrary.delete_loaded_asset(module)\n"
 	);
 
 	IPythonScriptPlugin::Get()->ExecPythonCommand(Script);
@@ -203,15 +203,15 @@ bool FEditor_Graph::RunTest(const FString& InParameters)
 {
 	using namespace UE::AnimNext;
 
-	const TStrongObjectPtr<UFactory> GraphFactory(NewObject<UAnimNextGraphFactory>());
-	UAnimNextGraph* Graph = Cast<UAnimNextGraph>(GraphFactory->FactoryCreateNew(UAnimNextGraph::StaticClass(), GetTransientPackage(), TEXT("TestAnimNextGraph"), RF_Transient, nullptr, nullptr, NAME_None));
+	const TStrongObjectPtr<UFactory> GraphFactory(NewObject<UAnimNextModuleFactory>());
+	UAnimNextModule* Graph = Cast<UAnimNextModule>(GraphFactory->FactoryCreateNew(UAnimNextModule::StaticClass(), GetTransientPackage(), TEXT("TestAnimNextGraph"), RF_Transient, nullptr, nullptr, NAME_None));
 	if(Graph == nullptr)
 	{
 		AddError(TEXT("Could not create graph."));
 		return false;
 	}
 
-	UAnimNextGraph_EditorData* EditorData = UncookedOnly::FUtils::GetEditorData(Graph);
+	UAnimNextModule_EditorData* EditorData = UncookedOnly::FUtils::GetEditorData(Graph);
 	if(EditorData == nullptr)
 	{
 		AddError(TEXT("Graph has no editor data."));
@@ -219,7 +219,7 @@ bool FEditor_Graph::RunTest(const FString& InParameters)
 	}
 
 	// Add graph
-	UAnimNextGraph_AnimationGraph* GraphEntry = nullptr;
+	UAnimNextModule_AnimationGraph* GraphEntry = nullptr;
 	{
 		FScopedTransaction Transaction(FText::GetEmpty());
 		GraphEntry = EditorData->AddAnimationGraph(TEXT("TestGraph"));
diff --git a/Engine/Plugins/Experimental/AnimNext/Source/AnimNextTestSuite/Private/AnimNextParametersTest.cpp b/Engine/Plugins/Experimental/AnimNext/Source/AnimNextTestSuite/Private/AnimNextParametersTest.cpp
--- a/Engine/Plugins/Experimental/AnimNext/Source/AnimNextTestSuite/Private/AnimNextParametersTest.cpp
+++ b/Engine/Plugins/Experimental/AnimNext/Source/AnimNextTestSuite/Private/AnimNextParametersTest.cpp
@@ -10,10 +10,10 @@
 #include "Animation/AnimSequence.h"
 #include "Animation/AnimMontage.h"
 #include "Async/ParallelFor.h"
-#include "Graph/AnimNextGraph.h"
-#include "Graph/AnimNextGraph_EditorData.h"
-#include "Graph/AnimNextGraph_Parameter.h"
-#include "Graph/GraphFactory.h"
+#include "Module/AnimNextModule.h"
+#include "Module/AnimNextModule_EditorData.h"
+#include "Module/AnimNextModule_Parameter.h"
+#include "Module/ModuleFactory.h"
 #include "Param/ParamStack.h"
 #include "UObject/StrongObjectPtr.h"
 
diff --git a/Engine/Plugins/Experimental/AnimNext/Source/AnimNextTestSuite/Private/AnimNextRuntimeTest.cpp b/Engine/Plugins/Experimental/AnimNext/Source/AnimNextTestSuite/Private/AnimNextRuntimeTest.cpp
--- a/Engine/Plugins/Experimental/AnimNext/Source/AnimNextTestSuite/Private/AnimNextRuntimeTest.cpp
+++ b/Engine/Plugins/Experimental/AnimNext/Source/AnimNextTestSuite/Private/AnimNextRuntimeTest.cpp
@@ -4,7 +4,7 @@
 #include "AnimNextRuntimeTest.h"
 
 #include "TraitCore/TraitReader.h"
-#include "Graph/AnimNextGraph.h"
+#include "Module/AnimNextModule.h"
 #include "Graph/RigUnit_AnimNextGraphRoot.h"
 #include "Misc/AutomationTest.h"
 #include "Serialization/MemoryReader.h"
@@ -25,30 +25,30 @@ namespace UE::AnimNext
 		Swap(Registry, TmpRegistry);
 	}
 
-	bool FTestUtils::LoadFromArchiveBuffer(UAnimNextGraph& Graph, TArray<FNodeHandle>& NodeHandles, const TArray<uint8>& SharedDataArchiveBuffer)
+	bool FTestUtils::LoadFromArchiveBuffer(UAnimNextModule& Module, TArray<FNodeHandle>& NodeHandles, const TArray<uint8>& SharedDataArchiveBuffer)
 	{
 		FAnimNextGraphEvaluatorExecuteDefinition ExecuteDefinition;
 		ExecuteDefinition.Hash = 0;
 		ExecuteDefinition.MethodName = TEXT("Execute_0");
 
 		// Manually add our entry point since we didn't go through a full RigVM graph
-		Graph.EntryPoints.Reset();
+		Module.EntryPoints.Reset();
 
-		FAnimNextGraphEntryPoint& EntryPoint = Graph.EntryPoints.AddDefaulted_GetRef();
-		EntryPoint.EntryPointName = Graph.GetDefaultEntryPoint();
+		FAnimNextGraphEntryPoint& EntryPoint = Module.EntryPoints.AddDefaulted_GetRef();
+		EntryPoint.EntryPointName = Module.GetDefaultEntryPoint();
 		EntryPoint.RootTraitHandle = FAnimNextEntryPointHandle(NodeHandles[0]);
-		Graph.ExecuteDefinition = ExecuteDefinition;
-		Graph.SharedDataArchiveBuffer = SharedDataArchiveBuffer;
-		Graph.GraphReferencedObjects.Empty();
+		Module.ExecuteDefinition = ExecuteDefinition;
+		Module.SharedDataArchiveBuffer = SharedDataArchiveBuffer;
+		Module.GraphReferencedObjects.Empty();
 
 		// Reconstruct our graph shared data
 		FMemoryReader GraphSharedDataArchive(SharedDataArchiveBuffer);
-		FTraitReader TraitReader(Graph.GraphReferencedObjects, GraphSharedDataArchive);
+		FTraitReader TraitReader(Module.GraphReferencedObjects, GraphSharedDataArchive);
 
-		const FTraitReader::EErrorState ErrorState = TraitReader.ReadGraph(Graph.SharedDataBuffer);
+		const FTraitReader::EErrorState ErrorState = TraitReader.ReadGraph(Module.SharedDataBuffer);
 		if (ErrorState == FTraitReader::EErrorState::None)
 		{
-			Graph.ResolvedRootTraitHandles.Add(Graph.GetDefaultEntryPoint(), TraitReader.ResolveEntryPointHandle(Graph.EntryPoints[0].RootTraitHandle));
+			Module.ResolvedRootTraitHandles.Add(Module.GetDefaultEntryPoint(), TraitReader.ResolveEntryPointHandle(Module.EntryPoints[0].RootTraitHandle));
 
 			for (FNodeHandle& NodeHandle : NodeHandles)
 			{
@@ -61,8 +61,8 @@ namespace UE::AnimNext
 		}
 		else
 		{
-			Graph.SharedDataBuffer.Empty(0);
-			Graph.ResolvedRootTraitHandles.Add(FRigUnit_AnimNextGraphRoot::DefaultEntryPoint, FAnimNextTraitHandle());
+			Module.SharedDataBuffer.Empty(0);
+			Module.ResolvedRootTraitHandles.Add(FRigUnit_AnimNextGraphRoot::DefaultEntryPoint, FAnimNextTraitHandle());
 			return false;
 		}
 	}
