Rename AnimNextGraph asset to AnimNextModule

diff --git a/Engine/Plugins/Experimental/AnimNext/Source/AnimNextTestSuite/Private/AnimNextTraitBaseTest.cpp b/Engine/Plugins/Experimental/AnimNext/Source/AnimNextTestSuite/Private/AnimNextTraitBaseTest.cpp
--- a/Engine/Plugins/Experimental/AnimNext/Source/AnimNextTestSuite/Private/AnimNextTraitBaseTest.cpp
+++ b/Engine/Plugins/Experimental/AnimNext/Source/AnimNextTestSuite/Private/AnimNextTraitBaseTest.cpp
@@ -19,8 +19,8 @@
 #include "TraitCore/NodeInstance.h"
 #include "TraitCore/NodeTemplateBuilder.h"
 #include "TraitCore/NodeTemplateRegistry.h"
-#include "Graph/AnimNextGraph.h"
-#include "Graph/GraphFactory.h"
+#include "Module/AnimNextModule.h"
+#include "Module/ModuleFactory.h"
 
 //****************************************************************************
 // AnimNext Runtime TraitBase Tests
@@ -521,9 +521,9 @@ bool FAnimationAnimNextRuntimeTest_NodeLifetime::RunTest(const FString& InParame
 		AUTO_REGISTER_ANIM_TRAIT(FTraitAB_Add)
 		AUTO_REGISTER_ANIM_TRAIT(FTraitAC_Add)
 
-		UFactory* GraphFactory = NewObject<UAnimNextGraphFactory>();
-		UAnimNextGraph* AnimNextGraph = CastChecked<UAnimNextGraph>(GraphFactory->FactoryCreateNew(UAnimNextGraph::StaticClass(), GetTransientPackage(), TEXT("TestAnimNextGraph"), RF_Transient, nullptr, nullptr, NAME_None));
-		UE_RETURN_ON_ERROR(AnimNextGraph != nullptr, "FAnimationAnimNextRuntimeTest_NodeLifetime -> Failed to create animation graph");
+		UFactory* ModuleFactory = NewObject<UAnimNextModuleFactory>();
+		UAnimNextModule* Module = CastChecked<UAnimNextModule>(ModuleFactory->FactoryCreateNew(UAnimNextModule::StaticClass(), GetTransientPackage(), TEXT("TestAnimNextModule"), RF_Transient, nullptr, nullptr, NAME_None));
+		UE_RETURN_ON_ERROR(Module != nullptr, "FAnimationAnimNextRuntimeTest_NodeLifetime -> Failed to create module");
 
 		FScopedClearNodeTemplateRegistry ScopedClearNodeTemplateRegistry;
 		FNodeTemplateRegistry& Registry = FNodeTemplateRegistry::Get();
@@ -582,10 +582,10 @@ bool FAnimationAnimNextRuntimeTest_NodeLifetime::RunTest(const FString& InParame
 		}
 
 		// Read our graph
-		FTestUtils::LoadFromArchiveBuffer(*AnimNextGraph, NodeHandles, GraphSharedDataArchiveBuffer);
+		FTestUtils::LoadFromArchiveBuffer(*Module, NodeHandles, GraphSharedDataArchiveBuffer);
 
 		FAnimNextGraphInstancePtr GraphInstance;
-		AnimNextGraph->AllocateInstance(GraphInstance);
+		Module->AllocateInstance(GraphInstance);
 
 		FExecutionContext Context(GraphInstance);
 
@@ -730,8 +730,8 @@ bool FAnimationAnimNextRuntimeTest_TraitStackBinding::RunTest(const FString& InP
 		AUTO_REGISTER_ANIM_TRAIT(FTraitAB_Add)
 		AUTO_REGISTER_ANIM_TRAIT(FTraitAC_Add)
 
-		UFactory* GraphFactory = NewObject<UAnimNextGraphFactory>();
-		UAnimNextGraph* AnimNextGraph = CastChecked<UAnimNextGraph>(GraphFactory->FactoryCreateNew(UAnimNextGraph::StaticClass(), GetTransientPackage(), TEXT("TestAnimNextGraph"), RF_Transient, nullptr, nullptr, NAME_None));
+		UFactory* GraphFactory = NewObject<UAnimNextModuleFactory>();
+		UAnimNextModule* AnimNextGraph = CastChecked<UAnimNextModule>(GraphFactory->FactoryCreateNew(UAnimNextModule::StaticClass(), GetTransientPackage(), TEXT("TestAnimNextGraph"), RF_Transient, nullptr, nullptr, NAME_None));
 		UE_RETURN_ON_ERROR(AnimNextGraph != nullptr, "FAnimationAnimNextRuntimeTest_TraitStackBinding -> Failed to create animation graph");
 
 		FScopedClearNodeTemplateRegistry ScopedClearNodeTemplateRegistry;
@@ -1154,9 +1154,9 @@ bool FAnimationAnimNextRuntimeTest_TraitSerialization::RunTest(const FString& In
 		AUTO_REGISTER_ANIM_TRAIT(FTraitSerialization_Add)
 		AUTO_REGISTER_ANIM_TRAIT(FTraitNativeSerialization_Add)
 
-		UFactory* GraphFactory = NewObject<UAnimNextGraphFactory>();
-		UAnimNextGraph* AnimNextGraph = CastChecked<UAnimNextGraph>(GraphFactory->FactoryCreateNew(UAnimNextGraph::StaticClass(), GetTransientPackage(), TEXT("TestAnimNextGraph"), RF_Transient, nullptr, nullptr, NAME_None));
-		UE_RETURN_ON_ERROR(AnimNextGraph != nullptr, "FAnimationAnimNextRuntimeTest_TraitSerialization -> Failed to create animation graph");
+		UFactory* ModuleFactory = NewObject<UAnimNextModuleFactory>();
+		UAnimNextModule* Module = CastChecked<UAnimNextModule>(ModuleFactory->FactoryCreateNew(UAnimNextModule::StaticClass(), GetTransientPackage(), TEXT("TestAnimNextGraph"), RF_Transient, nullptr, nullptr, NAME_None));
+		UE_RETURN_ON_ERROR(Module != nullptr, "FAnimationAnimNextRuntimeTest_TraitSerialization -> Failed to create module");
 
 		TArray<FTraitUID> NodeTemplateTraitList;
 		NodeTemplateTraitList.Add(FTraitSerialization_Base::TraitUID);
@@ -1359,10 +1359,10 @@ bool FAnimationAnimNextRuntimeTest_TraitSerialization::RunTest(const FString& In
 			AddErrorIfFalse(Registry.GetNum() == 0, "FAnimationAnimNextRuntimeTest_TraitSerialization -> Registry should contain 0 templates");
 
 			// Read our graph
-			FTestUtils::LoadFromArchiveBuffer(*AnimNextGraph, NodeHandles, GraphSharedDataArchiveBuffer);
+			FTestUtils::LoadFromArchiveBuffer(*Module, NodeHandles, GraphSharedDataArchiveBuffer);
 
 			FAnimNextGraphInstancePtr GraphInstance;
-			AnimNextGraph->AllocateInstance(GraphInstance);
+			Module->AllocateInstance(GraphInstance);
 
 			FExecutionContext Context(GraphInstance);
 
diff --git a/Engine/Plugins/Experimental/AnimNext/Source/AnimNextTestSuite/Private/AnimNextTraitEventTest.cpp b/Engine/Plugins/Experimental/AnimNext/Source/AnimNextTestSuite/Private/AnimNextTraitEventTest.cpp
--- a/Engine/Plugins/Experimental/AnimNext/Source/AnimNextTestSuite/Private/AnimNextTraitEventTest.cpp
+++ b/Engine/Plugins/Experimental/AnimNext/Source/AnimNextTestSuite/Private/AnimNextTraitEventTest.cpp
@@ -16,8 +16,8 @@
 #include "TraitCore/ExecutionContext.h"
 #include "TraitCore/NodeTemplateBuilder.h"
 #include "TraitCore/NodeTemplateRegistry.h"
-#include "Graph/AnimNextGraph.h"
-#include "Graph/GraphFactory.h"
+#include "Module/AnimNextModule.h"
+#include "Module/ModuleFactory.h"
 
 //****************************************************************************
 // AnimNext Runtime TraitEvent Tests
@@ -240,9 +240,9 @@ bool FAnimationAnimNextRuntimeTest_TraitEventRaising::RunTest(const FString& InP
 		AUTO_REGISTER_ANIM_TRAIT(FTraitCoreTest_EventA_Add)
 		AUTO_REGISTER_ANIM_TRAIT(FTraitCoreTest_EventB_Add)
 
-		UFactory* GraphFactory = NewObject<UAnimNextGraphFactory>();
-		UAnimNextGraph* AnimNextGraph = CastChecked<UAnimNextGraph>(GraphFactory->FactoryCreateNew(UAnimNextGraph::StaticClass(), GetTransientPackage(), TEXT("TestAnimNextGraph"), RF_Transient, nullptr, nullptr, NAME_None));
-		UE_RETURN_ON_ERROR(AnimNextGraph != nullptr, "FAnimationAnimNextRuntimeTest_TraitEventRaising -> Failed to create animation graph");
+		UFactory* ModuleFactory = NewObject<UAnimNextModuleFactory>();
+		UAnimNextModule* Module = CastChecked<UAnimNextModule>(ModuleFactory->FactoryCreateNew(UAnimNextModule::StaticClass(), GetTransientPackage(), TEXT("TestAnimNextGraph"), RF_Transient, nullptr, nullptr, NAME_None));
+		UE_RETURN_ON_ERROR(Module != nullptr, "FAnimationAnimNextRuntimeTest_TraitEventRaising -> Failed to create module");
 
 		FScopedClearNodeTemplateRegistry ScopedClearNodeTemplateRegistry;
 		FNodeTemplateRegistry& Registry = FNodeTemplateRegistry::Get();
@@ -291,10 +291,10 @@ bool FAnimationAnimNextRuntimeTest_TraitEventRaising::RunTest(const FString& InP
 		}
 
 		// Read our graph
-		FTestUtils::LoadFromArchiveBuffer(*AnimNextGraph, NodeHandles, GraphSharedDataArchiveBuffer);
+		FTestUtils::LoadFromArchiveBuffer(*Module, NodeHandles, GraphSharedDataArchiveBuffer);
 
 		FAnimNextGraphInstancePtr GraphInstance;
-		AnimNextGraph->AllocateInstance(GraphInstance);
+		Module->AllocateInstance(GraphInstance);
 
 		FExecutionContext Context(GraphInstance);
 
diff --git a/Engine/Plugins/Experimental/AnimNext/Source/AnimNextTestSuite/Public/AnimNextRuntimeTest.h b/Engine/Plugins/Experimental/AnimNext/Source/AnimNextTestSuite/Public/AnimNextRuntimeTest.h
--- a/Engine/Plugins/Experimental/AnimNext/Source/AnimNextTestSuite/Public/AnimNextRuntimeTest.h
+++ b/Engine/Plugins/Experimental/AnimNext/Source/AnimNextTestSuite/Public/AnimNextRuntimeTest.h
@@ -13,7 +13,7 @@
 #include "TraitCore/NodeTemplate.h"
 #include "TraitCore/NodeTemplateRegistry.h"
 
-class UAnimNextGraph;
+class UAnimNextModule;
 
 namespace UE::AnimNext
 {
@@ -82,9 +82,9 @@ namespace UE::AnimNext
 
 	struct ANIMNEXTTESTSUITE_API FTestUtils final
 	{
-		// Loads the graph data from the provided archive buffer and returns true on success, false otherwise
+		// Loads the module data from the provided archive buffer and returns true on success, false otherwise
 		// On success, we resolve every node handle provided as argument
-		static bool LoadFromArchiveBuffer(UAnimNextGraph& Graph, TArray<FNodeHandle>& NodeHandles, const TArray<uint8>& SharedDataArchiveBuffer);
+		static bool LoadFromArchiveBuffer(UAnimNextModule& Module, TArray<FNodeHandle>& NodeHandles, const TArray<uint8>& SharedDataArchiveBuffer);
 	};
 }
 
diff --git a/Engine/Plugins/Experimental/AnimNext/Source/AnimNextUncookedOnly/Internal/AnimNextEdGraph.h b/Engine/Plugins/Experimental/AnimNext/Source/AnimNextUncookedOnly/Internal/AnimNextEdGraph.h
new file mode 100644
--- /dev/null
+++ b/Engine/Plugins/Experimental/AnimNext/Source/AnimNextUncookedOnly/Internal/AnimNextEdGraph.h
@@ -0,0 +1,33 @@
+﻿// Copyright Epic Games, Inc. All Rights Reserved.
+
+#pragma once
+
+#include "CoreMinimal.h"
+#include "EdGraph/RigVMEdGraph.h"
+#include "AnimNextEdGraph.generated.h"
+
+class UAnimNextRigVMAssetEditorData;
+
+namespace UE::AnimNext::UncookedOnly
+{
+	struct FUtils;
+}
+
+/**
+  * Wraps UEdGraph which represents the node graph
+  */
+UCLASS(MinimalAPI)
+class UAnimNextEdGraph : public URigVMEdGraph
+{
+	GENERATED_BODY()
+	
+	friend class UAnimNextModule_EditorData;
+
+	// UObject interface
+	virtual void PostLoad() override;
+
+	// URigVMEdGraph interface
+	virtual FRigVMClient* GetRigVMClient() const override;
+	
+	void Initialize(UAnimNextRigVMAssetEditorData* InEditorData);
+};
\ No newline at end of file
diff --git a/Engine/Plugins/Experimental/AnimNext/Source/AnimNextUncookedOnly/Internal/AnimNextEdGraphNode.h b/Engine/Plugins/Experimental/AnimNext/Source/AnimNextUncookedOnly/Internal/AnimNextEdGraphNode.h
new file mode 100644
--- /dev/null
+++ b/Engine/Plugins/Experimental/AnimNext/Source/AnimNextUncookedOnly/Internal/AnimNextEdGraphNode.h
@@ -0,0 +1,35 @@
+﻿// Copyright Epic Games, Inc. All Rights Reserved.
+
+#pragma once
+
+#include "CoreMinimal.h"
+#include "EdGraph/RigVMEdGraphNode.h"
+
+#include "AnimNextEdGraphNode.generated.h"
+
+// EdGraphNode representation for AnimNext nodes
+// A node can hold a trait stack or a trait entry
+UCLASS(MinimalAPI)
+class UAnimNextEdGraphNode : public URigVMEdGraphNode
+{
+	GENERATED_BODY()
+
+public:
+	//////////////////////////////////////////////////////////////////////////
+	// UEdGraphNode implementation
+	virtual void GetNodeContextMenuActions(class UToolMenu* Menu, class UGraphNodeContextMenuContext* Context) const override;
+
+	//////////////////////////////////////////////////////////////////////////
+	// URigVMEdGraphNode implementation
+	virtual void ConfigurePin(UEdGraphPin* EdGraphPin, const URigVMPin* ModelPin) const override;
+
+	//////////////////////////////////////////////////////////////////////////
+	// Our implementation
+
+	// Returns whether this node is a trait stack or not
+	ANIMNEXTUNCOOKEDONLY_API bool IsTraitStack() const;
+
+private:
+	// Populates the SubMenu with entries for each trait that can be added through the context menu
+	void BuildAddTraitContextMenu(class UToolMenu* SubMenu);
+};
diff --git a/Engine/Plugins/Experimental/AnimNext/Source/AnimNextUncookedOnly/Internal/AnimNextEdGraphSchema.h b/Engine/Plugins/Experimental/AnimNext/Source/AnimNextUncookedOnly/Internal/AnimNextEdGraphSchema.h
new file mode 100644
--- /dev/null
+++ b/Engine/Plugins/Experimental/AnimNext/Source/AnimNextUncookedOnly/Internal/AnimNextEdGraphSchema.h
@@ -0,0 +1,21 @@
+﻿// Copyright Epic Games, Inc. All Rights Reserved.
+
+#pragma once
+
+#include "CoreMinimal.h"
+#include "EdGraph/RigVMEdGraphSchema.h"
+#include "AnimNextEdGraphNode.h"
+#include "AnimNextEdGraphSchema.generated.h"
+
+UCLASS(MinimalAPI)
+class UAnimNextEdGraphSchema : public URigVMEdGraphSchema
+{
+	GENERATED_BODY()
+
+	// URigVMEdGraphSchema interface
+	virtual TSubclassOf<URigVMEdGraphNode> GetGraphNodeClass(const URigVMEdGraph* InGraph) const override { return UAnimNextEdGraphNode::StaticClass(); }
+	virtual bool IsStructEditable(UStruct* InStruct) const override;
+
+	// UEdGraphSchema interface
+	virtual void GetGraphDisplayInformation(const UEdGraph& Graph, /*out*/ FGraphDisplayInfo& DisplayInfo) const override;
+};
\ No newline at end of file
diff --git a/Engine/Plugins/Experimental/AnimNext/Source/AnimNextUncookedOnly/Internal/AnimNextRigVMAssetEditorData.h b/Engine/Plugins/Experimental/AnimNext/Source/AnimNextUncookedOnly/Internal/AnimNextRigVMAssetEditorData.h
--- a/Engine/Plugins/Experimental/AnimNext/Source/AnimNextUncookedOnly/Internal/AnimNextRigVMAssetEditorData.h
+++ b/Engine/Plugins/Experimental/AnimNext/Source/AnimNextUncookedOnly/Internal/AnimNextRigVMAssetEditorData.h
@@ -13,7 +13,7 @@
 enum class ERigVMGraphNotifType : uint8;
 class UAnimNextRigVMAssetEntry;
 class UAnimNextRigVMAssetEditorData;
-class UAnimNextGraph_EdGraph;
+class UAnimNextEdGraph;
 
 namespace UE::AnimNext::UncookedOnly
 {
@@ -78,7 +78,7 @@ protected:
 	friend class UE::AnimNext::Editor::FWorkspaceEditor;
 	friend class UAnimNextRigVMAssetEntry;
 	friend class UAnimNextRigVMAssetLibrary;
-	friend class UAnimNextGraph_EdGraph;
+	friend class UAnimNextEdGraph;
 	friend class UE::AnimNext::Tests::FEditor_Graph;
 	friend class UE::AnimNext::Tests::FEditor_Parameters;
 	friend class UE::AnimNext::Editor::FParameterCustomization;
diff --git a/Engine/Plugins/Experimental/AnimNext/Source/AnimNextUncookedOnly/Internal/Graph/AnimNextAnimationGraphSchema.h b/Engine/Plugins/Experimental/AnimNext/Source/AnimNextUncookedOnly/Internal/Graph/AnimNextAnimationGraphSchema.h
new file mode 100644
--- /dev/null
+++ b/Engine/Plugins/Experimental/AnimNext/Source/AnimNextUncookedOnly/Internal/Graph/AnimNextAnimationGraphSchema.h
@@ -0,0 +1,17 @@
+﻿// Copyright Epic Games, Inc. All Rights Reserved.
+
+#pragma once
+
+#include "CoreMinimal.h"
+#include "AnimNextRigVMAssetSchema.h"
+#include "AnimNextAnimationGraphSchema.generated.h"
+
+UCLASS()
+class UAnimNextAnimationGraphSchema : public UAnimNextRigVMAssetSchema
+{
+	GENERATED_BODY()
+
+	// URigVMSchema interface
+	virtual bool SupportsUnitFunction(URigVMController* InController, const FRigVMFunction* InUnitFunction) const override;
+	virtual bool SupportsDispatchFactory(URigVMController* InController, const FRigVMDispatchFactory* InDispatchFactory) const override;
+};
diff --git a/Engine/Plugins/Experimental/AnimNext/Source/AnimNextUncookedOnly/Internal/Graph/AnimNextModule_AnimationGraph.h b/Engine/Plugins/Experimental/AnimNext/Source/AnimNextUncookedOnly/Internal/Graph/AnimNextModule_AnimationGraph.h
new file mode 100644
--- /dev/null
+++ b/Engine/Plugins/Experimental/AnimNext/Source/AnimNextUncookedOnly/Internal/Graph/AnimNextModule_AnimationGraph.h
@@ -0,0 +1,71 @@
+﻿// Copyright Epic Games, Inc. All Rights Reserved.
+
+#pragma once
+
+#include "CoreMinimal.h"
+#include "AnimNextRigVMAssetEntry.h"
+#include "IAnimNextRigVMExportInterface.h"
+#include "IAnimNextRigVMGraphInterface.h"
+#include "AnimNextModule_AnimationGraph.generated.h"
+
+class UAnimNextModule_EditorData;
+class UAnimNextEdGraph;
+enum class ERigVMGraphNotifType : uint8;
+
+namespace UE::AnimNext::Editor
+{
+	struct FUtils;
+}
+
+/** An Animation Graph entry in an AnimNext module asset */
+UCLASS(MinimalAPI, Category = "Animation Graphs")
+class UAnimNextModule_AnimationGraph : public UAnimNextRigVMAssetEntry, public IAnimNextRigVMGraphInterface, public IAnimNextRigVMExportInterface
+{
+	GENERATED_BODY()
+
+	friend class UAnimNextModule_EditorData;
+	friend struct UE::AnimNext::Editor::FUtils;	
+
+	// IAnimNextRigVMExportInterface interface
+	virtual FAnimNextParamType GetExportType() const override;
+	virtual FName GetExportName() const override;
+	virtual EAnimNextExportAccessSpecifier GetExportAccessSpecifier() const override;
+	virtual void SetExportAccessSpecifier(EAnimNextExportAccessSpecifier InAccessSpecifier, bool bSetupUndoRedo = true) override;
+
+	// UAnimNextRigVMAssetEntry interface
+	virtual FName GetEntryName() const override;
+	virtual void SetEntryName(FName InName, bool bSetupUndoRedo = true) override;
+	
+	// IAnimNextRigVMGraphInterface interface
+	virtual URigVMGraph* GetRigVMGraph() const override;
+	virtual URigVMEdGraph* GetEdGraph() const override;
+	virtual void SetRigVMGraph(URigVMGraph* InGraph) override;
+	virtual void SetEdGraph(URigVMEdGraph* InGraph) override;
+
+protected:
+	/** Access specifier - whether the graph's entry point is visible external to this asset */
+	UPROPERTY(EditAnywhere, Category = AnimationGraph)
+	EAnimNextExportAccessSpecifier Access = EAnimNextExportAccessSpecifier::Private;
+
+	/** The name of the graph */
+	UPROPERTY(VisibleAnywhere, Category = AnimationGraph)
+	FName GraphName;
+
+	/** RigVM graph */
+	UPROPERTY()
+	TObjectPtr<URigVMGraph> Graph;
+
+	/** Editor graph */
+	UPROPERTY()
+	TObjectPtr<UAnimNextEdGraph> EdGraph;
+};
+
+// Old deprecated entry
+UCLASS()
+class UAnimNextGraphEntry : public UAnimNextRigVMAssetEntry
+{
+	GENERATED_BODY()
+
+	// UAnimNextRigVMAssetEntry interface
+	virtual FName GetEntryName() const override { return NAME_None; }
+};
\ No newline at end of file
diff --git a/Engine/Plugins/Experimental/AnimNext/Source/AnimNextUncookedOnly/Internal/IAnimNextRigVMGraphInterface.h b/Engine/Plugins/Experimental/AnimNext/Source/AnimNextUncookedOnly/Internal/IAnimNextRigVMGraphInterface.h
--- a/Engine/Plugins/Experimental/AnimNext/Source/AnimNextUncookedOnly/Internal/IAnimNextRigVMGraphInterface.h
+++ b/Engine/Plugins/Experimental/AnimNext/Source/AnimNextUncookedOnly/Internal/IAnimNextRigVMGraphInterface.h
@@ -27,7 +27,7 @@ public:
 	virtual URigVMEdGraph* GetEdGraph() const = 0;
 
 private:
-	friend class UAnimNextGraph_EditorData;
+	friend class UAnimNextModule_EditorData;
 
 	// Set the RigVM graph
 	virtual void SetRigVMGraph(URigVMGraph* InGraph) = 0;
diff --git a/Engine/Plugins/Experimental/AnimNext/Source/AnimNextUncookedOnly/Internal/Module/AnimNextModuleWorkspaceAssetUserData.h b/Engine/Plugins/Experimental/AnimNext/Source/AnimNextUncookedOnly/Internal/Module/AnimNextModuleWorkspaceAssetUserData.h
new file mode 100644
--- /dev/null
+++ b/Engine/Plugins/Experimental/AnimNext/Source/AnimNextUncookedOnly/Internal/Module/AnimNextModuleWorkspaceAssetUserData.h
@@ -0,0 +1,55 @@
+// Copyright Epic Games, Inc. All Rights Reserved.
+
+#pragma once
+
+#include "WorkspaceAssetRegistryInfo.h"
+#include "Param/ParamType.h"
+#include "IAnimNextRigVMGraphInterface.h"
+
+#include "AnimNextModuleWorkspaceAssetUserData.generated.h"
+
+USTRUCT()
+struct FAnimNextGraphAssetOutlinerData : public FWorkspaceOutlinerItemData
+{
+	GENERATED_BODY()
+	
+	FAnimNextGraphAssetOutlinerData() = default;
+};
+
+USTRUCT()
+struct FAnimNextParameterOutlinerData : public FWorkspaceOutlinerItemData
+{
+	GENERATED_BODY()
+	
+	FAnimNextParameterOutlinerData() = default;
+
+	UPROPERTY(VisibleAnywhere, Category=AnimNext)
+	FAnimNextParamType Type;
+};
+
+USTRUCT()
+struct FAnimNextGraphOutlinerData : public FWorkspaceOutlinerItemData
+{
+	GENERATED_BODY()
+	
+	FAnimNextGraphOutlinerData() = default;
+
+	UPROPERTY(VisibleAnywhere, Category=AnimNext)
+	TScriptInterface<IAnimNextRigVMGraphInterface> GraphInterface;
+};
+
+USTRUCT()
+struct FAnimNextSchedulerData : public FAnimNextGraphAssetOutlinerData
+{
+	GENERATED_BODY()
+	
+	FAnimNextSchedulerData() = default;
+};
+
+UCLASS()
+class UAnimNextModuleWorkspaceAssetUserData : public UAssetUserData
+{	
+	GENERATED_BODY()
+
+	virtual void GetAssetRegistryTags(FAssetRegistryTagsContext Context) const override;
+};
\ No newline at end of file
diff --git a/Engine/Plugins/Experimental/AnimNext/Source/AnimNextUncookedOnly/Internal/Module/AnimNextModule_Controller.h b/Engine/Plugins/Experimental/AnimNext/Source/AnimNextUncookedOnly/Internal/Module/AnimNextModule_Controller.h
new file mode 100644
--- /dev/null
+++ b/Engine/Plugins/Experimental/AnimNext/Source/AnimNextUncookedOnly/Internal/Module/AnimNextModule_Controller.h
@@ -0,0 +1,78 @@
+// Copyright Epic Games, Inc. All Rights Reserved.
+
+#pragma once
+
+#include "CoreMinimal.h"
+#include "Param/ParamType.h"
+#include "RigVMModel/RigVMController.h"
+#include "Param/AnimNextParamInstanceIdentifier.h"
+#include "InstancedStruct.h"
+#include "AnimNextModule_Controller.generated.h"
+
+struct FAnimNextParamType;
+
+
+namespace UE::AnimNext
+{
+	struct FTrait;
+	struct FTraitUID;
+}
+
+
+/**
+  * Implements AnimNext RigVM controller extensions
+  */
+UCLASS(MinimalAPI)
+class UAnimNextModule_Controller : public URigVMController
+{
+	GENERATED_BODY()
+
+public:
+	
+	// Adds a new Trait to the Stack, with default struct values
+	// Returns Trait Instance Name (or NAME_None on failure)
+	UFUNCTION(BlueprintCallable, Category = RigVMController)
+	ANIMNEXTUNCOOKEDONLY_API FName AddTraitByName(FName InNodeName, FName InNewTraitTypeName, int32 InPinIndex, const FString& InNewTraitDefaultValue = TEXT(""), bool bSetupUndoRedo = true, bool bPrintPythonCommand = false);
+
+	// Removes a Trait from the Stack, using Trait Instance Name
+	// Returns operation success (true) or failure (false)
+	UFUNCTION(BlueprintCallable, Category = RigVMController)
+	ANIMNEXTUNCOOKEDONLY_API bool RemoveTraitByName(FName InNodeName, FName InTraitInstanceName, bool bSetupUndoRedo = true, bool bPrintPythonCommand = false);
+
+	// Swap a Trait from the Stack with a new one, using existing Trait Instance Name and new Trait Type Name
+	// Returns Trait Instance Name (or NAME_None on failure)
+	UFUNCTION(BlueprintCallable, Category = RigVMController)
+	ANIMNEXTUNCOOKEDONLY_API FName SwapTraitByName(FName InNodeName, FName InTraitInstanceName, int32 InCurrentTraitPinIndex, FName InNewTraitTypeName, const FString& InNewTraitDefaultValue = TEXT(""), bool bSetupUndoRedo = true, bool bPrintPythonCommand = false);
+
+	// Move a Trait from its current PinIndex to the specified one (moving it visually in the stack)
+	// Returns operation success (true) or failure (false)
+	UFUNCTION(BlueprintCallable, Category = RigVMController)
+	ANIMNEXTUNCOOKEDONLY_API bool SetTraitPinIndex(FName InNodeName, FName InTraitInstanceName, int32 InNewPinIndex, bool bSetupUndoRedo = true, bool bPrintPythonCommand = false);
+
+	// Adds a unit node with a dynamic number of pins
+	URigVMUnitNode* AddUnitNodeWithPins(UScriptStruct* InScriptStruct, const FRigVMPinInfoArray& PinArray, const FName& InMethodName = TEXT("Execute"), const FVector2D& InPosition = FVector2D::ZeroVector, const FString& InNodeName = TEXT(""), bool bSetupUndoRedo = true, bool bPrintPythonCommand = false);
+
+	// Sets the parameter referenced by a set/get parameter node
+	UFUNCTION(BlueprintCallable, Category = RigVMController)
+	ANIMNEXTUNCOOKEDONLY_API bool SetAnimNextParameterNode(URigVMNode* ParameterNode, FName ParameterName, EPropertyBagPropertyType ValueType, EPropertyBagContainerType ContainerType = EPropertyBagContainerType::None, const UObject* ValueTypeObject = nullptr, const TInstancedStruct<FAnimNextParamInstanceIdentifier>& InstanceId = TInstancedStruct<FAnimNextParamInstanceIdentifier>(), bool bSetupUndoRedo = true, bool bPrintPythonCommand = false);
+	ANIMNEXTUNCOOKEDONLY_API bool SetAnimNextParameterNode(URigVMNode* ParameterNode, FName ParameterName, const FAnimNextParamType& InType, const TInstancedStruct<FAnimNextParamInstanceIdentifier>& InstanceId = TInstancedStruct<FAnimNextParamInstanceIdentifier>(), bool bSetupUndoRedo = true, bool bPrintPythonCommand = false);
+
+	// Creates a Get Parameter node
+	UFUNCTION(BlueprintCallable, Category = RigVMController)
+	ANIMNEXTUNCOOKEDONLY_API URigVMNode* AddGetAnimNextParameterNode(const FVector2D& InPosition, FName ParameterName, EPropertyBagPropertyType ValueType, EPropertyBagContainerType ContainerType = EPropertyBagContainerType::None, const UObject* ValueTypeObject = nullptr, const TInstancedStruct<FAnimNextParamInstanceIdentifier>& InstanceId = TInstancedStruct<FAnimNextParamInstanceIdentifier>(), bool bSetupUndoRedo = true, bool bPrintPythonCommand = false);
+	ANIMNEXTUNCOOKEDONLY_API URigVMNode* AddGetAnimNextParameterNode(const FVector2D& InPosition, FName ParameterName, const FAnimNextParamType& InType, const TInstancedStruct<FAnimNextParamInstanceIdentifier>& InstanceId = TInstancedStruct<FAnimNextParamInstanceIdentifier>(), bool bSetupUndoRedo = true, bool bPrintPythonCommand = false);
+
+	// Creates a Get Graph Parameter node
+	UFUNCTION(BlueprintCallable, Category = RigVMController)
+	ANIMNEXTUNCOOKEDONLY_API URigVMNode* AddGetAnimNextGraphParameterNode(const FVector2D& InPosition, FName ParameterName, EPropertyBagPropertyType ValueType, EPropertyBagContainerType ContainerType = EPropertyBagContainerType::None, const UObject* ValueTypeObject = nullptr, const TInstancedStruct<FAnimNextParamInstanceIdentifier>& InstanceId = TInstancedStruct<FAnimNextParamInstanceIdentifier>(), bool bSetupUndoRedo = true, bool bPrintPythonCommand = false);
+	ANIMNEXTUNCOOKEDONLY_API URigVMNode* AddGetAnimNextGraphParameterNode(const FVector2D& InPosition, FName ParameterName, const FAnimNextParamType& InType, const TInstancedStruct<FAnimNextParamInstanceIdentifier>& InstanceId = TInstancedStruct<FAnimNextParamInstanceIdentifier>(), bool bSetupUndoRedo = true, bool bPrintPythonCommand = false);
+
+	// Creates a Set Graph Parameter node
+	UFUNCTION(BlueprintCallable, Category = RigVMController)
+	ANIMNEXTUNCOOKEDONLY_API URigVMNode* AddSetAnimNextGraphParameterNode(const FVector2D& InPosition, FName ParameterName, EPropertyBagPropertyType ValueType, EPropertyBagContainerType ContainerType = EPropertyBagContainerType::None, const UObject* ValueTypeObject = nullptr, const TInstancedStruct<FAnimNextParamInstanceIdentifier>& InstanceId = TInstancedStruct<FAnimNextParamInstanceIdentifier>(), bool bSetupUndoRedo = true, bool bPrintPythonCommand = false);
+	ANIMNEXTUNCOOKEDONLY_API URigVMNode* AddSetAnimNextGraphParameterNode(const FVector2D& InPosition, FName ParameterName, const FAnimNextParamType& InType, const TInstancedStruct<FAnimNextParamInstanceIdentifier>& InstanceId = TInstancedStruct<FAnimNextParamInstanceIdentifier>(), bool bSetupUndoRedo = true, bool bPrintPythonCommand = false);
+
+private:
+	URigVMNode* AddGetAnimNextParameterNodeInternal(UScriptStruct* NodeStruct, const FVector2D& Position, FName ParameterName, const FAnimNextParamType& Type, const TInstancedStruct<FAnimNextParamInstanceIdentifier>& InstanceId, bool bSetupUndoRedo, bool bPrintPythonCommand);
+
+};
diff --git a/Engine/Plugins/Experimental/AnimNext/Source/AnimNextUncookedOnly/Internal/Module/AnimNextModule_EditorData.h b/Engine/Plugins/Experimental/AnimNext/Source/AnimNextUncookedOnly/Internal/Module/AnimNextModule_EditorData.h
new file mode 100644
--- /dev/null
+++ b/Engine/Plugins/Experimental/AnimNext/Source/AnimNextUncookedOnly/Internal/Module/AnimNextModule_EditorData.h
@@ -0,0 +1,100 @@
+﻿// Copyright Epic Games, Inc. All Rights Reserved.
+
+#pragma once
+
+#include "CoreMinimal.h"
+#include "Module/AnimNextModule_Controller.h"
+#include "RigVMModel/RigVMGraph.h"
+#include "AnimNextEdGraph.h"
+#include "AnimNextRigVMAssetEditorData.h"
+#include "AnimNextExecuteContext.h"
+#include "Kismet/BlueprintFunctionLibrary.h"
+#include "AnimNextModule_EditorData.generated.h"
+
+class UAnimNextModule;
+enum class ERigVMGraphNotifType : uint8;
+class FAnimationAnimNextRuntimeTest_GraphAddTrait;
+class FAnimationAnimNextRuntimeTest_GraphExecute;
+class FAnimationAnimNextRuntimeTest_GraphExecuteLatent;
+
+namespace UE::AnimNext::UncookedOnly
+{
+	struct FUtils;
+}
+
+namespace UE::AnimNext::Editor
+{
+	class FModuleEditor;
+	class SAnimNextGraphView;
+	struct FUtils;
+}
+
+// Script-callable editor API hoisted onto UAnimNextModule
+UCLASS()
+class ANIMNEXTUNCOOKEDONLY_API UAnimNextModuleLibrary : public UBlueprintFunctionLibrary
+{
+	GENERATED_BODY()
+
+	/** Adds an animation graph to an AnimNext Module asset */
+	UFUNCTION(BlueprintCallable, Category = "AnimNext|Module", meta=(ScriptMethod))
+	static UAnimNextModule_AnimationGraph* AddAnimationGraph(UAnimNextModule* InModule, FName InName, bool bSetupUndoRedo = true, bool bPrintPythonCommand = true);
+
+	/** Adds a parameter to an AnimNext Module asset */
+	UFUNCTION(BlueprintCallable, Category = "AnimNext|Module", meta=(ScriptMethod))
+	static UAnimNextModule_Parameter* AddParameter(UAnimNextModule* InModule, FName InName, EPropertyBagPropertyType InValueType, EPropertyBagContainerType InContainerType = EPropertyBagContainerType::None, const UObject* InValueTypeObject = nullptr, bool bSetupUndoRedo = true, bool bPrintPythonCommand = true);
+
+	/** Adds an event graph to an AnimNext Module asset */
+	UFUNCTION(BlueprintCallable, Category = "AnimNext|Module", meta=(ScriptMethod))
+	static UAnimNextModule_EventGraph* AddEventGraph(UAnimNextModule* InModule, FName InName, bool bSetupUndoRedo = true, bool bPrintPythonCommand = true);
+};
+
+/** Editor data for AnimNext modules */
+UCLASS(MinimalAPI)
+class UAnimNextModule_EditorData : public UAnimNextRigVMAssetEditorData
+{
+	GENERATED_BODY()
+
+	friend class UAnimNextModuleFactory;
+	friend class UAnimNextEdGraph;
+	friend class UAnimNextGraphEntry;
+	friend struct UE::AnimNext::UncookedOnly::FUtils;
+	friend struct UE::AnimNext::Editor::FUtils;
+	friend class UE::AnimNext::Editor::FModuleEditor;
+	friend class UE::AnimNext::Editor::SAnimNextGraphView;
+	friend struct FAnimNextGraphSchemaAction_RigUnit;
+	friend struct FAnimNextGraphSchemaAction_DispatchFactory;
+	friend class FAnimationAnimNextEditorTest_GraphAddTrait;
+	friend class FAnimationAnimNextEditorTest_GraphTraitOperations;
+	friend class FAnimationAnimNextRuntimeTest_GraphExecute;
+	friend class FAnimationAnimNextRuntimeTest_GraphExecuteLatent;
+	
+public:
+	/** Adds an animation graph to this asset */
+	ANIMNEXTUNCOOKEDONLY_API UAnimNextModule_AnimationGraph* AddAnimationGraph(FName InName, bool bSetupUndoRedo = true, bool bPrintPythonCommand = true);
+
+	/** Adds a parameter to this asset */
+	ANIMNEXTUNCOOKEDONLY_API UAnimNextModule_Parameter* AddParameter(FName InName, FAnimNextParamType InType, bool bSetupUndoRedo = true, bool bPrintPythonCommand = true);
+
+	/** Adds an event graph to this asset */
+	ANIMNEXTUNCOOKEDONLY_API UAnimNextModule_EventGraph* AddEventGraph(FName InName, bool bSetupUndoRedo = true, bool bPrintPythonCommand = true);
+
+private:
+	// UObject interface
+	virtual void PostLoad() override;
+
+	// IRigVMClientHost interface
+	virtual void RecompileVM() override;
+
+	// UAnimNextRigVMAssetEditorData interface
+	virtual TSubclassOf<URigVMController> GetControllerClass() const override { return UAnimNextModule_Controller::StaticClass(); }
+	virtual UScriptStruct* GetExecuteContextStruct() const override { return FAnimNextExecuteContext::StaticStruct(); }
+	virtual UEdGraph* CreateEdGraph(URigVMGraph* InRigVMGraph, bool bForce) override;
+	virtual bool RemoveEdGraph(URigVMGraph* InModel) override;
+	virtual void CreateEdGraphForCollapseNode(URigVMCollapseNode* InNode) override;
+	virtual void HandleModifiedEvent(ERigVMGraphNotifType InNotifType, URigVMGraph* InGraph, UObject* InSubject) override;
+	virtual TConstArrayView<TSubclassOf<UAnimNextRigVMAssetEntry>> GetEntryClasses() const override;
+
+private:
+	UPROPERTY()
+	TArray<TObjectPtr<UAnimNextEdGraph>> Graphs_DEPRECATED;
+};
diff --git a/Engine/Plugins/Experimental/AnimNext/Source/AnimNextUncookedOnly/Internal/Module/AnimNextModule_EventGraph.h b/Engine/Plugins/Experimental/AnimNext/Source/AnimNextUncookedOnly/Internal/Module/AnimNextModule_EventGraph.h
new file mode 100644
--- /dev/null
+++ b/Engine/Plugins/Experimental/AnimNext/Source/AnimNextUncookedOnly/Internal/Module/AnimNextModule_EventGraph.h
@@ -0,0 +1,61 @@
+﻿// Copyright Epic Games, Inc. All Rights Reserved.
+
+#pragma once
+
+#include "CoreMinimal.h"
+#include "AnimNextRigVMAssetEntry.h"
+#include "IAnimNextRigVMExportInterface.h"
+#include "IAnimNextRigVMGraphInterface.h"
+#include "IAnimNextRigVMParameterInterface.h"
+#include "AnimNextModule_EventGraph.generated.h"
+
+class UAnimNextModule_EditorData;
+class UAnimNextEdGraph;
+
+namespace UE::AnimNext::Tests
+{
+	class FEditor_Parameters;
+}
+
+UCLASS(MinimalAPI, Category = "Event Graphs")
+class UAnimNextModule_EventGraph : public UAnimNextRigVMAssetEntry, public IAnimNextRigVMGraphInterface
+{
+	GENERATED_BODY()
+
+	friend class UAnimNextModule_EditorData;
+	friend class UE::AnimNext::Tests::FEditor_Parameters;
+
+	// UAnimNextRigVMAssetEntry interface
+	virtual FName GetEntryName() const override { return GraphName; }
+	virtual void SetEntryName(FName InName, bool bSetupUndoRedo = true) override;
+	virtual FText GetDisplayName() const override;
+	virtual FText GetDisplayNameTooltip() const override;
+
+	// IAnimNextRigVMGraphInterface interface
+	virtual URigVMGraph* GetRigVMGraph() const override;
+	virtual URigVMEdGraph* GetEdGraph() const override;
+	virtual void SetRigVMGraph(URigVMGraph* InGraph) override;
+	virtual void SetEdGraph(URigVMEdGraph* InGraph) override;
+
+	/** The name of the graph */
+	UPROPERTY(VisibleAnywhere, Category = "Event Graph")
+	FName GraphName;
+
+	/** Graph */
+	UPROPERTY()
+	TObjectPtr<URigVMGraph> Graph;
+
+	/** Graph */
+	UPROPERTY()
+	TObjectPtr<UAnimNextEdGraph> EdGraph;
+};
+
+// Old deprecated entry
+UCLASS()
+class UAnimNextParameterBlockGraph : public UAnimNextRigVMAssetEntry
+{
+	GENERATED_BODY()
+
+	// UAnimNextRigVMAssetEntry interface
+	virtual FName GetEntryName() const override { return NAME_None; }
+};
\ No newline at end of file
diff --git a/Engine/Plugins/Experimental/AnimNext/Source/AnimNextUncookedOnly/Internal/Module/AnimNextModule_Parameter.h b/Engine/Plugins/Experimental/AnimNext/Source/AnimNextUncookedOnly/Internal/Module/AnimNextModule_Parameter.h
new file mode 100644
--- /dev/null
+++ b/Engine/Plugins/Experimental/AnimNext/Source/AnimNextUncookedOnly/Internal/Module/AnimNextModule_Parameter.h
@@ -0,0 +1,77 @@
+// Copyright Epic Games, Inc. All Rights Reserved.
+
+#pragma once
+
+#include "CoreMinimal.h"
+#include "AnimNextRigVMAssetEntry.h"
+#include "IAnimNextRigVMExportInterface.h"
+#include "IAnimNextRigVMParameterInterface.h"
+#include "Param/ParamType.h"
+#include "AnimNextModule_Parameter.generated.h"
+
+class UAnimNextModule_EditorData;
+
+namespace UE::AnimNext::Editor
+{
+	class FParameterCustomization;
+}
+
+namespace UE::AnimNext::Tests
+{
+	class FEditor_Parameters;
+}
+
+UCLASS(MinimalAPI, Category = "Parameters")
+class UAnimNextModule_Parameter : public UAnimNextRigVMAssetEntry, public IAnimNextRigVMParameterInterface, public IAnimNextRigVMExportInterface
+{
+	GENERATED_BODY()
+
+	friend class UAnimNextModule_EditorData;
+	friend class UE::AnimNext::Tests::FEditor_Parameters;
+	friend class UE::AnimNext::Editor::FParameterCustomization;
+
+	// IAnimNextRigVMExportInterface interface
+	virtual FAnimNextParamType GetExportType() const override;
+	virtual FName GetExportName() const override;
+	virtual EAnimNextExportAccessSpecifier GetExportAccessSpecifier() const override;
+	virtual void SetExportAccessSpecifier(EAnimNextExportAccessSpecifier InAccessSpecifier, bool bSetupUndoRedo = true) override;
+
+	// UAnimNextRigVMAssetEntry interface
+	virtual FName GetEntryName() const override;
+	virtual void SetEntryName(FName InName, bool bSetupUndoRedo = true) override;
+	virtual FText GetDisplayName() const override;
+	virtual FText GetDisplayNameTooltip() const override;
+
+	// IAnimNextRigVMParameterInterface interface
+	virtual FAnimNextParamType GetParamType() const override;
+	virtual bool SetParamType(const FAnimNextParamType& InType, bool bSetupUndoRedo = true) override;
+	virtual FName GetParamName() const override;
+	virtual void SetParamName(FName InName, bool bSetupUndoRedo = true) override;
+	virtual FInstancedPropertyBag& GetPropertyBag() const override;
+
+	/** Access specifier - whether the parameter is visible external to this asset */
+	UPROPERTY(EditAnywhere, Category = "Parameter")
+	EAnimNextExportAccessSpecifier Access = EAnimNextExportAccessSpecifier::Private;
+
+	/** Parameter name we reference */
+	UPROPERTY(VisibleAnywhere, Category = Parameter)
+	FName ParameterName;
+
+	/** The parameter's type */
+	UPROPERTY(EditAnywhere, Category = "Parameter")
+	FAnimNextParamType Type = FAnimNextParamType::GetType<bool>();
+	
+	/** Comment to display in editor */
+	UPROPERTY(EditAnywhere, Category = "Parameter", meta=(MultiLine))
+	FString Comment;
+};
+
+// Old deprecated class
+UCLASS()
+class UAnimNextParameterBlockParameter : public UAnimNextRigVMAssetEntry
+{
+	GENERATED_BODY()
+
+	// UAnimNextRigVMAssetEntry interface
+	virtual FName GetEntryName() const { return NAME_None; }
+};
\ No newline at end of file
diff --git a/Engine/Plugins/Experimental/AnimNext/Source/AnimNextUncookedOnly/Internal/UncookedOnlyUtils.h b/Engine/Plugins/Experimental/AnimNext/Source/AnimNextUncookedOnly/Internal/UncookedOnlyUtils.h
--- a/Engine/Plugins/Experimental/AnimNext/Source/AnimNextUncookedOnly/Internal/UncookedOnlyUtils.h
+++ b/Engine/Plugins/Experimental/AnimNext/Source/AnimNextUncookedOnly/Internal/UncookedOnlyUtils.h
@@ -16,12 +16,12 @@ struct FAnimNextParam;
 struct FEdGraphPinType;
 struct FWorkspaceOutlinerItemExports;
 class UAnimNextSchedule;
-class UAnimNextGraph;
-class UAnimNextGraph_EditorData;
-class UAnimNextGraph_EdGraph;
+class UAnimNextModule;
+class UAnimNextModule_EditorData;
+class UAnimNextEdGraph;
 class URigVMController;
 class URigVMGraph;
-class UAnimNextGraph_EdGraph;
+class UAnimNextEdGraph;
 class UAnimNextRigVMAsset;
 class UAnimNextRigVMAssetEditorData;
 class UAnimNextRigVMAssetEntry;
@@ -104,19 +104,19 @@ namespace UE::AnimNext::UncookedOnly
 
 struct ANIMNEXTUNCOOKEDONLY_API FUtils
 {
-	static void Compile(UAnimNextGraph* InGraph);
+	static void Compile(UAnimNextModule* InModule);
 
-	static void CompileVM(UAnimNextGraph* InGraph);
+	static void CompileVM(UAnimNextModule* InModule);
 
-	static void CompileStruct(UAnimNextGraph* InGraph);
+	static void CompileStruct(UAnimNextModule* InModule);
 
-	static UAnimNextGraph_EditorData* GetEditorData(const UAnimNextGraph* InAnimNextGraph);
+	static UAnimNextModule_EditorData* GetEditorData(const UAnimNextModule* InModule);
 
-	static UAnimNextGraph* GetGraph(const UAnimNextGraph_EditorData* InEditorData);
+	static UAnimNextModule* GetGraph(const UAnimNextModule_EditorData* InEditorData);
 
-	static FInstancedPropertyBag* GetPropertyBag(UAnimNextGraph* InAnimNextGraph);
+	static FInstancedPropertyBag* GetPropertyBag(UAnimNextModule* InModule);
 
-	static void RecreateVM(UAnimNextGraph* InGraph);
+	static void RecreateVM(UAnimNextModule* InModule);
 
 	/**
 	 * Get an AnimNext parameter type handle from an FEdGraphPinType.
diff --git a/Engine/Plugins/Experimental/AnimNext/Source/AnimNextUncookedOnly/Private/AnimNextEdGraph.cpp b/Engine/Plugins/Experimental/AnimNext/Source/AnimNextUncookedOnly/Private/AnimNextEdGraph.cpp
new file mode 100644
--- /dev/null
+++ b/Engine/Plugins/Experimental/AnimNext/Source/AnimNextUncookedOnly/Private/AnimNextEdGraph.cpp
@@ -0,0 +1,30 @@
+﻿// Copyright Epic Games, Inc. All Rights Reserved.
+
+#include "AnimNextEdGraph.h"
+#include "Module/AnimNextModule_EditorData.h"
+
+void UAnimNextEdGraph::PostLoad()
+{
+	Super::PostLoad();
+
+	UAnimNextRigVMAssetEditorData* EditorData = GetTypedOuter<UAnimNextRigVMAssetEditorData>();
+	check(EditorData);
+	Initialize(EditorData);
+}
+
+void UAnimNextEdGraph::Initialize(UAnimNextRigVMAssetEditorData* InEditorData)
+{
+	InEditorData->RigVMGraphModifiedEvent.RemoveAll(this);
+	InEditorData->RigVMGraphModifiedEvent.AddUObject(this, &UAnimNextEdGraph::HandleModifiedEvent);
+	InEditorData->RigVMCompiledEvent.RemoveAll(this);
+	InEditorData->RigVMCompiledEvent.AddUObject(this, &UAnimNextEdGraph::HandleVMCompiledEvent);
+}
+
+FRigVMClient* UAnimNextEdGraph::GetRigVMClient() const
+{
+	if (const UAnimNextModule_EditorData* EditorData = GetTypedOuter<UAnimNextModule_EditorData>())
+	{
+		return (FRigVMClient*)EditorData->GetRigVMClient();
+	}
+	return nullptr;
+}
diff --git a/Engine/Plugins/Experimental/AnimNext/Source/AnimNextUncookedOnly/Private/AnimNextEdGraphNode.cpp b/Engine/Plugins/Experimental/AnimNext/Source/AnimNextUncookedOnly/Private/AnimNextEdGraphNode.cpp
new file mode 100644
--- /dev/null
+++ b/Engine/Plugins/Experimental/AnimNext/Source/AnimNextUncookedOnly/Private/AnimNextEdGraphNode.cpp
@@ -0,0 +1,130 @@
+﻿// Copyright Epic Games, Inc. All Rights Reserved.
+
+#include "AnimNextEdGraphNode.h"
+
+#include "TraitCore/TraitHandle.h"
+#include "Graph/RigUnit_AnimNextTraitStack.h"
+#include "Graph/RigDecorator_AnimNextCppTrait.h"
+#include "RigVMModel/RigVMController.h"
+#include "TraitCore/TraitRegistry.h"
+#include "ToolMenu.h"
+
+#define LOCTEXT_NAMESPACE "AnimNextEdGraphNode"
+
+void UAnimNextEdGraphNode::GetNodeContextMenuActions(UToolMenu* Menu, class UGraphNodeContextMenuContext* Context) const
+{
+	URigVMEdGraphNode::GetNodeContextMenuActions(Menu, Context);
+
+	if (IsTraitStack())
+	{
+		FToolMenuSection& Section = Menu->AddSection("AnimNextTraitNodeActions", LOCTEXT("AnimNextTraitNodeActionsMenuHeader", "Anim Next Trait Actions"));
+
+		UAnimNextEdGraphNode* NonConstThis = const_cast<UAnimNextEdGraphNode*>(this);
+
+		Section.AddSubMenu(
+			TEXT("AddTraitMenu"),
+			LOCTEXT("AddTraitMenu", "Add Trait"),
+			LOCTEXT("AddTraitMenuTooltip", "Add the chosen trait to currently selected node"),
+			FNewToolMenuDelegate::CreateUObject(NonConstThis, &UAnimNextEdGraphNode::BuildAddTraitContextMenu));
+	}
+}
+
+void UAnimNextEdGraphNode::ConfigurePin(UEdGraphPin* EdGraphPin, const URigVMPin* ModelPin) const
+{
+	Super::ConfigurePin(EdGraphPin, ModelPin);
+
+	// Trait handles always remain as a RigVM input pins so that we can still link things to them even if they are hidden
+	// We handle visibility for those explicitly here
+	const bool bIsInputPin = ModelPin->GetDirection() == ERigVMPinDirection::Input;
+	const bool bIsTraitHandle = ModelPin->GetCPPTypeObject() == FAnimNextTraitHandle::StaticStruct();
+	if (bIsInputPin && bIsTraitHandle)
+	{
+		if (const URigVMPin* DecoratorPin = ModelPin->GetParentPin())
+		{
+			if (DecoratorPin->IsTraitPin())
+			{
+				check(DecoratorPin->GetScriptStruct() == FRigDecorator_AnimNextCppDecorator::StaticStruct());
+
+				TSharedPtr<FStructOnScope> DecoratorScope = DecoratorPin->GetTraitInstance();
+				const FRigDecorator_AnimNextCppDecorator* VMDecorator = (const FRigDecorator_AnimNextCppDecorator*)DecoratorScope->GetStructMemory();
+
+				const UScriptStruct* TraitStruct = VMDecorator->GetTraitSharedDataStruct();
+				check(TraitStruct != nullptr);
+
+				const FProperty* PinProperty = TraitStruct->FindPropertyByName(ModelPin->GetFName());
+				EdGraphPin->bHidden = PinProperty->HasMetaData(FRigVMStruct::HiddenMetaName);
+			}
+		}
+	}
+}
+
+bool UAnimNextEdGraphNode::IsTraitStack() const
+{
+	if (const URigVMUnitNode* VMNode = Cast<URigVMUnitNode>(GetModelNode()))
+	{
+		const UScriptStruct* ScriptStruct = VMNode->GetScriptStruct();
+		return ScriptStruct == FRigUnit_AnimNextTraitStack::StaticStruct();
+	}
+
+	return false;
+}
+
+void UAnimNextEdGraphNode::BuildAddTraitContextMenu(UToolMenu* SubMenu)
+{
+	using namespace UE::AnimNext;
+
+	const FTraitRegistry& TraitRegistry = FTraitRegistry::Get();
+	TArray<const FTrait*> Traits = TraitRegistry.GetTraits();
+
+	URigVMController* VMController = GetController();
+	URigVMNode* VMNode = GetModelNode();
+
+	const UScriptStruct* CppDecoratorStruct = FRigDecorator_AnimNextCppDecorator::StaticStruct();
+
+	for (const FTrait* Trait : Traits)
+	{
+		UScriptStruct* ScriptStruct = Trait->GetTraitSharedDataStruct();
+
+		FString DefaultValue;
+		{
+			const FRigDecorator_AnimNextCppDecorator DefaultCppDecoratorStructInstance;
+			FRigDecorator_AnimNextCppDecorator CppDecoratorStructInstance;
+			CppDecoratorStructInstance.DecoratorSharedDataStruct = ScriptStruct;
+
+			if (!CppDecoratorStructInstance.CanBeAddedToNode(VMNode, nullptr))
+			{
+				continue;	// This trait isn't supported on this node
+			}
+
+			const FProperty* Prop = FAnimNextCppDecoratorWrapper::StaticStruct()->FindPropertyByName(GET_MEMBER_NAME_STRING_CHECKED(FAnimNextCppDecoratorWrapper, CppDecorator));
+			Prop->ExportText_Direct(DefaultValue, &CppDecoratorStructInstance, &DefaultCppDecoratorStructInstance, nullptr, PPF_SerializedAsImportText);
+		}
+
+		FString DisplayNameMetadata;
+		ScriptStruct->GetStringMetaDataHierarchical(FRigVMStruct::DisplayNameMetaName, &DisplayNameMetadata);
+		const FString DisplayName = DisplayNameMetadata.IsEmpty() ? Trait->GetTraitName() : DisplayNameMetadata;
+
+		const FText ToolTip = ScriptStruct->GetToolTipText();
+
+		FToolMenuEntry TraitEntry = FToolMenuEntry::InitMenuEntry(
+			*Trait->GetTraitName(),
+			FText::FromString(DisplayName),
+			ToolTip,
+			FSlateIcon(),
+			FUIAction(FExecuteAction::CreateLambda(
+				[this, Trait, VMController, VMNode, CppDecoratorStruct, DefaultValue, DisplayName]()
+				{
+					VMController->AddTrait(
+						VMNode->GetFName(),
+						*CppDecoratorStruct->GetPathName(),
+						*DisplayName,
+						DefaultValue, INDEX_NONE, true, true);
+				})
+			)
+		);
+
+		SubMenu->AddMenuEntry(NAME_None, TraitEntry);
+	}
+}
+
+#undef LOCTEXT_NAMESPACE
diff --git a/Engine/Plugins/Experimental/AnimNext/Source/AnimNextUncookedOnly/Private/AnimNextEdGraphSchema.cpp b/Engine/Plugins/Experimental/AnimNext/Source/AnimNextUncookedOnly/Private/AnimNextEdGraphSchema.cpp
new file mode 100644
--- /dev/null
+++ b/Engine/Plugins/Experimental/AnimNext/Source/AnimNextUncookedOnly/Private/AnimNextEdGraphSchema.cpp
@@ -0,0 +1,32 @@
+﻿// Copyright Epic Games, Inc. All Rights Reserved.
+
+#include "AnimNextEdGraphSchema.h"
+
+#include "AnimNextRigVMAsset.h"
+#include "AnimNextRigVMAssetEntry.h"
+#include "Param/AnimNextEditorParam.h"
+#include "Param/AnimNextParam.h"
+
+#define LOCTEXT_NAMESPACE "AnimNextEdGraphSchema"
+
+void UAnimNextEdGraphSchema::GetGraphDisplayInformation(const UEdGraph& Graph, /*out*/ FGraphDisplayInfo& DisplayInfo) const
+{
+	Super::GetGraphDisplayInformation(Graph, DisplayInfo);
+	
+	if(UAnimNextRigVMAssetEntry* AssetEntry = Cast<UAnimNextRigVMAssetEntry>(Graph.GetOuter()))
+	{
+		DisplayInfo.DisplayName = FText::Format(LOCTEXT("GraphTabTitleFormat", "{0}: {1}"), FText::FromName(AssetEntry->GetEntryName()), FText::FromName(AssetEntry->GetTypedOuter<UAnimNextRigVMAsset>()->GetFName()));
+		DisplayInfo.Tooltip = FText::Format(LOCTEXT("GraphTabTooltipFormat", "{0} in:\n{1}"), FText::FromName(AssetEntry->GetEntryName()), FText::FromString(AssetEntry->GetTypedOuter<UAnimNextRigVMAsset>()->GetPathName()));
+	}
+}
+
+bool UAnimNextEdGraphSchema::IsStructEditable(UStruct* InStruct) const
+{
+	if (InStruct == FAnimNextEditorParam::StaticStruct() || InStruct == FAnimNextParam::StaticStruct())
+	{
+		return true;
+	}
+	return Super::IsStructEditable(InStruct);
+}
+
+#undef LOCTEXT_NAMESPACE
\ No newline at end of file
diff --git a/Engine/Plugins/Experimental/AnimNext/Source/AnimNextUncookedOnly/Private/AnimNextRigVMAssetEditorData.cpp b/Engine/Plugins/Experimental/AnimNext/Source/AnimNextUncookedOnly/Private/AnimNextRigVMAssetEditorData.cpp
--- a/Engine/Plugins/Experimental/AnimNext/Source/AnimNextUncookedOnly/Private/AnimNextRigVMAssetEditorData.cpp
+++ b/Engine/Plugins/Experimental/AnimNext/Source/AnimNextUncookedOnly/Private/AnimNextRigVMAssetEditorData.cpp
@@ -5,7 +5,7 @@
 #include "AnimNextRigVMAsset.h"
 #include "AnimNextRigVMAssetEntry.h"
 #include "AnimNextRigVMAssetSchema.h"
-#include "AnimNextRigVMWorkspaceAssetUserData.h"
+#include "Module/AnimNextModuleWorkspaceAssetUserData.h"
 #include "ControlRigDefines.h"
 #include "ExternalPackageHelper.h"
 #include "IAnimNextRigVMGraphInterface.h"
@@ -119,9 +119,9 @@ void UAnimNextRigVMAssetEditorData::Initialize(bool bRecompileVM)
 
 	if (IInterface_AssetUserData* OuterUserData = Cast<IInterface_AssetUserData>(GetOuter()))
 	{
-		if(!OuterUserData->HasAssetUserDataOfClass(UAnimNextGraphWorkspaceAssetUserData::StaticClass()))
+		if(!OuterUserData->HasAssetUserDataOfClass(UAnimNextModuleWorkspaceAssetUserData::StaticClass()))
 		{
-			OuterUserData->AddAssetUserDataOfClass(UAnimNextGraphWorkspaceAssetUserData::StaticClass());
+			OuterUserData->AddAssetUserDataOfClass(UAnimNextModuleWorkspaceAssetUserData::StaticClass());
 		}
 	}
 }
diff --git a/Engine/Plugins/Experimental/AnimNext/Source/AnimNextUncookedOnly/Private/AnimNextUncookedOnlyModule.cpp b/Engine/Plugins/Experimental/AnimNext/Source/AnimNextUncookedOnly/Private/AnimNextUncookedOnlyModule.cpp
--- a/Engine/Plugins/Experimental/AnimNext/Source/AnimNextUncookedOnly/Private/AnimNextUncookedOnlyModule.cpp
+++ b/Engine/Plugins/Experimental/AnimNext/Source/AnimNextUncookedOnly/Private/AnimNextUncookedOnlyModule.cpp
@@ -2,7 +2,7 @@
 
 #include "AnimNextUncookedOnlyModule.h"
 
-#include "AnimNextRigVMWorkspaceAssetUserData.h"
+#include "Module/AnimNextModuleWorkspaceAssetUserData.h"
 #include "UncookedOnlyUtils.h"
 #include "Engine/Blueprint.h"
 #include "Modules/ModuleManager.h"
diff --git a/Engine/Plugins/Experimental/AnimNext/Source/AnimNextUncookedOnly/Private/Graph/AnimNextAnimationGraphSchema.cpp b/Engine/Plugins/Experimental/AnimNext/Source/AnimNextUncookedOnly/Private/Graph/AnimNextAnimationGraphSchema.cpp
new file mode 100644
--- /dev/null
+++ b/Engine/Plugins/Experimental/AnimNext/Source/AnimNextUncookedOnly/Private/Graph/AnimNextAnimationGraphSchema.cpp
@@ -0,0 +1,47 @@
+﻿// Copyright Epic Games, Inc. All Rights Reserved.
+
+#include "Graph/AnimNextAnimationGraphSchema.h"
+
+#include "AnimNextExecuteContext.h"
+#include "Graph/RigUnit_AnimNextBase.h"
+#include "Param/RigVMDispatch_GetParameter.h"
+#include "Param/RigVMDispatch_GetScopedParameter.h"
+
+bool UAnimNextAnimationGraphSchema::SupportsUnitFunction(URigVMController* InController, const FRigVMFunction* InUnitFunction) const
+{
+	if(InUnitFunction)
+	{
+		if(const UScriptStruct* FunctionExecuteContextStruct = InUnitFunction->GetExecuteContextStruct())
+		{
+			if(FunctionExecuteContextStruct == FAnimNextExecuteContext::StaticStruct())
+			{
+				// Only allow nodes that are children of FRigUnit_AnimNextBase
+				if(InUnitFunction->Struct)
+				{
+					return InUnitFunction->Struct->IsChildOf(FRigUnit_AnimNextBase::StaticStruct());
+				}
+			}
+		}
+	}
+
+	return Super::SupportsUnitFunction(InController, InUnitFunction);
+}
+
+bool UAnimNextAnimationGraphSchema::SupportsDispatchFactory(URigVMController* InController, const FRigVMDispatchFactory* InDispatchFactory) const
+{
+	if(InDispatchFactory)
+	{
+		if(const UScriptStruct* DispatchExecuteContextStruct = InDispatchFactory->GetExecuteContextStruct())
+		{
+			if(DispatchExecuteContextStruct == FAnimNextExecuteContext::StaticStruct())
+			{
+				// We only support the FRigVMDispatch_GetParameter/FRigVMDispatch_GetScopedParameter at the moment.
+				return
+				   (InDispatchFactory->GetScriptStruct() == FRigVMDispatch_GetParameter::StaticStruct() ||
+					InDispatchFactory->GetScriptStruct() == FRigVMDispatch_GetScopedParameter::StaticStruct());
+			}
+		}
+	}
+	
+	return Super::SupportsDispatchFactory(InController, InDispatchFactory);
+}
\ No newline at end of file
diff --git a/Engine/Plugins/Experimental/AnimNext/Source/AnimNextUncookedOnly/Private/Graph/AnimNextModule_AnimationGraph.cpp b/Engine/Plugins/Experimental/AnimNext/Source/AnimNextUncookedOnly/Private/Graph/AnimNextModule_AnimationGraph.cpp
new file mode 100644
--- /dev/null
+++ b/Engine/Plugins/Experimental/AnimNext/Source/AnimNextUncookedOnly/Private/Graph/AnimNextModule_AnimationGraph.cpp
@@ -0,0 +1,95 @@
+﻿// Copyright Epic Games, Inc. All Rights Reserved.
+
+#include "Graph/AnimNextModule_AnimationGraph.h"
+
+#include "AnimNextRigVMAsset.h"
+#include "UncookedOnlyUtils.h"
+#include "AnimNextEdGraph.h"
+#include "Graph/RigUnit_AnimNextGraphRoot.h"
+#include "Param/AnimNextTag.h"
+#include "Param/ParamType.h"
+
+FAnimNextParamType UAnimNextModule_AnimationGraph::GetExportType() const
+{
+	return FAnimNextParamType::GetType<FAnimNextEntryPoint>();
+}
+
+FName UAnimNextModule_AnimationGraph::GetExportName() const
+{
+	if(UAnimNextRigVMAsset* OuterAsset = GetTypedOuter<UAnimNextRigVMAsset>())
+	{
+		return UE::AnimNext::UncookedOnly::FUtils::GetQualifiedName(OuterAsset, GraphName);
+	}
+	return GraphName;
+}
+
+EAnimNextExportAccessSpecifier UAnimNextModule_AnimationGraph::GetExportAccessSpecifier() const
+{
+	return Access;
+}
+
+void UAnimNextModule_AnimationGraph::SetExportAccessSpecifier(EAnimNextExportAccessSpecifier InAccessSpecifier, bool bSetupUndoRedo)
+{
+	if(bSetupUndoRedo)
+	{
+		Modify();
+	};
+
+	Access = InAccessSpecifier;
+
+	BroadcastModified();
+}
+
+FName UAnimNextModule_AnimationGraph::GetEntryName() const
+{
+	return GraphName;
+}
+
+void UAnimNextModule_AnimationGraph::SetEntryName(FName InName, bool bSetupUndoRedo)
+{
+	if(bSetupUndoRedo)
+	{
+		Modify();
+	};
+
+	GraphName = InName;
+
+	// Forward to entry point node
+	URigVMController* Controller = GetImplementingOuter<IRigVMClientHost>()->GetController(Graph);
+	for(URigVMNode* Node : Graph->GetNodes())
+	{
+		if(URigVMUnitNode* UnitNode = Cast<URigVMUnitNode>(Node))
+		{
+			if(UnitNode->GetScriptStruct() == FRigUnit_AnimNextGraphRoot::StaticStruct())
+			{
+				URigVMPin* EntryPointPin = UnitNode->FindPin(GET_MEMBER_NAME_STRING_CHECKED(FRigUnit_AnimNextGraphRoot, EntryPoint));
+				check(EntryPointPin);
+				check(EntryPointPin->GetDirection() == ERigVMPinDirection::Hidden);
+
+				Controller->SetPinDefaultValue(EntryPointPin->GetPinPath(), InName.ToString());
+			}
+		}
+	}
+	
+	BroadcastModified();
+}
+
+URigVMGraph* UAnimNextModule_AnimationGraph::GetRigVMGraph() const
+{
+	return Graph;
+}
+
+URigVMEdGraph* UAnimNextModule_AnimationGraph::GetEdGraph() const
+{
+	return EdGraph;
+}
+
+void UAnimNextModule_AnimationGraph::SetRigVMGraph(URigVMGraph* InGraph)
+{
+	Graph = InGraph;
+}
+
+void UAnimNextModule_AnimationGraph::SetEdGraph(URigVMEdGraph* InGraph)
+{
+	EdGraph = CastChecked<UAnimNextEdGraph>(InGraph);
+}
diff --git a/Engine/Plugins/Experimental/AnimNext/Source/AnimNextUncookedOnly/Private/Module/AnimNextModuleWorkspaceAssetUserData.cpp b/Engine/Plugins/Experimental/AnimNext/Source/AnimNextUncookedOnly/Private/Module/AnimNextModuleWorkspaceAssetUserData.cpp
new file mode 100644
--- /dev/null
+++ b/Engine/Plugins/Experimental/AnimNext/Source/AnimNextUncookedOnly/Private/Module/AnimNextModuleWorkspaceAssetUserData.cpp
@@ -0,0 +1,33 @@
+// Copyright Epic Games, Inc. All Rights Reserved.
+
+#include "Module/AnimNextModuleWorkspaceAssetUserData.h"
+
+#include "UncookedOnlyUtils.h"
+#include "Module/AnimNextModule.h"
+#include "Module/AnimNextModule_EditorData.h"
+#include "UObject/AssetRegistryTagsContext.h"
+
+void UAnimNextModuleWorkspaceAssetUserData::GetAssetRegistryTags(FAssetRegistryTagsContext Context) const
+{
+	Super::GetAssetRegistryTags(Context);
+
+	FWorkspaceOutlinerItemExports Exports;
+	{
+		const UAnimNextModule* ModuleOuter = CastChecked<UAnimNextModule>(GetOuter());
+		const UAnimNextModule_EditorData* GraphEditorData = UE::AnimNext::UncookedOnly::FUtils::GetEditorData(ModuleOuter);	
+		{
+			FWorkspaceOutlinerItemExport& RootAssetExport = Exports.Exports.AddDefaulted_GetRef();	
+			RootAssetExport.Identifier = ModuleOuter->GetFName();
+			RootAssetExport.ParentIdentifier = NAME_None;
+			RootAssetExport.AssetPath = ModuleOuter;
+
+			RootAssetExport.Data.InitializeAsScriptStruct(FAnimNextGraphAssetOutlinerData::StaticStruct());
+		}
+	
+		UE::AnimNext::UncookedOnly::FUtils::GetAssetOutlinerItems(GraphEditorData, Exports);
+	}
+	
+	FString TagValue;
+	FWorkspaceOutlinerItemExports::StaticStruct()->ExportText(TagValue, &Exports, nullptr, nullptr, PPF_None, nullptr);
+	Context.AddTag(FAssetRegistryTag(UE::Workspace::ExportsWorkspaceItemsRegistryTag, TagValue, FAssetRegistryTag::TT_Hidden));
+}
diff --git a/Engine/Plugins/Experimental/AnimNext/Source/AnimNextUncookedOnly/Private/Module/AnimNextModule_Controller.cpp b/Engine/Plugins/Experimental/AnimNext/Source/AnimNextUncookedOnly/Private/Module/AnimNextModule_Controller.cpp
new file mode 100644
--- /dev/null
+++ b/Engine/Plugins/Experimental/AnimNext/Source/AnimNextUncookedOnly/Private/Module/AnimNextModule_Controller.cpp
@@ -0,0 +1,470 @@
+// Copyright Epic Games, Inc. All Rights Reserved.
+
+#include "Module/AnimNextModule_Controller.h"
+
+#include "AnimNextUnitNode.h"
+#include "Graph/RigDecorator_AnimNextCppTrait.h"
+#include "Param/AnimNextEditorParam.h"
+#include "Param/AnimNextParam.h"
+#include "Param/AnimNextParamInstanceIdentifier.h"
+#include "Param/ParamType.h"
+#include "Param/RigVMDispatch_GetLayerParameter.h"
+#include "Param/RigVMDispatch_GetParameter.h"
+#include "Param/RigVMDispatch_GetScopedParameter.h"
+#include "Param/RigVMDispatch_SetLayerParameter.h"
+#include "RigVMModel/RigVMControllerActions.h"
+#include "TraitCore/Trait.h"
+#include "TraitCore/TraitRegistry.h"
+#include "TraitCore/TraitUID.h"
+
+FName UAnimNextModule_Controller::AddTraitByName(FName InNodeName, FName InNewTraitTypeName, int32 InPinIndex, const FString& InNewTraitDefaultValue, bool bSetupUndoRedo, bool bPrintPythonCommand)
+{
+	if (!IsValidGraph())
+	{
+		return NAME_None;
+	}
+
+	if (!IsTransacting() && !IsGraphEditable())
+	{
+		return NAME_None;
+	}
+
+	if (InNodeName == NAME_None)
+	{
+		ReportError(TEXT("Invalid node name."));
+		return NAME_None;
+	}
+
+	const URigVMGraph* Graph = GetGraph();
+	check(Graph);
+
+	URigVMNode * Node = Graph->FindNodeByName(InNodeName);
+	if (Node == nullptr)
+	{
+		ReportError(TEXT("This graph does not contain a node with the provided name."));
+		return NAME_None;
+	}
+
+	UE::AnimNext::FTraitRegistry& TraitRegistry = UE::AnimNext::FTraitRegistry::Get();
+
+	const UE::AnimNext::FTrait* Trait = TraitRegistry.Find(InNewTraitTypeName);
+	if (Trait == nullptr)
+	{
+		ReportError(TEXT("Unknown Trait Type."));
+		return NAME_None;
+	}
+
+	const UScriptStruct* CppDecoratorStruct = FRigDecorator_AnimNextCppDecorator::StaticStruct();
+	UScriptStruct* ScriptStruct = Trait->GetTraitSharedDataStruct();
+
+	FString DefaultValue = InNewTraitDefaultValue;
+	if (DefaultValue.IsEmpty())
+	{
+		const FRigDecorator_AnimNextCppDecorator DefaultCppDecoratorStructInstance;
+		FRigDecorator_AnimNextCppDecorator CppDecoratorStructInstance;
+		CppDecoratorStructInstance.DecoratorSharedDataStruct = ScriptStruct;
+
+		if (!CppDecoratorStructInstance.CanBeAddedToNode(Node, nullptr))
+		{
+			ReportError(TEXT("Trait is not supported by the Node."));
+			return NAME_None;	// This trait isn't supported on this node
+		}
+
+		const FProperty* Prop = FAnimNextCppDecoratorWrapper::StaticStruct()->FindPropertyByName(GET_MEMBER_NAME_STRING_CHECKED(FAnimNextCppDecoratorWrapper, CppDecorator));
+		Prop->ExportText_Direct(DefaultValue, &CppDecoratorStructInstance, &DefaultCppDecoratorStructInstance, nullptr, PPF_SerializedAsImportText);
+	}
+
+	// Avoid multiple VM recompilations for internal operations
+	FRigVMControllerCompileBracketScope CompileScope(this);
+
+	const FName TraitName = InNewTraitTypeName;
+
+	return AddTrait(InNodeName, *CppDecoratorStruct->GetPathName(), TraitName, DefaultValue, InPinIndex, bSetupUndoRedo, bPrintPythonCommand);
+}
+
+bool UAnimNextModule_Controller::RemoveTraitByName(FName InNodeName, FName InTraitInstanceName, bool bSetupUndoRedo, bool bPrintPythonCommand)
+{
+	if (!IsValidGraph())
+	{
+		return false;
+	}
+
+	if (!IsTransacting() && !IsGraphEditable())
+	{
+		return false;
+	}
+
+	if (InNodeName == NAME_None)
+	{
+		ReportError(TEXT("Invalid node name."));
+		return false;
+	}
+
+	// Avoid multiple VM recompilations for internal operations
+	FRigVMControllerCompileBracketScope CompileScope(this);
+	return RemoveTrait(InNodeName, InTraitInstanceName, bSetupUndoRedo, bPrintPythonCommand);
+}
+
+FName UAnimNextModule_Controller::SwapTraitByName(FName InNodeName, FName InTraitInstanceName, int32 InCurrentTraitPinIndex, FName InNewTraitTypeName, const FString& InNewTraitDefaultValue, bool bSetupUndoRedo, bool bPrintPythonCommand)
+{
+	if (!IsValidGraph())
+	{
+		return NAME_None;
+	}
+
+	if (!IsTransacting() && !IsGraphEditable())
+	{
+		return NAME_None;
+	}
+
+	if (InNodeName == NAME_None)
+	{
+		ReportError(TEXT("Invalid node name."));
+		return NAME_None;
+	}
+
+	// Avoid multiple VM recompilations, for each operation
+	FRigVMControllerCompileBracketScope CompileScope(this);
+	if (RemoveTraitByName(InNodeName, InTraitInstanceName, bSetupUndoRedo, bPrintPythonCommand))
+	{
+		return AddTraitByName(InNodeName, InNewTraitTypeName, InCurrentTraitPinIndex, InNewTraitDefaultValue, bSetupUndoRedo, bPrintPythonCommand);
+	}
+
+	return NAME_None;
+}
+
+bool UAnimNextModule_Controller::SetTraitPinIndex(FName InNodeName, FName InTraitInstanceName, int32 InNewPinIndex, bool bSetupUndoRedo, bool bPrintPythonCommand)
+{
+	if (!IsValidGraph())
+	{
+		return false;
+	}
+
+	if (!IsTransacting() && !IsGraphEditable())
+	{
+		return false;
+	}
+
+	if (InNodeName == NAME_None)
+	{
+		ReportError(TEXT("Invalid node name."));
+		return false;
+	}
+
+	const URigVMGraph* Graph = GetGraph();
+	check(Graph);
+
+	URigVMNode* Node = Graph->FindNodeByName(InNodeName);
+	if (Node == nullptr)
+	{
+		ReportError(TEXT("This graph does not contain a node with the provided name."));
+		return false;
+	}
+
+	if (!IsTransacting() && !IsGraphEditable())
+	{
+		return false;
+	}
+
+	URigVMPin* TraitPin = Node->FindTrait(InTraitInstanceName);
+	if (TraitPin == nullptr)
+	{
+		ReportError(TEXT("The node does not contain a Trait with the provided name."));
+		return false;
+	}
+	
+	// Save current pin data for later
+	const FString TraitDefaultValue = TraitPin->GetDefaultValue();
+
+	// TODO zzz : Is there a better way to get a Trait* from a TraitPin ?
+	if (TSharedPtr<FStructOnScope> ScopedTrait = Node->GetTraitInstance(TraitPin->GetFName()))
+	{
+		const FRigVMTrait* Trait = (FRigVMTrait*)ScopedTrait->GetStructMemory();
+		if (const UScriptStruct* TraitSharedInstanceData = Trait->GetTraitSharedDataStruct())
+		{
+			UE::AnimNext::FTraitRegistry& TraitRegistry = UE::AnimNext::FTraitRegistry::Get();
+			if (const UE::AnimNext::FTrait* AnimNextTrait = TraitRegistry.Find(TraitSharedInstanceData))
+			{
+				// Avoid multiple VM recompilations, for each operation
+				FRigVMControllerCompileBracketScope CompileScope(this);
+				if (RemoveTraitByName(InNodeName, InTraitInstanceName, bSetupUndoRedo, bPrintPythonCommand))
+				{
+					// TOOO zzz : Why TraitName is a string ?
+					if (AddTraitByName(InNodeName, FName(AnimNextTrait->GetTraitName()), InNewPinIndex, TraitDefaultValue, bSetupUndoRedo, bPrintPythonCommand) != NAME_None)
+					{
+						return true;
+					}
+				}
+			}
+		}
+	}
+
+	return false;
+}
+
+URigVMUnitNode* UAnimNextModule_Controller::AddUnitNodeWithPins(UScriptStruct* InScriptStruct, const FRigVMPinInfoArray& PinArray, const FName& InMethodName, const FVector2D& InPosition, const FString& InNodeName, bool bSetupUndoRedo, bool bPrintPythonCommand)
+{
+	const bool bHasDynamicPins = PinArray.Num() != 0;
+
+	if (bHasDynamicPins)
+	{
+		OpenUndoBracket(TEXT("Add unit node with pins"));
+	}
+
+	URigVMUnitNode* Node = AddUnitNode(InScriptStruct, UAnimNextUnitNode::StaticClass(), InMethodName, InPosition, InNodeName, bSetupUndoRedo, bPrintPythonCommand);
+
+	if (Node == nullptr)
+	{
+		if (bHasDynamicPins)
+		{
+			CancelUndoBracket();
+		}
+
+		return nullptr;
+	}
+
+	if (bHasDynamicPins)
+	{
+		const FRigVMRegistry& Registry = FRigVMRegistry::Get();
+		const FRigVMPinInfoArray PreviousPins(Node, this);
+
+		for (int32 PinIndex = 0; PinIndex < PinArray.Num(); ++PinIndex)
+		{
+			const FString& PinPath = PinArray.GetPinPath(PinIndex);
+			FString ParentPinPath, PinName;
+			UObject* OuterForPin = Node;
+			if (URigVMPin::SplitPinPathAtEnd(PinPath, ParentPinPath, PinName))
+			{
+				OuterForPin = Node->FindPin(ParentPinPath);
+			}
+
+			CreatePinFromPinInfo(Registry, PreviousPins, PinArray[PinIndex], PinPath, OuterForPin);
+		}
+
+		CloseUndoBracket();
+	}
+
+	return Node;
+}
+
+bool UAnimNextModule_Controller::SetAnimNextParameterNode(URigVMNode* ParameterNode, FName ParameterName, EPropertyBagPropertyType ValueType, EPropertyBagContainerType ContainerType, const UObject* ValueTypeObject, const TInstancedStruct<FAnimNextParamInstanceIdentifier>& InstanceId, bool bSetupUndoRedo, bool bPrintPythonCommand)
+{
+	return SetAnimNextParameterNode(ParameterNode, ParameterName, FAnimNextParamType(ValueType, ContainerType, ValueTypeObject), InstanceId, bSetupUndoRedo, bPrintPythonCommand);
+}
+
+bool UAnimNextModule_Controller::SetAnimNextParameterNode(URigVMNode* ParameterNode, FName ParameterName, const FAnimNextParamType& Type, const TInstancedStruct<FAnimNextParamInstanceIdentifier>& InstanceId, bool bSetupUndoRedo, bool bPrintPythonCommand)
+{
+	if(!IsValidGraph())
+	{
+		return false;
+	}
+
+	if (!IsTransacting() && !IsGraphEditable())
+	{
+		return false;
+	}
+
+	if(ParameterNode == nullptr)
+	{
+		ReportError(TEXT("Invalid node."));
+		return false;
+	}
+
+	const URigVMTemplateNode* TemplateNode = Cast<URigVMTemplateNode>(ParameterNode);
+	if (TemplateNode == nullptr)
+	{
+		ReportError(TEXT("Not a template node."));
+		return false;
+	}
+
+	const FRigVMDispatchFactory* GetParameterFactory = FRigVMRegistry::Get().FindOrAddDispatchFactory(FRigVMDispatch_GetParameter::StaticStruct());
+	const FName GetParameterNotation = GetParameterFactory->GetTemplate()->GetNotation();
+	const FRigVMDispatchFactory* GetScopedParameterFactory = FRigVMRegistry::Get().FindOrAddDispatchFactory(FRigVMDispatch_GetScopedParameter::StaticStruct());
+	const FName GetScopedParameterNotation = GetScopedParameterFactory->GetTemplate()->GetNotation();
+	const FRigVMDispatchFactory* GetLayerParameterFactory = FRigVMRegistry::Get().FindOrAddDispatchFactory(FRigVMDispatch_GetLayerParameter::StaticStruct());
+	const FName GetLayerParameterNotation = GetLayerParameterFactory->GetTemplate()->GetNotation();
+	const FRigVMDispatchFactory* SetLayerParameterFactory = FRigVMRegistry::Get().FindOrAddDispatchFactory(FRigVMDispatch_SetLayerParameter::StaticStruct());
+	const FName SetLayerParameterNotation = SetLayerParameterFactory->GetTemplate()->GetNotation();
+
+	if( TemplateNode->GetNotation() != GetParameterNotation &&
+		TemplateNode->GetNotation() != GetScopedParameterNotation &&
+		TemplateNode->GetNotation() != GetLayerParameterNotation &&
+		TemplateNode->GetNotation() != SetLayerParameterNotation)
+	{
+		ReportError(TEXT("Not a parameter node."));
+		return false;
+	}
+
+	const bool bIsNamedParam = TemplateNode->GetNotation() == GetLayerParameterNotation || TemplateNode->GetNotation() == SetLayerParameterNotation;
+
+	if(ParameterName != NAME_None)
+	{
+		const FSoftObjectPath SoftObjectPath(ParameterName.ToString());
+		if(!SoftObjectPath.GetAssetPath().IsValid() || SoftObjectPath.GetSubPathString().Len() == 0)
+		{
+			ReportError(TEXT("InParameterName is an invalid format. Should be /AssetOrFieldPath/ClassOrAsset.ClassOrAsset:FieldOrParameter."));
+			return false;
+		}
+	}
+
+	if (!Type.IsNone() && !Type.IsValid())
+	{
+		ReportError(TEXT("Type is invalid."));
+		return false;
+	}
+
+	URigVMPin* ParameterPin = ParameterNode->FindPin(FRigVMDispatch_GetParameter::ParameterName.ToString());
+	if(ParameterPin == nullptr)
+	{
+		ReportError(TEXT("Parameter pin not found."));
+		return false;
+	}
+
+	if(bSetupUndoRedo)
+	{
+		OpenUndoBracket(TEXT("Set parameter"));
+	}
+
+	FString ValueAsString;
+	if(bIsNamedParam)
+	{
+		ValueAsString = ParameterName.ToString();
+	}
+	else
+	{
+		FAnimNextParam ParamValue(ParameterName, Type, InstanceId);
+		FAnimNextParam::StaticStruct()->ExportText(ValueAsString, &ParamValue, nullptr, nullptr, PPF_None, nullptr);
+	}
+
+	if(!SetPinDefaultValue(ParameterPin->GetPinPath(), ValueAsString, true, bSetupUndoRedo, true, bPrintPythonCommand))
+	{
+		if(bSetupUndoRedo)
+		{
+			CancelUndoBracket();
+		}
+		return false;
+	}
+
+	FRigVMTemplateArgumentType RigVMType = Type.ToRigVMTemplateArgument();
+	URigVMPin* OutputPin = ParameterNode->FindPin(FRigVMDispatch_GetParameter::ValueName.ToString());
+	if(RigVMType.IsValid() && OutputPin)
+	{
+		// Re-resolve the node's output
+		TArray<URigVMLink*> Links = OutputPin->GetLinks();
+		if(!TemplateNode->IsFullyUnresolved())
+		{
+			if(!UnresolveTemplateNodes({ ParameterNode->GetFName() }, bSetupUndoRedo, bPrintPythonCommand))
+			{
+				if(bSetupUndoRedo)
+				{
+					CloseUndoBracket();
+				}
+				return true;
+			}
+		}
+
+		if(TemplateNode->IsFullyUnresolved())
+		{
+			if(!ResolveWildCardPin(OutputPin, RigVMType, bSetupUndoRedo, bPrintPythonCommand))
+			{
+				if(bSetupUndoRedo)
+				{
+					CloseUndoBracket();
+				}
+				return true;
+			}
+		}
+
+		// Try to restore links
+		for(URigVMLink* Link : Links)
+		{
+			if(!AddLink(OutputPin, Link->GetOppositePin(OutputPin), bSetupUndoRedo))
+			{
+				if(bSetupUndoRedo)
+				{
+					CloseUndoBracket();
+				}
+				return true;
+			}
+		}
+	}
+
+	if(bSetupUndoRedo)
+	{
+		CloseUndoBracket();
+	}
+
+	return true;
+}
+
+URigVMNode* UAnimNextModule_Controller::AddGetAnimNextParameterNode(const FVector2D& Position, FName ParameterName, EPropertyBagPropertyType ValueType, EPropertyBagContainerType ContainerType, const UObject* ValueTypeObject, const TInstancedStruct<FAnimNextParamInstanceIdentifier>& InstanceId, bool bSetupUndoRedo, bool bPrintPythonCommand)
+{
+	return AddGetAnimNextParameterNode(Position, ParameterName, FAnimNextParamType(ValueType, ContainerType, ValueTypeObject), InstanceId, bSetupUndoRedo, bPrintPythonCommand);
+}
+
+URigVMNode* UAnimNextModule_Controller::AddGetAnimNextParameterNode(const FVector2D& Position, FName ParameterName, const FAnimNextParamType& Type, const TInstancedStruct<FAnimNextParamInstanceIdentifier>& InstanceId, bool bSetupUndoRedo, bool bPrintPythonCommand)
+{
+	return AddGetAnimNextParameterNodeInternal(FRigVMDispatch_GetScopedParameter::StaticStruct(), Position, ParameterName, Type, InstanceId, bSetupUndoRedo, bPrintPythonCommand);
+}
+
+URigVMNode* UAnimNextModule_Controller::AddGetAnimNextGraphParameterNode(const FVector2D& Position, FName ParameterName, EPropertyBagPropertyType ValueType, EPropertyBagContainerType ContainerType, const UObject* ValueTypeObject, const TInstancedStruct<FAnimNextParamInstanceIdentifier>& InstanceId, bool bSetupUndoRedo, bool bPrintPythonCommand)
+{
+	return AddGetAnimNextGraphParameterNode(Position, ParameterName, FAnimNextParamType(ValueType, ContainerType, ValueTypeObject), InstanceId, bSetupUndoRedo, bPrintPythonCommand);
+}
+
+URigVMNode* UAnimNextModule_Controller::AddGetAnimNextGraphParameterNode(const FVector2D& Position, FName ParameterName, const FAnimNextParamType& Type, const TInstancedStruct<FAnimNextParamInstanceIdentifier>& InstanceId, bool bSetupUndoRedo, bool bPrintPythonCommand)
+{
+return AddGetAnimNextParameterNodeInternal(FRigVMDispatch_GetLayerParameter::StaticStruct(), Position, ParameterName, Type, InstanceId, bSetupUndoRedo, bPrintPythonCommand);
+}
+
+URigVMNode* UAnimNextModule_Controller::AddSetAnimNextGraphParameterNode(const FVector2D& Position, FName ParameterName, EPropertyBagPropertyType ValueType, EPropertyBagContainerType ContainerType, const UObject* ValueTypeObject, const TInstancedStruct<FAnimNextParamInstanceIdentifier>& InstanceId, bool bSetupUndoRedo, bool bPrintPythonCommand)
+{
+	return AddSetAnimNextGraphParameterNode(Position, ParameterName, FAnimNextParamType(ValueType, ContainerType, ValueTypeObject), InstanceId, bSetupUndoRedo, bPrintPythonCommand);
+}
+
+URigVMNode* UAnimNextModule_Controller::AddSetAnimNextGraphParameterNode(const FVector2D& Position, FName ParameterName, const FAnimNextParamType& Type, const TInstancedStruct<FAnimNextParamInstanceIdentifier>& InstanceId, bool bSetupUndoRedo, bool bPrintPythonCommand)
+{
+	return AddGetAnimNextParameterNodeInternal(FRigVMDispatch_SetLayerParameter::StaticStruct(), Position, ParameterName, Type, InstanceId, bSetupUndoRedo, bPrintPythonCommand);
+}
+
+URigVMNode* UAnimNextModule_Controller::AddGetAnimNextParameterNodeInternal(UScriptStruct* NodeStruct, const FVector2D& Position, FName ParameterName, const FAnimNextParamType& Type, const TInstancedStruct<FAnimNextParamInstanceIdentifier>& InstanceId, bool bSetupUndoRedo, bool bPrintPythonCommand)
+{
+	if(!IsValidGraph())
+	{
+		return nullptr;
+	}
+
+	URigVMGraph* Graph = GetGraph();
+	check(Graph);
+
+	if(bSetupUndoRedo)
+	{
+		OpenUndoBracket(TEXT("Add parameter node"));
+	}
+
+	const FString Name = GetSchema()->GetValidNodeName(Graph, TEXT("ParameterNode"));
+	const FRigVMDispatchFactory* Factory = FRigVMRegistry::Get().FindOrAddDispatchFactory(NodeStruct);
+	URigVMNode* Node = AddTemplateNode(Factory->GetTemplate()->GetNotation(), Position, Name, bSetupUndoRedo, bPrintPythonCommand);
+	if(Node)
+	{
+		if(!SetAnimNextParameterNode(Node, ParameterName, Type, InstanceId, bSetupUndoRedo, bPrintPythonCommand))
+		{
+			if(bSetupUndoRedo)
+			{
+				CancelUndoBracket();
+			}
+			return nullptr;
+		}
+	}
+	else if(bSetupUndoRedo)
+	{
+		CancelUndoBracket();
+	}
+
+	if(bSetupUndoRedo)
+	{
+		CloseUndoBracket();
+	}
+
+	return Node;
+}
\ No newline at end of file
diff --git a/Engine/Plugins/Experimental/AnimNext/Source/AnimNextUncookedOnly/Private/Module/AnimNextModule_EditorData.cpp b/Engine/Plugins/Experimental/AnimNext/Source/AnimNextUncookedOnly/Private/Module/AnimNextModule_EditorData.cpp
new file mode 100644
--- /dev/null
+++ b/Engine/Plugins/Experimental/AnimNext/Source/AnimNextUncookedOnly/Private/Module/AnimNextModule_EditorData.cpp
@@ -0,0 +1,440 @@
+﻿// Copyright Epic Games, Inc. All Rights Reserved.
+
+#include "Module/AnimNextModule_EditorData.h"
+
+#include "UncookedOnlyUtils.h"
+#include "AssetRegistry/IAssetRegistry.h"
+#include "Curves/CurveFloat.h"
+#include "Module/AnimNextModule.h"
+#include "Graph/AnimNextModule_AnimationGraph.h"
+#include "AnimNextEdGraphSchema.h"
+#include "Module/AnimNextModule_EventGraph.h"
+#include "AnimNextEventGraphSchema.h"
+#include "Module/AnimNextModule_Parameter.h"
+#include "Graph/AnimNextAnimationGraphSchema.h"
+#include "RigVMModel/RigVMFunctionLibrary.h"
+#include "RigVMModel/RigVMNotifications.h"
+#include "RigVMModel/Nodes/RigVMCollapseNode.h"
+#include "UObject/AssetRegistryTagsContext.h"
+#include "UObject/FortniteMainBranchObjectVersion.h"
+#include "UObject/LinkerLoad.h"
+
+void UAnimNextModule_EditorData::PostLoad()
+{
+	Super::PostLoad();
+
+	auto FindEntryForRigVMGraph = [this](URigVMGraph* InRigVMGraph)
+	{
+		UAnimNextRigVMAssetEntry* FoundEntry = nullptr;
+		for(UAnimNextRigVMAssetEntry* Entry : Entries)
+		{
+			if(IAnimNextRigVMGraphInterface* GraphEntry = Cast<IAnimNextRigVMGraphInterface>(Entry))
+			{
+				if(InRigVMGraph == GraphEntry->GetRigVMGraph())
+				{
+					FoundEntry = Entry;
+					break;
+				}
+			}
+		}
+		return FoundEntry;
+	};
+
+	if(GetLinkerCustomVersion(FFortniteMainBranchObjectVersion::GUID) < FFortniteMainBranchObjectVersion::AnimNextCombineGraphContexts)
+	{
+		// Must preload entries so their data is populated or we cannot find the appropriate entries for graphs
+		for(UAnimNextRigVMAssetEntry* Entry : Entries) 
+		{
+			Entry->GetLinker()->Preload(Entry);
+		}
+
+		TArray<URigVMGraph*> AllModels = RigVMClient.GetAllModels(false, true);
+		for(URigVMGraph* Graph : AllModels)
+		{
+			Graph->SetExecuteContextStruct(FAnimNextExecuteContext::StaticStruct());
+			if(UAnimNextRigVMAssetEntry* FoundEntry = FindEntryForRigVMGraph(Graph))
+			{
+				if(FoundEntry->IsA(UAnimNextModule_AnimationGraph::StaticClass()))
+				{
+					Graph->SetSchemaClass(UAnimNextAnimationGraphSchema::StaticClass());
+				}
+				else
+				{
+					Graph->SetSchemaClass(UAnimNextEventGraphSchema::StaticClass());
+				}
+			}
+			else
+			{
+				Graph->SetSchemaClass(UAnimNextAnimationGraphSchema::StaticClass());
+			}
+		}
+	}
+
+	if(GetLinkerCustomVersion(FFortniteMainBranchObjectVersion::GUID) < FFortniteMainBranchObjectVersion::AnimNextMoveGraphsToEntries)
+	{
+		// Must preload entries so their data is populated or we cannot find the appropriate entries for graphs
+		for(UAnimNextRigVMAssetEntry* Entry : Entries) 
+		{
+			Entry->GetLinker()->Preload(Entry);
+		}
+		
+		for(TObjectPtr<UAnimNextEdGraph> Graph : Graphs_DEPRECATED)
+		{
+			URigVMGraph* FoundRigVMGraph = GetRigVMGraphForEditorObject(Graph);
+			if(FoundRigVMGraph)
+			{
+				if(UAnimNextRigVMAssetEntry* FoundEntry = FindEntryForRigVMGraph(FoundRigVMGraph))
+				{
+					if(UAnimNextModule_AnimationGraph* AnimationGraphEntry = Cast<UAnimNextModule_AnimationGraph>(FoundEntry))
+					{
+						AnimationGraphEntry->EdGraph = Graph;
+					}
+					else if(UAnimNextModule_EventGraph* EventGraphEntry = Cast<UAnimNextModule_EventGraph>(FoundEntry))
+					{
+						EventGraphEntry->EdGraph = Graph;
+					}
+
+					Graph->Rename(nullptr, FoundEntry, REN_ForceNoResetLoaders | REN_DoNotDirty | REN_DontCreateRedirectors | REN_NonTransactional);
+					Graph->Initialize(this);
+				}
+			}
+		}
+
+		// We used to add a default model that is no longer needed
+		URigVMGraph* DefaultModel = RigVMClient.GetDefaultModel();
+		if(DefaultModel && DefaultModel->GetName() == TEXT("RigVMGraph"))
+		{
+			bool bFound = false;
+			for(UAnimNextRigVMAssetEntry* Entry : Entries)
+			{
+				if(UAnimNextModule_EventGraph* EventGraphEntry = Cast<UAnimNextModule_EventGraph>(Entry))
+				{
+					if(DefaultModel == static_cast<IAnimNextRigVMGraphInterface*>(EventGraphEntry)->GetRigVMGraph())
+					{
+						bFound = true;
+						break;
+					}
+				}
+			}
+
+			if(!bFound)
+			{
+				TGuardValue<bool> DisablePythonPrint(bSuspendPythonMessagesForRigVMClient, false);
+				TGuardValue<bool> DisableAutoCompile(bAutoRecompileVM, false);
+				RigVMClient.RemoveModel(DefaultModel->GetNodePath(), false);
+			}
+		}
+
+		RecompileVM();
+	}
+
+	if(GetLinkerCustomVersion(FFortniteMainBranchObjectVersion::GUID) < FFortniteMainBranchObjectVersion::AnimNextGraphAccessSpecifiers)
+	{
+		// Must preload entries so their data is populated as we will be modifying them
+		for(UAnimNextRigVMAssetEntry* Entry : Entries) 
+		{
+			Entry->GetLinker()->Preload(Entry);
+		}
+
+		// Force older assets to all have public symbols so they work as-is. Newer assets need user intervention as entries default to private
+		for(UAnimNextRigVMAssetEntry* Entry : Entries)
+		{
+			if(UAnimNextModule_AnimationGraph* AnimationGraphEntry = Cast<UAnimNextModule_AnimationGraph>(Entry))
+			{
+				AnimationGraphEntry->Access = EAnimNextExportAccessSpecifier::Public;
+			}
+			else if(UAnimNextModule_Parameter* ParameterEntry = Cast<UAnimNextModule_Parameter>(Entry))
+			{
+				ParameterEntry->Access = EAnimNextExportAccessSpecifier::Public;
+			}
+		}
+	}
+}
+
+void UAnimNextModule_EditorData::RecompileVM()
+{
+	CachedExports = FAnimNextParameterProviderAssetRegistryExports();
+	UE::AnimNext::UncookedOnly::FUtils::GetAssetParameters(this, CachedExports.GetValue());
+	UE::AnimNext::UncookedOnly::FUtils::Compile(GetTypedOuter<UAnimNextModule>());
+	
+	if (IAssetRegistry* AssetRegistry = IAssetRegistry::Get())
+	{
+		AssetRegistry->AssetUpdateTags(GetTypedOuter<UAnimNextModule>(), EAssetRegistryTagsCaller::Fast);
+	}
+}
+
+void UAnimNextModule_EditorData::HandleModifiedEvent(ERigVMGraphNotifType InNotifType, URigVMGraph* InGraph, UObject* InSubject)
+{
+	switch(InNotifType)
+	{
+	case ERigVMGraphNotifType::PinAdded:
+		{
+			if (URigVMPin* Pin = Cast<URigVMPin>(InSubject))
+			{
+				if (Pin->IsTraitPin())
+				{
+					RequestAutoVMRecompilation();
+				}
+			}
+			break;
+		}
+	}
+
+	Super::HandleModifiedEvent(InNotifType, InGraph, InSubject);
+}
+
+TConstArrayView<TSubclassOf<UAnimNextRigVMAssetEntry>> UAnimNextModule_EditorData::GetEntryClasses() const
+{
+	static const TSubclassOf<UAnimNextRigVMAssetEntry> Classes[] =
+	{
+		UAnimNextModule_AnimationGraph::StaticClass(),
+		UAnimNextModule_EventGraph::StaticClass(),
+		UAnimNextModule_Parameter::StaticClass(),
+	};
+	
+	return Classes;
+}
+
+void UAnimNextModule_EditorData::CreateEdGraphForCollapseNode(URigVMCollapseNode* InNode)
+{
+	if (InNode->GetGraph()->IsA<URigVMFunctionLibrary>())
+	{
+		if (URigVMGraph* ContainedGraph = InNode->GetContainedGraph())
+		{
+			// create a sub graph
+			UAnimNextEdGraph* RigFunctionGraph = NewObject<UAnimNextEdGraph>(this, *InNode->GetName(), RF_Transactional);
+			RigFunctionGraph->Schema = UAnimNextEdGraphSchema::StaticClass();
+			RigFunctionGraph->bAllowRenaming = true;
+			RigFunctionGraph->bEditable = true;
+			RigFunctionGraph->bAllowDeletion = true;
+			RigFunctionGraph->ModelNodePath = ContainedGraph->GetNodePath();
+			RigFunctionGraph->bIsFunctionDefinition = true;
+
+			RigFunctionGraph->Initialize(this);
+
+			RigVMClient.GetOrCreateController(ContainedGraph)->ResendAllNotifications();
+		}
+	}
+}
+
+UEdGraph* UAnimNextModule_EditorData::CreateEdGraph(URigVMGraph* InRigVMGraph, bool bForce)
+{
+	check(InRigVMGraph);
+
+	if(InRigVMGraph->IsA<URigVMFunctionLibrary>())
+	{
+		return nullptr;
+	}
+
+	IAnimNextRigVMGraphInterface* Entry = Cast<IAnimNextRigVMGraphInterface>(FindEntryForRigVMGraph(InRigVMGraph));
+	if(Entry == nullptr)
+	{
+		// Not found, we could be adding a new entry, in which case the graph wont be assigned yet
+		check(Entries.Num() > 0);
+		check(Cast<IAnimNextRigVMGraphInterface>(Entries.Last()) != nullptr);
+		check(Cast<IAnimNextRigVMGraphInterface>(Entries.Last())->GetRigVMGraph() == nullptr);
+		Entry = Cast<IAnimNextRigVMGraphInterface>(FindEntryForRigVMGraph(nullptr));
+	}
+
+	if(Entry == nullptr)
+	{
+		return nullptr;
+	}
+	
+	if(bForce)
+	{
+		RemoveEdGraph(InRigVMGraph);
+	}
+
+	FString GraphName = InRigVMGraph->GetName();
+	check(!GraphName.IsEmpty());
+
+	UAnimNextEdGraph* RigFunctionGraph = NewObject<UAnimNextEdGraph>(CastChecked<UObject>(Entry), NAME_None, RF_Transactional);
+	RigFunctionGraph->Schema = UAnimNextEdGraphSchema::StaticClass();
+	RigFunctionGraph->bAllowDeletion = true;
+	RigFunctionGraph->bIsFunctionDefinition = false;
+	RigFunctionGraph->ModelNodePath = InRigVMGraph->GetNodePath();
+	RigFunctionGraph->Initialize(this);
+
+	Entry->SetEdGraph(RigFunctionGraph);
+	if(Entry->GetRigVMGraph() == nullptr)
+	{
+		Entry->SetRigVMGraph(InRigVMGraph);
+	}
+	else
+	{
+		check(Entry->GetRigVMGraph() == InRigVMGraph);
+	}
+
+	return RigFunctionGraph;
+}
+
+bool UAnimNextModule_EditorData::RemoveEdGraph(URigVMGraph* InModel)
+{
+	if(UAnimNextModule_AnimationGraph* Entry = Cast<UAnimNextModule_AnimationGraph>(FindEntryForRigVMGraph(InModel)))
+	{
+		RigVMClient.DestroyObject(Entry->EdGraph);
+		Entry->EdGraph = nullptr;
+		return true;
+	}
+	return false;
+}
+
+UAnimNextModule_Parameter* UAnimNextModuleLibrary::AddParameter(UAnimNextModule* InModule, FName InName, EPropertyBagPropertyType InValueType,
+	EPropertyBagContainerType InContainerType, const UObject* InValueTypeObject, bool bSetupUndoRedo, bool bPrintPythonCommand)
+{
+	return UE::AnimNext::UncookedOnly::FUtils::GetEditorData(InModule)->AddParameter(InName, FAnimNextParamType(InValueType, InContainerType, InValueTypeObject), bSetupUndoRedo, bPrintPythonCommand);
+}
+
+UAnimNextModule_Parameter* UAnimNextModule_EditorData::AddParameter(FName InName, FAnimNextParamType InType, bool bSetupUndoRedo, bool bPrintPythonCommand)
+{
+	if(InName == NAME_None)
+	{
+		ReportError(TEXT("UAnimNextModule_EditorData::AddParameter: Invalid parameter name supplied."));
+		return nullptr;
+	}
+
+	// Check for duplicate parameter
+	const bool bAlreadyExists = Entries.ContainsByPredicate([InName](const UAnimNextRigVMAssetEntry* InEntry)
+	{
+		if(const UAnimNextModule_Parameter* Parameter = Cast<UAnimNextModule_Parameter>(InEntry))
+		{
+			return Parameter->ParameterName == InName;
+		}
+		return false;
+	});
+
+	if(bAlreadyExists)
+	{
+		ReportError(TEXT("UAnimNextModule_EditorData::AddParameter: A parameter already exists for the supplied parameter name."));
+		return nullptr;
+	}
+
+	UAnimNextModule_Parameter* NewEntry = CreateNewSubEntry<UAnimNextModule_Parameter>(this);
+	NewEntry->ParameterName = InName;
+	NewEntry->Type = InType;
+
+	if(bSetupUndoRedo)
+	{
+		NewEntry->Modify();
+		Modify();
+	}
+	
+	Entries.Add(NewEntry);
+
+	BroadcastModified();
+
+	return NewEntry;
+}
+
+UAnimNextModule_EventGraph* UAnimNextModuleLibrary::AddEventGraph(UAnimNextModule* InModule, FName InName, bool bSetupUndoRedo, bool bPrintPythonCommand)
+{
+	return UE::AnimNext::UncookedOnly::FUtils::GetEditorData(InModule)->AddEventGraph(InName, bSetupUndoRedo, bPrintPythonCommand);
+}
+
+UAnimNextModule_EventGraph* UAnimNextModule_EditorData::AddEventGraph(FName InName, bool bSetupUndoRedo, bool bPrintPythonCommand)
+{
+	if(InName == NAME_None)
+	{
+		ReportError(TEXT("UAnimNextModule_EditorData::AddEventGraph: Invalid graph name supplied."));
+		return nullptr;
+	}
+
+	// Check for duplicate name
+	FName NewGraphName = InName;
+	auto DuplicateNamePredicate = [&NewGraphName](const UAnimNextRigVMAssetEntry* InEntry)
+	{
+		return InEntry->GetEntryName() == NewGraphName;
+	};
+
+	bool bAlreadyExists = Entries.ContainsByPredicate(DuplicateNamePredicate);
+	int32 NameNumber = InName.GetNumber() + 1;
+	while(bAlreadyExists)
+	{
+		NewGraphName = FName(InName, NameNumber++);
+		bAlreadyExists =  Entries.ContainsByPredicate(DuplicateNamePredicate);
+	}
+
+	UAnimNextModule_EventGraph* NewEntry = CreateNewSubEntry<UAnimNextModule_EventGraph>(this);
+	NewEntry->GraphName = NewGraphName;
+
+	if(bSetupUndoRedo)
+	{
+		NewEntry->Modify();
+		Modify();
+	}
+
+	Entries.Add(NewEntry);
+
+	// Add new graph
+	{
+		TGuardValue<bool> EnablePythonPrint(bSuspendPythonMessagesForRigVMClient, !bPrintPythonCommand);
+		TGuardValue<bool> DisableAutoCompile(bAutoRecompileVM, false);
+		URigVMGraph* NewGraph = RigVMClient.AddModel(URigVMGraph::StaticClass()->GetFName(), UAnimNextEventGraphSchema::StaticClass(), bSetupUndoRedo);
+		ensure(NewGraph);
+		NewEntry->Graph = NewGraph;
+
+		URigVMController* Controller = RigVMClient.GetController(NewGraph);
+		UE::AnimNext::UncookedOnly::FUtils::SetupEventGraph(Controller);
+	}
+
+	BroadcastModified();
+
+	return NewEntry;
+}
+
+UAnimNextModule_AnimationGraph* UAnimNextModuleLibrary::AddAnimationGraph(UAnimNextModule* InModule, FName InName, bool bSetupUndoRedo, bool bPrintPythonCommand)
+{
+	return UE::AnimNext::UncookedOnly::FUtils::GetEditorData(InModule)->AddAnimationGraph(InName, bSetupUndoRedo, bPrintPythonCommand);
+}
+
+UAnimNextModule_AnimationGraph* UAnimNextModule_EditorData::AddAnimationGraph(FName InName, bool bSetupUndoRedo, bool bPrintPythonCommand)
+{
+	if(InName == NAME_None)
+	{
+		ReportError(TEXT("UAnimNextModule_EditorData::AddAnimationGraph: Invalid graph name supplied."));
+		return nullptr;
+	}
+
+	// Check for duplicate name
+	FName NewGraphName = InName;
+	auto DuplicateNamePredicate = [&NewGraphName](const UAnimNextRigVMAssetEntry* InEntry)
+	{
+		return InEntry->GetEntryName() == NewGraphName;
+	};
+
+	bool bAlreadyExists = Entries.ContainsByPredicate(DuplicateNamePredicate);
+	int32 NameNumber = InName.GetNumber() + 1;
+	while(bAlreadyExists)
+	{
+		NewGraphName = FName(InName, NameNumber++);
+		bAlreadyExists =  Entries.ContainsByPredicate(DuplicateNamePredicate);
+	}
+
+	UAnimNextModule_AnimationGraph* NewEntry = CreateNewSubEntry<UAnimNextModule_AnimationGraph>(this);
+	NewEntry->GraphName = NewGraphName;
+
+	if(bSetupUndoRedo)
+	{
+		NewEntry->Modify();
+		Modify();
+	}
+	
+	Entries.Add(NewEntry);
+
+	// Add new graph
+	{
+		TGuardValue<bool> EnablePythonPrint(bSuspendPythonMessagesForRigVMClient, !bPrintPythonCommand);
+		TGuardValue<bool> DisableAutoCompile(bAutoRecompileVM, false);
+		URigVMGraph* NewGraph = RigVMClient.AddModel(URigVMGraph::StaticClass()->GetFName(), UAnimNextAnimationGraphSchema::StaticClass(), bSetupUndoRedo);
+		ensure(NewGraph);
+		NewEntry->Graph = NewGraph;
+
+		URigVMController* Controller = RigVMClient.GetController(NewGraph);
+		UE::AnimNext::UncookedOnly::FUtils::SetupAnimGraph(NewEntry, Controller);
+	}
+
+	BroadcastModified();
+
+	return NewEntry;
+}
\ No newline at end of file
diff --git a/Engine/Plugins/Experimental/AnimNext/Source/AnimNextUncookedOnly/Private/Module/AnimNextModule_EventGraph.cpp b/Engine/Plugins/Experimental/AnimNext/Source/AnimNextUncookedOnly/Private/Module/AnimNextModule_EventGraph.cpp
new file mode 100644
--- /dev/null
+++ b/Engine/Plugins/Experimental/AnimNext/Source/AnimNextUncookedOnly/Private/Module/AnimNextModule_EventGraph.cpp
@@ -0,0 +1,49 @@
+﻿// Copyright Epic Games, Inc. All Rights Reserved.
+
+#include "Module/AnimNextModule_EventGraph.h"
+
+#include "AnimNextRigVMAsset.h"
+#include "UncookedOnlyUtils.h"
+#include "AnimNextEdGraph.h"
+
+FText UAnimNextModule_EventGraph::GetDisplayName() const
+{
+	return FText::FromName(GraphName);
+}
+
+FText UAnimNextModule_EventGraph::GetDisplayNameTooltip() const
+{
+	return FText::FromName(GraphName);
+}
+
+void UAnimNextModule_EventGraph::SetEntryName(FName InName, bool bSetupUndoRedo)
+{
+	if(bSetupUndoRedo)
+	{
+		Modify();
+	}
+
+	GraphName = InName;
+
+	BroadcastModified();
+}
+
+URigVMGraph* UAnimNextModule_EventGraph::GetRigVMGraph() const
+{
+	return Graph;
+}
+
+URigVMEdGraph* UAnimNextModule_EventGraph::GetEdGraph() const
+{
+	return EdGraph;
+}
+
+void UAnimNextModule_EventGraph::SetRigVMGraph(URigVMGraph* InGraph)
+{
+	Graph = InGraph;
+}
+
+void UAnimNextModule_EventGraph::SetEdGraph(URigVMEdGraph* InGraph)
+{
+	EdGraph = CastChecked<UAnimNextEdGraph>(InGraph);
+}
\ No newline at end of file
diff --git a/Engine/Plugins/Experimental/AnimNext/Source/AnimNextUncookedOnly/Private/Module/AnimNextModule_EventGraphSchema.cpp b/Engine/Plugins/Experimental/AnimNext/Source/AnimNextUncookedOnly/Private/Module/AnimNextModule_EventGraphSchema.cpp
new file mode 100644
--- /dev/null
+++ b/Engine/Plugins/Experimental/AnimNext/Source/AnimNextUncookedOnly/Private/Module/AnimNextModule_EventGraphSchema.cpp
@@ -0,0 +1,23 @@
+﻿// Copyright Epic Games, Inc. All Rights Reserved.
+
+#include "AnimNextEventGraphSchema.h"
+
+#include "AnimNextExecuteContext.h"
+#include "Param/RigUnit_AnimNextParameterBase.h"
+
+bool UAnimNextEventGraphSchema::SupportsUnitFunction(URigVMController* InController, const FRigVMFunction* InUnitFunction) const
+{
+	if(const UScriptStruct* FunctionExecuteContextStruct = InUnitFunction->GetExecuteContextStruct())
+	{
+		if(FunctionExecuteContextStruct == FAnimNextExecuteContext::StaticStruct())
+		{
+			// Only allow nodes that are children of FRigUnit_AnimNextParameterBase
+			if(InUnitFunction->Struct)
+			{
+				return InUnitFunction->Struct->IsChildOf(FRigUnit_AnimNextParameterBase::StaticStruct());
+			}
+		}
+	}
+
+	return Super::SupportsUnitFunction(InController, InUnitFunction);
+}
diff --git a/Engine/Plugins/Experimental/AnimNext/Source/AnimNextUncookedOnly/Private/Module/AnimNextModule_Parameter.cpp b/Engine/Plugins/Experimental/AnimNext/Source/AnimNextUncookedOnly/Private/Module/AnimNextModule_Parameter.cpp
new file mode 100644
--- /dev/null
+++ b/Engine/Plugins/Experimental/AnimNext/Source/AnimNextUncookedOnly/Private/Module/AnimNextModule_Parameter.cpp
@@ -0,0 +1,101 @@
+// Copyright Epic Games, Inc. All Rights Reserved.
+
+#include "Module/AnimNextModule_Parameter.h"
+
+#include "UncookedOnlyUtils.h"
+#include "Module/AnimNextModule.h"
+#include "Param/AnimNextTag.h"
+
+FAnimNextParamType UAnimNextModule_Parameter::GetExportType() const
+{
+	return GetParamType();
+}
+
+FName UAnimNextModule_Parameter::GetExportName() const
+{
+	return GetParamName();
+}
+
+EAnimNextExportAccessSpecifier UAnimNextModule_Parameter::GetExportAccessSpecifier() const
+{
+	return Access;
+}
+
+void UAnimNextModule_Parameter::SetExportAccessSpecifier(EAnimNextExportAccessSpecifier InAccessSpecifier, bool bSetupUndoRedo)
+{
+	if(bSetupUndoRedo)
+	{
+		Modify();
+	};
+
+	Access = InAccessSpecifier;
+
+	BroadcastModified();
+}
+
+FAnimNextParamType UAnimNextModule_Parameter::GetParamType() const
+{
+	return Type;
+}
+
+FName UAnimNextModule_Parameter::GetEntryName() const
+{
+	return ParameterName;
+}
+
+bool UAnimNextModule_Parameter::SetParamType(const FAnimNextParamType& InType, bool bSetupUndoRedo)
+{
+	if(bSetupUndoRedo)
+	{
+		Modify();
+	}
+	
+	Type = InType;
+
+	BroadcastModified();
+
+	return true;
+}
+
+FName UAnimNextModule_Parameter::GetParamName() const
+{
+	if(UAnimNextRigVMAsset* OuterAsset = GetTypedOuter<UAnimNextRigVMAsset>())
+	{
+		return UE::AnimNext::UncookedOnly::FUtils::GetQualifiedName(OuterAsset, ParameterName);
+	}
+	return ParameterName;
+}
+
+void UAnimNextModule_Parameter::SetParamName(FName InName, bool bSetupUndoRedo)
+{
+	SetEntryName(InName, bSetupUndoRedo);
+}
+
+FInstancedPropertyBag& UAnimNextModule_Parameter::GetPropertyBag() const
+{
+	// TODO: move property bag for defaults onto this entry!
+	UAnimNextModule* Asset = GetTypedOuter<UAnimNextModule>();
+	check(Asset);
+	return Asset->DefaultState.State;
+}
+
+void UAnimNextModule_Parameter::SetEntryName(FName InName, bool bSetupUndoRedo)
+{
+	if(bSetupUndoRedo)
+	{
+		Modify();
+	}
+
+	ParameterName = InName;
+	BroadcastModified();
+}
+
+FText UAnimNextModule_Parameter::GetDisplayName() const
+{
+	return FText::FromName(ParameterName);
+}
+
+FText UAnimNextModule_Parameter::GetDisplayNameTooltip() const
+{
+	return FText::FromString(Comment);
+}
\ No newline at end of file
diff --git a/Engine/Plugins/Experimental/AnimNext/Source/AnimNextUncookedOnly/Private/Param/ObjectProxyType.cpp b/Engine/Plugins/Experimental/AnimNext/Source/AnimNextUncookedOnly/Private/Param/ObjectProxyType.cpp
--- a/Engine/Plugins/Experimental/AnimNext/Source/AnimNextUncookedOnly/Private/Param/ObjectProxyType.cpp
+++ b/Engine/Plugins/Experimental/AnimNext/Source/AnimNextUncookedOnly/Private/Param/ObjectProxyType.cpp
@@ -2,7 +2,7 @@
 
 #include "ObjectProxyType.h"
 
-#include "AnimNextModule.h"
+#include "AnimNextModuleImpl.h"
 #include "IUniversalObjectLocatorEditorModule.h"
 #include "UniversalObjectLocatorEditor.h"
 #include "UniversalObjectLocatorFragmentType.h"
@@ -135,7 +135,7 @@ bool FObjectProxyType::FindParameterInfo(const TInstancedStruct<FAnimNextParamIn
 		return false;
 	}
 
-	UE::AnimNext::FModule& AnimNextModule = FModuleManager::GetModuleChecked<UE::AnimNext::FModule>("AnimNext");
+	UE::AnimNext::FAnimNextModuleImpl& AnimNextModule = FModuleManager::GetModuleChecked<UE::AnimNext::FAnimNextModuleImpl>("AnimNext");
 	TSharedPtr<FObjectProxyFactory> ObjectProxyFactory = StaticCastSharedPtr<FObjectProxyFactory>(AnimNextModule.FindParameterSourceFactory("ObjectProxy"));
 	if(!ObjectProxyFactory.IsValid())
 	{
@@ -179,7 +179,7 @@ void FObjectProxyType::ForEachParameter(const TInstancedStruct<FAnimNextParamIns
 		return;
 	}
 
-	UE::AnimNext::FModule& AnimNextModule = FModuleManager::GetModuleChecked<UE::AnimNext::FModule>("AnimNext");
+	UE::AnimNext::FAnimNextModuleImpl& AnimNextModule = FModuleManager::GetModuleChecked<UE::AnimNext::FAnimNextModuleImpl>("AnimNext");
 	TSharedPtr<FObjectProxyFactory> ObjectProxyFactory = StaticCastSharedPtr<FObjectProxyFactory>(AnimNextModule.FindParameterSourceFactory("ObjectProxy"));
 	if(!ObjectProxyFactory.IsValid())
 	{
