Rename AnimNextGraph asset to AnimNextModule

diff --git a/Engine/Plugins/Experimental/AnimNext/Source/AnimNextUncookedOnly/Private/UncookedOnlyUtils.cpp b/Engine/Plugins/Experimental/AnimNext/Source/AnimNextUncookedOnly/Private/UncookedOnlyUtils.cpp
--- a/Engine/Plugins/Experimental/AnimNext/Source/AnimNextUncookedOnly/Private/UncookedOnlyUtils.cpp
+++ b/Engine/Plugins/Experimental/AnimNext/Source/AnimNextUncookedOnly/Private/UncookedOnlyUtils.cpp
@@ -7,9 +7,9 @@
 #include "AssetRegistry/AssetRegistryModule.h"
 #include "AssetRegistry/IAssetRegistry.h"
 #include "Component/AnimNextComponentParameter.h"
-#include "Graph/AnimNextGraph.h"
-#include "Graph/AnimNextGraph_Controller.h"
-#include "Graph/AnimNextGraph_EditorData.h"
+#include "Module/AnimNextModule.h"
+#include "Module/AnimNextModule_Controller.h"
+#include "Module/AnimNextModule_EditorData.h"
 #include "Graph/RigUnit_AnimNextTraitStack.h"
 #include "Graph/RigUnit_AnimNextGraphRoot.h"
 #include "Graph/RigUnit_AnimNextGraphEvaluator.h"
@@ -44,12 +44,12 @@
 #include "IAnimNextRigVMExportInterface.h"
 #include "GameFramework/Character.h"
 #include "Graph/AnimNextGraphEntryPoint.h"
-#include "Graph/AnimNextGraph_AnimationGraphSchema.h"
+#include "Graph/AnimNextAnimationGraphSchema.h"
 #include "Logging/StructuredLog.h"
 #include "Param/AnimNextParamInstanceIdentifier.h"
 #include "Param/IParameterSourceType.h"
 #include "Param/RigVMDispatch_GetScopedParameter.h"
-#include "AnimNextRigVMWorkspaceAssetUserData.h"
+#include "Module/AnimNextModuleWorkspaceAssetUserData.h"
 
 #define LOCTEXT_NAMESPACE "AnimNextUncookedOnlyUtils"
 
@@ -102,11 +102,11 @@ namespace Private
 		URigVMNode* RootNode;
 		TArray<FTraitStackMapping> TraitStackNodes;
 
-		explicit FTraitGraph(const UAnimNextGraph* InGraph, URigVMNode* InRootNode)
+		explicit FTraitGraph(const UAnimNextModule* InModule, URigVMNode* InRootNode)
 			: RootNode(InRootNode)
 		{
 			TStringBuilder<256> StringBuilder;
-			StringBuilder.Append(InGraph->GetPathName());
+			StringBuilder.Append(InModule->GetPathName());
 			StringBuilder.Append(TEXT(":"));
 			StringBuilder.Append(InRootNode->FindPin(GET_MEMBER_NAME_STRING_CHECKED(FRigUnit_AnimNextGraphRoot, EntryPoint))->GetDefaultValue());
 			EntryPoint = FName(StringBuilder.ToView());
@@ -349,7 +349,7 @@ namespace Private
 		}
 	}
 
-	FTraitGraph CollectGraphInfo(const UAnimNextGraph* InGraph,  const URigVMGraph* VMGraph, URigVMController* VMController)
+	FTraitGraph CollectGraphInfo(const UAnimNextModule* InModule, const URigVMGraph* VMGraph, URigVMController* VMController)
 	{
 		const TArray<URigVMNode*>& VMNodes = VMGraph->GetNodes();
 		URigVMUnitNode* VMRootNode = FindRootNode(VMNodes);
@@ -363,7 +363,7 @@ namespace Private
 		// Make sure we don't have empty input pins
 		AddMissingInputLinks(VMGraph, VMController);
 
-		FTraitGraph TraitGraph(InGraph, VMRootNode);
+		FTraitGraph TraitGraph(InModule, VMRootNode);
 
 		TArray<const URigVMNode*> NodesToVisit;
 		NodesToVisit.Add(VMRootNode);
@@ -500,21 +500,21 @@ namespace Private
 	}
 }
 
-void FUtils::Compile(UAnimNextGraph* InGraph)
+void FUtils::Compile(UAnimNextModule* InModule)
 {
-	check(InGraph);
+	check(InModule);
 
-	FMessageLog("AnimNextCompilerResults").NewPage(FText::FromName(InGraph->GetFName()));
+	FMessageLog("AnimNextCompilerResults").NewPage(FText::FromName(InModule->GetFName()));
 
-	CompileStruct(InGraph);
-	CompileVM(InGraph);
+	CompileStruct(InModule);
+	CompileVM(InModule);
 }
 
-void FUtils::CompileVM(UAnimNextGraph* InGraph)
+void FUtils::CompileVM(UAnimNextModule* InModule)
 {
-	check(InGraph);
+	check(InModule);
 
-	UAnimNextGraph_EditorData* EditorData = GetEditorData(InGraph);
+	UAnimNextModule_EditorData* EditorData = GetEditorData(InModule);
 
 	if (EditorData->bIsCompiling)
 	{
@@ -525,7 +525,7 @@ void FUtils::CompileVM(UAnimNextGraph* InGraph)
 	
 	// Before we re-compile a graph, we need to release and live instances since we need the metadata we are about to replace
 	// to call trait destructors etc
-	InGraph->FreezeGraphInstances();
+	InModule->FreezeGraphInstances();
 
 	EditorData->bErrorsDuringCompilation = false;
 
@@ -535,17 +535,17 @@ void FUtils::CompileVM(UAnimNextGraph* InGraph)
 	TGuardValue<bool> ReentrantGuardSelf(EditorData->bSuspendModelNotificationsForSelf, true);
 	TGuardValue<bool> ReentrantGuardOthers(EditorData->bSuspendModelNotificationsForOthers, true);
 
-	RecreateVM(InGraph);
+	RecreateVM(InModule);
 
-	InGraph->VMRuntimeSettings = EditorData->VMRuntimeSettings;
-	InGraph->EntryPoints.Empty();
-	InGraph->ResolvedRootTraitHandles.Empty();
-	InGraph->ResolvedEntryPoints.Empty();
-	InGraph->ExecuteDefinition = FAnimNextGraphEvaluatorExecuteDefinition();
-	InGraph->SharedDataBuffer.Empty();
-	InGraph->GraphReferencedObjects.Empty();
-	InGraph->RequiredParametersHash = 0;
-	InGraph->RequiredParameters.Empty();
+	InModule->VMRuntimeSettings = EditorData->VMRuntimeSettings;
+	InModule->EntryPoints.Empty();
+	InModule->ResolvedRootTraitHandles.Empty();
+	InModule->ResolvedEntryPoints.Empty();
+	InModule->ExecuteDefinition = FAnimNextGraphEvaluatorExecuteDefinition();
+	InModule->SharedDataBuffer.Empty();
+	InModule->GraphReferencedObjects.Empty();
+	InModule->RequiredParametersHash = 0;
+	InModule->RequiredParameters.Empty();
 
 	FRigVMClient* VMClient = EditorData->GetRigVMClient();
 	URigVMGraph* VMRootGraph = VMClient->GetDefaultModel();
@@ -569,7 +569,7 @@ void FUtils::CompileVM(UAnimNextGraph* InGraph)
 		return;
 	}
 
-	UAnimNextGraph_Controller* TempController = CastChecked<UAnimNextGraph_Controller>(VMClient->GetOrCreateController(TempGraphs[0]));
+	UAnimNextModule_Controller* TempController = CastChecked<UAnimNextModule_Controller>(VMClient->GetOrCreateController(TempGraphs[0]));
 
 	FTraitWriter TraitWriter;
 
@@ -580,13 +580,13 @@ void FUtils::CompileVM(UAnimNextGraph* InGraph)
 	// Build entry points and extract their required latent pins
 	for(const URigVMGraph* TempGraph : TempGraphs)
 	{
-		if(TempGraph->GetSchemaClass() == UAnimNextGraph_AnimationGraphSchema::StaticClass())
+		if(TempGraph->GetSchemaClass() == UAnimNextAnimationGraphSchema::StaticClass())
 		{
 			// Gather our trait stacks
-			Private::FTraitGraph& TraitGraph = TraitGraphs.Add_GetRef(Private::CollectGraphInfo(InGraph, TempGraph, TempController->GetControllerForGraph(TempGraph)));
+			Private::FTraitGraph& TraitGraph = TraitGraphs.Add_GetRef(Private::CollectGraphInfo(InModule, TempGraph, TempController->GetControllerForGraph(TempGraph)));
 			check(!TraitGraph.TraitStackNodes.IsEmpty());
 
-			FAnimNextGraphEntryPoint& EntryPoint = InGraph->EntryPoints.AddDefaulted_GetRef();
+			FAnimNextGraphEntryPoint& EntryPoint = InModule->EntryPoints.AddDefaulted_GetRef();
 			EntryPoint.EntryPointName = TraitGraph.EntryPoint;
 
 			// Extract latent pins for this graph
@@ -616,11 +616,11 @@ void FUtils::CompileVM(UAnimNextGraph* InGraph)
 	if(LatentPins.Num() > 0)
 	{
 		// We need a unique method name to match our unique argument list
-		InGraph->ExecuteDefinition = Private::GetGraphEvaluatorExecuteMethod(LatentPins);
+		InModule->ExecuteDefinition = Private::GetGraphEvaluatorExecuteMethod(LatentPins);
 
 		// Add our runtime shim root node
 		URigVMUnitNode* TempShimRootNode = TempController->AddUnitNode(FRigUnit_AnimNextShimRoot::StaticStruct(), FRigUnit_AnimNextShimRoot::EventName, FVector2D::ZeroVector, FString(), false);
-		URigVMUnitNode* GraphEvaluatorNode = TempController->AddUnitNodeWithPins(FRigUnit_AnimNextGraphEvaluator::StaticStruct(), LatentPins, *InGraph->ExecuteDefinition.MethodName, FVector2D::ZeroVector, FString(), false);
+		URigVMUnitNode* GraphEvaluatorNode = TempController->AddUnitNodeWithPins(FRigUnit_AnimNextGraphEvaluator::StaticStruct(), LatentPins, *InModule->ExecuteDefinition.MethodName, FVector2D::ZeroVector, FString(), false);
 
 		// Link our shim and evaluator nodes together using the execution context
 		TempController->AddLink(
@@ -652,42 +652,42 @@ void FUtils::CompileVM(UAnimNextGraph* InGraph)
 	TraitWriter.EndNodeWriting();
 
 	// Cache our compiled metadata
-	InGraph->SharedDataArchiveBuffer = TraitWriter.GetGraphSharedData();
-	InGraph->GraphReferencedObjects = TraitWriter.GetGraphReferencedObjects();
+	InModule->SharedDataArchiveBuffer = TraitWriter.GetGraphSharedData();
+	InModule->GraphReferencedObjects = TraitWriter.GetGraphReferencedObjects();
 
 	// Populate our runtime metadata
-	InGraph->LoadFromArchiveBuffer(InGraph->SharedDataArchiveBuffer);
+	InModule->LoadFromArchiveBuffer(InModule->SharedDataArchiveBuffer);
 
 	URigVMCompiler* Compiler = URigVMCompiler::StaticClass()->GetDefaultObject<URigVMCompiler>();
 	EditorData->VMCompileSettings.SetExecuteContextStruct(FAnimNextExecuteContext::StaticStruct());
 	FRigVMCompileSettings Settings = (EditorData->bCompileInDebugMode) ? FRigVMCompileSettings::Fast(EditorData->VMCompileSettings.GetExecuteContextStruct()) : EditorData->VMCompileSettings;
 	Settings.ASTSettings.bSetupTraits = false; // disable the default implementation of decorators for now
-	Settings.ASTSettings.ReportDelegate.BindLambda([InGraph](EMessageSeverity::Type InType, UObject* InObject, const FString& InString)
+	Settings.ASTSettings.ReportDelegate.BindLambda([InModule](EMessageSeverity::Type InType, UObject* InObject, const FString& InString)
 	{
 		FMessageLog("AnimNextCompilerResults").Message(InType, FText::FromString(InString));
 	});
 
-	Compiler->Compile(Settings, TempGraphs, TempController, InGraph->VM, InGraph->ExtendedExecuteContext, TArray<FRigVMExternalVariable>(), & EditorData->PinToOperandMap);
+	Compiler->Compile(Settings, TempGraphs, TempController, InModule->VM, InModule->ExtendedExecuteContext, TArray<FRigVMExternalVariable>(), & EditorData->PinToOperandMap);
 
 	// Initialize right away, in packaged builds we initialize during PostLoad
-	InGraph->VM->Initialize(InGraph->ExtendedExecuteContext);
-	InGraph->GenerateUserDefinedDependenciesData(InGraph->ExtendedExecuteContext);
+	InModule->VM->Initialize(InModule->ExtendedExecuteContext);
+	InModule->GenerateUserDefinedDependenciesData(InModule->ExtendedExecuteContext);
 
 	// Notable difference with vanilla RigVM host behavior - we init the VM here at the moment as we only have one 'instance'
-	InGraph->InitializeVM(FRigUnit_AnimNextBeginExecution::EventName);
+	InModule->InitializeVM(FRigUnit_AnimNextBeginExecution::EventName);
 
 	if (EditorData->bErrorsDuringCompilation)
 	{
 		if(Settings.SurpressErrors)
 		{
-			Settings.Reportf(EMessageSeverity::Info, InGraph,TEXT("Compilation Errors may be suppressed for AnimNext Interface Graph: %s. See VM Compile Settings for more Details"), *InGraph->GetName());
+			Settings.Reportf(EMessageSeverity::Info, InModule, TEXT("Compilation Errors may be suppressed for AnimNext Interface Graph: %s. See VM Compile Settings for more Details"), *InModule->GetName());
 		}
 	}
 
 	EditorData->bVMRecompilationRequired = false;
-	if(InGraph->VM)
+	if(InModule->VM)
 	{
-		EditorData->RigVMCompiledEvent.Broadcast(InGraph, InGraph->VM, InGraph->ExtendedExecuteContext);
+		EditorData->RigVMCompiledEvent.Broadcast(InModule, InModule->VM, InModule->ExtendedExecuteContext);
 	}
 
 	for(URigVMGraph* TempGraph : TempGraphs)
@@ -696,7 +696,7 @@ void FUtils::CompileVM(UAnimNextGraph* InGraph)
 	}
 
 	// Now that the graph has been re-compiled, re-allocate the previous live instances
-	InGraph->ThawGraphInstances();
+	InModule->ThawGraphInstances();
 	
 	FAnimNextParameterProviderAssetRegistryExports Exports;
 	GetAssetParameters(EditorData, Exports);
@@ -706,39 +706,39 @@ void FUtils::CompileVM(UAnimNextGraph* InGraph)
 		// Required parameters are those that are read in this asset but not declared in this asset as state
 		if(EnumHasAllFlags(Entry.GetFlags(), EAnimNextParameterFlags::Read) && !EnumHasAnyFlags(Entry.GetFlags(), EAnimNextParameterFlags::Declared))
 		{
-			InGraph->RequiredParameters.Emplace(Entry.Name, Entry.Type, Entry.InstanceId);
+			InModule->RequiredParameters.Emplace(Entry.Name, Entry.Type, Entry.InstanceId);
 		}
 	}
 
-	InGraph->RequiredParametersHash = SortAndHashParameters(InGraph->RequiredParameters);
+	InModule->RequiredParametersHash = SortAndHashParameters(InModule->RequiredParameters);
 
 #if WITH_EDITOR
 //	RefreshBreakpoints(EditorData);
 #endif
 }
 
-void FUtils::RecreateVM(UAnimNextGraph* InGraph)
+void FUtils::RecreateVM(UAnimNextModule* InModule)
 {
-	if (InGraph->VM == nullptr)
+	if (InModule->VM == nullptr)
 	{
-		InGraph->VM = NewObject<URigVM>(InGraph, TEXT("VM"), RF_NoFlags);
+		InModule->VM = NewObject<URigVM>(InModule, TEXT("VM"), RF_NoFlags);
 	}
-	InGraph->VM->Reset(InGraph->ExtendedExecuteContext);
-	InGraph->RigVM = InGraph->VM; // Local serialization
+	InModule->VM->Reset(InModule->ExtendedExecuteContext);
+	InModule->RigVM = InModule->VM; // Local serialization
 }
 
-UAnimNextGraph_EditorData* FUtils::GetEditorData(const UAnimNextGraph* InAnimNextGraph)
+UAnimNextModule_EditorData* FUtils::GetEditorData(const UAnimNextModule* InModule)
 {
-	check(InAnimNextGraph);
+	check(InModule);
 	
-	return CastChecked<UAnimNextGraph_EditorData>(InAnimNextGraph->EditorData);
+	return CastChecked<UAnimNextModule_EditorData>(InModule->EditorData);
 }
 
-UAnimNextGraph* FUtils::GetGraph(const UAnimNextGraph_EditorData* InEditorData)
+UAnimNextModule* FUtils::GetGraph(const UAnimNextModule_EditorData* InEditorData)
 {
 	check(InEditorData);
 
-	return CastChecked<UAnimNextGraph>(InEditorData->GetOuter());
+	return CastChecked<UAnimNextModule>(InEditorData->GetOuter());
 }
 
 FParamTypeHandle FUtils::GetParameterHandleFromPin(const FEdGraphPinType& InPinType)
@@ -845,11 +845,11 @@ FParamTypeHandle FUtils::GetParameterHandleFromPin(const FEdGraphPinType& InPinT
 	return FAnimNextParamType(ValueType, ContainerType, ValueTypeObject).GetHandle();
 }
 
-void FUtils::CompileStruct(UAnimNextGraph* InGraph)
+void FUtils::CompileStruct(UAnimNextModule* InModule)
 {
-	check(InGraph);
+	check(InModule);
 
-	UAnimNextGraph_EditorData* EditorData = GetEditorData(InGraph);
+	UAnimNextModule_EditorData* EditorData = GetEditorData(InModule);
 	if(EditorData->bIsCompiling)
 	{
 		return;
@@ -903,7 +903,7 @@ void FUtils::CompileStruct(UAnimNextGraph* InGraph)
 		TArray<FPropertyBagPropertyDesc> PropertyDescs;
 		PropertyDescs.Reserve(StructEntryInfos.Num());
 
-		InGraph->DefaultState.PublicParameterStartIndex = INDEX_NONE;
+		InModule->DefaultState.PublicParameterStartIndex = INDEX_NONE;
 
 		for(int32 EntryIndex = 0; EntryIndex < StructEntryInfos.Num(); ++EntryIndex)
 		{
@@ -911,7 +911,7 @@ void FUtils::CompileStruct(UAnimNextGraph* InGraph)
 			// Find the first parameter that is public and record it
 			if(StructEntryInfo.AccessSpecifier == EAnimNextExportAccessSpecifier::Public)
 			{
-				InGraph->DefaultState.PublicParameterStartIndex = EntryIndex;
+				InModule->DefaultState.PublicParameterStartIndex = EntryIndex;
 			}
 			PropertyDescs.Emplace(StructEntryInfo.Name, StructEntryInfo.Type.ContainerType, StructEntryInfo.Type.ValueType, StructEntryInfo.Type.ValueTypeObject);
 		}
@@ -920,9 +920,9 @@ void FUtils::CompileStruct(UAnimNextGraph* InGraph)
 		// TODO: linear search - we could cache the name->GUID lookup in editor to accelerate this.
 		for(FPropertyBagPropertyDesc& NewDesc : PropertyDescs)
 		{
-			if(InGraph->DefaultState.State.GetPropertyBagStruct())
+			if(InModule->DefaultState.State.GetPropertyBagStruct())
 			{
-				for(const FPropertyBagPropertyDesc& ExistingDesc : InGraph->DefaultState.State.GetPropertyBagStruct()->GetPropertyDescs())
+				for(const FPropertyBagPropertyDesc& ExistingDesc : InModule->DefaultState.State.GetPropertyBagStruct()->GetPropertyDescs())
 				{
 					if(ExistingDesc.Name == NewDesc.Name)
 					{
@@ -935,11 +935,11 @@ void FUtils::CompileStruct(UAnimNextGraph* InGraph)
 
 		// Create new property bag and migrate
 		const UPropertyBag* NewBagStruct = UPropertyBag::GetOrCreateFromDescs(PropertyDescs);
-		InGraph->DefaultState.State.MigrateToNewBagStruct(NewBagStruct);
+		InModule->DefaultState.State.MigrateToNewBagStruct(NewBagStruct);
 	}
 	else
 	{
-		InGraph->DefaultState.Reset();
+		InModule->DefaultState.Reset();
 	}
 }
 
@@ -955,9 +955,9 @@ UAnimNextRigVMAssetEditorData* FUtils::GetEditorData(UAnimNextRigVMAsset* InAsse
 	return CastChecked<UAnimNextRigVMAssetEditorData>(InAsset->EditorData);
 }
 
-FInstancedPropertyBag* FUtils::GetPropertyBag(UAnimNextGraph* InAnimNextGraph)
+FInstancedPropertyBag* FUtils::GetPropertyBag(UAnimNextModule* InModule)
 {
-	FInstancedPropertyBag* InstancedPropertyBag = &InAnimNextGraph->DefaultState.State;
+	FInstancedPropertyBag* InstancedPropertyBag = &InModule->DefaultState.State;
 
 	return InstancedPropertyBag;
 }
@@ -1872,14 +1872,14 @@ void FUtils::CompileSchedule(UAnimNextSchedule* InSchedule)
 			{
 				bool bValid = true;
 
-				if(GraphEntry->Graph == nullptr && !GraphEntry->DynamicGraph.IsValid())
+				if(GraphEntry->Module == nullptr && !GraphEntry->DynamicGraph.IsValid())
 				{
 					Log.Error(LOCTEXT("InvalidGraphOrParameterError", "Invalid graph or invalid parameter supplied in graph task"));
 					bValid = false;
 				}
-				else if(GraphEntry->Graph != nullptr)
+				else if(GraphEntry->Module != nullptr)
 				{
-					TConstArrayView<FScheduleTerm> Terms = GraphEntry->Graph->GetTerms();
+					TConstArrayView<FScheduleTerm> Terms = GraphEntry->Module->GetTerms();
 					if(GraphEntry->Terms.Num() != Terms.Num())
 					{
 						Log.Error(FText::Format(LOCTEXT("GraphIncorrectTermCountError", "Incorrect term count for graph: {0}"), FText::AsNumber(GraphEntry->Terms.Num())));
@@ -1951,11 +1951,11 @@ void FUtils::CompileSchedule(UAnimNextSchedule* InSchedule)
 					GraphTask.ParamScopeIndex = InSchedule->NumParameterScopes++;
 					GraphTask.ParamParentScopeIndex = ParentScopeIndex;
 					GraphTask.EntryPoint = FAnimNextParam(GraphEntry->EntryPoint);
-					GraphTask.Graph = GraphEntry->Graph;
-					GraphTask.DynamicGraph = FAnimNextParam(GraphEntry->DynamicGraph);
+					GraphTask.Module = GraphEntry->Module;
+					GraphTask.DynamicModule = FAnimNextParam(GraphEntry->DynamicGraph);
 					GraphTask.ReferencePose = FAnimNextParam(GraphEntry->ReferencePose);
 					GraphTask.LOD = FAnimNextParam(GraphEntry->LOD);
-					if(GraphEntry->Graph == nullptr && GraphEntry->DynamicGraph.IsValid())
+					if(GraphEntry->Module == nullptr && GraphEntry->DynamicGraph.IsValid())
 					{
 						Algo::Transform(GraphEntry->RequiredParameters, GraphTask.SuppliedParameters, [](const FAnimNextEditorParam& InParam){ return FAnimNextParam(InParam); });
 						GraphTask.SuppliedParametersHash = SortAndHashParameters(GraphTask.SuppliedParameters);
@@ -2126,10 +2126,10 @@ void FUtils::CompileSchedule(UAnimNextSchedule* InSchedule)
 			{
 				if (UAnimNextScheduleEntry_AnimNextGraph* GraphEntry = Cast<UAnimNextScheduleEntry_AnimNextGraph>(Entry))
 				{
-					if(GraphEntry->Graph)
+					if(GraphEntry->Module)
 					{
 						FAnimNextParameterProviderAssetRegistryExports Exports;
-						if(UncookedOnly::FUtils::GetExportedParametersForAsset(FAssetData(GraphEntry->Graph), Exports))
+						if(UncookedOnly::FUtils::GetExportedParametersForAsset(FAssetData(GraphEntry->Module), Exports))
 						{
 							TArray<FAnimNextEditorParam> RequiredParameters;
 							RequiredParameters.Reserve(Exports.Parameters.Num());
@@ -2161,7 +2161,7 @@ void FUtils::CompileSchedule(UAnimNextSchedule* InSchedule)
 				}
 				else if (UAnimNextScheduleEntry_ParamScope* ParamScopeTaskEntry = Cast<UAnimNextScheduleEntry_ParamScope>(Entry))
 				{
-					for(UAnimNextGraph* Parameters : ParamScopeTaskEntry->Parameters)
+					for(UAnimNextModule* Parameters : ParamScopeTaskEntry->Parameters)
 					{
 						if(Parameters)
 						{
diff --git a/Engine/Plugins/Experimental/AnimNextAnimGraph/Config/BaseAnimNextAnimGraph.ini b/Engine/Plugins/Experimental/AnimNextAnimGraph/Config/BaseAnimNextAnimGraph.ini
--- a/Engine/Plugins/Experimental/AnimNextAnimGraph/Config/BaseAnimNextAnimGraph.ini
+++ b/Engine/Plugins/Experimental/AnimNextAnimGraph/Config/BaseAnimNextAnimGraph.ini
@@ -23,3 +23,6 @@
 +StructRedirects=(OldName="/Script/AnimNext.AnimNextSubGraphHostTraitSharedData", NewName="/Script/AnimNextAnimGraph.AnimNextSubGraphHostTraitSharedData")
 +StructRedirects=(OldName="/Script/AnimNext.AnimNextSynchronizeUsingGroupsTraitSharedData", NewName="/Script/AnimNextAnimGraph.AnimNextSynchronizeUsingGroupsTraitSharedData")
 +StructRedirects=(OldName="/Script/AnimNext.AnimNode_AnimNextGraph", NewName="/Script/AnimNextAnimGraph.AnimNode_AnimNextGraph")
+
+; AnimNextGraph -> AnimNextModule
++PropertyRedirects=(OldName="/Script/AnimNextGraph.AnimNode_AnimNextGraph.AnimNextGraph", NewName="Module")
\ No newline at end of file
diff --git a/Engine/Plugins/Experimental/AnimNextAnimGraph/Source/AnimNextAnimGraph/Internal/Graph/AnimGraph/AnimNode_AnimNextGraph.h b/Engine/Plugins/Experimental/AnimNextAnimGraph/Source/AnimNextAnimGraph/Internal/Graph/AnimGraph/AnimNode_AnimNextGraph.h
--- a/Engine/Plugins/Experimental/AnimNextAnimGraph/Source/AnimNextAnimGraph/Internal/Graph/AnimGraph/AnimNode_AnimNextGraph.h
+++ b/Engine/Plugins/Experimental/AnimNextAnimGraph/Source/AnimNextAnimGraph/Internal/Graph/AnimGraph/AnimNode_AnimNextGraph.h
@@ -3,7 +3,7 @@
 #pragma once
 
 #include "Animation/AnimNode_CustomProperty.h"
-#include "Graph/AnimNextGraph.h"
+#include "Module/AnimNextModule.h"
 #include "Graph/AnimNextGraphInstancePtr.h"
 #include "TraitCore/TraitPtr.h"
 #include "Context.h"
@@ -42,7 +42,7 @@ private:
 	FPoseLink SourceLink;
 
 	UPROPERTY(EditAnywhere, Category = Settings, meta = (PinHiddenByDefault))
-	TObjectPtr<UAnimNextGraph> AnimNextGraph;
+	TObjectPtr<UAnimNextModule> Module;
 
 	// Our graph instance, we own it
 	UPROPERTY()
@@ -58,7 +58,7 @@ private:
 	int32 LODThreshold;
 
 protected:
-	virtual UClass* GetTargetClass() const override { return AnimNextGraph ? AnimNextGraph->StaticClass() : nullptr; }
+	virtual UClass* GetTargetClass() const override { return Module ? Module->StaticClass() : nullptr; }
 	
 public:
 
diff --git a/Engine/Plugins/Experimental/AnimNextAnimGraph/Source/AnimNextAnimGraph/Private/Graph/AnimGraph/AnimNode_AnimNextGraph.cpp b/Engine/Plugins/Experimental/AnimNextAnimGraph/Source/AnimNextAnimGraph/Private/Graph/AnimGraph/AnimNode_AnimNextGraph.cpp
--- a/Engine/Plugins/Experimental/AnimNextAnimGraph/Source/AnimNextAnimGraph/Private/Graph/AnimGraph/AnimNode_AnimNextGraph.cpp
+++ b/Engine/Plugins/Experimental/AnimNextAnimGraph/Source/AnimNextAnimGraph/Private/Graph/AnimGraph/AnimNode_AnimNextGraph.cpp
@@ -26,7 +26,7 @@ TAutoConsoleVariable<int32> CVarAnimNextForceAnimBP(TEXT("a.AnimNextForceAnimBP"
 
 FAnimNode_AnimNextGraph::FAnimNode_AnimNextGraph()
 	: FAnimNode_CustomProperty()
-	, AnimNextGraph(nullptr)
+	, Module(nullptr)
 	, LODThreshold(INDEX_NONE)
 {
 }
@@ -79,9 +79,9 @@ void FAnimNode_AnimNextGraph::Initialize_AnyThread(const FAnimationInitializeCon
 
 	SourceLink.Initialize(Context);
 
-	if (!GraphInstance.IsValid() && AnimNextGraph)
+	if (!GraphInstance.IsValid() && Module)
 	{
-		AnimNextGraph->AllocateInstance(GraphInstance);
+		Module->AllocateInstance(GraphInstance);
 	}
 
 	FAnimNode_CustomProperty::Initialize_AnyThread(Context);
@@ -163,7 +163,7 @@ void FAnimNode_AnimNextGraph::PostSerialize(const FArchive& Ar)
 	// since memory has changed
 	if (Ar.IsObjectReferenceCollector())
 	{
-		if (AnimNextGraph)
+		if (Module)
 		{
 			//AnimNextGraph->Initialize();
 		}
diff --git a/Engine/Plugins/Experimental/AnimNextAnimGraph/Source/AnimNextAnimGraph/Private/Module.cpp b/Engine/Plugins/Experimental/AnimNextAnimGraph/Source/AnimNextAnimGraph/Private/Module.cpp
--- a/Engine/Plugins/Experimental/AnimNextAnimGraph/Source/AnimNextAnimGraph/Private/Module.cpp
+++ b/Engine/Plugins/Experimental/AnimNextAnimGraph/Source/AnimNextAnimGraph/Private/Module.cpp
@@ -3,14 +3,14 @@
 #include "CoreMinimal.h"
 #include "Animation/BlendProfile.h"
 #include "Curves/CurveFloat.h"
-#include "Graph/AnimNextGraph.h"
+#include "Module/AnimNextModule.h"
 #include "Modules/ModuleInterface.h"
 #include "Modules/ModuleManager.h"
 #include "RigVMCore/RigVMRegistry.h"
 #include "Animation/AnimSequence.h"
 #include "Chooser.h"
 
-#include "IAnimNextModule.h"
+#include "IAnimNextModuleInterface.h"
 #include "TraitInterfaces/IEvaluate.h"
 #include "TraitInterfaces/IUpdate.h"
 #include "EvaluationVM/EvaluationVM.h"
@@ -29,18 +29,18 @@ public:
 			{ UScriptStruct::StaticClass(), FRigVMRegistry::ERegisterObjectOperation::Class },
 			{ UBlendProfile::StaticClass(), FRigVMRegistry::ERegisterObjectOperation::Class },
 			{ UCurveFloat::StaticClass(), FRigVMRegistry::ERegisterObjectOperation::Class },
-			{ UAnimNextGraph::StaticClass(), FRigVMRegistry::ERegisterObjectOperation::Class },
+			{ UAnimNextModule::StaticClass(), FRigVMRegistry::ERegisterObjectOperation::Class },
 			{ UChooserTable::StaticClass(), FRigVMRegistry::ERegisterObjectOperation::Class },
 		};
 
 		FRigVMRegistry::Get().RegisterObjectTypes(AllowedObjectTypes);
 
-		IAnimNextModule::Get().RegisterAnimNextAnimGraph(*this);
+		IAnimNextModuleInterface::Get().RegisterAnimNextAnimGraph(*this);
 	}
 
 	virtual void ShutdownModule() override
 	{
-		IAnimNextModule::Get().UnregisterAnimNextAnimGraph();
+		IAnimNextModuleInterface::Get().UnregisterAnimNextAnimGraph();
 	}
 
 	virtual void UpdateGraph(FAnimNextGraphInstancePtr& GraphInstance, float DeltaTime, UE::AnimNext::FTraitEventList& InputEventList, FTraitEventList& OutputEventList) const override
diff --git a/Engine/Plugins/Experimental/AnimNextAnimGraph/Source/AnimNextAnimGraph/Private/PlayAnim/PlayAnimSlotTrait.cpp b/Engine/Plugins/Experimental/AnimNextAnimGraph/Source/AnimNextAnimGraph/Private/PlayAnim/PlayAnimSlotTrait.cpp
--- a/Engine/Plugins/Experimental/AnimNextAnimGraph/Source/AnimNextAnimGraph/Private/PlayAnim/PlayAnimSlotTrait.cpp
+++ b/Engine/Plugins/Experimental/AnimNextAnimGraph/Source/AnimNextAnimGraph/Private/PlayAnim/PlayAnimSlotTrait.cpp
@@ -31,14 +31,14 @@ namespace UE::AnimNext
 	#undef TRAIT_INTERFACE_ENUMERATOR
 	#undef TRAIT_EVENT_ENUMERATOR
 
-	void FPlayAnimSlotTrait::FPlayAnimSlotRequest::Initialize(FPlayAnimRequestPtr InRequest, const FPlayAnimBlendSettings& InBlendSettings, const UAnimNextGraph* InSubGraph)
+	void FPlayAnimSlotTrait::FPlayAnimSlotRequest::Initialize(FPlayAnimRequestPtr InRequest, const FPlayAnimBlendSettings& InBlendSettings, const UAnimNextModule* InModule)
 	{
 		Request = InRequest;
 		BlendSettings = InBlendSettings;
-		SubGraph = InSubGraph;
+		Module = InModule;
 
 		// If no input is provided, we'll use the source
-		State = InSubGraph != nullptr ? EPlayAnimRequestState::Active : EPlayAnimRequestState::ActiveSource;
+		State = InModule != nullptr ? EPlayAnimRequestState::Active : EPlayAnimRequestState::ActiveSource;
 		bWasRelevant = false;
 	}
 
@@ -155,7 +155,7 @@ namespace UE::AnimNext
 			InstanceData->PendingRequest.Reset();
 
 			FPlayAnimBlendSettings BlendSettings;
-			const UAnimNextGraph* SubGraph = nullptr;
+			const UAnimNextModule* Module = nullptr;
 			if (Request)
 			{
 				// This is a new pending request, lookup the sub-graph to use with our chooser and the desired animation object
@@ -170,21 +170,21 @@ namespace UE::AnimNext
 						FChooserEvaluationContext ChooserContext;
 						ChooserContext.AddStructParam(ChooserParameters);
 
-						UChooserTable::EvaluateChooser(ChooserContext, Chooser, FObjectChooserBase::FObjectChooserIteratorCallback::CreateLambda([&SubGraph](UObject* InResult)
+						UChooserTable::EvaluateChooser(ChooserContext, Chooser, FObjectChooserBase::FObjectChooserIteratorCallback::CreateLambda([&Module](UObject* InResult)
 							{
-								SubGraph = Cast<UAnimNextGraph>(InResult);
+								Module = Cast<UAnimNextModule>(InResult);
 								return FObjectChooserBase::EIteratorStatus::Stop;
 							}));
 					}
 
-					if (SubGraph != nullptr)
+					if (Module != nullptr)
 					{
 						// Check for re-entrancy and early-out if we are linking back to the current instance or one of its parents
-						const FName EntryPoint = SubGraph->GetDefaultEntryPoint();
+						const FName EntryPoint = Module->GetDefaultEntryPoint();
 						const FAnimNextGraphInstance* OwnerGraphInstance = &Binding.GetTraitPtr().GetNodeInstance()->GetOwner();
 						while (OwnerGraphInstance != nullptr)
 						{
-							if (OwnerGraphInstance->UsesGraph(SubGraph) && OwnerGraphInstance->UsesEntryPoint(EntryPoint))
+							if (OwnerGraphInstance->UsesModule(Module) && OwnerGraphInstance->UsesEntryPoint(EntryPoint))
 							{
 								return;
 							}
@@ -215,7 +215,7 @@ namespace UE::AnimNext
 			const int32 FreeRequestIndex = FindFreeRequestIndexOrAdd(*InstanceData);
 
 			FPlayAnimSlotRequest& SlotRequest = InstanceData->SlotRequests[FreeRequestIndex];
-			SlotRequest.Initialize(Request, BlendSettings, SubGraph);
+			SlotRequest.Initialize(Request, BlendSettings, Module);
 
 			const int32 OldChildIndex = InstanceData->CurrentlyActiveRequestIndex;
 			const int32 NewChildIndex = FreeRequestIndex;
@@ -397,14 +397,14 @@ namespace UE::AnimNext
 
 			if (SlotRequest.State == EPlayAnimRequestState::Active)
 			{
-				const FName EntryPoint = SlotRequest.SubGraph->GetDefaultEntryPoint();
-				SlotRequest.SubGraph->AllocateInstance(Binding.GetTraitPtr().GetNodeInstance()->GetOwner(), SlotRequest.GraphInstance, EntryPoint);
+				const FName EntryPoint = SlotRequest.Module->GetDefaultEntryPoint();
+				SlotRequest.Module->AllocateInstance(Binding.GetTraitPtr().GetNodeInstance()->GetOwner(), SlotRequest.GraphInstance, EntryPoint);
 				SlotRequest.ChildPtr = SlotRequest.GraphInstance.GetGraphRootPtr();
 
 				// Setup our graph parameters
 				const FAnimNextPlayAnimRequestArgs& RequestArgs = SlotRequest.Request->GetArgs();
 
-				const FString SubGraphPath = SlotRequest.SubGraph->GetPathName();
+				const FString SubGraphPath = SlotRequest.Module->GetPathName();
 
 				const FName AnimationObjectName(FString::Printf(TEXT("%s:AnimationObject"), *SubGraphPath));
 				const FName StartPositionName(FString::Printf(TEXT("%s:StartPosition"), *SubGraphPath));
@@ -559,7 +559,7 @@ namespace UE::AnimNext
 				SlotRequest.Request->AddReferencedObjects(Collector);
 			}
 
-			Collector.AddReferencedObject(SlotRequest.SubGraph);
+			Collector.AddReferencedObject(SlotRequest.Module);
 		}
 	}
 }
diff --git a/Engine/Plugins/Experimental/AnimNextAnimGraph/Source/AnimNextAnimGraph/Private/PlayAnim/PlayAnimSlotTrait.h b/Engine/Plugins/Experimental/AnimNextAnimGraph/Source/AnimNextAnimGraph/Private/PlayAnim/PlayAnimSlotTrait.h
--- a/Engine/Plugins/Experimental/AnimNextAnimGraph/Source/AnimNextAnimGraph/Private/PlayAnim/PlayAnimSlotTrait.h
+++ b/Engine/Plugins/Experimental/AnimNextAnimGraph/Source/AnimNextAnimGraph/Private/PlayAnim/PlayAnimSlotTrait.h
@@ -11,7 +11,7 @@
 #include "TraitInterfaces/IInertializerBlend.h"
 #include "TraitInterfaces/ISmoothBlend.h"
 #include "TraitInterfaces/IUpdate.h"
-#include "Graph/AnimNextGraph.h"
+#include "Module/AnimNextModule.h"
 #include "Graph/AnimNextGraphInstancePtr.h"
 #include "PlayAnim/PlayAnimEvents.h"
 #include "Param/PropertyBagProxy.h"
@@ -108,8 +108,8 @@ namespace UE::AnimNext
 			// The blend settings to use
 			FPlayAnimBlendSettings BlendSettings;
 
-			// The sub-graph used by the graph instance, as selected by the chooser
-			TObjectPtr<const UAnimNextGraph> SubGraph;
+			// The module used by the graph instance, as selected by the chooser
+			TObjectPtr<const UAnimNextModule> Module;
 
 			// The graph instance
 			FAnimNextGraphInstancePtr GraphInstance;
@@ -125,7 +125,7 @@ namespace UE::AnimNext
 			bool bWasRelevant = false;
 
 			// Initializes a request to begin playing
-			void Initialize(FPlayAnimRequestPtr InRequest, const FPlayAnimBlendSettings& InBlendSettings, const UAnimNextGraph* InSubGraph);
+			void Initialize(FPlayAnimRequestPtr InRequest, const FPlayAnimBlendSettings& InBlendSettings, const UAnimNextModule* InModule);
 		};
 
 		using FSharedData = FAnimNextPlayAnimSlotTraitSharedData;
diff --git a/Engine/Plugins/Experimental/AnimNextAnimGraph/Source/AnimNextAnimGraph/Private/Traits/SubGraphHost.cpp b/Engine/Plugins/Experimental/AnimNextAnimGraph/Source/AnimNextAnimGraph/Private/Traits/SubGraphHost.cpp
--- a/Engine/Plugins/Experimental/AnimNextAnimGraph/Source/AnimNextAnimGraph/Private/Traits/SubGraphHost.cpp
+++ b/Engine/Plugins/Experimental/AnimNextAnimGraph/Source/AnimNextAnimGraph/Private/Traits/SubGraphHost.cpp
@@ -69,25 +69,25 @@ namespace UE::AnimNext
 
 		const bool bHasActiveSubGraph = InstanceData->CurrentlyActiveSubGraphIndex != INDEX_NONE;
 
-		TObjectPtr<const UAnimNextGraph> CurrentActiveSubGraph;
+		TObjectPtr<const UAnimNextModule> CurrentActiveModule;
 		FName CurrentActiveEntryPoint = NAME_None;
 		if (bHasActiveSubGraph)
 		{
 			FSubGraphSlot& SubGraphSlot = InstanceData->SubGraphSlots[InstanceData->CurrentlyActiveSubGraphIndex];
-			CurrentActiveSubGraph = SubGraphSlot.SubGraph;
+			CurrentActiveModule = SubGraphSlot.Module;
 			CurrentActiveEntryPoint = SubGraphSlot.EntryPoint;
 
 			SubGraphSlot.bWasRelevant = true;
 		}
 
-		const TObjectPtr<const UAnimNextGraph> DesiredSubGraph = SharedData->GetSubGraph(Binding);
+		const TObjectPtr<const UAnimNextModule> DesiredModule = SharedData->GetModule(Binding);
 		const FName EntryPoint = SharedData->GetEntryPoint(Binding);
 
 		// Check for re-entrancy and early-out if we are linking back to the current instance or one of its parents
 		const FAnimNextGraphInstance* OwnerGraphInstance = &Binding.GetTraitPtr().GetNodeInstance()->GetOwner();
 		while (OwnerGraphInstance != nullptr)
 		{
-			if (OwnerGraphInstance->UsesGraph(DesiredSubGraph) && OwnerGraphInstance->UsesEntryPoint(EntryPoint))
+			if (OwnerGraphInstance->UsesModule(DesiredModule) && OwnerGraphInstance->UsesEntryPoint(EntryPoint))
 			{
 				return;
 			}
@@ -95,7 +95,7 @@ namespace UE::AnimNext
 			OwnerGraphInstance = OwnerGraphInstance->GetParentGraphInstance();
 		}
 
-		if (!bHasActiveSubGraph || CurrentActiveSubGraph != DesiredSubGraph || CurrentActiveEntryPoint != EntryPoint)
+		if (!bHasActiveSubGraph || CurrentActiveModule != DesiredModule || CurrentActiveEntryPoint != EntryPoint)
 		{
 			// Find an empty slot we can use
 			int32 FreeSlotIndex = INDEX_NONE;
@@ -118,8 +118,8 @@ namespace UE::AnimNext
 			}
 
 			FSubGraphSlot& SubGraphSlot = InstanceData->SubGraphSlots[FreeSlotIndex];
-			SubGraphSlot.SubGraph = DesiredSubGraph;
-			SubGraphSlot.State = DesiredSubGraph ? ESlotState::ActiveWithGraph : ESlotState::ActiveWithReferencePose;
+			SubGraphSlot.Module = DesiredModule;
+			SubGraphSlot.State = DesiredModule ? ESlotState::ActiveWithGraph : ESlotState::ActiveWithReferencePose;
 			SubGraphSlot.EntryPoint = EntryPoint;
 
 			const int32 OldChildIndex = InstanceData->CurrentlyActiveSubGraphIndex;
@@ -210,7 +210,7 @@ namespace UE::AnimNext
 
 			if (SubGraphEntry.State == ESlotState::ActiveWithGraph)
 			{
-				SubGraphEntry.SubGraph->AllocateInstance(Binding.GetTraitPtr().GetNodeInstance()->GetOwner(), SubGraphEntry.GraphInstance, SubGraphEntry.EntryPoint);
+				SubGraphEntry.Module->AllocateInstance(Binding.GetTraitPtr().GetNodeInstance()->GetOwner(), SubGraphEntry.GraphInstance, SubGraphEntry.EntryPoint);
 			}
 		}
 	}
diff --git a/Engine/Plugins/Experimental/AnimNextAnimGraph/Source/AnimNextAnimGraph/Private/Traits/SubGraphHost.h b/Engine/Plugins/Experimental/AnimNextAnimGraph/Source/AnimNextAnimGraph/Private/Traits/SubGraphHost.h
--- a/Engine/Plugins/Experimental/AnimNextAnimGraph/Source/AnimNextAnimGraph/Private/Traits/SubGraphHost.h
+++ b/Engine/Plugins/Experimental/AnimNextAnimGraph/Source/AnimNextAnimGraph/Private/Traits/SubGraphHost.h
@@ -9,7 +9,7 @@
 #include "TraitInterfaces/IGarbageCollection.h"
 #include "TraitInterfaces/IHierarchy.h"
 #include "TraitInterfaces/IUpdate.h"
-#include "Graph/AnimNextGraph.h"
+#include "Module/AnimNextModule.h"
 #include "Graph/AnimNextGraphInstancePtr.h"
 
 #include "SubGraphHost.generated.h"
@@ -19,9 +19,9 @@ struct FAnimNextSubGraphHostTraitSharedData : public FAnimNextTraitSharedData
 {
 	GENERATED_BODY()
 
-	/** Subgraph that we host and manage. */
+	/** Module of the sub-graph that we host and manage. */
 	UPROPERTY(EditAnywhere, Category = "Default")
-	TObjectPtr<const UAnimNextGraph> SubGraph;
+	TObjectPtr<const UAnimNextModule> Module;
 
 	/** A dummy child that we can use to output the bind pose. This property is hidden and automatically populated during compilation. */
 	UPROPERTY(meta = (Hidden))
@@ -33,7 +33,7 @@ struct FAnimNextSubGraphHostTraitSharedData : public FAnimNextTraitSharedData
 
 	// Latent pin support boilerplate
 	#define TRAIT_LATENT_PROPERTIES_ENUMERATOR(GeneratorMacro) \
-		GeneratorMacro(SubGraph) \
+		GeneratorMacro(Module) \
 		GeneratorMacro(EntryPoint) \
 
 	GENERATE_TRAIT_LATENT_PROPERTIES(FAnimNextSubGraphHostTraitSharedData, TRAIT_LATENT_PROPERTIES_ENUMERATOR)
@@ -60,8 +60,8 @@ namespace UE::AnimNext
 
 		struct FSubGraphSlot
 		{
-			// The sub-graph to use
-			TObjectPtr<const UAnimNextGraph> SubGraph;
+			// The module to use
+			TObjectPtr<const UAnimNextModule> Module;
 
 			// The graph instance
 			FAnimNextGraphInstancePtr GraphInstance;
diff --git a/Engine/Plugins/Experimental/AnimNextAnimGraph/Source/AnimNextAnimGraphTestSuite/Private/AnimNextAnimGraphTraitEventTest.cpp b/Engine/Plugins/Experimental/AnimNextAnimGraph/Source/AnimNextAnimGraphTestSuite/Private/AnimNextAnimGraphTraitEventTest.cpp
--- a/Engine/Plugins/Experimental/AnimNextAnimGraph/Source/AnimNextAnimGraphTestSuite/Private/AnimNextAnimGraphTraitEventTest.cpp
+++ b/Engine/Plugins/Experimental/AnimNextAnimGraph/Source/AnimNextAnimGraphTestSuite/Private/AnimNextAnimGraphTraitEventTest.cpp
@@ -14,8 +14,8 @@
 #include "TraitCore/NodeTemplateRegistry.h"
 #include "TraitInterfaces/IHierarchy.h"
 #include "TraitInterfaces/IUpdate.h"
-#include "Graph/AnimNextGraph.h"
-#include "Graph/GraphFactory.h"
+#include "Module/AnimNextModule.h"
+#include "Module/ModuleFactory.h"
 #include "Misc/AutomationTest.h"
 
 #if WITH_DEV_AUTOMATION_TESTS
@@ -194,9 +194,9 @@ bool FAnimationAnimNextRuntimeTest_GraphTraitEvent::RunTest(const FString& InPar
 		AUTO_REGISTER_ANIM_TRAIT(FTraitGraphTest_EventAB_OneChild)
 		AUTO_REGISTER_ANIM_TRAIT(FTraitGraphTest_EventAB_TwoChildren)
 
-		UFactory* GraphFactory = NewObject<UAnimNextGraphFactory>();
-		UAnimNextGraph* AnimNextGraph = CastChecked<UAnimNextGraph>(GraphFactory->FactoryCreateNew(UAnimNextGraph::StaticClass(), GetTransientPackage(), TEXT("TestAnimNextGraph"), RF_Transient, nullptr, nullptr, NAME_None));
-		UE_RETURN_ON_ERROR(AnimNextGraph != nullptr, "FAnimationAnimNextRuntimeTest_GraphTraitEvent -> Failed to create animation graph");
+		UFactory* GraphFactory = NewObject<UAnimNextModuleFactory>();
+		UAnimNextModule* Module = CastChecked<UAnimNextModule>(GraphFactory->FactoryCreateNew(UAnimNextModule::StaticClass(), GetTransientPackage(), TEXT("TestAnimNextGraph"), RF_Transient, nullptr, nullptr, NAME_None));
+		UE_RETURN_ON_ERROR(Module != nullptr, "FAnimationAnimNextRuntimeTest_GraphTraitEvent -> Failed to create module");
 
 		FScopedClearNodeTemplateRegistry ScopedClearNodeTemplateRegistry;
 		FNodeTemplateRegistry& Registry = FNodeTemplateRegistry::Get();
@@ -296,10 +296,10 @@ bool FAnimationAnimNextRuntimeTest_GraphTraitEvent::RunTest(const FString& InPar
 		}
 
 		// Read our graph
-		FTestUtils::LoadFromArchiveBuffer(*AnimNextGraph, NodeHandles, GraphSharedDataArchiveBuffer);
+		FTestUtils::LoadFromArchiveBuffer(*Module, NodeHandles, GraphSharedDataArchiveBuffer);
 
 		FAnimNextGraphInstancePtr GraphInstance;
-		AnimNextGraph->AllocateInstance(GraphInstance);
+		Module->AllocateInstance(GraphInstance);
 
 		{
 			// Raise EventA and EventB on graph, every node sees them
diff --git a/Engine/Plugins/Experimental/AnimNextAnimGraph/Source/AnimNextAnimGraphTestSuite/Private/AnimNextAnimGraphTraitGraphTest.cpp b/Engine/Plugins/Experimental/AnimNextAnimGraph/Source/AnimNextAnimGraphTestSuite/Private/AnimNextAnimGraphTraitGraphTest.cpp
--- a/Engine/Plugins/Experimental/AnimNextAnimGraph/Source/AnimNextAnimGraphTestSuite/Private/AnimNextAnimGraphTraitGraphTest.cpp
+++ b/Engine/Plugins/Experimental/AnimNextAnimGraph/Source/AnimNextAnimGraphTestSuite/Private/AnimNextAnimGraphTraitGraphTest.cpp
@@ -12,9 +12,9 @@
 #include "TraitInterfaces/IUpdate.h"
 #include "Editor/Transactor.h"
 #include "AnimNextExecuteContext.h"
-#include "Graph/AnimNextGraph.h"
-#include "Graph/AnimNextGraph_EditorData.h"
-#include "Graph/GraphFactory.h"
+#include "Module/AnimNextModule.h"
+#include "Module/AnimNextModule_EditorData.h"
+#include "Module/ModuleFactory.h"
 #include "Graph/RigDecorator_AnimNextCppTrait.h"
 #include "Graph/RigUnit_AnimNextGraphRoot.h"
 #include "Graph/RigUnit_AnimNextTraitStack.h"
@@ -102,12 +102,12 @@ bool FAnimationAnimNextEditorTest_GraphAddTrait::RunTest(const FString& InParame
 
 		FScopedClearNodeTemplateRegistry ScopedClearNodeTemplateRegistry;
 
-		UFactory* GraphFactory = NewObject<UAnimNextGraphFactory>();
-		UAnimNextGraph* AnimNextGraph = CastChecked<UAnimNextGraph>(GraphFactory->FactoryCreateNew(UAnimNextGraph::StaticClass(), GetTransientPackage(), TEXT("TestAnimNextGraph"), RF_Transient, nullptr, nullptr, NAME_None));
-		UE_RETURN_ON_ERROR(AnimNextGraph != nullptr, "FAnimationAnimNextEditorTest_GraphAddTrait -> Failed to create animation graph");
+		UFactory* GraphFactory = NewObject<UAnimNextModuleFactory>();
+		UAnimNextModule* Module = CastChecked<UAnimNextModule>(GraphFactory->FactoryCreateNew(UAnimNextModule::StaticClass(), GetTransientPackage(), TEXT("TestAnimNextGraph"), RF_Transient, nullptr, nullptr, NAME_None));
+		UE_RETURN_ON_ERROR(Module != nullptr, "FAnimationAnimNextEditorTest_GraphAddTrait -> Failed to create module");
 
-		UAnimNextGraph_EditorData* EditorData = UncookedOnly::FUtils::GetEditorData(AnimNextGraph);
-		UE_RETURN_ON_ERROR(EditorData != nullptr, "FAnimationAnimNextEditorTest_GraphAddTrait -> Failed to find animation graph editor data");
+		UAnimNextModule_EditorData* EditorData = UncookedOnly::FUtils::GetEditorData(Module);
+		UE_RETURN_ON_ERROR(EditorData != nullptr, "FAnimationAnimNextEditorTest_GraphAddTrait -> Failed to find module editor data");
 
 		UE_RETURN_ON_ERROR(EditorData->AddAnimationGraph(FRigUnit_AnimNextGraphRoot::DefaultEntryPoint, false) != nullptr, "FAnimationAnimNextEditorTest_GraphAddTrait -> Failed to add animation graph");
 
@@ -208,16 +208,16 @@ bool FAnimationAnimNextEditorTest_GraphTraitOperations::RunTest(const FString& I
 
 		FScopedClearNodeTemplateRegistry ScopedClearNodeTemplateRegistry;
 
-		UFactory* GraphFactory = NewObject<UAnimNextGraphFactory>();
-		UAnimNextGraph* AnimNextGraph = CastChecked<UAnimNextGraph>(GraphFactory->FactoryCreateNew(UAnimNextGraph::StaticClass(), GetTransientPackage(), TEXT("TestAnimNextGraph"), RF_Transient, nullptr, nullptr, NAME_None));
-		UE_RETURN_ON_ERROR(AnimNextGraph != nullptr, "FAnimationAnimNextEditorTest_GraphTraitOperations -> Failed to create animation graph");
+		UFactory* GraphFactory = NewObject<UAnimNextModuleFactory>();
+		UAnimNextModule* Module = CastChecked<UAnimNextModule>(GraphFactory->FactoryCreateNew(UAnimNextModule::StaticClass(), GetTransientPackage(), TEXT("TestAnimNextGraph"), RF_Transient, nullptr, nullptr, NAME_None));
+		UE_RETURN_ON_ERROR(Module != nullptr, "FAnimationAnimNextEditorTest_GraphTraitOperations -> Failed to create module");
 
-		UAnimNextGraph_EditorData* EditorData = UncookedOnly::FUtils::GetEditorData(AnimNextGraph);
-		UE_RETURN_ON_ERROR(EditorData != nullptr, "FAnimationAnimNextEditorTest_GraphTraitOperations -> Failed to find animation graph editor data");
+		UAnimNextModule_EditorData* EditorData = UncookedOnly::FUtils::GetEditorData(Module);
+		UE_RETURN_ON_ERROR(EditorData != nullptr, "FAnimationAnimNextEditorTest_GraphTraitOperations -> Failed to find module editor data");
 
 		UE_RETURN_ON_ERROR(EditorData->AddAnimationGraph(FRigUnit_AnimNextGraphRoot::DefaultEntryPoint, false) != nullptr, "FAnimationAnimNextEditorTest_GraphTraitOperations -> Failed to add animation graph");
 
-		UAnimNextGraph_Controller* Controller = Cast<UAnimNextGraph_Controller>(EditorData->GetRigVMClient()->GetController(EditorData->GetRigVMClient()->GetDefaultModel()));
+		UAnimNextModule_Controller* Controller = Cast<UAnimNextModule_Controller>(EditorData->GetRigVMClient()->GetController(EditorData->GetRigVMClient()->GetDefaultModel()));
 		UE_RETURN_ON_ERROR(Controller != nullptr, "FAnimationAnimNextEditorTest_GraphTraitOperations -> Failed to get RigVM controller");
 
 		// Create an empty trait stack node
@@ -370,12 +370,12 @@ bool FAnimationAnimNextRuntimeTest_GraphExecute::RunTest(const FString& InParame
 
 		FScopedClearNodeTemplateRegistry ScopedClearNodeTemplateRegistry;
 
-		UFactory* GraphFactory = NewObject<UAnimNextGraphFactory>();
-		UAnimNextGraph* AnimNextGraph = CastChecked<UAnimNextGraph>(GraphFactory->FactoryCreateNew(UAnimNextGraph::StaticClass(), GetTransientPackage(), TEXT("TestAnimNextGraph"), RF_Transient, nullptr, nullptr, NAME_None));
-		UE_RETURN_ON_ERROR(AnimNextGraph != nullptr, "FAnimationAnimNextRuntimeTest_GraphExecute -> Failed to create animation graph");
+		UFactory* GraphFactory = NewObject<UAnimNextModuleFactory>();
+		UAnimNextModule* Module = CastChecked<UAnimNextModule>(GraphFactory->FactoryCreateNew(UAnimNextModule::StaticClass(), GetTransientPackage(), TEXT("TestAnimNextGraph"), RF_Transient, nullptr, nullptr, NAME_None));
+		UE_RETURN_ON_ERROR(Module != nullptr, "FAnimationAnimNextRuntimeTest_GraphExecute -> Failed to create module");
 
-		UAnimNextGraph_EditorData* EditorData = UncookedOnly::FUtils::GetEditorData(AnimNextGraph);
-		UE_RETURN_ON_ERROR(EditorData != nullptr, "FAnimationAnimNextRuntimeTest_GraphExecute -> Failed to find animation graph editor data");
+		UAnimNextModule_EditorData* EditorData = UncookedOnly::FUtils::GetEditorData(Module);
+		UE_RETURN_ON_ERROR(EditorData != nullptr, "FAnimationAnimNextRuntimeTest_GraphExecute -> Failed to find module editor data");
 
 		UE_RETURN_ON_ERROR(EditorData->AddAnimationGraph(FRigUnit_AnimNextGraphRoot::DefaultEntryPoint, false) != nullptr, "FAnimationAnimNextRuntimeTest_GraphExecute -> Failed to add animation graph");
 
@@ -442,7 +442,7 @@ bool FAnimationAnimNextRuntimeTest_GraphExecute::RunTest(const FString& InParame
 		FParamStack::AttachToCurrentThread(ParamStack);
 
 		FAnimNextGraphInstancePtr GraphInstance;
-		AnimNextGraph->AllocateInstance(GraphInstance);
+		Module->AllocateInstance(GraphInstance);
 
 		FParamStack::FPushedLayerHandle LayerHandle = ParamStack->PushValues(
 			"/Engine/Transient.TestAnimNextGraph:UpdateCount", (int32)0,
@@ -490,16 +490,16 @@ bool FAnimationAnimNextRuntimeTest_GraphExecuteLatent::RunTest(const FString& In
 		AUTO_REGISTER_ANIM_TRAIT(FTestTrait)
 		FScopedClearNodeTemplateRegistry ScopedClearNodeTemplateRegistry;
 
-		UFactory* GraphFactory = NewObject<UAnimNextGraphFactory>();
-		UAnimNextGraph* AnimNextGraph = CastChecked<UAnimNextGraph>(GraphFactory->FactoryCreateNew(UAnimNextGraph::StaticClass(), GetTransientPackage(), TEXT("TestAnimNextGraph"), RF_Transient, nullptr, nullptr, NAME_None));
-		UE_RETURN_ON_ERROR(AnimNextGraph != nullptr, "FAnimationAnimNextRuntimeTest_GraphExecuteLatent -> Failed to create animation graph");
+		UFactory* GraphFactory = NewObject<UAnimNextModuleFactory>();
+		UAnimNextModule* Module = CastChecked<UAnimNextModule>(GraphFactory->FactoryCreateNew(UAnimNextModule::StaticClass(), GetTransientPackage(), TEXT("TestAnimNextGraph"), RF_Transient, nullptr, nullptr, NAME_None));
+		UE_RETURN_ON_ERROR(Module != nullptr, "FAnimationAnimNextRuntimeTest_GraphExecuteLatent -> Failed to create module");
 
-		UAnimNextGraph_EditorData* EditorData = UncookedOnly::FUtils::GetEditorData(AnimNextGraph);
-		UE_RETURN_ON_ERROR(EditorData != nullptr, "FAnimationAnimNextRuntimeTest_GraphExecuteLatent -> Failed to find animation graph editor data");
+		UAnimNextModule_EditorData* EditorData = UncookedOnly::FUtils::GetEditorData(Module);
+		UE_RETURN_ON_ERROR(EditorData != nullptr, "FAnimationAnimNextRuntimeTest_GraphExecuteLatent -> Failed to find module editor data");
 
 		UE_RETURN_ON_ERROR(EditorData->AddAnimationGraph(FRigUnit_AnimNextGraphRoot::DefaultEntryPoint, false) != nullptr, "FAnimationAnimNextRuntimeTest_GraphExecuteLatent -> Failed to add animation graph");
 
-		UAnimNextGraph_Controller* Controller = Cast<UAnimNextGraph_Controller>(EditorData->GetRigVMClient()->GetController(EditorData->GetRigVMClient()->GetDefaultModel()));
+		UAnimNextModule_Controller* Controller = Cast<UAnimNextModule_Controller>(EditorData->GetRigVMClient()->GetController(EditorData->GetRigVMClient()->GetDefaultModel()));
 		UE_RETURN_ON_ERROR(Controller != nullptr, "FAnimationAnimNextRuntimeTest_GraphExecuteLatent -> Failed to get RigVM controller");
 
 		// Find graph entry point
@@ -574,7 +574,7 @@ bool FAnimationAnimNextRuntimeTest_GraphExecuteLatent::RunTest(const FString& In
 		}
 
 		{
-			URigVMNode* GetParameterNode = Controller->AddGetAnimNextParameterNode(FVector2D::ZeroVector, UncookedOnly::FUtils::GetQualifiedName(AnimNextGraph, "SomeSourceInt"), FAnimNextParamType::GetType<int32>());
+			URigVMNode* GetParameterNode = Controller->AddGetAnimNextParameterNode(FVector2D::ZeroVector, UncookedOnly::FUtils::GetQualifiedName(Module, "SomeSourceInt"), FAnimNextParamType::GetType<int32>());
 			UE_RETURN_ON_ERROR(GetParameterNode != nullptr, TEXT("FAnimationAnimNextRuntimeTest_GraphExecuteLatent -> Failed to create GetParameter node"));
 
 			Controller->AddLink(
@@ -586,17 +586,17 @@ bool FAnimationAnimNextRuntimeTest_GraphExecuteLatent::RunTest(const FString& In
 		FParamStack::AttachToCurrentThread(ParamStack);
 
 		FAnimNextGraphInstancePtr GraphInstance;
-		AnimNextGraph->AllocateInstance(GraphInstance);
+		Module->AllocateInstance(GraphInstance);
 
 		FParamStack::FPushedLayerHandle LayerHandle = ParamStack->PushValues(
-			UncookedOnly::FUtils::GetQualifiedName(AnimNextGraph, "UpdateCount"), (int32)0,
-			UncookedOnly::FUtils::GetQualifiedName(AnimNextGraph, "EvaluateCount"), (int32)0,
-			UncookedOnly::FUtils::GetQualifiedName(AnimNextGraph, "SomeSourceInt"), (int32)1223,
-			UncookedOnly::FUtils::GetQualifiedName(AnimNextGraph, "SomeInt32"), (int32)0,
-			UncookedOnly::FUtils::GetQualifiedName(AnimNextGraph, "SomeFloat"), 0.0f,
-			UncookedOnly::FUtils::GetQualifiedName(AnimNextGraph, "SomeLatentInt32"), (int32)0,
-			UncookedOnly::FUtils::GetQualifiedName(AnimNextGraph, "SomeOtherLatentInt32"), (int32)0,
-			UncookedOnly::FUtils::GetQualifiedName(AnimNextGraph, "SomeLatentFloat"), 0.0f
+			UncookedOnly::FUtils::GetQualifiedName(Module, "UpdateCount"), (int32)0,
+			UncookedOnly::FUtils::GetQualifiedName(Module, "EvaluateCount"), (int32)0,
+			UncookedOnly::FUtils::GetQualifiedName(Module, "SomeSourceInt"), (int32)1223,
+			UncookedOnly::FUtils::GetQualifiedName(Module, "SomeInt32"), (int32)0,
+			UncookedOnly::FUtils::GetQualifiedName(Module, "SomeFloat"), 0.0f,
+			UncookedOnly::FUtils::GetQualifiedName(Module, "SomeLatentInt32"), (int32)0,
+			UncookedOnly::FUtils::GetQualifiedName(Module, "SomeOtherLatentInt32"), (int32)0,
+			UncookedOnly::FUtils::GetQualifiedName(Module, "SomeLatentFloat"), 0.0f
 		);
 
 		{
@@ -606,13 +606,13 @@ bool FAnimationAnimNextRuntimeTest_GraphExecuteLatent::RunTest(const FString& In
 			(void)UE::AnimNext::EvaluateGraph(GraphInstance);
 		}
 
-		AddErrorIfFalse(ParamStack->GetParam<int32>(UncookedOnly::FUtils::GetQualifiedName(AnimNextGraph, "UpdateCount")) == 1, "FAnimationAnimNextRuntimeTest_GraphExecuteLatent -> Unexpected update count");
-		AddErrorIfFalse(ParamStack->GetParam<int32>(UncookedOnly::FUtils::GetQualifiedName(AnimNextGraph, "EvaluateCount")) == 1, "FAnimationAnimNextRuntimeTest_GraphExecuteLatent -> Unexpected evaluate count");
-		AddErrorIfFalse(ParamStack->GetParam<int32>(UncookedOnly::FUtils::GetQualifiedName(AnimNextGraph, "SomeInt32")) == 78, "FAnimationAnimNextRuntimeTest_GraphExecuteLatent -> Unexpected SomeInt32 value");
-		AddErrorIfFalse(ParamStack->GetParam<float>(UncookedOnly::FUtils::GetQualifiedName(AnimNextGraph, "SomeFloat")) == 142.33f, "FAnimationAnimNextRuntimeTest_GraphExecuteLatent -> Unexpected SomeFloat value");
-		AddErrorIfFalse(ParamStack->GetParam<int32>(UncookedOnly::FUtils::GetQualifiedName(AnimNextGraph, "SomeLatentInt32")) == 33, "FAnimationAnimNextRuntimeTest_GraphExecuteLatent -> Unexpected SomeLatentInt32 value");
-		AddErrorIfFalse(ParamStack->GetParam<int32>(UncookedOnly::FUtils::GetQualifiedName(AnimNextGraph, "SomeOtherLatentInt32")) == 1223, "FAnimationAnimNextRuntimeTest_GraphExecuteLatent -> Unexpected SomeOtherLatentInt32 value");
-		AddErrorIfFalse(ParamStack->GetParam<float>(UncookedOnly::FUtils::GetQualifiedName(AnimNextGraph, "SomeLatentFloat")) == 1123.31f, "FAnimationAnimNextRuntimeTest_GraphExecuteLatent -> Unexpected SomeLatentFloat value");
+		AddErrorIfFalse(ParamStack->GetParam<int32>(UncookedOnly::FUtils::GetQualifiedName(Module, "UpdateCount")) == 1, "FAnimationAnimNextRuntimeTest_GraphExecuteLatent -> Unexpected update count");
+		AddErrorIfFalse(ParamStack->GetParam<int32>(UncookedOnly::FUtils::GetQualifiedName(Module, "EvaluateCount")) == 1, "FAnimationAnimNextRuntimeTest_GraphExecuteLatent -> Unexpected evaluate count");
+		AddErrorIfFalse(ParamStack->GetParam<int32>(UncookedOnly::FUtils::GetQualifiedName(Module, "SomeInt32")) == 78, "FAnimationAnimNextRuntimeTest_GraphExecuteLatent -> Unexpected SomeInt32 value");
+		AddErrorIfFalse(ParamStack->GetParam<float>(UncookedOnly::FUtils::GetQualifiedName(Module, "SomeFloat")) == 142.33f, "FAnimationAnimNextRuntimeTest_GraphExecuteLatent -> Unexpected SomeFloat value");
+		AddErrorIfFalse(ParamStack->GetParam<int32>(UncookedOnly::FUtils::GetQualifiedName(Module, "SomeLatentInt32")) == 33, "FAnimationAnimNextRuntimeTest_GraphExecuteLatent -> Unexpected SomeLatentInt32 value");
+		AddErrorIfFalse(ParamStack->GetParam<int32>(UncookedOnly::FUtils::GetQualifiedName(Module, "SomeOtherLatentInt32")) == 1223, "FAnimationAnimNextRuntimeTest_GraphExecuteLatent -> Unexpected SomeOtherLatentInt32 value");
+		AddErrorIfFalse(ParamStack->GetParam<float>(UncookedOnly::FUtils::GetQualifiedName(Module, "SomeLatentFloat")) == 1123.31f, "FAnimationAnimNextRuntimeTest_GraphExecuteLatent -> Unexpected SomeLatentFloat value");
 
 		ParamStack->PopLayer(LayerHandle);
 		GraphInstance.Release();
diff --git a/Engine/Plugins/Experimental/AnimNextAnimGraph/Source/AnimNextAnimGraphTestSuite/Private/AnimNextAnimGraphTraitGraphTest.h b/Engine/Plugins/Experimental/AnimNextAnimGraph/Source/AnimNextAnimGraphTestSuite/Private/AnimNextAnimGraphTraitGraphTest.h
--- a/Engine/Plugins/Experimental/AnimNextAnimGraph/Source/AnimNextAnimGraphTestSuite/Private/AnimNextAnimGraphTraitGraphTest.h
+++ b/Engine/Plugins/Experimental/AnimNextAnimGraph/Source/AnimNextAnimGraphTestSuite/Private/AnimNextAnimGraphTraitGraphTest.h
@@ -6,7 +6,7 @@
 
 #include "TraitCore/TraitBinding.h"
 #include "TraitCore/TraitSharedData.h"
-#include "Graph/AnimNextGraph.h"
+#include "Module/AnimNextModule.h"
 
 #include "AnimNextAnimGraphTraitGraphTest.generated.h"
 
diff --git a/Engine/Plugins/Experimental/AnimNextAnimGraph/Source/AnimNextAnimGraphTestSuite/Private/AnimNextAnimGraphTraitInterfacesTest.cpp b/Engine/Plugins/Experimental/AnimNextAnimGraph/Source/AnimNextAnimGraphTestSuite/Private/AnimNextAnimGraphTraitInterfacesTest.cpp
--- a/Engine/Plugins/Experimental/AnimNextAnimGraph/Source/AnimNextAnimGraphTestSuite/Private/AnimNextAnimGraphTraitInterfacesTest.cpp
+++ b/Engine/Plugins/Experimental/AnimNextAnimGraph/Source/AnimNextAnimGraphTestSuite/Private/AnimNextAnimGraphTraitInterfacesTest.cpp
@@ -24,8 +24,8 @@
 #include "TraitInterfaces/IEvaluate.h"
 #include "TraitInterfaces/IHierarchy.h"
 #include "TraitInterfaces/IUpdate.h"
-#include "Graph/AnimNextGraph.h"
-#include "Graph/GraphFactory.h"
+#include "Module/AnimNextModule.h"
+#include "Module/ModuleFactory.h"
 
 //****************************************************************************
 // AnimNext Runtime TraitInterfaces Tests
@@ -477,9 +477,9 @@ bool FAnimationAnimNextRuntimeTest_IHierarchy::RunTest(const FString& InParamete
 		AUTO_REGISTER_ANIM_TRAIT(FTraitWithOneChild)
 		AUTO_REGISTER_ANIM_TRAIT(FTraitWithChildren)
 
-		UFactory* GraphFactory = NewObject<UAnimNextGraphFactory>();
-		UAnimNextGraph* AnimNextGraph = CastChecked<UAnimNextGraph>(GraphFactory->FactoryCreateNew(UAnimNextGraph::StaticClass(), GetTransientPackage(), TEXT("TestAnimNextGraph"), RF_Transient, nullptr, nullptr, NAME_None));
-		UE_RETURN_ON_ERROR(AnimNextGraph != nullptr, "FAnimationAnimNextRuntimeTest_IHierarchy -> Failed to create animation graph");
+		UFactory* ModuleFactory = NewObject<UAnimNextModuleFactory>();
+		UAnimNextModule* Module = CastChecked<UAnimNextModule>(ModuleFactory->FactoryCreateNew(UAnimNextModule::StaticClass(), GetTransientPackage(), TEXT("TestAnimNextGraph"), RF_Transient, nullptr, nullptr, NAME_None));
+		UE_RETURN_ON_ERROR(Module != nullptr, "FAnimationAnimNextRuntimeTest_IHierarchy -> Failed to create module");
 
 		FScopedClearNodeTemplateRegistry ScopedClearNodeTemplateRegistry;
 		FNodeTemplateRegistry& Registry = FNodeTemplateRegistry::Get();
@@ -591,10 +591,10 @@ bool FAnimationAnimNextRuntimeTest_IHierarchy::RunTest(const FString& InParamete
 		}
 
 		// Read our graph
-		FTestUtils::LoadFromArchiveBuffer(*AnimNextGraph, NodeHandles, GraphSharedDataArchiveBuffer);
+		FTestUtils::LoadFromArchiveBuffer(*Module, NodeHandles, GraphSharedDataArchiveBuffer);
 
 		FAnimNextGraphInstancePtr GraphInstance;
-		AnimNextGraph->AllocateInstance(GraphInstance);
+		Module->AllocateInstance(GraphInstance);
 
 		FExecutionContext Context(GraphInstance);
 
@@ -698,9 +698,9 @@ bool FAnimationAnimNextRuntimeTest_IUpdate::RunTest(const FString& InParameters)
 		AUTO_REGISTER_ANIM_TRAIT(FTraitWithOneChild)
 		AUTO_REGISTER_ANIM_TRAIT(FTraitWithChildren)
 
-		UFactory* GraphFactory = NewObject<UAnimNextGraphFactory>();
-		UAnimNextGraph* AnimNextGraph = CastChecked<UAnimNextGraph>(GraphFactory->FactoryCreateNew(UAnimNextGraph::StaticClass(), GetTransientPackage(), TEXT("TestAnimNextGraph"), RF_Transient, nullptr, nullptr, NAME_None));
-		UE_RETURN_ON_ERROR(AnimNextGraph != nullptr, "FAnimationAnimNextRuntimeTest_IUpdate -> Failed to create animation graph");
+		UFactory* ModuleFactory = NewObject<UAnimNextModuleFactory>();
+		UAnimNextModule* Module = CastChecked<UAnimNextModule>(ModuleFactory->FactoryCreateNew(UAnimNextModule::StaticClass(), GetTransientPackage(), TEXT("TestAnimNextGraph"), RF_Transient, nullptr, nullptr, NAME_None));
+		UE_RETURN_ON_ERROR(Module != nullptr, "FAnimationAnimNextRuntimeTest_IUpdate -> Failed to create module");
 
 		FScopedClearNodeTemplateRegistry ScopedClearNodeTemplateRegistry;
 		FNodeTemplateRegistry& Registry = FNodeTemplateRegistry::Get();
@@ -790,10 +790,10 @@ bool FAnimationAnimNextRuntimeTest_IUpdate::RunTest(const FString& InParameters)
 		}
 
 		// Read our graph
-		FTestUtils::LoadFromArchiveBuffer(*AnimNextGraph, NodeHandles, GraphSharedDataArchiveBuffer);
+		FTestUtils::LoadFromArchiveBuffer(*Module, NodeHandles, GraphSharedDataArchiveBuffer);
 
 		FAnimNextGraphInstancePtr GraphInstance;
-		AnimNextGraph->AllocateInstance(GraphInstance);
+		Module->AllocateInstance(GraphInstance);
 
 		FExecutionContext Context(GraphInstance);
 
@@ -840,9 +840,9 @@ bool FAnimationAnimNextRuntimeTest_IEvaluate::RunTest(const FString& InParameter
 		AUTO_REGISTER_ANIM_TRAIT(FTraitWithOneChild)
 		AUTO_REGISTER_ANIM_TRAIT(FTraitWithChildren)
 
-		UFactory* GraphFactory = NewObject<UAnimNextGraphFactory>();
-		UAnimNextGraph* AnimNextGraph = CastChecked<UAnimNextGraph>(GraphFactory->FactoryCreateNew(UAnimNextGraph::StaticClass(), GetTransientPackage(), TEXT("TestAnimNextGraph"), RF_Transient, nullptr, nullptr, NAME_None));
-		UE_RETURN_ON_ERROR(AnimNextGraph != nullptr, "FAnimationAnimNextRuntimeTest_IEvaluate -> Failed to create animation graph");
+		UFactory* ModuleFactory = NewObject<UAnimNextModuleFactory>();
+		UAnimNextModule* Module = CastChecked<UAnimNextModule>(ModuleFactory->FactoryCreateNew(UAnimNextModule::StaticClass(), GetTransientPackage(), TEXT("TestAnimNextGraph"), RF_Transient, nullptr, nullptr, NAME_None));
+		UE_RETURN_ON_ERROR(Module != nullptr, "FAnimationAnimNextRuntimeTest_IEvaluate -> Failed to create module");
 
 		FScopedClearNodeTemplateRegistry ScopedClearNodeTemplateRegistry;
 		FNodeTemplateRegistry& Registry = FNodeTemplateRegistry::Get();
@@ -933,10 +933,10 @@ bool FAnimationAnimNextRuntimeTest_IEvaluate::RunTest(const FString& InParameter
 		}
 
 		// Read our graph
-		FTestUtils::LoadFromArchiveBuffer(*AnimNextGraph, NodeHandles, GraphSharedDataArchiveBuffer);
+		FTestUtils::LoadFromArchiveBuffer(*Module, NodeHandles, GraphSharedDataArchiveBuffer);
 
 		FAnimNextGraphInstancePtr GraphInstance;
-		AnimNextGraph->AllocateInstance(GraphInstance);
+		Module->AllocateInstance(GraphInstance);
 
 		{
 			TArray<FTraitUID> EvaluatedTraits;
@@ -979,9 +979,9 @@ bool FAnimationAnimNextRuntimeTest_IScopedInterface::RunTest(const FString& InPa
 		AUTO_REGISTER_ANIM_TRAIT(FScopedTagTrait)
 		AUTO_REGISTER_ANIM_TRAIT(FTestScopedTagTrait)
 
-		UFactory* GraphFactory = NewObject<UAnimNextGraphFactory>();
-		UAnimNextGraph* AnimNextGraph = CastChecked<UAnimNextGraph>(GraphFactory->FactoryCreateNew(UAnimNextGraph::StaticClass(), GetTransientPackage(), TEXT("TestAnimNextGraph"), RF_Transient, nullptr, nullptr, NAME_None));
-		UE_RETURN_ON_ERROR(AnimNextGraph != nullptr, "FAnimationAnimNextRuntimeTest_IScopedInterface -> Failed to create animation graph");
+		UFactory* ModuleFactory = NewObject<UAnimNextModuleFactory>();
+		UAnimNextModule* Module = CastChecked<UAnimNextModule>(ModuleFactory->FactoryCreateNew(UAnimNextModule::StaticClass(), GetTransientPackage(), TEXT("TestAnimNextGraph"), RF_Transient, nullptr, nullptr, NAME_None));
+		UE_RETURN_ON_ERROR(Module != nullptr, "FAnimationAnimNextRuntimeTest_IScopedInterface -> Failed to create module");
 
 		FScopedClearNodeTemplateRegistry ScopedClearNodeTemplateRegistry;
 		FNodeTemplateRegistry& Registry = FNodeTemplateRegistry::Get();
@@ -1070,10 +1070,10 @@ bool FAnimationAnimNextRuntimeTest_IScopedInterface::RunTest(const FString& InPa
 		}
 
 		// Read our graph
-		FTestUtils::LoadFromArchiveBuffer(*AnimNextGraph, NodeHandles, GraphSharedDataArchiveBuffer);
+		FTestUtils::LoadFromArchiveBuffer(*Module, NodeHandles, GraphSharedDataArchiveBuffer);
 
 		FAnimNextGraphInstancePtr GraphInstance;
-		AnimNextGraph->AllocateInstance(GraphInstance);
+		Module->AllocateInstance(GraphInstance);
 
 		FExecutionContext Context(GraphInstance);
 
diff --git a/Engine/Plugins/Experimental/AnimNextAnimGraph/Source/AnimNextAnimGraphUncookedOnly/Private/Graph/AnimGraph/AnimGraphNode_AnimNextGraph.cpp b/Engine/Plugins/Experimental/AnimNextAnimGraph/Source/AnimNextAnimGraphUncookedOnly/Private/Graph/AnimGraph/AnimGraphNode_AnimNextGraph.cpp
--- a/Engine/Plugins/Experimental/AnimNextAnimGraph/Source/AnimNextAnimGraphUncookedOnly/Private/Graph/AnimGraph/AnimGraphNode_AnimNextGraph.cpp
+++ b/Engine/Plugins/Experimental/AnimNextAnimGraph/Source/AnimNextAnimGraphUncookedOnly/Private/Graph/AnimGraph/AnimGraphNode_AnimNextGraph.cpp
@@ -24,13 +24,13 @@ UAnimGraphNode_AnimNextGraph::UAnimGraphNode_AnimNextGraph(const FObjectInitiali
 FText UAnimGraphNode_AnimNextGraph::GetNodeTitle(ENodeTitleType::Type TitleType) const
 {
 	// display control rig here
-	return LOCTEXT("AnimGraphNode_AnimNextGraph_Title", "AnimNext Graph");
+	return LOCTEXT("AnimGraphNode_AnimNextModule_Title", "AnimNext Module");
 }
 
 FText UAnimGraphNode_AnimNextGraph::GetTooltipText() const
 {
 	// display control rig here
-	return LOCTEXT("AnimGraphNode_AnimNextGraph_Tooltip", "Evaluates a AnimNext Graph");
+	return LOCTEXT("AnimGraphNode_AnimNextModule_Tooltip", "Evaluates a AnimNext Module");
 }
 
 FText UAnimGraphNode_AnimNextGraph::GetMenuCategory() const
@@ -42,10 +42,10 @@ void UAnimGraphNode_AnimNextGraph::PreloadRequiredAssets()
 {
 	Super::PreloadRequiredAssets();
 
-	if (Node.AnimNextGraph)
+	if (Node.Module)
 	{
-		PreloadObject(Node.AnimNextGraph);
-		PreloadObject(Node.AnimNextGraph->EditorData);
+		PreloadObject(Node.Module);
+		PreloadObject(Node.Module->EditorData);
 	}
 }
 
@@ -184,7 +184,7 @@ void UAnimGraphNode_AnimNextGraph::PostEditChangeProperty(FPropertyChangedEvent&
 
 	if (ChangedProperty)
 	{
-		if (ChangedProperty->GetFName() == GET_MEMBER_NAME_CHECKED(FAnimNode_AnimNextGraph, AnimNextGraph))
+		if (ChangedProperty->GetFName() == GET_MEMBER_NAME_CHECKED(FAnimNode_AnimNextGraph, Module))
 		{
 			bRequiresNodeReconstruct = true;
 			RebuildExposedProperties();
@@ -210,7 +210,7 @@ void UAnimGraphNode_AnimNextGraph::CustomizePinData(UEdGraphPin* Pin, FName Sour
 
 UObject* UAnimGraphNode_AnimNextGraph::GetJumpTargetForDoubleClick() const
 {
-	return Node.AnimNextGraph;
+	return Node.Module;
 }
 
 #undef LOCTEXT_NAMESPACE
