Fixed shutdown crash at AnimNextModule due to incorrect module being loaded

diff --git a/Engine/Plugins/Experimental/AnimNext/Source/AnimNextEditor/Private/AnimNextEditorModule.cpp b/Engine/Plugins/Experimental/AnimNext/Source/AnimNextEditor/Private/AnimNextEditorModule.cpp
--- a/Engine/Plugins/Experimental/AnimNext/Source/AnimNextEditor/Private/AnimNextEditorModule.cpp
+++ b/Engine/Plugins/Experimental/AnimNext/Source/AnimNextEditor/Private/AnimNextEditorModule.cpp
@@ -823,7 +823,7 @@ void FModule::UnregisterWorkspaceDocumentTypes()
 {
 	if(FModuleManager::Get().IsModuleLoaded("WorkspaceEditor"))
 	{
-		Workspace::IWorkspaceEditorModule& WorkspaceEditorModule = FModuleManager::LoadModuleChecked<Workspace::IWorkspaceEditorModule>("PropertyEditor");
+		Workspace::IWorkspaceEditorModule& WorkspaceEditorModule = FModuleManager::LoadModuleChecked<Workspace::IWorkspaceEditorModule>("WorkspaceEditor");
 		WorkspaceEditorModule.UnregisterObjectDocumentType(FTopLevelAssetPath(TEXT("/Script/AnimNext.AnimNextSchedule")));
 		WorkspaceEditorModule.UnregisterObjectDocumentType(FTopLevelAssetPath(TEXT("/Script/AnimNext.AnimNextModule")));
 		WorkspaceEditorModule.UnregisterObjectDocumentType(FTopLevelAssetPath(TEXT("/Script/AnimNextUncookedOnly.AnimNextEdGraph")));
