AnimNextWorkspace schema extending (adding supported classes outside of AnimNext module)

diff --git a/Engine/Plugins/Experimental/AnimNext/Source/AnimNext/Internal/Scheduler/AnimNextSchedule.h b/Engine/Plugins/Experimental/AnimNext/Source/AnimNext/Internal/Scheduler/AnimNextSchedule.h
--- a/Engine/Plugins/Experimental/AnimNext/Source/AnimNext/Internal/Scheduler/AnimNextSchedule.h
+++ b/Engine/Plugins/Experimental/AnimNext/Source/AnimNext/Internal/Scheduler/AnimNextSchedule.h
@@ -39,7 +39,7 @@ namespace UE::AnimNext::UncookedOnly
 
 namespace UE::AnimNext::Editor
 {
-	class FModule;
+	class FAnimNextEditorModule;
 }
 
 USTRUCT()
@@ -248,7 +248,7 @@ private:
 	friend struct FAnimNextSchedulerEntry;
 	friend class UAnimNextComponent;
 	friend class UAnimNextSchedulerWorldSubsystem;
-	friend class UE::AnimNext::Editor::FModule;
+	friend class UE::AnimNext::Editor::FAnimNextEditorModule;
 	friend class UE::AnimNext::UncookedOnly::FModule;
 	friend struct UE::AnimNext::UncookedOnly::FUtils;
 
diff --git a/Engine/Plugins/Experimental/AnimNext/Source/AnimNextEditor/Internal/IAnimNextEditorModule.h b/Engine/Plugins/Experimental/AnimNext/Source/AnimNextEditor/Internal/IAnimNextEditorModule.h
deleted file mode 100644
diff --git a/Engine/Plugins/Experimental/AnimNext/Source/AnimNextEditor/Private/AnimNextEditorModule.cpp b/Engine/Plugins/Experimental/AnimNext/Source/AnimNextEditor/Private/AnimNextEditorModule.cpp
--- a/Engine/Plugins/Experimental/AnimNext/Source/AnimNextEditor/Private/AnimNextEditorModule.cpp
+++ b/Engine/Plugins/Experimental/AnimNext/Source/AnimNextEditor/Private/AnimNextEditorModule.cpp
@@ -48,13 +48,14 @@
 #include "Param/AnimNextActorLocatorEditor.h"
 #include "IWorkspaceEditor.h"
 #include "Framework/Docking/LayoutExtender.h"
+#include "Workspace/AnimNextWorkspaceSchema.h"
 
 #define LOCTEXT_NAMESPACE "AnimNextEditorModule"
 
 namespace UE::AnimNext::Editor
 {
 
-void FModule::StartupModule()
+void FAnimNextEditorModule::StartupModule()
 {
 	// Register settings for user editing
 	ISettingsModule& SettingsModule = FModuleManager::Get().LoadModuleChecked<ISettingsModule>("Settings");
@@ -230,7 +231,7 @@ void FModule::StartupModule()
 	FAnimNextGraphItemDetails::RegisterToolMenuExtensions();
 }
 
-void FModule::ShutdownModule()
+void FAnimNextEditorModule::ShutdownModule()
 {
 	if(FModuleManager::Get().IsModuleLoaded("PropertyEditor"))
 	{
@@ -282,23 +283,31 @@ void FModule::ShutdownModule()
 	UnregisterLocatorFragmentEditorType("AnimNextActor");
 }
 
-TSharedRef<SWidget> FModule::CreateParameterPicker(const FParameterPickerArgs& InArgs)
+TSharedRef<SWidget> FAnimNextEditorModule::CreateParameterPicker(const FParameterPickerArgs& InArgs)
 {
 	return SNew(SParameterPicker)
 		.Args(InArgs);
 }
 
-void FModule::RegisterLocatorFragmentEditorType(FName InLocatorFragmentEditorName)
+void FAnimNextEditorModule::RegisterLocatorFragmentEditorType(FName InLocatorFragmentEditorName)
 {
 	LocatorFragmentEditorNames.Add(InLocatorFragmentEditorName);
 }
 
-void FModule::UnregisterLocatorFragmentEditorType(FName InLocatorFragmentEditorName)
+void FAnimNextEditorModule::UnregisterLocatorFragmentEditorType(FName InLocatorFragmentEditorName)
 {
 	LocatorFragmentEditorNames.Remove(InLocatorFragmentEditorName);
 }
 
-void FModule::RegisterWorkspaceDocumentTypes(Workspace::IWorkspaceEditorModule& WorkspaceEditorModule)
+void FAnimNextEditorModule::AddWorkspaceSupportedAssetClass(const FTopLevelAssetPath& InClassAssetPath)
+{
+	if (InClassAssetPath.IsValid())
+	{
+		UAnimNextWorkspaceSchema::SupportedAssetClasses.AddUnique(InClassAssetPath);
+	}
+}
+
+void FAnimNextEditorModule::RegisterWorkspaceDocumentTypes(Workspace::IWorkspaceEditorModule& WorkspaceEditorModule)
 {
 	// --- AnimNextSchedule ---
 	Workspace::FObjectDocumentArgs ScheduleDocumentArgs(
@@ -909,7 +918,7 @@ void FModule::RegisterWorkspaceDocumentTypes(Workspace::IWorkspaceEditorModule&
 	WorkspaceEditorModule.RegisterObjectDocumentType(FTopLevelAssetPath(TEXT("/Script/AnimNextUncookedOnly.AnimNextEdGraph")), GraphDocumentArgs);
 }
 
-void FModule::UnregisterWorkspaceDocumentTypes()
+void FAnimNextEditorModule::UnregisterWorkspaceDocumentTypes()
 {
 	if(FModuleManager::Get().IsModuleLoaded("WorkspaceEditor"))
 	{
@@ -922,6 +931,6 @@ void FModule::UnregisterWorkspaceDocumentTypes()
 
 }
 
-IMPLEMENT_MODULE(UE::AnimNext::Editor::FModule, AnimNextEditor);
+IMPLEMENT_MODULE(UE::AnimNext::Editor::FAnimNextEditorModule, AnimNextEditor);
 
 #undef LOCTEXT_NAMESPACE
\ No newline at end of file
diff --git a/Engine/Plugins/Experimental/AnimNext/Source/AnimNextEditor/Private/AnimNextEditorModule.h b/Engine/Plugins/Experimental/AnimNext/Source/AnimNextEditor/Private/AnimNextEditorModule.h
--- a/Engine/Plugins/Experimental/AnimNext/Source/AnimNextEditor/Private/AnimNextEditorModule.h
+++ b/Engine/Plugins/Experimental/AnimNext/Source/AnimNextEditor/Private/AnimNextEditorModule.h
@@ -21,7 +21,7 @@ namespace UE::Workspace
 namespace UE::AnimNext::Editor
 {
 
-class FModule : public IModule
+class FAnimNextEditorModule : public IAnimNextEditorModule
 {
 
 private:
@@ -32,6 +32,7 @@ private:
 	virtual TSharedRef<SWidget> CreateParameterPicker(const FParameterPickerArgs& InArgs) override;
 	virtual void RegisterLocatorFragmentEditorType(FName InLocatorFragmentEditorName) override;
 	virtual void UnregisterLocatorFragmentEditorType(FName InLocatorFragmentEditorName) override;
+	virtual void AddWorkspaceSupportedAssetClass(const FTopLevelAssetPath& InClassAssetPath) override;
 
 private:
 	void RegisterWorkspaceDocumentTypes(Workspace::IWorkspaceEditorModule& WorkspaceEditorModule);
diff --git a/Engine/Plugins/Experimental/AnimNext/Source/AnimNextEditor/Private/Param/AnimNextLocatorContext.h b/Engine/Plugins/Experimental/AnimNext/Source/AnimNextEditor/Private/Param/AnimNextLocatorContext.h
--- a/Engine/Plugins/Experimental/AnimNext/Source/AnimNextEditor/Private/Param/AnimNextLocatorContext.h
+++ b/Engine/Plugins/Experimental/AnimNext/Source/AnimNextEditor/Private/Param/AnimNextLocatorContext.h
@@ -21,7 +21,7 @@ class FLocatorContext : public UE::UniversalObjectLocator::ILocatorFragmentEdito
 
 	virtual bool IsFragmentAllowed(FName InFragmentName) const override
 	{
-		FModule& EditorModule = FModuleManager::GetModuleChecked<FModule>("AnimNextEditor");
+		FAnimNextEditorModule& EditorModule = FModuleManager::GetModuleChecked<FAnimNextEditorModule>("AnimNextEditor");
 		return EditorModule.LocatorFragmentEditorNames.Contains(InFragmentName);
 	}
 };
diff --git a/Engine/Plugins/Experimental/AnimNext/Source/AnimNextEditor/Private/Workspace/AnimNextWorkspaceSchema.cpp b/Engine/Plugins/Experimental/AnimNext/Source/AnimNextEditor/Private/Workspace/AnimNextWorkspaceSchema.cpp
--- a/Engine/Plugins/Experimental/AnimNext/Source/AnimNextEditor/Private/Workspace/AnimNextWorkspaceSchema.cpp
+++ b/Engine/Plugins/Experimental/AnimNext/Source/AnimNextEditor/Private/Workspace/AnimNextWorkspaceSchema.cpp
@@ -6,6 +6,12 @@
 
 #define LOCTEXT_NAMESPACE "AnimNextWorkspaceSchema"
 
+TArray<FTopLevelAssetPath> UAnimNextWorkspaceSchema::SupportedAssetClasses =
+{
+	UAnimNextSchedule::StaticClass()->GetClassPathName(),
+	UAnimNextModule::StaticClass()->GetClassPathName()
+};
+
 FText UAnimNextWorkspaceSchema::GetDisplayName() const
 {
 	return LOCTEXT("DisplayName", "AnimNext Workspace");
@@ -13,13 +19,7 @@ FText UAnimNextWorkspaceSchema::GetDisplayName() const
 
 TConstArrayView<FTopLevelAssetPath> UAnimNextWorkspaceSchema::GetSupportedAssetClassPaths() const
 {
-	static const FTopLevelAssetPath Assets[] =
-	{
-		UAnimNextSchedule::StaticClass()->GetClassPathName(),
-		UAnimNextModule::StaticClass()->GetClassPathName(),
-	};
-	
-	return Assets;
+	return SupportedAssetClasses;
 }
 
 #undef LOCTEXT_NAMESPACE
\ No newline at end of file
diff --git a/Engine/Plugins/Experimental/AnimNext/Source/AnimNextEditor/Private/Workspace/AnimNextWorkspaceSchema.h b/Engine/Plugins/Experimental/AnimNext/Source/AnimNextEditor/Private/Workspace/AnimNextWorkspaceSchema.h
--- a/Engine/Plugins/Experimental/AnimNext/Source/AnimNextEditor/Private/Workspace/AnimNextWorkspaceSchema.h
+++ b/Engine/Plugins/Experimental/AnimNext/Source/AnimNextEditor/Private/Workspace/AnimNextWorkspaceSchema.h
@@ -14,4 +14,7 @@ class UAnimNextWorkspaceSchema : public UWorkspaceSchema
 	// UWorkspaceSchema interface
 	virtual FText GetDisplayName() const override;
 	virtual TConstArrayView<FTopLevelAssetPath> GetSupportedAssetClassPaths() const override;
+	
+public:
+	static TArray<FTopLevelAssetPath> SupportedAssetClasses;
 };
\ No newline at end of file
diff --git a/Engine/Plugins/Experimental/AnimNext/Source/AnimNextEditor/Public/IAnimNextEditorModule.h b/Engine/Plugins/Experimental/AnimNext/Source/AnimNextEditor/Public/IAnimNextEditorModule.h
new file mode 100644
--- /dev/null
+++ b/Engine/Plugins/Experimental/AnimNext/Source/AnimNextEditor/Public/IAnimNextEditorModule.h
@@ -0,0 +1,41 @@
+﻿// Copyright Epic Games, Inc. All Rights Reserved.
+
+#pragma once
+
+#include "CoreMinimal.h"
+#include "Modules/ModuleManager.h"
+
+class SWidget;
+struct FTopLevelAssetPath;
+
+namespace UE::Workspace
+{
+class IWorkspaceEditor;
+};
+
+namespace UE::AnimNext::Editor
+{
+
+struct FParameterPickerArgs;
+
+class IAnimNextEditorModule : public IModuleInterface
+{
+public:
+	// Create a parameter picker
+	// @param  InArgs          Arguments used for configuring the picker
+	virtual TSharedRef<SWidget> CreateParameterPicker(const FParameterPickerArgs& InArgs) = 0;
+
+	// Register a valid fragment type name to be used with parameter UOLs
+	// @param InLocatorFragmentEditorName The name of the locator fragment editor
+	virtual void RegisterLocatorFragmentEditorType(FName InLocatorFragmentEditorName) = 0;
+
+	// Unregister a valid fragment type name to be used with parameter UOLs
+	// @param InLocatorFragmentEditorName The name of the locator fragment editor
+	virtual void UnregisterLocatorFragmentEditorType(FName InLocatorFragmentEditorName) = 0;
+
+	// Add a UClass path to the set of classes which can be opened within an AnimNext Workspace
+	// @param InClassAssetPath Asset path for to-be-registered Class 
+	virtual void AddWorkspaceSupportedAssetClass(const FTopLevelAssetPath& InClassAssetPath) = 0;
+};
+
+}
diff --git a/Engine/Plugins/Experimental/AnimNext/Source/AnimNextUncookedOnly/Internal/AnimNextRigVMAssetEditorData.h b/Engine/Plugins/Experimental/AnimNext/Source/AnimNextUncookedOnly/Internal/AnimNextRigVMAssetEditorData.h
--- a/Engine/Plugins/Experimental/AnimNext/Source/AnimNextUncookedOnly/Internal/AnimNextRigVMAssetEditorData.h
+++ b/Engine/Plugins/Experimental/AnimNext/Source/AnimNextUncookedOnly/Internal/AnimNextRigVMAssetEditorData.h
@@ -28,7 +28,7 @@ namespace UE::AnimNext::Editor
 	class SParameterPicker;
 	class SRigVMAssetViewRow;
 	class FParameterCustomization;
-	class FModule;
+	class FAnimNextEditorModule;
 	class FWorkspaceEditor;
 }
 
@@ -74,7 +74,7 @@ protected:
 	friend class UE::AnimNext::Editor::SRigVMAssetViewRow;
 	friend struct UE::AnimNext::Editor::FUtils;
 	friend struct UE::AnimNext::UncookedOnly::FUtils;
-	friend class UE::AnimNext::Editor::FModule;
+	friend class UE::AnimNext::Editor::FAnimNextEditorModule;
 	friend class UE::AnimNext::Editor::FWorkspaceEditor;
 	friend class UAnimNextRigVMAssetEntry;
 	friend class UAnimNextRigVMAssetLibrary;
diff --git a/Engine/Plugins/Experimental/AnimNext/Source/AnimNextUncookedOnly/Private/Graph/AnimGraph/AnimGraphNodeBinding_AnimNextParameters.cpp b/Engine/Plugins/Experimental/AnimNext/Source/AnimNextUncookedOnly/Private/Graph/AnimGraph/AnimGraphNodeBinding_AnimNextParameters.cpp
--- a/Engine/Plugins/Experimental/AnimNext/Source/AnimNextUncookedOnly/Private/Graph/AnimGraph/AnimGraphNodeBinding_AnimNextParameters.cpp
+++ b/Engine/Plugins/Experimental/AnimNext/Source/AnimNextUncookedOnly/Private/Graph/AnimGraph/AnimGraphNodeBinding_AnimNextParameters.cpp
@@ -479,7 +479,7 @@ TSharedRef<SWidget> UAnimGraphNodeBinding_AnimNextParameters::MakePropertyBindin
 				}
 				MenuBuilder.EndSection();
 
-				IModule& EditorModule = FModuleManager::Get().LoadModuleChecked<IModule>("AnimNextEditor");
+				IAnimNextEditorModule& EditorModule = FModuleManager::Get().LoadModuleChecked<IAnimNextEditorModule>("AnimNextEditor");
 				FParameterPickerArgs Args;
 				Args.OnParameterPicked = FOnParameterPicked::CreateLambda([InArgs, BindingPropertyPath, PinArrayIndex, bIsArray, bIsArrayElement, Blueprint](const FParameterBindingReference& InBinding)
 				{
