[Backout] - CL35140883 - Introduced crash on editor startup with certain parameters
Original CL Desc
-----------------------------------------------------------------
AnimNextWorkspace schema extending (adding supported classes outside of AnimNext module)

diff --git a/Engine/Plugins/Experimental/AnimNext/Source/AnimNext/Internal/Scheduler/AnimNextSchedule.h b/Engine/Plugins/Experimental/AnimNext/Source/AnimNext/Internal/Scheduler/AnimNextSchedule.h
--- a/Engine/Plugins/Experimental/AnimNext/Source/AnimNext/Internal/Scheduler/AnimNextSchedule.h
+++ b/Engine/Plugins/Experimental/AnimNext/Source/AnimNext/Internal/Scheduler/AnimNextSchedule.h
@@ -39,7 +39,7 @@ namespace UE::AnimNext::UncookedOnly
 
 namespace UE::AnimNext::Editor
 {
-	class FAnimNextEditorModule;
+	class FModule;
 }
 
 USTRUCT()
@@ -248,7 +248,7 @@ private:
 	friend struct FAnimNextSchedulerEntry;
 	friend class UAnimNextComponent;
 	friend class UAnimNextSchedulerWorldSubsystem;
-	friend class UE::AnimNext::Editor::FAnimNextEditorModule;
+	friend class UE::AnimNext::Editor::FModule;
 	friend class UE::AnimNext::UncookedOnly::FModule;
 	friend struct UE::AnimNext::UncookedOnly::FUtils;
 
diff --git a/Engine/Plugins/Experimental/AnimNext/Source/AnimNextEditor/Internal/IAnimNextEditorModule.h b/Engine/Plugins/Experimental/AnimNext/Source/AnimNextEditor/Internal/IAnimNextEditorModule.h
new file mode 100644
--- /dev/null
+++ b/Engine/Plugins/Experimental/AnimNext/Source/AnimNextEditor/Internal/IAnimNextEditorModule.h
@@ -0,0 +1,31 @@
+﻿// Copyright Epic Games, Inc. All Rights Reserved.
+
+#pragma once
+
+#include "CoreMinimal.h"
+#include "Modules/ModuleManager.h"
+
+class SWidget;
+
+namespace UE::AnimNext::Editor
+{
+
+struct FParameterPickerArgs;
+
+class IModule : public IModuleInterface
+{
+public:
+	// Create a parameter picker
+	// @param  InArgs          Arguments used for configuring the picker
+	virtual TSharedRef<SWidget> CreateParameterPicker(const FParameterPickerArgs& InArgs) = 0;
+
+	// Register a valid fragment type name to be used with parameter UOLs
+	// @param InLocatorFragmentEditorName The name of the locator fragment editor
+	virtual void RegisterLocatorFragmentEditorType(FName InLocatorFragmentEditorName) = 0;
+
+	// Unregister a valid fragment type name to be used with parameter UOLs
+	// @param InLocatorFragmentEditorName The name of the locator fragment editor
+	virtual void UnregisterLocatorFragmentEditorType(FName InLocatorFragmentEditorName) = 0;
+};
+
+}
\ No newline at end of file
diff --git a/Engine/Plugins/Experimental/AnimNext/Source/AnimNextEditor/Private/AnimNextEditorModule.cpp b/Engine/Plugins/Experimental/AnimNext/Source/AnimNextEditor/Private/AnimNextEditorModule.cpp
--- a/Engine/Plugins/Experimental/AnimNext/Source/AnimNextEditor/Private/AnimNextEditorModule.cpp
+++ b/Engine/Plugins/Experimental/AnimNext/Source/AnimNextEditor/Private/AnimNextEditorModule.cpp
@@ -48,14 +48,13 @@
 #include "Param/AnimNextActorLocatorEditor.h"
 #include "IWorkspaceEditor.h"
 #include "Framework/Docking/LayoutExtender.h"
-#include "Workspace/AnimNextWorkspaceSchema.h"
 
 #define LOCTEXT_NAMESPACE "AnimNextEditorModule"
 
 namespace UE::AnimNext::Editor
 {
 
-void FAnimNextEditorModule::StartupModule()
+void FModule::StartupModule()
 {
 	// Register settings for user editing
 	ISettingsModule& SettingsModule = FModuleManager::Get().LoadModuleChecked<ISettingsModule>("Settings");
@@ -231,7 +230,7 @@ void FAnimNextEditorModule::StartupModule()
 	FAnimNextGraphItemDetails::RegisterToolMenuExtensions();
 }
 
-void FAnimNextEditorModule::ShutdownModule()
+void FModule::ShutdownModule()
 {
 	if(FModuleManager::Get().IsModuleLoaded("PropertyEditor"))
 	{
@@ -283,31 +282,23 @@ void FAnimNextEditorModule::ShutdownModule()
 	UnregisterLocatorFragmentEditorType("AnimNextActor");
 }
 
-TSharedRef<SWidget> FAnimNextEditorModule::CreateParameterPicker(const FParameterPickerArgs& InArgs)
+TSharedRef<SWidget> FModule::CreateParameterPicker(const FParameterPickerArgs& InArgs)
 {
 	return SNew(SParameterPicker)
 		.Args(InArgs);
 }
 
-void FAnimNextEditorModule::RegisterLocatorFragmentEditorType(FName InLocatorFragmentEditorName)
+void FModule::RegisterLocatorFragmentEditorType(FName InLocatorFragmentEditorName)
 {
 	LocatorFragmentEditorNames.Add(InLocatorFragmentEditorName);
 }
 
-void FAnimNextEditorModule::UnregisterLocatorFragmentEditorType(FName InLocatorFragmentEditorName)
+void FModule::UnregisterLocatorFragmentEditorType(FName InLocatorFragmentEditorName)
 {
 	LocatorFragmentEditorNames.Remove(InLocatorFragmentEditorName);
 }
 
-void FAnimNextEditorModule::AddWorkspaceSupportedAssetClass(const FTopLevelAssetPath& InClassAssetPath)
-{
-	if (InClassAssetPath.IsValid())
-	{
-		UAnimNextWorkspaceSchema::SupportedAssetClasses.AddUnique(InClassAssetPath);
-	}
-}
-
-void FAnimNextEditorModule::RegisterWorkspaceDocumentTypes(Workspace::IWorkspaceEditorModule& WorkspaceEditorModule)
+void FModule::RegisterWorkspaceDocumentTypes(Workspace::IWorkspaceEditorModule& WorkspaceEditorModule)
 {
 	// --- AnimNextSchedule ---
 	Workspace::FObjectDocumentArgs ScheduleDocumentArgs(
@@ -918,7 +909,7 @@ void FAnimNextEditorModule::RegisterWorkspaceDocumentTypes(Workspace::IWorkspace
 	WorkspaceEditorModule.RegisterObjectDocumentType(FTopLevelAssetPath(TEXT("/Script/AnimNextUncookedOnly.AnimNextEdGraph")), GraphDocumentArgs);
 }
 
-void FAnimNextEditorModule::UnregisterWorkspaceDocumentTypes()
+void FModule::UnregisterWorkspaceDocumentTypes()
 {
 	if(FModuleManager::Get().IsModuleLoaded("WorkspaceEditor"))
 	{
@@ -931,6 +922,6 @@ void FAnimNextEditorModule::UnregisterWorkspaceDocumentTypes()
 
 }
 
-IMPLEMENT_MODULE(UE::AnimNext::Editor::FAnimNextEditorModule, AnimNextEditor);
+IMPLEMENT_MODULE(UE::AnimNext::Editor::FModule, AnimNextEditor);
 
 #undef LOCTEXT_NAMESPACE
\ No newline at end of file
diff --git a/Engine/Plugins/Experimental/AnimNext/Source/AnimNextEditor/Private/AnimNextEditorModule.h b/Engine/Plugins/Experimental/AnimNext/Source/AnimNextEditor/Private/AnimNextEditorModule.h
--- a/Engine/Plugins/Experimental/AnimNext/Source/AnimNextEditor/Private/AnimNextEditorModule.h
+++ b/Engine/Plugins/Experimental/AnimNext/Source/AnimNextEditor/Private/AnimNextEditorModule.h
@@ -21,7 +21,7 @@ namespace UE::Workspace
 namespace UE::AnimNext::Editor
 {
 
-class FAnimNextEditorModule : public IAnimNextEditorModule
+class FModule : public IModule
 {
 
 private:
@@ -32,7 +32,6 @@ private:
 	virtual TSharedRef<SWidget> CreateParameterPicker(const FParameterPickerArgs& InArgs) override;
 	virtual void RegisterLocatorFragmentEditorType(FName InLocatorFragmentEditorName) override;
 	virtual void UnregisterLocatorFragmentEditorType(FName InLocatorFragmentEditorName) override;
-	virtual void AddWorkspaceSupportedAssetClass(const FTopLevelAssetPath& InClassAssetPath) override;
 
 private:
 	void RegisterWorkspaceDocumentTypes(Workspace::IWorkspaceEditorModule& WorkspaceEditorModule);
diff --git a/Engine/Plugins/Experimental/AnimNext/Source/AnimNextEditor/Private/Param/AnimNextLocatorContext.h b/Engine/Plugins/Experimental/AnimNext/Source/AnimNextEditor/Private/Param/AnimNextLocatorContext.h
--- a/Engine/Plugins/Experimental/AnimNext/Source/AnimNextEditor/Private/Param/AnimNextLocatorContext.h
+++ b/Engine/Plugins/Experimental/AnimNext/Source/AnimNextEditor/Private/Param/AnimNextLocatorContext.h
@@ -21,7 +21,7 @@ class FLocatorContext : public UE::UniversalObjectLocator::ILocatorFragmentEdito
 
 	virtual bool IsFragmentAllowed(FName InFragmentName) const override
 	{
-		FAnimNextEditorModule& EditorModule = FModuleManager::GetModuleChecked<FAnimNextEditorModule>("AnimNextEditor");
+		FModule& EditorModule = FModuleManager::GetModuleChecked<FModule>("AnimNextEditor");
 		return EditorModule.LocatorFragmentEditorNames.Contains(InFragmentName);
 	}
 };
diff --git a/Engine/Plugins/Experimental/AnimNext/Source/AnimNextEditor/Private/Workspace/AnimNextWorkspaceSchema.cpp b/Engine/Plugins/Experimental/AnimNext/Source/AnimNextEditor/Private/Workspace/AnimNextWorkspaceSchema.cpp
--- a/Engine/Plugins/Experimental/AnimNext/Source/AnimNextEditor/Private/Workspace/AnimNextWorkspaceSchema.cpp
+++ b/Engine/Plugins/Experimental/AnimNext/Source/AnimNextEditor/Private/Workspace/AnimNextWorkspaceSchema.cpp
@@ -6,12 +6,6 @@
 
 #define LOCTEXT_NAMESPACE "AnimNextWorkspaceSchema"
 
-TArray<FTopLevelAssetPath> UAnimNextWorkspaceSchema::SupportedAssetClasses =
-{
-	UAnimNextSchedule::StaticClass()->GetClassPathName(),
-	UAnimNextModule::StaticClass()->GetClassPathName()
-};
-
 FText UAnimNextWorkspaceSchema::GetDisplayName() const
 {
 	return LOCTEXT("DisplayName", "AnimNext Workspace");
@@ -19,7 +13,13 @@ FText UAnimNextWorkspaceSchema::GetDisplayName() const
 
 TConstArrayView<FTopLevelAssetPath> UAnimNextWorkspaceSchema::GetSupportedAssetClassPaths() const
 {
-	return SupportedAssetClasses;
+	static const FTopLevelAssetPath Assets[] =
+	{
+		UAnimNextSchedule::StaticClass()->GetClassPathName(),
+		UAnimNextModule::StaticClass()->GetClassPathName(),
+	};
+	
+	return Assets;
 }
 
 #undef LOCTEXT_NAMESPACE
\ No newline at end of file
diff --git a/Engine/Plugins/Experimental/AnimNext/Source/AnimNextEditor/Private/Workspace/AnimNextWorkspaceSchema.h b/Engine/Plugins/Experimental/AnimNext/Source/AnimNextEditor/Private/Workspace/AnimNextWorkspaceSchema.h
--- a/Engine/Plugins/Experimental/AnimNext/Source/AnimNextEditor/Private/Workspace/AnimNextWorkspaceSchema.h
+++ b/Engine/Plugins/Experimental/AnimNext/Source/AnimNextEditor/Private/Workspace/AnimNextWorkspaceSchema.h
@@ -14,7 +14,4 @@ class UAnimNextWorkspaceSchema : public UWorkspaceSchema
 	// UWorkspaceSchema interface
 	virtual FText GetDisplayName() const override;
 	virtual TConstArrayView<FTopLevelAssetPath> GetSupportedAssetClassPaths() const override;
-	
-public:
-	static TArray<FTopLevelAssetPath> SupportedAssetClasses;
 };
\ No newline at end of file
diff --git a/Engine/Plugins/Experimental/AnimNext/Source/AnimNextEditor/Public/IAnimNextEditorModule.h b/Engine/Plugins/Experimental/AnimNext/Source/AnimNextEditor/Public/IAnimNextEditorModule.h
deleted file mode 100644
diff --git a/Engine/Plugins/Experimental/AnimNext/Source/AnimNextUncookedOnly/Internal/AnimNextRigVMAssetEditorData.h b/Engine/Plugins/Experimental/AnimNext/Source/AnimNextUncookedOnly/Internal/AnimNextRigVMAssetEditorData.h
--- a/Engine/Plugins/Experimental/AnimNext/Source/AnimNextUncookedOnly/Internal/AnimNextRigVMAssetEditorData.h
+++ b/Engine/Plugins/Experimental/AnimNext/Source/AnimNextUncookedOnly/Internal/AnimNextRigVMAssetEditorData.h
@@ -28,7 +28,7 @@ namespace UE::AnimNext::Editor
 	class SParameterPicker;
 	class SRigVMAssetViewRow;
 	class FParameterCustomization;
-	class FAnimNextEditorModule;
+	class FModule;
 	class FWorkspaceEditor;
 }
 
@@ -74,7 +74,7 @@ protected:
 	friend class UE::AnimNext::Editor::SRigVMAssetViewRow;
 	friend struct UE::AnimNext::Editor::FUtils;
 	friend struct UE::AnimNext::UncookedOnly::FUtils;
-	friend class UE::AnimNext::Editor::FAnimNextEditorModule;
+	friend class UE::AnimNext::Editor::FModule;
 	friend class UE::AnimNext::Editor::FWorkspaceEditor;
 	friend class UAnimNextRigVMAssetEntry;
 	friend class UAnimNextRigVMAssetLibrary;
diff --git a/Engine/Plugins/Experimental/AnimNext/Source/AnimNextUncookedOnly/Private/Graph/AnimGraph/AnimGraphNodeBinding_AnimNextParameters.cpp b/Engine/Plugins/Experimental/AnimNext/Source/AnimNextUncookedOnly/Private/Graph/AnimGraph/AnimGraphNodeBinding_AnimNextParameters.cpp
--- a/Engine/Plugins/Experimental/AnimNext/Source/AnimNextUncookedOnly/Private/Graph/AnimGraph/AnimGraphNodeBinding_AnimNextParameters.cpp
+++ b/Engine/Plugins/Experimental/AnimNext/Source/AnimNextUncookedOnly/Private/Graph/AnimGraph/AnimGraphNodeBinding_AnimNextParameters.cpp
@@ -479,7 +479,7 @@ TSharedRef<SWidget> UAnimGraphNodeBinding_AnimNextParameters::MakePropertyBindin
 				}
 				MenuBuilder.EndSection();
 
-				IAnimNextEditorModule& EditorModule = FModuleManager::Get().LoadModuleChecked<IAnimNextEditorModule>("AnimNextEditor");
+				IModule& EditorModule = FModuleManager::Get().LoadModuleChecked<IModule>("AnimNextEditor");
 				FParameterPickerArgs Args;
 				Args.OnParameterPicked = FOnParameterPicked::CreateLambda([InArgs, BindingPropertyPath, PinArrayIndex, bIsArray, bIsArrayElement, Blueprint](const FParameterBindingReference& InBinding)
 				{
