Linux build fix

#rnx

diff --git a/Engine/Plugins/Experimental/AnimNext/Source/AnimNextUncookedOnly/Internal/AnimNextRigVMAssetEditorData.h b/Engine/Plugins/Experimental/AnimNext/Source/AnimNextUncookedOnly/Internal/AnimNextRigVMAssetEditorData.h
--- a/Engine/Plugins/Experimental/AnimNext/Source/AnimNextUncookedOnly/Internal/AnimNextRigVMAssetEditorData.h
+++ b/Engine/Plugins/Experimental/AnimNext/Source/AnimNextUncookedOnly/Internal/AnimNextRigVMAssetEditorData.h
@@ -218,7 +218,7 @@ protected:
 
 			TArray<UEdGraph*> SubGraphs;
 			RigVMEdGraph->GetAllChildrenGraphs(SubGraphs);
-			for (const TObjectPtr<UEdGraph>& SubGraph : SubGraphs)
+			for (const UEdGraph* SubGraph : SubGraphs)
 			{
 				if (SubGraph)
 				{
@@ -239,7 +239,7 @@ protected:
 
 				TArray<UEdGraph*> SubGraphs;
 				RigVMEdGraph->GetAllChildrenGraphs(SubGraphs);
-				for (const TObjectPtr<UEdGraph>& SubGraph : SubGraphs)
+				for (const UEdGraph* SubGraph : SubGraphs)
 				{
 					if (SubGraph)
 					{
