Linux build fix

#rnx

diff --git a/Engine/Plugins/Experimental/AnimNext/Source/AnimNextUncookedOnly/Private/AnimNextRigVMAssetEditorData.cpp b/Engine/Plugins/Experimental/AnimNext/Source/AnimNextUncookedOnly/Private/AnimNextRigVMAssetEditorData.cpp
--- a/Engine/Plugins/Experimental/AnimNext/Source/AnimNextUncookedOnly/Private/AnimNextRigVMAssetEditorData.cpp
+++ b/Engine/Plugins/Experimental/AnimNext/Source/AnimNextUncookedOnly/Private/AnimNextRigVMAssetEditorData.cpp
@@ -445,7 +445,7 @@ UObject* UAnimNextRigVMAssetEditorData::GetEditorObjectForRigVMGraph(URigVMGraph
 		{
 			TArray<UEdGraph*> SubGraphs;
 			EdGraph->GetAllChildrenGraphs(SubGraphs);
-			for (const TObjectPtr<UEdGraph>& SubGraph : SubGraphs)
+			for (UEdGraph* SubGraph : SubGraphs)
 			{
 				if (URigVMEdGraph* RigVMEdGraph = Cast<URigVMEdGraph>(SubGraph))
 				{
