AnimNextModule static analysis fix

#rnx

diff --git a/Engine/Plugins/Experimental/AnimNext/Source/AnimNextEditor/Private/AnimNextEditorModule.cpp b/Engine/Plugins/Experimental/AnimNext/Source/AnimNextEditor/Private/AnimNextEditorModule.cpp
--- a/Engine/Plugins/Experimental/AnimNext/Source/AnimNextEditor/Private/AnimNextEditorModule.cpp
+++ b/Engine/Plugins/Experimental/AnimNext/Source/AnimNextEditor/Private/AnimNextEditorModule.cpp
@@ -880,7 +880,7 @@ void FModule::RegisterWorkspaceDocumentTypes(Workspace::IWorkspaceEditorModule&
 				{
 					URigVMEdGraph* RigVMEdGraph = Cast<URigVMEdGraph>(GraphEditorData->GetEditorObjectForRigVMGraph(ModelGraph));
 
-					if (GraphEditorData && GraphEditorData->GetLocalFunctionLibrary() != RigVMEdGraph->GetModel())
+					if (RigVMEdGraph != nullptr && GraphEditorData->GetLocalFunctionLibrary() != RigVMEdGraph->GetModel())
 					{
 						const TSharedPtr<Workspace::FWorkspaceBreadcrumb>& GraphCrumb = OutBreadcrumbs.Add_GetRef(MakeShared<Workspace::FWorkspaceBreadcrumb>());
 
diff --git a/Engine/Plugins/Experimental/AnimNext/Source/AnimNextUncookedOnly/Private/Module/AnimNextModule_EditorData.cpp b/Engine/Plugins/Experimental/AnimNext/Source/AnimNextUncookedOnly/Private/Module/AnimNextModule_EditorData.cpp
--- a/Engine/Plugins/Experimental/AnimNext/Source/AnimNextUncookedOnly/Private/Module/AnimNextModule_EditorData.cpp
+++ b/Engine/Plugins/Experimental/AnimNext/Source/AnimNextUncookedOnly/Private/Module/AnimNextModule_EditorData.cpp
@@ -576,21 +576,22 @@ UAnimNextModule_AnimationGraph* UAnimNextModule_EditorData::AddAnimationGraph(FN
 
 		// Editor data has to be the graph outer, or RigVM unique name generator will not work
 		URigVMGraph* NewRigVMGraphModel = RigVMClient.CreateModel(URigVMGraph::StaticClass()->GetFName(), UAnimNextAnimationGraphSchema::StaticClass(), bSetupUndoRedo, this);
-		// Then, to avoid the graph losing ref due to external package, set the same package as the Entry
-		if (!NewRigVMGraphModel->HasAnyFlags(RF_Transient))
+		if (ensure(NewRigVMGraphModel))
 		{
-			NewRigVMGraphModel->SetExternalPackage(CastChecked<UObject>(NewEntry)->GetExternalPackage());
-		}
-
-		ensure(NewRigVMGraphModel);
+			// Then, to avoid the graph losing ref due to external package, set the same package as the Entry
+			if (!NewRigVMGraphModel->HasAnyFlags(RF_Transient))
+			{
+				NewRigVMGraphModel->SetExternalPackage(CastChecked<UObject>(NewEntry)->GetExternalPackage());
+			}
 
-		NewEntry->Graph = NewRigVMGraphModel;
+			NewEntry->Graph = NewRigVMGraphModel;
 		
-		RefreshExternalModels();
-		RigVMClient.AddModel(NewRigVMGraphModel, true);
+			RefreshExternalModels();
+			RigVMClient.AddModel(NewRigVMGraphModel, true);
 
-		URigVMController* Controller = RigVMClient.GetController(NewRigVMGraphModel);
-		UE::AnimNext::UncookedOnly::FUtils::SetupAnimGraph(NewEntry, Controller);
+			URigVMController* Controller = RigVMClient.GetController(NewRigVMGraphModel);
+			UE::AnimNext::UncookedOnly::FUtils::SetupAnimGraph(NewEntry, Controller);
+		}
 	}
 
 	BroadcastModified();
