Non-unity fix
#rnx

diff --git a/Engine/Plugins/Experimental/AnimNext/Source/AnimNextEditor/Private/AnimNextEditorModule.h b/Engine/Plugins/Experimental/AnimNext/Source/AnimNextEditor/Private/AnimNextEditorModule.h
--- a/Engine/Plugins/Experimental/AnimNext/Source/AnimNextEditor/Private/AnimNextEditorModule.h
+++ b/Engine/Plugins/Experimental/AnimNext/Source/AnimNextEditor/Private/AnimNextEditorModule.h
@@ -3,6 +3,7 @@
 #pragma once
 
 #include "IAnimNextEditorModule.h"
+#include "UObject/TopLevelAssetPath.h"
 
 class FAnimNextGraphPanelNodeFactory;
 
