Let's try that again, non-unity CIS fix
#rnx

diff --git a/Engine/Plugins/Experimental/AnimNext/Source/AnimNextEditor/Public/IAnimNextEditorModule.h b/Engine/Plugins/Experimental/AnimNext/Source/AnimNextEditor/Public/IAnimNextEditorModule.h
--- a/Engine/Plugins/Experimental/AnimNext/Source/AnimNextEditor/Public/IAnimNextEditorModule.h
+++ b/Engine/Plugins/Experimental/AnimNext/Source/AnimNextEditor/Public/IAnimNextEditorModule.h
@@ -6,6 +6,7 @@
 #include "Modules/ModuleManager.h"
 
 class SWidget;
+struct FTopLevelAssetPath;
 
 namespace UE::AnimNext::Editor
 {
