Forward declare AnimNewWorkspaceSchema outside of namespace (fails on Clang)
#rnx

diff --git a/Engine/Plugins/Experimental/AnimNext/Source/AnimNextEditor/Private/AnimNextEditorModule.h b/Engine/Plugins/Experimental/AnimNext/Source/AnimNextEditor/Private/AnimNextEditorModule.h
--- a/Engine/Plugins/Experimental/AnimNext/Source/AnimNextEditor/Private/AnimNextEditorModule.h
+++ b/Engine/Plugins/Experimental/AnimNext/Source/AnimNextEditor/Private/AnimNextEditorModule.h
@@ -6,6 +6,7 @@
 #include "UObject/TopLevelAssetPath.h"
 
 class FAnimNextGraphPanelNodeFactory;
+class UAnimNextWorkspaceSchema;
 
 namespace UE::AnimNext::Editor
 {
@@ -24,7 +25,7 @@ namespace UE::AnimNext::Editor
 
 class FAnimNextEditorModule : public IAnimNextEditorModule
 {
-	friend class UAnimNextWorkspaceSchema;
+	friend UAnimNextWorkspaceSchema;
 private:
 	virtual void StartupModule() override;
 	virtual void ShutdownModule() override;
