Fix for horde error in AnimNext EditorTests

#rnx

diff --git a/Engine/Plugins/Experimental/AnimNext/Source/AnimNextUncookedOnly/Internal/IAnimNextRigVMGraphInterface.h b/Engine/Plugins/Experimental/AnimNext/Source/AnimNextUncookedOnly/Internal/IAnimNextRigVMGraphInterface.h
--- a/Engine/Plugins/Experimental/AnimNext/Source/AnimNextUncookedOnly/Internal/IAnimNextRigVMGraphInterface.h
+++ b/Engine/Plugins/Experimental/AnimNext/Source/AnimNextUncookedOnly/Internal/IAnimNextRigVMGraphInterface.h
@@ -31,6 +31,7 @@ public:
 
 private:
 	friend class UAnimNextModule_EditorData;
+	friend class UAnimNextRigVMAssetEditorData;
 
 	// Set the RigVM graph
 	virtual void SetRigVMGraph(URigVMGraph* InGraph) = 0;
diff --git a/Engine/Plugins/Experimental/AnimNext/Source/AnimNextUncookedOnly/Private/AnimNextRigVMAssetEditorData.cpp b/Engine/Plugins/Experimental/AnimNext/Source/AnimNextUncookedOnly/Private/AnimNextRigVMAssetEditorData.cpp
--- a/Engine/Plugins/Experimental/AnimNext/Source/AnimNextUncookedOnly/Private/AnimNextRigVMAssetEditorData.cpp
+++ b/Engine/Plugins/Experimental/AnimNext/Source/AnimNextUncookedOnly/Private/AnimNextRigVMAssetEditorData.cpp
@@ -386,6 +386,15 @@ void UAnimNextRigVMAssetEditorData::HandleRigVMGraphRemoved(const FRigVMClient*
 {
 	if(URigVMGraph* RigVMGraph = InClient->GetModel(InNodePath))
 	{
+		if (UAnimNextRigVMAssetEntry* Entry = FindEntryForRigVMGraph(RigVMGraph))
+		{
+			if (IAnimNextRigVMGraphInterface* GraphInterface = Cast<IAnimNextRigVMGraphInterface>(Entry))
+			{
+				GraphInterface->SetRigVMGraph(nullptr);
+			}
+		}
+		GraphModels.Remove(RigVMGraph);
+
 		RemoveEdGraph(RigVMGraph);
 		RecompileVM();
 
@@ -742,13 +751,6 @@ bool UAnimNextRigVMAssetEditorData::RemoveEntry(UAnimNextRigVMAssetEntry* InEntr
 
 	// Remove from internal array
 	UAnimNextRigVMAssetEntry* EntryToRemove = *EntryToRemovePtr;
-	Entries.Remove(EntryToRemove);
-	RefreshExternalModels();
-
-	if (bSetupUndoRedo)
-	{
-		EntryToRemove->Modify();
-	}
 
 	bool bResult = true;
 	if(const IAnimNextRigVMGraphInterface* GraphInterface = Cast<IAnimNextRigVMGraphInterface>(EntryToRemove))
@@ -762,6 +764,13 @@ bool UAnimNextRigVMAssetEditorData::RemoveEntry(UAnimNextRigVMAssetEntry* InEntr
 		}
 	}
 
+	if (bSetupUndoRedo)
+	{
+		EntryToRemove->Modify();
+	}
+	Entries.Remove(EntryToRemove);
+	RefreshExternalModels();
+
 	// This will cause any external package to be removed when saved
 	EntryToRemove->MarkAsGarbage();
 
