Detect bones missing from parent LOD in AnimNext generation tools

diff --git a/Engine/Plugins/Experimental/AnimNext/Source/AnimNext/Private/GenerationTools.cpp b/Engine/Plugins/Experimental/AnimNext/Source/AnimNext/Private/GenerationTools.cpp
--- a/Engine/Plugins/Experimental/AnimNext/Source/AnimNext/Private/GenerationTools.cpp
+++ b/Engine/Plugins/Experimental/AnimNext/Source/AnimNext/Private/GenerationTools.cpp
@@ -94,8 +94,8 @@ struct FCompareBoneIndexType
 
 		if (bCanGenerateSingleBonesList)
 		{
-			bCanGenerateSingleBonesList &= CheckExcludedBones(NumLODs, GenerationLODData, SkeletalMesh);
-			//bCanGenerateSingleBonesList &= CheckExcludedBones(NumLODs, ComponentSpaceGenerationLODData, SkeletalMesh); // Commented : right now we only support skeletal meshes with all the sockets set to always animate
+			bCanGenerateSingleBonesList &= CheckExcludedAndRequiredBones(NumLODs, GenerationLODData, SkeletalMesh);
+			//bCanGenerateSingleBonesList &= CheckExcludedAndRequiredBones(NumLODs, ComponentSpaceGenerationLODData, SkeletalMesh); // Commented : right now we only support skeletal meshes with all the sockets set to always animate
 
 			TArray<FBoneIndexType> LODBoneIndexToMeshBoneIndexMap;
 			TArray<FBoneIndexType> ComponentSpaceOrderedBoneList;
@@ -254,13 +254,23 @@ struct FCompareBoneIndexType
 	DifferenceBoneIndexArrays(RequiredBones_ParentLOD, RequiredBones, ExcludedBonesFromPrevLOD);
 }
 
-/*static*/ bool FGenerationTools::CheckExcludedBones(const int32 NumLODs
+/*static*/ bool FGenerationTools::CheckExcludedAndRequiredBones(const int32 NumLODs
 	, const TArray<FGenerationLODData>& GenerationLODData
 	, const USkeletalMesh* SkeletalMesh)
 {
 	bool bCanGenerateSingleBonesList = true;
 
-	for (int32 LODIndex = NumLODs - 1; LODIndex > 1; --LODIndex)
+	auto GetBoneNameSafe = [](const USkeletalMesh* SkeletalMesh, const uint32 BoneIndex) -> FString
+		{
+			if (const USkeleton* Skeleton = SkeletalMesh->GetSkeleton())
+			{
+				return Skeleton->GetReferenceSkeleton().GetBoneName(BoneIndex).ToString();
+			}
+
+			return {};
+		};
+
+	for (int32 LODIndex = NumLODs - 1; LODIndex >= 1; --LODIndex)
 	{
 		const FGenerationLODData& LODData = GenerationLODData[LODIndex];
 		const FGenerationLODData& PrevLODData = GenerationLODData[LODIndex - 1];
@@ -278,7 +288,18 @@ struct FCompareBoneIndexType
 			if (LODData.ExcludedBones.Contains(*LODExcludedBonesIt) == false)
 			{
 				bCanGenerateSingleBonesList = false;
-				UE_LOG(LogAnimGenerationTools, Warning, TEXT("SkeletalMesh %s canonical ordered bone set can not be stored in LOD order. LOD %d does not contain all the bones of LOD %d"), *SkeletalMesh->GetPathName(), LODIndex, LODIndex - 1);
+				UE_LOG(LogAnimGenerationTools, Warning, TEXT("SkeletalMesh %s canonical ordered bone set can not be stored in LOD order. LOD %d does not contain all the bones of LOD %d, like e.g. '%s'."), *SkeletalMesh->GetPathName(), LODIndex, LODIndex - 1, *GetBoneNameSafe(SkeletalMesh, *LODExcludedBonesIt));
+				break;
+			}
+		}
+
+		// Check if the parent LOD contains all required bones set in the current LOD.
+		for (TArray<FBoneIndexType>::TConstIterator LODRequiredBonesIt(LODData.RequiredBones); LODRequiredBonesIt; ++LODRequiredBonesIt)
+		{
+			if (PrevLODData.RequiredBones.Contains(*LODRequiredBonesIt) == false)
+			{
+				bCanGenerateSingleBonesList = false;
+				UE_LOG(LogAnimGenerationTools, Warning, TEXT("SkeletalMesh %s canonical ordered bone set can not be stored in LOD order. LOD %d does not contain all the bones of LOD %d, like e.g. '%s'."), *SkeletalMesh->GetPathName(), LODIndex, LODIndex - 1, *GetBoneNameSafe(SkeletalMesh, *LODRequiredBonesIt));
 				break;
 			}
 		}
diff --git a/Engine/Plugins/Experimental/AnimNext/Source/AnimNext/Public/GenerationTools.h b/Engine/Plugins/Experimental/AnimNext/Source/AnimNext/Public/GenerationTools.h
--- a/Engine/Plugins/Experimental/AnimNext/Source/AnimNext/Public/GenerationTools.h
+++ b/Engine/Plugins/Experimental/AnimNext/Source/AnimNext/Public/GenerationTools.h
@@ -60,8 +60,8 @@ public:
 	static void CalculateDifferenceFromParentLOD(int32 LODIndex, TArray<FGenerationLODData>& GenerationLODData);
 
 	// For each LOD :
-	// Check the excluded bones in LOD(N) contain all the bones excluded in LOD(N-1)
-	static bool CheckExcludedBones(const int32 NumLODs
+	// Check the excluded and required bones in LOD(N) contain all the bones excluded in LOD(N-1)
+	static bool CheckExcludedAndRequiredBones(const int32 NumLODs
 		, const TArray<FGenerationLODData>& GenerationLODData
 		, const USkeletalMesh* SkeletalMesh);
 
