AnimNext : Added automatic LOD fixup for skeletal meshes that have a malformed LOD setup (missing bones in parent LOD that are used on children).
Also fixed the generation size of the skeleton to mesh lod bone index map, as it has different size than the mesh to skeleton

#rnx

diff --git a/Engine/Plugins/Experimental/AnimNext/Source/AnimNext/Private/GenerationTools.cpp b/Engine/Plugins/Experimental/AnimNext/Source/AnimNext/Private/GenerationTools.cpp
--- a/Engine/Plugins/Experimental/AnimNext/Source/AnimNext/Private/GenerationTools.cpp
+++ b/Engine/Plugins/Experimental/AnimNext/Source/AnimNext/Private/GenerationTools.cpp
@@ -85,18 +85,17 @@ struct FCompareBoneIndexType
 		GenerateRawLODData(SkeletalMeshComponent, SkeletalMesh, LOD0Index, LODRenderData, RequiredBones_LOD0, ComponentSpaceGenerationLODData[LOD0Index].RequiredBones);
 
 		// Now calculate the LODs > 1
-
 		constexpr int32 StartLOD = 1;
 		GenerateLODData(SkeletalMeshComponent, SkeletalMesh, StartLOD, NumLODs, LODRenderData, RequiredBones_LOD0, GenerationLODData, ComponentSpaceGenerationLODData);
 
+		// Add missing bones to parent LODs for skeletal meshes that arrive with malformed LOD setups
+		FixLODRequiredBones(NumLODs, SkeletalMesh, GenerationLODData, ComponentSpaceGenerationLODData);
+
 		// Check if the sockets are set to always animate, else the component space requires separated data (different bone indexes)
 		bool bCanGenerateSingleBonesList = CheckSkeletalAllMeshSocketsAlwaysAnimate(SkeletalMesh);
 
 		if (bCanGenerateSingleBonesList)
 		{
-			bCanGenerateSingleBonesList &= CheckExcludedAndRequiredBones(NumLODs, GenerationLODData, SkeletalMesh);
-			//bCanGenerateSingleBonesList &= CheckExcludedAndRequiredBones(NumLODs, ComponentSpaceGenerationLODData, SkeletalMesh); // Commented : right now we only support skeletal meshes with all the sockets set to always animate
-
 			TArray<FBoneIndexType> LODBoneIndexToMeshBoneIndexMap;
 			TArray<FBoneIndexType> ComponentSpaceOrderedBoneList;
 
@@ -126,12 +125,13 @@ struct FCompareBoneIndexType
 				const FSkeletonToMeshLinkup& LinkupTable = const_cast<USkeleton*>(Skeleton)->FindOrAddMeshLinkupData(SkeletalMesh);
 
 				// Generate a Skeleton to LOD look up table
+				const int32 SkeletonToMeshLinkupTableSize = LinkupTable.SkeletonToMeshTable.Num();
 				const int32 NumOrderedBones = LODBoneIndexToMeshBoneIndexMap.Num();
 
 				TArray<FBoneIndexType> SkeletonBoneIndexToLODBoneIndexMap;
 				TArray<FBoneIndexType> LODBoneIndexToSkeletonBoneIndexMap;
-				SkeletonBoneIndexToLODBoneIndexMap.SetNumZeroed(NumOrderedBones);
-				LODBoneIndexToSkeletonBoneIndexMap.SetNumZeroed(NumOrderedBones);
+				SkeletonBoneIndexToLODBoneIndexMap.Init(INDEX_NONE, SkeletonToMeshLinkupTableSize);
+				LODBoneIndexToSkeletonBoneIndexMap.Init(INDEX_NONE, NumOrderedBones);
 
 				for (int32 LODBoneIndex = 0; LODBoneIndex < NumOrderedBones; ++LODBoneIndex)
 				{
@@ -151,8 +151,6 @@ struct FCompareBoneIndexType
 				ReferencePoseGenerated = true;
 			}
 		}
-		// TODO : Decide if we only allow skeletal meshes that allow for single bone list generation
-		//        Or create an alternative
 	}
 
 	return ReferencePoseGenerated;
@@ -193,7 +191,7 @@ struct FCompareBoneIndexType
 		}
 
 
-			// Get socket bones set to animate and bones required to fill the component space base transforms
+		// Get socket bones set to animate and bones required to fill the component space base transforms
 		TArray<FBoneIndexType> NeededBonesForFillComponentSpaceTransforms;
 		USkeletalMeshComponent::GetSocketRequiredBones(SkeletalMesh, OutRequiredBones, NeededBonesForFillComponentSpaceTransforms);
 
@@ -254,21 +252,80 @@ struct FCompareBoneIndexType
 	DifferenceBoneIndexArrays(RequiredBones_ParentLOD, RequiredBones, ExcludedBonesFromPrevLOD);
 }
 
-/*static*/ bool FGenerationTools::CheckExcludedAndRequiredBones(const int32 NumLODs
-	, const TArray<FGenerationLODData>& GenerationLODData
-	, const USkeletalMesh* SkeletalMesh)
+static FString GetBoneNameSafe(const USkeletalMesh* SkeletalMesh, const uint32 BoneIndex)
 {
-	bool bCanGenerateSingleBonesList = true;
+	FString BoneName;
+
+	if (const USkeleton* Skeleton = SkeletalMesh->GetSkeleton())
+	{
+		if (Skeleton->GetReferenceSkeleton().IsValidIndex(BoneIndex))
+		{
+			BoneName = Skeleton->GetReferenceSkeleton().GetBoneName(BoneIndex).ToString();
+		}
+	}
+
+	return BoneName;
+};
+
+void FGenerationTools::FixLODRequiredBones(const int32 NumLODs
+	, const USkeletalMesh* SkeletalMesh
+	, TArray<FGenerationLODData>& GenerationLODData
+	, TArray<FGenerationLODData>& GenerationComponentSpaceLODData)
+{
+	for (int32 LODIndex = NumLODs - 1; LODIndex > 0; --LODIndex)
+	{
+		FGenerationLODData& LODData = GenerationLODData[LODIndex];
+		FGenerationLODData& PrevLODData = GenerationLODData[LODIndex - 1];
+
 
-	auto GetBoneNameSafe = [](const USkeletalMesh* SkeletalMesh, const uint32 BoneIndex) -> FString
+		TArray<FBoneIndexType> MissingBones;
+		MissingBones.Reserve(LODData.RequiredBones.Num());
+
+		// Check if all required bones by the current LOD are part of the parent LOD as well.
+		// If a bone is present at a lower LOD while missing from the parent LOD, we will automatically add it to the parent LOD.
+		for (TArray<FBoneIndexType>::TConstIterator ReqBoneIt(LODData.RequiredBones); ReqBoneIt; ++ReqBoneIt)
 		{
-			if (const USkeleton* Skeleton = SkeletalMesh->GetSkeleton())
+			if (!PrevLODData.RequiredBones.Contains(*ReqBoneIt))
 			{
-				return Skeleton->GetReferenceSkeleton().GetBoneName(BoneIndex).ToString();
+				MissingBones.Add(*ReqBoneIt);
+
+				UE_LOG(LogAnimGenerationTools, Warning, TEXT("SkeletalMesh LOD %d does not contain bone [%s] required by LOD %d. Please update the skeletal mesh asset [%s] or its corresponding LOD settings asset.")
+					, LODIndex - 1
+					, *GetBoneNameSafe(SkeletalMesh, *ReqBoneIt)
+					, LODIndex
+					, *SkeletalMesh->GetPathName());
 			}
+		}
 
-			return {};
-		};
+		if (MissingBones.Num() > 0) // Update the arrays and keep bones sorted
+		{
+			FGenerationLODData& ComponentSpacePrevLODData = GenerationLODData[LODIndex - 1];
+
+			USkeletalMeshComponent::MergeInBoneIndexArrays(PrevLODData.RequiredBones, MissingBones);
+			USkeletalMeshComponent::MergeInBoneIndexArrays(ComponentSpacePrevLODData.RequiredBones, MissingBones);
+
+			if (LODIndex > 1) // No need to remove excluded bones at LOD 1, as parent is LOD 0 and has none excluded
+			{
+				const int32 NumElem = MissingBones.Num();
+				for (int32 i = NumElem - 1; i >= 0; --i)
+				{
+					const FBoneIndexType& BoneIndex = MissingBones[i];
+
+					PrevLODData.ExcludedBones.Remove(BoneIndex);
+					PrevLODData.ExcludedBonesFromPrevLOD.Remove(BoneIndex);
+					ComponentSpacePrevLODData.ExcludedBones.Remove(BoneIndex);
+					ComponentSpacePrevLODData.ExcludedBonesFromPrevLOD.Remove(BoneIndex);
+				}
+			}
+		}
+	}
+}
+
+/*static*/ bool FGenerationTools::CheckExcludedBones(const int32 NumLODs
+	, const TArray<FGenerationLODData>& GenerationLODData
+	, const USkeletalMesh* SkeletalMesh)
+{
+	bool bCanGenerateSingleBonesList = true;
 
 	for (int32 LODIndex = NumLODs - 1; LODIndex >= 1; --LODIndex)
 	{
@@ -292,17 +349,6 @@ struct FCompareBoneIndexType
 				break;
 			}
 		}
-
-		// Check if the parent LOD contains all required bones set in the current LOD.
-		for (TArray<FBoneIndexType>::TConstIterator LODRequiredBonesIt(LODData.RequiredBones); LODRequiredBonesIt; ++LODRequiredBonesIt)
-		{
-			if (PrevLODData.RequiredBones.Contains(*LODRequiredBonesIt) == false)
-			{
-				bCanGenerateSingleBonesList = false;
-				UE_LOG(LogAnimGenerationTools, Warning, TEXT("SkeletalMesh %s canonical ordered bone set can not be stored in LOD order. LOD %d does not contain all the bones of LOD %d, like e.g. '%s'."), *SkeletalMesh->GetPathName(), LODIndex, LODIndex - 1, *GetBoneNameSafe(SkeletalMesh, *LODRequiredBonesIt));
-				break;
-			}
-		}
 	}
 	return bCanGenerateSingleBonesList;
 }
diff --git a/Engine/Plugins/Experimental/AnimNext/Source/AnimNext/Public/GenerationTools.h b/Engine/Plugins/Experimental/AnimNext/Source/AnimNext/Public/GenerationTools.h
--- a/Engine/Plugins/Experimental/AnimNext/Source/AnimNext/Public/GenerationTools.h
+++ b/Engine/Plugins/Experimental/AnimNext/Source/AnimNext/Public/GenerationTools.h
@@ -59,9 +59,16 @@ public:
 	// Calculate the bone indexes difference from LOD0 for LODIndex
 	static void CalculateDifferenceFromParentLOD(int32 LODIndex, TArray<FGenerationLODData>& GenerationLODData);
 
+	// Check the required bones in LOD(N) are required in LOD(N-1) 
+	// and add missing bones at LOD(N-1), to enable fast path on malformed LODs
+	static void FixLODRequiredBones(const int32 NumLODs
+		, const USkeletalMesh* SkeletalMesh
+		, TArray<FGenerationLODData>& GenerationLODData
+		, TArray<FGenerationLODData>& GenerationComponentSpaceLODData);
+
 	// For each LOD :
-	// Check the excluded and required bones in LOD(N) contain all the bones excluded in LOD(N-1)
-	static bool CheckExcludedAndRequiredBones(const int32 NumLODs
+	// Check the excluded bones in LOD(N) contain all the bones excluded in LOD(N-1)
+	static bool CheckExcludedBones(const int32 NumLODs
 		, const TArray<FGenerationLODData>& GenerationLODData
 		, const USkeletalMesh* SkeletalMesh);
 
