AnimNext: Module refactor

Removes parameter system (inc FParamStack and FParamTypeHandle) and replaces the functionality with Variables. These are the usual 'external variables' at the RigVM level.
Removes schedules, replaces them with module events. Module events are called at specific points in the frame (PrePhysics, PostPhysics etc.) and can perform animation graph logic using the FRigUnit_AnimNextRunAnimationGraph rig unit, outputting to a mesh's buffers via FRigUnit_AnimNextWriteSkeletalMeshComponentPose.
Two main asset types: UAnimNextModule and UAnimNextAnimationGraph. Components in the world/scene graph represent an instance of a module. Animation graphs are designed to be the basis for logical animation units and are intended to be the base class of any extensions (e.g. State Tree).
Communication of parameters to the system is done in one of two ways; either via the component's SetVariable function (which sets a variable by name and relies on double-buffering the public variables) or via a variable binding (which executes some logic to access & copy data on the game thread prior to a module's execution).

Sub-reviews:
https://p4-swarm.epicgames.net/reviews/35170229
https://p4-swarm.epicgames.net/reviews/35049182
https://p4-swarm.epicgames.net/reviews/35170496
https://p4-swarm.epicgames.net/reviews/35173062
https://p4-swarm.epicgames.net/reviews/35170485
https://p4-swarm.epicgames.net/reviews/35170468
https://p4-swarm.epicgames.net/reviews/35170558

#rnx

diff --git a/Engine/Plugins/Experimental/AnimNext/Config/BaseAnimNext.ini b/Engine/Plugins/Experimental/AnimNext/Config/BaseAnimNext.ini
--- a/Engine/Plugins/Experimental/AnimNext/Config/BaseAnimNext.ini
+++ b/Engine/Plugins/Experimental/AnimNext/Config/BaseAnimNext.ini
@@ -6,7 +6,6 @@
 ; Parameter block/graph merge
 +ClassRedirects=(OldName="/Script/AnimNext.AnimNextParameterBlock", NewName="/Script/AnimNext.AnimNextGraph")
 +ClassRedirects=(OldName="/Script/AnimNextUncookedOnly.AnimNextParameterBlock_EditorData", NewName="/Script/AnimNextUncookedOnly.AnimNextGraph_EditorData")
-+ClassRedirects=(OldName="/Script/AnimNextUncookedOnly.AnimNextParameterBlock_ParamSchema", NewName="/Script/AnimNextUncookedOnly.AnimNextGraph_EventGraphSchema")
 +ClassRedirects=(OldName="/Script/AnimNextUncookedOnly.AnimNextParameterBlock_EdGraph", NewName="/Script/AnimNextUncookedOnly.AnimNextGraph_EdGraph")
 +ClassRedirects=(OldName="/Script/AnimNextUncookedOnly.AnimNextParameterBlock_EdGraphNode", NewName="/Script/AnimNextUncookedOnly.AnimNextGraph_EdGraphNode")
 +ClassRedirects=(OldName="/Script/AnimNextUncookedOnly.AnimNextParameterBlock_EdGraphSchema", NewName="/Script/AnimNextUncookedOnly.AnimNextGraph_EdGraphSchema")
@@ -25,15 +24,14 @@
 
 ; AnimNextGraph -> AnimNextModule
 +ClassRedirects=(OldName="/Script/AnimNext.AnimNextGraph", NewName="/Script/AnimNext.AnimNextModule")
-+ClassRedirects=(OldName="/Script/AnimNextUncookedOnly.AnimNextGraph_AnimationGraph", NewName="/Script/AnimNextUncookedOnly.AnimNextModule_AnimationGraph")
 +ClassRedirects=(OldName="/Script/AnimNextUncookedOnly.AnimNextGraph_Controller", NewName="/Script/AnimNextUncookedOnly.AnimNextModule_Controller")
 +ClassRedirects=(OldName="/Script/AnimNextUncookedOnly.AnimNextGraph_EditorData", NewName="/Script/AnimNextUncookedOnly.AnimNextModule_EditorData")
-+ClassRedirects=(OldName="/Script/AnimNextUncookedOnly.AnimNextGraph_EventGraph", NewName="/Script/AnimNextUncookedOnly.AnimNextModule_EventGraph")
-+ClassRedirects=(OldName="/Script/AnimNextUncookedOnly.AnimNextGraph_Parameter", NewName="/Script/AnimNextUncookedOnly.AnimNextModule_Parameter")
 +ClassRedirects=(OldName="/Script/AnimNextUncookedOnly.AnimNextGraph_AnimationGraphSchema", NewName="/Script/AnimNextUncookedOnly.AnimNextAnimationGraphSchema")
 +ClassRedirects=(OldName="/Script/AnimNextUncookedOnly.AnimNextGraph_EdGraph", NewName="/Script/AnimNextUncookedOnly.AnimNextEdGraph")
 +ClassRedirects=(OldName="/Script/AnimNextUncookedOnly.AnimNextGraph_EdGraphNode", NewName="/Script/AnimNextUncookedOnly.AnimNextEdGraphNode")
 +ClassRedirects=(OldName="/Script/AnimNextUncookedOnly.AnimNextGraph_EdGraphSchema", NewName="/Script/AnimNextUncookedOnly.AnimNextEdGraphSchema")
++ClassRedirects=(OldName="/Script/AnimNextUncookedOnly.AnimNextParameterBlock_ParamSchema", NewName="/Script/AnimNextUncookedOnly.AnimNextEventGraphSchema")
++ClassRedirects=(OldName="/Script/AnimNextUncookedOnly.AnimNextGraph_EventGraphSchema", NewName="/Script/AnimNextUncookedOnly.AnimNextEventGraphSchema")
 +ClassRedirects=(OldName="/Script/AnimNextUncookedOnly.AnimNextGraph_UnitNode", NewName="/Script/AnimNextUncookedOnly.AnimNextUnitNode")
 +ClassRedirects=(OldName="/Script/AnimNextUncookedOnly.AnimNextGraphLibrary", NewName="/Script/AnimNextUncookedOnly.AnimNextModuleLibrary")
 +ClassRedirects=(OldName="/Script/AnimNextUncookedOnly.AnimNextRigVMWorkspaceAssetUserData", NewName="/Script/AnimNextUncookedOnly.AnimNextModuleWorkspaceAssetUserData")
@@ -45,3 +43,23 @@
 +PropertyRedirects=(OldName="/Script/AnimNext.AnimNextSubGraphHostTraitSharedData.SubGraph", NewName="Module")
 +PropertyRedirects=(OldName="/Script/AnimNext.AnimNextScheduleGraphTask.Graph", NewName="Module")
 +PropertyRedirects=(OldName="/Script/AnimNext.AnimNextScheduleGraphTask.DynamicGraph", NewName="DynamicModule")
++PropertyRedirects=(OldName="/Script/AnimNextUncookedOnly.AnimNextParameterProviderAssetRegistryExports.Parameters", NewName="Variables")
+
+; Splitting up modules & adding variables to graphs
++ClassRedirects=(OldName="/Script/AnimNextUncookedOnly.AnimNextModuleWorkspaceAssetUserData", NewName="/Script/AnimNextUncookedOnly.AnimNextAssetWorkspaceAssetUserData")
++ClassRedirects=(OldName="/Script/AnimNextUncookedOnly.AnimNextModule_Controller", NewName="/Script/AnimNextUncookedOnly.AnimNextController")
++ClassRedirects=(OldName="/Script/AnimNextUncookedOnly.AnimNextModule_AnimationGraph", NewName="/Script/AnimNextUncookedOnly.AnimNextAnimationGraphEntry")
++ClassRedirects=(OldName="/Script/AnimNextUncookedOnly.AnimNextModule_Parameter", NewName="/Script/AnimNextUncookedOnly.AnimNextVariableEntry")
++ClassRedirects=(OldName="/Script/AnimNextUncookedOnly.AnimNextModule_EventGraph", NewName="/Script/AnimNextUncookedOnly.AnimNextEventGraphEntry")
++ClassRedirects=(OldName="/Script/AnimNextUncookedOnly.AnimNextGraph_AnimationGraph", NewName="/Script/AnimNextUncookedOnly.AnimNextAnimationGraphEntry")
++ClassRedirects=(OldName="/Script/AnimNextUncookedOnly.AnimNextGraph_Parameter", NewName="/Script/AnimNextUncookedOnly.AnimNextVariableEntry")
++ClassRedirects=(OldName="/Script/AnimNextUncookedOnly.AnimNextGraph_EventGraph", NewName="/Script/AnimNextUncookedOnly.AnimNextEventGraphEntry")
+
++StructRedirects=(OldName="/Script/AnimNextUncookedOnly.AnimNextGraphAssetOutlinerData", NewName="/Script/AnimNextUncookedOnly.AnimNextModuleOutlinerData")
++StructRedirects=(OldName="/Script/AnimNextUncookedOnly.AnimNextParameterOutlinerData", NewName="/Script/AnimNextUncookedOnly.AnimNextVariableOutlinerData")
++StructRedirects=(OldName="/Script/AnimNext.RigUnit_AnimNextModuleSelf", NewName="/Script/AnimNext.FRigUnit_AnimNextAnimationGraphSelf")
++StructRedirects=(OldName="/Script/AnimNextUncookedOnly.AnimNextParameterProviderAssetRegistryExports", NewName="/Script/AnimNextUncookedOnly.AnimNextAssetRegistryExports")
++StructRedirects=(OldName="/Script/AnimNextUncookedOnly.AnimNextParameterAssetRegistryExportEntry", NewName="/Script/AnimNextUncookedOnly.AnimNextAssetRegistryExportedVariable")
+
++EnumRedirects=(OldName="EAnimNextModuleInitMethod",NewName="/Script/AnimNext.EAnimNextModuleInitMethod")
++EnumRedirects=(OldName="EAnimNextParameterFlags",NewName="/Script/AnimNext.EAnimNextExportedVariableFlags")
diff --git a/Engine/Plugins/Experimental/AnimNext/Source/AnimNext/Internal/AnimNextExecuteContext.h b/Engine/Plugins/Experimental/AnimNext/Source/AnimNext/Internal/AnimNextExecuteContext.h
--- a/Engine/Plugins/Experimental/AnimNext/Source/AnimNext/Internal/AnimNextExecuteContext.h
+++ b/Engine/Plugins/Experimental/AnimNext/Source/AnimNext/Internal/AnimNextExecuteContext.h
@@ -4,8 +4,8 @@
 
 #include "CoreMinimal.h"
 #include "RigVMCore/RigVMExecuteContext.h"
+#include "Module/AnimNextModuleContextData.h"
 #include "Graph/AnimNextGraphContextData.h"
-#include "Param/AnimNextParamContextData.h"
 #include "Misc/TVariant.h"
 #include "AnimNextExecuteContext.generated.h"
 
@@ -14,6 +14,7 @@ struct FAnimNextGraphInstance;
 namespace UE::AnimNext
 {
 	struct FLatentPropertyHandle;
+	struct FModuleEventTickFunction;
 }
 
 USTRUCT(BlueprintType)
@@ -54,9 +55,10 @@ protected:
 	}
 
 	// All possible known variants of our context data. IF we ever want this to be extensible, this can be converted into an FInstancedStruct
-	TVariant<FAnimNextGraphContextData, FAnimNextParamContextData> ContextData;
-
+	TVariant<FAnimNextModuleContextData, FAnimNextGraphContextData> ContextData;
+	
+	friend struct UE::AnimNext::FModuleEventTickFunction;
 	friend struct FAnimNextGraphInstance;
-	friend class UAnimNextModule;
+	friend class UAnimNextAnimationGraph;
 };
 
diff --git a/Engine/Plugins/Experimental/AnimNext/Source/AnimNext/Internal/AnimNextModuleImpl.h b/Engine/Plugins/Experimental/AnimNext/Source/AnimNext/Internal/AnimNextModuleImpl.h
--- a/Engine/Plugins/Experimental/AnimNext/Source/AnimNext/Internal/AnimNextModuleImpl.h
+++ b/Engine/Plugins/Experimental/AnimNext/Source/AnimNext/Internal/AnimNextModuleImpl.h
@@ -6,12 +6,10 @@
 #include "IAnimNextModuleInterface.h"
 
 struct FAnimNextGraphInstancePtr;
-struct FAnimNextParamInstanceIdentifier;
+struct FAnimNextVariableBindingData;
 
 namespace UE::AnimNext
 {
-	class IParameterSource;
-	class IParameterSourceFactory;
 	struct FParameterSourceContext;
 }
 
@@ -31,39 +29,15 @@ public:
 	// IAnimNextModuleInterface interface
 	virtual void RegisterAnimNextAnimGraph(const IAnimNextAnimGraph& InAnimGraphImpl) override;
 	virtual void UnregisterAnimNextAnimGraph() override;
-	virtual void UpdateGraph(FAnimNextGraphInstancePtr& GraphInstance, float DeltaTime, FTraitEventList& InputEventList, FTraitEventList& OutputEventList) override;
-	virtual void EvaluateGraph(FAnimNextGraphInstancePtr& GraphInstance, const FReferencePose& RefPose, int32 GraphLODLevel, FLODPoseHeap& OutputPose) const override;
-
-	// Factory method used to create a parameter source to access the specified instance ID, with a set of parameters that are initially required
-	// @param   InContext             Context used to set up the parameter source
-	// @param   InInstanceId          The instance identifier associated with the parameters that are required
-	// @param   InRequiredParameters  Any required parameters that the source should initially supply, can be empty, in which ase all parameters are created
-	// @return a new parameter source, or nullptr if the instance ID could not be handled
-	ANIMNEXT_API TUniquePtr<IParameterSource> CreateParameterSource(const FParameterSourceContext& InContext, const TInstancedStruct<FAnimNextParamInstanceIdentifier>& InInstanceId, TConstArrayView<FName> InRequiredParameters = TConstArrayView<FName>()) const;
-
-	// Register a factory that can be used to generate parameter sources
-	// @param   InName                Identifier for the factory
-	// @param   InFactory             The factory to register
-	ANIMNEXT_API void RegisterParameterSourceFactory(FName InName, TSharedRef<IParameterSourceFactory> InFactory);
-
-	// Unregister a factory previously passed to RegisterParameterSourceFactory
-	// @param   InName                Identifier for the factory
-	ANIMNEXT_API void UnregisterParameterSourceFactory(FName InName);
-
-	// Find a factory previously passed to RegisterParameterSourceFactory
-	// @param   InName                Identifier for the factory
-	// @return a shared ptr to the factory, if found, otherwise an invalid shared ptr
-	ANIMNEXT_API TSharedPtr<IParameterSourceFactory> FindParameterSourceFactory(FName InName);
+	virtual void UpdateGraph(const FAnimNextGraphInstancePtr& GraphInstance, float DeltaTime, FTraitEventList& InputEventList, FTraitEventList& OutputEventList) override;
+	virtual void EvaluateGraph(const FAnimNextGraphInstancePtr& GraphInstance, const FReferencePose& RefPose, int32 GraphLODLevel, FLODPoseHeap& OutputPose) const override;
 
 #if WITH_ANIMNEXT_CONSOLE_COMMANDS
 	void ListNodeTemplates(const TArray<FString>& Args);
-	void ListModules(const TArray<FString>& Args);
+	void ListAnimationGraphs(const TArray<FString>& Args);
 
 	TArray<IConsoleObject*> ConsoleCommands;
 #endif
-
-	// All known factories
-	TMap<FName, TSharedRef<IParameterSourceFactory>> ParameterSourceFactories;
 };
 
 }
diff --git a/Engine/Plugins/Experimental/AnimNext/Source/AnimNext/Internal/AnimNextRigVMAsset.h b/Engine/Plugins/Experimental/AnimNext/Source/AnimNext/Internal/AnimNextRigVMAsset.h
--- a/Engine/Plugins/Experimental/AnimNext/Source/AnimNext/Internal/AnimNextRigVMAsset.h
+++ b/Engine/Plugins/Experimental/AnimNext/Source/AnimNext/Internal/AnimNextRigVMAsset.h
@@ -7,7 +7,14 @@
 #include "RigVMHost.h"
 #include "AnimNextRigVMAsset.generated.h"
 
+class UAnimNextVariableEntry;
 class UAnimNextRigVMAssetEditorData;
+struct FAnimNextModuleInstance;
+
+namespace UE::AnimNext
+{
+	struct FProxyVariablesContext;
+}
 
 namespace UE::AnimNext::UncookedOnly
 {
@@ -21,7 +28,10 @@ class UAnimNextRigVMAsset : public URigVMHost
 	GENERATED_BODY()
 
 protected:
+	friend class UAnimNextVariableEntry;
 	friend class UAnimNextRigVMAssetEditorData;
+	friend struct FAnimNextModuleInstance;
+	friend struct UE::AnimNext::FProxyVariablesContext;
 	friend struct UE::AnimNext::UncookedOnly::FUtils;
 
 	UAnimNextRigVMAsset(const FObjectInitializer& ObjectInitializer);
@@ -31,6 +41,15 @@ protected:
 	virtual void PostLoad() override;
 	virtual void GetAssetRegistryTags(FAssetRegistryTagsContext Context) const override; 
 
+public:
+	// Get public variable defaults property bag
+	const FInstancedPropertyBag& GetPublicVariableDefaults() const { return PublicVariableDefaults; }
+
+private:
+	// URigVMHost interface
+	virtual TArray<FRigVMExternalVariable> GetExternalVariablesImpl(bool bFallbackToBlueprint) const;
+
+protected:
 	// The ExtendedExecuteContext object holds the common work data used by the RigVM internals. It is populated during the initial VM initialization.
 	// Each instance of an AnimGraph requires a copy of this context and a call to initialize the VM instance with the context copy, 
 	// so the cached memory handles are updated to the correct memory addresses.
@@ -41,6 +60,14 @@ protected:
 	UPROPERTY()
 	TObjectPtr<URigVM> RigVM;
 
+	// Variables and their defaults (including public variables, sorted first)
+	UPROPERTY()
+	FInstancedPropertyBag VariableDefaults;
+
+	// Public variables (for easy duplication)
+	UPROPERTY()
+	FInstancedPropertyBag PublicVariableDefaults;
+
 #if WITH_EDITORONLY_DATA
 	UPROPERTY(VisibleAnywhere, Instanced, Category = "Editor Data", meta = (ShowInnerProperties))
 	TObjectPtr<UObject> EditorData;
diff --git a/Engine/Plugins/Experimental/AnimNext/Source/AnimNext/Internal/Component/AnimNextComponent.h b/Engine/Plugins/Experimental/AnimNext/Source/AnimNext/Internal/Component/AnimNextComponent.h
--- a/Engine/Plugins/Experimental/AnimNext/Source/AnimNext/Internal/Component/AnimNextComponent.h
+++ b/Engine/Plugins/Experimental/AnimNext/Source/AnimNext/Internal/Component/AnimNextComponent.h
@@ -2,76 +2,97 @@
 
 #pragma once
 
-#include "AnimNextComponentParameter.h"
+#include "AnimNextPublicVariablesProxy.h"
 #include "Components/ActorComponent.h"
-#include "Scheduler/AnimNextSchedule.h"
-#include "Scheduler/ScheduleHandle.h"
+#include "Module/AnimNextModule.h"
+#include "Module/ModuleHandle.h"
+#include "Param/ParamType.h"
 #include "TraitCore/TraitEvent.h"
+#include "Variables/IAnimNextVariableProxyHost.h"
 
 #include "AnimNextComponent.generated.h"
 
-class UAnimNextSchedule;
 struct FAnimNextComponentInstanceData;
+class UAnimNextComponentWorldSubsystem;
+
+namespace UE::AnimNext
+{
+	struct FProxyVariablesContext;
+};
 
 namespace UE::AnimNext::UncookedOnly
 {
 	struct FUtils;
 }
 
-UENUM()
-enum class EAnimNextParameterScopeOrdering : uint8
-{
-	// Value will be pushed before the scope, allowing the static scope to potentially override the value
-	Before,
-
-	// Value will be pushed after the scope, potentially overriding the static scope
-	After,
-};
-
 UCLASS(MinimalAPI, meta = (BlueprintSpawnableComponent))
-class UAnimNextComponent : public UActorComponent
+class UAnimNextComponent : public UActorComponent, public IAnimNextVariableProxyHost
 {
 	GENERATED_BODY()
 
 	// UActorComponent interface
 	virtual void OnRegister() override;
 	virtual void OnUnregister() override;
+	virtual void BeginPlay() override;
+	virtual void EndPlay(const EEndPlayReason::Type EndPlayReason) override;
+
+	// IAnimNextVariableProxyHost interface
+	virtual void FlipPublicVariablesProxy(const UE::AnimNext::FProxyVariablesContext& InContext) override;
+
+#if WITH_EDITOR
+	// Called back to refresh any cached data on module compilation
+	void OnModuleCompiled();
+#endif
+
+	// (Re-)create the public variable proxy
+	void CreatePublicVariablesProxy();
+
+	// (Re-)create the public variable proxy
+	void DestroyPublicVariablesProxy();
 
 public:
-	// Sets a parameter's value in the supplied scope.
-	// @param    Scope    Scopes corresponding to an existing scope in a schedule, or "None". Passing "None" will apply the parameter to the whole schedule.
-	// @param    Ordering Where to apply the parameter in relation to the supplied scope. Ignored for scope "None".
-	// @param    Name     The name of the parameter to apply
-	// @param    Value    The value to set the parameter to
+	// Sets a module variable's value.
+	// @param    Name     The name of the variable to set
+	// @param    Value    The value to set the variable to
 	UFUNCTION(BlueprintCallable, Category = "AnimNext", CustomThunk, meta = (CustomStructureParam = Value, UnsafeDuringActorConstruction))
-	ANIMNEXT_API void SetParameterInScope(UPARAM(meta = (CustomWidget = "ParamName", AllowedParamType = "FAnimNextScope")) FName Scope, EAnimNextParameterScopeOrdering Ordering, UPARAM(meta = (CustomWidget = "ParamName")) FName Name, int32 Value);
+	ANIMNEXT_API void SetVariable(UPARAM(meta = (CustomWidget = "VariableName")) FName Name, int32 Value);
 
 	// Enable or disable this component's update
 	UFUNCTION(BlueprintCallable, Category = "AnimNext")
-	ANIMNEXT_API void Enable(bool bEnabled);
+	ANIMNEXT_API void SetEnabled(bool bEnabled);
 
 	// Queues an input trait event
 	// Input events will be processed in the next graph update after they are queued
 	ANIMNEXT_API void QueueInputTraitEvent(FAnimNextTraitEventPtr Event);
-	
+
 private:
-	DECLARE_FUNCTION(execSetParameterInScope);
+	DECLARE_FUNCTION(execSetVariable);
 
 private:
 	friend struct UE::AnimNext::UncookedOnly::FUtils;
+	friend class UAnimNextComponentWorldSubsystem;
+
+	// The AnimNext module that this component will run
+	UPROPERTY(EditAnywhere, Category="Module")
+	TObjectPtr<UAnimNextModule> Module = nullptr;
+
+	UPROPERTY(Transient)
+	TObjectPtr<UAnimNextComponentWorldSubsystem> Subsystem = nullptr;
+	
+	// Handle to the registered module
+	UE::AnimNext::FModuleHandle ModuleHandle;
 
-	// The execution schedule that this component will run
-	UPROPERTY(EditAnywhere, Category="Schedule")
-	TObjectPtr<UAnimNextSchedule> Schedule = nullptr;
+	// Lock for public variables proxy
+	FRWLock PublicVariablesLock;
 
-	// Parameters to apply on schedule/component registration
-	UPROPERTY(Instanced, EditAnywhere, Category = "Parameters")
-	TArray<TObjectPtr<UAnimNextComponentParameter>> Parameters;
+	// Proxy public variables
+	UPROPERTY()
+	FAnimNextPublicVariablesProxy PublicVariablesProxy;
 
-	// How to initialize the schedule
-	UPROPERTY(EditAnywhere, Category="Schedule")
-	EAnimNextScheduleInitMethod InitMethod = EAnimNextScheduleInitMethod::InitializeAndPauseInEditor;
+	// Map from name->proxy variable index
+	TMap<FName, int32> PublicVariablesProxyMap;
 
-	// Handle to the registered results/schedule
-	UE::AnimNext::FScheduleHandle SchedulerHandle;
+	// How to initialize the module
+	UPROPERTY(EditAnywhere, Category="Module")
+	EAnimNextModuleInitMethod InitMethod = EAnimNextModuleInitMethod::InitializeAndPauseInEditor;
 };
diff --git a/Engine/Plugins/Experimental/AnimNext/Source/AnimNext/Internal/Component/AnimNextPublicVariablesProxy.h b/Engine/Plugins/Experimental/AnimNext/Source/AnimNext/Internal/Component/AnimNextPublicVariablesProxy.h
new file mode 100644
--- /dev/null
+++ b/Engine/Plugins/Experimental/AnimNext/Source/AnimNext/Internal/Component/AnimNextPublicVariablesProxy.h
@@ -0,0 +1,41 @@
+﻿// Copyright Epic Games, Inc. All Rights Reserved.
+
+#pragma once
+
+#include "StructUtils/PropertyBag.h"
+#include "AnimNextPublicVariablesProxy.generated.h"
+
+// Proxy struct for public variables held on the component
+// TODO: Instead of a dirty flag bitfield (and copying parameters), with some minor refactoring of RigVM external variables
+// we could switch this to use a double-buffered external variables array too, so each write to a variable from the game thread
+// would just write to the current 'GT-side' buffer, then update the external variable ptr to point to the 'latest' version written
+// Then when it comes to swap buffers before WT execution, we just swap the ptrs for the external variables and dont have to copy. 
+USTRUCT()
+struct FAnimNextPublicVariablesProxy
+{
+	GENERATED_BODY()
+
+	void Reset()
+	{
+		Data.Reset();
+		DirtyFlags.Reset();
+		bIsDirty = false;
+	}
+
+	void Empty()
+	{
+		Data.Reset();
+		DirtyFlags.Empty();
+		bIsDirty = false;
+	}
+
+	// Proxy public variables
+	UPROPERTY()
+	FInstancedPropertyBag Data;
+
+	// Dirty flags for each public variable
+	TBitArray<> DirtyFlags;
+
+	// Global dirty flag
+	bool bIsDirty = false;
+};
\ No newline at end of file
diff --git a/Engine/Plugins/Experimental/AnimNext/Source/AnimNext/Internal/Component/AnimNextWorldSubsystem.h b/Engine/Plugins/Experimental/AnimNext/Source/AnimNext/Internal/Component/AnimNextWorldSubsystem.h
new file mode 100644
--- /dev/null
+++ b/Engine/Plugins/Experimental/AnimNext/Source/AnimNext/Internal/Component/AnimNextWorldSubsystem.h
@@ -0,0 +1,100 @@
+﻿// Copyright Epic Games, Inc. All Rights Reserved.
+
+#pragma once
+
+#include "CoreMinimal.h"
+#include "AnimNextPool.h"
+#include "Subsystems/WorldSubsystem.h"
+#include "Module/ModuleHandle.h"
+#include "Module/ModuleTickFunction.h"
+#include "Module/AnimNextModuleInstance.h"
+#include "Module/ModuleTaskContext.h"
+#include "AnimNextWorldSubsystem.generated.h"
+
+class UAnimNextComponent;
+
+namespace UE::AnimNext
+{
+
+// A queued action to complete next frame
+struct FModulePendingAction
+{
+	enum class EType : int32
+	{
+		None = 0,
+		ReleaseHandle,
+		EnableHandle,
+		DisableHandle,
+	};
+
+	FModulePendingAction() = default;
+
+	FModulePendingAction(EType InType, FModuleHandle InHandle);
+
+	FModuleHandle Handle;
+
+	EType Type = EType::None;
+};
+
+}
+
+// Represents AnimNext systems to the gameplay framework
+UCLASS(Abstract)
+class ANIMNEXT_API UAnimNextWorldSubsystem : public UWorldSubsystem
+{
+	GENERATED_BODY()
+
+public:
+	UAnimNextWorldSubsystem();
+	virtual ~UAnimNextWorldSubsystem() override;
+
+	static void AddReferencedObjects(UObject* InThis, FReferenceCollector& Collector);
+
+private:
+	// UWorldSubsystem interface
+	virtual bool DoesSupportWorldType(const EWorldType::Type WorldType) const override;
+
+	bool IsValidHandle(UE::AnimNext::FModuleHandle InHandle) const;
+
+	void FlushPendingActions();
+
+protected:
+#if WITH_EDITOR
+	// Refresh any entries that use the provided module as it has been recompiled.
+	virtual void OnModuleCompiled(UAnimNextModule* InModule);
+#endif
+	
+	// Register a handle to the subsystem
+	void RegisterHandle(UE::AnimNext::FModuleHandle& InOutHandle, UAnimNextModule* InModule, UObject* InObject, EAnimNextModuleInitMethod InitMethod);
+
+	// Unregister a handle from the subsystem
+	void UnregisterHandle(UE::AnimNext::FModuleHandle& InOutHandle);
+
+	// Enables or disables the module represented by the supplied handle
+	// This operation is deferred until the next time the schedule ticks
+	void EnableHandle(UE::AnimNext::FModuleHandle InHandle, bool bInEnabled);
+
+	// Queue a task to run at a particular point in a schedule
+	// @param	InHandle			The handle to execute the task on
+	// @param	InModuleEventName	The name of the event in the module to run the supplied task relative to. If this is NAME_None, then the first valid event will be used.
+	// @param	InTaskFunction		The function to run
+	// @param	InLocation			Where to run the task, before or after
+	void QueueTaskHandle(UE::AnimNext::FModuleHandle InHandle, FName InModuleEventName, TUniqueFunction<void(const UE::AnimNext::FModuleTaskContext&)>&& InTaskFunction, UE::AnimNext::ETaskRunLocation InLocation = UE::AnimNext::ETaskRunLocation::Before);
+
+protected:
+	// Currently running instances, pooled
+	UE::AnimNext::TPool<FAnimNextModuleInstance> Instances;
+
+	// Queued actions
+	TArray<UE::AnimNext::FModulePendingAction> PendingActions;
+
+	// Locks for concurrent modifications
+	FRWLock InstancesLock;
+	FRWLock PendingLock;
+
+	// Handle used to hook into pre-world tick
+	FDelegateHandle OnWorldPreActorTickHandle;
+
+	// Cached delta time
+	float DeltaTime;
+};
\ No newline at end of file
diff --git a/Engine/Plugins/Experimental/AnimNext/Source/AnimNext/Internal/Graph/AnimGraph/AnimNodeExposedValueHandler_AnimNextParameters.h b/Engine/Plugins/Experimental/AnimNext/Source/AnimNext/Internal/Graph/AnimGraph/AnimNodeExposedValueHandler_AnimNextParameters.h
--- a/Engine/Plugins/Experimental/AnimNext/Source/AnimNext/Internal/Graph/AnimGraph/AnimNodeExposedValueHandler_AnimNextParameters.h
+++ b/Engine/Plugins/Experimental/AnimNext/Source/AnimNext/Internal/Graph/AnimGraph/AnimNodeExposedValueHandler_AnimNextParameters.h
@@ -4,9 +4,7 @@
 
 #include "Animation/ExposedValueHandler.h"
 #include "PropertyAccess.h"
-#include "Param/ParamId.h"
 #include "Param/ParamType.h"
-#include "Param/ParamTypeHandle.h"
 #include "AnimNodeExposedValueHandler_AnimNextParameters.generated.h"
 
 USTRUCT()
@@ -18,16 +16,10 @@ struct FAnimNodeExposedValueHandler_AnimNextParameters_Entry
 	UPROPERTY()
 	FName ParameterName;
 
-	// Cached ID for the parameter
-	UE::AnimNext::FParamId ParamId;
-
 	// ParamType of the property that this entry will write to
 	UPROPERTY()
 	FAnimNextParamType PropertyParamType;
 
-	// Cached type handle created from ParamType
-	UE::AnimNext::FParamTypeHandle PropertyParamTypeHandle;
-
 	// Property access index
 	UPROPERTY()
 	int32 AccessIndex = INDEX_NONE;
diff --git a/Engine/Plugins/Experimental/AnimNext/Source/AnimNext/Internal/Graph/AnimNextAnimationGraph.h b/Engine/Plugins/Experimental/AnimNext/Source/AnimNext/Internal/Graph/AnimNextAnimationGraph.h
new file mode 100644
--- /dev/null
+++ b/Engine/Plugins/Experimental/AnimNext/Source/AnimNext/Internal/Graph/AnimNextAnimationGraph.h
@@ -0,0 +1,159 @@
+﻿// Copyright Epic Games, Inc. All Rights Reserved.
+
+#pragma once
+
+#include "CoreMinimal.h"
+#include "Graph/AnimNextGraphEntryPoint.h"
+#include "AnimNextRigVMAsset.h"
+#include "Graph/RigUnit_AnimNextGraphRoot.h"
+#include "RigVMCore/RigVM.h"
+#include "TraitCore/TraitPtr.h"
+#include "TraitCore/TraitHandle.h"
+#include "TraitCore/EntryPointHandle.h"
+#include "Graph/RigUnit_AnimNextGraphEvaluator.h"
+#include "RigVMHost.h"
+
+#include "AnimNextAnimationGraph.generated.h"
+
+class UEdGraph;
+class UAnimNextModule;
+class UAnimGraphNode_AnimNextGraph;
+struct FAnimNode_AnimNextGraph;
+struct FRigUnit_AnimNextGraphEvaluator;
+struct FAnimNextGraphInstancePtr;
+struct FAnimNextGraphInstance;
+struct FAnimNextScheduleGraphTask;
+struct FAnimNextEditorParam;
+struct FAnimNextParam;
+struct FAnimNextModuleInstance;
+
+namespace UE::AnimNext
+{
+	struct FContext;
+	struct FExecutionContext;
+	class FAnimNextModuleImpl;
+	struct FTestUtils;
+	struct FParametersProxy;
+	struct FPlayAnimSlotTrait;
+}
+
+namespace UE::AnimNext::UncookedOnly
+{
+	struct FUtils;
+}
+
+namespace UE::AnimNext::Editor
+{
+	class FModuleEditor;
+	class FVariableCustomization;
+}
+
+namespace UE::AnimNext::Graph
+{
+	extern ANIMNEXT_API const FName EntryPointName;
+	extern ANIMNEXT_API const FName ResultName;
+}
+
+// A user-created collection of animation logic & data
+UCLASS(BlueprintType)
+class ANIMNEXT_API UAnimNextAnimationGraph : public UAnimNextRigVMAsset
+{
+	GENERATED_BODY()
+
+public:
+	UAnimNextAnimationGraph(const FObjectInitializer& ObjectInitializer);
+
+	// UObject interface
+	virtual void Serialize(FArchive& Ar) override;
+
+	// Allocates an instance of the graph
+	// @param	OutInstance			The instance to allocate data for
+	// @param	InModuleInstance	The module instance to use
+	// @param	InEntryPoint		The entry point to use. If this is NAME_None then the default entry point for this graph is used
+	void AllocateInstance(FAnimNextGraphInstancePtr& OutInstance, FAnimNextModuleInstance* InModuleInstance = nullptr, FName InEntryPoint = NAME_None) const;
+
+	// Allocates an instance of the graph with the specified parent graph instance
+	// @param	InParentGraphInstance		The parent graph instance to use
+	// @param	OutInstance					The instance to allocate data for
+	// @param	InEntryPoint				The entry point to use. If this is NAME_None then the default entry point for this graph is used
+	void AllocateInstance(FAnimNextGraphInstance& InParentGraphInstance, FAnimNextGraphInstancePtr& OutInstance, FName InEntryPoint = NAME_None) const;
+
+protected:
+
+	// Loads the graph data from the provided archive buffer and returns true on success, false otherwise
+	bool LoadFromArchiveBuffer(const TArray<uint8>& SharedDataArchiveBuffer);
+
+	// Allocates an instance of the graph with an optional parent graph instance
+	void AllocateInstanceImpl(FAnimNextModuleInstance* InModuleInstance, FAnimNextGraphInstance* InParentGraphInstance, FAnimNextGraphInstancePtr& OutInstance, FName InEntryPoint) const;
+
+#if WITH_EDITORONLY_DATA
+	// During graph compilation, if we have existing graph instances, we freeze them by releasing their memory before thawing them
+	// Freezing is a partial release of resources that retains the necessary information to re-create things safely
+	void FreezeGraphInstances();
+
+	// During graph compilation, once compilation is done we thaw existing graph instances to reallocate their memory
+	void ThawGraphInstances();
+#endif
+	
+	friend class UAnimNextAnimationGraphFactory;
+	friend class UAnimNextAnimationGraph_EditorData;
+	friend class UAnimNextVariableEntry;
+	friend struct UE::AnimNext::UncookedOnly::FUtils;
+	friend class UE::AnimNext::Editor::FModuleEditor;
+	friend struct UE::AnimNext::FTestUtils;
+	friend FAnimNextGraphInstancePtr;
+	friend FAnimNextGraphInstance;
+	friend class UAnimGraphNode_AnimNextGraph;
+	friend UE::AnimNext::FExecutionContext;
+	friend struct FAnimNextScheduleGraphTask;
+	friend UE::AnimNext::FAnimNextModuleImpl;
+	friend class UE::AnimNext::Editor::FVariableCustomization;
+	friend struct UE::AnimNext::FParametersProxy;
+	friend struct UE::AnimNext::FPlayAnimSlotTrait;
+
+#if WITH_EDITORONLY_DATA
+	mutable FCriticalSection GraphInstancesLock;
+
+	// This is a list of live graph instances that have been allocated, used in the editor to reset instances when we re-compile/live edit
+	mutable TSet<FAnimNextGraphInstance*> GraphInstances;
+#endif
+
+	// This is the execute method definition used by a graph to evaluate latent pins
+	UPROPERTY()
+	FAnimNextGraphEvaluatorExecuteDefinition ExecuteDefinition;
+
+	// Data for each entry point in this graph
+	UPROPERTY()
+	TArray<FAnimNextGraphEntryPoint> EntryPoints;
+
+	// This is a resolved handle to the root trait in our graph, for each entry point 
+	TMap<FName, FAnimNextTraitHandle> ResolvedRootTraitHandles;
+
+	// This is an index into EntryPoints, for each entry point
+	TMap<FName, int32> ResolvedEntryPoints;
+
+	// This is the graph shared data used by the trait system, the output of FTraitReader
+	// We de-serialize manually into this buffer from the archive buffer, this is never saved on disk
+	TArray<uint8> SharedDataBuffer;
+
+	// This is a list of all referenced UObjects in the graph shared data
+	// We collect all the references here to make it quick and easy for the GC to query them
+	// It means that object references in the graph shared data are not visited at runtime by the GC (they are immutable)
+	// The shared data serialization archive stores indices to these to perform UObject serialization
+	UPROPERTY()
+	TArray<TObjectPtr<UObject>> GraphReferencedObjects;
+
+	// The entry point that this graph defaults to using
+	UPROPERTY(EditAnywhere, Category = "Graph")
+	FName DefaultEntryPoint = FRigUnit_AnimNextGraphRoot::DefaultEntryPoint;
+
+	// Default state for this graph
+	UPROPERTY()
+	FAnimNextGraphState DefaultState;
+
+#if WITH_EDITORONLY_DATA
+	// This buffer holds the output of the FTraitWriter post compilation
+	// We serialize it manually and it is discarded at runtime
+	TArray<uint8> SharedDataArchiveBuffer;
+#endif
+};
diff --git a/Engine/Plugins/Experimental/AnimNext/Source/AnimNext/Internal/Graph/AnimNextGraphContextData.h b/Engine/Plugins/Experimental/AnimNext/Source/AnimNext/Internal/Graph/AnimNextGraphContextData.h
--- a/Engine/Plugins/Experimental/AnimNext/Source/AnimNext/Internal/Graph/AnimNextGraphContextData.h
+++ b/Engine/Plugins/Experimental/AnimNext/Source/AnimNext/Internal/Graph/AnimNextGraphContextData.h
@@ -3,6 +3,7 @@
 #pragma once
 
 #include "CoreMinimal.h"
+#include "Module/AnimNextModuleContextData.h"
 #include "AnimNextGraphContextData.generated.h"
 
 struct FAnimNextGraphInstance;
@@ -13,14 +14,15 @@ namespace UE::AnimNext
 }
 
 USTRUCT()
-struct FAnimNextGraphContextData
+struct FAnimNextGraphContextData : public FAnimNextModuleContextData
 {
 	GENERATED_BODY()
 
 	FAnimNextGraphContextData() = default;
 
-	FAnimNextGraphContextData(const FAnimNextGraphInstance* InInstance, const TConstArrayView<UE::AnimNext::FLatentPropertyHandle>& InLatentHandles, void* InDestinationBasePtr, bool bInIsFrozen)
-		: Instance(InInstance)
+	FAnimNextGraphContextData(FAnimNextModuleInstance* InModuleInstance, const FAnimNextGraphInstance* InInstance, const TConstArrayView<UE::AnimNext::FLatentPropertyHandle>& InLatentHandles, void* InDestinationBasePtr, bool bInIsFrozen)
+		: FAnimNextModuleContextData(InModuleInstance)
+		, Instance(InInstance)
 		, LatentHandles(InLatentHandles)
 		, DestinationBasePtr(InDestinationBasePtr)
 		, bIsFrozen(bInIsFrozen)
diff --git a/Engine/Plugins/Experimental/AnimNext/Source/AnimNext/Internal/Graph/RigUnit_AnimNextGraphEvaluator.h b/Engine/Plugins/Experimental/AnimNext/Source/AnimNext/Internal/Graph/RigUnit_AnimNextGraphEvaluator.h
--- a/Engine/Plugins/Experimental/AnimNext/Source/AnimNext/Internal/Graph/RigUnit_AnimNextGraphEvaluator.h
+++ b/Engine/Plugins/Experimental/AnimNext/Source/AnimNext/Internal/Graph/RigUnit_AnimNextGraphEvaluator.h
@@ -41,7 +41,7 @@ struct FAnimNextGraphEvaluatorExecuteDefinition
  * It performs the animation graph update and evaluation through the data provided in the execution context.
  * It also holds all latent/lazy pins that the graph references in the editor.
  */
-USTRUCT(meta=(DisplayName="Animation Runtime Output", Category="Events", NodeColor="1, 0, 0"))
+USTRUCT(meta=(Hidden, DisplayName="Animation Runtime Output", Category="Events", NodeColor="1, 0, 0"))
 struct ANIMNEXT_API FRigUnit_AnimNextGraphEvaluator : public FRigUnit_AnimNextBase
 {
 	GENERATED_BODY()
diff --git a/Engine/Plugins/Experimental/AnimNext/Source/AnimNext/Internal/Graph/RigUnit_AnimNextGraphRoot.h b/Engine/Plugins/Experimental/AnimNext/Source/AnimNext/Internal/Graph/RigUnit_AnimNextGraphRoot.h
--- a/Engine/Plugins/Experimental/AnimNext/Source/AnimNext/Internal/Graph/RigUnit_AnimNextGraphRoot.h
+++ b/Engine/Plugins/Experimental/AnimNext/Source/AnimNext/Internal/Graph/RigUnit_AnimNextGraphRoot.h
@@ -14,7 +14,7 @@
  * to evaluate.
  * This node isn't used at runtime.
  */
-USTRUCT(meta=(DisplayName="Animation Output", Category="Events", NodeColor="1, 0, 0", Keywords="Root,Output"))
+USTRUCT(meta=(Hidden, DisplayName="Animation Output", Category="Events", NodeColor="1, 0, 0", Keywords="Root,Output"))
 struct ANIMNEXT_API FRigUnit_AnimNextGraphRoot : public FRigUnit_AnimNextBase
 {
 	GENERATED_BODY()
diff --git a/Engine/Plugins/Experimental/AnimNext/Source/AnimNext/Internal/Graph/RigUnit_AnimNextRunAnimationGraph.h b/Engine/Plugins/Experimental/AnimNext/Source/AnimNext/Internal/Graph/RigUnit_AnimNextRunAnimationGraph.h
new file mode 100644
--- /dev/null
+++ b/Engine/Plugins/Experimental/AnimNext/Source/AnimNext/Internal/Graph/RigUnit_AnimNextRunAnimationGraph.h
@@ -0,0 +1,44 @@
+// Copyright Epic Games, Inc. All Rights Reserved.
+
+#pragma once
+
+#include "AnimNextExecuteContext.h"
+#include "Graph/AnimNext_LODPose.h"
+#include "Graph/RigUnit_AnimNextBase.h"
+#include "Graph/AnimNextGraphInstancePtr.h"
+
+#include "RigUnit_AnimNextRunAnimationGraph.generated.h"
+
+/** Runs an animation graph */
+USTRUCT(meta=(DisplayName="Run Graph", Category="Animation Graph", NodeColor="0, 1, 1", Keywords="Trait,Stack"))
+struct ANIMNEXT_API FRigUnit_AnimNextRunAnimationGraph : public FRigUnit_AnimNextBase
+{
+	GENERATED_BODY()
+
+	RIGVM_METHOD()
+	void Execute();
+
+	// Graph to run
+	UPROPERTY(EditAnywhere, Category = "Graph", meta = (Input))
+	TObjectPtr<UAnimNextAnimationGraph> Graph;
+
+	// Instance used to hold graph state
+	UPROPERTY(EditAnywhere, Category = "Graph", meta = (Input, Output))
+	FAnimNextGraphInstancePtr Instance;
+
+	// LOD to run the graph at
+	UPROPERTY(EditAnywhere, Category = "Graph", meta = (Input))
+	int32 LOD = 0;
+
+	// Reference pose for the graph
+	UPROPERTY(EditAnywhere, Category = "Graph", meta = (Input))
+	FAnimNextGraphReferencePose ReferencePose;
+
+	// Pose result
+	UPROPERTY(EditAnywhere, Category = "Graph", meta = (Output))
+	FAnimNextGraphLODPose Result;
+
+	// The execution result
+	UPROPERTY(EditAnywhere, DisplayName = "Execute", Category = "BeginExecution", meta = (Input, Output))
+	FAnimNextExecuteContext ExecuteContext;
+};
diff --git a/Engine/Plugins/Experimental/AnimNext/Source/AnimNext/Internal/Graph/RigUnit_AnimNextShimRoot.h b/Engine/Plugins/Experimental/AnimNext/Source/AnimNext/Internal/Graph/RigUnit_AnimNextShimRoot.h
--- a/Engine/Plugins/Experimental/AnimNext/Source/AnimNext/Internal/Graph/RigUnit_AnimNextShimRoot.h
+++ b/Engine/Plugins/Experimental/AnimNext/Source/AnimNext/Internal/Graph/RigUnit_AnimNextShimRoot.h
@@ -12,7 +12,7 @@
  * The graph editor will not see this node. It is added during compilation as a shim to start the evaluation process.
  * This node is only used at runtime.
  */
-USTRUCT(meta=(DisplayName="Animation Output Shim", Category="Events", NodeColor="1, 0, 0", Keywords="Root,Output"))
+USTRUCT(meta=(Hidden, DisplayName="Animation Output Shim", Category="Events", NodeColor="1, 0, 0", Keywords="Root,Output"))
 struct ANIMNEXT_API FRigUnit_AnimNextShimRoot : public FRigUnit_AnimNextBase
 {
 	GENERATED_BODY()
diff --git a/Engine/Plugins/Experimental/AnimNext/Source/AnimNext/Internal/Graph/RigVMTrait_AnimNextPublicVariables.h b/Engine/Plugins/Experimental/AnimNext/Source/AnimNext/Internal/Graph/RigVMTrait_AnimNextPublicVariables.h
new file mode 100644
--- /dev/null
+++ b/Engine/Plugins/Experimental/AnimNext/Source/AnimNext/Internal/Graph/RigVMTrait_AnimNextPublicVariables.h
@@ -0,0 +1,30 @@
+﻿// Copyright Epic Games, Inc. All Rights Reserved.
+
+#pragma once
+
+#include "RigVMCore/RigVMTrait.h"
+#include "RigVMTrait_AnimNextPublicVariables.generated.h"
+
+class UAnimNextRigVMAsset;
+
+// Represents public variables of an asset via a trait 
+USTRUCT(BlueprintType)
+struct ANIMNEXT_API FRigVMTrait_AnimNextPublicVariables : public FRigVMTrait
+{
+	GENERATED_BODY()
+
+	// The asset that any programmatic pins will be derived from
+	UPROPERTY(meta = (Hidden))
+	TObjectPtr<UAnimNextRigVMAsset> Asset = nullptr;
+
+	// Variable names that are exposed
+	UPROPERTY(meta = (Hidden))
+	TArray<FName> VariableNames;
+
+	// FRigVMTrait interface
+#if WITH_EDITOR
+	virtual FString GetDisplayName() const override;
+	virtual void GetProgrammaticPins(URigVMController* InController, int32 InParentPinIndex, const FString& InDefaultValue, FRigVMPinInfoArray& OutPinArray) const override;
+	virtual bool ShouldCreatePinForProperty(const FProperty* InProperty) const override;
+#endif
+};
diff --git a/Engine/Plugins/Experimental/AnimNext/Source/AnimNext/Internal/Module/AnimNextModule.h b/Engine/Plugins/Experimental/AnimNext/Source/AnimNext/Internal/Module/AnimNextModule.h
--- a/Engine/Plugins/Experimental/AnimNext/Source/AnimNext/Internal/Module/AnimNextModule.h
+++ b/Engine/Plugins/Experimental/AnimNext/Source/AnimNext/Internal/Module/AnimNextModule.h
@@ -5,15 +5,8 @@
 #include "CoreMinimal.h"
 #include "Graph/AnimNextGraphEntryPoint.h"
 #include "AnimNextRigVMAsset.h"
-#include "Graph/RigUnit_AnimNextGraphRoot.h"
 #include "RigVMCore/RigVM.h"
 #include "TraitCore/TraitPtr.h"
-#include "TraitCore/TraitHandle.h"
-#include "TraitCore/EntryPointHandle.h"
-#include "Graph/GraphInstanceComponent.h"
-#include "Graph/RigUnit_AnimNextGraphEvaluator.h"
-#include "Param/ParamId.h"
-#include "Scheduler/IAnimNextScheduleTermInterface.h"
 #include "RigVMHost.h"
 
 #include "AnimNextModule.generated.h"
@@ -25,7 +18,6 @@ struct FAnimNode_AnimNextGraph;
 struct FRigUnit_AnimNextGraphEvaluator;
 struct FAnimNextGraphInstancePtr;
 struct FAnimNextGraphInstance;
-class UAnimNextSchedule;
 struct FAnimNextScheduleGraphTask;
 struct FAnimNextEditorParam;
 struct FAnimNextParam;
@@ -47,18 +39,28 @@ namespace UE::AnimNext::UncookedOnly
 namespace UE::AnimNext::Editor
 {
 	class FModuleEditor;
-	class FParameterCustomization;
+	class FVariableCustomization;
 }
 
-namespace UE::AnimNext::Graph
+UENUM()
+enum class EAnimNextModuleInitMethod : uint8
 {
-	extern ANIMNEXT_API const FName EntryPointName;
-	extern ANIMNEXT_API const FName ResultName;
-}
+	// Do not perform any initial update, set up data structures only
+	None,
+
+	// Set up data structures, perform an initial update and then pause
+	InitializeAndPause,
+
+	// Set up data structures, perform an initial update and then pause in editor only, otherwise act like InitializeAndRun
+	InitializeAndPauseInEditor,
+
+	// Set up data structures then continue updating
+	InitializeAndRun
+};
 
-// A user-created collection of animation logic & data
+// Root asset represented by a component when instantiated
 UCLASS(BlueprintType)
-class ANIMNEXT_API UAnimNextModule : public UAnimNextRigVMAsset, public IAnimNextScheduleTermInterface
+class ANIMNEXT_API UAnimNextModule : public UAnimNextRigVMAsset
 {
 	GENERATED_BODY()
 
@@ -68,56 +70,16 @@ public:
 	// UObject interface
 	virtual void Serialize(FArchive& Ar) override;
 	virtual void PostLoad() override;
-#if WITH_EDITOR
-	virtual void PostEditChangeProperty(FPropertyChangedEvent& PropertyChangedEvent) override;
-#endif
-
-	// IAnimNextScheduleTermInterface interface
-	virtual TConstArrayView<UE::AnimNext::FScheduleTerm> GetTerms() const override;
-
-	// Allocates an instance of the graph
-	// @param	OutInstance		The instance to allocate data for
-	// @param	InEntryPoint	The entry point to use. If this is NAME_None then the default entry point for this graph is used
-	void AllocateInstance(FAnimNextGraphInstancePtr& OutInstance, FName InEntryPoint = NAME_None) const;
-
-	// Allocates an instance of the graph with the specified parent graph instance
-	// @param	InOutParentGraphInstance	The parent graph instance to use
-	// @param	OutInstance					The instance to allocate data for
-	// @param	InEntryPoint				The entry point to use. If this is NAME_None then the default entry point for this graph is used
-	void AllocateInstance(FAnimNextGraphInstance& InOutParentGraphInstance, FAnimNextGraphInstancePtr& OutInstance, FName InEntryPoint = NAME_None) const;
-
-	// Update the parameter layer, if any
-	void UpdateLayer(UE::AnimNext::FParamStackLayerHandle& InHandle, float InDeltaTime) const;
-
-	// Get the fully-qualified default entry point name (/Path/To/Asset.Asset:EntryPoint)
-	FName GetDefaultEntryPoint() const;
-
-protected:
-
-	// Loads the graph data from the provided archive buffer and returns true on success, false otherwise
-	bool LoadFromArchiveBuffer(const TArray<uint8>& SharedDataArchiveBuffer);
-
-	// Allocates an instance of the graph with an optional parent graph instance
-	void AllocateInstanceImpl(FAnimNextGraphInstance* InOutParentGraphInstance, FAnimNextGraphInstancePtr& OutInstance, FName InEntryPoint) const;
 
 #if WITH_EDITORONLY_DATA
-	// During graph compilation, if we have existing graph instances, we freeze them by releasing their memory before thawing them
-	// Freezing is a partial release of resources that retains the necessary information to re-create things safely
-	void FreezeGraphInstances();
-
-	// During graph compilation, once compilation is done we thaw existing graph instances to reallocate their memory
-	void ThawGraphInstances();
+	// Delegate called in editor when a module is compiled
+	using FOnModuleCompiled = TTSMulticastDelegate<void(UAnimNextModule*)>;
+	static FOnModuleCompiled& OnModuleCompiled();
 #endif
 
-	// Set the default entry point name (unqualified)
-	void SetDefaultEntryPoint(FName InEntryPoint);
-
-	// Cache the default entry point
-	void CacheDefaultEntryPoint() const;
-
+protected:
 	friend class UAnimNextModuleFactory;
 	friend class UAnimNextModule_EditorData;
-	friend class UAnimNextModule_Parameter;
 	friend struct UE::AnimNext::UncookedOnly::FUtils;
 	friend class UE::AnimNext::Editor::FModuleEditor;
 	friend struct UE::AnimNext::FTestUtils;
@@ -125,70 +87,16 @@ protected:
 	friend FAnimNextGraphInstance;
 	friend class UAnimGraphNode_AnimNextGraph;
 	friend UE::AnimNext::FExecutionContext;
-	friend class UAnimNextSchedule;
 	friend struct FAnimNextScheduleGraphTask;
 	friend UE::AnimNext::FAnimNextModuleImpl;
-	friend class UE::AnimNext::Editor::FParameterCustomization;
+	friend class UE::AnimNext::Editor::FVariableCustomization;
 	friend struct UE::AnimNext::FParametersProxy;
 
 #if WITH_EDITORONLY_DATA
-	mutable FCriticalSection GraphInstancesLock;
-
-	// This is a list of live graph instances that have been allocated, used in the editor to reset instances when we re-compile/live edit
-	mutable TSet<FAnimNextGraphInstance*> GraphInstances;
-#endif
-
-	// This is the execute method definition used by a graph to evaluate latent pins
-	UPROPERTY()
-	FAnimNextGraphEvaluatorExecuteDefinition ExecuteDefinition;
-
-	// Data for each entry point in this graph
-	UPROPERTY()
-	TArray<FAnimNextGraphEntryPoint> EntryPoints;
-
-	// This is a resolved handle to the root trait in our graph, for each entry point 
-	TMap<FName, FAnimNextTraitHandle> ResolvedRootTraitHandles;
-
-	// This is an index into EntryPoints, for each entry point
-	TMap<FName, int32> ResolvedEntryPoints;
-
-	// This is the graph shared data used by the trait system, the output of FTraitReader
-	// We de-serialize manually into this buffer from the archive buffer, this is never saved on disk
-	TArray<uint8> SharedDataBuffer;
-
-	// This is a list of all referenced UObjects in the graph shared data
-	// We collect all the references here to make it quick and easy for the GC to query them
-	// It means that object references in the graph shared data are not visited at runtime by the GC (they are immutable)
-	// The shared data serialization archive stores indices to these to perform UObject serialization
-	UPROPERTY()
-	TArray<TObjectPtr<UObject>> GraphReferencedObjects;
-
-	// The entry point that this graph defaults to using (unqualified by asset). Use GetDefaultEntryPoint to get the fully-qualified name.
-	UPROPERTY(EditAnywhere, Category = "Graph")
-	FName DefaultEntryPoint = FRigUnit_AnimNextGraphRoot::DefaultEntryPoint;
-
-	// Cached fully-qualified entry point, set via by SetDefaultEntryPoint/GetDefaultEntryPoint.
-	UPROPERTY(Transient)
-	mutable FName CachedDefaultEntryPoint = NAME_None;
-
-	// Hash of required parameters
-	UPROPERTY()
-	uint64 RequiredParametersHash = 0;
-
-	// All the parameters that are required for this graph to run
 	UPROPERTY()
-	TArray<FAnimNextParam> RequiredParameters;
-
-	// Default state for this graph
-	UPROPERTY()
-	FAnimNextGraphState DefaultState;
-
-#if WITH_EDITORONLY_DATA
+	FAnimNextGraphState DefaultState_DEPRECATED;
+	
 	UPROPERTY()
 	FInstancedPropertyBag PropertyBag_DEPRECATED;
-
-	// This buffer holds the output of the FTraitWriter post compilation
-	// We serialize it manually and it is discarded at runtime
-	TArray<uint8> SharedDataArchiveBuffer;
 #endif
 };
diff --git a/Engine/Plugins/Experimental/AnimNext/Source/AnimNext/Internal/Module/AnimNextModuleContextData.h b/Engine/Plugins/Experimental/AnimNext/Source/AnimNext/Internal/Module/AnimNextModuleContextData.h
new file mode 100644
--- /dev/null
+++ b/Engine/Plugins/Experimental/AnimNext/Source/AnimNext/Internal/Module/AnimNextModuleContextData.h
@@ -0,0 +1,40 @@
+// Copyright Epic Games, Inc. All Rights Reserved.
+
+#pragma once
+
+#include "CoreMinimal.h"
+#include "AnimNextModuleContextData.generated.h"
+
+struct FAnimNextModuleInstance;
+
+USTRUCT()
+struct FAnimNextModuleContextData
+{
+	GENERATED_BODY()
+
+	FAnimNextModuleContextData() = default;
+
+	FAnimNextModuleContextData(FAnimNextModuleInstance* InModuleInstance)
+		: ModuleInstance(InModuleInstance)
+	{
+	}
+
+	FAnimNextModuleInstance& GetModuleInstance() const
+	{
+		check(ModuleInstance != nullptr);
+		return *ModuleInstance;
+	}
+
+private:
+	// Call this to reset the context to its original state to detect stale usage
+	void Reset()
+	{
+		ModuleInstance = nullptr;
+	}
+
+	// Module instance that is currently executing
+	FAnimNextModuleInstance* ModuleInstance = nullptr;
+
+	friend class UAnimNextModule;
+	friend struct FAnimNextExecuteContext;
+};
diff --git a/Engine/Plugins/Experimental/AnimNext/Source/AnimNext/Internal/Module/AnimNextModuleInstance.h b/Engine/Plugins/Experimental/AnimNext/Source/AnimNext/Internal/Module/AnimNextModuleInstance.h
new file mode 100644
--- /dev/null
+++ b/Engine/Plugins/Experimental/AnimNext/Source/AnimNext/Internal/Module/AnimNextModuleInstance.h
@@ -0,0 +1,160 @@
+// Copyright Epic Games, Inc. All Rights Reserved.
+
+#pragma once
+
+#include "CoreMinimal.h"
+#include "AnimNextModule.h"
+#include "ModuleHandle.h"
+#include "Component/AnimNextPublicVariablesProxy.h"
+#include "Module/ModuleTickFunction.h"
+#include "TraitCore/TraitEvent.h"
+#include "TraitCore/TraitEventList.h"
+#include "Module/ModuleHandle.h"
+#include "AnimNextModuleInstance.generated.h"
+
+struct FAnimNextModuleInstance;
+
+namespace UE::AnimNext
+{
+	struct FProxyVariablesContext;
+}
+
+namespace UE::AnimNext
+{
+
+// Phase is used as a general ordering constraint on event execution
+enum class EModuleEventPhase : uint8
+{
+	// Before any execution, e.g. for copying data from the game thread
+	PreExecute,
+
+	// Any initial setup work for new instances
+	Initialize,
+
+	// General execution, e.g. a prephysics event
+	Execute,
+};
+
+}
+
+namespace UE::AnimNext::Private
+{
+	// Pre-caches all possible module event info for events that have been predefined in code  
+	extern void CacheAllModuleEvents();
+}
+
+// Root memory owner of a parameterized schedule 
+USTRUCT()
+struct FAnimNextModuleInstance
+{
+	GENERATED_BODY()
+
+	FAnimNextModuleInstance() = default;
+	FAnimNextModuleInstance(
+		UAnimNextModule* InModule,
+		UObject* InObject,
+		EAnimNextModuleInitMethod InInitMethod);
+	~FAnimNextModuleInstance();
+
+	// Setup the entry
+	void Initialize();
+
+	// Clear data that binds the schedule to a runtime (e.g. tick functions) and any instance data
+	void ResetBindingsAndInstanceData();
+
+	// Used for pooling
+	void Invalidate();
+
+	// Enables/disables the ticking of this entry
+	void Enable(bool bInEnabled);
+
+	// Clears the bTickEvenWhenPaused flags of the entries' tick functions
+	void ClearTickFunctionPauseFlags();
+
+	// Queues an input trait event
+	// Input events will be processed in the next graph update after they are queued
+	void QueueInputTraitEvent(FAnimNextTraitEventPtr Event);
+
+	// Flip proxy variable buffers then copy any dirty values
+	void CopyProxyVariables();
+
+#if WITH_EDITOR
+	// Resets internal state if the module we are bound to is recompiled in editor
+	void OnModuleCompiled();
+#endif
+
+	UPROPERTY(Transient)
+	TObjectPtr<UAnimNextModule> Module = nullptr;
+
+	// Object this entry is bound to
+	UPROPERTY(Transient)
+	TObjectPtr<UObject> Object = nullptr;
+
+	// Copy of the handle that represents this entry to client systems
+	UE::AnimNext::FModuleHandle Handle;
+
+	// End tick function used to perform internal bookkeeping
+	UE::AnimNext::FModuleEndTickFunction EndTickFunction;
+
+	// Pre-allocated graph of tick functions
+	TArray<UE::AnimNext::FModuleEventTickFunction> TickFunctions;
+
+	// Input event list to be processed on the next update
+	UE::AnimNext::FTraitEventList InputEventList;
+
+	// Output event list to be processed at the end of the schedule tick
+	UE::AnimNext::FTraitEventList OutputEventList;
+
+	// Lock to ensure event list actions are thread safe
+	FRWLock EventListLock;
+
+	// Execute context for the RigVM instance to execute with
+	FRigVMExtendedExecuteContext ExtendedExecuteContext;
+
+	// Variables state
+	FInstancedPropertyBag Variables;
+
+	// Proxy public variables
+	FAnimNextPublicVariablesProxy PublicVariablesProxy;
+
+	enum class ERunState : uint8
+	{
+		None,
+
+		CreatingTasks,
+
+		BindingTasks,
+
+		PendingInitialUpdate,
+
+		Running,
+
+		Paused,
+	};
+
+	// Current running state
+	ERunState RunState = ERunState::None;
+
+	// Transition to the specified run state, verifying that the current state is valid
+	void TransitionToRunState(ERunState InNewState);
+	
+	// How this entry initializes
+	EAnimNextModuleInitMethod InitMethod = EAnimNextModuleInitMethod::InitializeAndPauseInEditor;
+
+	// Whether this represents an editor object 
+	bool bIsEditor : 1 = false;
+
+#if WITH_EDITOR
+	// Whether we are currently recreating this instance because of compilation/reinstancing
+	bool bIsRecreatingOnCompile : 1 = false;
+#endif
+};
+
+template<>
+struct TStructOpsTypeTraits<FAnimNextModuleInstance> : public TStructOpsTypeTraitsBase2<FAnimNextModuleInstance>
+{
+	enum
+	{
+		WithCopy = false
+	};
+};
diff --git a/Engine/Plugins/Experimental/AnimNext/Source/AnimNext/Internal/Module/ModuleEventTickFunctionBindings.h b/Engine/Plugins/Experimental/AnimNext/Source/AnimNext/Internal/Module/ModuleEventTickFunctionBindings.h
new file mode 100644
--- /dev/null
+++ b/Engine/Plugins/Experimental/AnimNext/Source/AnimNext/Internal/Module/ModuleEventTickFunctionBindings.h
@@ -0,0 +1,40 @@
+﻿// Copyright Epic Games, Inc. All Rights Reserved.
+
+#pragma once
+
+#include "CoreMinimal.h"
+#include "Templates/Function.h"
+
+struct FTickFunction;
+struct FAnimNextModuleInstance;
+
+namespace UE::AnimNext
+{
+
+// Binding context for tick functions
+struct FTickFunctionBindingContext
+{
+	FTickFunctionBindingContext(FAnimNextModuleInstance& InModuleInstance, UObject* InObject, UWorld* InWorld, int32 InEventIndex)
+		: ModuleInstance(InModuleInstance)
+		, Object(InObject)
+		, World(InWorld)
+		, EventIndex(InEventIndex)
+	{}
+
+	// The module instance
+	FAnimNextModuleInstance& ModuleInstance;
+
+	// Object that a module is bound on (e.g. a component)
+	UObject* Object = nullptr;
+
+	// World that the object is contained within
+	UWorld* World = nullptr;
+
+	// Index into the module's events array
+	int32 EventIndex = INDEX_NONE;
+};
+
+// Function called to bind a tick function to a module event, set up prerequisites etc.
+using FModuleEventBindingFunction = TFunction<void(const FTickFunctionBindingContext& InContext, FTickFunction& InTickFunction)>;
+
+}
diff --git a/Engine/Plugins/Experimental/AnimNext/Source/AnimNext/Internal/Module/ModuleHandle.h b/Engine/Plugins/Experimental/AnimNext/Source/AnimNext/Internal/Module/ModuleHandle.h
new file mode 100644
--- /dev/null
+++ b/Engine/Plugins/Experimental/AnimNext/Source/AnimNext/Internal/Module/ModuleHandle.h
@@ -0,0 +1,20 @@
+// Copyright Epic Games, Inc. All Rights Reserved.
+
+#pragma once
+
+#include "AnimNextPool.h"
+
+struct FAnimNextModuleInstance;
+
+namespace UE::AnimNext
+{
+	template<typename T> struct TPoolHandle;
+}
+
+namespace UE::AnimNext
+{
+
+// Opaque handle representing a module instance
+using FModuleHandle = TPoolHandle<FAnimNextModuleInstance>;
+
+}
\ No newline at end of file
diff --git a/Engine/Plugins/Experimental/AnimNext/Source/AnimNext/Internal/Module/ProxyVariablesContext.h b/Engine/Plugins/Experimental/AnimNext/Source/AnimNext/Internal/Module/ProxyVariablesContext.h
new file mode 100644
--- /dev/null
+++ b/Engine/Plugins/Experimental/AnimNext/Source/AnimNext/Internal/Module/ProxyVariablesContext.h
@@ -0,0 +1,32 @@
+﻿// Copyright Epic Games, Inc. All Rights Reserved.
+
+#pragma once
+
+#include "CoreMinimal.h"
+
+struct FAnimNextPublicVariablesProxy;
+struct FInstancedPropertyBag;
+struct FAnimNextModuleInstance;
+
+namespace UE::AnimNext
+{
+
+// Context passed to module proxy variables flip callback
+struct FProxyVariablesContext
+{
+	// Get the default public variables of the module
+	FAnimNextPublicVariablesProxy& GetPublicVariablesProxy() const;
+
+private:
+	explicit FProxyVariablesContext(FAnimNextModuleInstance& InModuleInstance)
+		: ModuleInstance(InModuleInstance)
+	{}
+
+	// The instance we wrap
+	FAnimNextModuleInstance& ModuleInstance;
+
+	friend struct ::FAnimNextModuleInstance;
+	friend struct FModuleBeginTickFunction;
+};
+
+}
\ No newline at end of file
diff --git a/Engine/Plugins/Experimental/AnimNext/Source/AnimNext/Internal/Module/RigUnit_AnimNextAnimationGraphSelf.h b/Engine/Plugins/Experimental/AnimNext/Source/AnimNext/Internal/Module/RigUnit_AnimNextAnimationGraphSelf.h
new file mode 100644
--- /dev/null
+++ b/Engine/Plugins/Experimental/AnimNext/Source/AnimNext/Internal/Module/RigUnit_AnimNextAnimationGraphSelf.h
@@ -0,0 +1,25 @@
+﻿// Copyright Epic Games, Inc. All Rights Reserved.
+
+#pragma once
+
+#include "Graph/RigUnit_AnimNextBase.h"
+#include "RigUnit_AnimNextAnimationGraphSelf.generated.h"
+
+class UAnimNextAnimationGraph;
+
+/** Get a reference to the currently executing animation graph */
+USTRUCT(meta=(Deprecated, DisplayName="Self", Category="Animation Graph", NodeColor="0, 0, 1", Keywords="Current,This"))
+struct ANIMNEXT_API FRigUnit_AnimNextAnimationGraphSelf : public FRigUnit_AnimNextBase
+{
+	GENERATED_BODY()
+
+	RIGVM_METHOD()
+	void Execute();
+
+	RIGVM_METHOD()
+	virtual FRigVMStructUpgradeInfo GetUpgradeInfo() const override;
+
+	// The currently-executing animation graph
+	UPROPERTY(meta=(Output))
+	TObjectPtr<UAnimNextAnimationGraph> Self;
+};
diff --git a/Engine/Plugins/Experimental/AnimNext/Source/AnimNext/Internal/Module/RigUnit_AnimNextModuleEvents.h b/Engine/Plugins/Experimental/AnimNext/Source/AnimNext/Internal/Module/RigUnit_AnimNextModuleEvents.h
new file mode 100644
--- /dev/null
+++ b/Engine/Plugins/Experimental/AnimNext/Source/AnimNext/Internal/Module/RigUnit_AnimNextModuleEvents.h
@@ -0,0 +1,103 @@
+﻿// Copyright Epic Games, Inc. All Rights Reserved.
+
+#pragma once
+
+#include "Graph/RigUnit_AnimNextBase.h"
+#include "AnimNextExecuteContext.h"
+#include "AnimNextModuleInstance.h"
+#include "ModuleEventTickFunctionBindings.h"
+#include "RigUnit_AnimNextModuleEvents.generated.h"
+
+/** Base schedule-level event, never instantiated */
+USTRUCT(meta=(Abstract, Category="Events", NodeColor="1, 0, 0", Keywords="Begin,Update,Tick,Forward,Event"))
+struct ANIMNEXT_API FRigUnit_AnimNextModuleEventBase : public FRigUnit_AnimNextBase
+{
+	GENERATED_BODY()
+
+	FRigUnit_AnimNextModuleEventBase() = default;
+
+	// FRigUnit interface
+	virtual FString GetUnitLabel() const override { return GetEventName().ToString(); };
+	virtual bool CanOnlyExistOnce() const override final { return true; }
+
+	// Get the general ordering phase of this event, used for linearization
+	virtual UE::AnimNext::EModuleEventPhase GetEventPhase() const { return UE::AnimNext::EModuleEventPhase::Execute; }
+
+	// Overriden in derived classes to provide binding function
+	virtual UE::AnimNext::FModuleEventBindingFunction GetBindingFunction() const { return [](const UE::AnimNext::FTickFunctionBindingContext& InContext, FTickFunction& InTickFunction){}; }
+
+	// The execution result
+	UPROPERTY(EditAnywhere, DisplayName = "Execute", Category = "Events", meta = (Output))
+	FAnimNextExecuteContext ExecuteContext;
+};
+
+/** Synthetic event injected by the compiler to process any variable bindings, not user instantiated */
+USTRUCT(meta=(Hidden, Category="Internal", NodeColor="1, 0, 0"))
+struct ANIMNEXT_API FRigUnit_AnimNextExecuteBindings : public FRigUnit_AnimNextModuleEventBase
+{
+	GENERATED_BODY()
+
+	FRigUnit_AnimNextExecuteBindings() = default;
+
+	static inline const FLazyName EventName = FLazyName("ExecuteBindings");
+	
+	RIGVM_METHOD()
+	virtual void Execute() override;
+	virtual FName GetEventName() const override final { return EventName; }
+
+	// FRigUnit_AnimNextModuleEventBase interface
+	virtual UE::AnimNext::EModuleEventPhase GetEventPhase() const override final { return UE::AnimNext::EModuleEventPhase::PreExecute; }
+	virtual UE::AnimNext::FModuleEventBindingFunction GetBindingFunction() const override final;
+};
+
+/** Schedule event called to set up a module */
+USTRUCT(meta=(DisplayName="Initialize", Keywords="Setup,Startup,Create"))
+struct ANIMNEXT_API FRigUnit_AnimNextInitializeEvent : public FRigUnit_AnimNextModuleEventBase
+{
+	GENERATED_BODY()
+
+	static inline const FLazyName EventName = FLazyName("Initialize");
+
+	RIGVM_METHOD()
+	virtual void Execute() override;
+	virtual FName GetEventName() const override final { return EventName; }
+
+	// FRigUnit_AnimNextModuleEventBase interface
+	virtual UE::AnimNext::EModuleEventPhase GetEventPhase() const override final { return UE::AnimNext::EModuleEventPhase::Initialize; }
+	virtual UE::AnimNext::FModuleEventBindingFunction GetBindingFunction() const override final;
+};
+
+/** Schedule event called before world physics is updated */
+USTRUCT(meta=(DisplayName="PrePhysics", Keywords="Start,Before"))
+struct ANIMNEXT_API FRigUnit_AnimNextPrePhysicsEvent : public FRigUnit_AnimNextModuleEventBase
+{
+	GENERATED_BODY()
+
+	static inline const FLazyName EventName = FLazyName("PrePhysics");
+
+	RIGVM_METHOD()
+	virtual void Execute() override;
+	virtual FName GetEventName() const override final { return EventName; }
+
+	// FRigUnit_AnimNextModuleEventBase interface
+	virtual UE::AnimNext::EModuleEventPhase GetEventPhase() const override final { return UE::AnimNext::EModuleEventPhase::Execute; }
+	virtual UE::AnimNext::FModuleEventBindingFunction GetBindingFunction() const override final;
+};
+
+/** Schedule event called after world physics is updated */
+USTRUCT(meta=(DisplayName="PostPhysics", Keywords="End,After"))
+struct ANIMNEXT_API FRigUnit_AnimNextPostPhysicsEvent : public FRigUnit_AnimNextModuleEventBase
+{
+	GENERATED_BODY()
+
+	static inline const FLazyName EventName = FLazyName("PostPhysics");
+
+	RIGVM_METHOD()
+	virtual void Execute() override;
+	virtual FName GetEventName() const override final { return EventName; }
+
+	// FRigUnit_AnimNextModuleEventBase interface
+	virtual UE::AnimNext::EModuleEventPhase GetEventPhase() const override final { return UE::AnimNext::EModuleEventPhase::Execute; }
+	virtual UE::AnimNext::FModuleEventBindingFunction GetBindingFunction() const override final;
+};
+
diff --git a/Engine/Plugins/Experimental/AnimNext/Source/AnimNext/Internal/Param/ParamUtils.h b/Engine/Plugins/Experimental/AnimNext/Source/AnimNext/Internal/Param/ParamUtils.h
--- a/Engine/Plugins/Experimental/AnimNext/Source/AnimNext/Internal/Param/ParamUtils.h
+++ b/Engine/Plugins/Experimental/AnimNext/Source/AnimNext/Internal/Param/ParamUtils.h
@@ -10,16 +10,10 @@ struct FUniversalObjectLocator;
 namespace UE::AnimNext
 {
 
-struct FParamTypeHandle;
 struct FParamCompatibility;
 
 struct FParamUtils
 {
-	// Check to see if a parameter type handle is compatible with another. The order of the parameters
-	// imply directionality, e.g. IsCompatible((int64)InLHS, (int32)InRHS) is allowed as no data loss 
-	// occurs, but IsCompatible((int32)InLHS, (int64)InRHS) is not as B could be truncated.
-	static ANIMNEXT_API FParamCompatibility GetCompatibility(const FParamTypeHandle& InLHS, const FParamTypeHandle& InRHS);
-
 	// Check to see if a parameter handle is compatible with another. The order of the parameters
 	// imply directionality, e.g. IsCompatible((int64)InLHS, (int32)InRHS) is allowed as no data loss 
 	// occurs, but IsCompatible((int32)InLHS, (int64)InRHS) is not as B could be truncated.
@@ -30,17 +24,17 @@ struct FParamUtils
 	// @param	InClass		The expected class to check - used to determine if hoisted functions are valid in this context.
 	static ANIMNEXT_API bool CanUseFunction(const UFunction* InFunction, const UClass* InExpectedClass);
 
-	// Check whether the supplied function can be used to access/map parameters. Checks type and returns it in OutTypeHandle
+	// Check whether the supplied function can be used to access/map parameters. Checks type and returns it in OutType
 	// @param	InFunction		The function to check
 	// @param	InClass			The expected class to check - used to determine if hoisted functions are valid in this context.
-	// @param	OutTypeHandle	Type handle of the function's return value
-	static ANIMNEXT_API bool CanUseFunction(const UFunction* InFunction, const UClass* InExpectedClass, FParamTypeHandle& OutTypeHandle);
+	// @param	OutType			Type of the function's return value
+	static ANIMNEXT_API bool CanUseFunction(const UFunction* InFunction, const UClass* InExpectedClass, FAnimNextParamType& OutType);
 	
 	// Check (via flags only) the supplied property can be used to access/map parameters
 	static ANIMNEXT_API bool CanUseProperty(const FProperty* InProperty);
 
-	// Check whether the supplied property can be used to access/map parameters. Checks type and returns it in OutTypeHandle
-	static ANIMNEXT_API bool CanUseProperty(const FProperty* InProperty, FParamTypeHandle& OutTypeHandle);
+	// Check whether the supplied property can be used to access/map parameters. Checks type and returns it in OutType
+	static ANIMNEXT_API bool CanUseProperty(const FProperty* InProperty, FAnimNextParamType& OutType);
 
 	// Convert a UOL to an FName
 	static ANIMNEXT_API FName LocatorToName(const FUniversalObjectLocator& InLocator);
diff --git a/Engine/Plugins/Experimental/AnimNext/Source/AnimNext/Internal/Param/RigUnit_AnimNextParameterBase.h b/Engine/Plugins/Experimental/AnimNext/Source/AnimNext/Internal/Param/RigUnit_AnimNextParameterBase.h
--- a/Engine/Plugins/Experimental/AnimNext/Source/AnimNext/Internal/Param/RigUnit_AnimNextParameterBase.h
+++ b/Engine/Plugins/Experimental/AnimNext/Source/AnimNext/Internal/Param/RigUnit_AnimNextParameterBase.h
@@ -6,6 +6,7 @@
 #include "Units/RigUnit.h"
 #include "RigUnit_AnimNextParameterBase.generated.h"
 
+/** Deprecated stub - remove once content is updated */
 USTRUCT(meta = (ExecuteContext = "FAnimNextExecuteContext"))
 struct FRigUnit_AnimNextParameterBase : public FRigUnit
 {
diff --git a/Engine/Plugins/Experimental/AnimNext/Source/AnimNext/Internal/Param/RigUnit_AnimNextParameterBeginExecution.h b/Engine/Plugins/Experimental/AnimNext/Source/AnimNext/Internal/Param/RigUnit_AnimNextParameterBeginExecution.h
--- a/Engine/Plugins/Experimental/AnimNext/Source/AnimNext/Internal/Param/RigUnit_AnimNextParameterBeginExecution.h
+++ b/Engine/Plugins/Experimental/AnimNext/Source/AnimNext/Internal/Param/RigUnit_AnimNextParameterBeginExecution.h
@@ -6,9 +6,7 @@
 #include "AnimNextExecuteContext.h"
 #include "RigUnit_AnimNextParameterBeginExecution.generated.h"
 
-/**
- * Event for driving the skeleton hierarchy with variables and rig elements
- */
+/** Deprecated stub - remove once content is updated */
 USTRUCT(meta=(DisplayName="Execute", Category="Events", NodeColor="1, 0, 0", Keywords="Begin,Update,Tick,Forward,Event"))
 struct ANIMNEXT_API FRigUnit_AnimNextParameterBeginExecution : public FRigUnit_AnimNextParameterBase
 {
diff --git a/Engine/Plugins/Experimental/AnimNext/Source/AnimNext/Internal/Param/RigVMDispatch_GetLayerParameter.h b/Engine/Plugins/Experimental/AnimNext/Source/AnimNext/Internal/Param/RigVMDispatch_GetLayerParameter.h
--- a/Engine/Plugins/Experimental/AnimNext/Source/AnimNext/Internal/Param/RigVMDispatch_GetLayerParameter.h
+++ b/Engine/Plugins/Experimental/AnimNext/Source/AnimNext/Internal/Param/RigVMDispatch_GetLayerParameter.h
@@ -14,7 +14,7 @@ namespace UE::AnimNext::UncookedOnly
 /*
  * Gets a parameter's current value from the graph
  */
-USTRUCT(meta = (DisplayName = "Get Graph Parameter", Category="Parameters", NodeColor = "0.8, 0, 0.2, 1"))
+USTRUCT(meta = (Deprecated, DisplayName = "Get Graph Parameter", Category="Parameters", NodeColor = "0.8, 0, 0.2, 1"))
 struct ANIMNEXT_API FRigVMDispatch_GetLayerParameter : public FRigVMDispatchFactory
 {
 	GENERATED_BODY()
diff --git a/Engine/Plugins/Experimental/AnimNext/Source/AnimNext/Internal/Param/RigVMDispatch_GetScopedParameter.h b/Engine/Plugins/Experimental/AnimNext/Source/AnimNext/Internal/Param/RigVMDispatch_GetScopedParameter.h
--- a/Engine/Plugins/Experimental/AnimNext/Source/AnimNext/Internal/Param/RigVMDispatch_GetScopedParameter.h
+++ b/Engine/Plugins/Experimental/AnimNext/Source/AnimNext/Internal/Param/RigVMDispatch_GetScopedParameter.h
@@ -19,7 +19,7 @@ namespace UE::AnimNext::UncookedOnly
 /*
  * Gets a parameter's value
  */
-USTRUCT(meta = (DisplayName = "Get Parameter", Category="Parameters", NodeColor = "0.8, 0, 0.2, 1"))
+USTRUCT(meta = (Deprecated, DisplayName = "Get Parameter", Category="Parameters", NodeColor = "0.8, 0, 0.2, 1"))
 struct ANIMNEXT_API FRigVMDispatch_GetScopedParameter : public FRigVMDispatchFactory
 {
 	GENERATED_BODY()
diff --git a/Engine/Plugins/Experimental/AnimNext/Source/AnimNext/Internal/Param/RigVMDispatch_SetLayerParameter.h b/Engine/Plugins/Experimental/AnimNext/Source/AnimNext/Internal/Param/RigVMDispatch_SetLayerParameter.h
--- a/Engine/Plugins/Experimental/AnimNext/Source/AnimNext/Internal/Param/RigVMDispatch_SetLayerParameter.h
+++ b/Engine/Plugins/Experimental/AnimNext/Source/AnimNext/Internal/Param/RigVMDispatch_SetLayerParameter.h
@@ -14,7 +14,7 @@ struct FUtils;
 /*
  * Sets a parameter's value
  */
-USTRUCT(meta=(DisplayName = "Set Graph Parameter", Category="Parameters", NodeColor = "0.8, 0, 0.2, 1"))
+USTRUCT(meta=(Deprecated, DisplayName = "Set Graph Parameter", Category="Parameters", NodeColor = "0.8, 0, 0.2, 1"))
 struct ANIMNEXT_API FRigVMDispatch_SetLayerParameter : public FRigVMDispatchFactory
 {
 	GENERATED_BODY()
diff --git a/Engine/Plugins/Experimental/AnimNext/Source/AnimNext/Internal/Variables/AnimNextFieldPath.h b/Engine/Plugins/Experimental/AnimNext/Source/AnimNext/Internal/Variables/AnimNextFieldPath.h
new file mode 100644
--- /dev/null
+++ b/Engine/Plugins/Experimental/AnimNext/Source/AnimNext/Internal/Variables/AnimNextFieldPath.h
@@ -0,0 +1,15 @@
+﻿// Copyright Epic Games, Inc. All Rights Reserved.
+
+#pragma once
+
+#include "AnimNextFieldPath.generated.h"
+
+/** Struct wrapper around TFieldPath, used to allow RigVM support */
+USTRUCT()
+struct ANIMNEXT_API FAnimNextFieldPath
+{
+	GENERATED_BODY()
+
+	UPROPERTY()
+	TFieldPath<FProperty> FieldPath;
+};
\ No newline at end of file
diff --git a/Engine/Plugins/Experimental/AnimNext/Source/AnimNext/Internal/Variables/IAnimNextVariableProxyHost.h b/Engine/Plugins/Experimental/AnimNext/Source/AnimNext/Internal/Variables/IAnimNextVariableProxyHost.h
new file mode 100644
--- /dev/null
+++ b/Engine/Plugins/Experimental/AnimNext/Source/AnimNext/Internal/Variables/IAnimNextVariableProxyHost.h
@@ -0,0 +1,31 @@
+﻿// Copyright Epic Games, Inc. All Rights Reserved.
+
+#pragma once
+
+#include "UObject/Interface.h"
+#include "IAnimNextVariableProxyHost.generated.h"
+
+struct FAnimNextModuleInstance;
+
+namespace UE::AnimNext
+{
+	struct FProxyVariablesContext;
+}
+
+/** Interface used for hosting proxy variables */
+UINTERFACE()
+class UAnimNextVariableProxyHost : public UInterface
+{
+	GENERATED_BODY()
+};
+
+class IAnimNextVariableProxyHost
+{
+	GENERATED_BODY()
+
+	// Flip the public variables proxy buffer
+	// Called from any thread to retrieve the latest copy of any public proxy variables
+	virtual void FlipPublicVariablesProxy(const UE::AnimNext::FProxyVariablesContext& InContext) = 0;
+
+	friend struct FAnimNextModuleInstance;
+};
diff --git a/Engine/Plugins/Experimental/AnimNext/Source/AnimNext/Internal/Variables/RigVMDispatch_CallHoistedAccessorFunction.h b/Engine/Plugins/Experimental/AnimNext/Source/AnimNext/Internal/Variables/RigVMDispatch_CallHoistedAccessorFunction.h
new file mode 100644
--- /dev/null
+++ b/Engine/Plugins/Experimental/AnimNext/Source/AnimNext/Internal/Variables/RigVMDispatch_CallHoistedAccessorFunction.h
@@ -0,0 +1,33 @@
+﻿// Copyright Epic Games, Inc. All Rights Reserved.
+
+#pragma once
+
+#include "RigVMCore/RigVMDispatchFactory.h"
+#include "AnimNextExecuteContext.h"
+#include "RigVMDispatch_CallHoistedAccessorFunction.generated.h"
+
+/** Synthetic dispatch injected by the compiler to get a value via a hoisted accessor UFunction, not user instantiated */
+USTRUCT(meta = (Hidden, DisplayName = "Call Hoisted Function", Category="Internal"))
+struct ANIMNEXT_API FRigVMDispatch_CallHoistedAccessorFunction : public FRigVMDispatchFactory
+{
+	GENERATED_BODY()
+
+	FRigVMDispatch_CallHoistedAccessorFunction();
+
+	static inline const FLazyName ObjectName = FLazyName("Object");
+	static inline const FLazyName FunctionName = FLazyName("Function");
+	static inline const FLazyName ValueName = FLazyName("Value");
+
+private:
+	virtual UScriptStruct* GetExecuteContextStruct() const override { return FAnimNextExecuteContext::StaticStruct(); }
+	virtual FName GetArgumentNameForOperandIndex(int32 InOperandIndex, int32 InTotalOperands) const override;
+	virtual const TArray<FRigVMTemplateArgumentInfo>& GetArgumentInfos() const override;
+	virtual FRigVMTemplateTypeMap OnNewArgumentType(const FName& InArgumentName, TRigVMTypeIndex InTypeIndex) const override;
+	virtual bool IsSingleton() const override { return true; }
+
+	virtual FRigVMFunctionPtr GetDispatchFunctionImpl(const FRigVMTemplateTypeMap& InTypes) const override
+	{
+		return &FRigVMDispatch_CallHoistedAccessorFunction::Execute;
+	}
+	static void Execute(FRigVMExtendedExecuteContext& InContext, FRigVMMemoryHandleArray Handles, FRigVMPredicateBranchArray RigVMBranches);
+};
diff --git a/Engine/Plugins/Experimental/AnimNext/Source/AnimNext/Internal/Variables/RigVMDispatch_CallObjectAccessorFunction.h b/Engine/Plugins/Experimental/AnimNext/Source/AnimNext/Internal/Variables/RigVMDispatch_CallObjectAccessorFunction.h
new file mode 100644
--- /dev/null
+++ b/Engine/Plugins/Experimental/AnimNext/Source/AnimNext/Internal/Variables/RigVMDispatch_CallObjectAccessorFunction.h
@@ -0,0 +1,33 @@
+﻿// Copyright Epic Games, Inc. All Rights Reserved.
+
+#pragma once
+
+#include "RigVMCore/RigVMDispatchFactory.h"
+#include "AnimNextExecuteContext.h"
+#include "RigVMDispatch_CallObjectAccessorFunction.generated.h"
+
+/** Synthetic dispatch injected by the compiler to get a value via an accessor UFunction, not user instantiated */
+USTRUCT(meta = (Hidden, DisplayName = "Call Function", Category="Internal"))
+struct ANIMNEXT_API FRigVMDispatch_CallObjectAccessorFunction : public FRigVMDispatchFactory
+{
+	GENERATED_BODY()
+
+	FRigVMDispatch_CallObjectAccessorFunction();
+
+	static inline const FLazyName ObjectName = FLazyName("Object");
+	static inline const FLazyName FunctionName = FLazyName("Function");
+	static inline const FLazyName ValueName = FLazyName("Value");
+
+private:
+	virtual UScriptStruct* GetExecuteContextStruct() const override { return FAnimNextExecuteContext::StaticStruct(); }
+	virtual FName GetArgumentNameForOperandIndex(int32 InOperandIndex, int32 InTotalOperands) const override;
+	virtual const TArray<FRigVMTemplateArgumentInfo>& GetArgumentInfos() const override;
+	virtual FRigVMTemplateTypeMap OnNewArgumentType(const FName& InArgumentName, TRigVMTypeIndex InTypeIndex) const override;
+	virtual bool IsSingleton() const override { return true; }
+
+	virtual FRigVMFunctionPtr GetDispatchFunctionImpl(const FRigVMTemplateTypeMap& InTypes) const override
+	{
+		return &FRigVMDispatch_CallObjectAccessorFunction::Execute;
+	}
+	static void Execute(FRigVMExtendedExecuteContext& InContext, FRigVMMemoryHandleArray Handles, FRigVMPredicateBranchArray RigVMBranches);
+};
diff --git a/Engine/Plugins/Experimental/AnimNext/Source/AnimNext/Internal/Variables/RigVMDispatch_GetObjectProperty.h b/Engine/Plugins/Experimental/AnimNext/Source/AnimNext/Internal/Variables/RigVMDispatch_GetObjectProperty.h
new file mode 100644
--- /dev/null
+++ b/Engine/Plugins/Experimental/AnimNext/Source/AnimNext/Internal/Variables/RigVMDispatch_GetObjectProperty.h
@@ -0,0 +1,33 @@
+﻿// Copyright Epic Games, Inc. All Rights Reserved.
+
+#pragma once
+
+#include "RigVMCore/RigVMDispatchFactory.h"
+#include "AnimNextExecuteContext.h"
+#include "RigVMDispatch_GetObjectProperty.generated.h"
+
+/** Synthetic dispatch injected by the compiler to get a value from an object's property, not user instantiated */
+USTRUCT(meta = (Hidden, DisplayName = "Get Property", Category="Internal"))
+struct ANIMNEXT_API FRigVMDispatch_GetObjectProperty : public FRigVMDispatchFactory
+{
+	GENERATED_BODY()
+
+	FRigVMDispatch_GetObjectProperty();
+
+	static inline const FLazyName ObjectName = FLazyName("Object");
+	static inline const FLazyName PropertyName = FLazyName("Property");
+	static inline const FLazyName ValueName = FLazyName("Value");
+	
+private:
+	virtual UScriptStruct* GetExecuteContextStruct() const override { return FAnimNextExecuteContext::StaticStruct(); }
+	virtual FName GetArgumentNameForOperandIndex(int32 InOperandIndex, int32 InTotalOperands) const override;
+	virtual const TArray<FRigVMTemplateArgumentInfo>& GetArgumentInfos() const override;
+	virtual FRigVMTemplateTypeMap OnNewArgumentType(const FName& InArgumentName, TRigVMTypeIndex InTypeIndex) const override;
+	virtual bool IsSingleton() const override { return true; }
+
+	virtual FRigVMFunctionPtr GetDispatchFunctionImpl(const FRigVMTemplateTypeMap& InTypes) const override
+	{
+		return &FRigVMDispatch_GetObjectProperty::Execute;
+	}
+	static void Execute(FRigVMExtendedExecuteContext& InContext, FRigVMMemoryHandleArray Handles, FRigVMPredicateBranchArray RigVMBranches);
+};
diff --git a/Engine/Plugins/Experimental/AnimNext/Source/AnimNext/Private/AnimNextModuleImpl.cpp b/Engine/Plugins/Experimental/AnimNext/Source/AnimNext/Private/AnimNextModuleImpl.cpp
--- a/Engine/Plugins/Experimental/AnimNext/Source/AnimNext/Private/AnimNextModuleImpl.cpp
+++ b/Engine/Plugins/Experimental/AnimNext/Source/AnimNext/Private/AnimNextModuleImpl.cpp
@@ -5,28 +5,28 @@
 #include "DataRegistry.h"
 #include "IUniversalObjectLocatorModule.h"
 #include "RigVMRuntimeDataRegistry.h"
+#include "UniversalObjectLocator.h"
 #include "Animation/AnimSequence.h"
 #include "Animation/BlendProfile.h"
 #include "Component/AnimNextComponent.h"
 #include "Curves/CurveFloat.h"
-#include "Module/AnimNextModule.h"
+#include "Graph/AnimNextAnimationGraph.h"
+#include "Graph/AnimNext_LODPose.h"
+#include "Graph/RigVMTrait_AnimNextPublicVariables.h"
+#include "Module/AnimNextModuleInstance.h"
 #include "Modules/ModuleManager.h"
 #include "Param/AnimNextActorLocatorFragment.h"
 #include "Param/AnimNextComponentLocatorFragment.h"
-#include "Param/AnimNextEditorParam.h"
 #include "Param/AnimNextObjectCastLocatorFragment.h"
 #include "Param/AnimNextObjectFunctionLocatorFragment.h"
 #include "Param/AnimNextObjectPropertyLocatorFragment.h"
-#include "Param/AnimNextParam.h"
 #include "Param/AnimNextTag.h"
-#include "Param/ObjectProxyFactory.h"
-#include "Param/PropertyBagProxy.h"
 #include "RigVMCore/RigVMRegistry.h"
-#include "Scheduler/AnimNextTickFunctionBinding.h"
-#include "Scheduler/Scheduler.h"
 #include "TraitCore/NodeTemplateRegistry.h"
 #include "TraitCore/TraitRegistry.h"
 #include "TraitCore/TraitInterfaceRegistry.h"
+#include "Variables/AnimNextFieldPath.h"
+#include "Variables/AnimNextSoftFunctionPtr.h"
 
 #if WITH_ANIMNEXT_CONSOLE_COMMANDS
 #include "HAL/IConsoleManager.h"
@@ -41,7 +41,6 @@
 
 namespace UE::AnimNext
 {
-
 	void FAnimNextModuleImpl::StartupModule()
 	{
 		GetMutableDefault<UAnimNextConfig>()->LoadConfig();
@@ -53,7 +52,7 @@ namespace UE::AnimNext
 			{ UScriptStruct::StaticClass(), FRigVMRegistry::ERegisterObjectOperation::Class },
 			{ UBlendProfile::StaticClass(), FRigVMRegistry::ERegisterObjectOperation::Class },
 			{ UCurveFloat::StaticClass(), FRigVMRegistry::ERegisterObjectOperation::Class },
-			{ UAnimNextModule::StaticClass(), FRigVMRegistry::ERegisterObjectOperation::Class },
+			{ UAnimNextAnimationGraph::StaticClass(), FRigVMRegistry::ERegisterObjectOperation::Class },
 			{ UAnimNextComponent::StaticClass(), FRigVMRegistry::ERegisterObjectOperation::Class },
 		};
 
@@ -61,23 +60,22 @@ namespace UE::AnimNext
 
 		static UScriptStruct* const AllowedStructTypes[] =
 		{
-			FAnimNextEditorParam::StaticStruct(),
-			FAnimNextParam::StaticStruct(),
 			FAnimNextScope::StaticStruct(),
 			FAnimNextEntryPoint::StaticStruct(),
 			FUniversalObjectLocator::StaticStruct(),
-			FAnimNextTickFunctionBinding::StaticStruct(),
-			FAnimNextGraphReferencePose::StaticStruct()
+			FAnimNextGraphReferencePose::StaticStruct(),
+			FAnimNextFieldPath::StaticStruct(),
+			FAnimNextSoftFunctionPtr::StaticStruct(),
+			FRigVMTrait_AnimNextPublicVariables::StaticStruct()
 		};
 
 		RigVMRegistry.RegisterStructTypes(AllowedStructTypes);
 
-		RegisterParameterSourceFactory("ObjectProxy", MakeShared<FObjectProxyFactory>());
+		Private::CacheAllModuleEvents();
 		FDataRegistry::Init();
 		FTraitRegistry::Init();
 		FTraitInterfaceRegistry::Init();
 		FNodeTemplateRegistry::Init();
-		FScheduler::Init();
 		FRigVMRuntimeDataRegistry::Init();
 
 		UE::UniversalObjectLocator::IUniversalObjectLocatorModule& UolModule = FModuleManager::Get().LoadModuleChecked<UE::UniversalObjectLocator::IUniversalObjectLocatorModule>("UniversalObjectLocator");
@@ -122,7 +120,7 @@ namespace UE::AnimNext
 			ConsoleCommands.Add(IConsoleManager::Get().RegisterConsoleCommand(
 				TEXT("AnimNext.Module"),
 				TEXT("Dumps statistics about modules to the log."),
-				FConsoleCommandWithArgsDelegate::CreateRaw(this, &FAnimNextModuleImpl::ListModules),
+				FConsoleCommandWithArgsDelegate::CreateRaw(this, &FAnimNextModuleImpl::ListAnimationGraphs),
 				ECVF_Default
 			));
 		}
@@ -132,12 +130,10 @@ namespace UE::AnimNext
 	void FAnimNextModuleImpl::ShutdownModule()
 	{
 		FRigVMRuntimeDataRegistry::Destroy();
-		FScheduler::Destroy();
 		FNodeTemplateRegistry::Destroy();
 		FTraitInterfaceRegistry::Destroy();
 		FTraitRegistry::Destroy();
 		FDataRegistry::Destroy();
-		UnregisterParameterSourceFactory("ObjectProxy");
 
 #if WITH_ANIMNEXT_CONSOLE_COMMANDS
 		for (IConsoleObject* Cmd : ConsoleCommands)
@@ -160,7 +156,7 @@ namespace UE::AnimNext
 		AnimGraphImpl = nullptr;
 	}
 
-	void FAnimNextModuleImpl::UpdateGraph(FAnimNextGraphInstancePtr& GraphInstance, float DeltaTime, FTraitEventList& InputEventList, FTraitEventList& OutputEventList)
+	void FAnimNextModuleImpl::UpdateGraph(const FAnimNextGraphInstancePtr& GraphInstance, float DeltaTime, FTraitEventList& InputEventList, FTraitEventList& OutputEventList)
 	{
 		if (AnimGraphImpl != nullptr)
 		{
@@ -168,7 +164,7 @@ namespace UE::AnimNext
 		}
 	}
 
-	void FAnimNextModuleImpl::EvaluateGraph(FAnimNextGraphInstancePtr& GraphInstance, const UE::AnimNext::FReferencePose& RefPose, int32 GraphLODLevel, FLODPoseHeap& OutputPose) const
+	void FAnimNextModuleImpl::EvaluateGraph(const FAnimNextGraphInstancePtr& GraphInstance, const UE::AnimNext::FReferencePose& RefPose, int32 GraphLODLevel, FLODPoseHeap& OutputPose) const
 	{
 		if (AnimGraphImpl != nullptr)
 		{
@@ -176,42 +172,6 @@ namespace UE::AnimNext
 		}
 	}
 
-	TUniquePtr<IParameterSource> FAnimNextModuleImpl::CreateParameterSource(const FParameterSourceContext& InContext, const TInstancedStruct<FAnimNextParamInstanceIdentifier>& InInstanceId, TConstArrayView<FName> InRequiredParameters) const
-	{
-		for(const TPair<FName, TSharedRef<IParameterSourceFactory>>& FactoryPair : ParameterSourceFactories)
-		{
-			// TODO: in future we could introduce a priority system or more complex logic here. For now it is earliest wins as are sources are
-			// all mutually exclusive at the moment.
-			TUniquePtr<IParameterSource> ParameterSource = FactoryPair.Value->CreateParameterSource(InContext, InInstanceId, InRequiredParameters);
-			if(ParameterSource.IsValid())
-			{
-				return ParameterSource;
-			}
-		}
-
-		return nullptr;
-	}
-
-	void FAnimNextModuleImpl::RegisterParameterSourceFactory(FName InName, TSharedRef<IParameterSourceFactory> InFactory)
-	{
-		ParameterSourceFactories.Add(InName, InFactory);
-	}
-
-	void FAnimNextModuleImpl::UnregisterParameterSourceFactory(FName InName)
-	{
-		ParameterSourceFactories.Remove(InName);
-	}
-
-	TSharedPtr<IParameterSourceFactory> FAnimNextModuleImpl::FindParameterSourceFactory(FName InName)
-	{
-		if(TSharedRef<IParameterSourceFactory>* FoundFactory = ParameterSourceFactories.Find(InName))
-		{
-			return *FoundFactory;
-		}
-
-		return nullptr;
-	}
-
 #if WITH_ANIMNEXT_CONSOLE_COMMANDS
 	void FAnimNextModuleImpl::ListNodeTemplates(const TArray<FString>& Args)
 	{
@@ -268,7 +228,7 @@ namespace UE::AnimNext
 		LogAnimation.SetVerbosity(OldVerbosity);
 	}
 
-	void FAnimNextModuleImpl::ListModules(const TArray<FString>& Args)
+	void FAnimNextModuleImpl::ListAnimationGraphs(const TArray<FString>& Args)
 	{
 		// Turn off log times to make diff-ing easier
 		TGuardValue<ELogTimes::Type> DisableLogTimes(GPrintLogTimes, ELogTimes::None);
@@ -277,30 +237,30 @@ namespace UE::AnimNext
 		const ELogVerbosity::Type OldVerbosity = LogAnimation.GetVerbosity();
 		LogAnimation.SetVerbosity(ELogVerbosity::All);
 
-		TArray<const UAnimNextModule*> Modules;
+		TArray<const UAnimNextAnimationGraph*> AnimationGraphs;
 
-		for (TObjectIterator<UAnimNextModule> It; It; ++It)
+		for (TObjectIterator<UAnimNextAnimationGraph> It; It; ++It)
 		{
-			Modules.Add(*It);
+			AnimationGraphs.Add(*It);
 		}
 
 		struct FCompareObjectNames
 		{
-			FORCEINLINE bool operator()(const UAnimNextModule& Lhs, const UAnimNextModule& Rhs) const
+			FORCEINLINE bool operator()(const UAnimNextAnimationGraph& Lhs, const UAnimNextAnimationGraph& Rhs) const
 			{
 				return Lhs.GetPathName().Compare(Rhs.GetPathName()) < 0;
 			}
 		};
-		Modules.Sort(FCompareObjectNames());
+		AnimationGraphs.Sort(FCompareObjectNames());
 
 		const FNodeTemplateRegistry& NodeTemplateRegistry = FNodeTemplateRegistry::Get();
 		const FTraitRegistry& TraitRegistry = FTraitRegistry::Get();
 		const bool bDetailedOutput = true;
 
 		UE_LOG(LogAnimation, Log, TEXT("===== AnimNext Modules ====="));
-		UE_LOG(LogAnimation, Log, TEXT("Num Graphs: %u"), Modules.Num());
+		UE_LOG(LogAnimation, Log, TEXT("Num Graphs: %u"), AnimationGraphs.Num());
 
-		for (const UAnimNextModule* Module : Modules)
+		for (const UAnimNextAnimationGraph* AnimationGraph : AnimationGraphs)
 		{
 			uint32 TotalInstanceSize = 0;
 			uint32 NumNodes = 0;
@@ -308,9 +268,9 @@ namespace UE::AnimNext
 				// We always have a node at offset 0
 				int32 NodeOffset = 0;
 
-				while (NodeOffset < Module->SharedDataBuffer.Num())
+				while (NodeOffset < AnimationGraph->SharedDataBuffer.Num())
 				{
-					const FNodeDescription* NodeDesc = reinterpret_cast<const FNodeDescription*>(&Module->SharedDataBuffer[NodeOffset]);
+					const FNodeDescription* NodeDesc = reinterpret_cast<const FNodeDescription*>(&AnimationGraph->SharedDataBuffer[NodeOffset]);
 
 					TotalInstanceSize += NodeDesc->GetNodeInstanceDataSize();
 					NumNodes++;
@@ -320,8 +280,8 @@ namespace UE::AnimNext
 				}
 			}
 
-			UE_LOG(LogAnimation, Log, TEXT("    %s ..."), *Module->GetPathName());
-			UE_LOG(LogAnimation, Log, TEXT("        Shared Data Size: %.2f KB"), double(Module->SharedDataBuffer.Num()) / 1024.0);
+			UE_LOG(LogAnimation, Log, TEXT("    %s ..."), *AnimationGraph->GetPathName());
+			UE_LOG(LogAnimation, Log, TEXT("        Shared Data Size: %.2f KB"), double(AnimationGraph->SharedDataBuffer.Num()) / 1024.0);
 			UE_LOG(LogAnimation, Log, TEXT("        Max Instance Data Size: %.2f KB"), double(TotalInstanceSize) / 1024.0);
 			UE_LOG(LogAnimation, Log, TEXT("        Num Nodes: %u"), NumNodes);
 
@@ -330,9 +290,9 @@ namespace UE::AnimNext
 				// We always have a node at offset 0
 				int32 NodeOffset = 0;
 
-				while (NodeOffset < Module->SharedDataBuffer.Num())
+				while (NodeOffset < AnimationGraph->SharedDataBuffer.Num())
 				{
-					const FNodeDescription* NodeDesc = reinterpret_cast<const FNodeDescription*>(&Module->SharedDataBuffer[NodeOffset]);
+					const FNodeDescription* NodeDesc = reinterpret_cast<const FNodeDescription*>(&AnimationGraph->SharedDataBuffer[NodeOffset]);
 					const FNodeTemplate* NodeTemplate = NodeTemplateRegistry.Find(NodeDesc->GetTemplateHandle());
 
 					const uint32 NumTraits = NodeTemplate->GetNumTraits();
diff --git a/Engine/Plugins/Experimental/AnimNext/Source/AnimNext/Private/AnimNextRigVMAsset.cpp b/Engine/Plugins/Experimental/AnimNext/Source/AnimNext/Private/AnimNextRigVMAsset.cpp
--- a/Engine/Plugins/Experimental/AnimNext/Source/AnimNext/Private/AnimNextRigVMAsset.cpp
+++ b/Engine/Plugins/Experimental/AnimNext/Source/AnimNext/Private/AnimNextRigVMAsset.cpp
@@ -35,6 +35,7 @@ void UAnimNextRigVMAsset::PostLoad()
 	if(VM != nullptr)
 	{
 		VM->ClearExternalVariables(ExtendedExecuteContext);
+		VM->SetExternalVariableDefs(GetExternalVariablesImpl(false));
 		VM->Initialize(ExtendedExecuteContext);
 		InitializeVM(FRigUnit_AnimNextBeginExecution::EventName);
 	}
@@ -63,3 +64,31 @@ void UAnimNextRigVMAsset::GetAssetRegistryTags(FAssetRegistryTagsContext Context
 	}
 #endif // WITH_EDITOR
 }
+
+TArray<FRigVMExternalVariable> UAnimNextRigVMAsset::GetExternalVariablesImpl(bool bFallbackToBlueprint) const
+{
+	TArray<FRigVMExternalVariable> ExternalVariables;
+
+	if(const UPropertyBag* PropertyBag = VariableDefaults.GetPropertyBagStruct())
+	{
+		TConstArrayView<FPropertyBagPropertyDesc> VariableDescs = PropertyBag->GetPropertyDescs();
+		if(VariableDescs.Num() > 0)
+		{
+			ExternalVariables.Reserve(VariableDescs.Num());
+			for(const FPropertyBagPropertyDesc& Desc : VariableDescs)
+			{
+				const FProperty* Property = Desc.CachedProperty;
+				FRigVMExternalVariable ExternalVariable = FRigVMExternalVariable::Make(Property, const_cast<uint8*>(VariableDefaults.GetValue().GetMemory()));
+				if(!ExternalVariable.IsValid())
+				{
+					UE_LOG(LogRigVM, Warning, TEXT("%s: Property '%s' of type '%s' is not supported."), *GetClass()->GetName(), *Property->GetName(), *Property->GetCPPType());
+					continue;
+				}
+
+				ExternalVariables.Add(ExternalVariable);
+			}
+		}
+	}
+
+	return ExternalVariables;
+}
diff --git a/Engine/Plugins/Experimental/AnimNext/Source/AnimNext/Private/Component/AnimNextComponent.cpp b/Engine/Plugins/Experimental/AnimNext/Source/AnimNext/Private/Component/AnimNextComponent.cpp
--- a/Engine/Plugins/Experimental/AnimNext/Source/AnimNext/Private/Component/AnimNextComponent.cpp
+++ b/Engine/Plugins/Experimental/AnimNext/Source/AnimNext/Private/Component/AnimNextComponent.cpp
@@ -2,87 +2,97 @@
 
 #include "Component/AnimNextComponent.h"
 
+#include "AnimNextComponentWorldSubsystem.h"
 #include "Blueprint/BlueprintExceptionInfo.h"
-#include "Scheduler/Scheduler.h"
-#include "Scheduler/ScheduleContext.h"
-#include "Param/PropertyBagProxy.h"
-#include "Scheduler/ScheduleTaskContext.h"
-#include "Scheduler/ScheduleInitializationContext.h"
-#include "IAnimNextModuleInterface.h"
-
-namespace UE::AnimNext::Private
+#include "Engine/World.h"
+#include "Module/ModuleTaskContext.h"
+#include "Module/ProxyVariablesContext.h"
+
+void UAnimNextComponent::OnRegister()
 {
+	using namespace UE::AnimNext;
+
+	Super::OnRegister();
+
+	Subsystem = GetWorld()->GetSubsystem<UAnimNextComponentWorldSubsystem>();
 
-template<typename ContextType>
-static void SetValuesInScopeHelper(const ContextType& InContext, FName InId, FName InScope, EAnimNextParameterScopeOrdering InOrdering, TConstArrayView<FPropertyBagPropertyDesc> InPropertyDescs, TConstArrayView<TConstArrayView<uint8>> InValues)
+	if (Subsystem && Module)
+	{
+		check(!ModuleHandle.IsValid());
+
+		CreatePublicVariablesProxy();
+		Subsystem->Register(this);
+	}
+}
+
+void UAnimNextComponent::OnUnregister()
 {
-	InContext.ApplyParametersToScope(InScope, (EParameterScopeOrdering)InOrdering, InId, InPropertyDescs, InValues);
+	Super::OnUnregister();
+
+	if(Subsystem)
+	{
+		Subsystem->Unregister(this);
+		Subsystem = nullptr;
+		DestroyPublicVariablesProxy();
+	}
 }
 
+void UAnimNextComponent::BeginPlay()
+{
+	Super::BeginPlay();
+
+	SetEnabled(true);
 }
 
-void UAnimNextComponent::OnRegister()
+void UAnimNextComponent::EndPlay(const EEndPlayReason::Type EndPlayReason)
 {
-	using namespace UE::AnimNext;
+	Super::EndPlay(EndPlayReason);
 
-	Super::OnRegister();
+	SetEnabled(false);
+}
 
-	if (Schedule)
+#if WITH_EDITOR
+void UAnimNextComponent::OnModuleCompiled()
+{
+	CreatePublicVariablesProxy();
+}
+#endif
+
+void UAnimNextComponent::CreatePublicVariablesProxy()
+{
+	PublicVariablesProxyMap.Reset();
+	PublicVariablesProxy.Reset();
+	if(Module && Module->GetPublicVariableDefaults().GetPropertyBagStruct())
 	{
-		// Initialization callback to set up any persistent external parameters
-		auto Initialize = [this](const FScheduleInitializationContext& InContext)
+		PublicVariablesProxy.Data = Module->GetPublicVariableDefaults();
+		TConstArrayView<FPropertyBagPropertyDesc> ProxyDescs = PublicVariablesProxy.Data.GetPropertyBagStruct()->GetPropertyDescs();
+		PublicVariablesProxyMap.Reserve(ProxyDescs.Num());
+		for(int32 DescIndex = 0; DescIndex < ProxyDescs.Num(); ++DescIndex)
 		{
-			// First group params into scopes
-			TMap<FName, TArray<UAnimNextComponentParameter*, TInlineAllocator<4>>, TInlineSetAllocator<4>> ParamsByScope;
-			for(UAnimNextComponentParameter* Parameter : Parameters)
-			{
-				if(Parameter && Parameter->IsValid())
-				{
-					ParamsByScope.FindOrAdd(Parameter->Scope).Add(Parameter);
-				}
-			}
-
-			// Now apply to each scope
-			for(const TPair<FName, TArray<UAnimNextComponentParameter*, TInlineAllocator<4>>>& ParamPair : ParamsByScope)
-			{
-				TArray<FPropertyBagPropertyDesc, TInlineAllocator<16>> PropertyDescs;
-				TArray<TConstArrayView<uint8>, TInlineAllocator<16>> Values;
-				PropertyDescs.Reserve(ParamPair.Value.Num());
-				Values.Reserve(ParamPair.Value.Num());
-				for(UAnimNextComponentParameter* Parameter : ParamPair.Value)
-				{
-					FName Name;
-					const FProperty* Property;
-					Parameter->GetParamInfo(Name, Property);
-					PropertyDescs.Emplace(Name, Property);
-					Values.Emplace(Property->ContainerPtrToValuePtr<uint8>(Parameter), Property->GetSize());
-				}
-
-				// NOTE: Layer is always applied 'before' currently. If we have a use case for 'After' we can add it to UAnimNextComponentParameter
-				Private::SetValuesInScopeHelper(InContext, "ComponentParams", ParamPair.Key, EAnimNextParameterScopeOrdering::Before, PropertyDescs, Values);
-			}
-		};
-
-		check(!SchedulerHandle.IsValid());
-		SchedulerHandle = FScheduler::AcquireHandle(this, Schedule, InitMethod, MoveTemp(Initialize));
+			PublicVariablesProxyMap.Add(ProxyDescs[DescIndex].Name, DescIndex);
+		}
+		PublicVariablesProxy.DirtyFlags.SetNum(ProxyDescs.Num(), false);
 	}
 }
 
-void UAnimNextComponent::OnUnregister()
+void UAnimNextComponent::DestroyPublicVariablesProxy()
 {
-	using namespace UE::AnimNext;
-
-	Super::OnUnregister();
+	PublicVariablesProxyMap.Empty();
+	PublicVariablesProxy.Empty();
+}
 
-	FScheduler::ReleaseHandle(this, SchedulerHandle);
+void UAnimNextComponent::FlipPublicVariablesProxy(const UE::AnimNext::FProxyVariablesContext& InContext)
+{
+	FRWScopeLock Lock(PublicVariablesLock, SLT_Write);
+	Swap(InContext.GetPublicVariablesProxy(), PublicVariablesProxy);
 }
 
-void UAnimNextComponent::SetParameterInScope(FName Scope, EAnimNextParameterScopeOrdering Ordering, FName Name, int32 Value)
+void UAnimNextComponent::SetVariable(FName Name, int32 Value)
 {
 	checkNoEntry();
 }
 
-DEFINE_FUNCTION(UAnimNextComponent::execSetParameterInScope)
+DEFINE_FUNCTION(UAnimNextComponent::execSetVariable)
 {
 	using namespace UE::AnimNext;
 
@@ -90,8 +100,6 @@ DEFINE_FUNCTION(UAnimNextComponent::execSetParameterInScope)
 	Stack.MostRecentPropertyAddress = nullptr;
 	Stack.MostRecentPropertyContainer = nullptr;
 
-	P_GET_PROPERTY(FNameProperty, Scope);
-	P_GET_ENUM(EAnimNextParameterScopeOrdering, Ordering);
 	P_GET_PROPERTY(FNameProperty, Name);
 
 	Stack.StepCompiledIn<FProperty>(nullptr);
@@ -104,7 +112,7 @@ DEFINE_FUNCTION(UAnimNextComponent::execSetParameterInScope)
 	{
 		FBlueprintExceptionInfo ExceptionInfo(
 			EBlueprintExceptionType::AbortExecution,
-			NSLOCTEXT("AnimNextComponent", "AnimNextComponent_SetParameterInScopeError", "Failed to resolve the Value for Set Parameter In Scope")
+			NSLOCTEXT("AnimNextComponent", "AnimNextComponent_SetVariableError", "Failed to resolve the Value for Set Variable")
 		);
 
 		FBlueprintCoreDelegates::ThrowScriptException(P_THIS, Stack, ExceptionInfo);
@@ -115,40 +123,57 @@ DEFINE_FUNCTION(UAnimNextComponent::execSetParameterInScope)
 	{
 		FBlueprintExceptionInfo ExceptionInfo(
 			EBlueprintExceptionType::NonFatalError,
-			NSLOCTEXT("AnimNextComponent", "AnimNextComponent_SetParameterInScopeWarning", "Invalid parameter name supplied to Set Parameter In Scope")
+			NSLOCTEXT("AnimNextComponent", "AnimNextComponent_SetVariableInvalidWarning", "Invalid variable name supplied to Set Variable")
 		);
 
 		FBlueprintCoreDelegates::ThrowScriptException(P_THIS, Stack, ExceptionInfo);
 		return;
 	}
 
-	P_NATIVE_BEGIN;
+	int32* IndexPtr = P_THIS->PublicVariablesProxyMap.Find(Name);
+	if (IndexPtr == nullptr)
+	{
+		FBlueprintExceptionInfo ExceptionInfo(
+			EBlueprintExceptionType::NonFatalError,
+			FText::Format(NSLOCTEXT("AnimNextComponent", "AnimNextComponent_SetVariableNotFoundWarning", "Unknown variable name '{0}' supplied to Set Variable"), FText::FromName(Name))
+		);
+
+		FBlueprintCoreDelegates::ThrowScriptException(P_THIS, Stack, ExceptionInfo);
+		return;
+	}
 
-	FInstancedPropertyBag PropertyBag;
-	PropertyBag.AddProperty(Name, ValueProp);
-	const FProperty* NewProperty = PropertyBag.GetPropertyBagStruct()->GetPropertyDescs()[0].CachedProperty;
-	const void* ValuePtr = ValueProp->ContainerPtrToValuePtr<void>(ContainerPtr);
-	NewProperty->SetValue_InContainer(PropertyBag.GetMutableValue().GetMemory(), ValuePtr);
+	P_NATIVE_BEGIN;
 
-	FScheduler::QueueTask(P_THIS, P_THIS->SchedulerHandle, Scope, [Scope, Name, PropertyBag = MoveTemp(PropertyBag), Ordering](const FScheduleTaskContext& InContext) mutable
 	{
-		InContext.ApplyParametersToScope(Scope, (EParameterScopeOrdering)Ordering, Name, MoveTemp(PropertyBag));
-	});
+		FRWScopeLock Lock(P_THIS->PublicVariablesLock, SLT_Write);
+		TConstArrayView<FPropertyBagPropertyDesc> ProxyDescs = P_THIS->PublicVariablesProxy.Data.GetPropertyBagStruct()->GetPropertyDescs();
+		const void* ValuePtr = ValueProp->ContainerPtrToValuePtr<void>(ContainerPtr);
+		ProxyDescs[*IndexPtr].CachedProperty->SetValue_InContainer(P_THIS->PublicVariablesProxy.Data.GetMutableValue().GetMemory(), ValuePtr);
+		P_THIS->PublicVariablesProxy.DirtyFlags[*IndexPtr] = true;
+		P_THIS->PublicVariablesProxy.bIsDirty = true;
+	}
 
 	P_NATIVE_END;
 }
 
-void UAnimNextComponent::Enable(bool bEnabled)
+void UAnimNextComponent::SetEnabled(bool bEnabled)
 {
-	UE::AnimNext::FScheduler::EnableHandle(this, SchedulerHandle, bEnabled);
+	if(Subsystem)
+	{
+		Subsystem->SetEnabled(this, bEnabled);
+	}
 }
 
 void UAnimNextComponent::QueueInputTraitEvent(FAnimNextTraitEventPtr Event)
 {
 	using namespace UE::AnimNext;
 
-	FScheduler::QueueTask(this, SchedulerHandle, NAME_None, [Event = MoveTemp(Event)](const FScheduleTaskContext& InContext)
+	if(Subsystem)
+	{
+		Subsystem->QueueTask(this, NAME_None, [Event = MoveTemp(Event)](const FModuleTaskContext& InContext)
 		{
 			InContext.QueueInputTraitEvent(Event);
-		});
+		},
+		ETaskRunLocation::After);
+	}
 }
