AnimNext: Module refactor

diff --git a/Engine/Plugins/Experimental/AnimNext/Source/AnimNext/Private/Component/AnimNextComponentWorldSubsystem.cpp b/Engine/Plugins/Experimental/AnimNext/Source/AnimNext/Private/Component/AnimNextComponentWorldSubsystem.cpp
new file mode 100644
--- /dev/null
+++ b/Engine/Plugins/Experimental/AnimNext/Source/AnimNext/Private/Component/AnimNextComponentWorldSubsystem.cpp
@@ -0,0 +1,60 @@
+// Copyright Epic Games, Inc. All Rights Reserved.
+
+#include "AnimNextComponentWorldSubsystem.h"
+
+#include "Component/AnimNextComponent.h"
+#include "Engine/World.h"
+#include "Logging/StructuredLog.h"
+#include "Module/AnimNextModule.h"
+#include "Module/ModuleHandle.h"
+#include "Module/ModuleTaskContext.h"
+
+void UAnimNextComponentWorldSubsystem::Register(UAnimNextComponent* InComponent)
+{
+	using namespace UE::AnimNext;
+	check(IsInGameThread());
+	check(InComponent);
+	RegisterHandle(InComponent->ModuleHandle, InComponent->Module, InComponent, InComponent->InitMethod);
+}
+
+void UAnimNextComponentWorldSubsystem::Unregister(UAnimNextComponent* InComponent)
+{
+	using namespace UE::AnimNext;
+	check(IsInGameThread());
+	check(InComponent);
+	UnregisterHandle(InComponent->ModuleHandle);
+}
+
+void UAnimNextComponentWorldSubsystem::SetEnabled(UAnimNextComponent* InComponent, bool bInEnabled)
+{
+	using namespace UE::AnimNext;
+	check(IsInGameThread());
+	check(InComponent);
+	EnableHandle(InComponent->ModuleHandle, bInEnabled);
+}
+
+void UAnimNextComponentWorldSubsystem::QueueTask(UAnimNextComponent* InComponent, FName InModuleEventName, TUniqueFunction<void(const UE::AnimNext::FModuleTaskContext&)>&& InTaskFunction, UE::AnimNext::ETaskRunLocation InLocation)
+{
+	using namespace UE::AnimNext;
+	check(IsInGameThread());
+	check(InComponent);
+	QueueTaskHandle(InComponent->ModuleHandle, InModuleEventName, MoveTemp(InTaskFunction), InLocation);
+}
+
+#if WITH_EDITOR
+
+void UAnimNextComponentWorldSubsystem::OnModuleCompiled(UAnimNextModule* InModule)
+{
+	Super::OnModuleCompiled(InModule);
+
+	for(FAnimNextModuleInstance& Instance : Instances)
+	{
+		if(Instance.Module == InModule)
+		{
+			UAnimNextComponent* AnimNextComponent = CastChecked<UAnimNextComponent>(Instance.Object);
+			AnimNextComponent->OnModuleCompiled();
+		}
+	}
+}
+
+#endif
\ No newline at end of file
diff --git a/Engine/Plugins/Experimental/AnimNext/Source/AnimNext/Private/Component/AnimNextComponentWorldSubsystem.h b/Engine/Plugins/Experimental/AnimNext/Source/AnimNext/Private/Component/AnimNextComponentWorldSubsystem.h
new file mode 100644
--- /dev/null
+++ b/Engine/Plugins/Experimental/AnimNext/Source/AnimNext/Private/Component/AnimNextComponentWorldSubsystem.h
@@ -0,0 +1,44 @@
+// Copyright Epic Games, Inc. All Rights Reserved.
+
+#pragma once
+
+#include "CoreMinimal.h"
+#include "AnimNextPool.h"
+#include "Component/AnimNextWorldSubsystem.h"
+#include "Module/ModuleHandle.h"
+#include "Module/ModuleTaskContext.h"
+#include "AnimNextComponentWorldSubsystem.generated.h"
+
+class UAnimNextComponent;
+
+// Represents AnimNext systems to the AActor/UActorComponent gameplay framework
+UCLASS()
+class UAnimNextComponentWorldSubsystem : public UAnimNextWorldSubsystem
+{
+	GENERATED_BODY()
+
+	friend class UAnimNextComponent;
+
+	// Register a component to the subsystem
+	void Register(UAnimNextComponent* InComponent);
+
+	// Unregister a component to the subsystem
+	// The full release of the module referenced by the component's handle will be deferred after this call is made
+	void Unregister(UAnimNextComponent* InComponent);
+
+	// Enables or disables the module represented by the supplied handle
+	// This operation is deferred until the next time the schedule ticks
+	void SetEnabled(UAnimNextComponent* InComponent, bool bInEnabled);
+
+	// Queue a task to run at a particular point in a schedule
+	// @param	InComponent			The component to execute the task on
+	// @param	InModuleEventName	The name of the event in the module to run the supplied task relative to. If this is NAME_None, then the first valid event will be used.
+	// @param	InTaskFunction		The function to run
+	// @param	InLocation			Where to run the task, before or after
+	void QueueTask(UAnimNextComponent* InComponent, FName InModuleEventName, TUniqueFunction<void(const UE::AnimNext::FModuleTaskContext&)>&& InTaskFunction, UE::AnimNext::ETaskRunLocation InLocation = UE::AnimNext::ETaskRunLocation::Before);
+
+#if WITH_EDITOR
+	// UAnimNextWorldSubsystem interface
+	virtual void OnModuleCompiled(UAnimNextModule* InModule) override;
+#endif
+};
\ No newline at end of file
diff --git a/Engine/Plugins/Experimental/AnimNext/Source/AnimNext/Private/Component/AnimNextWorldSubsystem.cpp b/Engine/Plugins/Experimental/AnimNext/Source/AnimNext/Private/Component/AnimNextWorldSubsystem.cpp
new file mode 100644
--- /dev/null
+++ b/Engine/Plugins/Experimental/AnimNext/Source/AnimNext/Private/Component/AnimNextWorldSubsystem.cpp
@@ -0,0 +1,233 @@
+﻿// Copyright Epic Games, Inc. All Rights Reserved.
+
+#include "Component/AnimNextWorldSubsystem.h"
+
+#include "Engine/World.h"
+#include "Logging/StructuredLog.h"
+#include "Module/AnimNextModule.h"
+#include "Module/ModuleHandle.h"
+#include "Module/ModuleTaskContext.h"
+
+namespace UE::AnimNext
+{
+
+FModulePendingAction::FModulePendingAction(EType InType, FModuleHandle InHandle)
+	: Handle(InHandle)
+	, Type(InType)
+{
+}
+
+}
+
+UAnimNextWorldSubsystem::UAnimNextWorldSubsystem()
+{
+	if(!HasAnyFlags(RF_ClassDefaultObject))
+	{
+#if WITH_EDITOR
+		UAnimNextModule::OnModuleCompiled().AddUObject(this, &UAnimNextWorldSubsystem::OnModuleCompiled);
+#endif
+
+		// Kick off root task at the start of each world tick
+		OnWorldPreActorTickHandle = FWorldDelegates::OnWorldPreActorTick.AddLambda([this](UWorld* InWorld, ELevelTick InTickType, float InDeltaSeconds)
+		{
+			if (InTickType == LEVELTICK_All || InTickType == LEVELTICK_ViewportsOnly)
+			{
+				// Flush actions here as they require game thread callbacks (e.g. to reconfigure tick functions)
+				FlushPendingActions();
+				DeltaTime = InDeltaSeconds;
+			}
+		});
+	}
+}
+
+UAnimNextWorldSubsystem::~UAnimNextWorldSubsystem()
+{
+	if(!HasAnyFlags(RF_ClassDefaultObject))
+	{
+#if WITH_EDITOR
+		UAnimNextModule::OnModuleCompiled().RemoveAll(this);
+#endif
+
+		FWorldDelegates::OnWorldPreActorTick.Remove(OnWorldPreActorTickHandle);
+	}
+}
+
+void UAnimNextWorldSubsystem::AddReferencedObjects(UObject* InThis, FReferenceCollector& Collector)
+{
+	UAnimNextWorldSubsystem* This = CastChecked<UAnimNextWorldSubsystem>(InThis);
+	for(FAnimNextModuleInstance& Instance : This->Instances)
+	{
+		Collector.AddReferencedObject(Instance.Module);
+		Collector.AddReferencedObject(Instance.Object);
+	}
+}
+
+bool UAnimNextWorldSubsystem::DoesSupportWorldType(const EWorldType::Type WorldType) const
+{
+	switch (WorldType)
+	{
+	case EWorldType::Game:
+	case EWorldType::Editor:
+	case EWorldType::PIE:
+	case EWorldType::EditorPreview:
+	case EWorldType::GamePreview:
+		return true;
+	}
+
+	return false;
+}
+
+bool UAnimNextWorldSubsystem::IsValidHandle(UE::AnimNext::FModuleHandle InHandle) const
+{
+	return Instances.IsValidHandle(InHandle);
+}
+
+void UAnimNextWorldSubsystem::FlushPendingActions()
+{
+	using namespace UE::AnimNext;
+
+	FRWScopeLock PendingLockScope(PendingLock, SLT_Write);
+
+	if (PendingActions.Num() > 0)
+	{
+		FRWScopeLock InstancesLockScope(InstancesLock, SLT_Write);
+
+		for (FModulePendingAction& PendingAction : PendingActions)
+		{
+			switch (PendingAction.Type)
+			{
+			case FModulePendingAction::EType::ReleaseHandle:
+				if (IsValidHandle(PendingAction.Handle))
+				{
+					Instances.Release(PendingAction.Handle);
+				}
+				break;
+			case FModulePendingAction::EType::EnableHandle:
+				if (IsValidHandle(PendingAction.Handle))
+				{
+					FAnimNextModuleInstance& Instance = Instances.Get(PendingAction.Handle);
+					Instance.Enable(true);
+				}
+				break;
+			case FModulePendingAction::EType::DisableHandle:
+				if (IsValidHandle(PendingAction.Handle))
+				{
+					FAnimNextModuleInstance& Instance = Instances.Get(PendingAction.Handle);
+					Instance.Enable(false);
+				}
+				break;
+			default:
+				break;
+			}
+		}
+
+		PendingActions.Reset();
+	}
+}
+
+void UAnimNextWorldSubsystem::RegisterHandle(UE::AnimNext::FModuleHandle& InOutHandle, UAnimNextModule* InModule, UObject* InObject, EAnimNextModuleInitMethod InitMethod)
+{
+	using namespace UE::AnimNext;
+	check(IsInGameThread());
+	FRWScopeLock InstancesLockScope(InstancesLock, SLT_Write);
+
+	InOutHandle = Instances.Emplace(InModule, InObject, InitMethod);
+	FAnimNextModuleInstance& Instance = Instances.Get(InOutHandle);
+	Instance.Handle = InOutHandle;
+	Instance.Initialize();
+}
+
+void UAnimNextWorldSubsystem::UnregisterHandle(UE::AnimNext::FModuleHandle& InOutHandle)
+{
+	using namespace UE::AnimNext;
+	check(IsInGameThread());
+
+	if(IsValidHandle(InOutHandle))
+	{
+		FRWScopeLock InstancesLockScope(InstancesLock, SLT_Write);
+
+		Instances.Release(InOutHandle);
+		InOutHandle.Reset();
+
+		// TODO: do not allow immediate handle releases outside of schedule runs - we can either defer or assert
+	//	PendingActions.Emplace(FModulePendingAction::EType::ReleaseHandle, InOutHandle);
+	}
+}
+
+void UAnimNextWorldSubsystem::EnableHandle(UE::AnimNext::FModuleHandle InOutHandle, bool bInEnabled)
+{
+	using namespace UE::AnimNext;
+	check(IsInGameThread());
+	if (IsValidHandle(InOutHandle))
+	{
+		PendingActions.Emplace(bInEnabled ? FModulePendingAction::EType::EnableHandle : FModulePendingAction::EType::DisableHandle, InOutHandle);
+	}
+}
+
+void UAnimNextWorldSubsystem::QueueTaskHandle(UE::AnimNext::FModuleHandle InOutHandle, FName InModuleEventName, TUniqueFunction<void(const UE::AnimNext::FModuleTaskContext&)>&& InTaskFunction, UE::AnimNext::ETaskRunLocation InLocation)
+{
+	using namespace UE::AnimNext;
+	check(IsInGameThread());
+	if (IsValidHandle(InOutHandle))
+	{
+		FAnimNextModuleInstance& Instance = Instances.Get(InOutHandle);
+
+		TSpscQueue<TUniqueFunction<void(const FModuleTaskContext&)>>* Queue = nullptr;
+		FModuleEventTickFunction* FoundTickFunction = nullptr;
+		if(Instance.TickFunctions.Num() > 0)
+		{
+			if(InModuleEventName == NAME_None)
+			{
+				// passing NAME_None means 'any tick function' in this context, so just use the first one
+				FoundTickFunction = &Instance.TickFunctions[0];
+			}
+			else
+			{
+				FoundTickFunction = Instance.TickFunctions.FindByPredicate([InModuleEventName](const FModuleEventTickFunction& InTickFunction)
+				{
+					return InTickFunction.EventName == InModuleEventName;
+				});
+			}
+		}
+
+		if (FoundTickFunction)
+		{
+			switch (InLocation)
+			{
+			case ETaskRunLocation::Before:
+				Queue = &FoundTickFunction->PreExecuteTasks;
+				break;
+			case ETaskRunLocation::After:
+				Queue = &FoundTickFunction->PostExecuteTasks;
+				break;
+			}
+		}
+
+		if (Queue)
+		{
+			Queue->Enqueue(MoveTemp(InTaskFunction));
+		}
+		else
+		{
+			UE_LOGFMT(LogAnimation, Warning, "QueueTask: Could not find event '{EventName}' in module '{ModuleName}'", InModuleEventName, Instance.Module->GetName());
+		}
+	}
+}
+
+#if WITH_EDITOR
+
+void UAnimNextWorldSubsystem::OnModuleCompiled(UAnimNextModule* InModule)
+{
+	// Cant do this while we are running in a world tick
+	check(!GetWorld()->bInTick); 
+
+	for(FAnimNextModuleInstance& Instance : Instances)
+	{
+		if(Instance.Module == InModule)
+		{
+			Instance.OnModuleCompiled();
+		}
+	}
+}
+
+#endif
\ No newline at end of file
diff --git a/Engine/Plugins/Experimental/AnimNext/Source/AnimNext/Private/DataRegistry.cpp b/Engine/Plugins/Experimental/AnimNext/Source/AnimNext/Private/DataRegistry.cpp
--- a/Engine/Plugins/Experimental/AnimNext/Source/AnimNext/Private/DataRegistry.cpp
+++ b/Engine/Plugins/Experimental/AnimNext/Source/AnimNext/Private/DataRegistry.cpp
@@ -176,7 +176,7 @@ void FDataRegistry::FreeAllocatedBlock(Private::FAllocatedBlock* AllocatedBlock)
 
 			{
 				FRWScopeLock ScopedDataTypeDefsLock(DataTypeDefsLock, SLT_ReadOnly);
-				FDataTypeDef* TypeDef = DataTypeDefs.Find(AllocatedBlock->TypeHandle);
+				FDataTypeDef* TypeDef = DataTypeDefs.Find(AllocatedBlock->Type);
 				if (ensure(TypeDef != nullptr))
 				{
 					TypeDef->DestroyTypeFn((uint8*)AllocatedBlock->Memory, AllocatedBlock->NumElem);
diff --git a/Engine/Plugins/Experimental/AnimNext/Source/AnimNext/Private/Graph/AnimGraph/AnimNodeExposedValueHandler_AnimNextParameters.cpp b/Engine/Plugins/Experimental/AnimNext/Source/AnimNext/Private/Graph/AnimGraph/AnimNodeExposedValueHandler_AnimNextParameters.cpp
--- a/Engine/Plugins/Experimental/AnimNext/Source/AnimNext/Private/Graph/AnimGraph/AnimNodeExposedValueHandler_AnimNextParameters.cpp
+++ b/Engine/Plugins/Experimental/AnimNext/Source/AnimNext/Private/Graph/AnimGraph/AnimNodeExposedValueHandler_AnimNextParameters.cpp
@@ -1,11 +1,7 @@
 // Copyright Epic Games, Inc. All Rights Reserved.
 
 #include "Graph/AnimGraph/AnimNodeExposedValueHandler_AnimNextParameters.h"
-#include "Param/ParamStack.h"
 #include "Animation/AnimNodeBase.h"
-#include "Animation/AnimSubsystem_PropertyAccess.h"
-#include "Animation/AnimClassInterface.h"
-#include "Graph/AnimGraphParamStackScope.h"
 
 void FAnimNodeExposedValueHandler_AnimNextParameters::Initialize(const UClass* InClass)
 {
@@ -13,453 +9,14 @@ void FAnimNodeExposedValueHandler_AnimNextParameters::Initialize(const UClass* I
 
 	Super::Initialize(InClass);
 
-	if (const FAnimSubsystem_PropertyAccess* Subsystem = IAnimClassInterface::GetFromClass(InClass)->FindSubsystem<FAnimSubsystem_PropertyAccess>())
-	{
-		PropertyAccessLibrary = &Subsystem->GetLibrary();
-	}
-
-	for (FAnimNodeExposedValueHandler_AnimNextParameters_Entry& Entry : Entries)
-	{
-		Entry.ParamId = FParamId(Entry.ParameterName);
-		Entry.PropertyParamTypeHandle = Entry.PropertyParamType.GetHandle();
-	}
-}
-
-namespace UE::AnimNext::Private
-{
-	template <typename SourceType, typename DestinationType>
-	static void CopyAndCastFloatingPointArray(
-		const TArray<SourceType>& SourceArray,
-		const FArrayProperty* DestinationArrayProperty,
-		void* DestinationAddress)
-	{
-		checkSlow(DestinationArrayProperty);
-		checkSlow(DestinationAddress);
-
-		FScriptArrayHelper DestinationArrayHelper(DestinationArrayProperty, DestinationAddress);
-
-		DestinationArrayHelper.Resize(SourceArray.Num());
-		for (int32 i = 0; i < SourceArray.Num(); ++i)
-		{
-			const SourceType& SourceData = SourceArray[i];
-			DestinationType* DestinationData = reinterpret_cast<DestinationType*>(DestinationArrayHelper.GetRawPtr(i));
-
-			*DestinationData = static_cast<DestinationType>(SourceData);
-		}
-	}
-
-	static void PerformCopy(EPropertyAccessCopyType InCopyType, const FProperty* InDestProperty, void* InDestAddr, const void* InSrcAddr)
-	{
-		switch (InCopyType)
-		{
-		case EPropertyAccessCopyType::None:
-			break;
-		case EPropertyAccessCopyType::Plain:
-			checkSlow(InDestProperty->PropertyFlags & CPF_IsPlainOldData);
-			FMemory::Memcpy(InDestAddr, InSrcAddr, InDestProperty->GetElementSize());
-			break;
-		case EPropertyAccessCopyType::Complex:
-			InDestProperty->CopyCompleteValue(InDestAddr, InSrcAddr);
-			break;
-		case EPropertyAccessCopyType::Bool:
-			checkSlow(InDestProperty->IsA<FBoolProperty>());
-			static_cast<const FBoolProperty*>(InDestProperty)->SetPropertyValue(InDestAddr, *reinterpret_cast<const bool*>(InSrcAddr));
-			break;
-		case EPropertyAccessCopyType::Struct:
-			checkSlow(InDestProperty->IsA<FStructProperty>());
-			static_cast<const FStructProperty*>(InDestProperty)->Struct->CopyScriptStruct(InDestAddr, InSrcAddr);
-			break;
-		case EPropertyAccessCopyType::Name:
-			checkSlow(InDestProperty->IsA<FNameProperty>());
-			static_cast<const FNameProperty*>(InDestProperty)->SetPropertyValue(InDestAddr, *reinterpret_cast<const FName*>(InSrcAddr));
-			break;
-		case EPropertyAccessCopyType::Object:
-			checkSlow(InDestProperty->IsA<FObjectPropertyBase>());
-			static_cast<const FObjectPropertyBase*>(InDestProperty)->SetObjectPropertyValue(InDestAddr, *reinterpret_cast<UObject* const*>(InSrcAddr));
-			break;
-		case EPropertyAccessCopyType::Array:
-		{
-			checkSlow(InDestProperty->IsA<FArrayProperty>());
-			const FArrayProperty* ArrayProperty = static_cast<const FArrayProperty*>(InDestProperty);
-			FScriptArrayHelper SourceArrayHelper(ArrayProperty, InSrcAddr);
-			FScriptArrayHelper DestArrayHelper(ArrayProperty, InDestAddr);
-
-			// Copy the minimum number of elements to the destination array without resizing
-			const int32 MinSize = FMath::Min(SourceArrayHelper.Num(), DestArrayHelper.Num());
-			for (int32 ElementIndex = 0; ElementIndex < MinSize; ++ElementIndex)
-			{
-				ArrayProperty->Inner->CopySingleValue(DestArrayHelper.GetRawPtr(ElementIndex), SourceArrayHelper.GetRawPtr(ElementIndex));
-			}
-			break;
-		}
-		case EPropertyAccessCopyType::PromoteBoolToByte:
-			checkSlow(InDestProperty->IsA<FByteProperty>());
-			static_cast<const FByteProperty*>(InDestProperty)->SetPropertyValue(InDestAddr, (uint8)(*reinterpret_cast<const bool*>(InSrcAddr)));
-			break;
-		case EPropertyAccessCopyType::PromoteBoolToInt32:
-			checkSlow(InDestProperty->IsA<FIntProperty>());
-			static_cast<const FIntProperty*>(InDestProperty)->SetPropertyValue(InDestAddr, (int32)(*reinterpret_cast<const bool*>(InSrcAddr)));
-			break;
-		case EPropertyAccessCopyType::PromoteBoolToInt64:
-			checkSlow(InDestProperty->IsA<FInt64Property>());
-			static_cast<const FInt64Property*>(InDestProperty)->SetPropertyValue(InDestAddr, (int64)(*reinterpret_cast<const bool*>(InSrcAddr)));
-			break;
-		case EPropertyAccessCopyType::PromoteBoolToFloat:
-			checkSlow(InDestProperty->IsA<FFloatProperty>());
-			static_cast<const FFloatProperty*>(InDestProperty)->SetPropertyValue(InDestAddr, (float)(*reinterpret_cast<const bool*>(InSrcAddr)));
-			break;
-		case EPropertyAccessCopyType::PromoteBoolToDouble:
-			checkSlow(InDestProperty->IsA<FDoubleProperty>());
-			static_cast<const FDoubleProperty*>(InDestProperty)->SetPropertyValue(InDestAddr, (double)(*reinterpret_cast<const bool*>(InSrcAddr)));
-			break;
-		case EPropertyAccessCopyType::PromoteByteToInt32:
-			checkSlow(InDestProperty->IsA<FIntProperty>());
-			static_cast<const FIntProperty*>(InDestProperty)->SetPropertyValue(InDestAddr, (int32)(*reinterpret_cast<const uint8*>(InSrcAddr)));
-			break;
-		case EPropertyAccessCopyType::PromoteByteToInt64:
-			checkSlow(InDestProperty->IsA<FInt64Property>());
-			static_cast<const FInt64Property*>(InDestProperty)->SetPropertyValue(InDestAddr, (int64)(*reinterpret_cast<const uint8*>(InSrcAddr)));
-			break;
-		case EPropertyAccessCopyType::PromoteByteToFloat:
-			checkSlow(InDestProperty->IsA<FFloatProperty>());
-			static_cast<const FFloatProperty*>(InDestProperty)->SetPropertyValue(InDestAddr, (float)(*reinterpret_cast<const uint8*>(InSrcAddr)));
-			break;
-		case EPropertyAccessCopyType::PromoteByteToDouble:
-			checkSlow(InDestProperty->IsA<FDoubleProperty>());
-			static_cast<const FDoubleProperty*>(InDestProperty)->SetPropertyValue(InDestAddr, (double)(*reinterpret_cast<const uint8*>(InSrcAddr)));
-			break;
-		case EPropertyAccessCopyType::PromoteInt32ToInt64:
-			checkSlow(InDestProperty->IsA<FInt64Property>());
-			static_cast<const FInt64Property*>(InDestProperty)->SetPropertyValue(InDestAddr, (int64)(*reinterpret_cast<const int32*>(InSrcAddr)));
-			break;
-		case EPropertyAccessCopyType::PromoteInt32ToFloat:
-			checkSlow(InDestProperty->IsA<FFloatProperty>());
-			static_cast<const FFloatProperty*>(InDestProperty)->SetPropertyValue(InDestAddr, (float)(*reinterpret_cast<const int32*>(InSrcAddr)));
-			break;
-		case EPropertyAccessCopyType::PromoteInt32ToDouble:
-			checkSlow(InDestProperty->IsA<FDoubleProperty>());
-			static_cast<const FDoubleProperty*>(InDestProperty)->SetPropertyValue(InDestAddr, (double)(*reinterpret_cast<const int32*>(InSrcAddr)));
-			break;
-		case EPropertyAccessCopyType::PromoteFloatToDouble:
-			checkSlow(InDestProperty->IsA<FDoubleProperty>());
-			static_cast<const FDoubleProperty*>(InDestProperty)->SetPropertyValue(InDestAddr, (double)(*reinterpret_cast<const float*>(InSrcAddr)));
-			break;
-		case EPropertyAccessCopyType::DemoteDoubleToFloat:
-			checkSlow(InDestProperty->IsA<FFloatProperty>());
-			static_cast<const FFloatProperty*>(InDestProperty)->SetPropertyValue(InDestAddr, (float)(*reinterpret_cast<const double*>(InSrcAddr)));
-			break;
-		case EPropertyAccessCopyType::PromoteArrayFloatToDouble:
-		{
-			checkSlow(InDestProperty->IsA<FArrayProperty>());
-			const FArrayProperty* DestArrayProperty = ExactCastField<const FArrayProperty>(InDestProperty);
-			Private::CopyAndCastFloatingPointArray<float, double>(*reinterpret_cast<const TArray<float>*>(InSrcAddr), DestArrayProperty, InDestAddr);
-			break;
-		}
-		case EPropertyAccessCopyType::DemoteArrayDoubleToFloat:
-		{
-			checkSlow(InDestProperty->IsA<FArrayProperty>());
-			const FArrayProperty* DestArrayProperty = ExactCastField<const FArrayProperty>(InDestProperty);
-			Private::CopyAndCastFloatingPointArray<double, float>(*reinterpret_cast<const TArray<double>*>(InSrcAddr), DestArrayProperty, InDestAddr);
-			break;
-		}
-		// TODO: no map support in AnimNext parameter at present
-		case EPropertyAccessCopyType::PromoteMapValueFloatToDouble:
-		case EPropertyAccessCopyType::DemoteMapValueDoubleToFloat:
-		default:
-			check(false);
-			break;
-		}
-	}
-
-	static EPropertyAccessCopyType GetCopyTypeForCompatibility(EPropertyAccessCopyType InAccessType, const FParamTypeHandle& InDestTypeHandle, const FParamTypeHandle& InSrcTypeHandle)
-	{
-		switch (InSrcTypeHandle.GetParameterType())
-		{
-		case FParamTypeHandle::EParamType::Bool:
-			switch (InDestTypeHandle.GetParameterType())
-			{
-			case FParamTypeHandle::EParamType::Bool:
-				return InAccessType;
-			}
-			break;
-		case FParamTypeHandle::EParamType::Byte:
-			switch (InDestTypeHandle.GetParameterType())
-			{
-			case FParamTypeHandle::EParamType::Byte:
-				return InAccessType;
-			case FParamTypeHandle::EParamType::Int32:
-				return EPropertyAccessCopyType::PromoteByteToInt32;
-			case FParamTypeHandle::EParamType::Int64:
-				return EPropertyAccessCopyType::PromoteByteToInt64;
-			case FParamTypeHandle::EParamType::Float:
-				return EPropertyAccessCopyType::PromoteByteToFloat;
-			case FParamTypeHandle::EParamType::Double:
-				return EPropertyAccessCopyType::PromoteByteToDouble;
-			}
-			break;
-		case FParamTypeHandle::EParamType::Int32:
-			switch (InDestTypeHandle.GetParameterType())
-			{
-			case FParamTypeHandle::EParamType::Byte:
-				return EPropertyAccessCopyType::None;
-			case FParamTypeHandle::EParamType::Int32:
-				return InAccessType;
-			case FParamTypeHandle::EParamType::Int64:
-				return EPropertyAccessCopyType::PromoteInt32ToInt64;
-			case FParamTypeHandle::EParamType::Float:
-				return EPropertyAccessCopyType::None;
-			case FParamTypeHandle::EParamType::Double:
-				return EPropertyAccessCopyType::PromoteInt32ToDouble;
-			}
-			break;
-		case FParamTypeHandle::EParamType::Int64:
-			switch (InDestTypeHandle.GetParameterType())
-			{
-			case FParamTypeHandle::EParamType::Byte:
-				return EPropertyAccessCopyType::None;
-			case FParamTypeHandle::EParamType::Int32:
-				return EPropertyAccessCopyType::None;
-			case FParamTypeHandle::EParamType::Int64:
-				return InAccessType;
-			case FParamTypeHandle::EParamType::Float:
-				return EPropertyAccessCopyType::None;
-			case FParamTypeHandle::EParamType::Double:
-				return EPropertyAccessCopyType::None;
-			}
-			break;
-		case FParamTypeHandle::EParamType::Float:
-			switch (InDestTypeHandle.GetParameterType())
-			{
-			case FParamTypeHandle::EParamType::Byte:
-				return EPropertyAccessCopyType::None;
-			case FParamTypeHandle::EParamType::Int32:
-				return EPropertyAccessCopyType::None;
-			case FParamTypeHandle::EParamType::Int64:
-				return EPropertyAccessCopyType::None;
-			case FParamTypeHandle::EParamType::Float:
-				return InAccessType;
-			case FParamTypeHandle::EParamType::Double:
-				return EPropertyAccessCopyType::PromoteFloatToDouble;
-			}
-			break;
-		case FParamTypeHandle::EParamType::Double:
-			switch (InDestTypeHandle.GetParameterType())
-			{
-			case FParamTypeHandle::EParamType::Byte:
-				return EPropertyAccessCopyType::None;
-			case FParamTypeHandle::EParamType::Int32:
-				return EPropertyAccessCopyType::None;
-			case FParamTypeHandle::EParamType::Int64:
-				return EPropertyAccessCopyType::None;
-			case FParamTypeHandle::EParamType::Float:
-				return EPropertyAccessCopyType::DemoteDoubleToFloat;
-			case FParamTypeHandle::EParamType::Double:
-				return InAccessType;
-			}
-			break;
-		case FParamTypeHandle::EParamType::Name:
-			switch (InDestTypeHandle.GetParameterType())
-			{
-			case FParamTypeHandle::EParamType::Name:
-				return InAccessType;
-			}
-			break;
-		case FParamTypeHandle::EParamType::String:
-			switch (InDestTypeHandle.GetParameterType())
-			{
-			case FParamTypeHandle::EParamType::String:
-				return InAccessType;;
-			}
-			break;
-		case FParamTypeHandle::EParamType::Text:
-			switch (InDestTypeHandle.GetParameterType())
-			{
-			case FParamTypeHandle::EParamType::Text:
-				return InAccessType;
-			}
-			break;
-		case FParamTypeHandle::EParamType::Vector:
-			switch (InDestTypeHandle.GetParameterType())
-			{
-			case FParamTypeHandle::EParamType::Vector:
-				return InAccessType;
-			}
-			break;
-		case FParamTypeHandle::EParamType::Vector4:
-			switch (InDestTypeHandle.GetParameterType())
-			{
-			case FParamTypeHandle::EParamType::Vector4:
-				return InAccessType;
-			}
-			break;
-		case FParamTypeHandle::EParamType::Quat:
-			switch (InDestTypeHandle.GetParameterType())
-			{
-			case FParamTypeHandle::EParamType::Quat:
-				return InAccessType;
-			}
-			break;
-		case FParamTypeHandle::EParamType::Transform:
-			switch (InDestTypeHandle.GetParameterType())
-			{
-			case FParamTypeHandle::EParamType::Transform:
-				return InAccessType;
-			}
-			break;
-		case FParamTypeHandle::EParamType::Custom:
-			switch (InDestTypeHandle.GetParameterType())
-			{
-			case FParamTypeHandle::EParamType::Custom:
-				{
-					FAnimNextParamType::EValueType ValueTypeDest, ValueTypeSrc;
-					FAnimNextParamType::EContainerType ContainerTypeDest, ContainerTypeSrc;
-					const UObject* ValueTypeObjectDest, * ValueTypeObjectSrc;
-
-					InDestTypeHandle.GetCustomTypeInfo(ValueTypeDest, ContainerTypeDest, ValueTypeObjectDest);
-					InSrcTypeHandle.GetCustomTypeInfo(ValueTypeSrc, ContainerTypeSrc, ValueTypeObjectSrc);
-
-					if (ContainerTypeDest == ContainerTypeSrc)
-					{
-						if (ContainerTypeDest == FAnimNextParamType::EContainerType::Array)
-						{
-							if (ValueTypeDest != ValueTypeSrc)
-							{
-								switch (ValueTypeDest)
-								{
-								case FAnimNextParamType::EValueType::Float:
-									if (ValueTypeSrc == FAnimNextParamType::EValueType::Double)
-									{
-										return EPropertyAccessCopyType::DemoteArrayDoubleToFloat;
-									}
-									break;
-								case FAnimNextParamType::EValueType::Double:
-									if (ValueTypeSrc == FAnimNextParamType::EValueType::Float)
-									{
-										return EPropertyAccessCopyType::PromoteArrayFloatToDouble;
-									}
-									break;
-								case FAnimNextParamType::EValueType::Struct:
-								case FAnimNextParamType::EValueType::Object:
-								case FAnimNextParamType::EValueType::SoftObject:
-								case FAnimNextParamType::EValueType::Class:
-								case FAnimNextParamType::EValueType::SoftClass:
-									if (ValueTypeObjectDest == ValueTypeObjectSrc)
-									{
-										return InAccessType;
-									}
-									break;
-								}
-							}
-							else
-							{
-								switch (ValueTypeDest)
-								{
-								case FAnimNextParamType::EValueType::Bool:
-								case FAnimNextParamType::EValueType::Byte:
-								case FAnimNextParamType::EValueType::Int32:
-								case FAnimNextParamType::EValueType::Int64:
-								case FAnimNextParamType::EValueType::Float:
-								case FAnimNextParamType::EValueType::Double:
-								case FAnimNextParamType::EValueType::Name:
-								case FAnimNextParamType::EValueType::String:
-								case FAnimNextParamType::EValueType::Text:
-									return InAccessType;
-								case FAnimNextParamType::EValueType::Enum:
-								case FAnimNextParamType::EValueType::Struct:
-								case FAnimNextParamType::EValueType::Object:
-								case FAnimNextParamType::EValueType::SoftObject:
-								case FAnimNextParamType::EValueType::Class:
-								case FAnimNextParamType::EValueType::SoftClass:
-									if (ValueTypeObjectDest == ValueTypeObjectSrc)
-									{
-										return InAccessType;
-									}
-								}
-							}
-						}
-						else
-						{
-							if (ValueTypeDest == ValueTypeSrc)
-							{
-								switch (ValueTypeDest)
-								{
-								case FAnimNextParamType::EValueType::Enum:
-									if (ValueTypeObjectDest == ValueTypeObjectSrc)
-									{
-										return InAccessType;
-									}
-									break;
-								case FAnimNextParamType::EValueType::Struct:
-									if (ValueTypeObjectDest == ValueTypeObjectSrc)
-									{
-										return InAccessType;
-									}
-									else if (CastChecked<UScriptStruct>(ValueTypeObjectDest)->IsChildOf(CastChecked<UScriptStruct>(ValueTypeObjectSrc)))
-									{
-										return InAccessType;
-									}
-									break;
-								case FAnimNextParamType::EValueType::Object:
-								case FAnimNextParamType::EValueType::SoftObject:
-								case FAnimNextParamType::EValueType::Class:
-								case FAnimNextParamType::EValueType::SoftClass:
-									if (ValueTypeObjectDest == ValueTypeObjectSrc)
-									{
-										return InAccessType;
-									}
-									else if (CastChecked<UClass>(ValueTypeObjectDest)->IsChildOf(CastChecked<UClass>(ValueTypeObjectSrc)))
-									{
-										return InAccessType;
-									}
-									break;
-								}
-							}
-						}
-					}
-				}
-				break;
-			}
-			break;
-		}
-
-		return EPropertyAccessCopyType::None;
-	}
+	// TODO: Rework this
 }
 
-
 void FAnimNodeExposedValueHandler_AnimNextParameters::Execute(const FAnimationBaseContext& InContext) const
 {
 	using namespace UE::AnimNext;
 
 	Super::Execute(InContext);
 
-	{
-		FAnimGraphParamStackScope Scope(InContext);
-
-		FParamStack& ParamStack = FParamStack::Get();
-		UObject* Object = InContext.GetAnimInstanceObject();
-
-		for (const FAnimNodeExposedValueHandler_AnimNextParameters_Entry& Entry : Entries)
-		{
-			EPropertyAccessCopyType CopyType = Entry.AccessType;
-			FParamTypeHandle FoundTypeHandle;
-			TConstArrayView<uint8> Value;
-			FParamResult Result = ParamStack.GetParamData(Entry.ParamId, Entry.PropertyParamTypeHandle, Value, FoundTypeHandle, FParamCompatibility::IncompatibleWithDataLoss());
-			if (Result.IsOfCompatibleType())
-			{
-				// Compatibility indicates we need to potentially modify the access type
-				CopyType = Private::GetCopyTypeForCompatibility(CopyType, Entry.PropertyParamTypeHandle, FoundTypeHandle);
-			}
-
-			if(Result.IsSuccessful())
-			{
-				const uint8* SrcAddr = Value.GetData();
-
-				PropertyAccess::GetAccessAddress(Object, *PropertyAccessLibrary, Entry.AccessIndex, [CopyType, SrcAddr](const FProperty* InProperty, void* InAddress)
-				{
-					Private::PerformCopy(CopyType, InProperty, InAddress, SrcAddr);
-				});
-			}
-		}
-	}
+	// TODO: Rework this
 }
\ No newline at end of file
diff --git a/Engine/Plugins/Experimental/AnimNext/Source/AnimNext/Private/Graph/AnimNextAnimationGraph.cpp b/Engine/Plugins/Experimental/AnimNext/Source/AnimNext/Private/Graph/AnimNextAnimationGraph.cpp
new file mode 100644
--- /dev/null
+++ b/Engine/Plugins/Experimental/AnimNext/Source/AnimNext/Private/Graph/AnimNextAnimationGraph.cpp
@@ -0,0 +1,232 @@
+﻿// Copyright Epic Games, Inc. All Rights Reserved.
+
+#include "Graph/AnimNextAnimationGraph.h"
+#include "RigVMCore/RigVMExecuteContext.h"
+#include "Graph/RigUnit_AnimNextShimRoot.h"
+#include "UObject/AssetRegistryTagsContext.h"
+#include "UObject/ObjectSaveContext.h"
+#include "TraitCore/TraitReader.h"
+#include "TraitCore/ExecutionContext.h"
+#include "Graph/AnimNext_LODPose.h"
+#include "Graph/AnimNextGraphInstance.h"
+#include "Serialization/MemoryReader.h"
+#include "AnimNextStats.h"
+#include "Graph/AnimNextGraphEntryPoint.h"
+#include "UObject/LinkerLoad.h"
+#include "UObject/ObjectResource.h"
+
+#if WITH_EDITOR	
+#include "Engine/ExternalAssetDependencyGatherer.h"
+REGISTER_ASSETDEPENDENCY_GATHERER(FExternalAssetDependencyGatherer, UAnimNextAnimationGraph);
+#endif // WITH_EDITOR
+
+DEFINE_STAT(STAT_AnimNext_Graph_AllocateInstance);
+DEFINE_STAT(STAT_AnimNext_Graph_UpdateParamLayer);
+
+#include UE_INLINE_GENERATED_CPP_BY_NAME(AnimNextAnimationGraph)
+
+UAnimNextAnimationGraph::UAnimNextAnimationGraph(const FObjectInitializer& ObjectInitializer)
+	: Super(ObjectInitializer)
+{
+	ExtendedExecuteContext.SetContextPublicDataStruct(FAnimNextExecuteContext::StaticStruct());
+}
+
+void UAnimNextAnimationGraph::AllocateInstance(FAnimNextGraphInstancePtr& OutInstance, FAnimNextModuleInstance* InModuleInstance, FName InEntryPoint) const
+{
+	AllocateInstanceImpl(InModuleInstance, nullptr, OutInstance, InEntryPoint);
+}
+
+void UAnimNextAnimationGraph::AllocateInstance(FAnimNextGraphInstance& InParentGraphInstance, FAnimNextGraphInstancePtr& OutInstance, FName InEntryPoint) const
+{
+	AllocateInstanceImpl(InParentGraphInstance.GetModuleInstance(), &InParentGraphInstance, OutInstance, InEntryPoint);
+}
+
+void UAnimNextAnimationGraph::AllocateInstanceImpl(FAnimNextModuleInstance* InModuleInstance, FAnimNextGraphInstance* ParentGraphInstance, FAnimNextGraphInstancePtr& Instance, FName InEntryPoint) const
+{
+	SCOPE_CYCLE_COUNTER(STAT_AnimNext_Graph_AllocateInstance);
+
+	Instance.Release();
+
+	const FName EntryPoint = (InEntryPoint == NAME_None) ? DefaultEntryPoint : InEntryPoint;
+	const FAnimNextTraitHandle ResolvedRootTraitHandle = ResolvedRootTraitHandles.FindRef(EntryPoint);
+	if (!ResolvedRootTraitHandle.IsValid())
+	{
+		return;
+	}
+
+	{
+		TSharedPtr<FAnimNextGraphInstance> InstanceImpl = MakeShared<FAnimNextGraphInstance>();
+
+		InstanceImpl->AnimationGraph = this;
+		InstanceImpl->ModuleInstance = InModuleInstance;
+		InstanceImpl->ParentGraphInstance = ParentGraphInstance;
+		InstanceImpl->EntryPoint = EntryPoint;
+
+		// If we have a parent graph, use its root since we share the same root, otherwise if we have no parent, we are the root
+		InstanceImpl->RootGraphInstance = ParentGraphInstance != nullptr ? ParentGraphInstance->GetRootGraphInstance() : InstanceImpl.Get();
+
+		InstanceImpl->Variables = VariableDefaults;
+		InstanceImpl->ExtendedExecuteContext = ExtendedExecuteContext;
+		if(VariableDefaults.GetPropertyBagStruct() != nullptr)
+		{
+			InstanceImpl->PublicVariablesState = FAnimNextGraphInstance::EPublicVariablesState::Unbound;
+
+			// Setup external variables memory ptrs manually as we dont follow the pattern of owning multiple URigVMHosts like control rig.
+			// InitializeVM() is called, but only sets up handles for the defaults in the module, not for an instance
+			const int32 NumVariables = InstanceImpl->Variables.GetNumPropertiesInBag();
+			TArray<FRigVMExternalVariableRuntimeData> ExternalVariableRuntimeData;
+			ExternalVariableRuntimeData.Reserve(NumVariables);
+			TConstArrayView<FPropertyBagPropertyDesc> Descs = InstanceImpl->Variables.GetPropertyBagStruct()->GetPropertyDescs();
+			uint8* BasePtr = InstanceImpl->Variables.GetMutableValue().GetMemory();
+			for(int32 VariableIndex = 0; VariableIndex < NumVariables; ++VariableIndex)
+			{
+				ExternalVariableRuntimeData.Emplace(Descs[VariableIndex].CachedProperty->ContainerPtrToValuePtr<uint8>(BasePtr));
+			}
+			InstanceImpl->ExtendedExecuteContext.ExternalVariableRuntimeData = MoveTemp(ExternalVariableRuntimeData);
+		}
+		else
+		{
+			InstanceImpl->PublicVariablesState = FAnimNextGraphInstance::EPublicVariablesState::None;
+		}
+
+		// Now initialize the 'instance', cache memory handles etc. in the context
+		VM->InitializeInstance(InstanceImpl->ExtendedExecuteContext);
+
+		// Move our implementation so that we can use the instance to allocate our root node
+		Instance.Impl = InstanceImpl;
+	}
+
+	{
+		UE::AnimNext::FExecutionContext Context(Instance);
+		Instance.Impl->GraphInstancePtr = Context.AllocateNodeInstance(*Instance.Impl, ResolvedRootTraitHandle);
+	}
+
+	if (!Instance.IsValid())
+	{
+		// We failed to allocate our instance, clear everything
+		Instance.Release();
+	}
+
+#if WITH_EDITORONLY_DATA
+	if (Instance.IsValid())
+	{
+		FScopeLock Lock(&GraphInstancesLock);
+		check(!GraphInstances.Contains(Instance.Impl.Get()));
+		GraphInstances.Add(Instance.Impl.Get());
+	}
+#endif
+}
+
+void UAnimNextAnimationGraph::Serialize(FArchive& Ar)
+{
+	Super::Serialize(Ar);
+
+	if (Ar.IsLoading())
+	{
+		if(Ar.CustomVer(FFortniteMainBranchObjectVersion::GUID) < FFortniteMainBranchObjectVersion::AnimNextCombineParameterBlocksAndGraphs)
+		{
+			// Skip over shared archive buffer if we are loading from an older version
+			if (const FLinkerLoad* Linker = GetLinker())
+			{
+				const int32 LinkerIndex = GetLinkerIndex();
+				const FObjectExport& Export = Linker->ExportMap[LinkerIndex];
+				Ar.Seek(Export.SerialOffset + Export.SerialSize);
+			}
+		}
+		else
+		{
+			int32 SharedDataArchiveBufferSize = 0;
+			Ar << SharedDataArchiveBufferSize;
+
+#if !WITH_EDITORONLY_DATA
+			// When editor data isn't present, we don't persist the archive buffer as it is only needed on load
+			// to populate the graph shared data
+			TArray<uint8> SharedDataArchiveBuffer;
+#endif
+
+			SharedDataArchiveBuffer.SetNumUninitialized(SharedDataArchiveBufferSize);
+			Ar.Serialize(SharedDataArchiveBuffer.GetData(), SharedDataArchiveBufferSize);
+
+			if (Ar.IsLoadingFromCookedPackage())
+			{
+				// If we are cooked, we populate our graph shared data otherwise in the editor we'll compile on load
+				// and re-populate everything then to account for changes in code/content
+				LoadFromArchiveBuffer(SharedDataArchiveBuffer);
+			}
+		}
+	}
+	else if (Ar.IsSaving())
+	{
+#if WITH_EDITORONLY_DATA
+		// We only save the archive buffer, if code changes we'll be able to de-serialize from it when
+		// building the runtime buffer
+		// This allows us to have editor only/non-shipping only properties that are stripped out on load
+		int32 SharedDataArchiveBufferSize = SharedDataArchiveBuffer.Num();
+		Ar << SharedDataArchiveBufferSize;
+		Ar.Serialize(SharedDataArchiveBuffer.GetData(), SharedDataArchiveBufferSize);
+#endif
+	}
+	else
+	{
+		// Counting, etc
+		Ar << SharedDataBuffer;
+
+#if WITH_EDITORONLY_DATA
+		Ar << SharedDataArchiveBuffer;
+#endif
+	}
+}
+
+bool UAnimNextAnimationGraph::LoadFromArchiveBuffer(const TArray<uint8>& InSharedDataArchiveBuffer)
+{
+	using namespace UE::AnimNext;
+
+	// Reconstruct our graph shared data
+	FMemoryReader GraphSharedDataArchive(InSharedDataArchiveBuffer);
+	FTraitReader TraitReader(GraphReferencedObjects, GraphSharedDataArchive);
+
+	const FTraitReader::EErrorState ErrorState = TraitReader.ReadGraph(SharedDataBuffer);
+	if (ErrorState == FTraitReader::EErrorState::None)
+	{
+		for(int32 EntryPointIndex = 0; EntryPointIndex < EntryPoints.Num(); ++EntryPointIndex)
+		{
+			const FAnimNextGraphEntryPoint& EntryPoint = EntryPoints[EntryPointIndex];
+			ResolvedRootTraitHandles.Add(EntryPoint.EntryPointName, TraitReader.ResolveEntryPointHandle(EntryPoint.RootTraitHandle));
+			ResolvedEntryPoints.Add(EntryPoint.EntryPointName, EntryPointIndex);
+		}
+
+		// Make sure our execute method is registered
+		FRigUnit_AnimNextGraphEvaluator::RegisterExecuteMethod(ExecuteDefinition);
+		return true;
+	}
+	else
+	{
+		SharedDataBuffer.Empty(0);
+		ResolvedRootTraitHandles.Add(DefaultEntryPoint, FAnimNextTraitHandle());
+		return false;
+	}
+}
+
+#if WITH_EDITORONLY_DATA
+void UAnimNextAnimationGraph::FreezeGraphInstances()
+{
+	FScopeLock Lock(&GraphInstancesLock);
+
+	TSet<FAnimNextGraphInstance*> GraphInstancesCopy = GraphInstances;
+	for (FAnimNextGraphInstance* GraphInstance : GraphInstancesCopy)
+	{
+		GraphInstance->Freeze();
+	}
+}
+
+void UAnimNextAnimationGraph::ThawGraphInstances()
+{
+	FScopeLock Lock(&GraphInstancesLock);
+
+	TSet<FAnimNextGraphInstance*> GraphInstancesCopy = GraphInstances;
+	for (FAnimNextGraphInstance* GraphInstance : GraphInstancesCopy)
+	{
+		GraphInstance->Thaw();
+	}
+}
+#endif
diff --git a/Engine/Plugins/Experimental/AnimNext/Source/AnimNext/Private/Graph/AnimNextGraphInstance.cpp b/Engine/Plugins/Experimental/AnimNext/Source/AnimNext/Private/Graph/AnimNextGraphInstance.cpp
--- a/Engine/Plugins/Experimental/AnimNext/Source/AnimNext/Private/Graph/AnimNextGraphInstance.cpp
+++ b/Engine/Plugins/Experimental/AnimNext/Source/AnimNext/Private/Graph/AnimNextGraphInstance.cpp
@@ -3,17 +3,23 @@
 #include "Graph/AnimNextGraphInstance.h"
 
 #include "AnimNextStats.h"
+#include "Graph/AnimNextAnimationGraph.h"
 #include "TraitCore/ExecutionContext.h"
-#include "Module/AnimNextModule.h"
 #include "Graph/GC_GraphInstanceComponent.h"
 #include "Graph/RigUnit_AnimNextShimRoot.h"
-#include "Misc/ScopeRWLock.h"
-#include "Param/IParameterSource.h"
-#include "Param/ParametersProxy.h"
+#include "Graph/RigVMTrait_AnimNextPublicVariables.h"
+#include "Logging/StructuredLog.h"
+#include "Module/AnimNextModule.h"
+#include "Module/AnimNextModuleInstance.h"
 
 DEFINE_STAT(STAT_AnimNext_Graph_RigVM);
 
-FAnimNextGraphInstance::FAnimNextGraphInstance() = default;
+FAnimNextGraphInstance::FAnimNextGraphInstance()
+{
+#if WITH_EDITORONLY_DATA
+	UAnimNextModule::OnModuleCompiled().AddRaw(this, &FAnimNextGraphInstance::OnModuleCompiled);
+#endif
+}
 
 FAnimNextGraphInstance::~FAnimNextGraphInstance()
 {
@@ -22,18 +28,28 @@ FAnimNextGraphInstance::~FAnimNextGraphInstance()
 
 void FAnimNextGraphInstance::Release()
 {
+#if WITH_EDITORONLY_DATA
+	UAnimNextModule::OnModuleCompiled().RemoveAll(this);
+
+	if(const UAnimNextAnimationGraph* Graph = GetAnimationGraph())
+	{
+		FScopeLock Lock(&Graph->GraphInstancesLock);
+		Graph->GraphInstances.Remove(this);
+	}
+#endif
+
 	if (!GraphInstancePtr.IsValid())
 	{
 		return;
 	}
 
 	GraphInstancePtr.Reset();
+	ModuleInstance = nullptr;
 	ParentGraphInstance = nullptr;
 	RootGraphInstance = nullptr;
 	ExtendedExecuteContext.Reset();
 	Components.Empty();
-	Module = nullptr;
-	GraphState = nullptr;
+	AnimationGraph = nullptr;
 }
 
 bool FAnimNextGraphInstance::IsValid() const
@@ -41,9 +57,9 @@ bool FAnimNextGraphInstance::IsValid() const
 	return GraphInstancePtr.IsValid();
 }
 
-const UAnimNextModule* FAnimNextGraphInstance::GetModule() const
+const UAnimNextAnimationGraph* FAnimNextGraphInstance::GetAnimationGraph() const
 {
-	return Module;
+	return AnimationGraph;
 }
 
 FName FAnimNextGraphInstance::GetEntryPoint() const
@@ -56,6 +72,11 @@ UE::AnimNext::FWeakTraitPtr FAnimNextGraphInstance::GetGraphRootPtr() const
 	return GraphInstancePtr;
 }
 
+FAnimNextModuleInstance* FAnimNextGraphInstance::GetModuleInstance() const
+{
+	return ModuleInstance;
+}
+
 FAnimNextGraphInstance* FAnimNextGraphInstance::GetParentGraphInstance() const
 {
 	return ParentGraphInstance;
@@ -66,18 +87,18 @@ FAnimNextGraphInstance* FAnimNextGraphInstance::GetRootGraphInstance() const
 	return RootGraphInstance;
 }
 
-bool FAnimNextGraphInstance::UsesModule(const UAnimNextModule* InModule) const
+bool FAnimNextGraphInstance::UsesAnimationGraph(const UAnimNextAnimationGraph* InAnimationGraph) const
 {
-	return Module == InModule;
+	return AnimationGraph == InAnimationGraph;
 }
 
 bool FAnimNextGraphInstance::UsesEntryPoint(FName InEntryPoint) const
 {
-	if(Module != nullptr)
+	if(AnimationGraph != nullptr)
 	{
 		if(InEntryPoint == NAME_None)
 		{
-			return EntryPoint == Module->DefaultEntryPoint;
+			return EntryPoint == AnimationGraph->DefaultEntryPoint;
 		}
 
 		return InEntryPoint == EntryPoint;
@@ -106,11 +127,6 @@ void FAnimNextGraphInstance::AddStructReferencedObjects(FReferenceCollector& Col
 	{
 		Component->AddReferencedObjects(Collector);
 	}
-
-	if(GraphState.IsValid())
-	{
-		GraphState->AddReferencedObjects(Collector);
-	}
 }
 
 UE::AnimNext::FGraphInstanceComponent* FAnimNextGraphInstance::TryGetComponent(int32 ComponentNameHash, FName ComponentName) const
@@ -138,6 +154,11 @@ void FAnimNextGraphInstance::Update()
 	bHasUpdatedOnce = true;
 }
 
+FRigVMExtendedExecuteContext& FAnimNextGraphInstance::GetExtendedExecuteContext()
+{
+	return ExtendedExecuteContext;
+}
+
 void FAnimNextGraphInstance::ExecuteLatentPins(const TConstArrayView<UE::AnimNext::FLatentPropertyHandle>& LatentHandles, void* DestinationBasePtr, bool bIsFrozen)
 {
 	SCOPE_CYCLE_COUNTER(STAT_AnimNext_Graph_RigVM);
@@ -147,10 +168,10 @@ void FAnimNextGraphInstance::ExecuteLatentPins(const TConstArrayView<UE::AnimNex
 		return;
 	}
 
-	if (URigVM* VM = Module->VM)
+	if (URigVM* VM = AnimationGraph->VM)
 	{
 		FAnimNextExecuteContext& AnimNextContext = ExtendedExecuteContext.GetPublicDataSafe<FAnimNextExecuteContext>();
-		AnimNextContext.SetContextData<FAnimNextGraphContextData>(this, LatentHandles, DestinationBasePtr, bIsFrozen);
+		AnimNextContext.SetContextData<FAnimNextGraphContextData>(ModuleInstance, this, LatentHandles, DestinationBasePtr, bIsFrozen);
 
 		VM->ExecuteVM(ExtendedExecuteContext, FRigUnit_AnimNextShimRoot::EventName);
 
@@ -159,6 +180,7 @@ void FAnimNextGraphInstance::ExecuteLatentPins(const TConstArrayView<UE::AnimNex
 	}
 }
 
+#if WITH_EDITORONLY_DATA
 void FAnimNextGraphInstance::Freeze()
 {
 	if (!IsValid())
@@ -167,7 +189,6 @@ void FAnimNextGraphInstance::Freeze()
 	}
 
 	GraphInstancePtr.Reset();
-	GraphState.Reset();
 	ExtendedExecuteContext.Reset();
 	Components.Empty();
 	bHasUpdatedOnce = false;
@@ -175,16 +196,17 @@ void FAnimNextGraphInstance::Freeze()
 
 void FAnimNextGraphInstance::Thaw()
 {
-	if (const UAnimNextModule* ModulePtr = Module)
+	if (const UAnimNextAnimationGraph* AnimationGraphPtr = AnimationGraph)
 	{
-		GraphState = MakeUnique<UE::AnimNext::FParametersProxy>(Module);
+		Variables.MigrateToNewBagInstance(AnimationGraphPtr->VariableDefaults);
+
+		ExtendedExecuteContext = AnimationGraphPtr->ExtendedExecuteContext;
 
-		ExtendedExecuteContext.CopyMemoryStorage(ModulePtr->ExtendedExecuteContext);
-		ModulePtr->VM->InitializeInstance(ExtendedExecuteContext);
+		RebindPublicVariables();
 
 		{
 			UE::AnimNext::FExecutionContext Context(*this);
-			if(const FAnimNextTraitHandle* FoundHandle = ModulePtr->ResolvedRootTraitHandles.Find(EntryPoint))
+			if(const FAnimNextTraitHandle* FoundHandle = AnimationGraphPtr->ResolvedRootTraitHandles.Find(EntryPoint))
 			{
 				GraphInstancePtr = Context.AllocateNodeInstance(*this, *FoundHandle);
 			}
@@ -198,17 +220,142 @@ void FAnimNextGraphInstance::Thaw()
 	}
 }
 
-UE::AnimNext::FParamStack::FPushedLayerHandle FAnimNextGraphInstance::UpdateAndPushGraphState(float InDeltaTime) const
+void FAnimNextGraphInstance::RebindPublicVariables()
 {
-	if(GraphState.IsValid())
+	if(PublicVariablesState != EPublicVariablesState::Bound)
 	{
-		GraphState->Update(InDeltaTime);
-		return UE::AnimNext::FParamStack::Get().PushLayer(GraphState->GetLayerHandle());
+		return;
 	}
-	return UE::AnimNext::FParamStack::FPushedLayerHandle();
+
+	// Setup external variables memory ptrs manually as we dont follow the pattern of owning multiple URigVMHosts like control rig.
+	// InitializeVM() is called, but only sets up handles for the defaults in the module, not for an instance
+	const int32 NumVariables = Variables.GetNumPropertiesInBag();
+	TArray<FRigVMExternalVariableRuntimeData> ExternalVariableRuntimeData;
+	ExternalVariableRuntimeData.Reserve(NumVariables);
+	TConstArrayView<FPropertyBagPropertyDesc> Descs = Variables.GetPropertyBagStruct()->GetPropertyDescs();
+	uint8* BasePtr = Variables.GetMutableValue().GetMemory();
+	for(int32 VariableIndex = 0; VariableIndex < NumVariables; ++VariableIndex)
+	{
+		ExternalVariableRuntimeData.Emplace(Descs[VariableIndex].CachedProperty->ContainerPtrToValuePtr<uint8>(BasePtr));
+	}
+	ExtendedExecuteContext.ExternalVariableRuntimeData = MoveTemp(ExternalVariableRuntimeData);
+
+	// Re-apply any bindings from our host
+	for(const FCachedVariableBinding& CachedBinding : CachedVariableBindings)
+	{
+		// TODO: remove this linear search in FindPropertyDescByName with a hash table?
+		if(const FPropertyBagPropertyDesc* Desc = Variables.FindPropertyDescByName(CachedBinding.VariableName))
+		{
+			int32 VariableIndex = Desc - Variables.GetPropertyBagStruct()->GetPropertyDescs().GetData();
+			ExtendedExecuteContext.ExternalVariableRuntimeData[VariableIndex].Memory = CachedBinding.Memory;
+		}
+	}
+
+	// Now reinitialize the 'instance', cache memory handles etc. in the context
+	AnimationGraph->VM->InitializeInstance(ExtendedExecuteContext);
 }
 
-void FAnimNextGraphInstance::PopGraphState(UE::AnimNext::FParamStack::FPushedLayerHandle InHandle) const
+void FAnimNextGraphInstance::OnModuleCompiled(UAnimNextModule* InModule)
 {
-	UE::AnimNext::FParamStack::Get().PopLayer(InHandle);
-}
\ No newline at end of file
+	// If we are hosted directly by a module, invalidate and mark our bindings as needing update 
+	if(ModuleInstance && InModule == ModuleInstance->Module && ParentGraphInstance == nullptr)
+	{
+		UnbindPublicVariables();
+	}
+}
+
+#endif
+
+void FAnimNextGraphInstance::BindPublicVariables(TConstArrayView<FRigVMTraitScope> InTraitScopes)
+{
+	if(AnimationGraph == nullptr)
+	{
+		return;
+	}
+
+	if(PublicVariablesState == EPublicVariablesState::Bound)
+	{
+		return;
+	}
+
+#if WITH_EDITORONLY_DATA
+	CachedVariableBindings.Reset();
+#endif
+
+	bool bPublicVariablesBound = false;
+	for(const FRigVMTraitScope& TraitScope : InTraitScopes)
+	{
+		const FRigVMTrait_AnimNextPublicVariables* VariablesTrait = TraitScope.GetTrait<FRigVMTrait_AnimNextPublicVariables>();
+		if(VariablesTrait == nullptr)
+		{
+			continue;
+		}
+
+		for(int32 TraitVariableIndex = 0; TraitVariableIndex < VariablesTrait->VariableNames.Num(); ++TraitVariableIndex)
+		{
+			FName VariableName = VariablesTrait->VariableNames[TraitVariableIndex];
+			// TODO: remove this linear search in FindPropertyDescByName with a hash table?
+			// TODO: multiple traits with the same name can collide at the moment, we should validate against this is at the compiler level to avoid issues
+			const FPropertyBagPropertyDesc* Desc = Variables.FindPropertyDescByName(VariableName);
+			if(Desc == nullptr)
+			{
+				continue;
+			}
+
+			if(!TraitScope.GetAdditionalMemoryHandles().IsValidIndex(TraitVariableIndex))
+			{
+				continue;
+			}
+
+			const FRigVMMemoryHandle& MemoryHandle = TraitScope.GetAdditionalMemoryHandles()[TraitVariableIndex];
+			int32 VariableIndex = Desc - Variables.GetPropertyBagStruct()->GetPropertyDescs().GetData();
+			if(Desc->CachedProperty->GetClass() != MemoryHandle.GetProperty()->GetClass())
+			{
+				UE_LOGFMT(LogAnimation, Warning, "Mismatched variable types when binding AnimNext graph: {Name}:{Type} vs {OtherType}", Desc->Name, Desc->CachedProperty->GetFName(), MemoryHandle.GetProperty()->GetFName());
+				continue;
+			}
+
+			uint8* Memory = const_cast<uint8*>(MemoryHandle.GetData());
+			ExtendedExecuteContext.ExternalVariableRuntimeData[VariableIndex].Memory = Memory;
+#if WITH_EDITORONLY_DATA
+			CachedVariableBindings.Add({ VariableName, Memory });
+#endif
+			bPublicVariablesBound = true;
+		}
+	}
+
+	if(bPublicVariablesBound)
+	{
+		// Re-initialize memory handles
+		AnimationGraph->VM->InitializeInstance(ExtendedExecuteContext, /* bCopyMemory = */false);
+	}
+
+	PublicVariablesState = EPublicVariablesState::Bound;
+}
+
+void FAnimNextGraphInstance::UnbindPublicVariables()
+{
+	if(PublicVariablesState != EPublicVariablesState::Bound)
+	{
+		return;
+	}
+
+#if WITH_EDITORONLY_DATA
+	CachedVariableBindings.Reset();
+#endif
+
+	// Reset external variable ptrs to point to internal public vars
+	const int32 NumVariables = Variables.GetNumPropertiesInBag();
+	TConstArrayView<FPropertyBagPropertyDesc> Descs = Variables.GetPropertyBagStruct()->GetPropertyDescs();
+	uint8* BasePtr = Variables.GetMutableValue().GetMemory();
+	for(int32 VariableIndex = 0; VariableIndex < NumVariables; ++VariableIndex)
+	{
+		ExtendedExecuteContext.ExternalVariableRuntimeData[VariableIndex].Memory = Descs[VariableIndex].CachedProperty->ContainerPtrToValuePtr<uint8>(BasePtr);
+	}
+
+	// Re-initialize memory handles
+	AnimationGraph->VM->InitializeInstance(ExtendedExecuteContext, /* bCopyMemory = */false);
+
+	PublicVariablesState = EPublicVariablesState::Unbound;
+}
+
diff --git a/Engine/Plugins/Experimental/AnimNext/Source/AnimNext/Private/Graph/AnimNextGraphInstancePtr.cpp b/Engine/Plugins/Experimental/AnimNext/Source/AnimNext/Private/Graph/AnimNextGraphInstancePtr.cpp
--- a/Engine/Plugins/Experimental/AnimNext/Source/AnimNext/Private/Graph/AnimNextGraphInstancePtr.cpp
+++ b/Engine/Plugins/Experimental/AnimNext/Source/AnimNext/Private/Graph/AnimNextGraphInstancePtr.cpp
@@ -2,13 +2,10 @@
 
 #include "Graph/AnimNextGraphInstancePtr.h"
 
-#include "Module/AnimNextModule.h"
+#include "Graph/AnimNextAnimationGraph.h"
 #include "Graph/AnimNextGraphInstance.h"
-#include "Param/IParameterSource.h"
 
 FAnimNextGraphInstancePtr::FAnimNextGraphInstancePtr() = default;
-FAnimNextGraphInstancePtr::FAnimNextGraphInstancePtr(FAnimNextGraphInstancePtr&&) = default;
-FAnimNextGraphInstancePtr& FAnimNextGraphInstancePtr::operator=(FAnimNextGraphInstancePtr&&) = default;
 
 FAnimNextGraphInstancePtr::~FAnimNextGraphInstancePtr()
 {
@@ -17,22 +14,7 @@ FAnimNextGraphInstancePtr::~FAnimNextGraphInstancePtr()
 
 void FAnimNextGraphInstancePtr::Release()
 {
-	if (Impl)
-	{
-#if WITH_EDITORONLY_DATA
-		{
-			const UAnimNextModule* Module = Impl->GetModule();
-			FScopeLock Lock(&Module->GraphInstancesLock);
-			Module->GraphInstances.Remove(Impl.Get());
-		}
-#endif
-
-		// Destroy the graph instance
-		Impl->Release();
-
-		// Reset our unique ptr
-		Impl.Reset();
-	}
+	Impl.Reset();
 }
 
 bool FAnimNextGraphInstancePtr::IsValid() const
@@ -40,9 +22,9 @@ bool FAnimNextGraphInstancePtr::IsValid() const
 	return Impl && Impl->IsValid();
 }
 
-const UAnimNextModule* FAnimNextGraphInstancePtr::GetModule() const
+const UAnimNextAnimationGraph* FAnimNextGraphInstancePtr::GetAnimationGraph() const
 {
-	return Impl ? Impl->GetModule() : nullptr;
+	return Impl->GetAnimationGraph();
 }
 
 UE::AnimNext::FWeakTraitPtr FAnimNextGraphInstancePtr::GetGraphRootPtr() const
@@ -55,9 +37,9 @@ FAnimNextGraphInstance* FAnimNextGraphInstancePtr::GetImpl() const
 	return Impl.Get();
 }
 
-bool FAnimNextGraphInstancePtr::UsesModule(const UAnimNextModule* InModule) const
+bool FAnimNextGraphInstancePtr::UsesAnimationGraph(const UAnimNextAnimationGraph* InAnimationGraph) const
 {
-	return Impl ? Impl->UsesModule(InModule) : false;
+	return Impl ? Impl->UsesAnimationGraph(InAnimationGraph) : false;
 }
 
 bool FAnimNextGraphInstancePtr::IsRoot() const
@@ -95,20 +77,26 @@ GraphInstanceComponentMapType::TConstIterator FAnimNextGraphInstancePtr::GetComp
 	return Impl->GetComponentIterator();
 }
 
-void FAnimNextGraphInstancePtr::Update()
+void FAnimNextGraphInstancePtr::Update() const
 {
 	check(Impl);
 	Impl->Update();
 }
 
-UE::AnimNext::FParamStack::FPushedLayerHandle FAnimNextGraphInstancePtr::UpdateAndPushGraphState(float InDeltaTime) const
+FRigVMExtendedExecuteContext& FAnimNextGraphInstancePtr::GetExtendedExecuteContext() const
+{
+	check(Impl);
+	return Impl->GetExtendedExecuteContext();
+}
+
+bool FAnimNextGraphInstancePtr::RequiresPublicVariableBinding() const
 {
 	check(Impl);
-	return Impl->UpdateAndPushGraphState(InDeltaTime);
+	return Impl->PublicVariablesState == FAnimNextGraphInstance::EPublicVariablesState::Unbound;
 }
 
-void FAnimNextGraphInstancePtr::PopGraphState(UE::AnimNext::FParamStack::FPushedLayerHandle InHandle) const
+void FAnimNextGraphInstancePtr::BindPublicVariables(TConstArrayView<FRigVMTraitScope> InTraitScopes) const
 {
 	check(Impl);
-	Impl->PopGraphState(InHandle);
+	return Impl->BindPublicVariables(InTraitScopes);
 }
diff --git a/Engine/Plugins/Experimental/AnimNext/Source/AnimNext/Private/Graph/RigUnit_AnimNextGraphEvaluator.cpp b/Engine/Plugins/Experimental/AnimNext/Source/AnimNext/Private/Graph/RigUnit_AnimNextGraphEvaluator.cpp
--- a/Engine/Plugins/Experimental/AnimNext/Source/AnimNext/Private/Graph/RigUnit_AnimNextGraphEvaluator.cpp
+++ b/Engine/Plugins/Experimental/AnimNext/Source/AnimNext/Private/Graph/RigUnit_AnimNextGraphEvaluator.cpp
@@ -1,7 +1,6 @@
 // Copyright Epic Games, Inc. All Rights Reserved.
 
 #include "Graph/RigUnit_AnimNextGraphEvaluator.h"
-#include "Context.h"
 #include "TraitCore/LatentPropertyHandle.h"
 
 #include UE_INLINE_GENERATED_CPP_BY_NAME(RigUnit_AnimNextGraphEvaluator)
@@ -49,11 +48,9 @@ void FRigUnit_AnimNextGraphEvaluator::StaticExecute(FRigVMExtendedExecuteContext
 
 		FRigVMMemoryHandle& MemoryHandle = RigVMMemoryHandles[Handle.GetLatentPropertyIndex()];
 
-		// This should be an assert. If this triggers, it means that we have a bug in how lazy memory handles
-		// are assigned during compilation. We keep it as an ensure because in this case, we can recover
-		// as even if the memory handle isn't lazy, it remains valid and we can use it. It won't have the
-		// value we expect but it'll work. The ensure will signal that we need to fix the bug.
-		if (ensure(MemoryHandle.IsLazy()))
+		// We cannot currently determine from our memory handle whether this value is a direct wire-up to a variable, in which
+		// case it will not be lazy, thus we cannot verify lazy-correctness before calling into this so we just guard instead.
+		if (MemoryHandle.IsLazy())
 		{
 			MemoryHandle.ComputeLazyValueIfNecessary(RigVMExecuteContext, RigVMExecuteContext.GetSliceHash());
 		}
diff --git a/Engine/Plugins/Experimental/AnimNext/Source/AnimNext/Private/Graph/RigUnit_AnimNextRunAnimationGraph.cpp b/Engine/Plugins/Experimental/AnimNext/Source/AnimNext/Private/Graph/RigUnit_AnimNextRunAnimationGraph.cpp
new file mode 100644
--- /dev/null
+++ b/Engine/Plugins/Experimental/AnimNext/Source/AnimNext/Private/Graph/RigUnit_AnimNextRunAnimationGraph.cpp
@@ -0,0 +1,112 @@
+﻿// Copyright Epic Games, Inc. All Rights Reserved.
+
+#include "Graph/RigUnit_AnimNextRunAnimationGraph.h"
+
+#include "Units/RigUnitContext.h"
+#include "AnimNextStats.h"
+#include "IAnimNextModuleInterface.h"
+#include "Module/AnimNextModuleInstance.h"
+#include "Graph/AnimNextAnimationGraph.h"
+#include "TraitCore/TraitEventList.h"
+
+#include UE_INLINE_GENERATED_CPP_BY_NAME(RigUnit_AnimNextRunAnimationGraph)
+
+DEFINE_STAT(STAT_AnimNext_Run_Graph);
+
+FRigUnit_AnimNextRunAnimationGraph_Execute()
+{
+	SCOPE_CYCLE_COUNTER(STAT_AnimNext_Run_Graph);
+
+	using namespace UE::AnimNext;
+
+	const FAnimNextModuleContextData& ModuleContextData = ExecuteContext.GetContextData<FAnimNextModuleContextData>();
+	FAnimNextModuleInstance& ModuleInstance = ModuleContextData.GetModuleInstance();
+
+	if (!ReferencePose.ReferencePose.IsValid())
+	{
+		return;
+	}
+
+	if(Graph == nullptr)
+	{
+		Instance.Release();
+		return;
+	}
+
+	// Lazily (re-)allocate graph instance if required
+	if(!Instance.IsValid() || !Instance.UsesAnimationGraph(Graph))
+	{
+		Graph->AllocateInstance(Instance, &ModuleInstance);
+
+		if(!Instance.IsValid())
+		{
+			return;
+		}
+	}
+
+	if(!Instance.IsValid())
+	{
+		return;
+	}
+
+	if(Instance.RequiresPublicVariableBinding())
+	{
+		// Bind any public variables we find
+		Instance.BindPublicVariables(ExecuteContext.GetTraits());
+	}
+
+	const UE::AnimNext::FReferencePose& RefPose = ReferencePose.ReferencePose.GetRef<UE::AnimNext::FReferencePose>();
+
+	// TODO: Currently forcing additive flag to false here
+	if (Result.LODPose.ShouldPrepareForLOD(RefPose, LOD, false))
+	{
+		Result.LODPose.PrepareForLOD(RefPose, LOD, true, false);
+	}
+
+	ensure(Result.LODPose.LODLevel == LOD);
+
+	// Every graph in a schedule will see the same input events (if they were queued before the schedule started)
+	FTraitEventList InputEventList;
+	FTraitEventList OutputEventList;
+
+	// A schedule can contain multiple graphs, we copy the input event list since it might be appended to during our update
+	{
+		FReadScopeLock ReadLock(ModuleInstance.EventListLock);
+		InputEventList = ModuleInstance.InputEventList;
+	}
+
+	// Track how many input events we started with, we'll append the new ones
+	const int32 NumOriginalInputEvents = InputEventList.Num();
+
+	// Internally we use memstack allocation, so we need a mark here
+	FMemStack& MemStack = FMemStack::Get();
+	FMemMark MemMark(MemStack);
+
+	// We allocate a dummy buffer to trigger the allocation of a large chunk if this is the first mark
+	// This reduces churn internally by avoiding a chunk to be repeatedly allocated and freed as we push/pop marks
+	MemStack.Alloc(size_t(FPageAllocator::SmallPageSize) + 1, 16);
+
+	IAnimNextModuleInterface& AnimNextModule = IAnimNextModuleInterface::Get();
+	AnimNextModule.UpdateGraph(Instance, ExecuteContext.GetDeltaTime(), InputEventList, OutputEventList);
+	AnimNextModule.EvaluateGraph(Instance, RefPose, LOD, Result.LODPose);
+
+	// We might have appended new input/output events, append them
+	{
+		const int32 NumInputEvents = InputEventList.Num();
+
+		FWriteScopeLock WriteLock(ModuleInstance.EventListLock);
+
+		// Append the new input events
+		for (int32 EventIndex = NumOriginalInputEvents; EventIndex < NumInputEvents; ++EventIndex)
+		{
+			FAnimNextTraitEventPtr& Event = InputEventList[EventIndex];
+			if (Event->IsValid())
+			{
+				ModuleInstance.InputEventList.Push(MoveTemp(Event));
+			}
+		}
+
+		// Append our output events
+		ModuleInstance.OutputEventList.Append(OutputEventList);
+	}
+}
diff --git a/Engine/Plugins/Experimental/AnimNext/Source/AnimNext/Private/Graph/RigUnit_AnimNextWriteSkeletalMeshComponentPose.cpp b/Engine/Plugins/Experimental/AnimNext/Source/AnimNext/Private/Graph/RigUnit_AnimNextWriteSkeletalMeshComponentPose.cpp
new file mode 100644
--- /dev/null
+++ b/Engine/Plugins/Experimental/AnimNext/Source/AnimNext/Private/Graph/RigUnit_AnimNextWriteSkeletalMeshComponentPose.cpp
@@ -0,0 +1,55 @@
+// Copyright Epic Games, Inc. All Rights Reserved.
+
+#include "RigUnit_AnimNextWriteSkeletalMeshComponentPose.h"
+
+#include "Units/RigUnitContext.h"
+#include "AnimNextStats.h"
+#include "GenerationTools.h"
+#include "Component/SkinnedMeshComponentExtensions.h"
+#include "Components/SkeletalMeshComponent.h"
+#include "Engine/SkeletalMesh.h"
+
+#include UE_INLINE_GENERATED_CPP_BY_NAME(RigUnit_AnimNextWriteSkeletalMeshComponentPose)
+
+DEFINE_STAT(STAT_AnimNext_Write_Pose);
+
+FRigUnit_AnimNextWriteSkeletalMeshComponentPose_Execute()
+{
+	SCOPE_CYCLE_COUNTER(STAT_AnimNext_Write_Pose);
+
+	using namespace UE::AnimNext;
+
+	if(SkeletalMeshComponent == nullptr)
+	{
+		return;
+	}
+	
+	if(!Pose.LODPose.IsValid())
+	{
+		return;
+	}
+
+	USkeletalMesh* SkeletalMesh = SkeletalMeshComponent->GetSkeletalMeshAsset();
+	if(SkeletalMesh == nullptr)
+	{
+		return;
+	}
+
+	FDataHandle RefPoseHandle = FDataRegistry::Get()->GetOrGenerateReferencePose(SkeletalMeshComponent);
+	const UE::AnimNext::FReferencePose& RefPose = RefPoseHandle.GetRef<UE::AnimNext::FReferencePose>();
+
+	FMemMark MemMark(FMemStack::Get());
+
+	TArray<FTransform, TMemStackAllocator<>> LocalSpaceTransforms;
+	LocalSpaceTransforms.SetNumUninitialized(SkeletalMesh->GetRefSkeleton().GetNum());
+
+	// Map LOD pose into local-space scratch buffer
+	FGenerationTools::RemapPose(Pose.LODPose, LocalSpaceTransforms);
+
+	// Convert and dispatch to renderer
+	UE::Anim::FSkinnedMeshComponentExtensions::CompleteAndDispatch(
+		SkeletalMeshComponent,
+		RefPose.GetParentIndices(),
+		RefPose.GetLODBoneIndexToMeshBoneIndexMap(Pose.LODPose.LODLevel),
+		LocalSpaceTransforms);
+}
diff --git a/Engine/Plugins/Experimental/AnimNext/Source/AnimNext/Private/Graph/RigUnit_AnimNextWriteSkeletalMeshComponentPose.h b/Engine/Plugins/Experimental/AnimNext/Source/AnimNext/Private/Graph/RigUnit_AnimNextWriteSkeletalMeshComponentPose.h
new file mode 100644
--- /dev/null
+++ b/Engine/Plugins/Experimental/AnimNext/Source/AnimNext/Private/Graph/RigUnit_AnimNextWriteSkeletalMeshComponentPose.h
@@ -0,0 +1,34 @@
+﻿// Copyright Epic Games, Inc. All Rights Reserved.
+
+#pragma once
+
+#include "AnimNextExecuteContext.h"
+#include "Graph/AnimNext_LODPose.h"
+#include "Graph/RigUnit_AnimNextBase.h"
+#include "Graph/AnimNextGraphInstancePtr.h"
+
+#include "RigUnit_AnimNextWriteSkeletalMeshComponentPose.generated.h"
+
+/** Writes a pose to a skeletal mesh component */
+USTRUCT(meta=(DisplayName="Write Pose", Category="Animation Graph", NodeColor="0, 1, 1", Keywords="Output,Pose,Port"))
+struct ANIMNEXT_API FRigUnit_AnimNextWriteSkeletalMeshComponentPose : public FRigUnit_AnimNextBase
+{
+	GENERATED_BODY()
+
+	RIGVM_METHOD()
+	void Execute();
+
+	virtual FString GetUnitSubTitle() const { return TEXT("Skeletal Mesh Component"); };
+
+	// Pose to write
+	UPROPERTY(EditAnywhere, Category = "Graph", meta = (Input))
+	FAnimNextGraphLODPose Pose;
+
+	// Mesh to write to
+	UPROPERTY(EditAnywhere, Category = "Graph", meta = (Input))
+	TObjectPtr<USkeletalMeshComponent> SkeletalMeshComponent;
+
+	// The execution result
+	UPROPERTY(EditAnywhere, DisplayName = "Execute", Category = "BeginExecution", meta = (Input, Output))
+	FAnimNextExecuteContext ExecuteContext;
+};
diff --git a/Engine/Plugins/Experimental/AnimNext/Source/AnimNext/Private/Graph/RigUnit_MakeReferencePoseFromSkeletalMeshComponent.cpp b/Engine/Plugins/Experimental/AnimNext/Source/AnimNext/Private/Graph/RigUnit_MakeReferencePoseFromSkeletalMeshComponent.cpp
new file mode 100644
--- /dev/null
+++ b/Engine/Plugins/Experimental/AnimNext/Source/AnimNext/Private/Graph/RigUnit_MakeReferencePoseFromSkeletalMeshComponent.cpp
@@ -0,0 +1,22 @@
+﻿// Copyright Epic Games, Inc. All Rights Reserved.
+
+#include "RigUnit_MakeReferencePoseFromSkeletalMeshComponent.h"
+
+#include "Units/RigUnitContext.h"
+#include "GenerationTools.h"
+#include "Components/SkeletalMeshComponent.h"
+#include "Engine/SkeletalMesh.h"
+
+#include UE_INLINE_GENERATED_CPP_BY_NAME(RigUnit_MakeReferencePoseFromSkeletalMeshComponent)
+
+FRigUnit_MakeReferencePoseFromSkeletalMeshComponent_Execute()
+{
+	using namespace UE::AnimNext;
+
+	if(SkeletalMeshComponent == nullptr)
+	{
+		return;
+	}
+
+	ReferencePose.ReferencePose = FDataRegistry::Get()->GetOrGenerateReferencePose(SkeletalMeshComponent);
+}
diff --git a/Engine/Plugins/Experimental/AnimNext/Source/AnimNext/Private/Graph/RigUnit_MakeReferencePoseFromSkeletalMeshComponent.h b/Engine/Plugins/Experimental/AnimNext/Source/AnimNext/Private/Graph/RigUnit_MakeReferencePoseFromSkeletalMeshComponent.h
new file mode 100644
--- /dev/null
+++ b/Engine/Plugins/Experimental/AnimNext/Source/AnimNext/Private/Graph/RigUnit_MakeReferencePoseFromSkeletalMeshComponent.h
@@ -0,0 +1,34 @@
+﻿// Copyright Epic Games, Inc. All Rights Reserved.
+
+#pragma once
+
+#include "AnimNextExecuteContext.h"
+#include "Graph/AnimNext_LODPose.h"
+#include "Graph/RigUnit_AnimNextBase.h"
+#include "Graph/AnimNextGraphInstancePtr.h"
+
+#include "RigUnit_MakeReferencePoseFromSkeletalMeshComponent.generated.h"
+
+/** Makes a reference pose from a skeletal mesh component */
+USTRUCT(meta=(DisplayName="Make Reference Pose", Category="Animation Graph", NodeColor="0, 1, 1", Keywords="Output,Pose,Port"))
+struct ANIMNEXT_API FRigUnit_MakeReferencePoseFromSkeletalMeshComponent : public FRigUnit_AnimNextBase
+{
+	GENERATED_BODY()
+
+	RIGVM_METHOD()
+	void Execute();
+
+	virtual FString GetUnitSubTitle() const { return TEXT("Skeletal Mesh Component"); };
+
+	// Reference pose to write
+	UPROPERTY(EditAnywhere, Category = "Graph", meta = (Input, Output))
+	FAnimNextGraphReferencePose ReferencePose;
+
+	// Mesh to use to generate the reference pose
+	UPROPERTY(EditAnywhere, Category = "Graph", meta = (Input))
+	TObjectPtr<USkeletalMeshComponent> SkeletalMeshComponent;
+
+	// The execution result
+	UPROPERTY(EditAnywhere, DisplayName = "Execute", Category = "BeginExecution", meta = (Input, Output))
+	FAnimNextExecuteContext ExecuteContext;
+};
diff --git a/Engine/Plugins/Experimental/AnimNext/Source/AnimNext/Private/Graph/RigVMTrait_AnimNextPublicVariables.cpp b/Engine/Plugins/Experimental/AnimNext/Source/AnimNext/Private/Graph/RigVMTrait_AnimNextPublicVariables.cpp
new file mode 100644
--- /dev/null
+++ b/Engine/Plugins/Experimental/AnimNext/Source/AnimNext/Private/Graph/RigVMTrait_AnimNextPublicVariables.cpp
@@ -0,0 +1,61 @@
+﻿// Copyright Epic Games, Inc. All Rights Reserved.
+
+#include "Graph/RigVMTrait_AnimNextPublicVariables.h"
+#include "AnimNextRigVMAsset.h"
+#if WITH_EDITOR
+#include "RigVMModel/RigVMController.h"
+#include "RigVMModel/RigVMPin.h"
+#endif
+
+#if WITH_EDITOR
+FString FRigVMTrait_AnimNextPublicVariables::GetDisplayName() const
+{
+	TStringBuilder<256> StringBuilder;
+	StringBuilder.Appendf(TEXT("Variables: %s"), Asset ? *Asset->GetFName().ToString() : TEXT("None"));
+	return StringBuilder.ToString();
+}
+
+void FRigVMTrait_AnimNextPublicVariables::GetProgrammaticPins(URigVMController* InController, int32 InParentPinIndex, const FString& InDefaultValue, FRigVMPinInfoArray& OutPinArray) const
+{
+	if (Asset == nullptr)
+	{
+		return;
+	}
+
+	const FInstancedPropertyBag& UserDefaults = Asset->GetPublicVariableDefaults();
+	if(!UserDefaults.IsValid())
+	{
+		return;
+	}
+
+	FInstancedPropertyBag Defaults;
+	Defaults.InitializeFromBagStruct(UserDefaults.GetPropertyBagStruct());
+
+	const TFunction<ERigVMPinDefaultValueType(const FName&)> DefaultValueTypeGetter = [&UserDefaults, &Defaults](const FName& InPropertyName)
+	{
+		if(const FProperty* Property = UserDefaults.GetPropertyBagStruct()->FindPropertyByName(InPropertyName))
+		{
+			if(Property->Identical_InContainer(UserDefaults.GetValue().GetMemory(), Defaults.GetValue().GetMemory()))
+			{
+				return ERigVMPinDefaultValueType::Unset;
+			}
+			return ERigVMPinDefaultValueType::Override;
+		}
+		return ERigVMPinDefaultValueType::AutoDetect;
+	};
+
+	OutPinArray.AddPins(const_cast<UPropertyBag*>(UserDefaults.GetPropertyBagStruct()), InController, ERigVMPinDirection::Input, InParentPinIndex, DefaultValueTypeGetter, UserDefaults.GetValue().GetMemory(), true);
+}
+
+bool FRigVMTrait_AnimNextPublicVariables::ShouldCreatePinForProperty(const FProperty* InProperty) const
+{
+	if(!Super::ShouldCreatePinForProperty(InProperty))
+	{
+		return false;
+	}
+	return
+		InProperty->GetFName() == GET_MEMBER_NAME_CHECKED(FRigVMTrait_AnimNextPublicVariables, Asset) ||
+		InProperty->GetFName() == GET_MEMBER_NAME_CHECKED(FRigVMTrait_AnimNextPublicVariables, VariableNames) ||
+		VariableNames.Contains(InProperty->GetFName());
+}
+#endif
\ No newline at end of file
diff --git a/Engine/Plugins/Experimental/AnimNext/Source/AnimNext/Private/Module/AnimNextModule.cpp b/Engine/Plugins/Experimental/AnimNext/Source/AnimNext/Private/Module/AnimNextModule.cpp
--- a/Engine/Plugins/Experimental/AnimNext/Source/AnimNext/Private/Module/AnimNextModule.cpp
+++ b/Engine/Plugins/Experimental/AnimNext/Source/AnimNext/Private/Module/AnimNextModule.cpp
@@ -1,25 +1,15 @@
 // Copyright Epic Games, Inc. All Rights Reserved.
 
 #include "Module/AnimNextModule.h"
-#include "RigVMCore/RigVMMemoryStorage.h"
 #include "RigVMCore/RigVMExecuteContext.h"
 #include "Graph/RigUnit_AnimNextShimRoot.h"
 #include "UObject/AssetRegistryTagsContext.h"
 #include "UObject/ObjectSaveContext.h"
-#include "UObject/Package.h"
-#include "Param/ParamTypeHandle.h"
-#include "TraitCore/TraitReader.h"
 #include "TraitCore/ExecutionContext.h"
 #include "Graph/AnimNext_LODPose.h"
-#include "Graph/GC_GraphInstanceComponent.h"
-#include "Graph/AnimNextGraphInstance.h"
 #include "Serialization/MemoryReader.h"
-#include "AnimNextStats.h"
-#include "RigVMRuntimeDataRegistry.h"
 #include "Graph/AnimNextGraphEntryPoint.h"
 #include "Graph/RigUnit_AnimNextBeginExecution.h"
-#include "Graph/RigUnit_AnimNextGraphRoot.h"
-#include "Param/ParametersProxy.h"
 #include "UObject/LinkerLoad.h"
 #include "UObject/ObjectResource.h"
 
@@ -28,114 +18,26 @@
 REGISTER_ASSETDEPENDENCY_GATHERER(FExternalAssetDependencyGatherer, UAnimNextModule);
 #endif // WITH_EDITOR
 
-DEFINE_STAT(STAT_AnimNext_Graph_AllocateInstance);
-DEFINE_STAT(STAT_AnimNext_Graph_UpdateParamLayer);
-
 #include UE_INLINE_GENERATED_CPP_BY_NAME(AnimNextModule)
 
-UAnimNextModule::UAnimNextModule(const FObjectInitializer& ObjectInitializer)
-	: Super(ObjectInitializer)
-{
-	ExtendedExecuteContext.SetContextPublicDataStruct(FAnimNextExecuteContext::StaticStruct());
-}
-
-void UAnimNextModule::AllocateInstance(FAnimNextGraphInstancePtr& Instance, FName InEntryPoint) const
-{
-	AllocateInstanceImpl(nullptr, Instance, InEntryPoint);
-}
-
-void UAnimNextModule::AllocateInstance(FAnimNextGraphInstance& ParentGraphInstance, FAnimNextGraphInstancePtr& Instance, FName InEntryPoint) const
-{
-	AllocateInstanceImpl(&ParentGraphInstance, Instance, InEntryPoint);
-}
-
-void UAnimNextModule::AllocateInstanceImpl(FAnimNextGraphInstance* ParentGraphInstance, FAnimNextGraphInstancePtr& Instance, FName InEntryPoint) const
-{
-	SCOPE_CYCLE_COUNTER(STAT_AnimNext_Graph_AllocateInstance);
-
-	Instance.Release();
-
-	const FName EntryPoint = (InEntryPoint == NAME_None) ? GetDefaultEntryPoint() : InEntryPoint;
-	const FAnimNextTraitHandle ResolvedRootTraitHandle = ResolvedRootTraitHandles.FindRef(EntryPoint);
-	if (!ResolvedRootTraitHandle.IsValid())
-	{
-		return;
-	}
-
-	{
-		TUniquePtr<FAnimNextGraphInstance> InstanceImpl = MakeUnique<FAnimNextGraphInstance>();
-
-		InstanceImpl->Module = this;
-		InstanceImpl->ParentGraphInstance = ParentGraphInstance;
-		InstanceImpl->EntryPoint = EntryPoint;
-		InstanceImpl->GraphState = MakeUnique<UE::AnimNext::FParametersProxy>(this);
-
-		// If we have a parent graph, use its root since we share the same root, otherwise if we have no parent, we are the root
-		InstanceImpl->RootGraphInstance = ParentGraphInstance != nullptr ? ParentGraphInstance->GetRootGraphInstance() : InstanceImpl.Get();
-
-		InstanceImpl->ExtendedExecuteContext.CopyMemoryStorage(ExtendedExecuteContext);
-		VM->InitializeInstance(InstanceImpl->ExtendedExecuteContext);
-
-		// Move our implementation so that we can use the instance to allocate our root node
-		Instance.Impl = MoveTemp(InstanceImpl);
-	}
-
-	{
-		UE::AnimNext::FExecutionContext Context(Instance);
-		Instance.Impl->GraphInstancePtr = Context.AllocateNodeInstance(*Instance.Impl, ResolvedRootTraitHandle);
-	}
-
-	if (!Instance.IsValid())
-	{
-		// We failed to allocate our instance, clear everything
-		Instance.Release();
-	}
-
 #if WITH_EDITORONLY_DATA
-	if (Instance.IsValid())
-	{
-		FScopeLock Lock(&GraphInstancesLock);
-		check(!GraphInstances.Contains(Instance.Impl.Get()));
-		GraphInstances.Add(Instance.Impl.Get());
-	}
+static UAnimNextModule::FOnModuleCompiled GOnModuleCompiled; 
 #endif
-}
 
-void UAnimNextModule::UpdateLayer(UE::AnimNext::FParamStackLayerHandle& InHandle, float InDeltaTime) const
+UAnimNextModule::UAnimNextModule(const FObjectInitializer& ObjectInitializer)
+	: Super(ObjectInitializer)
 {
-	SCOPE_CYCLE_COUNTER(STAT_AnimNext_Graph_UpdateParamLayer);
-	
-	if (VM)
-	{
-		if (TSharedPtr<UE::AnimNext::FRigVMRuntimeData> RuntimeData = UE::AnimNext::FRigVMRuntimeDataRegistry::FindOrAddLocalRuntimeData(VM, GetRigVMExtendedExecuteContext()).Pin())
-		{
-			FRigVMExtendedExecuteContext& Context = RuntimeData->Context;
-
-			check(Context.VMHash == VM->GetVMHash());
-
-			FAnimNextExecuteContext& AnimNextContext = Context.GetPublicDataSafe<FAnimNextExecuteContext>();
-
-			// Parameter setup
-			AnimNextContext.SetContextData<FAnimNextParamContextData>(InHandle);
-
-			// RigVM setup
-			AnimNextContext.SetDeltaTime(InDeltaTime);
-
-			VM->ExecuteVM(Context, FRigUnit_AnimNextBeginExecution::EventName);
-
-			// Reset the context to avoid issues if we forget to reset it the next time we use it
-			AnimNextContext.DebugReset<FAnimNextParamContextData>();
-		}
-	}
+	ExtendedExecuteContext.SetContextPublicDataStruct(FAnimNextExecuteContext::StaticStruct());
 }
 
 void UAnimNextModule::Serialize(FArchive& Ar)
 {
 	Super::Serialize(Ar);
 
+#if WITH_EDITORONLY_DATA
 	if (Ar.IsLoading())
 	{
-		if(Ar.CustomVer(FFortniteMainBranchObjectVersion::GUID) < FFortniteMainBranchObjectVersion::AnimNextCombineParameterBlocksAndGraphs)
+		if(Ar.CustomVer(FFortniteMainBranchObjectVersion::GUID) < FFortniteMainBranchObjectVersion::AnimNextModuleRefactor)
 		{
 			// Skip over shared archive buffer if we are loading from an older version
 			if (const FLinkerLoad* Linker = GetLinker())
@@ -145,50 +47,8 @@ void UAnimNextModule::Serialize(FArchive& Ar)
 				Ar.Seek(Export.SerialOffset + Export.SerialSize);
 			}
 		}
-		else
-		{
-			int32 SharedDataArchiveBufferSize = 0;
-			Ar << SharedDataArchiveBufferSize;
-
-#if !WITH_EDITORONLY_DATA
-			// When editor data isn't present, we don't persist the archive buffer as it is only needed on load
-			// to populate the graph shared data
-			TArray<uint8> SharedDataArchiveBuffer;
-#endif
-
-			SharedDataArchiveBuffer.SetNumUninitialized(SharedDataArchiveBufferSize);
-			Ar.Serialize(SharedDataArchiveBuffer.GetData(), SharedDataArchiveBufferSize);
-
-			if (Ar.IsLoadingFromCookedPackage())
-			{
-				// If we are cooked, we populate our graph shared data otherwise in the editor we'll compile on load
-				// and re-populate everything then to account for changes in code/content
-				LoadFromArchiveBuffer(SharedDataArchiveBuffer);
-			}
-		}
-
-		CacheDefaultEntryPoint();
 	}
-	else if (Ar.IsSaving())
-	{
-#if WITH_EDITORONLY_DATA
-		// We only save the archive buffer, if code changes we'll be able to de-serialize from it when
-		// building the runtime buffer
-		// This allows us to have editor only/non-shipping only properties that are stripped out on load
-		int32 SharedDataArchiveBufferSize = SharedDataArchiveBuffer.Num();
-		Ar << SharedDataArchiveBufferSize;
-		Ar.Serialize(SharedDataArchiveBuffer.GetData(), SharedDataArchiveBufferSize);
-#endif
-	}
-	else
-	{
-		// Counting, etc
-		Ar << SharedDataBuffer;
-
-#if WITH_EDITORONLY_DATA
-		Ar << SharedDataArchiveBuffer;
 #endif
-	}
 }
 
 void UAnimNextModule::PostLoad()
@@ -198,111 +58,16 @@ void UAnimNextModule::PostLoad()
 #if WITH_EDITORONLY_DATA
 	if(GetLinkerCustomVersion(FFortniteMainBranchObjectVersion::GUID) < FFortniteMainBranchObjectVersion::AnimNextGraphAccessSpecifiers)
 	{
-		DefaultState.State = PropertyBag_DEPRECATED;
+		DefaultState_DEPRECATED.State = PropertyBag_DEPRECATED;
 	}
 #endif
 }
 
 #if WITH_EDITOR
-void UAnimNextModule::PostEditChangeProperty(FPropertyChangedEvent& PropertyChangedEvent)
-{
-	if(PropertyChangedEvent.Property && PropertyChangedEvent.Property->GetFName() == GET_MEMBER_NAME_CHECKED(UAnimNextModule, DefaultEntryPoint))
-	{
-		CacheDefaultEntryPoint(); 
-	}
-}
-#endif
-
-TConstArrayView<UE::AnimNext::FScheduleTerm> UAnimNextModule::GetTerms() const
-{
-	using namespace UE::AnimNext;
-
-	static const FParamId OutputPoseId("UE_Internal_Graph_OutputPose");
-
-	static const FScheduleTerm Terms[] =
-	{
-		FScheduleTerm(OutputPoseId, FAnimNextParamType::GetType<FAnimNextGraphLODPose>(), EScheduleTermDirection::Output)
-	};
-
-	return Terms;
-}
-
-bool UAnimNextModule::LoadFromArchiveBuffer(const TArray<uint8>& InSharedDataArchiveBuffer)
-{
-	using namespace UE::AnimNext;
-
-	// Reconstruct our graph shared data
-	FMemoryReader GraphSharedDataArchive(InSharedDataArchiveBuffer);
-	FTraitReader TraitReader(GraphReferencedObjects, GraphSharedDataArchive);
-
-	const FTraitReader::EErrorState ErrorState = TraitReader.ReadGraph(SharedDataBuffer);
-	if (ErrorState == FTraitReader::EErrorState::None)
-	{
-		for(int32 EntryPointIndex = 0; EntryPointIndex < EntryPoints.Num(); ++EntryPointIndex)
-		{
-			const FAnimNextGraphEntryPoint& EntryPoint = EntryPoints[EntryPointIndex];
-			ResolvedRootTraitHandles.Add(EntryPoint.EntryPointName, TraitReader.ResolveEntryPointHandle(EntryPoint.RootTraitHandle));
-			ResolvedEntryPoints.Add(EntryPoint.EntryPointName, EntryPointIndex);
-		}
-
-		// Make sure our execute method is registered
-		FRigUnit_AnimNextGraphEvaluator::RegisterExecuteMethod(ExecuteDefinition);
-		return true;
-	}
-	else
-	{
-		SharedDataBuffer.Empty(0);
-		ResolvedRootTraitHandles.Add(GetDefaultEntryPoint(), FAnimNextTraitHandle());
-		return false;
-	}
-}
-
-#if WITH_EDITORONLY_DATA
-void UAnimNextModule::FreezeGraphInstances()
-{
-	FScopeLock Lock(&GraphInstancesLock);
-
-	TSet<FAnimNextGraphInstance*> GraphInstancesCopy = GraphInstances;
-	for (FAnimNextGraphInstance* GraphInstance : GraphInstancesCopy)
-	{
-		GraphInstance->Freeze();
-	}
-}
-
-void UAnimNextModule::ThawGraphInstances()
-{
-	FScopeLock Lock(&GraphInstancesLock);
-
-	TSet<FAnimNextGraphInstance*> GraphInstancesCopy = GraphInstances;
-	for (FAnimNextGraphInstance* GraphInstance : GraphInstancesCopy)
-	{
-		GraphInstance->Thaw();
-	}
-}
-#endif
-
-void UAnimNextModule::SetDefaultEntryPoint(FName InEntryPoint)
-{
-	DefaultEntryPoint = InEntryPoint;
-
-	CacheDefaultEntryPoint();
-}
 
-FName UAnimNextModule::GetDefaultEntryPoint() const
+UAnimNextModule::FOnModuleCompiled& UAnimNextModule::OnModuleCompiled()
 {
-	if(CachedDefaultEntryPoint == NAME_None)
-	{
-		CacheDefaultEntryPoint();
-	}
-
-	return CachedDefaultEntryPoint;
+	return GOnModuleCompiled;
 }
 
-void UAnimNextModule::CacheDefaultEntryPoint() const
-{
-	TStringBuilder<256> StringBuilder;
-	StringBuilder.Append(GetPathName());
-	StringBuilder.Append(TEXT(":"));
-	DefaultEntryPoint.AppendString(StringBuilder);
-	CachedDefaultEntryPoint = FName(StringBuilder.ToView());
-}
\ No newline at end of file
+#endif
\ No newline at end of file
diff --git a/Engine/Plugins/Experimental/AnimNext/Source/AnimNext/Private/Module/AnimNextModuleInstance.cpp b/Engine/Plugins/Experimental/AnimNext/Source/AnimNext/Private/Module/AnimNextModuleInstance.cpp
new file mode 100644
--- /dev/null
+++ b/Engine/Plugins/Experimental/AnimNext/Source/AnimNext/Private/Module/AnimNextModuleInstance.cpp
@@ -0,0 +1,364 @@
+// Copyright Epic Games, Inc. All Rights Reserved.
+
+#include "Module/AnimNextModuleInstance.h"
+
+#include "Engine/World.h"
+#include "AnimNextStats.h"
+#include "Algo/TopologicalSort.h"
+#include "Logging/StructuredLog.h"
+#include "Misc/EnumerateRange.h"
+#include "Module/RigUnit_AnimNextModuleEvents.h"
+#include "Module/ProxyVariablesContext.h"
+#include "UObject/UObjectIterator.h"
+#include "Variables/IAnimNextVariableProxyHost.h"
+
+DEFINE_STAT(STAT_AnimNext_InitializeInstance);
+
+FAnimNextModuleInstance::FAnimNextModuleInstance(
+		UAnimNextModule* InModule,
+		UObject* InObject,
+		EAnimNextModuleInitMethod InInitMethod)
+	: Module(InModule)
+	, Object(InObject)
+	, RunState(ERunState::None)
+	, InitMethod(InInitMethod)
+{
+}
+
+FAnimNextModuleInstance::~FAnimNextModuleInstance()
+{
+	ResetBindingsAndInstanceData();
+
+	Object = nullptr;
+	Module = nullptr;
+	Handle.Reset();
+}
+
+namespace UE::AnimNext::Private
+{
+
+struct FImplementedModuleEvent
+{
+	UScriptStruct* Struct = nullptr;
+	FModuleEventBindingFunction Binding;
+	FName EventName;
+	EModuleEventPhase Phase = EModuleEventPhase::Execute;
+};
+
+static TArray<FImplementedModuleEvent> GAllModuleEvents;
+static TArray<FImplementedModuleEvent> GImplementedModuleEvents;
+
+void CacheAllModuleEvents()
+{
+	check(IsInGameThread());	// This function cannot be run concurrently because of static usage
+
+	if(GAllModuleEvents.Num() == 0)
+	{
+		for(TObjectIterator<UScriptStruct> It; It; ++It)
+		{
+			UScriptStruct* Struct = *It;
+			if(Struct->IsChildOf(FRigUnit_AnimNextModuleEventBase::StaticStruct()) && Struct != FRigUnit_AnimNextModuleEventBase::StaticStruct())
+			{
+				TInstancedStruct<FRigUnit_AnimNextModuleEventBase> StructInstance;
+				StructInstance.InitializeAsScriptStruct(Struct);
+				FImplementedModuleEvent& NewEvent = GAllModuleEvents.AddDefaulted_GetRef();
+				NewEvent.Struct = Struct;
+				NewEvent.Binding = StructInstance.Get().GetBindingFunction();
+				NewEvent.EventName = StructInstance.Get().GetEventName();
+				NewEvent.Phase = StructInstance.Get().GetEventPhase();
+			}
+		}
+
+		Algo::SortBy(GAllModuleEvents, &FImplementedModuleEvent::Phase);
+		GImplementedModuleEvents.Reserve(GAllModuleEvents.Num());
+	}
+}
+
+// Gets information about the module events that are implemented by the supplied VM
+static TConstArrayView<FImplementedModuleEvent> GetImplementedModuleEvents(URigVM* VM)
+{
+	check(IsInGameThread());	// This function cannot be run concurrently because of static usage
+	check(GAllModuleEvents.Num() > 0);	// Call CacheAllModuleEvents before this function
+
+	// Get all the module events from the VM entry points
+	const TArray<FName>& EntryNames = VM->GetEntryNames();
+	GImplementedModuleEvents.Reset();
+	for(const FImplementedModuleEvent& ModuleEvent : GAllModuleEvents)
+	{
+		if(EntryNames.Contains(ModuleEvent.EventName))
+		{
+			GImplementedModuleEvents.Add(ModuleEvent);
+		}
+	}
+
+	return GImplementedModuleEvents;
+}
+
+}
+
+void FAnimNextModuleInstance::Initialize()
+{
+	SCOPE_CYCLE_COUNTER(STAT_AnimNext_InitializeInstance);
+	
+	using namespace UE::AnimNext;
+
+	check(IsInGameThread());
+
+	check(Object)
+	check(Module);
+	check(Handle.IsValid());
+
+	UWorld* World = Object->GetWorld();
+	bIsEditor = World && World->WorldType == EWorldType::Editor;
+
+	// Get all the module events from the VM entry points, sorted by phase
+	URigVM* VM = Module->GetVM();
+	TConstArrayView<Private::FImplementedModuleEvent> ImplementedModuleEvents = Private::GetImplementedModuleEvents(VM);
+
+	// Setup tick function graph using module events
+	if (ImplementedModuleEvents.Num() > 0)
+	{
+		TransitionToRunState(ERunState::CreatingTasks);
+
+		// Allocate tick functions
+		TickFunctions.SetNum(ImplementedModuleEvents.Num()); 
+		for (TEnumerateRef<const Private::FImplementedModuleEvent> ModuleEvent : EnumerateRange(ImplementedModuleEvents))
+		{
+			FModuleEventTickFunction& TickFunction = TickFunctions[ModuleEvent.GetIndex()];
+			TickFunction.ModuleInstance = this;
+			TickFunction.EventName = ModuleEvent->EventName;
+		}
+
+		// Bind tick functions (done in a second pass to allow functions to add prerequisites between each other)
+		for (TEnumerateRef<const Private::FImplementedModuleEvent> ModuleEvent : EnumerateRange(ImplementedModuleEvents))
+		{
+			FModuleEventTickFunction& TickFunction = TickFunctions[ModuleEvent.GetIndex()];
+			// Add prerequisite on end tick function
+			EndTickFunction.AddPrerequisite(Object, TickFunction);
+			FTickFunctionBindingContext Context(*this, Object, World, ModuleEvent.GetIndex());
+			ModuleEvent->Binding(Context, TickFunction);
+		}
+
+		// Bind end tick function to this instance
+		EndTickFunction.ModuleInstance = this;
+
+		TransitionToRunState(ERunState::BindingTasks);
+
+		// Register our tick functions
+		if(World)
+		{
+			ULevel* Level = World->PersistentLevel;
+			for (FModuleEventTickFunction& TickFunction : TickFunctions)
+			{
+				TickFunction.RegisterTickFunction(Level);
+			}
+			EndTickFunction.RegisterTickFunction(Level);
+		}
+
+		TransitionToRunState(ERunState::PendingInitialUpdate);
+
+		// Initialize variables
+		const int32 NumVariables = Module->VariableDefaults.GetNumPropertiesInBag();
+#if WITH_EDITOR
+		if(bIsRecreatingOnCompile)
+		{
+			Variables.MigrateToNewBagInstance(Module->VariableDefaults);
+		}
+		else
+#endif
+		{
+			Variables = Module->VariableDefaults;
+		}
+
+		if(Module->GetPublicVariableDefaults().GetPropertyBagStruct())
+		{
+			PublicVariablesProxy.Data = Module->GetPublicVariableDefaults();
+			TConstArrayView<FPropertyBagPropertyDesc> ProxyDescs = PublicVariablesProxy.Data.GetPropertyBagStruct()->GetPropertyDescs();
+			PublicVariablesProxy.DirtyFlags.SetNum(ProxyDescs.Num(), false);
+		}
+
+		// Initialize the RigVM context
+		ExtendedExecuteContext = Module->GetRigVMExtendedExecuteContext();
+
+		// Setup external variables memory ptrs manually as we dont follow the pattern of owning multiple URigVMHosts like control rig.
+		// InitializeVM() is called, but only sets up handles for the defaults in the module, not for an instance
+		TArray<FRigVMExternalVariableRuntimeData> ExternalVariableRuntimeData;
+		ExternalVariableRuntimeData.Reserve(NumVariables);
+		TConstArrayView<FPropertyBagPropertyDesc> Descs = Variables.GetPropertyBagStruct()->GetPropertyDescs();
+		uint8* BasePtr = Variables.GetMutableValue().GetMemory();
+		for(int32 VariableIndex = 0; VariableIndex < NumVariables; ++VariableIndex)
+		{
+			ExternalVariableRuntimeData.Emplace(Descs[VariableIndex].CachedProperty->ContainerPtrToValuePtr<uint8>(BasePtr));
+		}
+		ExtendedExecuteContext.ExternalVariableRuntimeData = MoveTemp(ExternalVariableRuntimeData);
+
+		// Now initialize the 'instance', cache memory handles etc. in the context
+		VM->InitializeInstance(ExtendedExecuteContext);
+
+		// Just pause now if we arent needing an initial update
+		if(InitMethod == EAnimNextModuleInitMethod::None)
+		{
+			Enable(false);
+		}
+#if WITH_EDITOR
+		else if(World)
+		{
+			// In editor worlds we run a linearized 'initial tick' to ensure we generate an initial output pose, as these worlds dont always tick
+			if( World->WorldType == EWorldType::Editor ||
+				World->WorldType == EWorldType::EditorPreview)
+			{
+				FModuleEventTickFunction::InitializeAndRunModule(*this);
+			}
+		}
+#endif
+	}
+}
+
+void FAnimNextModuleInstance::ResetBindingsAndInstanceData()
+{
+	using namespace UE::AnimNext;
+
+	check(IsInGameThread());
+
+	TransitionToRunState(ERunState::None);
+
+	for (FModuleEventTickFunction& TickFunction : TickFunctions)
+	{
+		TickFunction.UnRegisterTickFunction();
+	}
+
+	EndTickFunction.UnRegisterTickFunction();
+
+	TickFunctions.Reset();
+
+	ExtendedExecuteContext.Reset();
+
+#if WITH_EDITOR
+	if(!bIsRecreatingOnCompile)
+#endif
+	{
+		Variables.Reset();
+	}
+}
+
+void FAnimNextModuleInstance::Invalidate()
+{
+	using namespace UE::AnimNext;
+
+	ResetBindingsAndInstanceData();
+
+	Module = nullptr;
+	Object = nullptr;
+	Handle.Reset();
+}
+
+void FAnimNextModuleInstance::ClearTickFunctionPauseFlags()
+{
+	using namespace UE::AnimNext;
+
+	check(IsInGameThread());
+
+	for (FModuleEventTickFunction& TickFunction : TickFunctions)
+	{
+		TickFunction.bTickEvenWhenPaused = false;
+	}
+	EndTickFunction.bTickEvenWhenPaused = false;
+}
+
+void FAnimNextModuleInstance::QueueInputTraitEvent(FAnimNextTraitEventPtr Event)
+{
+	InputEventList.Push(MoveTemp(Event));
+}
+
+void FAnimNextModuleInstance::Enable(bool bInEnabled)
+{
+	using namespace UE::AnimNext;
+
+	check(IsInGameThread());
+
+	if(RunState == ERunState::PendingInitialUpdate || RunState == ERunState::Paused || RunState == ERunState::Running)
+	{
+		for (FModuleEventTickFunction& TickFunction : TickFunctions)
+		{
+			TickFunction.SetTickFunctionEnable(bInEnabled);
+		}
+		EndTickFunction.SetTickFunctionEnable(bInEnabled);
+
+		TransitionToRunState(bInEnabled ? ERunState::Running : ERunState::Paused);
+	}
+}
+
+void FAnimNextModuleInstance::TransitionToRunState(ERunState InNewState)
+{
+	switch(InNewState)
+	{
+	case ERunState::None:
+		check(RunState == ERunState::None || RunState == ERunState::PendingInitialUpdate || RunState == ERunState::Paused || RunState == ERunState::Running);
+		break;
+	case ERunState::CreatingTasks:
+		check(RunState == ERunState::None);
+		break;
+	case ERunState::BindingTasks:
+		check(RunState == ERunState::CreatingTasks);
+		break;
+	case ERunState::PendingInitialUpdate:
+		check(RunState == ERunState::BindingTasks);
+		break;
+	case ERunState::Running:
+		check(RunState == ERunState::PendingInitialUpdate || RunState == ERunState::Paused || RunState == ERunState::Running);
+		break;
+	case ERunState::Paused:
+		check(RunState == ERunState::PendingInitialUpdate || RunState == ERunState::Paused || RunState == ERunState::Running);
+		break;
+	default:
+		checkNoEntry();
+	}
+
+	RunState = InNewState;
+}
+
+void FAnimNextModuleInstance::CopyProxyVariables()
+{
+	// TODO: we can avoid the copies here by adopting a scheme where we:
+	// - Hold double-buffered memory handles
+	// - Update the memory handle's ptr to the currently-written double-buffered public variable on write
+	// - Swap the memory handles in ExtendedExecuteContext here
+	if(IAnimNextVariableProxyHost* ProxyHost = Cast<IAnimNextVariableProxyHost>(Object))
+	{
+		// Flip the proxy
+		ProxyHost->FlipPublicVariablesProxy(UE::AnimNext::FProxyVariablesContext(*this));
+
+		if(PublicVariablesProxy.bIsDirty)
+		{
+			// Copy dirty properties
+			TConstArrayView<FPropertyBagPropertyDesc> ProxyDescs = Variables.GetPropertyBagStruct()->GetPropertyDescs();
+			TConstArrayView<FPropertyBagPropertyDesc> PublicProxyDescs = PublicVariablesProxy.Data.GetPropertyBagStruct()->GetPropertyDescs();
+			const uint8* SourceContainerPtr = PublicVariablesProxy.Data.GetValue().GetMemory();
+			uint8* TargetContainerPtr = Variables.GetMutableValue().GetMemory();
+			for (TBitArray<>::FConstIterator It(PublicVariablesProxy.DirtyFlags); It; ++It)
+			{
+				const int32 Index = It.GetIndex();
+				const FProperty* SourceProperty = PublicProxyDescs[Index].CachedProperty;
+				const FProperty* TargetProperty = ProxyDescs[Index].CachedProperty;
+				checkSlow(SourceProperty->GetClass() == TargetProperty->GetClass());
+				ProxyDescs[Index].CachedProperty->CopyCompleteValue_InContainer(TargetContainerPtr, SourceContainerPtr);
+			}
+
+			// Reset dirty flags
+			PublicVariablesProxy.DirtyFlags.SetRange(0, PublicVariablesProxy.DirtyFlags.Num(), false);
+			PublicVariablesProxy.bIsDirty = false;
+		}
+	}
+}
+
+#if WITH_EDITOR
+void FAnimNextModuleInstance::OnModuleCompiled()
+{
+	using namespace UE::AnimNext;
+
+	FGuardValue_Bitfield(bIsRecreatingOnCompile, true);
+
+	ResetBindingsAndInstanceData();
+	Initialize();
+}
+#endif
\ No newline at end of file
