AnimNext: Module refactor

diff --git a/Engine/Plugins/Experimental/AnimNext/Source/AnimNext/Private/Module/ModuleTickFunction.cpp b/Engine/Plugins/Experimental/AnimNext/Source/AnimNext/Private/Module/ModuleTickFunction.cpp
new file mode 100644
--- /dev/null
+++ b/Engine/Plugins/Experimental/AnimNext/Source/AnimNext/Private/Module/ModuleTickFunction.cpp
@@ -0,0 +1,218 @@
+// Copyright Epic Games, Inc. All Rights Reserved.
+
+#include "Module/ModuleTickFunction.h"
+
+#include "RigVMRuntimeDataRegistry.h"
+#include "Algo/TopologicalSort.h"
+#include "Module/AnimNextModule.h"
+#include "Module/ModuleTaskContext.h"
+#include "Module/ProxyVariablesContext.h"
+#include "Module/RigUnit_AnimNextModuleEvents.h"
+#include "Module/ModuleEvents.h"
+#include "TraitCore/TraitEventList.h"
+
+namespace UE::AnimNext
+{
+
+void FModuleEndTickFunction::ExecuteTick(float DeltaTime, ELevelTick TickType, ENamedThreads::Type CurrentThread, const FGraphEventRef& MyCompletionGraphEvent)
+{
+	Run();
+}
+
+void FModuleEndTickFunction::Run()
+{
+	// Decrement the remaining lifetime of the input events we processed and queue up any remaining events
+	DecrementLifetimeAndPurgeExpired(ModuleInstance->InputEventList, ModuleInstance->OutputEventList);
+
+	// Filter out our schedule action events, we'll hand them off to the main thread to execute
+	FTraitEventList MainThreadActionEventList;
+	if (!ModuleInstance->OutputEventList.IsEmpty())
+	{
+		for (FAnimNextTraitEventPtr& Event : ModuleInstance->OutputEventList)
+		{
+			if (!Event->IsValid())
+			{
+				continue;
+			}
+
+			if (FAnimNextModule_ActionEvent* ActionEvent = Event->AsType<FAnimNextModule_ActionEvent>())
+			{
+				if (ActionEvent->IsThreadSafe())
+				{
+					// Execute this action now
+					ActionEvent->Execute();
+				}
+				else
+				{
+					// Defer this action and execute it on the main thread
+					MainThreadActionEventList.Push(Event);
+				}
+			}
+		}
+
+		// Reset our list of output events, we don't retain any
+		ModuleInstance->OutputEventList.Reset();
+	}
+	
+	auto RunTaskOnGameThread = [](TUniqueFunction<void(void)>&& InFunction)
+	{
+		if(IsInGameThread())
+		{
+			InFunction();
+		}
+		else
+		{
+			FFunctionGraphTask::CreateAndDispatchWhenReady(MoveTemp(InFunction), TStatId(), nullptr, ENamedThreads::GameThread);
+		}
+	};
+
+	auto DisableInitializeEvent = [this]()
+	{
+		if(ModuleInstance->TickFunctions.Num() > 0 && ModuleInstance->TickFunctions[0].EventName == FRigUnit_AnimNextInitializeEvent::EventName)
+		{
+			// Disable the initialize event as we have already run it
+			ModuleInstance->TickFunctions[0].SetTickFunctionEnable(false);
+		}
+	};
+
+	if(ModuleInstance->RunState == FAnimNextModuleInstance::ERunState::PendingInitialUpdate)
+	{
+		if(ModuleInstance->InitMethod == EAnimNextModuleInitMethod::InitializeAndPause
+#if WITH_EDITOR
+			|| (ModuleInstance->InitMethod == EAnimNextModuleInitMethod::InitializeAndPauseInEditor && ModuleInstance->bIsEditor)
+#endif
+			)
+		{
+			// Queue task to disable our tick functions now we have performed our initial update
+			RunTaskOnGameThread([this, &DisableInitializeEvent]()
+			{
+				check(IsInGameThread());
+				DisableInitializeEvent();
+				ModuleInstance->Enable(false);
+			});
+		}
+	}
+	else if(ModuleInstance->RunState != FAnimNextModuleInstance::ERunState::Running)
+	{
+		RunTaskOnGameThread([this, &DisableInitializeEvent]()
+		{
+			check(IsInGameThread());
+			DisableInitializeEvent();
+			ModuleInstance->TransitionToRunState(FAnimNextModuleInstance::ERunState::Running);
+		});
+	}
+
+	if (!MainThreadActionEventList.IsEmpty())
+	{
+		RunTaskOnGameThread([MainThreadActionEventList = MoveTemp(MainThreadActionEventList)]()
+			{
+				check(IsInGameThread());
+				for (const FAnimNextTraitEventPtr& Event : MainThreadActionEventList)
+				{
+					FAnimNextModule_ActionEvent* ActionEvent = Event->AsType<FAnimNextModule_ActionEvent>();
+					ActionEvent->Execute();
+				}
+			});
+	}
+}
+
+FString FModuleEndTickFunction::DiagnosticMessage()
+{
+	return TEXT("AnimNext: ModuleEnd");
+}
+
+void FModuleEventTickFunction::ExecuteTick(float DeltaTime, ELevelTick TickType, ENamedThreads::Type CurrentThread, const FGraphEventRef& MyCompletionGraphEvent)
+{
+	Run(DeltaTime);
+}
+
+void FModuleEventTickFunction::Run(float InDeltaTime)
+{
+	while (!PreExecuteTasks.IsEmpty())
+	{
+		TOptional<TUniqueFunction<void(const FModuleTaskContext&)>> Function = PreExecuteTasks.Dequeue();
+		check(Function.IsSet());
+		Function.GetValue()(FModuleTaskContext(*ModuleInstance));
+	}
+
+	if (URigVM* VM = ModuleInstance->Module->GetVM())
+	{
+		FRigVMExtendedExecuteContext& Context = ModuleInstance->ExtendedExecuteContext;
+		check(Context.VMHash == VM->GetVMHash());
+
+		FAnimNextExecuteContext& AnimNextContext = Context.GetPublicDataSafe<FAnimNextExecuteContext>();
+
+		// RigVM setup
+		AnimNextContext.SetDeltaTime(InDeltaTime);
+
+		// Module setup
+		AnimNextContext.SetContextData<FAnimNextModuleContextData>(ModuleInstance);
+
+		// Run the VM for this event
+		VM->ExecuteVM(Context, EventName);
+
+		// Reset the context to avoid issues if we forget to reset it the next time we use it
+		AnimNextContext.DebugReset<FAnimNextModuleContextData>();
+	}
+
+	while (!PostExecuteTasks.IsEmpty())
+	{
+		TOptional<TUniqueFunction<void(const FModuleTaskContext&)>> Function = PostExecuteTasks.Dequeue();
+		check(Function.IsSet());
+		Function.GetValue()(FModuleTaskContext(*ModuleInstance));
+	}
+}
+
+#if WITH_EDITOR
+
+void FModuleEventTickFunction::InitializeAndRunModule(FAnimNextModuleInstance& InModuleInstance)
+{
+	// Sort tick functions topologically
+	TArray<FModuleEventTickFunction*> TickFunctionPtrs;
+	TickFunctionPtrs.Reserve(InModuleInstance.TickFunctions.Num());
+	Algo::Transform(InModuleInstance.TickFunctions, TickFunctionPtrs, [](const FModuleEventTickFunction& InTickFunction){ return const_cast<FModuleEventTickFunction*>(&InTickFunction); });
+	TArray<FModuleEventTickFunction*> DependencyPtrs;
+	Algo::TopologicalSort(
+		TickFunctionPtrs,
+		[&DependencyPtrs, &InModuleInstance](FModuleEventTickFunction* InTickFunction)
+		{
+			DependencyPtrs.Reset();
+
+			// Add direct prereqs
+			for(const FTickPrerequisite& Prerequisite : InTickFunction->GetPrerequisites())
+			{
+				DependencyPtrs.Add(static_cast<FModuleEventTickFunction*>(Prerequisite.PrerequisiteTickFunction));
+			}
+
+			// Also add any tick functions in prior tick groups that we don't already have
+			for(FModuleEventTickFunction& TickFunction : InModuleInstance.TickFunctions)
+			{
+				if(TickFunction.TickGroup < InTickFunction->TickGroup && !DependencyPtrs.Contains(&TickFunction))
+				{
+					DependencyPtrs.Add(&TickFunction);
+				}
+			}
+
+			return DependencyPtrs;
+		},
+		Algo::ETopologicalSort::None);
+
+	// Run sorted tick functions
+	for(FModuleEventTickFunction* TickFunction : TickFunctionPtrs)
+	{
+		TickFunction->Run(0.0f);
+	}
+	InModuleInstance.EndTickFunction.Run();
+}
+
+#endif
+
+FString FModuleEventTickFunction::DiagnosticMessage()
+{
+	TStringBuilder<256> Builder;
+	Builder.Append(TEXT("AnimNext: "));
+	EventName.AppendString(Builder);
+	return Builder.ToString();
+}
+
+}
diff --git a/Engine/Plugins/Experimental/AnimNext/Source/AnimNext/Private/Module/ProxyVariablesContext.cpp b/Engine/Plugins/Experimental/AnimNext/Source/AnimNext/Private/Module/ProxyVariablesContext.cpp
new file mode 100644
--- /dev/null
+++ b/Engine/Plugins/Experimental/AnimNext/Source/AnimNext/Private/Module/ProxyVariablesContext.cpp
@@ -0,0 +1,14 @@
+﻿// Copyright Epic Games, Inc. All Rights Reserved.
+
+#include "Module/ProxyVariablesContext.h"
+#include "Module/AnimNextModuleInstance.h"
+
+namespace UE::AnimNext
+{
+
+FAnimNextPublicVariablesProxy& FProxyVariablesContext::GetPublicVariablesProxy() const
+{
+	return ModuleInstance.PublicVariablesProxy;
+}
+
+}
\ No newline at end of file
diff --git a/Engine/Plugins/Experimental/AnimNext/Source/AnimNext/Private/Module/RigUnit_AnimNextAnimationGraphSelf.cpp b/Engine/Plugins/Experimental/AnimNext/Source/AnimNext/Private/Module/RigUnit_AnimNextAnimationGraphSelf.cpp
new file mode 100644
--- /dev/null
+++ b/Engine/Plugins/Experimental/AnimNext/Source/AnimNext/Private/Module/RigUnit_AnimNextAnimationGraphSelf.cpp
@@ -0,0 +1,14 @@
+﻿// Copyright Epic Games, Inc. All Rights Reserved.
+
+#include "Module/RigUnit_AnimNextAnimationGraphSelf.h"
+#include "Graph/AnimNextGraphInstance.h"
+
+FRigUnit_AnimNextAnimationGraphSelf_Execute()
+{
+	// Deprecated stub
+}
+
+FRigVMStructUpgradeInfo FRigUnit_AnimNextAnimationGraphSelf::GetUpgradeInfo() const
+{
+	return FRigVMStructUpgradeInfo();
+}
\ No newline at end of file
diff --git a/Engine/Plugins/Experimental/AnimNext/Source/AnimNext/Private/Module/RigUnit_AnimNextModuleEvents.cpp b/Engine/Plugins/Experimental/AnimNext/Source/AnimNext/Private/Module/RigUnit_AnimNextModuleEvents.cpp
new file mode 100644
--- /dev/null
+++ b/Engine/Plugins/Experimental/AnimNext/Source/AnimNext/Private/Module/RigUnit_AnimNextModuleEvents.cpp
@@ -0,0 +1,68 @@
+﻿// Copyright Epic Games, Inc. All Rights Reserved.
+
+#include "Module/RigUnit_AnimNextModuleEvents.h"
+
+#include "Module/AnimNextModuleInstance.h"
+
+FRigUnit_AnimNextExecuteBindings_Execute()
+{
+}
+
+UE::AnimNext::FModuleEventBindingFunction FRigUnit_AnimNextExecuteBindings::GetBindingFunction() const
+{
+	return [](const UE::AnimNext::FTickFunctionBindingContext& InContext, FTickFunction& InTickFunction)
+	{
+		InTickFunction.TickGroup = ETickingGroup::TG_PrePhysics;
+		InTickFunction.bRunOnAnyThread = false; // TODO: This will need to vary depending on whether the event is deemed thread-safe 
+
+		// Add prerequisite on this for all subsequent tick functions (assumes tick functions are sorted by phase)
+		for(int32 TickFunctionIndex = InContext.EventIndex + 1; TickFunctionIndex < InContext.ModuleInstance.TickFunctions.Num(); ++TickFunctionIndex)
+		{
+			FTickFunction& TickFunctions = InContext.ModuleInstance.TickFunctions[TickFunctionIndex];
+			TickFunctions.AddPrerequisite(InContext.ModuleInstance.Object, InContext.ModuleInstance.TickFunctions[InContext.EventIndex]);
+		}
+	};
+}
+
+FRigUnit_AnimNextInitializeEvent_Execute()
+{
+}
+
+UE::AnimNext::FModuleEventBindingFunction FRigUnit_AnimNextInitializeEvent::GetBindingFunction() const
+{
+	return [](const UE::AnimNext::FTickFunctionBindingContext& InContext, FTickFunction& InTickFunction)
+	{
+		InTickFunction.TickGroup = ETickingGroup::TG_PrePhysics;
+
+		// Add prerequisite on this for all subsequent tick functions (assumes tick functions are sorted by phase)
+		for(int32 TickFunctionIndex = InContext.EventIndex + 1; TickFunctionIndex < InContext.ModuleInstance.TickFunctions.Num(); ++TickFunctionIndex)
+		{
+			FTickFunction& TickFunction = InContext.ModuleInstance.TickFunctions[TickFunctionIndex];
+			TickFunction.AddPrerequisite(InContext.ModuleInstance.Object, InContext.ModuleInstance.TickFunctions[InContext.EventIndex]);
+		}
+	};
+}
+
+FRigUnit_AnimNextPrePhysicsEvent_Execute()
+{
+}
+
+UE::AnimNext::FModuleEventBindingFunction FRigUnit_AnimNextPrePhysicsEvent::GetBindingFunction() const
+{
+	return [](const UE::AnimNext::FTickFunctionBindingContext& InContext, FTickFunction& InTickFunction)
+	{
+		InTickFunction.TickGroup = ETickingGroup::TG_PrePhysics;
+	};
+}
+
+FRigUnit_AnimNextPostPhysicsEvent_Execute()
+{
+}
+
+UE::AnimNext::FModuleEventBindingFunction FRigUnit_AnimNextPostPhysicsEvent::GetBindingFunction() const
+{
+	return [](const UE::AnimNext::FTickFunctionBindingContext& InContext, FTickFunction& InTickFunction)
+	{
+		InTickFunction.TickGroup = ETickingGroup::TG_PostPhysics;
+	};
+}
\ No newline at end of file
diff --git a/Engine/Plugins/Experimental/AnimNext/Source/AnimNext/Private/Param/AnimNextWorldLibrary.cpp b/Engine/Plugins/Experimental/AnimNext/Source/AnimNext/Private/Param/AnimNextWorldLibrary.cpp
--- a/Engine/Plugins/Experimental/AnimNext/Source/AnimNext/Private/Param/AnimNextWorldLibrary.cpp
+++ b/Engine/Plugins/Experimental/AnimNext/Source/AnimNext/Private/Param/AnimNextWorldLibrary.cpp
@@ -2,7 +2,6 @@
 
 #include "AnimNextWorldLibrary.h"
 #include "Engine/World.h"
-#include "Scheduler/AnimNextTickFunctionBinding.h"
 
 double UAnimNextWorldLibrary::GetTimeSeconds(UWorld* InWorld)
 {
@@ -23,11 +22,3 @@ float UAnimNextWorldLibrary::GetDeltaSeconds(UWorld* InWorld)
 {
 	return InWorld->GetDeltaSeconds();
 }
-
-FAnimNextTickFunctionBinding UAnimNextWorldLibrary::GetPhysicsTick(UWorld* InWorld)
-{
-	FAnimNextTickFunctionBinding Binding;
-	Binding.Object = InWorld;
-	Binding.TickFunction = &InWorld->EndPhysicsTickFunction;
-	return Binding;
-}
diff --git a/Engine/Plugins/Experimental/AnimNext/Source/AnimNext/Private/Param/AnimNextWorldLibrary.h b/Engine/Plugins/Experimental/AnimNext/Source/AnimNext/Private/Param/AnimNextWorldLibrary.h
--- a/Engine/Plugins/Experimental/AnimNext/Source/AnimNext/Private/Param/AnimNextWorldLibrary.h
+++ b/Engine/Plugins/Experimental/AnimNext/Source/AnimNext/Private/Param/AnimNextWorldLibrary.h
@@ -31,8 +31,4 @@ public:
 	// Returns the frame delta time in seconds adjusted by e.g. time dilation.
 	UFUNCTION(BlueprintCallable, BlueprintInternalUseOnly)
 	static float GetDeltaSeconds(UWorld* InWorld);
-
-	// Returns the physics tick function
-	UFUNCTION(BlueprintCallable, BlueprintInternalUseOnly)
-	static FAnimNextTickFunctionBinding GetPhysicsTick(UWorld* InWorld);
 };
diff --git a/Engine/Plugins/Experimental/AnimNext/Source/AnimNext/Private/Param/ParamType.cpp b/Engine/Plugins/Experimental/AnimNext/Source/AnimNext/Private/Param/ParamType.cpp
--- a/Engine/Plugins/Experimental/AnimNext/Source/AnimNext/Private/Param/ParamType.cpp
+++ b/Engine/Plugins/Experimental/AnimNext/Source/AnimNext/Private/Param/ParamType.cpp
@@ -3,7 +3,6 @@
 #include "Param/ParamType.h"
 
 #include "Graph/AnimNext_LODPose.h"
-#include "Param/ParamTypeHandle.h"
 #include "Misc/StringBuilder.h"
 #include "UObject/Class.h"
 #include "Templates/SubclassOf.h"
@@ -170,6 +169,112 @@ FAnimNextParamType FAnimNextParamType::FromRigVMTemplateArgument(const FRigVMTem
 	return Type;
 }
 
+FAnimNextParamType FAnimNextParamType::FromProperty(const FProperty* InProperty)
+{
+	FAnimNextParamType Type;
+
+	if(InProperty == nullptr)
+	{
+		return Type;
+	}
+
+	if (InProperty->IsA<FArrayProperty>())
+	{
+		Type.ContainerType = EPropertyBagContainerType::Array;
+		InProperty = CastField<FArrayProperty>(InProperty)->Inner;
+	}
+
+	if (InProperty->IsA<FBoolProperty>())
+	{
+		Type.ValueType = EPropertyBagPropertyType::Bool;
+	}
+	else if (InProperty->IsA<FByteProperty>())
+	{
+		Type.ValueType = EPropertyBagPropertyType::Byte;
+	}
+	else if (InProperty->IsA<FIntProperty>())
+	{
+		Type.ValueType = EPropertyBagPropertyType::Int32;
+	}
+	else if (InProperty->IsA<FInt64Property>())
+	{
+		Type.ValueType = EPropertyBagPropertyType::Int64;
+	}
+	else if (InProperty->IsA<FUInt32Property>())
+	{
+		Type.ValueType = EPropertyBagPropertyType::UInt32;
+	}
+	else if (InProperty->IsA<FUInt64Property>())
+	{
+		Type.ValueType = EPropertyBagPropertyType::UInt64;
+	}
+	else if (InProperty->IsA<FFloatProperty>())
+	{
+		Type.ValueType = EPropertyBagPropertyType::Float;
+	}
+	else if (InProperty->IsA<FDoubleProperty>())
+	{
+		Type.ValueType = EPropertyBagPropertyType::Double;
+	}
+	else if (InProperty->IsA<FNameProperty>())
+	{
+		Type.ValueType = EPropertyBagPropertyType::Name;
+	}
+	else if (InProperty->IsA<FStrProperty>())
+	{
+		Type.ValueType = EPropertyBagPropertyType::String;
+	}
+	else if (InProperty->IsA<FTextProperty>())
+	{
+		Type.ValueType = EPropertyBagPropertyType::Text;
+	}
+	else if (InProperty->IsA<FStructProperty>())
+	{
+		Type.ValueType = EPropertyBagPropertyType::Struct;
+		Type.ValueTypeObject = CastField<FStructProperty>(InProperty)->Struct;
+	}
+	else if(InProperty->IsA<FObjectPropertyBase>())
+	{
+		UClass* Class = CastField<FObjectPropertyBase>(InProperty)->PropertyClass;
+		if (InProperty->IsA<FClassProperty>())
+		{
+			Type.ValueType = EPropertyBagPropertyType::Class;
+			Type.ValueTypeObject = CastField<FClassProperty>(InProperty)->MetaClass;
+		}
+		else if (InProperty->IsA<FObjectProperty>())
+		{
+			if(Class == UClass::StaticClass())
+			{
+				Type.ValueType = EPropertyBagPropertyType::Class;
+				Type.ValueTypeObject = Class;
+			}
+			else
+			{
+				Type.ValueType = EPropertyBagPropertyType::Object;
+				Type.ValueTypeObject = Class;
+			}
+		}
+		else if (InProperty->IsA<FSoftClassProperty>())
+		{
+			Type.ValueType = EPropertyBagPropertyType::SoftClass;
+			Type.ValueTypeObject = CastField<FSoftClassProperty>(InProperty)->MetaClass;
+		}
+		else if (InProperty->IsA<FSoftObjectProperty>())
+		{
+			Type.ValueType = EPropertyBagPropertyType::SoftObject;
+			Type.ValueTypeObject = Class;
+		}
+	}
+	else if(InProperty->IsA<FEnumProperty>())
+	{
+		Type.ValueType = EPropertyBagPropertyType::Enum;
+		Type.ValueTypeObject = CastField<FEnumProperty>(InProperty)->GetEnum();
+	}
+
+	return Type;
+}
+
+
 bool FAnimNextParamType::IsValidObject() const
 {
 	switch(ValueType)
@@ -207,201 +312,6 @@ bool FAnimNextParamType::IsValidObject() const
 	}
 }
 
-UE::AnimNext::FParamTypeHandle FAnimNextParamType::GetHandle() const
-{
-	using namespace UE::AnimNext;
-
-	FParamTypeHandle Handle;
-	
-	switch(ContainerType)
-	{
-	case EPropertyBagContainerType::None:
-		switch(ValueType)
-		{
-		default:
-		case EValueType::None:
-			Handle.SetParameterType(FParamTypeHandle::EParamType::None);
-			break;
-		case EValueType::Bool:
-			Handle.SetParameterType(FParamTypeHandle::EParamType::Bool);
-			break;
-		case EValueType::Byte:
-			Handle.SetParameterType(FParamTypeHandle::EParamType::Byte);
-			break;
-		case EValueType::Int32:
-			Handle.SetParameterType(FParamTypeHandle::EParamType::Int32);
-			break;
-		case EValueType::Int64:
-			Handle.SetParameterType(FParamTypeHandle::EParamType::Int64);
-			break;
-		case EValueType::Float:
-			Handle.SetParameterType(FParamTypeHandle::EParamType::Float);
-			break;
-		case EValueType::Double:
-			Handle.SetParameterType(FParamTypeHandle::EParamType::Double);
-			break;
-		case EValueType::Name:
-			Handle.SetParameterType(FParamTypeHandle::EParamType::Name);
-			break;
-		case EValueType::String:
-			Handle.SetParameterType(FParamTypeHandle::EParamType::String);
-			break;
-		case EValueType::Text:
-			Handle.SetParameterType(FParamTypeHandle::EParamType::Text);
-			break;
-		case EValueType::Struct:
-			if(const UScriptStruct* ScriptStruct = Cast<UScriptStruct>(ValueTypeObject.Get()))
-			{
-				if(ScriptStruct == TBaseStructure<FVector>::Get())
-				{
-					Handle.SetParameterType(FParamTypeHandle::EParamType::Vector);
-				}
-				else if(ScriptStruct == TBaseStructure<FVector4>::Get())
-				{
-					Handle.SetParameterType(FParamTypeHandle::EParamType::Vector4);
-				}
-				else if(ScriptStruct == TBaseStructure<FQuat>::Get())
-				{
-					Handle.SetParameterType(FParamTypeHandle::EParamType::Quat);
-				}
-				else if(ScriptStruct == TBaseStructure<FTransform>::Get())
-				{
-					Handle.SetParameterType(FParamTypeHandle::EParamType::Transform);
-				}
-				else if(ScriptStruct == FAnimNextGraphLODPose::StaticStruct())
-				{
-					Handle.SetParameterType(FParamTypeHandle::EParamType::AnimNextGraphLODPose);
-				}
-				else if(ScriptStruct == FAnimNextGraphReferencePose::StaticStruct())
-				{
-					Handle.SetParameterType(FParamTypeHandle::EParamType::AnimNextGraphReferencePose);
-				}
-				else
-				{
-					Handle.SetParameterType(FParamTypeHandle::EParamType::Custom);
-					Handle.SetCustomTypeIndex(FParamTypeHandle::GetOrAllocateCustomTypeIndex(EValueType::Struct, EContainerType::None, ScriptStruct));
-				}
-			}
-			else
-			{
-				Handle.SetParameterType(FParamTypeHandle::EParamType::None);
-			}
-			break;
-		case EValueType::Enum:
-			if(const UEnum* Enum = Cast<UEnum>(ValueTypeObject.Get()))
-			{
-				Handle.SetParameterType(FParamTypeHandle::EParamType::Custom);
-				Handle.SetCustomTypeIndex(FParamTypeHandle::GetOrAllocateCustomTypeIndex(ValueType, EContainerType::None, Enum));
-			}
-			else
-			{
-				Handle.SetParameterType(FParamTypeHandle::EParamType::None);
-			}
-			break;
-		case EValueType::Object:
-			if(const UClass* Class = Cast<UClass>(ValueTypeObject.Get()))
-			{
-				if (Class == UObject::StaticClass())
-				{
-					Handle.SetParameterType(FParamTypeHandle::EParamType::Object);
-					break;
-				}
-				else if (Class == UCharacterMovementComponent::StaticClass())
-				{
-					Handle.SetParameterType(FParamTypeHandle::EParamType::CharacterMovementComponent);
-					break;
-				}
-				else if (Class == USkeletalMeshComponent::StaticClass())
-				{
-					Handle.SetParameterType(FParamTypeHandle::EParamType::SkeletalMeshComponent);
-					break;
-				}
-				else if (Class == UAnimSequence::StaticClass())
-				{
-					Handle.SetParameterType(FParamTypeHandle::EParamType::AnimSequence);
-					break;
-				}
-			}
-			// fall through
-		case EValueType::SoftObject:
-		case EValueType::Class:
-		case EValueType::SoftClass:
-			if(const UClass* Class = Cast<UClass>(ValueTypeObject.Get()))
-			{
-				Handle.SetParameterType(FParamTypeHandle::EParamType::Custom);
-				Handle.SetCustomTypeIndex(FParamTypeHandle::GetOrAllocateCustomTypeIndex(ValueType, EContainerType::None, Class));
-			}
-			else
-			{
-				Handle.SetParameterType(FParamTypeHandle::EParamType::None);
-			}
-			break;
-		}
-		break;
-	case EPropertyBagContainerType::Array:
-		switch(ValueType)
-		{
-		default:
-		case EValueType::None:
-			Handle.SetParameterType(FParamTypeHandle::EParamType::None);
-			break;
-		case EValueType::Bool:
-		case EValueType::Byte:
-		case EValueType::Int32:
-		case EValueType::Int64:
-		case EValueType::Float:
-		case EValueType::Double:
-		case EValueType::Name:
-		case EValueType::String:
-		case EValueType::Text:
-			Handle.SetParameterType(FParamTypeHandle::EParamType::Custom);
-			Handle.SetCustomTypeIndex(FParamTypeHandle::GetOrAllocateCustomTypeIndex(ValueType, ContainerType, nullptr));
-			break;
-		case EValueType::Struct:
-			if(const UScriptStruct* ScriptStruct = Cast<UScriptStruct>(ValueTypeObject.Get()))
-			{
-				Handle.SetParameterType(FParamTypeHandle::EParamType::Custom);
-				Handle.SetCustomTypeIndex(FParamTypeHandle::GetOrAllocateCustomTypeIndex(ValueType, ContainerType, ScriptStruct));
-			}
-			else
-			{
-				Handle.SetParameterType(FParamTypeHandle::EParamType::None);
-			}
-			break;
-		case EValueType::Enum:
-			if(const UEnum* Enum = Cast<UEnum>(ValueTypeObject.Get()))
-			{
-				Handle.SetParameterType(FParamTypeHandle::EParamType::Custom);
-				Handle.SetCustomTypeIndex(FParamTypeHandle::GetOrAllocateCustomTypeIndex(ValueType, ContainerType, Enum));
-			}
-			else
-			{
-				Handle.SetParameterType(FParamTypeHandle::EParamType::None);
-			}
-			break;
-		case EValueType::Object:
-		case EValueType::SoftObject:
-		case EValueType::Class:
-		case EValueType::SoftClass:
-			if(const UClass* Class = Cast<UClass>(ValueTypeObject.Get()))
-			{
-				Handle.SetParameterType(FParamTypeHandle::EParamType::Custom);
-				Handle.SetCustomTypeIndex(FParamTypeHandle::GetOrAllocateCustomTypeIndex(ValueType, ContainerType, Class));
-			}
-			else
-			{
-				Handle.SetParameterType(FParamTypeHandle::EParamType::None);
-			}
-			break;
-		}
-		break;
-	default:
-		break;
-	}
-
-	return Handle;
-}
-
 size_t FAnimNextParamType::GetSize() const
 {
 	switch(ContainerType)
diff --git a/Engine/Plugins/Experimental/AnimNext/Source/AnimNext/Private/Param/ParamUtils.cpp b/Engine/Plugins/Experimental/AnimNext/Source/AnimNext/Private/Param/ParamUtils.cpp
--- a/Engine/Plugins/Experimental/AnimNext/Source/AnimNext/Private/Param/ParamUtils.cpp
+++ b/Engine/Plugins/Experimental/AnimNext/Source/AnimNext/Private/Param/ParamUtils.cpp
@@ -7,395 +7,186 @@
 #include "GameFramework/CharacterMovementComponent.h"
 #include "Animation/AnimSequence.h"
 #include "Param/ParamType.h"
-#include "Param/ParamTypeHandle.h"
 #include "Param/ParamCompatibility.h"
 
 namespace UE::AnimNext
 {
 
-FParamCompatibility FParamUtils::GetCompatibility(const FParamTypeHandle& InLHS, const FParamTypeHandle& InRHS)
+FParamCompatibility FParamUtils::GetCompatibility(const FAnimNextParamType& InLHS, const FAnimNextParamType& InRHS)
 {
-	auto CheckClassCastToCustom = [](const FParamTypeHandle& InLHS, const UClass* InRHSClass)
+	switch (InRHS.GetValueType())
 	{
-		FAnimNextParamType::EValueType ValueTypeLHS;
-		FAnimNextParamType::EContainerType ContainerTypeLHS;
-		const UObject* ValueTypeObjectLHS;
-
-		InLHS.GetCustomTypeInfo(ValueTypeLHS, ContainerTypeLHS, ValueTypeObjectLHS);
-		if(ContainerTypeLHS == FAnimNextParamType::EContainerType::None && ValueTypeLHS == FAnimNextParamType::EValueType::Object)
+	case EPropertyBagPropertyType::Bool:
+		switch (InLHS.GetValueType())
 		{
-			if(const UClass* Class = Cast<UClass>(ValueTypeObjectLHS))
-			{
-				if(InRHSClass->IsChildOf(Class))
-				{
-					return true;
-				}
-			}
-		}
-		return false;
-	};
-
-	auto CheckClassCastFromCustom = [](const UClass* InLHSClass, const FParamTypeHandle& InRHS)
-	{
-		FAnimNextParamType::EValueType ValueTypeRHS;
-		FAnimNextParamType::EContainerType ContainerTypeRHS;
-		const UObject* ValueTypeObjectRHS;
-
-		InRHS.GetCustomTypeInfo(ValueTypeRHS, ContainerTypeRHS, ValueTypeObjectRHS);
-		if(ContainerTypeRHS == FAnimNextParamType::EContainerType::None && ValueTypeRHS == FAnimNextParamType::EValueType::Object)
-		{
-			if(const UClass* Class = Cast<UClass>(ValueTypeObjectRHS))
-			{
-				if(Class->IsChildOf(InLHSClass))
-				{
-					return true;
-				}
-			}
-		}
-		return false;
-	};
-	
-	switch (InRHS.GetParameterType())
-	{
-	case FParamTypeHandle::EParamType::Bool:
-		switch (InLHS.GetParameterType())
-		{
-		case FParamTypeHandle::EParamType::Bool:
+		case EPropertyBagPropertyType::Bool:
 			return EParamCompatibility::Compatible_Equal;
 		}
 		break;
-	case FParamTypeHandle::EParamType::Byte:
-		switch (InLHS.GetParameterType())
+	case EPropertyBagPropertyType::Byte:
+		switch (InLHS.GetValueType())
 		{
-		case FParamTypeHandle::EParamType::Byte:
+		case EPropertyBagPropertyType::Byte:
 			return EParamCompatibility::Compatible_Equal;
-		case FParamTypeHandle::EParamType::Int32:
+		case EPropertyBagPropertyType::Int32:
 			return EParamCompatibility::Compatible_Promotion;
-		case FParamTypeHandle::EParamType::Int64:
+		case EPropertyBagPropertyType::Int64:
 			return EParamCompatibility::Compatible_Promotion;
-		case FParamTypeHandle::EParamType::Float:
+		case EPropertyBagPropertyType::Float:
 			return EParamCompatibility::Compatible_Promotion;
-		case FParamTypeHandle::EParamType::Double:
+		case EPropertyBagPropertyType::Double:
 			return EParamCompatibility::Compatible_Promotion;
 		}
 		break;
-	case FParamTypeHandle::EParamType::Int32:
-		switch (InLHS.GetParameterType())
+	case EPropertyBagPropertyType::Int32:
+		switch (InLHS.GetValueType())
 		{
-		case FParamTypeHandle::EParamType::Byte:
+		case EPropertyBagPropertyType::Byte:
 			return EParamCompatibility::Incompatible_DataLoss;
-		case FParamTypeHandle::EParamType::Int32:
+		case EPropertyBagPropertyType::Int32:
 			return EParamCompatibility::Compatible_Equal;
-		case FParamTypeHandle::EParamType::Int64:
+		case EPropertyBagPropertyType::Int64:
 			return EParamCompatibility::Compatible_Promotion;
-		case FParamTypeHandle::EParamType::Float:
+		case EPropertyBagPropertyType::Float:
 			return EParamCompatibility::Incompatible_DataLoss;
-		case FParamTypeHandle::EParamType::Double:
+		case EPropertyBagPropertyType::Double:
 			return EParamCompatibility::Compatible_Promotion;
 		}
 		break;
-	case FParamTypeHandle::EParamType::Int64:
-		switch (InLHS.GetParameterType())
+	case EPropertyBagPropertyType::Int64:
+		switch (InLHS.GetValueType())
 		{
-		case FParamTypeHandle::EParamType::Byte:
+		case EPropertyBagPropertyType::Byte:
 			return EParamCompatibility::Incompatible_DataLoss;
-		case FParamTypeHandle::EParamType::Int32:
+		case EPropertyBagPropertyType::Int32:
 			return EParamCompatibility::Incompatible_DataLoss;
-		case FParamTypeHandle::EParamType::Int64:
+		case EPropertyBagPropertyType::Int64:
 			return EParamCompatibility::Compatible_Equal;
-		case FParamTypeHandle::EParamType::Float:
+		case EPropertyBagPropertyType::Float:
 			return EParamCompatibility::Incompatible_DataLoss;
-		case FParamTypeHandle::EParamType::Double:
+		case EPropertyBagPropertyType::Double:
 			return EParamCompatibility::Incompatible_DataLoss;
 		}
 		break;
-	case FParamTypeHandle::EParamType::Float:
-		switch (InLHS.GetParameterType())
+	case EPropertyBagPropertyType::Float:
+		switch (InLHS.GetValueType())
 		{
-		case FParamTypeHandle::EParamType::Byte:
+		case EPropertyBagPropertyType::Byte:
 			return EParamCompatibility::Incompatible_DataLoss;
-		case FParamTypeHandle::EParamType::Int32:
+		case EPropertyBagPropertyType::Int32:
 			return EParamCompatibility::Incompatible_DataLoss;
-		case FParamTypeHandle::EParamType::Int64:
+		case EPropertyBagPropertyType::Int64:
 			return EParamCompatibility::Incompatible_DataLoss;
-		case FParamTypeHandle::EParamType::Float:
+		case EPropertyBagPropertyType::Float:
 			return EParamCompatibility::Compatible_Equal;
-		case FParamTypeHandle::EParamType::Double:
+		case EPropertyBagPropertyType::Double:
 			return EParamCompatibility::Compatible_Promotion;
 		}
 		break;
-	case FParamTypeHandle::EParamType::Double:
-		switch (InLHS.GetParameterType())
+	case EPropertyBagPropertyType::Double:
+		switch (InLHS.GetValueType())
 		{
-		case FParamTypeHandle::EParamType::Byte:
+		case EPropertyBagPropertyType::Byte:
 			return EParamCompatibility::Incompatible_DataLoss;
-		case FParamTypeHandle::EParamType::Int32:
+		case EPropertyBagPropertyType::Int32:
 			return EParamCompatibility::Incompatible_DataLoss;
-		case FParamTypeHandle::EParamType::Int64:
+		case EPropertyBagPropertyType::Int64:
 			return EParamCompatibility::Incompatible_DataLoss;
-		case FParamTypeHandle::EParamType::Float:
+		case EPropertyBagPropertyType::Float:
 			return EParamCompatibility::Incompatible_DataLoss;
-		case FParamTypeHandle::EParamType::Double:
+		case EPropertyBagPropertyType::Double:
 			return EParamCompatibility::Compatible_Equal;
 		}
 		break;
-	case FParamTypeHandle::EParamType::Name:
-		switch (InLHS.GetParameterType())
+	case EPropertyBagPropertyType::Name:
+		switch (InLHS.GetValueType())
 		{
-		case FParamTypeHandle::EParamType::Name:
+		case EPropertyBagPropertyType::Name:
 			return EParamCompatibility::Compatible_Equal;
 		}
 		break;
-	case FParamTypeHandle::EParamType::String:
-		switch (InLHS.GetParameterType())
+	case EPropertyBagPropertyType::String:
+		switch (InLHS.GetValueType())
 		{
-		case FParamTypeHandle::EParamType::String:
+		case EPropertyBagPropertyType::String:
 			return EParamCompatibility::Compatible_Equal;
 		}
 		break;
-	case FParamTypeHandle::EParamType::Text:
-		switch (InLHS.GetParameterType())
+	case EPropertyBagPropertyType::Text:
+		switch (InLHS.GetValueType())
 		{
-		case FParamTypeHandle::EParamType::Text:
+		case EPropertyBagPropertyType::Text:
 			return EParamCompatibility::Compatible_Equal;
 		}
 		break;
-	case FParamTypeHandle::EParamType::Vector:
-		switch (InLHS.GetParameterType())
+	case EPropertyBagPropertyType::Struct:
+		switch (InLHS.GetValueType())
 		{
-		case FParamTypeHandle::EParamType::Vector:
-			return EParamCompatibility::Compatible_Equal;
-		}
-		break;
-	case FParamTypeHandle::EParamType::Vector4:
-		switch (InLHS.GetParameterType())
-		{
-		case FParamTypeHandle::EParamType::Vector4:
-			return EParamCompatibility::Compatible_Equal;
-		}
-		break;
-	case FParamTypeHandle::EParamType::Quat:
-		switch (InLHS.GetParameterType())
-		{
-		case FParamTypeHandle::EParamType::Quat:
-			return EParamCompatibility::Compatible_Equal;
-		}
-		break;
-	case FParamTypeHandle::EParamType::Transform:
-		switch (InLHS.GetParameterType())
-		{
-		case FParamTypeHandle::EParamType::Transform:
-			return EParamCompatibility::Compatible_Equal;
-		}
-		break;
-	case FParamTypeHandle::EParamType::Object:
-		switch (InLHS.GetParameterType())
-		{
-		case FParamTypeHandle::EParamType::Object:
-			return EParamCompatibility::Compatible_Equal;
-		case FParamTypeHandle::EParamType::Custom:
-			if(CheckClassCastToCustom(InLHS, UObject::StaticClass()))
+		case EPropertyBagPropertyType::Struct:
+			if(InLHS.GetValueTypeObject() == InRHS.GetValueTypeObject())
 			{
-				return EParamCompatibility::Compatible_Cast;
+				return EParamCompatibility::Compatible_Equal;
 			}
-		}
-		break;
-	case FParamTypeHandle::EParamType::CharacterMovementComponent:
-		switch (InLHS.GetParameterType())
-		{
-		case FParamTypeHandle::EParamType::Object:
-			return EParamCompatibility::Compatible_Cast;
-		case FParamTypeHandle::EParamType::CharacterMovementComponent:
-			return EParamCompatibility::Compatible_Equal;
-		case FParamTypeHandle::EParamType::Custom:
-			if(CheckClassCastToCustom(InLHS, UCharacterMovementComponent::StaticClass()))
+			else if(InLHS.GetValueTypeObject() && InRHS.GetValueTypeObject() && CastChecked<UScriptStruct>(InRHS.GetValueTypeObject())->IsChildOf(CastChecked<UScriptStruct>(InLHS.GetValueTypeObject())))
 			{
 				return EParamCompatibility::Compatible_Cast;
 			}
 		}
 		break;
-	case FParamTypeHandle::EParamType::SkeletalMeshComponent:
-		switch (InLHS.GetParameterType())
+	case EPropertyBagPropertyType::Object:
+		switch (InLHS.GetValueType())
 		{
-		case FParamTypeHandle::EParamType::Object:
-			return EParamCompatibility::Compatible_Cast;
-		case FParamTypeHandle::EParamType::SkeletalMeshComponent:
-			return EParamCompatibility::Compatible_Equal;
-		case FParamTypeHandle::EParamType::Custom:
-			if(CheckClassCastToCustom(InLHS, USkeletalMeshComponent::StaticClass()))
+		case EPropertyBagPropertyType::Object:
+			if(InLHS.GetValueTypeObject() == InRHS.GetValueTypeObject())
+			{
+				return EParamCompatibility::Compatible_Equal;
+			}
+			if(InLHS.GetValueTypeObject() && InRHS.GetValueTypeObject() && InRHS.GetValueTypeObject()->GetClass()->IsChildOf(InLHS.GetValueTypeObject()->GetClass()))
 			{
 				return EParamCompatibility::Compatible_Cast;
 			}
 		}
 		break;
-	case FParamTypeHandle::EParamType::AnimSequence:
-		switch (InLHS.GetParameterType())
+	case EPropertyBagPropertyType::SoftObject:
+		switch (InLHS.GetValueType())
 		{
-		case FParamTypeHandle::EParamType::Object:
-			return EParamCompatibility::Compatible_Cast;
-		case FParamTypeHandle::EParamType::AnimSequence:
-			return EParamCompatibility::Compatible_Equal;
-		case FParamTypeHandle::EParamType::Custom:
-			if(CheckClassCastToCustom(InLHS, UAnimSequence::StaticClass()))
+		case EPropertyBagPropertyType::SoftObject:
+			if(InLHS.GetValueTypeObject() == InRHS.GetValueTypeObject())
+			{
+				return EParamCompatibility::Compatible_Equal;
+			}
+			if(InLHS.GetValueTypeObject() && InRHS.GetValueTypeObject() && InRHS.GetValueTypeObject()->GetClass()->IsChildOf(InLHS.GetValueTypeObject()->GetClass()))
 			{
 				return EParamCompatibility::Compatible_Cast;
 			}
 		}
 		break;
-	case FParamTypeHandle::EParamType::AnimNextGraphLODPose:
-		switch (InLHS.GetParameterType())
-		{
-		case FParamTypeHandle::EParamType::AnimNextGraphLODPose:
-			return EParamCompatibility::Compatible_Equal;
-		}
-		break;
-	case FParamTypeHandle::EParamType::AnimNextGraphReferencePose:
-		switch (InLHS.GetParameterType())
-		{
-		case FParamTypeHandle::EParamType::AnimNextGraphReferencePose:
-			return EParamCompatibility::Compatible_Equal;
-		}
-		break;
-	case FParamTypeHandle::EParamType::Custom:
-		switch (InLHS.GetParameterType())
+	case EPropertyBagPropertyType::Class:
+		switch (InLHS.GetValueType())
 		{
-		case FParamTypeHandle::EParamType::Object:
-			if(CheckClassCastFromCustom(UObject::StaticClass(), InRHS))
+		case EPropertyBagPropertyType::Class:
+			if(InLHS.GetValueTypeObject() == InRHS.GetValueTypeObject())
 			{
-				return EParamCompatibility::Compatible_Cast;
+				return EParamCompatibility::Compatible_Equal;
 			}
-			break;
-		case FParamTypeHandle::EParamType::CharacterMovementComponent:
-			if(CheckClassCastFromCustom(UCharacterMovementComponent::StaticClass(), InRHS))
+			if(InLHS.GetValueTypeObject() && InRHS.GetValueTypeObject() && CastChecked<UClass>(InRHS.GetValueTypeObject())->IsChildOf(CastChecked<UClass>(InLHS.GetValueTypeObject())))
 			{
 				return EParamCompatibility::Compatible_Cast;
 			}
-			break;
-		case FParamTypeHandle::EParamType::SkeletalMeshComponent:
-			if(CheckClassCastFromCustom(USkeletalMeshComponent::StaticClass(), InRHS))
+		}
+		break;
+	case EPropertyBagPropertyType::SoftClass:
+		switch (InLHS.GetValueType())
+		{
+		case EPropertyBagPropertyType::SoftClass:
+			if(InLHS.GetValueTypeObject() == InRHS.GetValueTypeObject())
 			{
-				return EParamCompatibility::Compatible_Cast;
+				return EParamCompatibility::Compatible_Equal;
 			}
-			break;
-		case FParamTypeHandle::EParamType::AnimSequence:
-			if(CheckClassCastFromCustom(UAnimSequence::StaticClass(), InRHS))
+			if(InLHS.GetValueTypeObject() && InRHS.GetValueTypeObject() && CastChecked<UClass>(InRHS.GetValueTypeObject())->IsChildOf(CastChecked<UClass>(InLHS.GetValueTypeObject())))
 			{
 				return EParamCompatibility::Compatible_Cast;
 			}
-			break;
-		case FParamTypeHandle::EParamType::Custom:
-			{
-				FAnimNextParamType::EValueType ValueTypeLHS, ValueTypeRHS;
-				FAnimNextParamType::EContainerType ContainerTypeLHS, ContainerTypeRHS;
-				const UObject* ValueTypeObjectLHS, *ValueTypeObjectRHS;
-
-				InLHS.GetCustomTypeInfo(ValueTypeLHS, ContainerTypeLHS, ValueTypeObjectLHS);
-				InRHS.GetCustomTypeInfo(ValueTypeRHS, ContainerTypeRHS, ValueTypeObjectRHS);
-
-				if (ContainerTypeLHS == ContainerTypeRHS)
-				{
-					if (ContainerTypeLHS == FAnimNextParamType::EContainerType::Array)
-					{
-						if (ValueTypeLHS != ValueTypeRHS)
-						{
-							switch (ValueTypeLHS)
-							{
-							case FAnimNextParamType::EValueType::Float:
-								if(ValueTypeRHS == FAnimNextParamType::EValueType::Double)
-								{
-									return EParamCompatibility::Incompatible_DataLoss;
-								}
-								break;
-							case FAnimNextParamType::EValueType::Double:
-								if (ValueTypeRHS == FAnimNextParamType::EValueType::Float)
-								{
-									return EParamCompatibility::Compatible_Promotion;
-								}
-								break;
-							case FAnimNextParamType::EValueType::Struct:
-							case FAnimNextParamType::EValueType::Object:
-							case FAnimNextParamType::EValueType::SoftObject:
-							case FAnimNextParamType::EValueType::Class:
-							case FAnimNextParamType::EValueType::SoftClass:
-								if (ValueTypeObjectLHS == ValueTypeObjectRHS)
-								{
-									return EParamCompatibility::Compatible_Equal;
-								}
-								break;
-							}
-						}
-						else
-						{
-							switch (ValueTypeLHS)
-							{
-							case FAnimNextParamType::EValueType::Bool:
-							case FAnimNextParamType::EValueType::Byte:
-							case FAnimNextParamType::EValueType::Int32:
-							case FAnimNextParamType::EValueType::Int64:
-							case FAnimNextParamType::EValueType::Float:
-							case FAnimNextParamType::EValueType::Double:
-							case FAnimNextParamType::EValueType::Name:
-							case FAnimNextParamType::EValueType::String:
-							case FAnimNextParamType::EValueType::Text:
-								return EParamCompatibility::Compatible_Equal;
-							case FAnimNextParamType::EValueType::Enum:
-							case FAnimNextParamType::EValueType::Struct:
-							case FAnimNextParamType::EValueType::Object:
-							case FAnimNextParamType::EValueType::SoftObject:
-							case FAnimNextParamType::EValueType::Class:
-							case FAnimNextParamType::EValueType::SoftClass:
-								if (ValueTypeObjectLHS == ValueTypeObjectRHS)
-								{
-									return EParamCompatibility::Compatible_Equal;
-								}
-								break;
-							}
-						}
-					}
-					else
-					{
-						if (ValueTypeLHS == ValueTypeRHS)
-						{
-							switch (ValueTypeLHS)
-							{
-							default:
-							case FAnimNextParamType::EValueType::None:
-								return EParamCompatibility::Incompatible;
-							case FAnimNextParamType::EValueType::Enum:
-								if (ValueTypeObjectLHS == ValueTypeObjectRHS)
-								{
-									return EParamCompatibility::Compatible_Equal;
-								}
-								break;
-							case FAnimNextParamType::EValueType::Struct:
-								if (ValueTypeObjectLHS == ValueTypeObjectRHS)
-								{
-									return EParamCompatibility::Compatible_Equal;
-								}
-								else if (CastChecked<UScriptStruct>(ValueTypeObjectRHS)->IsChildOf(CastChecked<UScriptStruct>(ValueTypeObjectLHS)))
-								{
-									return EParamCompatibility::Compatible_Cast;
-								}
-								break;
-							case FAnimNextParamType::EValueType::Object:
-							case FAnimNextParamType::EValueType::SoftObject:
-							case FAnimNextParamType::EValueType::Class:
-							case FAnimNextParamType::EValueType::SoftClass:
-								if (ValueTypeObjectLHS == ValueTypeObjectRHS)
-								{
-									return EParamCompatibility::Compatible_Equal;
-								}
-								else if (CastChecked<UClass>(ValueTypeObjectRHS)->IsChildOf(CastChecked<UClass>(ValueTypeObjectLHS)))
-								{
-									return EParamCompatibility::Compatible_Cast;
-								}
-								break;
-							}
-						}
-					}
-				}
-			}
-			break;
 		}
 		break;
 	}
@@ -403,11 +194,6 @@ FParamCompatibility FParamUtils::GetCompatibility(const FParamTypeHandle& InLHS,
 	return EParamCompatibility::Incompatible;
 }
 
-FParamCompatibility FParamUtils::GetCompatibility(const FAnimNextParamType& InLHS, const FAnimNextParamType& InRHS)
-{
-	return GetCompatibility(InLHS.GetHandle(), InRHS.GetHandle());
-}
-
 static bool CanUseFunctionInternal(const UFunction* InFunction, const UClass* InExpectedClass, FProperty*& OutReturnProperty)
 {
 	UClass* FunctionClass = InFunction->GetOuterUClass();
@@ -492,7 +278,7 @@ bool FParamUtils::CanUseFunction(const UFunction* InFunction, const UClass* InEx
 	return CanUseFunctionInternal(InFunction, InExpectedClass, ReturnProperty);
 }
 
-bool FParamUtils::CanUseFunction(const UFunction* InFunction, const UClass* InExpectedClass, FParamTypeHandle& OutTypeHandle)
+bool FParamUtils::CanUseFunction(const UFunction* InFunction, const UClass* InExpectedClass, FAnimNextParamType& OutType)
 {
 	FProperty* ReturnProperty = nullptr;
 	if(!CanUseFunctionInternal(InFunction, InExpectedClass, ReturnProperty))
@@ -501,8 +287,8 @@ bool FParamUtils::CanUseFunction(const UFunction* InFunction, const UClass* InEx
 	}
 
 	check(ReturnProperty);
-	OutTypeHandle = FParamTypeHandle::FromProperty(ReturnProperty);
-	if(!OutTypeHandle.IsValid())
+	OutType = FAnimNextParamType::FromProperty(ReturnProperty);
+	if(!OutType.IsValid())
 	{
 		return false;
 	}
@@ -519,15 +305,15 @@ bool FParamUtils::CanUseProperty(const FProperty* InProperty)
 	return true;
 }
 
-bool FParamUtils::CanUseProperty(const FProperty* InProperty, FParamTypeHandle& OutTypeHandle)
+bool FParamUtils::CanUseProperty(const FProperty* InProperty, FAnimNextParamType& OutType)
 {
 	if(!CanUseProperty(InProperty))
 	{
 		return false;
 	}
 
-	OutTypeHandle = FParamTypeHandle::FromProperty(InProperty);
-	if(!OutTypeHandle.IsValid())
+	OutType = FAnimNextParamType::FromProperty(InProperty);
+	if(!OutType.IsValid())
 	{
 		return false;
 	}
diff --git a/Engine/Plugins/Experimental/AnimNext/Source/AnimNext/Private/Param/RigVMDispatch_GetLayerParameter.cpp b/Engine/Plugins/Experimental/AnimNext/Source/AnimNext/Private/Param/RigVMDispatch_GetLayerParameter.cpp
--- a/Engine/Plugins/Experimental/AnimNext/Source/AnimNext/Private/Param/RigVMDispatch_GetLayerParameter.cpp
+++ b/Engine/Plugins/Experimental/AnimNext/Source/AnimNext/Private/Param/RigVMDispatch_GetLayerParameter.cpp
@@ -2,7 +2,6 @@
 
 #include "Param/RigVMDispatch_GetLayerParameter.h"
 #include "RigVMCore/RigVMStruct.h"
-#include "Param/ParamStack.h"
 
 const FName FRigVMDispatch_GetLayerParameter::ValueName = TEXT("Value");
 const FName FRigVMDispatch_GetLayerParameter::TypeHandleName = TEXT("Type");
@@ -76,31 +75,6 @@ FRigVMTemplateTypeMap FRigVMDispatch_GetLayerParameter::OnNewArgumentType(const
 
 void FRigVMDispatch_GetLayerParameter::Execute(FRigVMExtendedExecuteContext& InContext, FRigVMMemoryHandleArray Handles, FRigVMPredicateBranchArray RigVMBranches)
 {
-	using namespace UE::AnimNext;
-
-	const FName Parameter = *(FName*)Handles[0].GetData();
-	const FProperty* ValueProperty = Handles[1].GetResolvedProperty();
-	check(ValueProperty);
-	uint8* TargetDataPtr = Handles[1].GetData();
-
-	uint32& ParameterHash = *(uint32*)Handles[2].GetData();
-	if (ParameterHash == 0 && Parameter != NAME_None)
-	{
-		ParameterHash = FParamId::CalculateHash(Parameter, NAME_None);
-	}
-
-	uint32& TypeHandle = *(uint32*)Handles[3].GetData();
-	if (TypeHandle == 0)
-	{
-		TypeHandle = FParamTypeHandle::FromProperty(ValueProperty).ToRaw();
-	}
-
-	FAnimNextExecuteContext& Context = InContext.GetPublicData<FAnimNextExecuteContext>();
-	const FAnimNextParamContextData& ParamContextData = Context.GetContextData<FAnimNextParamContextData>();
-	TConstArrayView<uint8> SourceData;
-	if (ParamContextData.GetLayerHandle().GetValueRaw(FParamId(Parameter, ParameterHash), FParamTypeHandle::FromRaw(TypeHandle), SourceData).IsSuccessful())
-	{
-		ValueProperty->CopyCompleteValue(TargetDataPtr, SourceData.GetData());
-	}
+	// Deprecated stub
 }
 
diff --git a/Engine/Plugins/Experimental/AnimNext/Source/AnimNext/Private/Param/RigVMDispatch_GetParameter.cpp b/Engine/Plugins/Experimental/AnimNext/Source/AnimNext/Private/Param/RigVMDispatch_GetParameter.cpp
--- a/Engine/Plugins/Experimental/AnimNext/Source/AnimNext/Private/Param/RigVMDispatch_GetParameter.cpp
+++ b/Engine/Plugins/Experimental/AnimNext/Source/AnimNext/Private/Param/RigVMDispatch_GetParameter.cpp
@@ -2,8 +2,6 @@
 
 #include "Param/RigVMDispatch_GetParameter.h"
 #include "RigVMCore/RigVMStruct.h"
-#include "Context.h"
-#include "Param/ParamStack.h"
 
 const FName FRigVMDispatch_GetParameter::ValueName = TEXT("Value");
 const FName FRigVMDispatch_GetParameter::TypeHandleName = TEXT("Type");
@@ -77,29 +75,6 @@ FRigVMTemplateTypeMap FRigVMDispatch_GetParameter::OnNewArgumentType(const FName
 
 void FRigVMDispatch_GetParameter::Execute(FRigVMExtendedExecuteContext& InContext, FRigVMMemoryHandleArray Handles, FRigVMPredicateBranchArray RigVMBranches)
 {
-	using namespace UE::AnimNext;
-
-	const FName Parameter = *(FName*)Handles[0].GetData();
-	const FProperty* ValueProperty = Handles[1].GetResolvedProperty();
-	check(ValueProperty);
-	uint8* TargetDataPtr = Handles[1].GetData();
-
-	uint32& ParameterHash = *(uint32*)Handles[2].GetData();
-	if (ParameterHash == 0 && Parameter != NAME_None)
-	{
-		ParameterHash = GetTypeHash(Parameter);
-	}
-
-	uint32& TypeHandle = *(uint32*)Handles[3].GetData();
-	if (TypeHandle == 0)
-	{
-		TypeHandle = FParamTypeHandle::FromProperty(ValueProperty).ToRaw();
-	}
-
-	TConstArrayView<uint8> SourceData;
-	if (FParamStack::Get().GetParamData(FParamId(Parameter, ParameterHash), FParamTypeHandle::FromRaw(TypeHandle), SourceData).IsSuccessful())
-	{
-		ValueProperty->CopyCompleteValue(TargetDataPtr, SourceData.GetData());
-	}
+	// Deprecated stub
 }
 
diff --git a/Engine/Plugins/Experimental/AnimNext/Source/AnimNext/Private/Param/RigVMDispatch_GetScopedParameter.cpp b/Engine/Plugins/Experimental/AnimNext/Source/AnimNext/Private/Param/RigVMDispatch_GetScopedParameter.cpp
--- a/Engine/Plugins/Experimental/AnimNext/Source/AnimNext/Private/Param/RigVMDispatch_GetScopedParameter.cpp
+++ b/Engine/Plugins/Experimental/AnimNext/Source/AnimNext/Private/Param/RigVMDispatch_GetScopedParameter.cpp
@@ -1,9 +1,7 @@
 // Copyright Epic Games, Inc. All Rights Reserved.
 
 #include "Param/RigVMDispatch_GetScopedParameter.h"
-#include "Param/ParamStack.h"
 #include "RigVMCore/RigVMStruct.h"
-#include "Param/AnimNextParam.h"
 
 const FName FRigVMDispatch_GetScopedParameter::ParameterName = TEXT("Parameter");
 const FName FRigVMDispatch_GetScopedParameter::ValueName = TEXT("Value");
@@ -56,7 +54,7 @@ const TArray<FRigVMTemplateArgumentInfo>& FRigVMDispatch_GetScopedParameter::Get
 		};
 
 		const FRigVMRegistry_NoLock& Registry = FRigVMRegistry_NoLock::GetForRead();
-		Infos.Emplace(ParameterName, ERigVMPinDirection::Input, Registry.GetTypeIndex_NoLock<FAnimNextParam>());
+		Infos.Emplace(ParameterName, ERigVMPinDirection::Input, RigVMTypeUtils::TypeIndex::FString);
 		Infos.Emplace(ValueName, ERigVMPinDirection::Output, ValueCategories);
 		Infos.Emplace(ParameterIdName, ERigVMPinDirection::Hidden, RigVMTypeUtils::TypeIndex::UInt32);
 		Infos.Emplace(TypeHandleName, ERigVMPinDirection::Hidden, RigVMTypeUtils::TypeIndex::UInt32);
@@ -69,7 +67,7 @@ FRigVMTemplateTypeMap FRigVMDispatch_GetScopedParameter::OnNewArgumentType(const
 {
 	FRigVMTemplateTypeMap Types;
 	const FRigVMRegistry_NoLock& Registry = FRigVMRegistry_NoLock::GetForRead();
-	Types.Add(ParameterName, Registry.GetTypeIndex_NoLock<FAnimNextParam>());
+	Types.Add(ParameterName, RigVMTypeUtils::TypeIndex::FString);
 	Types.Add(ValueName, InTypeIndex);
 	Types.Add(ParameterIdName, RigVMTypeUtils::TypeIndex::UInt32);
 	Types.Add(TypeHandleName, RigVMTypeUtils::TypeIndex::UInt32);
@@ -78,28 +76,5 @@ FRigVMTemplateTypeMap FRigVMDispatch_GetScopedParameter::OnNewArgumentType(const
 
 void FRigVMDispatch_GetScopedParameter::Execute(FRigVMExtendedExecuteContext& InContext, FRigVMMemoryHandleArray Handles, FRigVMPredicateBranchArray RigVMBranches)
 {
-	using namespace UE::AnimNext;
-
-	const FAnimNextParam& Parameter = *(FAnimNextParam*)Handles[0].GetData();
-	const FProperty* ValueProperty = Handles[1].GetResolvedProperty();
-	check(ValueProperty);
-	uint8* TargetDataPtr = Handles[1].GetData();
-
-	uint32& ParameterHash = *(uint32*)Handles[2].GetData();
-	if (ParameterHash == 0 && Parameter.Name != NAME_None)
-	{
-		ParameterHash = FParamId::CalculateHash(Parameter.Name, Parameter.InstanceId);
-	}
-
-	uint32& TypeHandle = *(uint32*)Handles[3].GetData();
-	if (TypeHandle == 0)
-	{
-		TypeHandle = FParamTypeHandle::FromProperty(ValueProperty).ToRaw();
-	}
-
-	TConstArrayView<uint8> SourceData;
-	if (FParamStack::Get().GetParamData(FParamId(Parameter.Name, Parameter.InstanceId, ParameterHash), FParamTypeHandle::FromRaw(TypeHandle), SourceData).IsSuccessful())
-	{
-		ValueProperty->CopyCompleteValue(TargetDataPtr, SourceData.GetData());
-	}
+	// Deprecated stub
 }
diff --git a/Engine/Plugins/Experimental/AnimNext/Source/AnimNext/Private/Param/RigVMDispatch_SetLayerParameter.cpp b/Engine/Plugins/Experimental/AnimNext/Source/AnimNext/Private/Param/RigVMDispatch_SetLayerParameter.cpp
--- a/Engine/Plugins/Experimental/AnimNext/Source/AnimNext/Private/Param/RigVMDispatch_SetLayerParameter.cpp
+++ b/Engine/Plugins/Experimental/AnimNext/Source/AnimNext/Private/Param/RigVMDispatch_SetLayerParameter.cpp
@@ -3,10 +3,6 @@
 #include "Param/RigVMDispatch_SetLayerParameter.h"
 #include "RigVMCore/RigVMStruct.h"
 #include "RigVMCore/RigVM.h"
-#include "AnimNextExecuteContext.h"
-#include "Param/ParamId.h"
-#include "Param/ParamTypeHandle.h"
-#include "Param/ParamStackLayerHandle.h"
 
 const FName FRigVMDispatch_SetLayerParameter::ExecuteContextName = TEXT("ExecuteContext");
 const FName FRigVMDispatch_SetLayerParameter::ValueName = TEXT("Value");
@@ -87,33 +83,6 @@ FRigVMTemplateTypeMap FRigVMDispatch_SetLayerParameter::OnNewArgumentType(const
 
 void FRigVMDispatch_SetLayerParameter::Execute(FRigVMExtendedExecuteContext& InContext, FRigVMMemoryHandleArray Handles, FRigVMPredicateBranchArray RigVMBranches)
 {
-	using namespace UE::AnimNext;
-
-	const FName Parameter = *(FName*)Handles[0].GetData();
-	const FProperty* ValueProperty = Handles[1].GetResolvedProperty();
-	check(ValueProperty);
-	const uint8* SourceData =  Handles[1].GetData();
-
-	uint32& ParameterHash = *(uint32*)Handles[2].GetData();
-	if (ParameterHash == 0 && Parameter != NAME_None)
-	{
-		ParameterHash = FParamId::CalculateHash(Parameter, NAME_None);
-	}
-
-	uint32& RawTypeHandle = *(uint32*)Handles[3].GetData();
-	if (RawTypeHandle == 0)
-	{
-		RawTypeHandle = FParamTypeHandle::FromProperty(ValueProperty).ToRaw();
-	}
-
-	FParamId ParameterId(Parameter, ParameterHash);
-	FParamTypeHandle TypeHandle = FParamTypeHandle::FromRaw(RawTypeHandle);
-	if(ParameterId.IsValid() && TypeHandle.IsValid())
-	{
-		FAnimNextExecuteContext& AnimNextContext = InContext.GetPublicData<FAnimNextExecuteContext>();
-		const FAnimNextParamContextData& ParamContextData = AnimNextContext.GetContextData<FAnimNextParamContextData>();
-		TConstArrayView<uint8> SourceDataView(SourceData, ValueProperty->GetSize());
-		ParamContextData.GetLayerHandle().SetValueRaw(ParameterId, TypeHandle, SourceDataView);
-	}
+	// Deprecated stub
 }
 
diff --git a/Engine/Plugins/Experimental/AnimNext/Source/AnimNext/Private/TraitCore/ExecutionContext.cpp b/Engine/Plugins/Experimental/AnimNext/Source/AnimNext/Private/TraitCore/ExecutionContext.cpp
--- a/Engine/Plugins/Experimental/AnimNext/Source/AnimNext/Private/TraitCore/ExecutionContext.cpp
+++ b/Engine/Plugins/Experimental/AnimNext/Source/AnimNext/Private/TraitCore/ExecutionContext.cpp
@@ -9,7 +9,7 @@
 #include "TraitCore/NodeInstance.h"
 #include "TraitCore/NodeTemplate.h"
 #include "TraitCore/NodeTemplateRegistry.h"
-#include "Module/AnimNextModule.h"
+#include "Graph/AnimNextAnimationGraph.h"
 #include "Graph/AnimNextGraphInstance.h"
 
 namespace UE::AnimNext
@@ -94,7 +94,7 @@ namespace UE::AnimNext
 		BindTo(InGraphInstance);
 	}
 
-	void FExecutionContext::BindTo(FAnimNextGraphInstancePtr& InGraphInstance)
+	void FExecutionContext::BindTo(const FAnimNextGraphInstancePtr& InGraphInstance)
 	{
 		if (FAnimNextGraphInstance* InstanceImpl = InGraphInstance.GetImpl())
 		{
@@ -153,7 +153,7 @@ namespace UE::AnimNext
 			return FTraitPtr();	// The execution context must be bound to a valid graph instance
 		}
 
-		if (!ensure(GraphInstance.GetModule() != nullptr))
+		if (!ensure(GraphInstance.GetAnimationGraph() != nullptr))
 		{
 			return FTraitPtr();	// We need a valid graph instance to allocate into
 		}
@@ -565,8 +565,8 @@ namespace UE::AnimNext
 	const FNodeDescription& FExecutionContext::GetNodeDescription(const FAnimNextGraphInstance& GraphInstance, FNodeHandle NodeHandle) const
 	{
 		// Grab the node description from the specified graph
-		const UAnimNextModule* Module = GraphInstance.GetModule();
-		return *reinterpret_cast<const FNodeDescription*>(&Module->SharedDataBuffer[NodeHandle.GetSharedOffset()]);
+		const UAnimNextAnimationGraph* AnimationGraph = GraphInstance.GetAnimationGraph();
+		return *reinterpret_cast<const FNodeDescription*>(&AnimationGraph->SharedDataBuffer[NodeHandle.GetSharedOffset()]);
 	}
 
 	const FNodeDescription& FExecutionContext::GetNodeDescription(const FNodeInstance& NodeInstance) const
diff --git a/Engine/Plugins/Experimental/AnimNext/Source/AnimNext/Private/Variables/RigVMDispatch_CallHoistedAccessorFunction.cpp b/Engine/Plugins/Experimental/AnimNext/Source/AnimNext/Private/Variables/RigVMDispatch_CallHoistedAccessorFunction.cpp
new file mode 100644
--- /dev/null
+++ b/Engine/Plugins/Experimental/AnimNext/Source/AnimNext/Private/Variables/RigVMDispatch_CallHoistedAccessorFunction.cpp
@@ -0,0 +1,66 @@
+﻿// Copyright Epic Games, Inc. All Rights Reserved.
+
+#include "Variables/RigVMDispatch_CallHoistedAccessorFunction.h"
+#include "RigVMCore/RigVMStruct.h"
+#include "Variables/AnimNextSoftFunctionPtr.h"
+
+FRigVMDispatch_CallHoistedAccessorFunction::FRigVMDispatch_CallHoistedAccessorFunction()
+{
+	FactoryScriptStruct = StaticStruct();
+}
+
+FName FRigVMDispatch_CallHoistedAccessorFunction::GetArgumentNameForOperandIndex(int32 InOperandIndex, int32 InTotalOperands) const
+{
+	static const FName ArgumentNames[] =
+	{
+		ObjectName,
+		FunctionName,
+		ValueName,
+	};
+	check(InTotalOperands == UE_ARRAY_COUNT(ArgumentNames));
+	return ArgumentNames[InOperandIndex];
+}
+
+const TArray<FRigVMTemplateArgumentInfo>& FRigVMDispatch_CallHoistedAccessorFunction::GetArgumentInfos() const
+{
+	static TArray<FRigVMTemplateArgumentInfo> Infos;
+	if(Infos.IsEmpty())
+	{
+		static const TArray<FRigVMTemplateArgument::ETypeCategory> ValueCategories =
+		{
+			FRigVMTemplateArgument::ETypeCategory_SingleAnyValue,
+			FRigVMTemplateArgument::ETypeCategory_ArrayAnyValue
+		};
+
+		const FRigVMRegistry_NoLock& Registry = FRigVMRegistry_NoLock::GetForRead();
+		Infos.Emplace(ObjectName, ERigVMPinDirection::Input, Registry.GetTypeIndex_NoLock<UObject>());
+		Infos.Emplace(FunctionName, ERigVMPinDirection::Input, Registry.GetTypeIndex_NoLock<FAnimNextSoftFunctionPtr>());
+		Infos.Emplace(ValueName, ERigVMPinDirection::Output, ValueCategories);
+	}
+
+	return Infos;
+}
+
+FRigVMTemplateTypeMap FRigVMDispatch_CallHoistedAccessorFunction::OnNewArgumentType(const FName& InArgumentName, TRigVMTypeIndex InTypeIndex) const
+{
+	const FRigVMRegistry_NoLock& Registry = FRigVMRegistry_NoLock::GetForRead();
+	
+	FRigVMTemplateTypeMap Types;
+	Types.Add(ObjectName, Registry.GetTypeIndex_NoLock<UObject>());
+	Types.Add(FunctionName, Registry.GetTypeIndex_NoLock<FAnimNextSoftFunctionPtr>());
+	Types.Add(ValueName, InTypeIndex);
+	return Types;
+}
+
+void FRigVMDispatch_CallHoistedAccessorFunction::Execute(FRigVMExtendedExecuteContext& InContext, FRigVMMemoryHandleArray Handles, FRigVMPredicateBranchArray RigVMBranches)
+{
+	const UObject* ObjectPtr = *reinterpret_cast<UObject**>(Handles[0].GetData());
+	if(ObjectPtr == nullptr)
+	{
+		// Something failed to resolve upstream, OK to just skip this work
+		return;
+	}
+
+	// TODO
+}
+
diff --git a/Engine/Plugins/Experimental/AnimNext/Source/AnimNext/Private/Variables/RigVMDispatch_CallObjectAccessorFunction.cpp b/Engine/Plugins/Experimental/AnimNext/Source/AnimNext/Private/Variables/RigVMDispatch_CallObjectAccessorFunction.cpp
new file mode 100644
--- /dev/null
+++ b/Engine/Plugins/Experimental/AnimNext/Source/AnimNext/Private/Variables/RigVMDispatch_CallObjectAccessorFunction.cpp
@@ -0,0 +1,66 @@
+﻿// Copyright Epic Games, Inc. All Rights Reserved.
+
+#include "Variables/RigVMDispatch_CallObjectAccessorFunction.h"
+#include "RigVMCore/RigVMStruct.h"
+#include "Variables/AnimNextSoftFunctionPtr.h"
+
+FRigVMDispatch_CallObjectAccessorFunction::FRigVMDispatch_CallObjectAccessorFunction()
+{
+	FactoryScriptStruct = StaticStruct();
+}
+
+FName FRigVMDispatch_CallObjectAccessorFunction::GetArgumentNameForOperandIndex(int32 InOperandIndex, int32 InTotalOperands) const
+{
+	static const FName ArgumentNames[] =
+	{
+		ObjectName,
+		FunctionName,
+		ValueName,
+	};
+	check(InTotalOperands == UE_ARRAY_COUNT(ArgumentNames));
+	return ArgumentNames[InOperandIndex];
+}
+
+const TArray<FRigVMTemplateArgumentInfo>& FRigVMDispatch_CallObjectAccessorFunction::GetArgumentInfos() const
+{
+	static TArray<FRigVMTemplateArgumentInfo> Infos;
+	if(Infos.IsEmpty())
+	{
+		static const TArray<FRigVMTemplateArgument::ETypeCategory> ValueCategories =
+		{
+			FRigVMTemplateArgument::ETypeCategory_SingleAnyValue,
+			FRigVMTemplateArgument::ETypeCategory_ArrayAnyValue
+		};
+
+		const FRigVMRegistry_NoLock& Registry = FRigVMRegistry_NoLock::GetForRead();
+		Infos.Emplace(ObjectName, ERigVMPinDirection::Input, Registry.GetTypeIndex_NoLock<UObject>());
+		Infos.Emplace(FunctionName, ERigVMPinDirection::Input, Registry.GetTypeIndex_NoLock<FAnimNextSoftFunctionPtr>());
+		Infos.Emplace(ValueName, ERigVMPinDirection::Output, ValueCategories);
+	}
+
+	return Infos;
+}
+
+FRigVMTemplateTypeMap FRigVMDispatch_CallObjectAccessorFunction::OnNewArgumentType(const FName& InArgumentName, TRigVMTypeIndex InTypeIndex) const
+{
+	const FRigVMRegistry_NoLock& Registry = FRigVMRegistry_NoLock::GetForRead();
+	
+	FRigVMTemplateTypeMap Types;
+	Types.Add(ObjectName, Registry.GetTypeIndex_NoLock<UObject>());
+	Types.Add(FunctionName, Registry.GetTypeIndex_NoLock<FAnimNextSoftFunctionPtr>());
+	Types.Add(ValueName, InTypeIndex);
+	return Types;
+}
+
+void FRigVMDispatch_CallObjectAccessorFunction::Execute(FRigVMExtendedExecuteContext& InContext, FRigVMMemoryHandleArray Handles, FRigVMPredicateBranchArray RigVMBranches)
+{
+	const UObject* ObjectPtr = *reinterpret_cast<UObject**>(Handles[0].GetData());
+	if(ObjectPtr == nullptr)
+	{
+		// Something failed to resolve upstream, OK to just skip this work
+		return;
+	}
+
+	// TODO
+}
+
diff --git a/Engine/Plugins/Experimental/AnimNext/Source/AnimNext/Private/Variables/RigVMDispatch_GetObjectProperty.cpp b/Engine/Plugins/Experimental/AnimNext/Source/AnimNext/Private/Variables/RigVMDispatch_GetObjectProperty.cpp
new file mode 100644
--- /dev/null
+++ b/Engine/Plugins/Experimental/AnimNext/Source/AnimNext/Private/Variables/RigVMDispatch_GetObjectProperty.cpp
@@ -0,0 +1,87 @@
+﻿// Copyright Epic Games, Inc. All Rights Reserved.
+
+#include "Variables/RigVMDispatch_GetObjectProperty.h"
+#include "RigVMCore/RigVMStruct.h"
+#include "Variables/AnimNextFieldPath.h"
+
+FRigVMDispatch_GetObjectProperty::FRigVMDispatch_GetObjectProperty()
+{
+	FactoryScriptStruct = StaticStruct();
+}
+
+FName FRigVMDispatch_GetObjectProperty::GetArgumentNameForOperandIndex(int32 InOperandIndex, int32 InTotalOperands) const
+{
+	static const FName ArgumentNames[] =
+	{
+		ObjectName,
+		PropertyName,
+		ValueName,
+	};
+	check(InTotalOperands == UE_ARRAY_COUNT(ArgumentNames));
+	return ArgumentNames[InOperandIndex];
+}
+
+const TArray<FRigVMTemplateArgumentInfo>& FRigVMDispatch_GetObjectProperty::GetArgumentInfos() const
+{
+	static TArray<FRigVMTemplateArgumentInfo> Infos;
+	if(Infos.IsEmpty())
+	{
+		static const TArray<FRigVMTemplateArgument::ETypeCategory> ValueCategories =
+		{
+			FRigVMTemplateArgument::ETypeCategory_SingleAnyValue,
+			FRigVMTemplateArgument::ETypeCategory_ArrayAnyValue
+		};
+
+		const FRigVMRegistry_NoLock& Registry = FRigVMRegistry_NoLock::GetForRead();
+		Infos.Emplace(ObjectName, ERigVMPinDirection::Input, Registry.GetTypeIndex_NoLock<UObject>());
+		Infos.Emplace(PropertyName, ERigVMPinDirection::Input, Registry.GetTypeIndex_NoLock<FAnimNextFieldPath>());
+		Infos.Emplace(ValueName, ERigVMPinDirection::Output, ValueCategories);
+	}
+
+	return Infos;
+}
+
+FRigVMTemplateTypeMap FRigVMDispatch_GetObjectProperty::OnNewArgumentType(const FName& InArgumentName, TRigVMTypeIndex InTypeIndex) const
+{
+	const FRigVMRegistry_NoLock& Registry = FRigVMRegistry_NoLock::GetForRead();
+	
+	FRigVMTemplateTypeMap Types;
+	Types.Add(ObjectName, Registry.GetTypeIndex_NoLock<UObject>());
+	Types.Add(PropertyName, Registry.GetTypeIndex_NoLock<FAnimNextFieldPath>());
+	Types.Add(ValueName, InTypeIndex);
+	return Types;
+}
+
+void FRigVMDispatch_GetObjectProperty::Execute(FRigVMExtendedExecuteContext& InContext, FRigVMMemoryHandleArray Handles, FRigVMPredicateBranchArray RigVMBranches)
+{
+	const UObject* ObjectPtr = *reinterpret_cast<UObject**>(Handles[0].GetData());
+	if(ObjectPtr == nullptr)
+	{
+		// Something failed to resolve upstream, OK to just skip this work
+		return;
+	}
+
+	FAnimNextFieldPath* FieldPathPtr = reinterpret_cast<FAnimNextFieldPath*>(Handles[1].GetData());
+	const FProperty* SourceProperty = FieldPathPtr->FieldPath.Get();
+	checkSlow(SourceProperty);
+	const uint8* SourceAddress = SourceProperty->ContainerPtrToValuePtr<uint8>(ObjectPtr);
+	checkSlow(SourceAddress);
+
+	const FProperty* TargetProperty = Handles[2].GetResolvedProperty();
+	checkSlow(TargetProperty);
+	checkSlow(TargetProperty->GetClass() == SourceProperty->GetClass());
+	uint8* TargetAddress = Handles[2].GetData();
+	checkSlow(TargetAddress);
+
+	// TODO: add a specialization for bool to GetDispatchFunctionImpl rather than branching here
+	const FBoolProperty* SourceBoolProperty = CastField<FBoolProperty>(SourceProperty);
+	if(SourceBoolProperty)
+	{
+		static_cast<const FBoolProperty*>(TargetProperty)->SetPropertyValue(TargetAddress, SourceBoolProperty->GetPropertyValue(SourceAddress));
+	}
+	else
+	{
+		SourceProperty->CopyCompleteValue(TargetAddress, SourceAddress);
+	}
+}
+
diff --git a/Engine/Plugins/Experimental/AnimNext/Source/AnimNext/Public/AnimNextStats.h b/Engine/Plugins/Experimental/AnimNext/Source/AnimNext/Public/AnimNextStats.h
--- a/Engine/Plugins/Experimental/AnimNext/Source/AnimNext/Public/AnimNextStats.h
+++ b/Engine/Plugins/Experimental/AnimNext/Source/AnimNext/Public/AnimNextStats.h
@@ -6,14 +6,11 @@
 
 DECLARE_STATS_GROUP(TEXT("AnimNext"), STATGROUP_AnimNext, STATCAT_Advanced);
 
-DECLARE_CYCLE_STAT_EXTERN(TEXT("AnimNext: Initialize Entry"), STAT_AnimNext_InitializeEntry, STATGROUP_AnimNext, ANIMNEXT_API);
+DECLARE_CYCLE_STAT_EXTERN(TEXT("AnimNext: Initialize Instance"), STAT_AnimNext_InitializeInstance, STATGROUP_AnimNext, ANIMNEXT_API);
 DECLARE_CYCLE_STAT_EXTERN(TEXT("AnimNext: Create Instance Data"), STAT_AnimNext_CreateInstanceData, STATGROUP_AnimNext, ANIMNEXT_API);
 
-DECLARE_CYCLE_STAT_EXTERN(TEXT("AnimNext Task: Graph"), STAT_AnimNext_Task_Graph, STATGROUP_AnimNext, ANIMNEXT_API);
-DECLARE_CYCLE_STAT_EXTERN(TEXT("AnimNext Task: External Params"), STAT_AnimNext_Task_ExternalParams, STATGROUP_AnimNext, ANIMNEXT_API);
-DECLARE_CYCLE_STAT_EXTERN(TEXT("AnimNext Task: Port"), STAT_AnimNext_Task_Port, STATGROUP_AnimNext, ANIMNEXT_API);
-DECLARE_CYCLE_STAT_EXTERN(TEXT("AnimNext Task: Scope Entry"), STAT_AnimNext_Task_ScopeEntry, STATGROUP_AnimNext, ANIMNEXT_API);
-DECLARE_CYCLE_STAT_EXTERN(TEXT("AnimNext Task: Scope Exit"), STAT_AnimNext_Task_ScopeExit, STATGROUP_AnimNext, ANIMNEXT_API);
+DECLARE_CYCLE_STAT_EXTERN(TEXT("AnimNext: Run Graph"), STAT_AnimNext_Run_Graph, STATGROUP_AnimNext, ANIMNEXT_API);
+DECLARE_CYCLE_STAT_EXTERN(TEXT("AnimNext: Write Pose"), STAT_AnimNext_Write_Pose, STATGROUP_AnimNext, ANIMNEXT_API);
 
 DECLARE_CYCLE_STAT_EXTERN(TEXT("AnimNext: Allocate Graph Instance"), STAT_AnimNext_Graph_AllocateInstance, STATGROUP_AnimNext, ANIMNEXT_API);
 DECLARE_CYCLE_STAT_EXTERN(TEXT("AnimNext: Graph RigVM"), STAT_AnimNext_Graph_RigVM, STATGROUP_AnimNext, ANIMNEXT_API);
diff --git a/Engine/Plugins/Experimental/AnimNext/Source/AnimNext/Public/DataRegistry.h b/Engine/Plugins/Experimental/AnimNext/Source/AnimNext/Public/DataRegistry.h
--- a/Engine/Plugins/Experimental/AnimNext/Source/AnimNext/Public/DataRegistry.h
+++ b/Engine/Plugins/Experimental/AnimNext/Source/AnimNext/Public/DataRegistry.h
@@ -74,13 +74,13 @@ public:
 	template<typename DataType>
 	FDataHandle PreAllocateMemory(const int32 NumElements)
 	{
-		FParamTypeHandle ParamTypeHandle = FParamTypeHandle::GetHandle<DataType>();
+		FAnimNextParamType ParamType = FAnimNextParamType::GetType<DataType>();
 
 		bool bIsTypeDefValid = false;
 		FDataTypeDef TypeDef;
 		{
 			FRWScopeLock Lock(DataTypeDefsLock, SLT_ReadOnly);
-			if (const FDataTypeDef* TypeDefPtr = DataTypeDefs.Find(ParamTypeHandle))
+			if (const FDataTypeDef* TypeDefPtr = DataTypeDefs.Find(ParamType))
 			{
 				bIsTypeDefValid = true;
 				TypeDef = *TypeDefPtr;
@@ -93,7 +93,7 @@ public:
 			// TODO : Log if we allocate more than DEFAULT_BLOCK_SIZE elements of that type
 		}
 
-		if (ensure(TypeDef.ParamTypeHandle.IsValid()))
+		if (ensure(TypeDef.ParamType.IsValid()))
 		{
 			const int32 ElementSize = TypeDef.ElementSize;
 			const int32 ElementAlign = TypeDef.ElementAlign;
@@ -103,7 +103,7 @@ public:
 
 			uint8* Memory = (uint8*)FMemory::Malloc(BufferSize, TypeDef.ElementAlign);    // TODO : This should come from preallocated chunks, use malloc / free for now
 
-			Private::FAllocatedBlock* AllocatedBlock = new Private::FAllocatedBlock(Memory, NumElements, ParamTypeHandle); // TODO : avoid memory fragmentation
+			Private::FAllocatedBlock* AllocatedBlock = new Private::FAllocatedBlock(Memory, NumElements, ParamType); // TODO : avoid memory fragmentation
 			AllocatedBlock->AddRef();
 
 			FRWScopeLock Lock(AllocatedBlocksLock, SLT_Write);
@@ -140,7 +140,7 @@ private:
 	// structure holding each registered type information
 	struct FDataTypeDef
 	{
-		FParamTypeHandle ParamTypeHandle;
+		FAnimNextParamType ParamType;
 		DestroyFnSignature DestroyTypeFn = nullptr;
 		int32 ElementSize = 0;
 		int32 ElementAlign = 0;
@@ -168,7 +168,7 @@ private:
 	};
 
 	// Map holding registered types
-	TMap<FParamTypeHandle, FDataTypeDef> DataTypeDefs;
+	TMap<FAnimNextParamType, FDataTypeDef> DataTypeDefs;
 	// Lock for registered types map
 	FRWLock DataTypeDefsLock;
 
@@ -188,11 +188,11 @@ private:
 	template<typename DataType>
 	FDataTypeDef RegisterDataType_Impl(int32 AllocationBlockSize)
 	{
-		FParamTypeHandle ParamTypeHandle = FParamTypeHandle::GetHandle<DataType>();
-		check(ParamTypeHandle.IsValid());
+		FAnimNextParamType ParamType = FAnimNextParamType::GetType<DataType>();
+		check(ParamType.IsValid());
 
-		const int32 ElementSize = ParamTypeHandle.GetSize();
-		const int32 ElementAlign = ParamTypeHandle.GetAlignment();
+		const int32 ElementSize = ParamType.GetSize();
+		const int32 ElementAlign = ParamType.GetAlignment();
 
 		// If we use raw types, I need a per element destructor
 		DestroyFnSignature DestroyFn = [](uint8* TargetBuffer, int32 NumElem)->void
@@ -207,8 +207,8 @@ private:
 		{
 			FRWScopeLock WriteLock(DataTypeDefsLock, SLT_Write);
 
-			FDataTypeDef* AddedDef = &DataTypeDefs.FindOrAdd(ParamTypeHandle, { ParamTypeHandle, DestroyFn, ElementSize, ElementAlign, AllocationBlockSize });
-			check(AddedDef->ParamTypeHandle == ParamTypeHandle); // check we have not added two different types with the same ID
+			FDataTypeDef* AddedDef = &DataTypeDefs.FindOrAdd(ParamType, { ParamType, DestroyFn, ElementSize, ElementAlign, AllocationBlockSize });
+			check(AddedDef->ParamType == ParamType); // check we have not added two different types with the same ID
 
 			return *AddedDef;
 		}
diff --git a/Engine/Plugins/Experimental/AnimNext/Source/AnimNext/Public/DataRegistryTypes.h b/Engine/Plugins/Experimental/AnimNext/Source/AnimNext/Public/DataRegistryTypes.h
--- a/Engine/Plugins/Experimental/AnimNext/Source/AnimNext/Public/DataRegistryTypes.h
+++ b/Engine/Plugins/Experimental/AnimNext/Source/AnimNext/Public/DataRegistryTypes.h
@@ -4,7 +4,7 @@
 
 #include "CoreMinimal.h"
 #include "UObject/WeakObjectPtr.h"
-#include "Param/ParamTypeHandle.h"
+#include "Param/ParamType.h"
 #include "Animation/AnimTypes.h"
 #include "ReferenceSkeleton.h"
 #include "BoneIndices.h"
@@ -20,12 +20,12 @@ struct FAllocatedBlock
 	void* Memory = nullptr;
 	int32 NumElem = 0;
 	mutable int32 NumRefs = 0;
-	FParamTypeHandle TypeHandle;
+	FAnimNextParamType Type;
 
-	FAllocatedBlock(void* InMemory, int32 InNumElem, FParamTypeHandle InTypeHandle)
+	FAllocatedBlock(void* InMemory, int32 InNumElem, FAnimNextParamType InType)
 		: Memory(InMemory)
 		, NumElem(InNumElem)
-		, TypeHandle(InTypeHandle)
+		, Type(InType)
 	{
 	}
 
@@ -152,9 +152,9 @@ struct ANIMNEXT_API FDataHandle
 		return *Data;
 	}
 
-	inline FParamTypeHandle GetTypeHandle() const
+	inline FAnimNextParamType GetType() const
 	{
-		return AllocatedBlock != nullptr ? AllocatedBlock->TypeHandle : FParamTypeHandle();
+		return AllocatedBlock != nullptr ? AllocatedBlock->Type : FAnimNextParamType();
 	}
 
 private:
diff --git a/Engine/Plugins/Experimental/AnimNext/Source/AnimNext/Public/Graph/AnimNextGraphInstance.h b/Engine/Plugins/Experimental/AnimNext/Source/AnimNext/Public/Graph/AnimNextGraphInstance.h
--- a/Engine/Plugins/Experimental/AnimNext/Source/AnimNext/Public/Graph/AnimNextGraphInstance.h
+++ b/Engine/Plugins/Experimental/AnimNext/Source/AnimNext/Public/Graph/AnimNextGraphInstance.h
@@ -4,7 +4,6 @@
 
 #include "CoreMinimal.h"
 #include "HAL/CriticalSection.h"
-#include "Param/ParamStack.h"
 #include "TraitCore/TraitEvent.h"
 #include "TraitCore/TraitEventList.h"
 #include "TraitCore/TraitPtr.h"
@@ -16,7 +15,11 @@ class FReferenceCollector;
 
 struct FAnimNextGraphInstancePtr;
 struct FRigUnit_AnimNextGraphEvaluator;
+class UAnimNextAnimationGraph;
 class UAnimNextModule;
+struct FAnimNextModuleInstance;
+class FRigVMTraitScope;
+struct FRigVMExtendedExecuteContext;
 
 namespace UE::AnimNext
 {
@@ -53,8 +56,8 @@ struct ANIMNEXT_API FAnimNextGraphInstance
 	// Returns true if we have a live graph instance, false otherwise
 	bool IsValid() const;
 
-	// Returns the module used by this instance or nullptr if the instance is invalid
-	const UAnimNextModule* GetModule() const;
+	// Returns the animation graph used by this instance or nullptr if the instance is invalid
+	const UAnimNextAnimationGraph* GetAnimationGraph() const;
 
 	// Returns the entry point in Graph that this instance corresponds to 
 	FName GetEntryPoint() const;
@@ -62,14 +65,17 @@ struct ANIMNEXT_API FAnimNextGraphInstance
 	// Returns a weak handle to the root trait instance
 	UE::AnimNext::FWeakTraitPtr GetGraphRootPtr() const;
 
+	// Returns the module instance that owns us or nullptr if we are invalid
+	FAnimNextModuleInstance* GetModuleInstance() const;
+
 	// Returns the parent graph instance that owns us or nullptr for the root graph instance or if we are invalid
 	FAnimNextGraphInstance* GetParentGraphInstance() const;
 
 	// Returns the root graph instance that owns us and the components or nullptr if we are invalid
 	FAnimNextGraphInstance* GetRootGraphInstance() const;
 
-	// Check to see if this instance data matches the provided module
-	bool UsesModule(const UAnimNextModule* InModule) const;
+	// Check to see if this instance data matches the provided animation graph
+	bool UsesAnimationGraph(const UAnimNextAnimationGraph* InAnimationGraph) const;
 
 	// Check to see if this instance data matches the provided graph entry point
 	bool UsesEntryPoint(FName InEntryPoint) const;
@@ -101,6 +107,9 @@ struct ANIMNEXT_API FAnimNextGraphInstance
 	// Called each time the graph updates
 	void Update();
 
+	// Get the extended execute context that we own
+	FRigVMExtendedExecuteContext& GetExtendedExecuteContext();
+
 private:
 	// Returns a pointer to the specified component, or nullptr if not found
 	UE::AnimNext::FGraphInstanceComponent* TryGetComponent(int32 ComponentNameHash, FName ComponentName) const;
@@ -112,6 +121,7 @@ private:
 	// When frozen, latent handles that can freeze are skipped, all others will execute
 	void ExecuteLatentPins(const TConstArrayView<UE::AnimNext::FLatentPropertyHandle>& LatentHandles, void* DestinationBasePtr, bool bIsFrozen);
 
+#if WITH_EDITORONLY_DATA
 	// During graph compilation, if we have existing graph instances, we freeze them by releasing their memory before thawing them
 	// Freezing is a partial release of resources that retains the necessary information to re-create things safely
 	void Freeze();
@@ -119,15 +129,22 @@ private:
 	// During graph compilation, once compilation is done we thaw existing graph instances to reallocate their memory
 	void Thaw();
 
-	// Update and push any graph state that this graph has onto the parameter stack
-	UE::AnimNext::FParamStack::FPushedLayerHandle UpdateAndPushGraphState(float InDeltaTime) const;
+	// Hook into module compilation
+	void OnModuleCompiled(UAnimNextModule* InModule);
+
+	// Rebind any public variables when this instance is recompiled
+	void RebindPublicVariables();
+#endif
+
+	// Bind the variables in the supplied traits in scope to their respective public variables, so they point at host memory
+	void BindPublicVariables(TConstArrayView<FRigVMTraitScope> InTraitScopes);
 
-	// Pop any graph state that this graph has off the parameter stack
-	void PopGraphState(UE::AnimNext::FParamStack::FPushedLayerHandle InHandle) const;
+	// Unbind any public variables that were pointing at host memory and re-point them at the internal defaults
+	void UnbindPublicVariables();
 
 	// Hard reference to the graph used to create this instance to ensure we can release it safely
 	UPROPERTY()
-	TObjectPtr<const UAnimNextModule> Module;
+	TObjectPtr<const UAnimNextAnimationGraph> AnimationGraph;
 
 	// The entry point in Graph that this instance corresponds to 
 	FName EntryPoint;
@@ -135,14 +152,28 @@ private:
 	// Hard reference to the graph instance data, we own it
 	UE::AnimNext::FTraitPtr GraphInstancePtr;
 
+	// The module instance that owns the root, us and the components
+	FAnimNextModuleInstance* ModuleInstance = nullptr;
+
 	// The graph instance that owns us
 	FAnimNextGraphInstance* ParentGraphInstance = nullptr;
 
 	// The root graph instance that owns us and the components
 	FAnimNextGraphInstance* RootGraphInstance = nullptr;
 
-	// Graph state parameter source
-	TUniquePtr<UE::AnimNext::IParameterSource> GraphState;
+	// User variables used to operate the graph
+	FInstancedPropertyBag Variables;
+
+#if WITH_EDITORONLY_DATA
+	struct FCachedVariableBinding
+	{
+		FName VariableName;
+		uint8* Memory = nullptr;
+	};
+
+	// Cached public variable bindings used to correctly thaw instances with input pin bindings
+	TArray<FCachedVariableBinding> CachedVariableBindings;
+#endif
 
 	// Extended execute context instance for this graph instance, we own it
 	UPROPERTY()
@@ -151,10 +182,20 @@ private:
 	// Graph instance components that persist from update to update
 	GraphInstanceComponentMapType Components;
 
+	// The current state of public variable bindings to the host
+	enum class EPublicVariablesState : uint8
+	{
+		None,		// No public variables present
+		Unbound,	// Present, but currently unbound
+		Bound		// Present and bound
+	};
+
+	EPublicVariablesState PublicVariablesState;
+
 	// Whether or not this graph has updated once
-	bool bHasUpdatedOnce = false;
+	bool bHasUpdatedOnce : 1 = false;
 
-	friend UAnimNextModule;					// The module is the one that allocates instances
+	friend UAnimNextAnimationGraph;			// The graph is the one that allocates instances
 	friend FRigUnit_AnimNextGraphEvaluator;	// We evaluate the instance
 	friend UE::AnimNext::FExecutionContext;
 	friend FAnimNextGraphInstancePtr;
diff --git a/Engine/Plugins/Experimental/AnimNext/Source/AnimNext/Public/Graph/AnimNextGraphInstancePtr.h b/Engine/Plugins/Experimental/AnimNext/Source/AnimNext/Public/Graph/AnimNextGraphInstancePtr.h
--- a/Engine/Plugins/Experimental/AnimNext/Source/AnimNext/Public/Graph/AnimNextGraphInstancePtr.h
+++ b/Engine/Plugins/Experimental/AnimNext/Source/AnimNext/Public/Graph/AnimNextGraphInstancePtr.h
@@ -3,7 +3,6 @@
 #pragma once
 
 #include "CoreMinimal.h"
-#include "Param/ParamStack.h"
 #include "TraitCore/TraitEvent.h"
 #include "TraitCore/TraitPtr.h"
 
@@ -13,7 +12,10 @@ class FReferenceCollector;
 
 struct FAnimNextGraphInstance;
 struct FRigUnit_AnimNextGraphEvaluator;
-class UAnimNextModule;
+struct FRigUnit_AnimNextRunAnimationGraph;
+class UAnimNextAnimationGraph;
+class FRigVMTraitScope;
+struct FRigVMExtendedExecuteContext;
 
 namespace UE::AnimNext
 {
@@ -26,7 +28,7 @@ using GraphInstanceComponentMapType = TMap<FName, TSharedPtr<UE::AnimNext::FGrap
 // Represents an instance of an AnimNext graph
 // This struct uses UE reflection because we wish for the GC to keep the graph
 // alive while we own a reference to it. It is not intended to be serialized on disk with a live instance.
-USTRUCT()
+USTRUCT(BlueprintType, DisplayName="Animation Graph Instance")
 struct ANIMNEXT_API FAnimNextGraphInstancePtr
 {
 	GENERATED_BODY()
@@ -34,12 +36,6 @@ struct ANIMNEXT_API FAnimNextGraphInstancePtr
 	// Creates an empty graph instance that doesn't reference anything
 	FAnimNextGraphInstancePtr();
 
-	// No copying, only moving
-	FAnimNextGraphInstancePtr(const FAnimNextGraphInstancePtr&) = delete;
-	FAnimNextGraphInstancePtr& operator=(const FAnimNextGraphInstancePtr&) = delete;
-	FAnimNextGraphInstancePtr(FAnimNextGraphInstancePtr&&);
-	FAnimNextGraphInstancePtr& operator=(FAnimNextGraphInstancePtr&&);
-
 	// If the graph instance is allocated, we release it during destruction
 	~FAnimNextGraphInstancePtr();
 
@@ -49,8 +45,8 @@ struct ANIMNEXT_API FAnimNextGraphInstancePtr
 	// Returns true if we have a live graph instance, false otherwise
 	bool IsValid() const;
 
-	// Returns the module used by this instance or nullptr if the instance is invalid
-	const UAnimNextModule* GetModule() const;
+	// Returns the animation graph that will be used by this instance
+	const UAnimNextAnimationGraph* GetAnimationGraph() const;
 
 	// Returns a weak handle to the root trait instance
 	UE::AnimNext::FWeakTraitPtr GetGraphRootPtr() const;
@@ -58,8 +54,8 @@ struct ANIMNEXT_API FAnimNextGraphInstancePtr
 	// Returns the graph instance implementation
 	FAnimNextGraphInstance* GetImpl() const;
 
-	// Check to see if this instance data matches the provided module
-	bool UsesModule(const UAnimNextModule* InModule) const;
+	// Check to see if this instance data matches the provided animation graph
+	bool UsesAnimationGraph(const UAnimNextAnimationGraph* InAnimationGraph) const;
 
 	// Returns whether or not this graph instance is the root graph instance or false otherwise
 	bool IsRoot() const;
@@ -86,13 +82,13 @@ struct ANIMNEXT_API FAnimNextGraphInstancePtr
 	GraphInstanceComponentMapType::TConstIterator GetComponentIterator() const;
 
 	// Called each time the graph updates
-	void Update();
-
-	// Push any graph state that this graph has onto the parameter stack
-	UE::AnimNext::FParamStack::FPushedLayerHandle UpdateAndPushGraphState(float InDeltaTime) const;
+	void Update() const;
 
-	// Pop any graph state that this graph has off the parameter stack
-	void PopGraphState(UE::AnimNext::FParamStack::FPushedLayerHandle InHandle) const;
+	// Get the extended execute context that we own
+	FRigVMExtendedExecuteContext& GetExtendedExecuteContext() const;
+	
+	// Whether public variables require a binding
+	bool RequiresPublicVariableBinding() const;
 
 private:
 	// Returns a pointer to the specified component, or nullptr if not found
@@ -101,11 +97,15 @@ private:
 	// Adds the specified component and returns a reference to it
 	UE::AnimNext::FGraphInstanceComponent& AddComponent(int32 ComponentNameHash, FName ComponentName, TSharedPtr<UE::AnimNext::FGraphInstanceComponent>&& Component);
 
+	// Bind the variables in the supplied traits in scope to their respective public variables
+	void BindPublicVariables(TConstArrayView<FRigVMTraitScope> InTraitScopes) const;
+
 	// Indirection to hide implementation details and to fix the graph instance into a single memory location
-	TUniquePtr<FAnimNextGraphInstance> Impl;
+	TSharedPtr<FAnimNextGraphInstance> Impl;
 
-	friend UAnimNextModule;					// The graph is the one that allocates instances
-	friend FRigUnit_AnimNextGraphEvaluator;	// We evaluate the instance
+	friend UAnimNextAnimationGraph;			// The graph is the one that allocates instances
+	friend FRigUnit_AnimNextGraphEvaluator;		// We evaluate the instance
+	friend FRigUnit_AnimNextRunAnimationGraph;	// We evaluate the instance
 	friend UE::AnimNext::FExecutionContext;
 };
 
@@ -115,7 +115,6 @@ struct TStructOpsTypeTraits<FAnimNextGraphInstancePtr> : public TStructOpsTypeTr
 	enum
 	{
 		WithAddStructReferencedObjects = true,
-		WithCopy = false,
 	};
 };
 
diff --git a/Engine/Plugins/Experimental/AnimNext/Source/AnimNext/Public/IAnimNextModuleInterface.h b/Engine/Plugins/Experimental/AnimNext/Source/AnimNext/Public/IAnimNextModuleInterface.h
--- a/Engine/Plugins/Experimental/AnimNext/Source/AnimNext/Public/IAnimNextModuleInterface.h
+++ b/Engine/Plugins/Experimental/AnimNext/Source/AnimNext/Public/IAnimNextModuleInterface.h
@@ -8,12 +8,11 @@
 #include "ReferencePose.h"
 
 struct FAnimNextGraphInstancePtr;
-struct FAnimNextParamInstanceIdentifier;
+struct FAnimNextVariableBindingData;
 
 namespace UE::AnimNext
 {
 	class IParameterSource;
-	class IParameterSourceFactory;
 	struct FParameterSourceContext;
 	struct FTraitEventList;
 }
@@ -26,8 +25,8 @@ class IAnimNextAnimGraph
 public:
 	virtual ~IAnimNextAnimGraph() = default;
 
-	virtual void UpdateGraph(FAnimNextGraphInstancePtr& GraphInstance, float DeltaTime, FTraitEventList& InputEventList, FTraitEventList& OutputEventList) const = 0;
-	virtual void EvaluateGraph(FAnimNextGraphInstancePtr& GraphInstance, const FReferencePose& RefPose, int32 GraphLODLevel, FLODPoseHeap& OutputPose) const = 0;
+	virtual void UpdateGraph(const FAnimNextGraphInstancePtr& GraphInstance, float DeltaTime, FTraitEventList& InputEventList, FTraitEventList& OutputEventList) const = 0;
+	virtual void EvaluateGraph(const FAnimNextGraphInstancePtr& GraphInstance, const FReferencePose& RefPose, int32 GraphLODLevel, FLODPoseHeap& OutputPose) const = 0;
 };
 
 class IAnimNextModuleInterface : public IModuleInterface
@@ -43,8 +42,8 @@ public:
 	static ANIMNEXT_API IAnimNextModuleInterface& Get();
 	virtual void RegisterAnimNextAnimGraph(const IAnimNextAnimGraph& InAnimGraphImpl) = 0;
 	virtual void UnregisterAnimNextAnimGraph() = 0;
-	virtual void UpdateGraph(FAnimNextGraphInstancePtr& GraphInstance, float DeltaTime, FTraitEventList& InputEventList, FTraitEventList& OutputEventList) = 0;
-	virtual void EvaluateGraph(FAnimNextGraphInstancePtr& GraphInstance, const FReferencePose& RefPose, int32 GraphLODLevel, FLODPoseHeap& OutputPose) const = 0;
+	virtual void UpdateGraph(const FAnimNextGraphInstancePtr& GraphInstance, float DeltaTime, FTraitEventList& InputEventList, FTraitEventList& OutputEventList) = 0;
+	virtual void EvaluateGraph(const FAnimNextGraphInstancePtr& GraphInstance, const FReferencePose& RefPose, int32 GraphLODLevel, FLODPoseHeap& OutputPose) const = 0;
 };
 
 }
diff --git a/Engine/Plugins/Experimental/AnimNext/Source/AnimNext/Public/Param/ParamType.h b/Engine/Plugins/Experimental/AnimNext/Source/AnimNext/Public/Param/ParamType.h
--- a/Engine/Plugins/Experimental/AnimNext/Source/AnimNext/Public/Param/ParamType.h
+++ b/Engine/Plugins/Experimental/AnimNext/Source/AnimNext/Public/Param/ParamType.h
@@ -14,7 +14,6 @@
 
 namespace UE::AnimNext
 {
-	struct FParamTypeHandle;
 	struct FParamHelpers;
 }
 
@@ -62,7 +61,6 @@ public:
 	using EValueType = ::EPropertyBagPropertyType;
 	using EContainerType = ::EPropertyBagContainerType;
 
-	friend struct UE::AnimNext::FParamTypeHandle;
 	friend struct UE::AnimNext::FParamHelpers;
 	friend struct UE::AnimNext::UncookedOnly::FUtils;
 
@@ -160,7 +158,7 @@ private:
 			if constexpr (std::is_same_v<ParamType, TObjectPtr<UClass>>)
 			{
 				ParameterType.ValueType = EValueType::Class;
-				ParameterType.ValueTypeObject = ParamType::ElementType::StaticClass();
+				ParameterType.ValueTypeObject = UObject::StaticClass();
 			}
 			else
 			{
@@ -222,9 +220,6 @@ public:
 		return ParameterType;
 	}
 
-	/** Get a parameter type handle that represents this type */
-	UE::AnimNext::FParamTypeHandle GetHandle() const;
-
 	/** Get the pointer to the object that defines the Enum, Struct, or Class. */
 	const UObject* GetValueTypeObject() const
 	{
@@ -287,7 +282,10 @@ public:
 	
 	/** Construct a parameter type from the passed in FRigVMTemplateArgumentType. */
 	static FAnimNextParamType FromRigVMTemplateArgument(const FRigVMTemplateArgumentType& RigVMType);
-	
+
+	/** Construct a parameter type from the passed in FProperty. */
+	static FAnimNextParamType FromProperty(const FProperty* InProperty);
+
 	/** Equality operator */
 	friend bool operator==(const FAnimNextParamType& InLHS, const FAnimNextParamType& InRHS)
 	{
diff --git a/Engine/Plugins/Experimental/AnimNext/Source/AnimNext/Public/TraitCore/ExecutionContext.h b/Engine/Plugins/Experimental/AnimNext/Source/AnimNext/Public/TraitCore/ExecutionContext.h
--- a/Engine/Plugins/Experimental/AnimNext/Source/AnimNext/Public/TraitCore/ExecutionContext.h
+++ b/Engine/Plugins/Experimental/AnimNext/Source/AnimNext/Public/TraitCore/ExecutionContext.h
@@ -62,7 +62,7 @@ namespace UE::AnimNext
 		// In order to be used, the execution context must be bound to a valid root graph instance
 
 		// Binds the execution context to the specified graph instance if it differs from the currently bound instance
-		void BindTo(FAnimNextGraphInstancePtr& InGraphInstance);
+		void BindTo(const FAnimNextGraphInstancePtr& InGraphInstance);
 
 		// Binds the execution context to the specified graph instance if it differs from the currently bound instance
 		void BindTo(FAnimNextGraphInstance& InGraphInstance);
diff --git a/Engine/Plugins/Experimental/AnimNext/Source/AnimNextEditor/AnimNextEditor.Build.cs b/Engine/Plugins/Experimental/AnimNext/Source/AnimNextEditor/AnimNextEditor.Build.cs
--- a/Engine/Plugins/Experimental/AnimNext/Source/AnimNextEditor/AnimNextEditor.Build.cs
+++ b/Engine/Plugins/Experimental/AnimNext/Source/AnimNextEditor/AnimNextEditor.Build.cs
@@ -47,7 +47,8 @@ namespace UnrealBuildTool.Rules
 					"UniversalObjectLocator",
 					"UniversalObjectLocatorEditor",
 					"Kismet",
-					"AdvancedWidgets",
+					"AdvancedWidgets", 
+					"SceneOutliner",
 				}
 			);
 		}
