AnimNext: Module refactor

diff --git a/Engine/Plugins/Experimental/AnimNext/Source/AnimNextEditor/Private/AnimNextEditorModule.cpp b/Engine/Plugins/Experimental/AnimNext/Source/AnimNextEditor/Private/AnimNextEditorModule.cpp
--- a/Engine/Plugins/Experimental/AnimNext/Source/AnimNextEditor/Private/AnimNextEditorModule.cpp
+++ b/Engine/Plugins/Experimental/AnimNext/Source/AnimNextEditor/Private/AnimNextEditorModule.cpp
@@ -21,33 +21,34 @@
 #include "Graph/TraitEditorTabSummoner.h"
 #include "Graph/AnimNextCompilerResultsTabSummoner.h"
 #include "HAL/PlatformApplicationMisc.h"
-#include "Param/ParameterCustomization.h"
-#include "Param/ParameterPickerArgs.h"
-#include "Param/ParametersGraphPanelPinFactory.h"
-#include "Param/ParamNamePropertyCustomization.h"
-#include "Param/ParamPropertyCustomization.h"
+#include "Variables/VariableCustomization.h"
+#include "Variables/VariableBindingPropertyCustomization.h"
 #include "Param/ParamTypePropertyCustomization.h"
-#include "Param/SParameterPicker.h"
-#include "Scheduler/AnimNextSchedule.h"
 #include "IWorkspaceEditorModule.h"
 #include "Common/SActionMenu.h"
-#include "AnimNextRigVMAssetEntry.h"
+#include "Entries/AnimNextRigVMAssetEntry.h"
 #include "Editor/RigVMGraphDetailCustomization.h"
 #include "EditorUtils.h"
 #include "FileHelpers.h"
 #include "IUniversalObjectLocatorEditorModule.h"
 #include "Param/AnimNextComponentLocatorEditor.h"
-#include "Param/AnimNextParam.h"
 #include "Param/AnimNextLocatorContext.h"
 #include "Param/ObjectCastLocatorEditor.h"
 #include "Param/ObjectFunctionLocatorEditor.h"
 #include "Param/ObjectPropertyLocatorEditor.h"
-#include "Param/SAddParametersDialog.h"
-#include "Module/AnimNextModuleWorkspaceAssetUserData.h"
+#include "Variables/SAddVariablesDialog.h"
+#include "AnimNextAssetWorkspaceAssetUserData.h"
 #include "Graph/AnimNextGraphItemDetails.h"
 #include "Param/AnimNextActorLocatorEditor.h"
 #include "IWorkspaceEditor.h"
+#include "Entries/AnimNextAnimationGraphEntry.h"
 #include "Framework/Docking/LayoutExtender.h"
+#include "Graph/AnimNextAnimationGraph.h"
+#include "Graph/AnimNextAnimationGraph_EditorData.h"
+#include "Common/AnimNextAssetItemDetails.h"
+#include "Module/RigUnit_AnimNextModuleEvents.h"
+#include "Variables/SVariablesView.h"
+#include "Workspace/AnimNextWorkspaceSchema.h"
 
 #define LOCTEXT_NAMESPACE "AnimNextEditorModule"
 
@@ -69,30 +70,17 @@ void FAnimNextEditorModule::StartupModule()
 	PropertyModule.RegisterCustomPropertyTypeLayout(
 		"AnimNextParamType",
 		FOnGetPropertyTypeCustomizationInstance::CreateLambda([] { return MakeShared<FParamTypePropertyTypeCustomization>(); }));
-
-	PropertyModule.RegisterCustomPropertyTypeLayout(
-		"AnimNextParam",
-		FOnGetPropertyTypeCustomizationInstance::CreateLambda([] { return MakeShared<FParamPropertyCustomization>(); }));
-
-	PropertyModule.RegisterCustomPropertyTypeLayout(
-		"AnimNextEditorParam",
-		FOnGetPropertyTypeCustomizationInstance::CreateLambda([] { return MakeShared<FParamPropertyCustomization>(); }));
 	
-	Identifier = MakeShared<FParamNamePropertyTypeIdentifier>();
 	PropertyModule.RegisterCustomPropertyTypeLayout(
-		FNameProperty::StaticClass()->GetFName(),
-		FOnGetPropertyTypeCustomizationInstance::CreateLambda([] { return MakeShared<FParamNamePropertyTypeCustomization>(); }),
-		Identifier);
+		"AnimNextVariableBinding",
+		FOnGetPropertyTypeCustomizationInstance::CreateLambda([] { return MakeShared<FVariableBindingPropertyCustomization>(); }));
 
-	PropertyModule.RegisterCustomClassLayout("AnimNextModule_Parameter", 
-		FOnGetDetailCustomizationInstance::CreateLambda([] { return MakeShared<FParameterCustomization>(); }));
+	PropertyModule.RegisterCustomClassLayout("AnimNextVariableEntry", 
+		FOnGetDetailCustomizationInstance::CreateLambda([] { return MakeShared<FVariableCustomization>(); }));
 
 	AnimNextGraphPanelNodeFactory = MakeShared<FAnimNextGraphPanelNodeFactory>();
 	FEdGraphUtilities::RegisterVisualNodeFactory(AnimNextGraphPanelNodeFactory);
 
-	ParametersGraphPanelPinFactory = MakeShared<FParametersGraphPanelPinFactory>();
-	FEdGraphUtilities::RegisterVisualPinFactory(ParametersGraphPanelPinFactory);
-
 	Workspace::IWorkspaceEditorModule& WorkspaceEditorModule = FModuleManager::Get().LoadModuleChecked<Workspace::IWorkspaceEditorModule>("WorkspaceEditor");
 
 	WorkspaceEditorModule.OnRegisterTabsForEditor().AddLambda([](FWorkflowAllowedTabSet& TabFactories, const TSharedRef<FTabManager>& InTabManager, TSharedPtr<UE::Workspace::IWorkspaceEditor> InEditorPtr)
@@ -104,6 +92,10 @@ void FAnimNextEditorModule::StartupModule()
 			TSharedRef<FAnimNextCompilerResultsTabSummoner> CompilerResultsTabSummoner = MakeShared<FAnimNextCompilerResultsTabSummoner>(InEditorPtr);
 			TabFactories.RegisterFactory(CompilerResultsTabSummoner);
 			CompilerResultsTabSummoner->RegisterTabSpawner(InTabManager, nullptr);
+
+			TSharedRef<FAnimNextVariablesTabSummoner> VariablesTabSummoner = MakeShared<FAnimNextVariablesTabSummoner>(InEditorPtr);
+			TabFactories.RegisterFactory(VariablesTabSummoner);
+			VariablesTabSummoner->RegisterTabSpawner(InTabManager, nullptr);
 		});
 
 	WorkspaceEditorModule.OnExtendTabs().AddLambda([](FLayoutExtender& InLayoutExtender, TSharedPtr<UE::Workspace::IWorkspaceEditor> InEditorPtr)
@@ -113,6 +105,9 @@ void FAnimNextEditorModule::StartupModule()
 
 		FTabManager::FTab CompilerResultsTab(FTabId(CompilerResultsTabName), ETabState::ClosedTab);
 		InLayoutExtender.ExtendLayout(FTabId(Workspace::WorkspaceTabs::BottomMiddleDocumentArea), ELayoutExtensionPosition::After, CompilerResultsTab);
+
+		FTabManager::FTab VariablesTab(FTabId(VariablesTabName), ETabState::OpenedTab);
+		InLayoutExtender.ExtendLayout(FTabId(Workspace::WorkspaceTabs::BottomLeftDocumentArea), ELayoutExtensionPosition::After, VariablesTab);
 	});
 
 	RegisterWorkspaceDocumentTypes(WorkspaceEditorModule);
@@ -142,28 +137,28 @@ void FAnimNextEditorModule::StartupModule()
 			}
 		});
 
-	SRigVMAssetView::RegisterCategoryFactory("Parameters", [](UAnimNextRigVMAssetEditorData* InEditorData)
+	SRigVMAssetView::RegisterCategoryFactory("Variables", [](UAnimNextRigVMAssetEditorData* InEditorData)
 	{
-		UAnimNextModule_EditorData* EditorData = CastChecked<UAnimNextModule_EditorData>(InEditorData);
+		UAnimNextRigVMAssetEditorData* EditorData = CastChecked<UAnimNextRigVMAssetEditorData>(InEditorData);
 		UAnimNextRigVMAsset* Asset = UncookedOnly::FUtils::GetAsset(InEditorData);
 		return SNew(SSimpleButton)
-			.Text(LOCTEXT("AddParameterButton", "Add Parameter"))
+			.Text(LOCTEXT("AddVariableButton", "Add Variable"))
 			.Icon(FAppStyle::Get().GetBrush("Icons.Plus"))
 			.OnClicked_Lambda([EditorData, Asset]()
 			{
-				TSharedRef<SAddParametersDialog> AddParametersDialog =
-					SNew(SAddParametersDialog, FAssetData(EditorData));
+				TSharedRef<SAddVariablesDialog> AddVariablesDialog =
+					SNew(SAddVariablesDialog, FAssetData(EditorData));
 
-				TArray<FParameterToAdd> ParametersToAdd;
-				if(AddParametersDialog->ShowModal(ParametersToAdd))
+				TArray<FVariableToAdd> VariablesToAdd;
+				if(AddVariablesDialog->ShowModal(VariablesToAdd))
 				{
-					if(ParametersToAdd.Num() > 0)
+					if(VariablesToAdd.Num() > 0)
 					{
-						FScopedTransaction Transaction(LOCTEXT("AddParameter", "Add parameter"));
-						for (const FParameterToAdd& ParameterToAdd : ParametersToAdd)
+						FScopedTransaction Transaction(LOCTEXT("AddVariable", "Add variable"));
+						for (const FVariableToAdd& VariableToAdd : VariablesToAdd)
 						{
-							check(EditorData->FindEntry(ParameterToAdd.Name) == nullptr);
-							EditorData->AddParameter(ParameterToAdd.Name, ParameterToAdd.Type);
+							check(EditorData->FindEntry(VariableToAdd.Name) == nullptr);
+							EditorData->AddVariable(VariableToAdd.Name, VariableToAdd.Type);
 						}
 					}
 				}
@@ -173,7 +168,7 @@ void FAnimNextEditorModule::StartupModule()
 
 	SRigVMAssetView::RegisterCategoryFactory("Event Graphs", [](UAnimNextRigVMAssetEditorData* InEditorData)
 	{
-		UAnimNextModule_EditorData* EditorData = CastChecked<UAnimNextModule_EditorData>(InEditorData);
+		UAnimNextRigVMAssetEditorData* EditorData = CastChecked<UAnimNextRigVMAssetEditorData>(InEditorData);
 		return SNew(SSimpleButton)
 			.Text(LOCTEXT("AddEventGraphButton", "Add Event Graph"))
 			.Icon(FAppStyle::Get().GetBrush("Icons.Plus"))
@@ -182,7 +177,7 @@ void FAnimNextEditorModule::StartupModule()
 				FScopedTransaction Transaction(LOCTEXT("AddEventGraph", "Add Event Graph"));
 
 				// Create a new entry for the graph
-				EditorData->AddEventGraph(TEXT("NewGraph"));
+				EditorData->AddEventGraph(TEXT("NewGraph"), FRigUnit_AnimNextPrePhysicsEvent::StaticStruct());
 
 				return FReply::Handled();
 			});
@@ -190,7 +185,7 @@ void FAnimNextEditorModule::StartupModule()
 
 	SRigVMAssetView::RegisterCategoryFactory("Animation Graphs", [](UAnimNextRigVMAssetEditorData* InEditorData)
 	{
-		UAnimNextModule_EditorData* EditorData = CastChecked<UAnimNextModule_EditorData>(InEditorData);
+		UAnimNextRigVMAssetEditorData* EditorData = CastChecked<UAnimNextRigVMAssetEditorData>(InEditorData);
 		return SNew(SSimpleButton)
 			.Text(LOCTEXT("AddGraphButton", "Add Animation Graph"))
 			.Icon(FAppStyle::Get().GetBrush("Icons.Plus"))
@@ -225,16 +220,21 @@ void FAnimNextEditorModule::StartupModule()
 	RegisterLocatorFragmentEditorType("AnimNextActor");
 
 	Workspace::IWorkspaceEditorModule& WorkspaceModule = FModuleManager::Get().LoadModuleChecked<Workspace::IWorkspaceEditorModule>("WorkspaceEditor");
-	const TSharedPtr<FAnimNextGraphItemDetails> GraphItemDetails = MakeShareable<FAnimNextGraphItemDetails>(new FAnimNextGraphItemDetails());
-	WorkspaceModule.RegisterWorkspaceItemDetails(Workspace::FOutlinerItemDetailsId(FAnimNextGraphOutlinerData::StaticStruct()->GetFName()), StaticCastSharedPtr<UE::Workspace::IWorkspaceOutlinerItemDetails>(GraphItemDetails));
-	WorkspaceModule.RegisterWorkspaceItemDetails(Workspace::FOutlinerItemDetailsId(FAnimNextCollapseGraphOutlinerData::StaticStruct()->GetFName()), StaticCastSharedPtr<UE::Workspace::IWorkspaceOutlinerItemDetails>(GraphItemDetails));
-	WorkspaceModule.RegisterWorkspaceItemDetails(Workspace::FOutlinerItemDetailsId(FAnimNextGraphFunctionOutlinerData::StaticStruct()->GetFName()), StaticCastSharedPtr<UE::Workspace::IWorkspaceOutlinerItemDetails>(GraphItemDetails));
+	const TSharedPtr<FAnimNextGraphItemDetails> GraphItemDetails = MakeShared<FAnimNextGraphItemDetails>();
+	WorkspaceModule.RegisterWorkspaceItemDetails(Workspace::FOutlinerItemDetailsId(FAnimNextGraphOutlinerData::StaticStruct()->GetFName()), GraphItemDetails);
+	WorkspaceModule.RegisterWorkspaceItemDetails(Workspace::FOutlinerItemDetailsId(FAnimNextCollapseGraphOutlinerData::StaticStruct()->GetFName()), GraphItemDetails);
+	WorkspaceModule.RegisterWorkspaceItemDetails(Workspace::FOutlinerItemDetailsId(FAnimNextGraphFunctionOutlinerData::StaticStruct()->GetFName()), GraphItemDetails);
+	FAnimNextGraphItemDetails::RegisterToolMenuExtensions();
+
+	const TSharedPtr<FAnimNextAssetItemDetails> AssetItemDetails = MakeShared<FAnimNextAssetItemDetails>();
+	WorkspaceModule.RegisterWorkspaceItemDetails(Workspace::FOutlinerItemDetailsId(FAnimNextModuleOutlinerData::StaticStruct()->GetFName()), AssetItemDetails);
+	WorkspaceModule.RegisterWorkspaceItemDetails(Workspace::FOutlinerItemDetailsId(FAnimNextAnimationGraphOutlinerData::StaticStruct()->GetFName()), AssetItemDetails);
 
 	FAnimNextGraphItemDetails::RegisterToolMenuExtensions();
 
 	SupportedAssetClasses.Append(
 		{
-			UAnimNextSchedule::StaticClass()->GetClassPathName(),
+			UAnimNextAnimationGraph::StaticClass()->GetClassPathName(),
 			UAnimNextModule::StaticClass()->GetClassPathName()
 		});
 }
@@ -245,16 +245,12 @@ void FAnimNextEditorModule::ShutdownModule()
 	{
 		FPropertyEditorModule& PropertyModule = FModuleManager::LoadModuleChecked<FPropertyEditorModule>("PropertyEditor");
 		PropertyModule.UnregisterCustomPropertyTypeLayout("AnimNextParamType");
-		PropertyModule.UnregisterCustomPropertyTypeLayout("AnimNextParam");
-		PropertyModule.UnregisterCustomPropertyTypeLayout("AnimNextEditorParam");
-		PropertyModule.UnregisterCustomPropertyTypeLayout("NameProperty");
-		PropertyModule.UnregisterCustomClassLayout("AnimNextModule_Parameter");
+		PropertyModule.UnregisterCustomPropertyTypeLayout("AnimNextVariableBinding");
+		PropertyModule.UnregisterCustomClassLayout("AnimNextVariableEntry");
 	}
 
 	FEdGraphUtilities::UnregisterVisualNodeFactory(AnimNextGraphPanelNodeFactory);
 
-	FEdGraphUtilities::UnregisterVisualPinFactory(ParametersGraphPanelPinFactory);
-
 	UnregisterWorkspaceDocumentTypes();
 
 	SRigVMAssetView::UnregisterCategoryFactory("Parameters");
@@ -280,6 +276,9 @@ void FAnimNextEditorModule::ShutdownModule()
 		WorkspaceModule.UnregisterWorkspaceItemDetails(Workspace::FOutlinerItemDetailsId(FAnimNextCollapseGraphOutlinerData::StaticStruct()->GetFName()));
 		WorkspaceModule.UnregisterWorkspaceItemDetails(Workspace::FOutlinerItemDetailsId(FAnimNextGraphFunctionOutlinerData::StaticStruct()->GetFName()));
 		FAnimNextGraphItemDetails::UnregisterToolMenuExtensions();
+		WorkspaceModule.UnregisterWorkspaceItemDetails(Workspace::FOutlinerItemDetailsId(FAnimNextModuleOutlinerData::StaticStruct()->GetFName()));
+		WorkspaceModule.UnregisterWorkspaceItemDetails(Workspace::FOutlinerItemDetailsId(FAnimNextAnimationGraphOutlinerData::StaticStruct()->GetFName()));
+		FAnimNextAssetItemDetails::UnregisterToolMenuExtensions();
 	}
 
 	UnregisterLocatorFragmentEditorType("Actor");
@@ -293,12 +292,6 @@ void FAnimNextEditorModule::ShutdownModule()
 	UnregisterLocatorFragmentEditorType("AnimNextActor");
 }
 
-TSharedRef<SWidget> FAnimNextEditorModule::CreateParameterPicker(const FParameterPickerArgs& InArgs)
-{
-	return SNew(SParameterPicker)
-		.Args(InArgs);
-}
-
 void FAnimNextEditorModule::RegisterLocatorFragmentEditorType(FName InLocatorFragmentEditorName)
 {
 	LocatorFragmentEditorNames.Add(InLocatorFragmentEditorName);
@@ -319,88 +312,17 @@ void FAnimNextEditorModule::AddWorkspaceSupportedAssetClass(const FTopLevelAsset
 
 void FAnimNextEditorModule::RegisterWorkspaceDocumentTypes(Workspace::IWorkspaceEditorModule& WorkspaceEditorModule)
 {
-	// --- AnimNextSchedule ---
-	Workspace::FObjectDocumentArgs ScheduleDocumentArgs(
-		Workspace::FOnMakeDocumentWidget::CreateLambda([](const Workspace::FWorkspaceEditorContext& InContext)
-		{
-			UAnimNextSchedule* Schedule = CastChecked<UAnimNextSchedule>(InContext.Object);
-			FPropertyEditorModule& PropertyEditorModule = FModuleManager::GetModuleChecked<FPropertyEditorModule>( "PropertyEditor" );
-			FDetailsViewArgs DetailsViewArgs;
-			DetailsViewArgs.NameAreaSettings = FDetailsViewArgs::HideNameArea;
-			TSharedRef<IDetailsView> DetailsView = PropertyEditorModule.CreateDetailView(DetailsViewArgs);
-			DetailsView->SetObject(Schedule);
-
-			TWeakPtr<Workspace::IWorkspaceEditor> WeakWorkspaceEditor = InContext.WorkspaceEditor;
-			Schedule->CompiledEvent.RemoveAll(&InContext.WorkspaceEditor.Get());
-			Schedule->CompiledEvent.AddSPLambda(&InContext.WorkspaceEditor.Get(), [WeakWorkspaceEditor]()
-				{
-					if (TSharedPtr<Workspace::IWorkspaceEditor> WorkspaceEditor = WeakWorkspaceEditor.Pin())
-					{
-						int32 NumEntries = FMessageLog("AnimNextCompilerResults").NumMessages(EMessageSeverity::Warning);
-						if(NumEntries > 0)
-						{
-							WorkspaceEditor->GetTabManager()->TryInvokeTab(CompilerResultsTabName);
-						}
-					}
-				});
-			
-			return DetailsView;
-		}
-	), Workspace::WorkspaceTabs::TopMiddleDocumentArea);
-	ScheduleDocumentArgs.OnGetTabName = Workspace::FOnGetTabName::CreateLambda([](const Workspace::FWorkspaceEditorContext& InContext)
-	{
-		const UAnimNextSchedule* Schedule = CastChecked<UAnimNextSchedule>(InContext.Object);
-		return FText::FromName(Schedule->GetFName());
-	});
-
-	ScheduleDocumentArgs.OnGetDocumentBreadcrumbTrail = Workspace::FOnGetDocumentBreadcrumbTrail::CreateLambda([](const Workspace::FWorkspaceEditorContext& InContext, TArray<TSharedPtr<Workspace::FWorkspaceBreadcrumb>>& OutBreadcrumbs)
-	{
-		if (const UAnimNextSchedule* Schedule = Cast<UAnimNextSchedule>(InContext.Object))
-		{
-			const TSharedPtr<Workspace::FWorkspaceBreadcrumb>& GraphCrumb = OutBreadcrumbs.Add_GetRef(MakeShared<Workspace::FWorkspaceBreadcrumb>());
-			GraphCrumb->OnGetLabel = Workspace::FWorkspaceBreadcrumb::FOnGetBreadcrumbLabel::CreateLambda([ScheduleName = Schedule->GetFName()]{ return FText::FromName(ScheduleName); });
-
-			TWeakObjectPtr<const UAnimNextSchedule> WeakSchedule = Schedule;
-			GraphCrumb->CanSave = Workspace::FWorkspaceBreadcrumb::FCanSaveBreadcrumb::CreateLambda(
-				[WeakSchedule]
-				{
-					if (const UAnimNextSchedule* Schedule = WeakSchedule.Get())
-					{
-						return Schedule->GetPackage()->IsDirty();
-					}
-
-					return false;
-				}
-			);
-
-			GraphCrumb->OnSave = Workspace::FWorkspaceBreadcrumb::FOnSaveBreadcrumb::CreateLambda(
-				[WeakSchedule]
-				{
-					if (const UAnimNextSchedule* Schedule = WeakSchedule.Get())
-					{
-						FEditorFileUtils::PromptForCheckoutAndSave({Schedule->GetPackage()}, false, /*bPromptToSave=*/ false);
-					}
-				}
-			);
-		}
-	});
-
-	ScheduleDocumentArgs.OnGetTabIcon = Workspace::FOnGetTabIcon::CreateLambda([](const Workspace::FWorkspaceEditorContext& InContext)
-	{
-		return FAppStyle::GetBrush(TEXT("ClassIcon.Default"));
-	});
-
-	WorkspaceEditorModule.RegisterObjectDocumentType(FTopLevelAssetPath(TEXT("/Script/AnimNext.AnimNextSchedule")), ScheduleDocumentArgs);
-
-	// --- AnimNextGraph ---
-	Workspace::FObjectDocumentArgs AnimNextGraphDocumentArgs(
+	// --- AnimNextRigVMAsset ---
+	Workspace::FObjectDocumentArgs AnimNextAssetDocumentArgs(
 			Workspace::FOnMakeDocumentWidget::CreateLambda([](const Workspace::FWorkspaceEditorContext& InContext)
 			{
-				UAnimNextModule* Module = CastChecked<UAnimNextModule>(InContext.Object);
-				UAnimNextModule_EditorData* EditorData = UncookedOnly::FUtils::GetEditorData(Module);
+				UAnimNextRigVMAsset* Asset = CastChecked<UAnimNextRigVMAsset>(InContext.Object);
+				UAnimNextRigVMAssetEditorData* EditorData = UncookedOnly::FUtils::GetEditorData<UAnimNextRigVMAssetEditorData>(Asset);
 
 				TWeakPtr<Workspace::IWorkspaceEditor> WeakWorkspaceEditor = InContext.WorkspaceEditor;
 
+				// TODO: For now (so we can convert older assets over to new formats manually) we open the asset view, but in the future we
+				// should just open any nested graphs (and set the variables view up)
 				TSharedPtr<SRigVMAssetView> SharedAssetView = SNew(SRigVMAssetView, EditorData)
 					.OnOpenGraph_Lambda([WeakWorkspaceEditor](URigVMGraph* InGraph)
 					{
@@ -438,7 +360,7 @@ void FAnimNextEditorModule::RegisterWorkspaceDocumentTypes(Workspace::IWorkspace
 							}
 						}
 					});
-				
+
 				TWeakPtr<SRigVMAssetView> WeakAssetView = SharedAssetView;
 				SharedAssetView->SetOnSelectionChanged(SRigVMAssetView::FOnSelectionChanged::CreateLambda([WeakWorkspaceEditor, WeakAssetView](const TArray<UObject*>& InEntries)
 				{
@@ -455,48 +377,62 @@ void FAnimNextEditorModule::RegisterWorkspaceDocumentTypes(Workspace::IWorkspace
 				return SharedAssetView.ToSharedRef();
 			}),
 			Workspace::WorkspaceTabs::TopMiddleDocumentArea);
-	AnimNextGraphDocumentArgs.OnGetTabName = Workspace::FOnGetTabName::CreateLambda([](const Workspace::FWorkspaceEditorContext& InContext)
+	AnimNextAssetDocumentArgs.OnGetTabName = Workspace::FOnGetTabName::CreateLambda([](const Workspace::FWorkspaceEditorContext& InContext)
 	{
-		const UAnimNextModule* Graph = CastChecked<UAnimNextModule>(InContext.Object);
-		return FText::FromName(Graph->GetFName());
+		const UAnimNextRigVMAsset* Asset = CastChecked<UAnimNextRigVMAsset>(InContext.Object);
+		return FText::FromName(Asset->GetFName());
 	});
 
-	AnimNextGraphDocumentArgs.OnGetDocumentBreadcrumbTrail = Workspace::FOnGetDocumentBreadcrumbTrail::CreateLambda([](const Workspace::FWorkspaceEditorContext& InContext, TArray<TSharedPtr<Workspace::FWorkspaceBreadcrumb>>& OutBreadcrumbs)
+	AnimNextAssetDocumentArgs.OnGetDocumentBreadcrumbTrail = Workspace::FOnGetDocumentBreadcrumbTrail::CreateLambda([](const Workspace::FWorkspaceEditorContext& InContext, TArray<TSharedPtr<Workspace::FWorkspaceBreadcrumb>>& OutBreadcrumbs)
 	{
-		if (const UAnimNextModule* Module = Cast<UAnimNextModule>(InContext.Object))
+		if (const UAnimNextRigVMAsset* Asset = Cast<UAnimNextRigVMAsset>(InContext.Object))
 		{
-			TWeakObjectPtr<const UAnimNextModule> WeakModule = Module;
+			TWeakObjectPtr<const UAnimNextRigVMAsset> WeakAsset = Asset;
 			const TSharedPtr<Workspace::FWorkspaceBreadcrumb>& GraphCrumb = OutBreadcrumbs.Add_GetRef(MakeShared<Workspace::FWorkspaceBreadcrumb>());
-			GraphCrumb->OnGetLabel = Workspace::FWorkspaceBreadcrumb::FOnGetBreadcrumbLabel::CreateLambda([ModuleName = Module->GetFName()]{ return FText::FromName(ModuleName); });
+			GraphCrumb->OnGetLabel = Workspace::FWorkspaceBreadcrumb::FOnGetBreadcrumbLabel::CreateLambda([AssetName = Asset->GetFName()]{ return FText::FromName(AssetName); });
 			GraphCrumb->CanSave = Workspace::FWorkspaceBreadcrumb::FCanSaveBreadcrumb::CreateLambda(
-				[WeakModule]
+				[WeakAsset]
 				{
-					if (const UAnimNextModule* Module = WeakModule.Get())
+					if (const UAnimNextRigVMAsset* Asset = WeakAsset.Get())
 					{
-						return Module->GetPackage()->IsDirty();						
+						return Asset->GetPackage()->IsDirty();						
 					}
 
 					return false;
 				}
 			);
 			GraphCrumb->OnSave = Workspace::FWorkspaceBreadcrumb::FOnSaveBreadcrumb::CreateLambda(
-				[WeakModule]
+				[WeakAsset]
 				{
-					if (const UAnimNextModule* Module = WeakModule.Get())
+					if (const UAnimNextRigVMAsset* Asset = WeakAsset.Get())
 					{
-						FEditorFileUtils::PromptForCheckoutAndSave({Module->GetPackage()}, false, /*bPromptToSave=*/ false);
+						FEditorFileUtils::PromptForCheckoutAndSave({Asset->GetPackage()}, false, /*bPromptToSave=*/ false);
 					}
 				}
 			);
 		}
 	});
 
-	AnimNextGraphDocumentArgs.OnGetTabIcon = Workspace::FOnGetTabIcon::CreateLambda([](const Workspace::FWorkspaceEditorContext& InContext)
+	AnimNextAssetDocumentArgs.OnGetTabIcon = Workspace::FOnGetTabIcon::CreateLambda([](const Workspace::FWorkspaceEditorContext& InContext)
 	{
 		return FAppStyle::GetBrush(TEXT("ClassIcon.Default"));
 	});
 
-	WorkspaceEditorModule.RegisterObjectDocumentType(FTopLevelAssetPath(TEXT("/Script/AnimNext.AnimNextModule")), AnimNextGraphDocumentArgs);
+	WorkspaceEditorModule.RegisterObjectDocumentType(FTopLevelAssetPath(TEXT("/Script/AnimNext.AnimNextModule")), AnimNextAssetDocumentArgs);
+
+	// --- AnimNextAnimationGraph ---
+	Workspace::FObjectDocumentArgs AnimNextAnimationGraphDocumentArgs(
+			Workspace::FOnRedirectWorkspaceContext::CreateLambda([](UObject* InObject)
+			{
+				UAnimNextAnimationGraph* AnimationGraph = CastChecked<UAnimNextAnimationGraph>(InObject);
+				UAnimNextAnimationGraph_EditorData* EditorData = UncookedOnly::FUtils::GetEditorData<UAnimNextAnimationGraph_EditorData>(AnimationGraph);
+
+				// Redirect to the inner graph
+				UAnimNextAnimationGraphEntry* AnimationGraphEntry = CastChecked<UAnimNextAnimationGraphEntry>(EditorData->FindEntry(FRigUnit_AnimNextGraphRoot::DefaultEntryPoint));
+				return AnimationGraphEntry->GetEdGraph();
+			}));
+
+	WorkspaceEditorModule.RegisterObjectDocumentType(FTopLevelAssetPath(TEXT("/Script/AnimNext.AnimNextAnimationGraph")), AnimNextAnimationGraphDocumentArgs);
 
 	// --- AnimNextEdGraph ---
 	Workspace::FGraphDocumentWidgetArgs GraphArgs;
@@ -849,10 +785,10 @@ void FAnimNextEditorModule::RegisterWorkspaceDocumentTypes(Workspace::IWorkspace
 		if (UAnimNextEdGraph* EdGraph = Cast<UAnimNextEdGraph>(InContext.Object))
 		{
 			UAnimNextRigVMAssetEditorData* EditorData = EdGraph->GetTypedOuter<UAnimNextRigVMAssetEditorData>();
-			check(EditorData);
-
-			EditorData->InteractionBracketFinished.RemoveAll(&InContext.WorkspaceEditor.Get());
-			EditorData->InteractionBracketFinished.AddSPLambda(&InContext.WorkspaceEditor.Get(), [WeakWorkspaceEditor](UAnimNextRigVMAssetEditorData* InEditorData)
+			if(EditorData)
+			{
+				EditorData->InteractionBracketFinished.RemoveAll(&InContext.WorkspaceEditor.Get());
+				EditorData->InteractionBracketFinished.AddSPLambda(&InContext.WorkspaceEditor.Get(), [WeakWorkspaceEditor](UAnimNextRigVMAssetEditorData* InEditorData)
 				{
 					if (TSharedPtr<Workspace::IWorkspaceEditor> WorkspaceEditor = WeakWorkspaceEditor.Pin())
 					{
@@ -860,18 +796,19 @@ void FAnimNextEditorModule::RegisterWorkspaceDocumentTypes(Workspace::IWorkspace
 					}
 				});
 
-			EditorData->RigVMCompiledEvent.RemoveAll(&InContext.WorkspaceEditor.Get());
-			EditorData->RigVMCompiledEvent.AddSPLambda(&InContext.WorkspaceEditor.Get(), [WeakWorkspaceEditor](UObject*, URigVM*, FRigVMExtendedExecuteContext&)
+				EditorData->RigVMCompiledEvent.RemoveAll(&InContext.WorkspaceEditor.Get());
+				EditorData->RigVMCompiledEvent.AddSPLambda(&InContext.WorkspaceEditor.Get(), [WeakWorkspaceEditor](UObject*, URigVM*, FRigVMExtendedExecuteContext&)
 				{
 					if (TSharedPtr<Workspace::IWorkspaceEditor> WorkspaceEditor = WeakWorkspaceEditor.Pin())
 					{
 						int32 NumEntries = FMessageLog("AnimNextCompilerResults").NumMessages(EMessageSeverity::Warning);
 						if(NumEntries > 0)
 						{
-							WorkspaceEditor->GetTabManager()->TryInvokeTab(CompilerResultsTabName);
+							WorkspaceEditor->GetTabManager()->TryInvokeTab(FTabId(CompilerResultsTabName));
 						}
 					}
 				});
+			}
 		}
 
 		if (WorkspaceMakeDocumentWidgetDelegate.IsBound())
@@ -886,21 +823,21 @@ void FAnimNextEditorModule::RegisterWorkspaceDocumentTypes(Workspace::IWorkspace
 	{
 		if (URigVMEdGraph* EdGraph = Cast<URigVMEdGraph>(InContext.Object))
 		{
-			if (const UAnimNextModule_EditorData* GraphEditorData = EdGraph->GetTypedOuter<UAnimNextModule_EditorData>())
+			if (UAnimNextRigVMAssetEditorData* EditorData = EdGraph->GetTypedOuter<UAnimNextRigVMAssetEditorData>())
 			{
 				// Iterate model tree, so we display all graph parents until we reach the Entry
 				URigVMGraph* ModelGraph = EdGraph->GetModel();
 				while (ModelGraph != nullptr)
 				{
-					URigVMEdGraph* RigVMEdGraph = Cast<URigVMEdGraph>(GraphEditorData->GetEditorObjectForRigVMGraph(ModelGraph));
+					URigVMEdGraph* RigVMEdGraph = Cast<URigVMEdGraph>(EditorData->GetEditorObjectForRigVMGraph(ModelGraph));
 
-					if (RigVMEdGraph != nullptr && GraphEditorData->GetLocalFunctionLibrary() != RigVMEdGraph->GetModel())
+					if (RigVMEdGraph != nullptr && EditorData->GetLocalFunctionLibrary() != RigVMEdGraph->GetModel())
 					{
 						const TSharedPtr<Workspace::FWorkspaceBreadcrumb>& GraphCrumb = OutBreadcrumbs.Add_GetRef(MakeShared<Workspace::FWorkspaceBreadcrumb>());
 
 						TWeakObjectPtr<URigVMEdGraph> WeakEdGraph = RigVMEdGraph;
 						TWeakPtr<Workspace::IWorkspaceEditor> WeakWorkspaceEditor = InContext.WorkspaceEditor;
-						TWeakObjectPtr<const UAnimNextModule_EditorData> WeakGraphEditorData = GraphEditorData;
+						TWeakObjectPtr<const UAnimNextRigVMAssetEditorData> WeakEditorData = EditorData;
 
 						FText GraphName;
 						if (WeakEdGraph.IsValid())
@@ -919,10 +856,14 @@ void FAnimNextEditorModule::RegisterWorkspaceDocumentTypes(Workspace::IWorkspace
 
 							if (GraphName.IsEmpty())
 							{
-								if (WeakGraphEditorData.IsValid() && WeakGraphEditorData->GetLocalFunctionLibrary() == WeakEdGraph->GetModel())
+								if (WeakEditorData.IsValid() && WeakEditorData->GetLocalFunctionLibrary() == WeakEdGraph->GetModel())
 								{
 									GraphName = UE::AnimNext::UncookedOnly::FUtils::GetFunctionLibraryDisplayName();
 								}
+								else if(UAnimNextRigVMAssetEntry* Entry = WeakEdGraph->GetTypedOuter<UAnimNextRigVMAssetEntry>())
+								{
+									GraphName = Entry->GetDisplayName();
+								}
 								else
 								{
 									GraphName = FText::FromName(WeakEdGraph->GetFName());
@@ -946,7 +887,7 @@ void FAnimNextEditorModule::RegisterWorkspaceDocumentTypes(Workspace::IWorkspace
 							}
 						);
 						GraphCrumb->OnClicked = Workspace::FWorkspaceBreadcrumb::FOnBreadcrumbClicked::CreateLambda(
-							[WeakGraphEditorData, WeakEdGraph, WeakWorkspaceEditor]
+							[WeakEditorData, WeakEdGraph, WeakWorkspaceEditor]
 							{
 								if (const TSharedPtr<Workspace::IWorkspaceEditor> SharedWorkspaceEditor = WeakWorkspaceEditor.Pin())
 								{
@@ -968,39 +909,39 @@ void FAnimNextEditorModule::RegisterWorkspaceDocumentTypes(Workspace::IWorkspace
 					ModelGraph = ModelGraph->GetTypedOuter<URigVMGraph>();
 				}
 
-				// Display the Module
-				if(UAnimNextModule* OuterGraph = UncookedOnly::FUtils::GetGraph(GraphEditorData))
+				// Display the Asset
+				if(UAnimNextRigVMAsset* OuterAsset = UncookedOnly::FUtils::GetAsset<UAnimNextRigVMAsset>(EditorData))
 				{
 					const TSharedPtr<Workspace::FWorkspaceBreadcrumb>& OuterGraphCrumb = OutBreadcrumbs.Add_GetRef(MakeShared<Workspace::FWorkspaceBreadcrumb>());
-					TWeakObjectPtr<UAnimNextModule> WeakOuterGraph = OuterGraph;
+					TWeakObjectPtr<UAnimNextRigVMAsset> WeakOuterAsset = OuterAsset;
 					TWeakPtr<Workspace::IWorkspaceEditor> WeakWorkspaceEditor = InContext.WorkspaceEditor;
-					OuterGraphCrumb->OnGetLabel = Workspace::FWorkspaceBreadcrumb::FOnGetBreadcrumbLabel::CreateLambda([GraphName = OuterGraph->GetFName()]{ return FText::FromName(GraphName); });
+					OuterGraphCrumb->OnGetLabel = Workspace::FWorkspaceBreadcrumb::FOnGetBreadcrumbLabel::CreateLambda([AssetName = OuterAsset->GetFName()]{ return FText::FromName(AssetName); });
 					OuterGraphCrumb->OnClicked = Workspace::FWorkspaceBreadcrumb::FOnBreadcrumbClicked::CreateLambda(
-						[WeakOuterGraph, WeakWorkspaceEditor]
+						[WeakOuterAsset, WeakWorkspaceEditor]
 						{
 							if (const TSharedPtr<Workspace::IWorkspaceEditor> SharedWorkspaceEditor = WeakWorkspaceEditor.Pin())
 							{
-								SharedWorkspaceEditor->OpenObjects({WeakOuterGraph.Get()});
+								SharedWorkspaceEditor->OpenObjects({WeakOuterAsset.Get()});
 							}
 						}
 					);
 					OuterGraphCrumb->CanSave = Workspace::FWorkspaceBreadcrumb::FCanSaveBreadcrumb::CreateLambda(
-						[WeakOuterGraph]
+						[WeakOuterAsset]
 						{
-							if (UAnimNextModule* Module = WeakOuterGraph.Get())
+							if (UAnimNextRigVMAsset* Asset = WeakOuterAsset.Get())
 							{
-								return Module->GetPackage()->IsDirty();
+								return Asset->GetPackage()->IsDirty();
 							}
 
 							return false;
 						}
 					);
 					OuterGraphCrumb->OnSave = Workspace::FWorkspaceBreadcrumb::FOnSaveBreadcrumb::CreateLambda(
-						[WeakOuterGraph]
+						[WeakOuterAsset]
 							{
-								if (UAnimNextModule* Module = WeakOuterGraph.Get())
+								if (UAnimNextRigVMAsset* Asset = WeakOuterAsset.Get())
 								{
-									FEditorFileUtils::PromptForCheckoutAndSave({Module->GetPackage()}, false, /*bPromptToSave=*/ false);
+									FEditorFileUtils::PromptForCheckoutAndSave({Asset->GetPackage()}, false, /*bPromptToSave=*/ false);
 								}
 							}
 						);
@@ -1017,8 +958,8 @@ void FAnimNextEditorModule::UnregisterWorkspaceDocumentTypes()
 	if(FModuleManager::Get().IsModuleLoaded("WorkspaceEditor"))
 	{
 		Workspace::IWorkspaceEditorModule& WorkspaceEditorModule = FModuleManager::LoadModuleChecked<Workspace::IWorkspaceEditorModule>("WorkspaceEditor");
-		WorkspaceEditorModule.UnregisterObjectDocumentType(FTopLevelAssetPath(TEXT("/Script/AnimNext.AnimNextSchedule")));
 		WorkspaceEditorModule.UnregisterObjectDocumentType(FTopLevelAssetPath(TEXT("/Script/AnimNext.AnimNextModule")));
+		WorkspaceEditorModule.UnregisterObjectDocumentType(FTopLevelAssetPath(TEXT("/Script/AnimNext.AnimNextAnimationGraph")));
 		WorkspaceEditorModule.UnregisterObjectDocumentType(FTopLevelAssetPath(TEXT("/Script/AnimNextUncookedOnly.AnimNextEdGraph")));
 	}
 }
diff --git a/Engine/Plugins/Experimental/AnimNext/Source/AnimNextEditor/Private/AnimNextEditorModule.h b/Engine/Plugins/Experimental/AnimNext/Source/AnimNextEditor/Private/AnimNextEditorModule.h
--- a/Engine/Plugins/Experimental/AnimNext/Source/AnimNextEditor/Private/AnimNextEditorModule.h
+++ b/Engine/Plugins/Experimental/AnimNext/Source/AnimNextEditor/Private/AnimNextEditorModule.h
@@ -31,7 +31,6 @@ private:
 	virtual void ShutdownModule() override;
 
 	// IAnimNextEditorModule interface
-	virtual TSharedRef<SWidget> CreateParameterPicker(const FParameterPickerArgs& InArgs) override;
 	virtual void RegisterLocatorFragmentEditorType(FName InLocatorFragmentEditorName) override;
 	virtual void UnregisterLocatorFragmentEditorType(FName InLocatorFragmentEditorName) override;
 	virtual void AddWorkspaceSupportedAssetClass(const FTopLevelAssetPath& InClassAssetPath) override;
@@ -43,9 +42,6 @@ private:
 private:
 	/** Node factory for the AnimNext graph */
 	TSharedPtr<FAnimNextGraphPanelNodeFactory> AnimNextGraphPanelNodeFactory;
-	
-	/** Pin factory for parameters */
-	TSharedPtr<FParametersGraphPanelPinFactory> ParametersGraphPanelPinFactory;
 
 	/** Type identifier for parameter names */
 	TSharedPtr<FParamNamePropertyTypeIdentifier> Identifier;
diff --git a/Engine/Plugins/Experimental/AnimNext/Source/AnimNextEditor/Private/Common/AnimNextAssetItemDetails.cpp b/Engine/Plugins/Experimental/AnimNext/Source/AnimNextEditor/Private/Common/AnimNextAssetItemDetails.cpp
new file mode 100644
--- /dev/null
+++ b/Engine/Plugins/Experimental/AnimNext/Source/AnimNextEditor/Private/Common/AnimNextAssetItemDetails.cpp
@@ -0,0 +1,195 @@
+// Copyright Epic Games, Inc. All Rights Reserved.
+
+#include "AnimNextAssetItemDetails.h"
+
+#include "ClassIconFinder.h"
+#include "Toolkits/AssetEditorToolkitMenuContext.h"
+#include "Entries/AnimNextAnimationGraphEntry.h"
+#include "Entries/AnimNextEventGraphEntry.h"
+#include "InstancedStruct.h"
+#include "AnimNextAssetWorkspaceAssetUserData.h"
+#include "RigVMModel/RigVMGraph.h"
+#include "WorkspaceItemMenuContext.h"
+#include "IWorkspaceEditor.h"
+#include "ScopedTransaction.h"
+#include "RigVMModel/RigVMClient.h"
+#include "ToolMenus.h"
+#include "Framework/Commands/GenericCommands.h"
+#include "Module/AnimNextModule_EditorData.h"
+#include "Module/RigUnit_AnimNextModuleEvents.h"
+#include "UObject/UObjectIterator.h"
+#include "Variables/SAddVariablesDialog.h"
+
+#define LOCTEXT_NAMESPACE "FAnimNextGraphItemDetails"
+
+namespace UE::AnimNext::Editor
+{
+
+const FSlateBrush* FAnimNextAssetItemDetails::GetItemIcon(const FWorkspaceOutlinerItemExport& Export) const
+{
+	return FAppStyle::GetBrush(TEXT("LevelEditor.Tabs.Outliner"));
+}
+
+void FAnimNextAssetItemDetails::RegisterToolMenuExtensions()
+{
+	FToolMenuOwnerScoped OwnerScoped(TEXT("FAnimNextModuleItemDetails"));
+	UToolMenu* Menu = UToolMenus::Get()->ExtendMenu("WorkspaceOutliner.ItemContextMenu");
+	if (Menu == nullptr)
+	{
+		return;
+	}
+
+	Menu->AddDynamicSection(TEXT("AnimNextModuleItem"), FNewToolMenuDelegate::CreateLambda([](UToolMenu* InMenu)
+	{
+		UWorkspaceItemMenuContext* WorkspaceItemContext = InMenu->FindContext<UWorkspaceItemMenuContext>();
+		const UAssetEditorToolkitMenuContext* AssetEditorContext = InMenu->FindContext<UAssetEditorToolkitMenuContext>();
+		if (WorkspaceItemContext == nullptr || AssetEditorContext == nullptr)
+		{
+			return;
+		}
+
+		const TSharedPtr<UE::Workspace::IWorkspaceEditor> WorkspaceEditor = StaticCastSharedPtr<UE::Workspace::IWorkspaceEditor>(AssetEditorContext->Toolkit.Pin());
+		if(!WorkspaceEditor.IsValid())
+		{
+			return;
+		}
+
+		if (WorkspaceItemContext->SelectedExports.Num() != 1)
+		{
+			return;
+		}
+
+		TInstancedStruct<FWorkspaceOutlinerItemData>& Data = WorkspaceItemContext->SelectedExports[0].GetData();
+		if (!Data.IsValid() || !Data.GetScriptStruct()->IsChildOf(FAnimNextRigVMAssetOutlinerData::StaticStruct()))
+		{
+			return;
+		}
+
+		const FAnimNextRigVMAssetOutlinerData& OutlinerData = Data.Get<FAnimNextRigVMAssetOutlinerData>();
+		UAnimNextRigVMAsset* Asset = OutlinerData.Asset;
+		if(Asset == nullptr)
+		{
+			return;
+		}
+
+		UAnimNextRigVMAssetEditorData* EditorData = UncookedOnly::FUtils::GetEditorData<UAnimNextRigVMAssetEditorData>(Asset);
+		if(EditorData == nullptr)
+		{
+			return;
+		}
+
+		FToolMenuSection& AssetSection = InMenu->AddSection("AnimNextAsset", LOCTEXT("AnimNextAssetSectionLabel", "AnimNext Asset"));
+
+		// Per sub-item type addition
+		TConstArrayView<TSubclassOf<UAnimNextRigVMAssetEntry>> AllEntryClasses = EditorData->GetEntryClasses();
+		for(UClass* SubEntryClass : AllEntryClasses)
+		{
+			static FTextFormat AddEntryLabelFormat(LOCTEXT("AddEntryLabelFormat", "Add {0}"));
+			static FTextFormat AddEntryTooltipFormat(LOCTEXT("AddEntryTooltipFormat", "Adds a new {0} to this asset"));
+
+			if(!EditorData->CanAddNewEntry(SubEntryClass))
+			{
+				continue;
+			}
+
+			if(SubEntryClass == UAnimNextEventGraphEntry::StaticClass())
+			{
+				AssetSection.AddSubMenu(
+					SubEntryClass->GetFName(),
+					FText::Format(AddEntryLabelFormat, SubEntryClass->GetDisplayNameText()),
+					FText::Format(AddEntryTooltipFormat, SubEntryClass->GetDisplayNameText()),
+					FNewToolMenuDelegate::CreateLambda([SubEntryClass, WorkspaceItemContext, EditorData, Asset](UToolMenu* InToolMenu)
+					{
+						FToolMenuSection& Section = InToolMenu->AddSection(SubEntryClass->GetFName());
+
+						static FTextFormat AddEventTooltipFormat(LOCTEXT("AddEventGraphTooltipFormat", "Adds a {0} event graph to this asset"));
+						
+						for(TObjectIterator<UScriptStruct> It; It; ++It)
+						{
+							UScriptStruct* Struct = *It;
+							if(!Struct->IsChildOf(FRigUnit_AnimNextModuleEventBase::StaticStruct()) || Struct == FRigUnit_AnimNextModuleEventBase::StaticStruct())
+							{
+								continue;
+							}
+							if(Struct->HasMetaData(FRigVMStruct::HiddenMetaName) || Struct->HasMetaData(FRigVMStruct::AbstractMetaName))
+							{
+								continue;
+							}
+							
+							TInstancedStruct<FRigUnit_AnimNextModuleEventBase> StructInstance;
+							StructInstance.InitializeAsScriptStruct(Struct);
+							FName EventName = StructInstance.Get<FRigUnit_AnimNextModuleEventBase>().GetEventName();
+
+							Section.AddMenuEntry(
+								StructInstance.Get<FRigUnit_AnimNextModuleEventBase>().GetEventName(),
+								FText::FromName(EventName),
+								FText::Format(AddEventTooltipFormat, FText::FromName(EventName)),
+								FSlateIconFinder::FindIconForClass(SubEntryClass, "ClassIcon.Object"),
+								FUIAction(
+									FExecuteAction::CreateWeakLambda(WorkspaceItemContext, [EventName, SubEntryClass, EditorData, Struct]()
+									{
+										FScopedTransaction Transaction(LOCTEXT("AddEventGraph", "Add Event Graph"));
+										EditorData->AddEventGraph(EventName, Struct);
+									}),
+									FCanExecuteAction::CreateWeakLambda(WorkspaceItemContext, [Asset, EventName]()
+									{
+										return !Asset->GetVM()->ContainsEntry(EventName);
+									})
+								));
+						}
+					}),
+					false,
+					FSlateIconFinder::FindIconForClass(SubEntryClass, "ClassIcon.Object"));
+			}
+			else
+			{
+				AssetSection.AddMenuEntry(
+					SubEntryClass->GetFName(),
+					FText::Format(AddEntryLabelFormat, SubEntryClass->GetDisplayNameText()),
+					FText::Format(AddEntryTooltipFormat, SubEntryClass->GetDisplayNameText()),
+					FSlateIconFinder::FindIconForClass(SubEntryClass, "ClassIcon.Object"),
+					FUIAction(
+						FExecuteAction::CreateWeakLambda(WorkspaceItemContext, [&OutlinerData, SubEntryClass, EditorData]()
+						{
+							if(SubEntryClass == UAnimNextVariableEntry::StaticClass())
+							{
+								TSharedRef<SAddVariablesDialog> AddVariablesDialog =
+									SNew(SAddVariablesDialog, FAssetData(OutlinerData.Asset));
+
+								TArray<FVariableToAdd> VariablesToAdd;
+								if(AddVariablesDialog->ShowModal(VariablesToAdd))
+								{
+									if(VariablesToAdd.Num() > 0)
+									{
+										FScopedTransaction Transaction(LOCTEXT("AddVariable", "Add Variable"));
+										for (const FVariableToAdd& VariableToAdd : VariablesToAdd)
+										{
+											check(EditorData->FindEntry(VariableToAdd.Name) == nullptr);
+											EditorData->AddVariable(VariableToAdd.Name, VariableToAdd.Type);
+										}
+									}
+								}
+							}
+							else if(SubEntryClass == UAnimNextAnimationGraphEntry::StaticClass())
+							{
+								FScopedTransaction Transaction(LOCTEXT("AddAnimationGraph", "Add Animation Graph"));
+								EditorData->AddAnimationGraph(TEXT("Root"));
+							}
+						})
+					));
+			}
+		}
+	}));
+}
+
+void FAnimNextAssetItemDetails::UnregisterToolMenuExtensions()
+{
+	if(UToolMenus* ToolMenus = UToolMenus::Get())
+	{
+		ToolMenus->UnregisterOwnerByName("FAnimNextAssetItemDetails");
+	}
+}
+
+}
+
+#undef LOCTEXT_NAMESPACE // "FAnimNextGraphItemDetails"
\ No newline at end of file
diff --git a/Engine/Plugins/Experimental/AnimNext/Source/AnimNextEditor/Private/Common/AnimNextAssetItemDetails.h b/Engine/Plugins/Experimental/AnimNext/Source/AnimNextEditor/Private/Common/AnimNextAssetItemDetails.h
new file mode 100644
--- /dev/null
+++ b/Engine/Plugins/Experimental/AnimNext/Source/AnimNextEditor/Private/Common/AnimNextAssetItemDetails.h
@@ -0,0 +1,24 @@
+// Copyright Epic Games, Inc. All Rights Reserved.
+
+#pragma once
+
+#include "IWorkspaceOutlinerItemDetails.h"
+
+class UAnimNextRigVMAssetEditorData;
+enum class EAnimNextEditorDataNotifType : uint8;
+
+namespace UE::AnimNext::Editor
+{
+
+class FAnimNextAssetItemDetails : public UE::Workspace::IWorkspaceOutlinerItemDetails
+{
+public:
+	FAnimNextAssetItemDetails() = default;
+
+	virtual const FSlateBrush* GetItemIcon(const FWorkspaceOutlinerItemExport& Export) const override;
+
+	static void RegisterToolMenuExtensions();
+	static void UnregisterToolMenuExtensions();
+};
+
+}
diff --git a/Engine/Plugins/Experimental/AnimNext/Source/AnimNextEditor/Private/Common/AnimNextEntryAssetDefinitions.cpp b/Engine/Plugins/Experimental/AnimNext/Source/AnimNextEditor/Private/Common/AnimNextEntryAssetDefinitions.cpp
new file mode 100644
--- /dev/null
+++ b/Engine/Plugins/Experimental/AnimNext/Source/AnimNextEditor/Private/Common/AnimNextEntryAssetDefinitions.cpp
@@ -0,0 +1,25 @@
+﻿// Copyright Epic Games, Inc. All Rights Reserved.
+
+#include "AnimNextEntryAssetDefinitions.h"
+
+#define LOCTEXT_NAMESPACE "AnimNextAssetDefinitions"
+
+FText UAssetDefinition_AnimNextVariableEntry::GetObjectDisplayNameText(UObject* Object) const
+{
+	UAnimNextRigVMAssetEntry* Parameter = CastChecked<UAnimNextRigVMAssetEntry>(Object);
+	return Parameter->GetDisplayName();
+}
+
+FText UAssetDefinition_AnimNextAnimationGraphEntry::GetObjectDisplayNameText(UObject* Object) const
+{
+	UAnimNextRigVMAssetEntry* Parameter = CastChecked<UAnimNextRigVMAssetEntry>(Object);
+	return Parameter->GetDisplayName();
+}
+
+FText UAssetDefinition_AnimNextEventGraphEntry::GetObjectDisplayNameText(UObject* Object) const
+{
+	UAnimNextRigVMAssetEntry* Parameter = CastChecked<UAnimNextRigVMAssetEntry>(Object);
+	return Parameter->GetDisplayName();
+}
+
+#undef LOCTEXT_NAMESPACE
\ No newline at end of file
diff --git a/Engine/Plugins/Experimental/AnimNext/Source/AnimNextEditor/Private/Common/AnimNextEntryAssetDefinitions.h b/Engine/Plugins/Experimental/AnimNext/Source/AnimNextEditor/Private/Common/AnimNextEntryAssetDefinitions.h
new file mode 100644
--- /dev/null
+++ b/Engine/Plugins/Experimental/AnimNext/Source/AnimNextEditor/Private/Common/AnimNextEntryAssetDefinitions.h
@@ -0,0 +1,52 @@
+﻿// Copyright Epic Games, Inc. All Rights Reserved.
+
+#pragma once
+
+#include "AssetDefinitionDefault.h"
+#include "Entries/AnimNextEventGraphEntry.h"
+#include "Entries/AnimNextVariableEntry.h"
+#include "Entries/AnimNextAnimationGraphEntry.h"
+#include "AnimNextEntryAssetDefinitions.generated.h"
+
+#define LOCTEXT_NAMESPACE "AnimNextAssetDefinitions"
+
+UCLASS()
+class UAssetDefinition_AnimNextVariableEntry : public UAssetDefinitionDefault
+{
+	GENERATED_BODY()
+
+public:
+	// UAssetDefinition interface
+	virtual FText GetAssetDisplayName() const override { return LOCTEXT("AnimNextVariable", "Variable"); }
+	virtual FLinearColor GetAssetColor() const override { return FLinearColor(FColor(128,64,64)); }
+	virtual TSoftClassPtr<UObject> GetAssetClass() const override { return UAnimNextVariableEntry::StaticClass(); }
+	virtual FText GetObjectDisplayNameText(UObject* Object) const override;
+};
+
+UCLASS()
+class UAssetDefinition_AnimNextAnimationGraphEntry : public UAssetDefinitionDefault
+{
+	GENERATED_BODY()
+
+public:
+	// UAssetDefinition interface
+	virtual FText GetAssetDisplayName() const override { return LOCTEXT("AnimNextAnimationGraph", "Animation Graph"); }
+	virtual FLinearColor GetAssetColor() const override { return FLinearColor(FColor(128,64,64)); }
+	virtual TSoftClassPtr<UObject> GetAssetClass() const override { return UAnimNextAnimationGraphEntry::StaticClass(); }
+	virtual FText GetObjectDisplayNameText(UObject* Object) const override;
+};
+
+UCLASS()
+class UAssetDefinition_AnimNextEventGraphEntry : public UAssetDefinitionDefault
+{
+	GENERATED_BODY()
+
+public:
+	// UAssetDefinition interface
+	virtual FText GetAssetDisplayName() const override { return LOCTEXT("AnimNextEventGraph", "Event Graph"); }
+	virtual FLinearColor GetAssetColor() const override { return FLinearColor(FColor(128,64,64)); }
+	virtual TSoftClassPtr<UObject> GetAssetClass() const override { return UAnimNextEventGraphEntry::StaticClass(); }
+	virtual FText GetObjectDisplayNameText(UObject* Object) const override;
+};
+
+#undef LOCTEXT_NAMESPACE
\ No newline at end of file
diff --git a/Engine/Plugins/Experimental/AnimNext/Source/AnimNextEditor/Private/Common/GraphEditorSchemaActions.cpp b/Engine/Plugins/Experimental/AnimNext/Source/AnimNextEditor/Private/Common/GraphEditorSchemaActions.cpp
--- a/Engine/Plugins/Experimental/AnimNext/Source/AnimNextEditor/Private/Common/GraphEditorSchemaActions.cpp
+++ b/Engine/Plugins/Experimental/AnimNext/Source/AnimNextEditor/Private/Common/GraphEditorSchemaActions.cpp
@@ -4,10 +4,16 @@
 #include "EditorUtils.h"
 #include "AnimNextEdGraph.h"
 #include "AnimNextEdGraphNode.h"
+#include "EdGraphNode_Comment.h"
+#include "Editor.h"
 #include "Module/AnimNextModule_EditorData.h"
 #include "RigVMModel/Nodes/RigVMUnitNode.h"
 #include "Settings/ControlRigSettings.h"
 #include "Units/RigUnit.h"
+#include "EdGraphSchema_K2.h"
+#include "GraphEditor.h"
+
+#define LOCTEXT_NAMESPACE "AnimNextSchemaActions"
 
 UEdGraphNode* FAnimNextSchemaAction_RigUnit::PerformAction(UEdGraph* ParentGraph, TArray<UEdGraphPin*>& FromPins, const FVector2D Location, bool bSelectNewNode)
 {
@@ -74,52 +80,174 @@ UEdGraphNode* FAnimNextSchemaAction_RigUnit::PerformAction(UEdGraph* ParentGraph
 UEdGraphNode* FAnimNextSchemaAction_DispatchFactory::PerformAction(UEdGraph* ParentGraph, TArray<UEdGraphPin*>& FromPins, const FVector2D Location, bool bSelectNewNode)
 {
 	IRigVMClientHost* Host = ParentGraph->GetImplementingOuter<IRigVMClientHost>();
-	URigVMEdGraphNode* NewNode = nullptr;
+	if(Host == nullptr)
+	{
+		return nullptr;
+	}
+
 	URigVMEdGraph* EdGraph = Cast<URigVMEdGraph>(ParentGraph);
+	if(EdGraph == nullptr)
+	{
+		return nullptr;
+	}
 
+	const FRigVMTemplate* Template = FRigVMRegistry::Get().FindTemplate(Notation);
+	if (Template == nullptr)
+	{
+		return nullptr;
+	}
+	
+	URigVMEdGraphNode* NewNode = nullptr;
 	UEdGraphPin* FromPin = nullptr;
 	if (FromPins.Num() > 0)
 	{
 		FromPin = FromPins[0];
 	}
-	
-	if (Host != nullptr && EdGraph != nullptr)
+
+	const int32 NotationHash = (int32)GetTypeHash(Notation);
+	const FString TemplateName = TEXT("RigVMTemplate_") + FString::FromInt(NotationHash);
+
+	FName Name = UE::AnimNext::Editor::FUtils::ValidateName(Cast<UObject>(Host), Template->GetName().ToString());
+	URigVMController* Controller = Host->GetRigVMClient()->GetController(ParentGraph);
+
+	Controller->OpenUndoBracket(FString::Printf(TEXT("Add '%s' Node"), *Name.ToString()));
+
+	if (URigVMTemplateNode* ModelNode = Controller->AddTemplateNode(Notation, Location, Name.ToString(), true, false))
 	{
-		const FRigVMTemplate* Template = FRigVMRegistry::Get().FindTemplate(Notation);
-		if (Template == nullptr)
+		NewNode = Cast<URigVMEdGraphNode>(EdGraph->FindNodeForModelNodeName(ModelNode->GetFName()));
+
+		if (NewNode)
 		{
-			return nullptr;
+			if(FromPin)
+			{
+				NewNode->AutowireNewNode(FromPin);
+			}
+
+			Controller->ClearNodeSelection(true);
+			Controller->SelectNode(ModelNode, true, true);
 		}
 
-		const int32 NotationHash = (int32)GetTypeHash(Notation);
-		const FString TemplateName = TEXT("RigVMTemplate_") + FString::FromInt(NotationHash);
+		Controller->CloseUndoBracket();
+	}
+	else
+	{
+		Controller->CancelUndoBracket();
+	}
 
-		FName Name = UE::AnimNext::Editor::FUtils::ValidateName(Cast<UObject>(Host), Template->GetName().ToString());
-		URigVMController* Controller = Host->GetRigVMClient()->GetController(ParentGraph);
+	return NewNode;
+}
 
-		Controller->OpenUndoBracket(FString::Printf(TEXT("Add '%s' Node"), *Name.ToString()));
+FAnimNextSchemaAction_Variable::FAnimNextSchemaAction_Variable(const FRigVMExternalVariable& InExternalVariable, bool bInIsGetter)
+	: ExternalVariable(InExternalVariable)
+	, bIsGetter(bInIsGetter)
+{
+	static const FText VariablesCategory = LOCTEXT("Variables", "Variables");
+	static const FTextFormat GetVariableFormat = LOCTEXT("GetVariableFormat", "Get {0}");
+	static const FTextFormat SetVariableFormat = LOCTEXT("SetVariableFormat", "Set {0}");
 
-		if (URigVMTemplateNode* ModelNode = Controller->AddTemplateNode(Notation, Location, Name.ToString(), true, false))
-		{
-			NewNode = Cast<URigVMEdGraphNode>(EdGraph->FindNodeForModelNodeName(ModelNode->GetFName()));
+	FText MenuDesc;
+	FText ToolTip;
 
-			if (NewNode)
+	if(bInIsGetter)
+	{
+		MenuDesc = FText::Format(GetVariableFormat, FText::FromName(ExternalVariable.Name));
+		ToolTip = FText::FromString(FString::Printf(TEXT("Get the value of variable %s"), *ExternalVariable.Name.ToString()));
+	}
+	else
+	{
+		MenuDesc = FText::Format(SetVariableFormat, FText::FromName(ExternalVariable.Name));
+		ToolTip = FText::FromString(FString::Printf(TEXT("Set the value of variable %s"), *ExternalVariable.Name.ToString()));
+	}
+
+	UpdateSearchData(MenuDesc, ToolTip, VariablesCategory, FText::GetEmpty());
+
+	FEdGraphPinType PinType = RigVMTypeUtils::PinTypeFromExternalVariable(ExternalVariable);
+	VariableColor = GetDefault<UEdGraphSchema_K2>()->GetPinTypeColor(PinType);
+}
+
+UEdGraphNode* FAnimNextSchemaAction_Variable::PerformAction(UEdGraph* ParentGraph, TArray<UEdGraphPin*>& FromPins, const FVector2D Location, bool bSelectNewNode)
+{
+	IRigVMClientHost* Host = ParentGraph->GetImplementingOuter<IRigVMClientHost>();
+	if(Host == nullptr)
+	{
+		return nullptr;
+	}
+
+	URigVMEdGraph* EdGraph = Cast<URigVMEdGraph>(ParentGraph);
+	if(EdGraph == nullptr)
+	{
+		return nullptr;
+	}
+
+	URigVMEdGraphNode* NewNode = nullptr;
+
+	FString ObjectPath;
+	if (ExternalVariable.TypeObject)
+	{
+		ObjectPath = ExternalVariable.TypeObject->GetPathName();
+	}
+
+	FString TypeName = ExternalVariable.TypeName.ToString();
+	if (ExternalVariable.bIsArray)
+	{
+		TypeName = FString::Printf(TEXT("TArray<%s>"), *TypeName);
+	}
+
+	FString NodeName;
+
+	URigVMController* Controller = Host->GetRigVMClient()->GetController(ParentGraph);
+	Controller->OpenUndoBracket(TEXT("Add Variable"));
+
+	if (URigVMNode* ModelNode = Controller->AddVariableNodeFromObjectPath(ExternalVariable.Name, TypeName, ObjectPath, bIsGetter, FString(), Location, NodeName, true, true))
+	{
+		for (UEdGraphNode* Node : ParentGraph->Nodes)
+		{
+			if (URigVMEdGraphNode* RigNode = Cast<URigVMEdGraphNode>(Node))
 			{
-				if(FromPin)
+				if (RigNode->GetModelNodeName() == ModelNode->GetFName())
 				{
-					NewNode->AutowireNewNode(FromPin);
+					NewNode = RigNode;
+					break;
 				}
-
-				Controller->ClearNodeSelection(true);
-				Controller->SelectNode(ModelNode, true, true);
 			}
-
-			Controller->CloseUndoBracket();
 		}
-		else
+
+		if (NewNode)
 		{
-			Controller->CancelUndoBracket();
+			Controller->ClearNodeSelection(true);
+			Controller->SelectNode(ModelNode, true, true);
 		}
+		Controller->CloseUndoBracket();
 	}
+	else
+	{
+		Controller->CancelUndoBracket();
+	}
+
 	return NewNode;
-}
\ No newline at end of file
+}
+
+FAnimNextSchemaAction_AddComment::FAnimNextSchemaAction_AddComment()
+	: FAnimNextSchemaAction(FText(), LOCTEXT("AddComment", "Add Comment..."), LOCTEXT("AddCommentTooltip", "Create a resizable comment box."))
+{
+}
+
+UEdGraphNode* FAnimNextSchemaAction_AddComment::PerformAction(UEdGraph* ParentGraph, TArray<UEdGraphPin*>& FromPins, const FVector2D Location, bool bSelectNewNode/* = true*/)
+{
+	UEdGraphNode_Comment* const CommentTemplate = NewObject<UEdGraphNode_Comment>();
+
+	FVector2D SpawnLocation = Location;
+	FSlateRect Bounds;
+
+	TSharedPtr<SGraphEditor> GraphEditorPtr = SGraphEditor::FindGraphEditorForGraph(ParentGraph);
+	if (GraphEditorPtr.IsValid() && GraphEditorPtr->GetBoundsForSelectedNodes(/*out*/ Bounds, 50.0f))
+	{
+		CommentTemplate->SetBounds(Bounds);
+		SpawnLocation.X = CommentTemplate->NodePosX;
+		SpawnLocation.Y = CommentTemplate->NodePosY;
+	}
+
+	return FEdGraphSchemaAction_NewNode::SpawnNodeFromTemplate<UEdGraphNode_Comment>(ParentGraph, CommentTemplate, SpawnLocation, bSelectNewNode);
+}
+
+#undef LOCTEXT_NAMESPACE
\ No newline at end of file
diff --git a/Engine/Plugins/Experimental/AnimNext/Source/AnimNextEditor/Private/Common/GraphEditorSchemaActions.h b/Engine/Plugins/Experimental/AnimNext/Source/AnimNextEditor/Private/Common/GraphEditorSchemaActions.h
--- a/Engine/Plugins/Experimental/AnimNext/Source/AnimNextEditor/Private/Common/GraphEditorSchemaActions.h
+++ b/Engine/Plugins/Experimental/AnimNext/Source/AnimNextEditor/Private/Common/GraphEditorSchemaActions.h
@@ -5,6 +5,7 @@
 #include "CoreMinimal.h"
 #include "EdGraph/EdGraphSchema.h"
 #include "Editor/RigVMEditorStyle.h"
+#include "RigVMCore/RigVMExternalVariable.h"
 #include "Styling/AppStyle.h"
 #include "GraphEditorSchemaActions.generated.h"
 
@@ -80,3 +81,50 @@ private:
 	// Notation for dispatch factory
 	FName Notation;
 };
+
+
+USTRUCT()
+struct FAnimNextSchemaAction_Variable : public FAnimNextSchemaAction
+{
+	GENERATED_BODY()
+
+	FAnimNextSchemaAction_Variable() = default;
+
+	FAnimNextSchemaAction_Variable(const FRigVMExternalVariable& InExternalVariable, bool bInIsGetter = true);
+
+	// FEdGraphSchemaAction Interface
+	virtual UEdGraphNode* PerformAction(UEdGraph* ParentGraph, TArray<UEdGraphPin*>& FromPins, const FVector2D Location, bool bSelectNewNode = true) override;
+	virtual UEdGraphNode* PerformAction(UEdGraph* ParentGraph, UEdGraphPin* FromPin, const FVector2D Location, bool bSelectNewNode = true) { return nullptr; }
+
+	virtual const FSlateBrush* GetIconBrush() const override
+	{
+		return  FAppStyle::GetBrush(TEXT("Kismet.VariableList.TypeIcon"));
+	}
+
+	virtual const FLinearColor& GetIconColor() const override
+	{
+		return VariableColor;
+	}
+
+private:
+	FRigVMExternalVariable ExternalVariable;
+	bool bIsGetter = false;
+	FLinearColor VariableColor;
+};
+
+USTRUCT()
+struct FAnimNextSchemaAction_AddComment : public FAnimNextSchemaAction
+{
+	GENERATED_BODY()
+
+	FAnimNextSchemaAction_AddComment();
+
+	// FEdGraphSchemaAction Interface
+	virtual UEdGraphNode* PerformAction(UEdGraph* ParentGraph, TArray<UEdGraphPin*>& FromPins, const FVector2D Location, bool bSelectNewNode = true) override;
+	virtual UEdGraphNode* PerformAction(UEdGraph* ParentGraph, UEdGraphPin* FromPin, const FVector2D Location, bool bSelectNewNode = true) { return nullptr; }
+
+	virtual const FSlateBrush* GetIconBrush() const
+	{
+		return FAppStyle::Get().GetBrush("Icons.Comment");
+	}
+};
\ No newline at end of file
diff --git a/Engine/Plugins/Experimental/AnimNext/Source/AnimNextEditor/Private/Common/SActionMenu.cpp b/Engine/Plugins/Experimental/AnimNext/Source/AnimNextEditor/Private/Common/SActionMenu.cpp
--- a/Engine/Plugins/Experimental/AnimNext/Source/AnimNextEditor/Private/Common/SActionMenu.cpp
+++ b/Engine/Plugins/Experimental/AnimNext/Source/AnimNextEditor/Private/Common/SActionMenu.cpp
@@ -7,6 +7,7 @@
 #include "IDocumentation.h"
 #include "SSubobjectEditor.h"
 #include "AnimNextEdGraphSchema.h"
+#include "RigVMHost.h"
 #include "Widgets/Images/SImage.h"
 #include "Widgets/Layout/SBox.h"
 #include "Widgets/SBoxPanel.h"
@@ -45,9 +46,19 @@ void SActionMenu::CollectAllAnimNextGraphActions(FGraphContextMenuBuilder& MenuB
 			continue;
 		}
 
+		// skip hidden units
+		if(Function.Struct->HasMetaData(FRigVMStruct::HiddenMetaName))
+		{
+			continue;
+		}
+
 		FString CategoryMetadata, DisplayNameMetadata, MenuDescSuffixMetadata;
 		Struct->GetStringMetaDataHierarchical(FRigVMStruct::CategoryMetaName, &CategoryMetadata);
 		Struct->GetStringMetaDataHierarchical(FRigVMStruct::DisplayNameMetaName, &DisplayNameMetadata);
+		if(DisplayNameMetadata.IsEmpty())
+		{
+			DisplayNameMetadata = Function.GetMethodName().ToString();
+		}
 		Struct->GetStringMetaDataHierarchical(FRigVMStruct::MenuDescSuffixMetaName, &MenuDescSuffixMetadata);
 		if (!MenuDescSuffixMetadata.IsEmpty())
 		{
@@ -84,6 +95,12 @@ void SActionMenu::CollectAllAnimNextGraphActions(FGraphContextMenuBuilder& MenuB
 			continue;
 		}
 
+		// skip hidden factories
+		if(Factory->GetScriptStruct()->HasMetaData(FRigVMStruct::HiddenMetaName))
+		{
+			continue;
+		}
+
 		FText NodeCategory = FText::FromString(Factory->GetCategory());
 		FText MenuDesc = FText::FromString(Factory->GetNodeTitle(FRigVMTemplateTypeMap()));
 		FText ToolTip = Factory->GetNodeTooltip(FRigVMTemplateTypeMap());
@@ -91,6 +108,15 @@ void SActionMenu::CollectAllAnimNextGraphActions(FGraphContextMenuBuilder& MenuB
 		MenuBuilder.AddAction(MakeShared<FAnimNextSchemaAction_DispatchFactory>(Template->GetNotation(), NodeCategory, MenuDesc, ToolTip));
 	};
 
+	TArray<FRigVMExternalVariable> ExternalVariables = RigVMHost->GetVM()->GetExternalVariables(RigVMHost->GetRigVMExtendedExecuteContext());
+	for (const FRigVMExternalVariable& ExternalVariable : ExternalVariables)
+	{
+		MenuBuilder.AddAction(MakeShared<FAnimNextSchemaAction_Variable>(ExternalVariable, true));
+		MenuBuilder.AddAction(MakeShared<FAnimNextSchemaAction_Variable>(ExternalVariable, false));
+	}
+
+	MenuBuilder.AddAction(MakeShared<FAnimNextSchemaAction_AddComment>());
+
 	RigVMController->EnableReporting(true);
 }
 
@@ -113,6 +139,8 @@ void SActionMenu::Construct(const FArguments& InArgs, UEdGraph* InGraph)
 
 	RigVMClientHost = Graph->GetImplementingOuter<IRigVMClientHost>();
 	check(RigVMClientHost);
+	RigVMHost = Graph->GetTypedOuter<URigVMHost>();
+	check(RigVMHost);
 	RigVMController = RigVMClientHost->GetRigVMClient()->GetController(Graph);
 	check(RigVMController);
 	RigVMSchema = RigVMController->GetGraph()->GetSchema();
diff --git a/Engine/Plugins/Experimental/AnimNext/Source/AnimNextEditor/Private/Common/SRigVMAssetView.cpp b/Engine/Plugins/Experimental/AnimNext/Source/AnimNextEditor/Private/Common/SRigVMAssetView.cpp
--- a/Engine/Plugins/Experimental/AnimNext/Source/AnimNextEditor/Private/Common/SRigVMAssetView.cpp
+++ b/Engine/Plugins/Experimental/AnimNext/Source/AnimNextEditor/Private/Common/SRigVMAssetView.cpp
@@ -4,7 +4,7 @@
 
 #include "AnimNextRigVMAsset.h"
 #include "AnimNextRigVMAssetEditorData.h"
-#include "AnimNextRigVMAssetEntry.h"
+#include "Entries/AnimNextRigVMAssetEntry.h"
 #include "Common/SCategoryTableRow.h"
 #include "Widgets/Layout/SBox.h"
 #include "Widgets/Input/SSearchBox.h"
@@ -13,7 +13,7 @@
 #include "ISourceControlProvider.h"
 #include "EditorUtils.h"
 #include "IAnimNextRigVMExportInterface.h"
-#include "IAnimNextRigVMParameterInterface.h"
+#include "Variables/IAnimNextRigVMVariableInterface.h"
 #include "InstancedPropertyBagStructureDataProvider.h"
 #include "UncookedOnlyUtils.h"
 #include "RigVMAssetViewMenuContext.h"
@@ -150,7 +150,7 @@ void SRigVMAssetView::Construct(const FArguments& InArgs, UAnimNextRigVMAssetEdi
 	{
 		Categories.Add(MakeShared<FRigVMAssetViewEntry>(CategoryPair.Key, CategoryPair.Value));
 	}
-
+	
 	// Cache asset data for comparisons/filtering
 	AssetData = FAssetData(UncookedOnly::FUtils::GetAsset(EditorData));
 
@@ -436,11 +436,17 @@ void SRigVMAssetView::BindCommands()
 		FCanExecuteAction::CreateSP(this, &SRigVMAssetView::HasValidSingleSelection));
 }
 
-void SRigVMAssetView::HandleAssetModified(UAnimNextRigVMAssetEditorData* InEditorData)
+void SRigVMAssetView::HandleAssetModified(UAnimNextRigVMAssetEditorData* InEditorData, EAnimNextEditorDataNotifType InType, UObject* InSubject)
 {
 	check(InEditorData == EditorData);
 
-	RequestRefresh();
+	switch (InType)
+	{
+	case EAnimNextEditorDataNotifType::EntryAdded:
+	case EAnimNextEditorDataNotifType::EntryRemoved:
+		RequestRefresh();
+		break;
+	}
 }
 
 TSharedRef<SWidget> SRigVMAssetView::HandleGetContextContent()
@@ -563,7 +569,7 @@ class SRigVMAssetViewRow : public SMultiColumnTableRow<TSharedRef<FRigVMAssetVie
 		{
 			if (UAnimNextRigVMAssetEditorData* EditorData = View->EditorData)
 			{
-				if (UAnimNextRigVMAssetEntry* AssetEntry = EditorData->FindEntry(UncookedOnly::FUtils::GetParameterNameFromQualifiedName(PropertyAboutToChange->GetFName())))
+				if (UAnimNextRigVMAssetEntry* AssetEntry = EditorData->FindEntry(PropertyAboutToChange->GetFName()))
 				{
 					UAnimNextRigVMAsset* Asset = UE::AnimNext::UncookedOnly::FUtils::GetAsset(EditorData);
 
@@ -581,14 +587,14 @@ class SRigVMAssetViewRow : public SMultiColumnTableRow<TSharedRef<FRigVMAssetVie
 		{
 			if (UAnimNextRigVMAssetEditorData* EditorData = View->EditorData)
 			{
-				if (UAnimNextRigVMAssetEntry* AssetEntry = EditorData->FindEntry(UncookedOnly::FUtils::GetParameterNameFromQualifiedName(PropertyChangedEvent.GetMemberPropertyName())))
+				if (UAnimNextRigVMAssetEntry* AssetEntry = EditorData->FindEntry(PropertyChangedEvent.GetMemberPropertyName()))
 				{
 					// Needed to show the changed sign a the table when we modify the PropertyBag
 					// TODO: remove this by moving defaults into entries
 					AssetEntry->MarkPackageDirty();
 
 					// Ensure that default values get picked up and forwarded to compiler
-					EditorData->BroadcastModified();
+					EditorData->BroadcastModified(EAnimNextEditorDataNotifType::PropertyChanged, AssetEntry);
 				}
 			}
 		}
@@ -692,7 +698,7 @@ class SRigVMAssetViewRow : public SMultiColumnTableRow<TSharedRef<FRigVMAssetVie
 		}
 		else if(InColumnName == Column_Type)
 		{
-			if(Entry->WeakEntry.Get()->Implements<UAnimNextRigVMParameterInterface>())
+			if(Entry->WeakEntry.Get()->Implements<UAnimNextRigVMVariableInterface>())
 			{
 				return
 					SNew(SBox)
@@ -702,21 +708,21 @@ class SRigVMAssetViewRow : public SMultiColumnTableRow<TSharedRef<FRigVMAssetVie
 						SNew(SPinTypeSelector, FGetPinTypeTree::CreateStatic(&Editor::FUtils::GetFilteredVariableTypeTree))
 							.TargetPinType_Lambda([this]()
 							{
-								if(const IAnimNextRigVMParameterInterface* Binding = Cast<IAnimNextRigVMParameterInterface>(Entry->WeakEntry.Get()))
+								if(const IAnimNextRigVMVariableInterface* Binding = Cast<IAnimNextRigVMVariableInterface>(Entry->WeakEntry.Get()))
 								{
-									return UncookedOnly::FUtils::GetPinTypeFromParamType(Binding->GetParamType());
+									return UncookedOnly::FUtils::GetPinTypeFromParamType(Binding->GetType());
 								}
 
 								return FEdGraphPinType();
 							})
 							.OnPinTypeChanged_Lambda([this](const FEdGraphPinType& PinType)
 							{
-								if(IAnimNextRigVMParameterInterface* Binding = Cast<IAnimNextRigVMParameterInterface>(Entry->WeakEntry.Get()))
+								if(IAnimNextRigVMVariableInterface* Binding = Cast<IAnimNextRigVMVariableInterface>(Entry->WeakEntry.Get()))
 								{
 									const FAnimNextParamType ParamType = UncookedOnly::FUtils::GetParamTypeFromPinType(PinType);
 									if(ParamType.IsValid())
 									{
-										Binding->SetParamType(ParamType);
+										Binding->SetType(ParamType);
 									}
 								}
 							})
@@ -726,6 +732,17 @@ class SRigVMAssetViewRow : public SMultiColumnTableRow<TSharedRef<FRigVMAssetVie
 							.Font(IDetailLayoutBuilder::GetDetailFont())
 					];
 			}
+			else
+			{
+				return
+					SNew(SBox)
+					.HAlign(HAlign_Left)
+					.VAlign(VAlign_Center)
+					[
+						SNew(STextBlock)
+						.Text(Entry->WeakEntry.Get()->GetClass()->GetDisplayNameText())
+					];
+			}
 		}
 		else if(InColumnName == Column_Name)
 		{
@@ -754,7 +771,7 @@ class SRigVMAssetViewRow : public SMultiColumnTableRow<TSharedRef<FRigVMAssetVie
 						.IsSelected(this, &SRigVMAssetViewRow::IsSelectedExclusively)
 						.IsReadOnly_Lambda([this]()
 						{
-							if (Cast<IAnimNextRigVMParameterInterface>(Entry->WeakEntry.Get()) 
+							if (Cast<IAnimNextRigVMVariableInterface>(Entry->WeakEntry.Get()) 
 								|| Cast<IAnimNextRigVMGraphInterface>(Entry->WeakEntry.Get()))
 							{
 								return false;
@@ -858,11 +875,11 @@ class SRigVMAssetViewRow : public SMultiColumnTableRow<TSharedRef<FRigVMAssetVie
 			{
 				if(UAnimNextRigVMAssetEntry* AssetEntry = Cast<UAnimNextRigVMAssetEntry>(Entry->WeakEntry.Get()))
 				{
-					if (const IAnimNextRigVMParameterInterface* ParameterInterface = Cast<IAnimNextRigVMParameterInterface>(AssetEntry))
+					if (IAnimNextRigVMVariableInterface* ParameterInterface = Cast<IAnimNextRigVMVariableInterface>(AssetEntry))
 					{
-						FInstancedPropertyBag& PropertyBag = ParameterInterface->GetPropertyBag();
-						const FName ParameterName = ParameterInterface->GetParamName();
-						if (const FPropertyBagPropertyDesc* PropertyDesc = PropertyBag.FindPropertyDescByName(ParameterName))
+						FInstancedPropertyBag& PropertyBag = ParameterInterface->GetMutablePropertyBag();
+						const FName EntryName = AssetEntry->GetEntryName();
+						if (const FPropertyBagPropertyDesc* PropertyDesc = PropertyBag.FindPropertyDescByName(EntryName))
 						{
 							if (PropertyDesc->ContainerTypes.IsEmpty()) // avoid trying to inline containers
 							{
@@ -871,7 +888,7 @@ class SRigVMAssetViewRow : public SMultiColumnTableRow<TSharedRef<FRigVMAssetVie
 								SinglePropertyArgs.NotifyHook = this;
 
 								FPropertyEditorModule& PropertyEditorModule = FModuleManager::Get().LoadModuleChecked<FPropertyEditorModule>("PropertyEditor");
-								const TSharedPtr<ISinglePropertyView> SingleStructPropertyView = PropertyEditorModule.CreateSingleProperty(MakeShared<FInstancePropertyBagStructureDataProvider>(PropertyBag), ParameterName, SinglePropertyArgs);
+								const TSharedPtr<ISinglePropertyView> SingleStructPropertyView = PropertyEditorModule.CreateSingleProperty(MakeShared<FInstancePropertyBagStructureDataProvider>(PropertyBag), EntryName, SinglePropertyArgs);
 								if (SingleStructPropertyView.IsValid())
 								{
 									ColumnWidget = SingleStructPropertyView.ToSharedRef();
diff --git a/Engine/Plugins/Experimental/AnimNext/Source/AnimNextEditor/Private/Common/SRigVMAssetView.h b/Engine/Plugins/Experimental/AnimNext/Source/AnimNextEditor/Private/Common/SRigVMAssetView.h
--- a/Engine/Plugins/Experimental/AnimNext/Source/AnimNextEditor/Private/Common/SRigVMAssetView.h
+++ b/Engine/Plugins/Experimental/AnimNext/Source/AnimNextEditor/Private/Common/SRigVMAssetView.h
@@ -11,6 +11,7 @@ class URigVMGraph;
 class FUICommandList;
 class UAnimNextRigVMAssetEntry;
 class UAnimNextRigVMAssetEditorData;
+enum class EAnimNextEditorDataNotifType : uint8;
 
 namespace UE::AnimNext::Editor
 {
@@ -20,16 +21,26 @@ struct FRigVMAssetViewEntry;
 class SRigVMAssetView : public SCompoundWidget
 {
 public:
-	DECLARE_DELEGATE_OneParam(FOnSelectionChanged, const TArray<UObject*>& /*InEntries*/);
+	using FOnSelectionChanged = TDelegate<void(const TArray<UObject*>& InEntries)>;
 
-	DECLARE_DELEGATE_OneParam(FOnOpenGraph, URigVMGraph* /*InGraph*/);
+	using FOnOpenGraph = TDelegate<void(URigVMGraph* InGraph)>;
 
-	DECLARE_DELEGATE_OneParam(FOnDeleteEntries, const TArray<UAnimNextRigVMAssetEntry*>& /*InEntries*/);
+	using FOnDeleteEntries = TDelegate<void(const TArray<UAnimNextRigVMAssetEntry*>& InEntries)>;
+
+	enum class EFilterResult
+	{
+		Exclude,
+		Include
+	};
+	
+	using FOnFilterEntry = TDelegate<EFilterResult(const UAnimNextRigVMAssetEntry* InEntry)>;
+
+	using FOnFilterCategory = TDelegate<EFilterResult(FName InCategory)>;
 	
 	SLATE_BEGIN_ARGS(SRigVMAssetView) {}
 
 	SLATE_EVENT(SRigVMAssetView::FOnSelectionChanged, OnSelectionChanged)
-	
+
 	SLATE_EVENT(SRigVMAssetView::FOnOpenGraph, OnOpenGraph)
 
 	SLATE_EVENT(SRigVMAssetView::FOnDeleteEntries, OnDeleteEntries)
@@ -59,7 +70,7 @@ private:
 	void BindCommands();
 
 	// Handle modifications to the asset
-	void HandleAssetModified(UAnimNextRigVMAssetEditorData* InEditorData);
+	void HandleAssetModified(UAnimNextRigVMAssetEditorData* InEditorData, EAnimNextEditorDataNotifType InType, UObject* InSubject);
 
 	// Get the content for the context menu
 	TSharedRef<SWidget> HandleGetContextContent();
@@ -115,6 +126,9 @@ private:
 	bool bRefreshRequested = false;
 
 	static TMap<FName, FCategoryWidgetFactoryFunction> CategoryFactories;
+
+	// Map from category name -> display text
+	TMap<FName, FText> CategoryNameMap;
 };
 
 }
\ No newline at end of file
diff --git a/Engine/Plugins/Experimental/AnimNext/Source/AnimNextEditor/Private/EditorUtils.cpp b/Engine/Plugins/Experimental/AnimNext/Source/AnimNextEditor/Private/EditorUtils.cpp
--- a/Engine/Plugins/Experimental/AnimNext/Source/AnimNextEditor/Private/EditorUtils.cpp
+++ b/Engine/Plugins/Experimental/AnimNext/Source/AnimNextEditor/Private/EditorUtils.cpp
@@ -2,15 +2,19 @@
 
 #include "EditorUtils.h"
 
-#include "AnimNextRigVMAssetEntry.h"
+#include "IAnimNextUncookedOnlyModule.h"
+#include "Entries/AnimNextRigVMAssetEntry.h"
 #include "Module/AnimNextModule_EditorData.h"
 #include "Param/ParamType.h"
 #include "Kismet2/Kismet2NameValidators.h"
 #include "Module/AnimNextModule.h"
 #include "PropertyBagDetails.h"
 #include "UncookedOnlyUtils.h"
+#include "Modules/ModuleManager.h"
 #include "Param/RigVMDispatch_GetParameter.h"
 #include "String/ParseTokens.h"
+#include "Variables/IVariableBindingType.h"
+#include "Widgets/SNullWidget.h"
 
 #define LOCTEXT_NAMESPACE "AnimNextEditorUtils"
 
@@ -402,11 +406,12 @@ bool FUtils::IsValidParameterName(const FName InName, FText& OutErrorText)
 
 bool FUtils::DoesParameterNameExistInAsset(const FName InName, const FAssetData& InAsset)
 {
-	FAnimNextParameterProviderAssetRegistryExports Exports;
-	UncookedOnly::FUtils::GetExportedParametersForAsset(InAsset, Exports);
-	return Exports.Parameters.ContainsByPredicate([InName](const FAnimNextParameterAssetRegistryExportEntry& Entry) { return Entry.Name == InName; });
+	FAnimNextAssetRegistryExports Exports;
+	UncookedOnly::FUtils::GetExportedVariablesForAsset(InAsset, Exports);
+	return Exports.Variables.ContainsByPredicate([InName](const FAnimNextAssetRegistryExportedVariable& Entry) { return Entry.Name == InName; });
 }
 
+
 }
 
 #undef LOCTEXT_NAMESPACE
\ No newline at end of file
diff --git a/Engine/Plugins/Experimental/AnimNext/Source/AnimNextEditor/Private/EditorUtils.h b/Engine/Plugins/Experimental/AnimNext/Source/AnimNextEditor/Private/EditorUtils.h
--- a/Engine/Plugins/Experimental/AnimNext/Source/AnimNextEditor/Private/EditorUtils.h
+++ b/Engine/Plugins/Experimental/AnimNext/Source/AnimNextEditor/Private/EditorUtils.h
@@ -5,14 +5,17 @@
 #include "CoreMinimal.h"
 #include "Param/ParamType.h"
 #include "EdGraphSchema_K2.h"
+#include "StructView.h"
 
+struct FAnimNextVariableBindingData;
 class UAnimNextModule;
 class UAnimNextRigVMAssetEditorData;
 struct FAnimNextParamType;
 class URigVMController;
 struct FAnimNextWorkspaceAssetRegistryExports;
+class SWidget;
 
-struct FAnimNextParameterProviderAssetRegistryExports;
+struct FAnimNextAssetRegistryExports;
 
 namespace UE::AnimNext::Editor
 {
diff --git a/Engine/Plugins/Experimental/AnimNext/Source/AnimNextEditor/Private/Graph/AnimNextAnimationGraphAssetDefinition.cpp b/Engine/Plugins/Experimental/AnimNext/Source/AnimNextEditor/Private/Graph/AnimNextAnimationGraphAssetDefinition.cpp
new file mode 100644
--- /dev/null
+++ b/Engine/Plugins/Experimental/AnimNext/Source/AnimNextEditor/Private/Graph/AnimNextAnimationGraphAssetDefinition.cpp
@@ -0,0 +1,75 @@
+﻿// Copyright Epic Games, Inc. All Rights Reserved.
+
+#include "AnimNextAnimationGraphAssetDefinition.h"
+#include "ContentBrowserMenuContexts.h"
+#include "FileHelpers.h"
+#include "IWorkspaceEditorModule.h"
+#include "ToolMenus.h"
+#include "Workspace/AnimNextWorkspaceFactory.h"
+
+#define LOCTEXT_NAMESPACE "AnimNextAssetDefinitions"
+
+EAssetCommandResult UAssetDefinition_AnimNextAnimationGraph::OpenAssets(const FAssetOpenArgs& OpenArgs) const
+{
+	using namespace UE::AnimNext::Editor;
+	using namespace UE::Workspace;
+
+	for (UAnimNextAnimationGraph* Asset : OpenArgs.LoadObjects<UAnimNextAnimationGraph>())
+	{
+		IWorkspaceEditorModule& WorkspaceEditorModule = FModuleManager::Get().LoadModuleChecked<IWorkspaceEditorModule>("WorkspaceEditor");
+		WorkspaceEditorModule.OpenWorkspaceForObject(Asset, EOpenWorkspaceMethod::Default, UAnimNextWorkspaceFactory::StaticClass());
+	}
+
+	return EAssetCommandResult::Handled;
+}
+
+namespace UE::AnimNext::Editor
+{
+
+static FDelayedAutoRegisterHelper AutoRegisterGraphMenuItems(EDelayedRegisterRunPhase::EndOfEngineInit, []
+{
+	UToolMenus::RegisterStartupCallback(FSimpleMulticastDelegate::FDelegate::CreateLambda([]()
+	{
+		FToolMenuOwnerScoped OwnerScoped(UE_MODULE_NAME);
+		UToolMenu* Menu = UE::ContentBrowser::ExtendToolMenu_AssetContextMenu(UAnimNextAnimationGraph::StaticClass());
+		 
+		FToolMenuSection& Section = Menu->FindOrAddSection("GetAssetActions");
+		Section.AddDynamicEntry(NAME_None, FNewToolMenuSectionDelegate::CreateLambda([](FToolMenuSection& InSection)
+		{
+			InSection.AddMenuEntry(
+				"ForceSave",
+				LOCTEXT("ForceSaveLabel", "Force Save"),
+				LOCTEXT("ForceSaveTooltip", "Force the save of this item and all its subobjects."),
+				FSlateIcon(FAppStyle::Get().GetStyleSetName(), "Icons.Save"),
+				FToolUIAction(FToolMenuExecuteAction::CreateLambda([](const FToolMenuContext& InContext)
+				{
+					if(UContentBrowserAssetContextMenuContext* AssetContext = InContext.FindContext<UContentBrowserAssetContextMenuContext>())
+					{
+						TArray<UPackage*> PackagesToSave;
+						for (const FAssetData& AssetData : AssetContext->SelectedAssets)
+						{
+							if (UPackage* Package = AssetData.GetPackage())
+							{
+								PackagesToSave.Add(Package);
+								PackagesToSave.Append(Package->GetExternalPackages());
+							}
+						}
+
+						if(PackagesToSave.Num() > 0)
+						{
+							FEditorFileUtils::FPromptForCheckoutAndSaveParams SaveParams;
+							SaveParams.bCheckDirty = false;
+							SaveParams.bPromptToSave = false;
+							SaveParams.bIsExplicitSave = true;
+
+							FEditorFileUtils::PromptForCheckoutAndSave(PackagesToSave, SaveParams);
+						}
+					}
+				})));
+		}));
+	}));
+});
+
+}
+
+#undef LOCTEXT_NAMESPACE
\ No newline at end of file
diff --git a/Engine/Plugins/Experimental/AnimNext/Source/AnimNextEditor/Private/Graph/AnimNextAnimationGraphAssetDefinition.h b/Engine/Plugins/Experimental/AnimNext/Source/AnimNextEditor/Private/Graph/AnimNextAnimationGraphAssetDefinition.h
new file mode 100644
--- /dev/null
+++ b/Engine/Plugins/Experimental/AnimNext/Source/AnimNextEditor/Private/Graph/AnimNextAnimationGraphAssetDefinition.h
@@ -0,0 +1,30 @@
+﻿// Copyright Epic Games, Inc. All Rights Reserved.
+
+#pragma once
+
+#include "Graph/AnimNextAnimationGraph.h"
+#include "AssetDefinitionDefault.h"
+#include "AnimNextAnimationGraphAssetDefinition.generated.h"
+
+#define LOCTEXT_NAMESPACE "AnimNextAssetDefinitions"
+
+UCLASS()
+class UAssetDefinition_AnimNextAnimationGraph : public UAssetDefinitionDefault
+{
+	GENERATED_BODY()
+
+public:
+	// UAssetDefinition interface
+	virtual FText GetAssetDisplayName() const override { return LOCTEXT("AnimNextAnimationGraph", "AnimNext Animation Graph"); }
+	virtual FLinearColor GetAssetColor() const override { return FLinearColor(FColor(128,96,48)); }
+	virtual TSoftClassPtr<UObject> GetAssetClass() const override { return UAnimNextAnimationGraph::StaticClass(); }
+	virtual EAssetCommandResult OpenAssets(const FAssetOpenArgs& OpenArgs) const override;
+	virtual TConstArrayView<FAssetCategoryPath> GetAssetCategories() const override
+	{
+		static const auto Categories = { FAssetCategoryPath(EAssetCategoryPaths::Animation, LOCTEXT("AnimNextSubMenu", "AnimNext")) };
+		return Categories;
+	}
+	virtual bool ShouldSaveExternalPackages() const override { return true; }
+};
+
+#undef LOCTEXT_NAMESPACE
\ No newline at end of file
diff --git a/Engine/Plugins/Experimental/AnimNext/Source/AnimNextEditor/Private/Graph/AnimNextAnimationGraphFactory.cpp b/Engine/Plugins/Experimental/AnimNext/Source/AnimNextEditor/Private/Graph/AnimNextAnimationGraphFactory.cpp
new file mode 100644
--- /dev/null
+++ b/Engine/Plugins/Experimental/AnimNext/Source/AnimNextEditor/Private/Graph/AnimNextAnimationGraphFactory.cpp
@@ -0,0 +1,47 @@
+﻿// Copyright Epic Games, Inc. All Rights Reserved.
+
+#include "Graph/AnimNextAnimationGraphFactory.h"
+#include "Graph/AnimNextAnimationGraph.h"
+#include "Graph/AnimNextAnimationGraph_EditorData.h"
+#include "UncookedOnlyUtils.h"
+#include "Entries/AnimNextAnimationGraphEntry.h"
+
+UAnimNextAnimationGraphFactory::UAnimNextAnimationGraphFactory()
+{
+	bCreateNew = true;
+	bEditAfterNew = true;
+	SupportedClass = UAnimNextAnimationGraph::StaticClass();
+}
+
+bool UAnimNextAnimationGraphFactory::ConfigureProperties()
+{
+	return true;
+}
+
+UObject* UAnimNextAnimationGraphFactory::FactoryCreateNew(UClass* Class, UObject* InParent, FName Name, EObjectFlags Flags, UObject* Context, FFeedbackContext* Warn, FName CallingContext)
+{
+	EObjectFlags FlagsToUse = Flags | RF_Public | RF_Standalone | RF_Transactional | RF_LoadCompleted;
+	if(InParent == GetTransientPackage())
+	{
+		FlagsToUse &= ~RF_Standalone;
+	}
+
+	UAnimNextAnimationGraph* NewAnimationGraph = NewObject<UAnimNextAnimationGraph>(InParent, Class, Name, FlagsToUse);
+
+	// Create internal editor data
+	UAnimNextAnimationGraph_EditorData* EditorData = NewObject<UAnimNextAnimationGraph_EditorData>(NewAnimationGraph, TEXT("EditorData"), RF_Transactional);
+	NewAnimationGraph->EditorData = EditorData;
+	EditorData->Initialize(/*bRecompileVM*/false);
+
+	// Add a single internal graph entry
+	UAnimNextAnimationGraphEntry* NewGraph = EditorData->AddAnimationGraph(FRigUnit_AnimNextGraphRoot::DefaultEntryPoint);
+
+	// Hide in the outliner - the asset itself represents the graph to users
+	NewGraph->SetHiddenInOutliner(true);
+
+	// Compile the initial skeleton
+	EditorData->RecompileVM();
+	check(!EditorData->bErrorsDuringCompilation);
+
+	return NewAnimationGraph;
+}
\ No newline at end of file
diff --git a/Engine/Plugins/Experimental/AnimNext/Source/AnimNextEditor/Private/Graph/AnimNextCompilerResultsTabSummoner.cpp b/Engine/Plugins/Experimental/AnimNext/Source/AnimNextEditor/Private/Graph/AnimNextCompilerResultsTabSummoner.cpp
--- a/Engine/Plugins/Experimental/AnimNext/Source/AnimNextEditor/Private/Graph/AnimNextCompilerResultsTabSummoner.cpp
+++ b/Engine/Plugins/Experimental/AnimNext/Source/AnimNextEditor/Private/Graph/AnimNextCompilerResultsTabSummoner.cpp
@@ -10,7 +10,7 @@
 namespace UE::AnimNext::Editor
 {
 
-const FName CompilerResultsTabName("CompilerResultsTab");
+const FLazyName CompilerResultsTabName("CompilerResultsTab");
 
 // ***************************************************************************
 
diff --git a/Engine/Plugins/Experimental/AnimNext/Source/AnimNextEditor/Private/Graph/AnimNextCompilerResultsTabSummoner.h b/Engine/Plugins/Experimental/AnimNext/Source/AnimNextEditor/Private/Graph/AnimNextCompilerResultsTabSummoner.h
--- a/Engine/Plugins/Experimental/AnimNext/Source/AnimNextEditor/Private/Graph/AnimNextCompilerResultsTabSummoner.h
+++ b/Engine/Plugins/Experimental/AnimNext/Source/AnimNextEditor/Private/Graph/AnimNextCompilerResultsTabSummoner.h
@@ -13,7 +13,7 @@ namespace UE::Workspace
 namespace UE::AnimNext::Editor
 {
 
-extern const FName CompilerResultsTabName;
+extern const FLazyName CompilerResultsTabName;
 
 class SAnimNextCompilerResultsWidget : public SCompoundWidget
 {
diff --git a/Engine/Plugins/Experimental/AnimNext/Source/AnimNextEditor/Private/Graph/AnimNextEdGraphNodeCustomization.cpp b/Engine/Plugins/Experimental/AnimNext/Source/AnimNextEditor/Private/Graph/AnimNextEdGraphNodeCustomization.cpp
--- a/Engine/Plugins/Experimental/AnimNext/Source/AnimNextEditor/Private/Graph/AnimNextEdGraphNodeCustomization.cpp
+++ b/Engine/Plugins/Experimental/AnimNext/Source/AnimNextEditor/Private/Graph/AnimNextEdGraphNodeCustomization.cpp
@@ -30,7 +30,7 @@ void FAnimNextEdGraphNodeCustomization::PendingDelete()
 
 	if (TSharedPtr<UE::Workspace::IWorkspaceEditor> WorkspaceEditor = WorkspaceEditorWeak.Pin())
 	{
-		if (TSharedPtr<SDockTab> DockTab = WorkspaceEditor->GetTabManager()->FindExistingLiveTab(UE::AnimNext::Editor::TraitEditorTabName))
+		if (TSharedPtr<SDockTab> DockTab = WorkspaceEditor->GetTabManager()->FindExistingLiveTab(FTabId(UE::AnimNext::Editor::TraitEditorTabName)))
 		{
 			if (TSharedPtr<STraitEditorView> TraitEditorView = StaticCastSharedPtr<STraitEditorView>(DockTab->GetContent().ToSharedPtr()))
 			{
@@ -83,7 +83,7 @@ void FAnimNextEdGraphNodeCustomization::CustomizeObjects(IDetailLayoutBuilder& D
 		{
 			if (TSharedPtr<UE::Workspace::IWorkspaceEditor> WorkspaceEditor = WorkspaceEditorWeak.Pin())
 			{
-				if (TSharedPtr<SDockTab> DockTab = WorkspaceEditor->GetTabManager()->FindExistingLiveTab(UE::AnimNext::Editor::TraitEditorTabName))
+				if (TSharedPtr<SDockTab> DockTab = WorkspaceEditor->GetTabManager()->FindExistingLiveTab(FTabId(UE::AnimNext::Editor::TraitEditorTabName)))
 				{
 					if (TSharedPtr<STraitEditorView> TraitEditorView = StaticCastSharedPtr<STraitEditorView>(DockTab->GetContent().ToSharedPtr()))
 					{
diff --git a/Engine/Plugins/Experimental/AnimNext/Source/AnimNextEditor/Private/Graph/AnimNextGraphItemDetails.cpp b/Engine/Plugins/Experimental/AnimNext/Source/AnimNextEditor/Private/Graph/AnimNextGraphItemDetails.cpp
--- a/Engine/Plugins/Experimental/AnimNext/Source/AnimNextEditor/Private/Graph/AnimNextGraphItemDetails.cpp
+++ b/Engine/Plugins/Experimental/AnimNext/Source/AnimNextEditor/Private/Graph/AnimNextGraphItemDetails.cpp
@@ -3,16 +3,19 @@
 #include "AnimNextGraphItemDetails.h"
 
 #include "Toolkits/AssetEditorToolkitMenuContext.h"
-#include "Graph/AnimNextModule_AnimationGraph.h"
-#include "Module/AnimNextModule_EventGraph.h"
+#include "Entries/AnimNextAnimationGraphEntry.h"
+#include "Entries/AnimNextEventGraphEntry.h"
 #include "StructUtils/InstancedStruct.h"
-#include "Module/AnimNextModuleWorkspaceAssetUserData.h"
+#include "AnimNextAssetWorkspaceAssetUserData.h"
 #include "RigVMModel/RigVMGraph.h"
 #include "WorkspaceItemMenuContext.h"
 #include "IWorkspaceEditor.h"
+#include "ScopedTransaction.h"
 #include "RigVMModel/RigVMClient.h"
 #include "EdGraph/RigVMEdGraph.h"
 #include "ToolMenus.h"
+#include "Framework/Commands/GenericCommands.h"
+#include "Module/AnimNextModule_EditorData.h"
 
 #define LOCTEXT_NAMESPACE "FAnimNextGraphItemDetails"
 
@@ -69,6 +72,86 @@ void FAnimNextGraphItemDetails::HandleDoubleClick(const FToolMenuContext& ToolMe
 	}
 }
 
+void FAnimNextGraphItemDetails::Delete(TConstArrayView<FWorkspaceOutlinerItemExport> Exports) const
+{
+	TMap<UAnimNextRigVMAssetEditorData*, TArray<UAnimNextRigVMAssetEntry*>> EntriesToDelete;
+	for(const FWorkspaceOutlinerItemExport& Export : Exports)
+	{
+		const TInstancedStruct<FWorkspaceOutlinerItemData>& Data = Export.GetData();
+		if (Data.IsValid() && Data.GetScriptStruct() == FAnimNextGraphOutlinerData::StaticStruct())
+		{
+			const FAnimNextGraphOutlinerData& GraphData = Export.GetData().Get<FAnimNextGraphOutlinerData>();
+			UAnimNextRigVMAssetEditorData* EditorData = GraphData.Entry->GetTypedOuter<UAnimNextRigVMAssetEditorData>();
+			if(EditorData == nullptr)
+			{
+				continue;
+			}
+			UAnimNextRigVMAssetEntry* Entry = GraphData.Entry;
+			if(Entry == nullptr)
+			{
+				continue;
+			}
+			TArray<UAnimNextRigVMAssetEntry*>& Entries = EntriesToDelete.FindOrAdd(EditorData);
+			Entries.Add(Entry);
+		}
+	}
+
+	if(EntriesToDelete.Num() > 0)
+	{
+		FScopedTransaction Transaction(LOCTEXT("DeleteEntries", "Delete Entries"));
+		for(TPair<UAnimNextRigVMAssetEditorData*, TArray<UAnimNextRigVMAssetEntry*>> EntryToDeletePair : EntriesToDelete)
+		{
+			EntryToDeletePair.Key->RemoveEntries(EntryToDeletePair.Value);
+		}
+	}
+}
+
+bool FAnimNextGraphItemDetails::CanRename(const FWorkspaceOutlinerItemExport& Export) const
+{
+	const TInstancedStruct<FWorkspaceOutlinerItemData>& Data = Export.GetData();
+	if (Data.IsValid() && Data.GetScriptStruct() == FAnimNextGraphOutlinerData::StaticStruct())
+	{
+		return true;
+	}
+	return false;
+}
+
+void FAnimNextGraphItemDetails::Rename(const FWorkspaceOutlinerItemExport& Export, const FText& InName) const
+{
+	const TInstancedStruct<FWorkspaceOutlinerItemData>& Data = Export.GetData();
+	if (Data.IsValid() && Data.GetScriptStruct() == FAnimNextGraphOutlinerData::StaticStruct())
+	{
+		const FAnimNextGraphOutlinerData& GraphData = Export.GetData().Get<FAnimNextGraphOutlinerData>();
+		UAnimNextRigVMAssetEditorData* EditorData = GraphData.Entry->GetTypedOuter<UAnimNextRigVMAssetEditorData>();
+		FName NewName = FName(*InName.ToString());
+		UAnimNextRigVMAssetEntry* ExistingEntry = EditorData->FindEntry(NewName);
+		if(ExistingEntry == nullptr && GraphData.Entry->GetEntryName() != NewName)
+		{
+			FScopedTransaction Transaction(LOCTEXT("SetName", "Set Name"));
+			GraphData.Entry->SetEntryName(NewName);
+		}
+	}
+}
+
+bool FAnimNextGraphItemDetails::ValidateName(const FWorkspaceOutlinerItemExport& Export, const FText& InName, FText& OutErrorMessage) const
+{
+	const TInstancedStruct<FWorkspaceOutlinerItemData>& Data = Export.GetData();
+	if (Data.IsValid() && Data.GetScriptStruct() == FAnimNextGraphOutlinerData::StaticStruct())
+	{
+		const FAnimNextGraphOutlinerData& GraphData = Export.GetData().Get<FAnimNextGraphOutlinerData>();
+		UAnimNextRigVMAssetEditorData* EditorData = GraphData.Entry->GetTypedOuter<UAnimNextRigVMAssetEditorData>();
+		UAnimNextRigVMAssetEntry* ExistingEntry = EditorData->FindEntry(FName(*InName.ToString()));
+		if(ExistingEntry)
+		{
+			OutErrorMessage = LOCTEXT("NameAlreadyExistsError", "Name already exists in this module");
+			return false;
+		}
+
+		return true;
+	}
+	return false;
+}
+
 UPackage* FAnimNextGraphItemDetails::GetPackage(const FWorkspaceOutlinerItemExport& Export) const 
 {
 	const TInstancedStruct<FWorkspaceOutlinerItemData>& Data = Export.GetData();
@@ -102,139 +185,19 @@ UPackage* FAnimNextGraphItemDetails::GetPackage(const FWorkspaceOutlinerItemExpo
 	return nullptr;
 }
 
-const FSlateBrush* FAnimNextGraphItemDetails::GetItemIcon() const
+const FSlateBrush* FAnimNextGraphItemDetails::GetItemIcon(const FWorkspaceOutlinerItemExport& Export) const
 {
 	return FAppStyle::GetBrush(TEXT("GraphEditor.EventGraph_24x"));
 }
 
 void FAnimNextGraphItemDetails::RegisterToolMenuExtensions()
 {
-	FToolMenuOwnerScoped OwnerScoped(TEXT("FAnimNextGraphItemDetails"));
-	if (UToolMenu* Menu = UToolMenus::Get()->ExtendMenu("WorkspaceOutliner.ItemContextMenu"))
-	{
-		Menu->AddDynamicSection(TEXT("AnimNextGraphItem"), FNewToolMenuDelegate::CreateLambda([](UToolMenu* InMenu)
-		{
-			UWorkspaceItemMenuContext* WorkspaceItemContext = InMenu->FindContext<UWorkspaceItemMenuContext>();
-			const UAssetEditorToolkitMenuContext* AssetEditorContext = InMenu->FindContext<UAssetEditorToolkitMenuContext>();			
-			if (WorkspaceItemContext && AssetEditorContext)
-			{
-				FToolMenuSection& Section = InMenu->AddSection("WorkspaceOutliner.ItemContextMenu.RootAsset", FText::FromString(TEXT("Animation Next")));
-				if(const TSharedPtr<UE::Workspace::IWorkspaceEditor> WorkspaceEditor = StaticCastSharedPtr<UE::Workspace::IWorkspaceEditor>(AssetEditorContext->Toolkit.Pin()))
-				{
-					TArray<FWorkspaceOutlinerItemExport> GraphExports;
-					Algo::TransformIf(WorkspaceItemContext->SelectedExports, GraphExports, [](const FWorkspaceOutlinerItemExport& Export)
-						{
-							return Export.GetData().IsValid() 
-								&& (Export.GetData().GetScriptStruct() == FAnimNextGraphOutlinerData::StaticStruct()
-									|| Export.GetData().GetScriptStruct() == FAnimNextGraphFunctionOutlinerData::StaticStruct()
-									|| Export.GetData().GetScriptStruct() == FAnimNextCollapseGraphOutlinerData::StaticStruct());
-						},
-						[](const FWorkspaceOutlinerItemExport& Export)
-						{
-							return Export;
-						});
-
-					if (GraphExports.Num() == WorkspaceItemContext->SelectedExports.Num() && GraphExports.Num() > 0)
-					{
-						TWeakPtr<UE::Workspace::IWorkspaceEditor> WeakWorkspaceEditor = WorkspaceEditor;
-
-						TInstancedStruct<FWorkspaceOutlinerItemData>& Data = WorkspaceItemContext->SelectedExports[0].GetData();
-						if (Data.IsValid())
-						{
-							FName EntryText;
-							FText EntryLabelText;
-							FText EntryTooltipText;
-							FSlateIcon EntryIcon;
-
-							if (Data.GetScriptStruct() == FAnimNextGraphOutlinerData::StaticStruct())
-							{
-								EntryText = TEXT("OpenGraphMenuEntry");
-								EntryLabelText = FText::FormatOrdered(LOCTEXT("OpenGraphMenuEntryLabel", "Open {0}|plural(one=Animation Graph,other=Graphs)"), GraphExports.Num());
-								EntryTooltipText = FText::FormatOrdered(LOCTEXT("OpenGraphMenuEntryTooltip", "Open the selected {0}|plural(one=Animation Graph,other=Graphs)"), GraphExports.Num());
-								EntryIcon = FSlateIcon(FAppStyle::GetAppStyleSetName(), "GraphEditor.EventGraph_16x");
-							}
-							else if (Data.GetScriptStruct() == FAnimNextGraphFunctionOutlinerData::StaticStruct())
-							{
-								EntryText = TEXT("OpenFunctionActionEntry");
-								EntryLabelText = FText::FormatOrdered(LOCTEXT("OpenFunctionActionLabel", "Open {0}|plural(one=Function Graph,other=Grapsh)"), GraphExports.Num());
-								EntryTooltipText = FText::FormatOrdered(LOCTEXT("OpenFunctionActionTooltip", "Open the selected {0}|plural(one=Function Graph,other=Graphs)"), GraphExports.Num());
-								EntryIcon = FSlateIcon(FAppStyle::GetAppStyleSetName(), "GraphEditor.Function_16x");
-							}
-							else if (Data.GetScriptStruct() == FAnimNextCollapseGraphOutlinerData::StaticStruct())
-							{
-								EntryText = TEXT("OpenCollapseNodeActionEntry");
-								EntryLabelText = FText::FormatOrdered(LOCTEXT("OpenCollapseGraphActionLabel", "Open {0}|plural(one=Collapse Graph,other=Graphs)"), GraphExports.Num());
-								EntryTooltipText = FText::FormatOrdered(LOCTEXT("OpenCollapseGraphActionTooltip", "Open the selected {0}|plural(one=Collapse Graph,other=Graphs)"), GraphExports.Num());
-								EntryIcon = FSlateIcon(FAppStyle::GetAppStyleSetName(), "GraphEditor.SubGraph_16x");
-							}
 
-							Section.AddMenuEntry(
-								EntryText,
-								EntryLabelText,
-								EntryTooltipText,
-								EntryIcon,
-								FUIAction(
-									FExecuteAction::CreateWeakLambda(WorkspaceItemContext, [GraphExports, WeakWorkspaceEditor]()
-									{
-										if (TSharedPtr<UE::Workspace::IWorkspaceEditor> SharedWorkspaceEditor = WeakWorkspaceEditor.Pin())
-										{
-											TArray<UObject*> ObjectsToOpen;
-											for (const FWorkspaceOutlinerItemExport& Export : GraphExports)
-											{
-												const TInstancedStruct<FWorkspaceOutlinerItemData>& Data = Export.GetData();
-												if (Data.GetScriptStruct() == FAnimNextGraphOutlinerData::StaticStruct())
-												{
-													const FAnimNextGraphOutlinerData& GraphData = Export.GetData().Get<FAnimNextGraphOutlinerData>();
-													if (GraphData.GraphInterface)
-													{
-														if (URigVMGraph* RigVMGraph = GraphData.GraphInterface->GetRigVMGraph())
-														{
-															if (const IRigVMClientHost* RigVMClientHost = RigVMGraph->GetImplementingOuter<IRigVMClientHost>())
-															{
-																if (UObject* EditorObject = RigVMClientHost->GetEditorObjectForRigVMGraph(RigVMGraph))
-																{
-																	ObjectsToOpen.Add(EditorObject);
-																}
-															}
-														}
-													}
-												}
-												else if (Data.GetScriptStruct() == FAnimNextGraphFunctionOutlinerData::StaticStruct())
-												{
-													const FAnimNextGraphFunctionOutlinerData& GraphFunctionData = Data.Get<FAnimNextGraphFunctionOutlinerData>();
-													if (GraphFunctionData.EditorObject.IsValid())
-													{
-														ObjectsToOpen.Add(GraphFunctionData.EditorObject.Get());
-													}
-												}
-												else if (Data.GetScriptStruct() == FAnimNextCollapseGraphOutlinerData::StaticStruct())
-												{
-													const FAnimNextCollapseGraphOutlinerData& GraphFunctionData = Data.Get<FAnimNextCollapseGraphOutlinerData>();
-													if (GraphFunctionData.EditorObject.IsValid())
-													{
-														ObjectsToOpen.Add(GraphFunctionData.EditorObject.Get());
-													}
-												}
-											}
-
-											SharedWorkspaceEditor->OpenObjects({ ObjectsToOpen });
-										}
-									})
-								));
-						}
-					}
-				}
-			}
-		}), FToolMenuInsert(NAME_None, EToolMenuInsertType::First));
-	}
 }
 
 void FAnimNextGraphItemDetails::UnregisterToolMenuExtensions()
 {
-	if(UToolMenus* ToolMenus = UToolMenus::Get())
-	{
-		ToolMenus->UnregisterOwnerByName("FAnimNextGraphItemDetails");
-	}
+
 }
 
 } // UE::AnimNext::Editor
diff --git a/Engine/Plugins/Experimental/AnimNext/Source/AnimNextEditor/Private/Graph/AnimNextGraphItemDetails.h b/Engine/Plugins/Experimental/AnimNext/Source/AnimNextEditor/Private/Graph/AnimNextGraphItemDetails.h
--- a/Engine/Plugins/Experimental/AnimNext/Source/AnimNextEditor/Private/Graph/AnimNextGraphItemDetails.h
+++ b/Engine/Plugins/Experimental/AnimNext/Source/AnimNextEditor/Private/Graph/AnimNextGraphItemDetails.h
@@ -11,9 +11,13 @@ class FAnimNextGraphItemDetails : public UE::Workspace::IWorkspaceOutlinerItemDe
 public:
 	FAnimNextGraphItemDetails() = default;
 	virtual ~FAnimNextGraphItemDetails() override = default;
-	virtual void HandleDoubleClick(const FToolMenuContext& ToolMenuContext) const override;	
+	virtual void HandleDoubleClick(const FToolMenuContext& ToolMenuContext) const override;
+	virtual void Delete(TConstArrayView<FWorkspaceOutlinerItemExport> Exports) const override; 
+	virtual bool CanRename(const FWorkspaceOutlinerItemExport& Export) const override;
+	virtual void Rename(const FWorkspaceOutlinerItemExport& Export, const FText& InName) const override;
+	virtual bool ValidateName(const FWorkspaceOutlinerItemExport& Export, const FText& InName, FText& OutErrorMessage) const override;
 	virtual UPackage* GetPackage(const FWorkspaceOutlinerItemExport& Export) const override;
-	virtual const FSlateBrush* GetItemIcon() const override;
+	virtual const FSlateBrush* GetItemIcon(const FWorkspaceOutlinerItemExport& Export) const override;
 
 	static void RegisterToolMenuExtensions();
 	static void UnregisterToolMenuExtensions();
diff --git a/Engine/Plugins/Experimental/AnimNext/Source/AnimNextEditor/Private/Graph/STraitEditorView.cpp b/Engine/Plugins/Experimental/AnimNext/Source/AnimNextEditor/Private/Graph/STraitEditorView.cpp
--- a/Engine/Plugins/Experimental/AnimNext/Source/AnimNextEditor/Private/Graph/STraitEditorView.cpp
+++ b/Engine/Plugins/Experimental/AnimNext/Source/AnimNextEditor/Private/Graph/STraitEditorView.cpp
@@ -16,7 +16,7 @@
 #include "MessageLogModule.h"
 #include "IWorkspaceEditor.h"
 #include "IMessageLogListing.h"
-#include "Module/AnimNextModule_Controller.h"
+#include "AnimNextController.h"
 #include "Graph/AnimNextCompilerResultsTabSummoner.h"
 #include "Logging/MessageLog.h"
 #include "Widgets/Docking/SDockTab.h"
@@ -102,7 +102,7 @@ FReply STraitEditorView::OnTraitClicked(const FTraitUID InClickedTraitUID)
 	{
 		if (UAnimNextEdGraphNode* EdGraphNode = Cast<UAnimNextEdGraphNode>(TraitEditorSharedData->EdGraphNodeWeak))
 		{
-			if (UAnimNextModule_Controller* Controller = Cast<UAnimNextModule_Controller>(EdGraphNode->GetController()))
+			if (UAnimNextController* Controller = Cast<UAnimNextController>(EdGraphNode->GetController()))
 			{
 				int32 TraitIndex = INDEX_NONE;
 				TSharedPtr<FTraitDataEditorDef> SwapTraitData = FTraitEditorUtils::FindTraitInCurrentStackData(SelectedTraitUID, TraitEditorSharedData->CurrentTraitsDataShared, &TraitIndex);
@@ -147,7 +147,7 @@ FReply STraitEditorView::OnTraitDeleteRequest(const FTraitUID InTraitUIDToDelete
 				{
 					if (UAnimNextEdGraphNode* EdGraphNode = Cast<UAnimNextEdGraphNode>(TraitEditorSharedData->EdGraphNodeWeak.Get()))
 					{
-						if (UAnimNextModule_Controller* Controller = Cast<UAnimNextModule_Controller>(EdGraphNode->GetController()))
+						if (UAnimNextController* Controller = Cast<UAnimNextController>(EdGraphNode->GetController()))
 						{
 							Controller->RemoveTraitByName(EdGraphNode->GetFName(), TraitDataEditorDef->TraitName);
 						}
@@ -182,7 +182,7 @@ void STraitEditorView::ExecuteTraitDrag(const FTraitUID DraggedTraitUID, const F
 	{
 		if (UAnimNextEdGraphNode* EdGraphNode = Cast<UAnimNextEdGraphNode>(TraitEditorSharedData->EdGraphNodeWeak))
 		{
-			if (UAnimNextModule_Controller* Controller = Cast<UAnimNextModule_Controller>(EdGraphNode->GetController()))
+			if (UAnimNextController* Controller = Cast<UAnimNextController>(EdGraphNode->GetController()))
 			{
 				int32 TargetTraitIndex = INDEX_NONE;
 				TSharedPtr<FTraitDataEditorDef> SwapTraitData = FTraitEditorUtils::FindTraitInCurrentStackData(TargetTraitUID, TraitEditorSharedData->CurrentTraitsDataShared, &TargetTraitIndex);
@@ -298,7 +298,7 @@ void STraitEditorView::RefreshTraitStackTraitsStatus()
 					// Open tab to display errors
 					if (TSharedPtr<UE::Workspace::IWorkspaceEditor> WorkspaceEditor = WorkspaceEditorWeak.Pin())
 					{
-						WorkspaceEditor->GetTabManager()->TryInvokeTab(CompilerResultsTabName);
+						WorkspaceEditor->GetTabManager()->TryInvokeTab(FTabId(CompilerResultsTabName));
 					}
 				}
 			}
diff --git a/Engine/Plugins/Experimental/AnimNext/Source/AnimNextEditor/Private/Graph/STraitEditorView.h b/Engine/Plugins/Experimental/AnimNext/Source/AnimNextEditor/Private/Graph/STraitEditorView.h
--- a/Engine/Plugins/Experimental/AnimNext/Source/AnimNextEditor/Private/Graph/STraitEditorView.h
+++ b/Engine/Plugins/Experimental/AnimNext/Source/AnimNextEditor/Private/Graph/STraitEditorView.h
@@ -9,7 +9,7 @@
 #include "Graph/TraitEditorDefs.h"
 
 
-class UAnimNextModule_Controller;
+class UAnimNextController;
 class UAnimNextEdGraphNode;
 class IMessageLogListing;
 
diff --git a/Engine/Plugins/Experimental/AnimNext/Source/AnimNextEditor/Private/Graph/TraitEditorTabSummoner.cpp b/Engine/Plugins/Experimental/AnimNext/Source/AnimNextEditor/Private/Graph/TraitEditorTabSummoner.cpp
--- a/Engine/Plugins/Experimental/AnimNext/Source/AnimNextEditor/Private/Graph/TraitEditorTabSummoner.cpp
+++ b/Engine/Plugins/Experimental/AnimNext/Source/AnimNextEditor/Private/Graph/TraitEditorTabSummoner.cpp
@@ -8,8 +8,7 @@
 namespace UE::AnimNext::Editor
 {
 
-const FName TraitEditorTabName("TraitEditorTab");
-
+const FLazyName TraitEditorTabName("TraitEditorTab");
 
 FTraitEditorTabSummoner::FTraitEditorTabSummoner(const TSharedPtr<UE::Workspace::IWorkspaceEditor>& InHostingApp)
 	: FWorkflowTabFactory(TraitEditorTabName, StaticCastSharedPtr<FAssetEditorToolkit>(InHostingApp))
