AnimNext: Module refactor

diff --git a/Engine/Plugins/Experimental/AnimNext/Source/AnimNextEditor/Private/Graph/TraitEditorTabSummoner.h b/Engine/Plugins/Experimental/AnimNext/Source/AnimNextEditor/Private/Graph/TraitEditorTabSummoner.h
--- a/Engine/Plugins/Experimental/AnimNext/Source/AnimNextEditor/Private/Graph/TraitEditorTabSummoner.h
+++ b/Engine/Plugins/Experimental/AnimNext/Source/AnimNextEditor/Private/Graph/TraitEditorTabSummoner.h
@@ -13,7 +13,7 @@ namespace UE::Workspace
 namespace UE::AnimNext::Editor
 {
 
-extern const FName TraitEditorTabName;
+extern const FLazyName TraitEditorTabName;
 
 DECLARE_DELEGATE_OneParam(FOnTraitEditorCreated, TSharedRef<STraitEditorView>);
 
diff --git a/Engine/Plugins/Experimental/AnimNext/Source/AnimNextEditor/Private/Module/AnimNextModuleAssetDefinition.cpp b/Engine/Plugins/Experimental/AnimNext/Source/AnimNextEditor/Private/Module/AnimNextModuleAssetDefinition.cpp
new file mode 100644
--- /dev/null
+++ b/Engine/Plugins/Experimental/AnimNext/Source/AnimNextEditor/Private/Module/AnimNextModuleAssetDefinition.cpp
@@ -0,0 +1,76 @@
+﻿// Copyright Epic Games, Inc. All Rights Reserved.
+
+#include "AnimNextModuleAssetDefinition.h"
+#include "ContentBrowserMenuContexts.h"
+#include "FileHelpers.h"
+#include "IWorkspaceEditorModule.h"
+#include "IWorkspaceEditor.h"
+#include "ToolMenus.h"
+#include "Workspace/AnimNextWorkspaceFactory.h"
+
+#define LOCTEXT_NAMESPACE "AssetDefinition_AnimNextGraph"
+
+EAssetCommandResult UAssetDefinition_AnimNextModule::OpenAssets(const FAssetOpenArgs& OpenArgs) const
+{
+	using namespace UE::AnimNext::Editor;
+	using namespace UE::Workspace;
+
+	for (UAnimNextModule* Asset : OpenArgs.LoadObjects<UAnimNextModule>())
+	{
+		IWorkspaceEditorModule& WorkspaceEditorModule = FModuleManager::Get().LoadModuleChecked<IWorkspaceEditorModule>("WorkspaceEditor");
+		WorkspaceEditorModule.OpenWorkspaceForObject(Asset, EOpenWorkspaceMethod::Default, UAnimNextWorkspaceFactory::StaticClass());
+	}
+
+	return EAssetCommandResult::Handled;
+}
+
+namespace UE::AnimNext::Editor
+{
+
+static FDelayedAutoRegisterHelper AutoRegisterModuleMenuItems(EDelayedRegisterRunPhase::EndOfEngineInit, []
+{
+	UToolMenus::RegisterStartupCallback(FSimpleMulticastDelegate::FDelegate::CreateLambda([]()
+	{
+		FToolMenuOwnerScoped OwnerScoped(UE_MODULE_NAME);
+		UToolMenu* Menu = UE::ContentBrowser::ExtendToolMenu_AssetContextMenu(UAnimNextModule::StaticClass());
+		 
+		FToolMenuSection& Section = Menu->FindOrAddSection("GetAssetActions");
+		Section.AddDynamicEntry(NAME_None, FNewToolMenuSectionDelegate::CreateLambda([](FToolMenuSection& InSection)
+		{
+			InSection.AddMenuEntry(
+				"ForceSave",
+				LOCTEXT("ForceSaveLabel", "Force Save"),
+				LOCTEXT("ForceSaveTooltip", "Force the save of this item and all its subobjects."),
+				FSlateIcon(FAppStyle::Get().GetStyleSetName(), "Icons.Save"),
+				FToolUIAction(FToolMenuExecuteAction::CreateLambda([](const FToolMenuContext& InContext)
+				{
+					if(UContentBrowserAssetContextMenuContext* AssetContext = InContext.FindContext<UContentBrowserAssetContextMenuContext>())
+					{
+						TArray<UPackage*> PackagesToSave;
+						for (const FAssetData& AssetData : AssetContext->SelectedAssets)
+						{
+							if (UPackage* Package = AssetData.GetPackage())
+							{
+								PackagesToSave.Add(Package);
+								PackagesToSave.Append(Package->GetExternalPackages());
+							}
+						}
+
+						if(PackagesToSave.Num() > 0)
+						{
+							FEditorFileUtils::FPromptForCheckoutAndSaveParams SaveParams;
+							SaveParams.bCheckDirty = false;
+							SaveParams.bPromptToSave = false;
+							SaveParams.bIsExplicitSave = true;
+
+							FEditorFileUtils::PromptForCheckoutAndSave(PackagesToSave, SaveParams);
+						}
+					}
+				})));
+		}));
+	}));
+});
+
+}
+
+#undef LOCTEXT_NAMESPACE
\ No newline at end of file
diff --git a/Engine/Plugins/Experimental/AnimNext/Source/AnimNextEditor/Private/Module/AnimNextModuleAssetDefinition.h b/Engine/Plugins/Experimental/AnimNext/Source/AnimNextEditor/Private/Module/AnimNextModuleAssetDefinition.h
new file mode 100644
--- /dev/null
+++ b/Engine/Plugins/Experimental/AnimNext/Source/AnimNextEditor/Private/Module/AnimNextModuleAssetDefinition.h
@@ -0,0 +1,30 @@
+﻿// Copyright Epic Games, Inc. All Rights Reserved.
+
+#pragma once
+
+#include "Module/AnimNextModule.h"
+#include "AssetDefinitionDefault.h"
+#include "AnimNextModuleAssetDefinition.generated.h"
+
+#define LOCTEXT_NAMESPACE "AnimNextAssetDefinitions"
+
+UCLASS()
+class UAssetDefinition_AnimNextModule : public UAssetDefinitionDefault
+{
+	GENERATED_BODY()
+
+public:
+	// UAssetDefinition interface
+	virtual FText GetAssetDisplayName() const override { return LOCTEXT("AnimNextModule", "AnimNext Module"); }
+	virtual FLinearColor GetAssetColor() const override { return FLinearColor(FColor(128,96,48)); }
+	virtual TSoftClassPtr<UObject> GetAssetClass() const override { return UAnimNextModule::StaticClass(); }
+	virtual EAssetCommandResult OpenAssets(const FAssetOpenArgs& OpenArgs) const override;
+	virtual TConstArrayView<FAssetCategoryPath> GetAssetCategories() const override
+	{
+		static const auto Categories = { FAssetCategoryPath(EAssetCategoryPaths::Animation, LOCTEXT("AnimNextSubMenu", "AnimNext")) };
+		return Categories;
+	}
+	virtual bool ShouldSaveExternalPackages() const override { return true; }
+};
+
+#undef LOCTEXT_NAMESPACE
\ No newline at end of file
diff --git a/Engine/Plugins/Experimental/AnimNext/Source/AnimNextEditor/Private/Module/AnimNextModuleFactory.cpp b/Engine/Plugins/Experimental/AnimNext/Source/AnimNextEditor/Private/Module/AnimNextModuleFactory.cpp
new file mode 100644
--- /dev/null
+++ b/Engine/Plugins/Experimental/AnimNext/Source/AnimNextEditor/Private/Module/AnimNextModuleFactory.cpp
@@ -0,0 +1,43 @@
+﻿// Copyright Epic Games, Inc. All Rights Reserved.
+
+#include "Module/AnimNextModuleFactory.h"
+#include "Module/AnimNextModule.h"
+#include "Module/AnimNextModule_EditorData.h"
+#include "UncookedOnlyUtils.h"
+#include "Module/RigUnit_AnimNextModuleEvents.h"
+
+UAnimNextModuleFactory::UAnimNextModuleFactory()
+{
+	bCreateNew = true;
+	bEditAfterNew = true;
+	SupportedClass = UAnimNextModule::StaticClass();
+}
+
+bool UAnimNextModuleFactory::ConfigureProperties()
+{
+	return true;
+}
+
+UObject* UAnimNextModuleFactory::FactoryCreateNew(UClass* Class, UObject* InParent, FName Name, EObjectFlags Flags, UObject* Context, FFeedbackContext* Warn, FName CallingContext)
+{
+	EObjectFlags FlagsToUse = Flags | RF_Public | RF_Standalone | RF_Transactional | RF_LoadCompleted;
+	if(InParent == GetTransientPackage())
+	{
+		FlagsToUse &= ~RF_Standalone;
+	}
+
+	UAnimNextModule* NewModule = NewObject<UAnimNextModule>(InParent, Class, Name, FlagsToUse);
+
+	// Create internal editor data
+	UAnimNextModule_EditorData* EditorData = NewObject<UAnimNextModule_EditorData>(NewModule, TEXT("EditorData"), RF_Transactional);
+	NewModule->EditorData = EditorData;
+	EditorData->Initialize(/*bRecompileVM*/false);
+
+	EditorData->AddEventGraph(FRigUnit_AnimNextPrePhysicsEvent::EventName, FRigUnit_AnimNextPrePhysicsEvent::StaticStruct());
+
+	// Compile the initial skeleton
+	EditorData->RecompileVM();
+	check(!EditorData->bErrorsDuringCompilation);
+
+	return NewModule;
+}
\ No newline at end of file
diff --git a/Engine/Plugins/Experimental/AnimNext/Source/AnimNextEditor/Private/Param/ObjectCastLocatorEditor.cpp b/Engine/Plugins/Experimental/AnimNext/Source/AnimNextEditor/Private/Param/ObjectCastLocatorEditor.cpp
--- a/Engine/Plugins/Experimental/AnimNext/Source/AnimNextEditor/Private/Param/ObjectCastLocatorEditor.cpp
+++ b/Engine/Plugins/Experimental/AnimNext/Source/AnimNextEditor/Private/Param/ObjectCastLocatorEditor.cpp
@@ -12,9 +12,9 @@
 #include "Param/AnimNextObjectFunctionLocatorFragment.h"
 #include "Textures/SlateIcon.h"
 #include "Widgets/Input/SComboButton.h"
-#include "SParameterPicker.h"
 #include "Modules/ModuleManager.h"
 #include "Param/AnimNextObjectCastLocatorFragment.h"
+#include "Widgets/Layout/SBox.h"
 
 #define LOCTEXT_NAMESPACE "ObjectCastLocatorEditor"
 
diff --git a/Engine/Plugins/Experimental/AnimNext/Source/AnimNextEditor/Private/Param/ObjectFunctionLocatorEditor.cpp b/Engine/Plugins/Experimental/AnimNext/Source/AnimNextEditor/Private/Param/ObjectFunctionLocatorEditor.cpp
--- a/Engine/Plugins/Experimental/AnimNext/Source/AnimNextEditor/Private/Param/ObjectFunctionLocatorEditor.cpp
+++ b/Engine/Plugins/Experimental/AnimNext/Source/AnimNextEditor/Private/Param/ObjectFunctionLocatorEditor.cpp
@@ -12,11 +12,11 @@
 #include "Param/AnimNextObjectFunctionLocatorFragment.h"
 #include "Textures/SlateIcon.h"
 #include "Widgets/Input/SComboButton.h"
-#include "SParameterPicker.h"
 #include "Modules/ModuleManager.h"
 #include "Framework/PropertyViewer/IFieldIterator.h"
 #include "Widgets/PropertyViewer/SPropertyViewer.h"
 #include "Param/ParamUtils.h"
+#include "Widgets/Layout/SBox.h"
 
 #define LOCTEXT_NAMESPACE "ObjectFunctionLocatorEditor"
 
diff --git a/Engine/Plugins/Experimental/AnimNext/Source/AnimNextEditor/Private/Param/ObjectPropertyLocatorEditor.cpp b/Engine/Plugins/Experimental/AnimNext/Source/AnimNextEditor/Private/Param/ObjectPropertyLocatorEditor.cpp
--- a/Engine/Plugins/Experimental/AnimNext/Source/AnimNextEditor/Private/Param/ObjectPropertyLocatorEditor.cpp
+++ b/Engine/Plugins/Experimental/AnimNext/Source/AnimNextEditor/Private/Param/ObjectPropertyLocatorEditor.cpp
@@ -5,10 +5,10 @@
 #include "ClassViewerModule.h"
 #include "IAnimNextModuleInterface.h"
 #include "IUniversalObjectLocatorCustomization.h"
-#include "SParameterPicker.h"
 #include "UniversalObjectLocator.h"
 #include "UniversalObjectLocatorFragmentTypeHandle.h"
 #include "StructUtils/UserDefinedStruct.h"
+#include "Framework/PropertyViewer/IFieldExpander.h"
 #include "Framework/PropertyViewer/IFieldIterator.h"
 #include "Modules/ModuleManager.h"
 #include "Param/AnimNextObjectPropertyLocatorFragment.h"
@@ -16,6 +16,7 @@
 #include "Widgets/SCompoundWidget.h"
 #include "Widgets/Input/SComboButton.h"
 #include "Widgets/Input/SEditableTextBox.h"
+#include "Widgets/Layout/SBox.h"
 #include "Widgets/PropertyViewer/SPropertyViewer.h"
 
 #define LOCTEXT_NAMESPACE "ObjectPropertyLocatorEditor"
diff --git a/Engine/Plugins/Experimental/AnimNext/Source/AnimNextEditor/Private/Variables/AddVariableDialogMenuContext.h b/Engine/Plugins/Experimental/AnimNext/Source/AnimNextEditor/Private/Variables/AddVariableDialogMenuContext.h
new file mode 100644
--- /dev/null
+++ b/Engine/Plugins/Experimental/AnimNext/Source/AnimNextEditor/Private/Variables/AddVariableDialogMenuContext.h
@@ -0,0 +1,25 @@
+﻿// Copyright Epic Games, Inc. All Rights Reserved.
+
+#pragma once
+
+#include "AddVariableDialogMenuContext.generated.h"
+
+namespace UE::AnimNext::Editor
+{
+	class SAddVariablesDialog;
+	struct FVariableToAdd;
+}
+
+UCLASS()
+class UAddVariableDialogMenuContext : public UObject
+{
+	GENERATED_BODY()
+
+	friend class UE::AnimNext::Editor::SAddVariablesDialog;
+
+	// The add variable dialog that we are editing
+	TWeakPtr<UE::AnimNext::Editor::SAddVariablesDialog> AddVariablesDialog;
+
+	// The entry that we are editing
+	TWeakPtr<UE::AnimNext::Editor::FVariableToAdd> Entry;
+};
\ No newline at end of file
diff --git a/Engine/Plugins/Experimental/AnimNext/Source/AnimNextEditor/Private/Variables/Outliner/VariablesOutlinerAssetItem.cpp b/Engine/Plugins/Experimental/AnimNext/Source/AnimNextEditor/Private/Variables/Outliner/VariablesOutlinerAssetItem.cpp
new file mode 100644
--- /dev/null
+++ b/Engine/Plugins/Experimental/AnimNext/Source/AnimNextEditor/Private/Variables/Outliner/VariablesOutlinerAssetItem.cpp
@@ -0,0 +1,316 @@
+﻿// Copyright Epic Games, Inc. All Rights Reserved.
+
+#include "VariablesOutlinerAssetItem.h"
+
+#include "ISceneOutliner.h"
+#include "Styling/SlateColor.h"
+#include "Misc/Attribute.h"
+#include "Misc/Optional.h"
+#include "SceneOutlinerStandaloneTypes.h"
+#include "AnimNextEditorModule.h"
+#include "AnimNextRigVMAsset.h"
+#include "AnimNextRigVMAssetEditorData.h"
+#include "AssetToolsModule.h"
+#include "AssetViewUtils.h"
+#include "EditorUtils.h"
+#include "IAssetTools.h"
+#include "ScopedTransaction.h"
+#include "UncookedOnlyUtils.h"
+#include "Styling/SlateIconFinder.h"
+#include "Styling/StarshipCoreStyle.h"
+#include "UObject/Package.h"
+#include "Widgets/SCompoundWidget.h"
+#include "Widgets/DeclarativeSyntaxSupport.h"
+#include "Widgets/Images/SImage.h"
+#include "Widgets/Images/SThrobber.h"
+#include "Widgets/Text/SInlineEditableTextBlock.h"
+#include "Widgets/Text/STextBlock.h"
+#include "Widgets/Views/STableRow.h"
+#include "Widgets/Views/SListView.h"
+
+#define LOCTEXT_NAMESPACE "VariablesOutlinerTreeItem"
+
+namespace UE::AnimNext::Editor
+{
+
+const FSceneOutlinerTreeItemType FVariablesOutlinerAssetItem::Type;
+
+class SVariablesOutlinerAssetlabel : FSceneOutlinerCommonLabelData, public SCompoundWidget
+{
+	SLATE_BEGIN_ARGS(SVariablesOutlinerAssetlabel) {}
+	SLATE_END_ARGS()
+
+	void Construct(const FArguments& InArgs, FVariablesOutlinerAssetItem& InTreeItem, ISceneOutliner& SceneOutliner, const STableRow<FSceneOutlinerTreeItemPtr>& InRow)
+	{
+		WeakSceneOutliner = StaticCastSharedRef<ISceneOutliner>(SceneOutliner.AsShared());
+		TreeItem = StaticCastSharedRef<FVariablesOutlinerAssetItem>(InTreeItem.AsShared());
+
+		ChildSlot
+		[
+			SNew(SHorizontalBox)
+			+SHorizontalBox::Slot()
+			.AutoWidth()
+			.VAlign(VAlign_Center)
+			.Padding(2.0f, 2.0f)
+			[
+				SNew(SImage)
+				.Image(this, &SVariablesOutlinerAssetlabel::GetAssetIcon)
+			]
+			+SHorizontalBox::Slot()
+			.AutoWidth()
+			.VAlign(VAlign_Center)
+			.Padding(0.0f, 2.0f)
+			[
+				SAssignNew(TextBlock, SInlineEditableTextBlock)
+				.Font(FStyleFonts::Get().NormalBold)
+				.Text(this, &SVariablesOutlinerAssetlabel::GetDisplayText)
+				.HighlightText(SceneOutliner.GetFilterHighlightText())
+				.ColorAndOpacity(this, &SVariablesOutlinerAssetlabel::GetForegroundColor)
+				.OnTextCommitted(this, &SVariablesOutlinerAssetlabel::OnTextCommited)
+				.OnVerifyTextChanged(this, &SVariablesOutlinerAssetlabel::OnVerifyTextChanged)
+				.IsReadOnly(this, &SVariablesOutlinerAssetlabel::CanRename)
+			]
+			+SHorizontalBox::Slot()
+			.VAlign(VAlign_Center)
+			.Padding(2.0f, 0.0f, 2.0f, 3.0f)
+			.AutoWidth()
+			[
+				SNew(SImage)
+				.Visibility(this, &SVariablesOutlinerAssetlabel::GetDirtyImageVisibility)
+				.ToolTipText(this, &SVariablesOutlinerAssetlabel::GetDirtyTooltipText)
+				.Image(this, &SVariablesOutlinerAssetlabel::GetDirtyImageBrush)
+			]
+			+SHorizontalBox::Slot()
+			.VAlign(VAlign_Center)
+			.Padding(2.0f, 0.0f)
+			.AutoWidth()
+			[
+				SNew(SThrobber)
+				.Visibility(this, &SVariablesOutlinerAssetlabel::GetLoadingIndicatorVisibility)
+				.ToolTipText(LOCTEXT("LoadingTooltip", "Asset is loading..."))
+			]
+		];
+	}
+
+	bool CanRename() const
+	{
+		if (const TSharedPtr<FVariablesOutlinerAssetItem> Item = TreeItem.Pin())
+		{
+			return Item->SoftAsset.Get() != nullptr;
+		}
+		return false;
+	}
+
+	FText GetDirtyTooltipText() const
+	{
+		if (const TSharedPtr<FVariablesOutlinerAssetItem> Item = TreeItem.Pin())
+		{
+			FTextBuilder TextBuilder;
+			TextBuilder.AppendLine(LOCTEXT("ModifiedTooltip", "Modified"));
+
+			if(UAnimNextRigVMAsset* Asset = Item->SoftAsset.Get())
+			{
+				const UPackage* Package = Asset->GetPackage();
+				check(Package);
+				if(Package->IsDirty())
+				{
+					TextBuilder.AppendLine(FText::FromName(Package->GetFName()));
+				}
+			}
+
+			return TextBuilder.ToText();
+		}
+		return FText::GetEmpty();
+	}
+
+	const FSlateBrush* GetDirtyImageBrush() const
+	{
+		if (const TSharedPtr<FVariablesOutlinerAssetItem> Item = TreeItem.Pin())
+		{
+			if(UAnimNextRigVMAsset* Asset = Item->SoftAsset.Get())
+			{
+				bool bIsDirty = false;
+				const UPackage* Package = Asset->GetPackage();
+				check(Package);
+				if(Package->IsDirty())
+				{
+					bIsDirty = true;
+				}
+
+				return bIsDirty ? FAppStyle::GetBrush("Icons.DirtyBadge") : nullptr;
+			}
+		}
+		return nullptr;
+	}
+
+	EVisibility GetDirtyImageVisibility() const
+	{
+		if (const TSharedPtr<FVariablesOutlinerAssetItem> Item = TreeItem.Pin())
+		{
+			if(UAnimNextRigVMAsset* Asset = Item->SoftAsset.Get())
+			{
+				bool bIsDirty = false;
+				const UPackage* Package = Asset->GetPackage();
+				check(Package);
+				if(Package->IsDirty())
+				{
+					bIsDirty = true;
+				}
+
+				return bIsDirty ? EVisibility::Visible : EVisibility::Collapsed;
+			}
+		}
+		return EVisibility::Collapsed;
+	}
+
+	EVisibility GetLoadingIndicatorVisibility() const
+	{
+		if (const TSharedPtr<FVariablesOutlinerAssetItem> Item = TreeItem.Pin())
+		{
+			UAnimNextRigVMAsset* Asset = Item->SoftAsset.Get();
+			return Asset != nullptr ? EVisibility::Collapsed : EVisibility::Visible;
+		}
+		return EVisibility::Collapsed;
+	}
+	
+	FText GetDisplayText() const
+	{
+		if (const TSharedPtr<FVariablesOutlinerAssetItem> Item = TreeItem.Pin())
+		{
+			return FText::FromString(Item->GetDisplayString());
+		}
+		return FText();
+	}
+
+	const FSlateBrush* GetAssetIcon() const
+	{
+		if (const TSharedPtr<FVariablesOutlinerAssetItem> Item = TreeItem.Pin())
+		{
+			UAnimNextRigVMAsset* Asset = Item->SoftAsset.Get();
+			return Asset != nullptr ? FSlateIconFinder::FindIconBrushForClass(Asset->GetClass()) : FAppStyle::Get().GetBrush("ClassIcon.Object");
+		}
+		return FAppStyle::Get().GetBrush("ClassIcon.Object");
+	}
+
+	void OnTextCommited(const FText& InLabel, ETextCommit::Type InCommitInfo)
+	{
+		if(InCommitInfo == ETextCommit::OnEnter)
+		{
+			if (const TSharedPtr<FVariablesOutlinerAssetItem> Item = TreeItem.Pin())
+			{
+				Item->Rename(InLabel); 
+			}
+		}
+	}
+
+	bool OnVerifyTextChanged(const FText& InLabel, FText& OutErrorMessage)
+	{
+		if (const TSharedPtr<FVariablesOutlinerAssetItem> Item = TreeItem.Pin())
+		{
+			return Item->ValidateName(InLabel, OutErrorMessage); 
+		}
+		return false;
+	}
+
+	virtual FSlateColor GetForegroundColor() const override
+	{
+		const TOptional<FLinearColor> BaseColor = FSceneOutlinerCommonLabelData::GetForegroundColor(*TreeItem.Pin());
+		return BaseColor.IsSet() ? BaseColor.GetValue() : FSlateColor::UseForeground();
+	}
+
+	TWeakPtr<FVariablesOutlinerAssetItem> TreeItem;
+	TSharedPtr<SInlineEditableTextBlock> TextBlock;
+};
+
+FVariablesOutlinerAssetItem::FVariablesOutlinerAssetItem(TSoftObjectPtr<UAnimNextRigVMAsset> InAsset)
+	: ISceneOutlinerTreeItem(FVariablesOutlinerAssetItem::Type)
+	, SoftAsset(InAsset)
+{
+}
+
+bool FVariablesOutlinerAssetItem::IsValid() const
+{
+	return !SoftAsset.IsNull();
+}
+
+FSceneOutlinerTreeItemID FVariablesOutlinerAssetItem::GetID() const
+{
+	return GetTypeHash(SoftAsset);
+}
+
+FString FVariablesOutlinerAssetItem::GetDisplayString() const
+{
+	return SoftAsset.GetAssetName();
+}
+
+TSharedRef<SWidget> FVariablesOutlinerAssetItem::GenerateLabelWidget(ISceneOutliner& Outliner, const STableRow<FSceneOutlinerTreeItemPtr>& InRow)
+{
+	TSharedRef<SVariablesOutlinerAssetlabel> LabelWidget = SNew(SVariablesOutlinerAssetlabel, *this, Outliner, InRow);
+	RenameRequestEvent.BindSP(LabelWidget->TextBlock.Get(), &SInlineEditableTextBlock::EnterEditingMode);
+	return LabelWidget;
+}
+
+FString FVariablesOutlinerAssetItem::GetPackageName() const
+{
+	return SoftAsset.GetLongPackageName();
+}
+
+void FVariablesOutlinerAssetItem::Rename(const FText& InNewName)
+{
+	UAnimNextRigVMAsset* Asset = SoftAsset.Get();
+	if(Asset == nullptr)
+	{
+		return;
+	}
+
+	FAssetToolsModule& AssetToolsModule = FModuleManager::LoadModuleChecked<FAssetToolsModule>("AssetTools");
+	const FString CurrentAssetPath = FPackageName::GetLongPackagePath(Asset->GetPackage()->GetName());
+	TArray<FAssetRenameData> AssetsToRename = { FAssetRenameData(Asset, CurrentAssetPath, InNewName.ToString()) };
+	AssetToolsModule.Get().RenameAssets(AssetsToRename);
+}
+
+bool FVariablesOutlinerAssetItem::ValidateName(const FText& InNewName, FText& OutErrorMessage) const
+{
+	UAnimNextRigVMAsset* Asset = SoftAsset.Get();
+	if(Asset == nullptr)
+	{
+		OutErrorMessage = LOCTEXT("InvalidAssetError", "Asset is invalid");
+		return false;
+	}
+
+	FString NewName = InNewName.ToString();
+	if (NewName.Len() >= NAME_SIZE)
+	{
+		OutErrorMessage = LOCTEXT("Error_AssetNameTooLarge", "This asset name is too long. Please choose a shorter name.");
+		return false;
+	}
+
+	if (Asset->GetFName() != FName(*NewName)) // Deliberately ignore case here to allow case-only renames of existing assets
+	{
+		const FString PackageName = Asset->GetPackage()->GetPathName() / NewName;
+		const FString ObjectPath = FString::Printf(TEXT("%s.%s"), *PackageName, *NewName);
+
+		FText ValidationErrorMsg;
+		if (!AssetViewUtils::IsValidObjectPathForCreate(ObjectPath, ValidationErrorMsg))
+		{
+			OutErrorMessage = ValidationErrorMsg;
+			return false;
+		}
+	}
+
+	FAssetToolsModule& AssetToolsModule = FAssetToolsModule::GetModule();
+	TWeakPtr<IAssetTypeActions> WeakAssetTypeActions = AssetToolsModule.Get().GetAssetTypeActionsForClass(UAnimNextRigVMAsset::StaticClass());
+	if (TSharedPtr<IAssetTypeActions> AssetTypeActions = WeakAssetTypeActions.Pin())
+	{
+		if (!AssetTypeActions->CanRename(FAssetData(Asset), &OutErrorMessage))
+		{
+			return false;
+		}
+	}
+
+	return true;
+}
+
+}
+
+#undef LOCTEXT_NAMESPACE // "VariablesOutlinerTreeItem"
\ No newline at end of file
diff --git a/Engine/Plugins/Experimental/AnimNext/Source/AnimNextEditor/Private/Variables/Outliner/VariablesOutlinerColumns.cpp b/Engine/Plugins/Experimental/AnimNext/Source/AnimNextEditor/Private/Variables/Outliner/VariablesOutlinerColumns.cpp
new file mode 100644
--- /dev/null
+++ b/Engine/Plugins/Experimental/AnimNext/Source/AnimNextEditor/Private/Variables/Outliner/VariablesOutlinerColumns.cpp
@@ -0,0 +1,363 @@
+// Copyright Epic Games, Inc. All Rights Reserved.
+
+#include "VariablesOutlinerColumns.h"
+
+#include "AnimNextEditorModule.h"
+#include "AnimNextRigVMAssetEditorData.h"
+#include "DetailLayoutBuilder.h"
+#include "EditorUtils.h"
+#include "IAnimNextRigVMExportInterface.h"
+#include "Variables/IAnimNextRigVMVariableInterface.h"
+#include "InstancedPropertyBagStructureDataProvider.h"
+#include "ISinglePropertyView.h"
+#include "PropertyBagDetails.h"
+#include "ScopedTransaction.h"
+#include "SPinTypeSelector.h"
+#include "UncookedOnlyUtils.h"
+#include "VariablesOutlinerEntryItem.h"
+#include "Entries/AnimNextRigVMAssetEntry.h"
+#include "Widgets/Images/SImage.h"
+#include "UObject/Package.h"
+#include "Widgets/Input/SButton.h"
+#include "Widgets/Text/SInlineEditableTextBlock.h"
+
+#define LOCTEXT_NAMESPACE "VariablesOutlinerColumns"
+
+namespace UE::AnimNext::Editor
+{
+
+FLazyName VariablesOutlinerType("Type");
+
+FName FVariablesOutlinerTypeColumn::GetID()
+{
+	return VariablesOutlinerType;
+}
+
+SHeaderRow::FColumn::FArguments FVariablesOutlinerTypeColumn::ConstructHeaderRowColumn()
+{
+	return SHeaderRow::Column(GetColumnID())
+		.FixedWidth(24.f)
+		.HAlignHeader(HAlign_Center)
+		.VAlignHeader(VAlign_Center)
+		.HAlignCell(HAlign_Center)
+		.VAlignCell(VAlign_Center)
+		[
+			SNew(SBox)
+			.WidthOverride(16.0f)
+			.HeightOverride(16.0f)
+			.VAlign(VAlign_Center)
+			.HAlign(HAlign_Center)
+			[
+				SNew(SImage)
+				.ColorAndOpacity(FSlateColor::UseForeground())
+				.Image(FAppStyle::GetBrush(TEXT("Kismet.VariableList.TypeIcon")))
+				.ToolTipText(LOCTEXT("TypeTooltip", "Type of this entry"))
+			]
+		];
+}
+
+const TSharedRef<SWidget> FVariablesOutlinerTypeColumn::ConstructRowWidget(FSceneOutlinerTreeItemRef Item, const STableRow<FSceneOutlinerTreeItemPtr>& Row)
+{
+	const FVariablesOutlinerEntryItem* TreeItem = Item->CastTo<FVariablesOutlinerEntryItem>();
+	if (TreeItem == nullptr)
+	{
+		return SNullWidget::NullWidget;
+	}
+
+	return
+		SNew(SBox)
+		.HAlign(HAlign_Left)
+		.VAlign(VAlign_Center)
+		[
+			SNew(SPinTypeSelector, FGetPinTypeTree::CreateStatic(&Editor::FUtils::GetFilteredVariableTypeTree))
+			.TargetPinType_Lambda([WeakEntry = TreeItem->WeakEntry]()
+			{
+				if(const IAnimNextRigVMVariableInterface* Variable = Cast<IAnimNextRigVMVariableInterface>(WeakEntry.Get()))
+				{
+					return UncookedOnly::FUtils::GetPinTypeFromParamType(Variable->GetType());
+				}
+
+				return FEdGraphPinType();
+			})
+			.OnPinTypeChanged_Lambda([WeakEntry = TreeItem->WeakEntry](const FEdGraphPinType& PinType)
+			{
+				if(IAnimNextRigVMVariableInterface* Variable = Cast<IAnimNextRigVMVariableInterface>(WeakEntry.Get()))
+				{
+					const FAnimNextParamType ParamType = UncookedOnly::FUtils::GetParamTypeFromPinType(PinType);
+					if(ParamType.IsValid())
+					{
+						FScopedTransaction Transaction(LOCTEXT("SetTypeTransaction", "Set Variable Type"));
+						Variable->SetType(ParamType);
+					}
+				}
+			})
+			.Schema(GetDefault<UPropertyBagSchema>())
+			.bAllowArrays(true)
+			.TypeTreeFilter(ETypeTreeFilter::None)
+			.Font(IDetailLayoutBuilder::GetDetailFont())
+			.SelectorType(SPinTypeSelector::ESelectorType::Compact)
+		];
+}
+
+FLazyName VariablesOutlinerValue("Value");
+
+FName FVariablesOutlinerValueColumn::GetID()
+{
+	return VariablesOutlinerValue;
+}
+
+SHeaderRow::FColumn::FArguments FVariablesOutlinerValueColumn::ConstructHeaderRowColumn()
+{
+	return
+		SHeaderRow::Column(GetColumnID())
+		.FillWidth(1.0f)
+		.HAlignHeader(HAlign_Left)
+		.VAlignHeader(VAlign_Center)
+		.HAlignCell(HAlign_Left)
+		.VAlignCell(VAlign_Center)
+		[
+			SNew(SBox) 
+			.VAlign(VAlign_Center)
+			.HAlign(HAlign_Left)
+			[
+				SNew(STextBlock)
+				.Text(LOCTEXT("ValueLabel", "Value"))
+				.ToolTipText(LOCTEXT("ValueTooltip", "Value of the variable"))
+			]
+		];
+}
+
+class SVariablesOutlinerValue : public SCompoundWidget, public FNotifyHook
+{
+	SLATE_BEGIN_ARGS(SVariablesOutlinerValue) {}
+	SLATE_END_ARGS()
+
+	void Construct(const FArguments& InArgs, FVariablesOutlinerEntryItem& InTreeItem, ISceneOutliner& SceneOutliner, const STableRow<FSceneOutlinerTreeItemPtr>& InRow)
+	{
+		WeakTreeItem = StaticCastSharedRef<FVariablesOutlinerEntryItem>(InTreeItem.AsShared());
+
+		TSharedPtr<SWidget> ValueWidget = SNullWidget::NullWidget;
+		IAnimNextRigVMVariableInterface* Variable = CastChecked<IAnimNextRigVMVariableInterface>(InTreeItem.WeakEntry.Get());
+		FInstancedPropertyBag& PropertyBag = Variable->GetMutablePropertyBag();
+		const FPropertyBagPropertyDesc* PropertyDesc = PropertyBag.FindPropertyDescByName(IAnimNextRigVMVariableInterface::ValueName);
+		if (PropertyDesc != nullptr)
+		{
+			if (PropertyDesc->ContainerTypes.IsEmpty()) // avoid trying to inline containers
+			{
+				FSinglePropertyParams SinglePropertyArgs;
+				SinglePropertyArgs.NamePlacement = EPropertyNamePlacement::Hidden;
+				SinglePropertyArgs.NotifyHook = this;
+
+				FPropertyEditorModule& PropertyEditorModule = FModuleManager::Get().LoadModuleChecked<FPropertyEditorModule>("PropertyEditor");
+				
+				const TSharedPtr<ISinglePropertyView> SingleStructPropertyView = PropertyEditorModule.CreateSingleProperty(MakeShared<FInstancePropertyBagStructureDataProvider>(PropertyBag), IAnimNextRigVMVariableInterface::ValueName, SinglePropertyArgs);
+				if (SingleStructPropertyView.IsValid())
+				{
+					ValueWidget = SingleStructPropertyView;
+				}
+			}
+		}
+
+		ChildSlot
+		[
+			SNew(SHorizontalBox)
+			+SHorizontalBox::Slot()
+			.HAlign(HAlign_Center)
+			.AutoWidth()
+			[
+				ValueWidget.ToSharedRef()
+			]
+		];
+	}
+
+	// FNotifyHook interface
+	virtual void NotifyPostChange(const FPropertyChangedEvent& PropertyChangedEvent, FProperty* PropertyThatChanged) override
+	{
+		TSharedPtr<FVariablesOutlinerEntryItem> TreeItem = WeakTreeItem.Pin();
+		if(!TreeItem.IsValid())
+		{
+			return;
+		}
+		
+		UAnimNextRigVMAssetEntry* Entry = CastChecked<UAnimNextRigVMAssetEntry>(TreeItem->WeakEntry.Get());
+		if(Entry == nullptr)
+		{
+			return;
+		}
+
+		UAnimNextRigVMAssetEditorData* EditorData = Entry->GetTypedOuter<UAnimNextRigVMAssetEditorData>();
+		if (EditorData == nullptr)
+		{
+			return;
+		}
+
+		// Needed to show the changed sign in the table when we modify the PropertyBag
+		Entry->MarkPackageDirty();
+
+		// Ensure that default values get picked up and forwarded to compiler
+		EditorData->BroadcastModified(EAnimNextEditorDataNotifType::PropertyChanged, Entry);
+	}
+
+	TWeakPtr<FVariablesOutlinerEntryItem> WeakTreeItem;
+};
+
+const TSharedRef<SWidget> FVariablesOutlinerValueColumn::ConstructRowWidget(FSceneOutlinerTreeItemRef Item, const STableRow<FSceneOutlinerTreeItemPtr>& Row)
+{
+	FVariablesOutlinerEntryItem* TreeItem = Item->CastTo<FVariablesOutlinerEntryItem>();
+	if (TreeItem == nullptr)
+	{
+		return SNullWidget::NullWidget;
+	}
+
+	IAnimNextRigVMVariableInterface* Variable = Cast<IAnimNextRigVMVariableInterface>(TreeItem->WeakEntry.Get());
+	if(Variable == nullptr)
+	{
+		return SNullWidget::NullWidget;
+	}
+	
+	return SNew(SVariablesOutlinerValue, *TreeItem, WeakSceneOutliner.Pin().ToSharedRef().Get(), Row);
+}
+
+FLazyName VariablesOutlinerAccessSpecifier("AccessSpecifier");
+
+FName FVariablesOutlinerAccessSpecifierColumn::GetID()
+{
+	return VariablesOutlinerAccessSpecifier;
+}
+
+SHeaderRow::FColumn::FArguments FVariablesOutlinerAccessSpecifierColumn::ConstructHeaderRowColumn()
+{
+	return
+		SHeaderRow::Column(GetColumnID())
+		.FixedWidth(24.0f)
+		.HAlignHeader(HAlign_Center)
+		.VAlignHeader(VAlign_Center)
+		.HAlignCell(HAlign_Fill)
+		.VAlignCell(VAlign_Fill)
+		[
+			SNew(SBox) 
+			.VAlign(VAlign_Center)
+			.HAlign(HAlign_Center)
+			[
+				SNew(SImage)
+				.ColorAndOpacity(FSlateColor::UseForeground())
+				.Image(FAppStyle::GetBrush("Level.VisibleIcon16x"))
+				.ToolTipText(LOCTEXT("AccessSpecifierAccessLevelTooltip", "Access level of this entry"))
+			]
+		];
+}
+
+class SVariablesOutlinerAccessSpecifier : public SCompoundWidget, public FNotifyHook
+{
+	SLATE_BEGIN_ARGS(SVariablesOutlinerAccessSpecifier) {}
+	SLATE_END_ARGS()
+
+	void Construct(const FArguments& InArgs, FVariablesOutlinerEntryItem& InTreeItem, ISceneOutliner& SceneOutliner, const STableRow<FSceneOutlinerTreeItemPtr>& InRow)
+	{
+		WeakTreeItem = StaticCastSharedRef<FVariablesOutlinerEntryItem>(InTreeItem.AsShared());
+
+		ChildSlot
+		[
+			SNew(SBox)
+			.WidthOverride(16.0f)
+			.HeightOverride(16.0f)
+			.VAlign(VAlign_Center)
+			.HAlign(HAlign_Center)
+			[
+				SNew(SButton)
+				.ButtonStyle(FAppStyle::Get(), "HoverHintOnly")
+				.OnClicked(this, &SVariablesOutlinerAccessSpecifier::OnClicked)
+				.Content()
+				[
+					SNew(SImage)
+					.ColorAndOpacity(FSlateColor::UseForeground())
+					.Image(this, &SVariablesOutlinerAccessSpecifier::GetImage)
+					.ToolTipText(this, &SVariablesOutlinerAccessSpecifier::GetTooltipText)
+				]
+			]
+		];
+	}
+
+	FReply OnClicked()
+	{
+		TSharedPtr<FVariablesOutlinerEntryItem> TreeItem = WeakTreeItem.Pin();
+		if(!TreeItem.IsValid())
+		{
+			return FReply::Unhandled();
+		}
+		
+		IAnimNextRigVMExportInterface* Export = Cast<IAnimNextRigVMExportInterface>(TreeItem->WeakEntry.Get());
+		if(Export == nullptr)
+		{
+			return FReply::Unhandled();
+		}
+
+		FScopedTransaction Transaction(LOCTEXT("SetAccessSpecifierTransaction", "Set Access Specifier"));
+		Export->SetExportAccessSpecifier(Export->GetExportAccessSpecifier() == EAnimNextExportAccessSpecifier::Public ? EAnimNextExportAccessSpecifier::Private : EAnimNextExportAccessSpecifier::Public);  
+
+		return FReply::Unhandled();	// Fall through so we dont deselect our item
+	}
+
+	const FSlateBrush* GetImage() const
+	{
+		TSharedPtr<FVariablesOutlinerEntryItem> TreeItem = WeakTreeItem.Pin();
+		if(!TreeItem.IsValid())
+		{
+			return nullptr;
+		}
+
+		IAnimNextRigVMExportInterface* Export = Cast<IAnimNextRigVMExportInterface>(TreeItem->WeakEntry.Get());
+		if(Export == nullptr)
+		{
+			return nullptr;
+		}
+
+		return Export->GetExportAccessSpecifier() == EAnimNextExportAccessSpecifier::Public ?
+			FAppStyle::GetBrush("Level.VisibleIcon16x") :
+			FAppStyle::GetBrush("Level.NotVisibleHighlightIcon16x");
+	}
+
+	FText GetTooltipText() const
+	{
+		TSharedPtr<FVariablesOutlinerEntryItem> TreeItem = WeakTreeItem.Pin();
+		if(!TreeItem.IsValid())
+		{
+			return FText::GetEmpty();
+		}
+
+		IAnimNextRigVMExportInterface* Export = Cast<IAnimNextRigVMExportInterface>(TreeItem->WeakEntry.Get());
+		if(Export == nullptr)
+		{
+			return FText::GetEmpty();
+		}
+
+		FText AccessSpecifier = Export->GetExportAccessSpecifier() == EAnimNextExportAccessSpecifier::Public ? LOCTEXT("PublicSpecifier", "public") : LOCTEXT("PrivateSpecifier", "private");
+		FText AccessSpecifierDesc = Export->GetExportAccessSpecifier() == EAnimNextExportAccessSpecifier::Public ?
+			LOCTEXT("PublicSpecifierDesc", "This means that the entry is usable from gameplay and from other AnimNext assets") :
+			LOCTEXT("PrivateSpecifierDesc", "This means that the entry is only usable inside this asset");
+		return FText::Format(LOCTEXT("AccessSpecifierEntryTooltip", "This entry is {0}.\n{1}"), AccessSpecifier, AccessSpecifierDesc);
+	}
+
+	TWeakPtr<FVariablesOutlinerEntryItem> WeakTreeItem;
+};
+
+const TSharedRef<SWidget> FVariablesOutlinerAccessSpecifierColumn::ConstructRowWidget(FSceneOutlinerTreeItemRef Item, const STableRow<FSceneOutlinerTreeItemPtr>& Row)
+{
+	FVariablesOutlinerEntryItem* TreeItem = Item->CastTo<FVariablesOutlinerEntryItem>();
+	if (TreeItem == nullptr)
+	{
+		return SNullWidget::NullWidget;
+	}
+
+	IAnimNextRigVMVariableInterface* Variable = Cast<IAnimNextRigVMVariableInterface>(TreeItem->WeakEntry.Get());
+	if(Variable == nullptr)
+	{
+		return SNullWidget::NullWidget;
+	}
+	
+	return SNew(SVariablesOutlinerAccessSpecifier, *TreeItem, WeakSceneOutliner.Pin().ToSharedRef().Get(), Row);
+}
+
+}
+
+#undef LOCTEXT_NAMESPACE // "VariablesOutlinerColumns"
\ No newline at end of file
diff --git a/Engine/Plugins/Experimental/AnimNext/Source/AnimNextEditor/Private/Variables/Outliner/VariablesOutlinerEntryItem.cpp b/Engine/Plugins/Experimental/AnimNext/Source/AnimNextEditor/Private/Variables/Outliner/VariablesOutlinerEntryItem.cpp
new file mode 100644
--- /dev/null
+++ b/Engine/Plugins/Experimental/AnimNext/Source/AnimNextEditor/Private/Variables/Outliner/VariablesOutlinerEntryItem.cpp
@@ -0,0 +1,241 @@
+// Copyright Epic Games, Inc. All Rights Reserved.
+
+#include "VariablesOutlinerEntryItem.h"
+
+#include "ISceneOutliner.h"
+#include "Styling/SlateColor.h"
+#include "Misc/Attribute.h"
+#include "Misc/Optional.h"
+#include "SceneOutlinerStandaloneTypes.h"
+#include "AnimNextEditorModule.h"
+#include "AnimNextRigVMAssetEditorData.h"
+#include "EditorUtils.h"
+#include "Variables/IAnimNextRigVMVariableInterface.h"
+#include "ScopedTransaction.h"
+#include "UncookedOnlyUtils.h"
+#include "Entries/AnimNextRigVMAssetEntry.h"
+#include "UObject/Package.h"
+#include "Widgets/SCompoundWidget.h"
+#include "Widgets/DeclarativeSyntaxSupport.h"
+#include "Widgets/Images/SImage.h"
+#include "Widgets/Text/SInlineEditableTextBlock.h"
+#include "Widgets/Text/STextBlock.h"
+#include "Widgets/Views/STableRow.h"
+#include "Widgets/Views/SListView.h"
+#include "EdGraphSchema_K2.h"
+
+#define LOCTEXT_NAMESPACE "VariablesOutlinerTreeItem"
+
+namespace UE::AnimNext::Editor
+{
+
+const FSceneOutlinerTreeItemType FVariablesOutlinerEntryItem::Type;
+
+class SVariablesOutlinerEntrylabel : FSceneOutlinerCommonLabelData, public SCompoundWidget
+{
+	SLATE_BEGIN_ARGS(SVariablesOutlinerEntrylabel) {}
+	SLATE_END_ARGS()
+
+	void Construct(const FArguments& InArgs, FVariablesOutlinerEntryItem& InTreeItem, ISceneOutliner& SceneOutliner, const STableRow<FSceneOutlinerTreeItemPtr>& InRow)
+	{
+		WeakSceneOutliner = StaticCastSharedRef<ISceneOutliner>(SceneOutliner.AsShared());
+		TreeItem = StaticCastSharedRef<FVariablesOutlinerEntryItem>(InTreeItem.AsShared());
+
+		ChildSlot
+		[
+			SNew(SHorizontalBox)
+			+SHorizontalBox::Slot()
+			.AutoWidth()
+			.VAlign(VAlign_Center)
+			.Padding(0.0f, 2.0f)
+			[
+				SAssignNew(TextBlock, SInlineEditableTextBlock)
+				.Text(this, &SVariablesOutlinerEntrylabel::GetDisplayText)
+				.HighlightText(SceneOutliner.GetFilterHighlightText())
+				.ColorAndOpacity(this, &SVariablesOutlinerEntrylabel::GetForegroundColor)
+				.OnTextCommitted(this, &SVariablesOutlinerEntrylabel::OnTextCommited)
+				.OnVerifyTextChanged(this, &SVariablesOutlinerEntrylabel::OnVerifyTextChanged)
+			]
+			+SHorizontalBox::Slot()
+			.VAlign(VAlign_Center)
+			.Padding(2.0f, 0.0f, 2.0f, 3.0f)
+			.AutoWidth()
+			[
+				SNew(SImage)
+				.ToolTipText(this, &SVariablesOutlinerEntrylabel::GetDirtyTooltipText)
+				.Image(this, &SVariablesOutlinerEntrylabel::GetDirtyImageBrush)
+			]
+		];
+	}
+
+	FText GetDirtyTooltipText() const
+	{
+		if (const TSharedPtr<FVariablesOutlinerEntryItem> Item = TreeItem.Pin())
+		{
+			FTextBuilder TextBuilder;
+			TextBuilder.AppendLine(LOCTEXT("ModifiedTooltip", "Modified"));
+
+			if(UAnimNextRigVMAssetEntry* AssetEntry = Item->WeakEntry.Get())
+			{
+				const UPackage* ExternalPackage = AssetEntry->GetExternalPackage();
+				if(ExternalPackage && ExternalPackage->IsDirty())
+				{
+					TextBuilder.AppendLine(FText::FromName(ExternalPackage->GetFName()));
+				}
+			}
+
+			return TextBuilder.ToText();
+		}
+		return FText::GetEmpty();
+	}
+
+	const FSlateBrush* GetDirtyImageBrush() const
+	{
+		if (const TSharedPtr<FVariablesOutlinerEntryItem> Item = TreeItem.Pin())
+		{
+			if(UAnimNextRigVMAssetEntry* AssetEntry = Item->WeakEntry.Get())
+			{
+				bool bIsDirty = false;
+				const UPackage* ExternalPackage = AssetEntry->GetExternalPackage();
+				if(ExternalPackage && ExternalPackage->IsDirty())
+				{
+					bIsDirty = true;
+				}
+
+				return bIsDirty ? FAppStyle::GetBrush("Icons.DirtyBadge") : nullptr;
+			}
+		}
+		return nullptr;
+	}
+	
+	FText GetDisplayText() const
+	{
+		if (const TSharedPtr<FVariablesOutlinerEntryItem> Item = TreeItem.Pin())
+		{
+			return FText::FromString(Item->GetDisplayString());
+		}
+		return FText();
+	}
+
+	void OnTextCommited(const FText& InLabel, ETextCommit::Type InCommitInfo)
+	{
+		if(InCommitInfo == ETextCommit::OnEnter)
+		{
+			if (const TSharedPtr<FVariablesOutlinerEntryItem> Item = TreeItem.Pin())
+			{
+				Item->Rename(InLabel); 
+			}
+		}
+	}
+
+	bool OnVerifyTextChanged(const FText& InLabel, FText& OutErrorMessage)
+	{
+		if (const TSharedPtr<FVariablesOutlinerEntryItem> Item = TreeItem.Pin())
+		{
+			return Item->ValidateName(InLabel, OutErrorMessage); 
+		}
+		return false;
+	}
+
+	virtual FSlateColor GetForegroundColor() const override
+	{
+		const TOptional<FLinearColor> BaseColor = FSceneOutlinerCommonLabelData::GetForegroundColor(*TreeItem.Pin());
+		return BaseColor.IsSet() ? BaseColor.GetValue() : FSlateColor::UseForeground();
+	}
+
+	TWeakPtr<FVariablesOutlinerEntryItem> TreeItem;
+	TSharedPtr<SInlineEditableTextBlock> TextBlock;
+};
+
+FVariablesOutlinerEntryItem::FVariablesOutlinerEntryItem(UAnimNextRigVMAssetEntry* InEntry)
+	: ISceneOutlinerTreeItem(FVariablesOutlinerEntryItem::Type)
+	, WeakEntry(InEntry)
+{
+}
+
+bool FVariablesOutlinerEntryItem::IsValid() const
+{
+	return WeakEntry.Get() != nullptr;
+}
+
+FSceneOutlinerTreeItemID FVariablesOutlinerEntryItem::GetID() const
+{
+	return GetTypeHash(WeakEntry);
+}
+
+FString FVariablesOutlinerEntryItem::GetDisplayString() const
+{
+	UAnimNextRigVMAssetEntry* Entry = WeakEntry.Get();
+	if(Entry == nullptr)
+	{
+		return FString();
+	}
+
+	return Entry->GetDisplayName().ToString();
+}
+
+TSharedRef<SWidget> FVariablesOutlinerEntryItem::GenerateLabelWidget(ISceneOutliner& Outliner, const STableRow<FSceneOutlinerTreeItemPtr>& InRow)
+{
+	TSharedRef<SVariablesOutlinerEntrylabel> LabelWidget = SNew(SVariablesOutlinerEntrylabel, *this, Outliner, InRow);
+	RenameRequestEvent.BindSP(LabelWidget->TextBlock.Get(), &SInlineEditableTextBlock::EnterEditingMode);
+	return LabelWidget;
+}
+
+FString FVariablesOutlinerEntryItem::GetPackageName() const
+{
+	UAnimNextRigVMAssetEntry* Entry = WeakEntry.Get();
+	if(Entry == nullptr)
+	{
+		return ISceneOutlinerTreeItem::GetPackageName();
+	}
+
+	return Entry->GetPackage()->GetName();
+}
+
+void FVariablesOutlinerEntryItem::Rename(const FText& InNewName)
+{
+	UAnimNextRigVMAssetEntry* Entry = WeakEntry.Get();
+	if(Entry == nullptr)
+	{
+		return;
+	}
+
+	FScopedTransaction Transaction(LOCTEXT("RenameVariable", "Rename variable"));
+	Entry->SetEntryName(FName(*InNewName.ToString()));
+}
+
+bool FVariablesOutlinerEntryItem::ValidateName(const FText& InNewName, FText& OutErrorMessage) const
+{
+	UAnimNextRigVMAssetEntry* Entry = WeakEntry.Get();
+	if(Entry == nullptr)
+	{
+		OutErrorMessage = LOCTEXT("InvalidVariableError", "Variable is invalid");
+		return false;
+	}
+
+	UAnimNextRigVMAssetEditorData* EditorData = Entry->GetTypedOuter<UAnimNextRigVMAssetEditorData>();
+	if(EditorData == nullptr)
+	{
+		OutErrorMessage = LOCTEXT("InvalidEditorDataError", "Variable has invalid editor data");
+		return false;
+	}
+
+	const FString NewString = InNewName.ToString();
+	if (!Editor::FUtils::IsValidParameterNameString(NewString, OutErrorMessage))
+	{
+		return false;
+	}
+
+	FName Name(*NewString);
+	if(UAnimNextRigVMAssetEntry* ExistingEntry = EditorData->FindEntry(Name))
+	{
+		OutErrorMessage = LOCTEXT("NameExistsError", "Variable name already exists in this asset");
+		return false;
+	}
+
+	return true;
+}
+
+}
+
+#undef LOCTEXT_NAMESPACE // "VariablesOutlinerTreeItem"
\ No newline at end of file
diff --git a/Engine/Plugins/Experimental/AnimNext/Source/AnimNextEditor/Private/Variables/Outliner/VariablesOutlinerEntryItem.h b/Engine/Plugins/Experimental/AnimNext/Source/AnimNextEditor/Private/Variables/Outliner/VariablesOutlinerEntryItem.h
new file mode 100644
--- /dev/null
+++ b/Engine/Plugins/Experimental/AnimNext/Source/AnimNextEditor/Private/Variables/Outliner/VariablesOutlinerEntryItem.h
@@ -0,0 +1,39 @@
+// Copyright Epic Games, Inc. All Rights Reserved.
+
+#pragma once
+
+#include "ISceneOutlinerTreeItem.h"
+
+class UAnimNextRigVMAssetEntry;
+
+namespace UE::AnimNext::Editor
+{
+
+class IWorkspaceOutlinerItemDetails;
+
+struct FVariablesOutlinerEntryItem : ISceneOutlinerTreeItem
+{
+	static const FSceneOutlinerTreeItemType Type;
+
+	FVariablesOutlinerEntryItem(UAnimNextRigVMAssetEntry* InEntry);
+
+	// Begin ISceneOutlinerTreeItem overrides
+	virtual bool IsValid() const override;
+	virtual FSceneOutlinerTreeItemID GetID() const override;
+	virtual FString GetDisplayString() const override;
+	virtual bool CanInteract() const override { return true; }
+	virtual TSharedRef<SWidget> GenerateLabelWidget(ISceneOutliner& Outliner, const STableRow<FSceneOutlinerTreeItemPtr>& InRow) override;
+	virtual FString GetPackageName() const override;
+	// End ISceneOutlinerTreeItem overrides
+
+	// Renames the item to the specified name
+	void Rename(const FText& InNewName);
+
+	// Validates the new item name
+	bool ValidateName(const FText& InNewName, FText& OutErrorMessage) const;
+
+	// Ptr to the underlying entry
+	TWeakObjectPtr<UAnimNextRigVMAssetEntry> WeakEntry;
+};
+
+}
diff --git a/Engine/Plugins/Experimental/AnimNext/Source/AnimNextEditor/Private/Variables/Outliner/VariablesOutlinerHierarchy.cpp b/Engine/Plugins/Experimental/AnimNext/Source/AnimNextEditor/Private/Variables/Outliner/VariablesOutlinerHierarchy.cpp
new file mode 100644
--- /dev/null
+++ b/Engine/Plugins/Experimental/AnimNext/Source/AnimNextEditor/Private/Variables/Outliner/VariablesOutlinerHierarchy.cpp
@@ -0,0 +1,88 @@
+// Copyright Epic Games, Inc. All Rights Reserved.
+
+#include "VariablesOutlinerHierarchy.h"
+
+#include "AnimNextRigVMAsset.h"
+#include "ISceneOutlinerMode.h"
+#include "AnimNextRigVMAssetEditorData.h"
+#include "Variables/IAnimNextRigVMVariableInterface.h"
+#include "VariablesOutlinerAssetItem.h"
+#include "VariablesOutlinerMode.h"
+#include "WorkspaceAssetRegistryInfo.h"
+#include "VariablesOutlinerEntryItem.h"
+#include "Entries/AnimNextRigVMAssetEntry.h"
+#include "Variables/SVariablesView.h"
+
+namespace UE::AnimNext::Editor
+{
+
+FVariablesOutlinerHierarchy::FVariablesOutlinerHierarchy(ISceneOutlinerMode* Mode)
+	: ISceneOutlinerHierarchy(Mode)
+{
+}
+
+void FVariablesOutlinerHierarchy::CreateItems(TArray<FSceneOutlinerTreeItemPtr>& OutItems) const
+{
+	SVariablesOutliner* Outliner = static_cast<FVariablesOutlinerMode* const>(Mode)->GetOutliner();
+	for(const TSoftObjectPtr<UAnimNextRigVMAsset>& SoftAsset : Outliner->Assets)
+	{
+		UAnimNextRigVMAsset* Asset = SoftAsset.Get();
+		if(Asset == nullptr)
+		{
+			continue;
+		}
+
+		const UAnimNextRigVMAssetEditorData* EditorData = UncookedOnly::FUtils::GetEditorData<UAnimNextRigVMAssetEditorData>(Asset);
+		if (EditorData == nullptr)
+		{
+			continue;
+		}
+
+		if (FSceneOutlinerTreeItemPtr Item = Mode->CreateItemFor<FVariablesOutlinerAssetItem>(SoftAsset))
+		{
+			OutItems.Add(Item);
+		}
+
+		EditorData->ForEachEntryOfType<IAnimNextRigVMVariableInterface>([this, &OutItems](IAnimNextRigVMVariableInterface* InVariable)
+		{
+			if (FSceneOutlinerTreeItemPtr Item = Mode->CreateItemFor<FVariablesOutlinerEntryItem>(CastChecked<UAnimNextRigVMAssetEntry>(InVariable)))
+			{
+				OutItems.Add(Item);
+			}
+			return true;
+		});
+	}
+}
+
+FSceneOutlinerTreeItemPtr FVariablesOutlinerHierarchy::FindOrCreateParentItem(const ISceneOutlinerTreeItem& Item, const TMap<FSceneOutlinerTreeItemID, FSceneOutlinerTreeItemPtr>& Items, bool bCreate)
+{
+	const FVariablesOutlinerEntryItem* TreeItem = Item.CastTo<FVariablesOutlinerEntryItem>();
+	if (TreeItem == nullptr)
+	{
+		return nullptr;
+	}
+
+	UAnimNextRigVMAssetEntry* Entry = TreeItem->WeakEntry.Get();
+	if (Entry == nullptr)
+	{
+		return nullptr;
+	}
+
+	UAnimNextRigVMAsset* Asset = Entry->GetTypedOuter<UAnimNextRigVMAsset>();
+	if (Entry == nullptr)
+	{
+		return nullptr;
+	}
+
+	TSoftObjectPtr<UAnimNextRigVMAsset> SoftObjectPtr(Asset);
+	uint32 Hash = GetTypeHash(SoftObjectPtr);
+
+	if (const FSceneOutlinerTreeItemPtr* ParentItem = Items.Find(Hash))
+	{
+		return *ParentItem;
+	}
+
+	return nullptr;
+}
+
+}
diff --git a/Engine/Plugins/Experimental/AnimNext/Source/AnimNextEditor/Private/Variables/Outliner/VariablesOutlinerHierarchy.h b/Engine/Plugins/Experimental/AnimNext/Source/AnimNextEditor/Private/Variables/Outliner/VariablesOutlinerHierarchy.h
new file mode 100644
--- /dev/null
+++ b/Engine/Plugins/Experimental/AnimNext/Source/AnimNextEditor/Private/Variables/Outliner/VariablesOutlinerHierarchy.h
@@ -0,0 +1,27 @@
+// Copyright Epic Games, Inc. All Rights Reserved.
+
+#pragma once
+
+#include "CoreMinimal.h"
+#include "ISceneOutlinerHierarchy.h"
+
+class UAnimNextRigVMAssetEditorData;
+
+namespace UE::AnimNext::Editor
+{
+
+class FVariablesOutlinerHierarchy : public ISceneOutlinerHierarchy
+{
+public:
+	FVariablesOutlinerHierarchy(ISceneOutlinerMode* Mode);
+	FVariablesOutlinerHierarchy(const FVariablesOutlinerHierarchy&) = delete;
+	FVariablesOutlinerHierarchy& operator=(const FVariablesOutlinerHierarchy&) = delete;
+
+	// Begin ISceneOutlinerHierarchy overrides
+	virtual void CreateItems(TArray<FSceneOutlinerTreeItemPtr>& OutItems) const override;
+	virtual void CreateChildren(const FSceneOutlinerTreeItemPtr& Item, TArray<FSceneOutlinerTreeItemPtr>& OutChildren) const override {}
+	virtual FSceneOutlinerTreeItemPtr FindOrCreateParentItem(const ISceneOutlinerTreeItem& Item, const TMap<FSceneOutlinerTreeItemID, FSceneOutlinerTreeItemPtr>& Items, bool bCreate = false) override;
+	// End ISceneOutlinerHierarchy overrides
+};
+
+}
\ No newline at end of file
diff --git a/Engine/Plugins/Experimental/AnimNext/Source/AnimNextEditor/Private/Variables/Outliner/VariablesOutlinerMode.cpp b/Engine/Plugins/Experimental/AnimNext/Source/AnimNextEditor/Private/Variables/Outliner/VariablesOutlinerMode.cpp
new file mode 100644
--- /dev/null
+++ b/Engine/Plugins/Experimental/AnimNext/Source/AnimNextEditor/Private/Variables/Outliner/VariablesOutlinerMode.cpp
@@ -0,0 +1,333 @@
+// Copyright Epic Games, Inc. All Rights Reserved.
+
+#include "VariablesOutlinerMode.h"
+
+#include "ISourceControlModule.h"
+#include "WorkspaceItemMenuContext.h"
+#include "VariablesOutlinerHierarchy.h"
+#include "VariablesOutlinerEntryItem.h"
+#include "Toolkits/AssetEditorToolkitMenuContext.h"
+#include "AnimNextRigVMAssetEditorData.h"
+#include "IWorkspaceEditor.h"
+#include "AnimNextEditorModule.h"
+#include "AnimNextRigVMAssetEditorData.h"
+#include "AssetRegistry/AssetRegistryModule.h"
+#include "ScopedTransaction.h"
+#include "ToolMenus.h"
+#include "Framework/Commands/GenericCommands.h"
+#include "Entries/AnimNextRigVMAssetEntry.h"
+#include "Variables/SAddVariablesDialog.h"
+#include "Variables/AnimNextVariableItemMenuContext.h"
+#include "AnimNextRigVMAsset.h"
+#include "VariablesOutlinerAssetItem.h"
+#include "VariablesOutlinerDragDrop.h"
+#include "Common/GraphEditorSchemaActions.h"
+#include "Entries/AnimNextVariableEntry.h"
+#include "Variables/SVariablesView.h"
+
+#define LOCTEXT_NAMESPACE "VariablesOutlinerMode"
+
+namespace UE::AnimNext::Editor
+{
+
+FVariablesOutlinerMode::FVariablesOutlinerMode(SVariablesOutliner* InVariablesOutliner, const TSharedRef<UE::Workspace::IWorkspaceEditor>& InWorkspaceEditor)
+	: ISceneOutlinerMode(InVariablesOutliner)
+	, WeakWorkspaceEditor(InWorkspaceEditor)
+{
+	CommandList = MakeShared<FUICommandList>();
+}
+
+void FVariablesOutlinerMode::Rebuild()
+{
+	Hierarchy = CreateHierarchy();
+}
+
+TSharedPtr<SWidget> FVariablesOutlinerMode::CreateContextMenu()
+{
+	static const FName MenuName("VariablesOutliner.ItemContextMenu");
+	UToolMenus* ToolMenus = UToolMenus::Get();
+	if (!ToolMenus->IsMenuRegistered(MenuName))
+	{
+		FToolMenuOwnerScoped ToolMenuOwnerScope(this);
+		if (UToolMenu* Menu = ToolMenus->RegisterMenu(MenuName))
+		{
+			Menu->AddDynamicSection(TEXT("Assets"), FNewToolMenuDelegate::CreateLambda([](UToolMenu* InMenu)
+			{
+				const UAssetEditorToolkitMenuContext* EditorContext = InMenu->FindContext<UAssetEditorToolkitMenuContext>();
+				const UAnimNextVariableItemMenuContext* MenuContext = InMenu->FindContext<UAnimNextVariableItemMenuContext>();
+				if(EditorContext == nullptr || MenuContext == nullptr)
+				{
+					return;
+				}
+
+				FToolMenuSection& VariablesSection = InMenu->AddSection("Variables", LOCTEXT("VariablesSectionLabel", "Variables"));
+
+				VariablesSection.AddMenuEntry(TEXT("AddVariables"),
+					LOCTEXT("AddVariablesMenuItem", "Add Variable(s)"),
+					LOCTEXT("AddVariablesMenuItemTooltip", "Adds variables to assets.\nIf multiple assets are selected, then variables will be added to each.\nIf no assets are selected and there are multiple assets, variables will be added to all assets."),
+					FSlateIcon(FAppStyle::Get().GetStyleSetName(), "Icons.Plus"),
+					FUIAction(FExecuteAction::CreateLambda([WeakOutliner = MenuContext->WeakOutliner]()
+					{
+						if (TSharedPtr<SVariablesOutliner> Outliner = WeakOutliner.Pin())
+						{
+							Outliner->HandleAddVariablesClicked();
+						}
+					}))
+				);
+
+				VariablesSection.AddMenuEntryWithCommandList(FGenericCommands::Get().Delete, MenuContext->WeakCommandList.Pin());
+				VariablesSection.AddMenuEntryWithCommandList(FGenericCommands::Get().Rename, MenuContext->WeakCommandList.Pin());
+
+				if (TSharedPtr<SVariablesOutliner> Outliner = MenuContext->WeakOutliner.Pin())
+				{
+					Outliner->AddSourceControlMenuOptions(InMenu);
+				}
+			}));
+		}
+	}
+
+	{
+		UAnimNextVariableItemMenuContext* MenuContext = NewObject<UAnimNextVariableItemMenuContext>();
+		MenuContext->WeakWorkspaceEditor = WeakWorkspaceEditor;
+		MenuContext->WeakOutliner = StaticCastSharedRef<SVariablesOutliner>(SceneOutliner->AsShared());
+		MenuContext->WeakCommandList = CommandList;
+		TArray<FSceneOutlinerTreeItemPtr> SelectedItems = GetOutliner()->GetSelectedItems();
+		for(const FSceneOutlinerTreeItemPtr& Item : SelectedItems)
+		{
+			if (const FVariablesOutlinerAssetItem* AssetItem = Item->CastTo<FVariablesOutlinerAssetItem>())
+			{
+				UAnimNextRigVMAsset* Asset = AssetItem->SoftAsset.Get();
+				if(Asset == nullptr)
+				{
+					continue;
+				}
+
+				UAnimNextRigVMAssetEditorData* EditorData = UncookedOnly::FUtils::GetEditorData<UAnimNextRigVMAssetEditorData>(Asset);
+				if(EditorData == nullptr)
+				{
+					continue;
+				}
+				
+				MenuContext->WeakEditorDatas.Add(EditorData);
+			}
+			else if (const FVariablesOutlinerEntryItem* EntryItem = Item->CastTo<FVariablesOutlinerEntryItem>())
+			{
+				UAnimNextRigVMAssetEntry* Entry = EntryItem->WeakEntry.Get();
+				if(Entry == nullptr)
+				{
+					continue;
+				}
+
+				MenuContext->WeakEntries.Add(Entry);
+			}
+		}
+
+		FToolMenuContext Context;
+		Context.AddObject(MenuContext);
+		WeakWorkspaceEditor.Pin()->InitToolMenuContext(Context);
+		return UToolMenus::Get()->GenerateWidget(MenuName, Context);
+	}
+}
+
+FReply FVariablesOutlinerMode::OnKeyDown(const FKeyEvent& InKeyEvent)
+{
+	if (CommandList->ProcessCommandBindings(InKeyEvent))
+	{
+		return FReply::Handled();
+	}
+
+	return FReply::Unhandled();
+}
+
+void FVariablesOutlinerMode::OnItemClicked(FSceneOutlinerTreeItemPtr Item)
+{
+	const FSceneOutlinerItemSelection& Selection = SceneOutliner->GetSelection();
+	HandleItemSelection(Selection);
+}
+
+void FVariablesOutlinerMode::HandleItemSelection(const FSceneOutlinerItemSelection& Selection)
+{
+	if (TSharedPtr<UE::Workspace::IWorkspaceEditor> WorkspaceEditor = WeakWorkspaceEditor.Pin())
+	{
+		TArray<FSceneOutlinerTreeItemPtr> SelectedItems;
+		Selection.Get(SelectedItems);
+		TArray<UObject*> EntriesToShow;
+		EntriesToShow.Reserve(SelectedItems.Num());
+		for(const FSceneOutlinerTreeItemPtr& Item : SelectedItems)
+		{
+			if (const FVariablesOutlinerEntryItem* VariablesItem = Item->CastTo<FVariablesOutlinerEntryItem>())
+			{
+				if(UAnimNextRigVMAssetEntry* Entry = VariablesItem->WeakEntry.Get())
+				{
+					EntriesToShow.Add(Entry);
+				}
+			}
+		}
+
+		WorkspaceEditor->SetDetailsObjects(EntriesToShow);
+	}
+}
+
+void FVariablesOutlinerMode::OnItemSelectionChanged(FSceneOutlinerTreeItemPtr Item, ESelectInfo::Type SelectionType, const FSceneOutlinerItemSelection& Selection)
+{
+	HandleItemSelection(Selection);
+
+	if (TSharedPtr<UE::Workspace::IWorkspaceEditor> SharedWorkspaceEditor = WeakWorkspaceEditor.Pin())
+	{
+		SharedWorkspaceEditor->SetGlobalSelection(SceneOutliner->AsShared(), UE::Workspace::FOnClearGlobalSelection::CreateRaw(this, &FVariablesOutlinerMode::ResetOutlinerSelection));
+	}
+}
+
+bool FVariablesOutlinerMode::CanRenameItem(const ISceneOutlinerTreeItem& Item) const
+{
+	if (const FVariablesOutlinerEntryItem* EntryItem = Item.CastTo<FVariablesOutlinerEntryItem>())
+	{
+		return true;
+	}
+	else if (const FVariablesOutlinerAssetItem* AssetItem = Item.CastTo<FVariablesOutlinerAssetItem>())
+	{
+		return true;
+	}
+	return false;
+}
+
+void FVariablesOutlinerMode::BindCommands(const TSharedRef<FUICommandList>& OutCommandList)
+{
+	CommandList->MapAction( 
+		FGenericCommands::Get().Rename, 
+		FExecuteAction::CreateRaw(this, &FVariablesOutlinerMode::Rename),
+		FCanExecuteAction::CreateRaw(this, &FVariablesOutlinerMode::CanRename)
+	);
+
+	CommandList->MapAction(
+		FGenericCommands::Get().Delete,
+		FExecuteAction::CreateRaw(this, &FVariablesOutlinerMode::Delete),
+		FCanExecuteAction::CreateRaw(this, &FVariablesOutlinerMode::CanDelete));
+}
+
+TUniquePtr<ISceneOutlinerHierarchy> FVariablesOutlinerMode::CreateHierarchy()
+{
+	return MakeUnique<FVariablesOutlinerHierarchy>(this);
+}
+
+void FVariablesOutlinerMode::ResetOutlinerSelection()
+{
+	SceneOutliner->ClearSelection();
+}
+
+SVariablesOutliner* FVariablesOutlinerMode::GetOutliner() const
+{
+	return static_cast<SVariablesOutliner*>(SceneOutliner);
+}
+
+void FVariablesOutlinerMode::Rename()
+{
+	const FSceneOutlinerItemSelection& Selection = SceneOutliner->GetSelection();
+	if (Selection.Num() == 1)
+	{
+		FSceneOutlinerTreeItemPtr ItemToRename = Selection.SelectedItems[0].Pin();
+
+		if (ItemToRename.IsValid() && CanRenameItem(*ItemToRename) && ItemToRename->CanInteract())
+		{
+			SceneOutliner->SetPendingRenameItem(ItemToRename);
+			SceneOutliner->ScrollItemIntoView(ItemToRename);
+		}
+	}
+}
+
+bool FVariablesOutlinerMode::CanRename() const
+{
+	const FSceneOutlinerItemSelection& Selection = SceneOutliner->GetSelection();
+	if (Selection.Num() == 1)
+	{
+		FSceneOutlinerTreeItemPtr ItemToRename = Selection.SelectedItems[0].Pin();
+		return ItemToRename.IsValid() && CanRenameItem(*ItemToRename) && ItemToRename->CanInteract();
+	}
+	return false;
+}
+
+void FVariablesOutlinerMode::Delete()
+{
+	int32 NumEntries = 0;
+	TMap<UAnimNextRigVMAssetEditorData*, TArray<UAnimNextRigVMAssetEntry*>> EntriesToDeletePerAsset;
+	TArray<FSceneOutlinerTreeItemPtr> SelectedItems = SceneOutliner->GetSelectedItems();
+	for(const FSceneOutlinerTreeItemPtr& Item : SelectedItems)
+	{
+		if (const FVariablesOutlinerEntryItem* VariablesItem = Item->CastTo<FVariablesOutlinerEntryItem>())
+		{
+			UAnimNextRigVMAssetEntry* Entry = VariablesItem->WeakEntry.Get();
+			if(Entry == nullptr)
+			{
+				continue;
+			}
+
+			UAnimNextRigVMAssetEditorData* EditorData = Entry->GetTypedOuter<UAnimNextRigVMAssetEditorData>();
+			if(EditorData == nullptr)
+			{
+				continue;
+			}
+
+			TArray<UAnimNextRigVMAssetEntry*>& EntriesToDelete = EntriesToDeletePerAsset.FindOrAdd(EditorData);
+			EntriesToDelete.Add(Entry);
+			NumEntries++;
+		}
+	}
+
+	if(NumEntries > 0)
+	{
+		FScopedTransaction Transaction(FText::Format(LOCTEXT("DeleteVariablesFormat", "Delete {0}|plural(one=variable, other=variables)"), NumEntries));
+		for(const TPair<UAnimNextRigVMAssetEditorData*, TArray<UAnimNextRigVMAssetEntry*>>& EntriesPair : EntriesToDeletePerAsset)
+		{
+			EntriesPair.Key->RemoveEntries(EntriesPair.Value);
+		}
+	}
+}
+
+bool FVariablesOutlinerMode::CanDelete() const
+{
+	return true;
+}
+
+TSharedPtr<FDragDropOperation> FVariablesOutlinerMode::CreateDragDropOperation(const FPointerEvent& MouseEvent, const TArray<FSceneOutlinerTreeItemPtr>& InTreeItems) const
+{
+	if(InTreeItems.Num() == 0)
+	{
+		return nullptr;
+	}
+
+	FVariablesOutlinerEntryItem* VariableItem = InTreeItems[0]->CastTo<FVariablesOutlinerEntryItem>();
+	if(VariableItem == nullptr)
+	{
+		return nullptr;
+	}
+
+	UAnimNextVariableEntry* Entry = Cast<UAnimNextVariableEntry>(VariableItem->WeakEntry.Get());
+	if(Entry == nullptr)
+	{
+		return nullptr;
+	}
+
+	URigVMHost* RigVMHost = Entry->GetTypedOuter<URigVMHost>();
+	if(RigVMHost == nullptr)
+	{
+		return nullptr;
+	}
+
+	FName VariableName = Entry->GetVariableName();
+	TArray<FRigVMExternalVariable> ExternalVariables = RigVMHost->GetVM()->GetExternalVariables(RigVMHost->GetRigVMExtendedExecuteContext());
+	for (const FRigVMExternalVariable& ExternalVariable : ExternalVariables)
+	{
+		if(ExternalVariable.Name == VariableName)
+		{
+			const bool bIsGetter = MouseEvent.IsControlDown();
+			TSharedPtr<FAnimNextSchemaAction_Variable> Action = MakeShared<FAnimNextSchemaAction_Variable>(ExternalVariable, bIsGetter);
+			return FVariableDragDropOp::New(Action);
+		}
+	}
+
+	return nullptr;
+}
+
+}
+
+#undef LOCTEXT_NAMESPACE // "FVariablesOutlinerMode"
\ No newline at end of file
diff --git a/Engine/Plugins/Experimental/AnimNext/Source/AnimNextEditor/Private/Variables/Outliner/VariablesOutlinerMode.h b/Engine/Plugins/Experimental/AnimNext/Source/AnimNextEditor/Private/Variables/Outliner/VariablesOutlinerMode.h
new file mode 100644
--- /dev/null
+++ b/Engine/Plugins/Experimental/AnimNext/Source/AnimNextEditor/Private/Variables/Outliner/VariablesOutlinerMode.h
@@ -0,0 +1,64 @@
+// Copyright Epic Games, Inc. All Rights Reserved.
+
+#pragma once
+
+#include "AnimNextRigVMAssetEditorData.h"
+#include "ISceneOutlinerMode.h"
+
+class UAnimNextRigVMAssetEditorData;
+
+namespace UE::Workspace
+{
+	class IWorkspaceEditor;
+}
+
+namespace UE::AnimNext::Editor
+{
+	class SVariablesOutliner;
+}
+
+namespace UE::AnimNext::Editor
+{
+class FVariablesOutlinerMode : public ISceneOutlinerMode
+{
+public:
+	FVariablesOutlinerMode(SVariablesOutliner* InVariablesOutliner, const TSharedRef<UE::Workspace::IWorkspaceEditor>& InWorkspaceEditor);
+
+	// Begin ISceneOutlinerMode overrides
+	virtual void Rebuild() override;
+	virtual TSharedPtr<SWidget> CreateContextMenu() override;
+	virtual FReply OnKeyDown(const FKeyEvent& InKeyEvent) override;
+	virtual void OnItemClicked(FSceneOutlinerTreeItemPtr Item) override;
+	void HandleItemSelection(const FSceneOutlinerItemSelection& Selection);
+	virtual void OnItemSelectionChanged(FSceneOutlinerTreeItemPtr Item, ESelectInfo::Type SelectionType, const FSceneOutlinerItemSelection& Selection) override;
+	virtual bool CanCustomizeToolbar() const override { return true; }
+	virtual bool CanRenameItem(const ISceneOutlinerTreeItem& Item) const override;
+	virtual ESelectionMode::Type GetSelectionMode() const override { return ESelectionMode::Multi; }
+	virtual void BindCommands(const TSharedRef<FUICommandList>& OutCommandList) override;
+	virtual TSharedPtr<FDragDropOperation> CreateDragDropOperation(const FPointerEvent& MouseEvent, const TArray<FSceneOutlinerTreeItemPtr>& InTreeItems) const override;
+
+protected:
+	virtual TUniquePtr<ISceneOutlinerHierarchy> CreateHierarchy() override;
+	// End ISceneOutlinerMode overrides
+
+	void ResetOutlinerSelection();
+
+	SVariablesOutliner* GetOutliner() const;
+
+	void Rename();
+
+	bool CanRename() const;
+
+	void Delete();
+
+	bool CanDelete() const;
+
+private:
+	friend class FVariablesOutlinerHierarchy;
+
+	TWeakPtr<UE::Workspace::IWorkspaceEditor> WeakWorkspaceEditor;
+
+	TSharedPtr<FUICommandList> CommandList;
+};
+
+}
\ No newline at end of file
diff --git a/Engine/Plugins/Experimental/AnimNext/Source/AnimNextEditor/Private/Variables/SAddVariablesDialog.cpp b/Engine/Plugins/Experimental/AnimNext/Source/AnimNextEditor/Private/Variables/SAddVariablesDialog.cpp
new file mode 100644
--- /dev/null
+++ b/Engine/Plugins/Experimental/AnimNext/Source/AnimNextEditor/Private/Variables/SAddVariablesDialog.cpp
@@ -0,0 +1,415 @@
+// Copyright Epic Games, Inc. All Rights Reserved.
+
+#include "SAddVariablesDialog.h"
+
+#include "AddVariableDialogMenuContext.h"
+#include "AnimNextVariableSettings.h"
+#include "AssetToolsModule.h"
+#include "ContentBrowserModule.h"
+#include "DetailLayoutBuilder.h"
+#include "IContentBrowserSingleton.h"
+#include "EditorUtils.h"
+#include "UncookedOnlyUtils.h"
+#include "PropertyBagDetails.h"
+#include "SPinTypeSelector.h"
+#include "Framework/Application/SlateApplication.h"
+#include "Framework/Docking/TabManager.h"
+#include "Widgets/Layout/SUniformGridPanel.h"
+#include "Widgets/Input/SButton.h"
+#include "Widgets/Layout/SBox.h"
+#include "Widgets/Text/SInlineEditableTextBlock.h"
+#include "SSimpleButton.h"
+#include "SSimpleComboButton.h"
+#include "ToolMenus.h"
+#include "String/ParseTokens.h"
+
+#define LOCTEXT_NAMESPACE "SAddVariablesDialog"
+
+namespace UE::AnimNext::Editor
+{
+
+namespace AddVariablesDialog
+{
+static FName Column_Name(TEXT("Name"));
+static FName Column_Type(TEXT("Type"));
+static FName SelectLibraryMenuName(TEXT("AnimNext.AddVariablesDialog.SelectedLibraryMenu"));
+}
+
+bool FVariableToAdd::IsValid(FText& OutReason) const
+{
+	if(Name == NAME_None)
+	{
+		OutReason = LOCTEXT("InvalidVariableName", "Invalid Variable Name");
+	}
+
+	if(!Type.IsValid())
+	{
+		OutReason = LOCTEXT("InvalidVariableType", "Invalid Variable Type");
+	}
+	
+	return true; 
+}
+
+void SAddVariablesDialog::Construct(const FArguments& InArgs, const FAssetData& InAsset)
+{
+	using namespace AddVariablesDialog;
+
+	OnFilterVariableType = InArgs._OnFilterVariableType;
+	Asset = InAsset;
+
+	SWindow::Construct(SWindow::FArguments()
+		.Title(LOCTEXT("WindowTitle", "Add Parameters"))
+		.SizingRule(ESizingRule::UserSized)
+		.ClientSize(InArgs._AllowMultiple ? FVector2D(500.f, 500.f) : FVector2D(500.f, 100.f))
+		.SupportsMaximize(false)
+		.SupportsMinimize(false)
+		[
+			SNew(SBox)
+			.Padding(5.0f)
+			[
+				SNew(SVerticalBox)
+				+SVerticalBox::Slot()
+				.AutoHeight()
+				.HAlign(HAlign_Left)
+				.Padding(0.0f, 5.0f)
+				[
+					SNew(SSimpleButton)
+					.Visibility(InArgs._AllowMultiple ? EVisibility::Visible : EVisibility::Collapsed)
+					.Text(LOCTEXT("AddButton", "Add"))
+					.ToolTipText(LOCTEXT("AddButtonTooltip", "Queue a new parameter for adding. New parameters will re-use the settings from the last queued parameter."))
+					.Icon(FAppStyle::Get().GetBrush("Icons.Plus"))
+					.OnClicked_Lambda([this]()
+					{
+						AddEntry();
+						return FReply::Handled();
+					})
+				]
+				+SVerticalBox::Slot()
+				.FillHeight(1.0f)
+				[
+					SAssignNew(EntriesList, SListView<TSharedRef<FVariableToAddEntry>>)
+					.ListItemsSource(&Entries)
+					.OnGenerateRow(this, &SAddVariablesDialog::HandleGenerateRow)
+					.HeaderRow(
+						SNew(SHeaderRow)
+						+SHeaderRow::Column(Column_Name)
+						.DefaultLabel(LOCTEXT("NameColumnHeader", "Name"))
+						.ToolTipText(LOCTEXT("NameColumnHeaderTooltip", "The name of the new parameter"))
+						.FillWidth(0.25f)
+
+						+SHeaderRow::Column(Column_Type)
+						.DefaultLabel(LOCTEXT("TypeColumnHeader", "Type"))
+						.ToolTipText(LOCTEXT("TypeColumnHeaderTooltip", "The type of the new parameter"))
+						.FillWidth(0.25f)
+					)
+				]
+				+SVerticalBox::Slot()
+				.AutoHeight()
+				.HAlign(HAlign_Right)
+				[
+					SNew(SUniformGridPanel)
+					.SlotPadding(FAppStyle::Get().GetMargin("StandardDialog.SlotPadding"))
+					.MinDesiredSlotWidth(FAppStyle::GetFloat("StandardDialog.MinDesiredSlotWidth"))
+					.MinDesiredSlotHeight(FAppStyle::GetFloat("StandardDialog.MinDesiredSlotHeight"))
+					+SUniformGridPanel::Slot(0, 0)
+					[
+						SNew(SButton)
+						.HAlign(HAlign_Center)
+						.ButtonStyle(&FAppStyle::Get().GetWidgetStyle<FButtonStyle>("PrimaryButton"))
+						.IsEnabled_Lambda([this]()
+						{
+							// Check each entry to see if the button can be pressed
+							for(TSharedRef<FVariableToAdd> Entry : Entries)
+							{
+								if(!Entry->IsValid())
+								{
+									return false;
+								}
+							}
+
+							return true;
+						})
+						.Text_Lambda([this]()
+						{
+							return FText::Format(LOCTEXT("AddParametersButtonFormat", "Add {0} {0}|plural(one=Parameter,other=Parameters)"), FText::AsNumber(Entries.Num()));
+						})
+						.ToolTipText_Lambda([this]()
+						{
+							// Check each entry to see if the button can be pressed
+							for(TSharedRef<FVariableToAdd> Entry : Entries)
+							{
+								FText Reason;
+								if(!Entry->IsValid(Reason))
+								{
+									return FText::Format(LOCTEXT("AddParametersButtonTooltip_InvalidEntry", "A parameter to add is not valid: {0}"), Reason);
+								}
+							}
+							return LOCTEXT("AddParametersButtonTooltip", "Add the selected parameters to the current graph");
+						})
+						.OnClicked_Lambda([this]()
+						{
+							bOKPressed = true;
+							RequestDestroyWindow();
+							return FReply::Handled();
+						})
+					]
+					+SUniformGridPanel::Slot(1, 0)
+					[
+						SNew(SButton)
+						.HAlign(HAlign_Center)
+						.ButtonStyle(&FAppStyle::Get().GetWidgetStyle<FButtonStyle>("Button"))
+						.Text(LOCTEXT("CancelButton", "Cancel"))
+						.ToolTipText(LOCTEXT("CancelButtonTooltip", "Cancel adding new parameters"))
+						.OnClicked_Lambda([this]()
+						{
+							RequestDestroyWindow();
+							return FReply::Handled();
+						})
+					]
+				]
+			]
+		]);
+
+	// Add an initial item
+	AddEntry(InArgs._InitialParamType);
+}
+
+FReply SAddVariablesDialog::OnKeyDown(const FGeometry& MyGeometry, const FKeyEvent& InKeyEvent)
+{
+	if(InKeyEvent.GetKey() == EKeys::Escape)
+	{
+		RequestDestroyWindow();
+		return FReply::Handled();
+	}
+
+	return FReply::Unhandled();
+}
+
+void SAddVariablesDialog::AddEntry(const FAnimNextParamType& InParamType)
+{
+	const UAnimNextVariableSettings* Settings = GetDefault<UAnimNextVariableSettings>();
+	
+	TArray<FName> PendingNames;
+	PendingNames.Reserve(Entries.Num());
+	for(const TSharedRef<FVariableToAddEntry>& QueuedAdd : Entries)
+	{
+		PendingNames.Add(QueuedAdd->Name);
+	}
+	FName ParameterName = FUtils::GetNewParameterName(Settings->GetLastVariableName(), Asset, PendingNames);
+	Entries.Add(MakeShared<FVariableToAddEntry>(InParamType.IsValid() ? InParamType : Settings->GetLastVariableType(), ParameterName));
+
+	RefreshEntries();
+}
+
+void SAddVariablesDialog::RefreshEntries()
+{
+	EntriesList->RequestListRefresh();
+}
+
+class SVariableToAdd : public SMultiColumnTableRow<TSharedRef<SAddVariablesDialog::FVariableToAddEntry>>
+{
+	SLATE_BEGIN_ARGS(SVariableToAdd) {}
+	SLATE_END_ARGS()
+
+	void Construct(const FArguments& InArgs, const TSharedRef<STableViewBase>& InOwnerTableView, TSharedRef<SAddVariablesDialog::FVariableToAddEntry> InEntry, TSharedRef<SAddVariablesDialog> InDialog)
+	{
+		Entry = InEntry;
+		WeakDialog = InDialog;
+		
+		SMultiColumnTableRow<TSharedRef<SAddVariablesDialog::FVariableToAddEntry>>::Construct( SMultiColumnTableRow<TSharedRef<SAddVariablesDialog::FVariableToAddEntry>>::FArguments(), InOwnerTableView);
+	}
+
+	virtual TSharedRef<SWidget> GenerateWidgetForColumn(const FName& InColumnName) override
+	{
+		using namespace AddVariablesDialog;
+
+		if(InColumnName == Column_Name)
+		{
+			TSharedPtr<SInlineEditableTextBlock> EditableText;
+			TSharedRef<SWidget> Widget =
+				SNew(SBox)
+				.HAlign(HAlign_Left)
+				.VAlign(VAlign_Center)
+				[
+					SAssignNew(EditableText, SInlineEditableTextBlock)
+					.Font(IDetailLayoutBuilder::GetDetailFont())
+					.IsSelected(this, &SVariableToAdd::IsSelectedExclusively)
+					.ToolTipText(LOCTEXT("NameTooltip", "The name of the new variable"))
+					.Text_Lambda([this]()
+					{
+						return FText::FromName(Entry->Name);
+					})
+					.OnTextCommitted_Lambda([this](const FText& InText, ETextCommit::Type InCommitType)
+					{
+						Entry->Name = *InText.ToString();
+
+						UAnimNextVariableSettings* Settings = GetMutableDefault<UAnimNextVariableSettings>();
+						Settings->SetLastVariableName(Entry->Name);
+					})
+					.OnVerifyTextChanged_Lambda([this](const FText& InNewText, FText& OutErrorText)
+					{
+						const FString NewString = InNewText.ToString();
+
+						if(!FUtils::IsValidParameterNameString(NewString, OutErrorText))
+						{
+							return false;
+						}
+
+						if(TSharedPtr<SAddVariablesDialog> Dialog = WeakDialog.Pin())
+						{
+							const FName Name(*NewString);
+							if(FUtils::DoesParameterNameExistInAsset(Name, Dialog->Asset))
+							{
+								OutErrorText = LOCTEXT("Error_NameExists", "This name already exists in the project");
+								return false;
+							}
+
+							return true;
+						}
+
+						return false;
+					})
+				];
+
+			if(Entry->bIsNew)
+			{
+				EditableText->RegisterActiveTimer(1/60.0f, FWidgetActiveTimerDelegate::CreateSPLambda(EditableText.Get(), [WeakEditableText = TWeakPtr<SInlineEditableTextBlock>(EditableText)](double, float)
+				{
+					if(TSharedPtr<SInlineEditableTextBlock> PinnedEditableText = WeakEditableText.Pin())
+					{
+						PinnedEditableText->EnterEditingMode();
+					}
+					return EActiveTimerReturnType::Stop;
+				}));
+	
+				Entry->bIsNew = false;
+			}
+
+			return Widget;
+		}
+		else if(InColumnName == Column_Type)
+		{
+			auto GetPinInfo = [this]()
+			{
+				return UncookedOnly::FUtils::GetPinTypeFromParamType(Entry->Type);
+			};
+
+			auto PinInfoChanged = [this](const FEdGraphPinType& PinType)
+			{
+				Entry->Type = UncookedOnly::FUtils::GetParamTypeFromPinType(PinType);
+
+				UAnimNextVariableSettings* Settings = GetMutableDefault<UAnimNextVariableSettings>();
+				Settings->SetLastVariableType(Entry->Type);
+			};
+			
+			auto GetFilteredVariableTypeTree = [this](TArray<TSharedPtr<UEdGraphSchema_K2::FPinTypeTreeInfo>>& TypeTree, ETypeTreeFilter TypeTreeFilter)
+			{
+				FUtils::GetFilteredVariableTypeTree(TypeTree, TypeTreeFilter);
+
+				if(TSharedPtr<SAddVariablesDialog> Dialog = WeakDialog.Pin())
+				{
+					if(Dialog->OnFilterVariableType.IsBound())
+					{
+						auto IsPinTypeAllowed = [&Dialog](const FEdGraphPinType& InType)
+						{
+							FAnimNextParamType Type = UncookedOnly::FUtils::GetParamTypeFromPinType(InType);
+							if(Type.IsValid())
+							{
+								return Dialog->OnFilterVariableType.Execute(Type) == EFilterVariableResult::Include;
+							}
+							return false;
+						};
+
+						// Additionally filter by allowed types
+						for (int32 Index = 0; Index < TypeTree.Num(); )
+						{
+							TSharedPtr<UEdGraphSchema_K2::FPinTypeTreeInfo>& PinType = TypeTree[Index];
+
+							if (PinType->Children.Num() == 0 && !IsPinTypeAllowed(PinType->GetPinType(/*bForceLoadSubCategoryObject*/false)))
+							{
+								TypeTree.RemoveAt(Index);
+								continue;
+							}
+
+							for (int32 ChildIndex = 0; ChildIndex < PinType->Children.Num(); )
+							{
+								TSharedPtr<UEdGraphSchema_K2::FPinTypeTreeInfo> Child = PinType->Children[ChildIndex];
+								if (Child.IsValid())
+								{
+									if (!IsPinTypeAllowed(Child->GetPinType(/*bForceLoadSubCategoryObject*/false)))
+									{
+										PinType->Children.RemoveAt(ChildIndex);
+										continue;
+									}
+								}
+								++ChildIndex;
+							}
+
+							++Index;
+						}
+					}
+				}
+			};
+
+			return
+				SNew(SBox)
+				.HAlign(HAlign_Left)
+				.VAlign(VAlign_Center)
+				[
+					SNew(SPinTypeSelector, FGetPinTypeTree::CreateLambda(GetFilteredVariableTypeTree))
+						.TargetPinType_Lambda(GetPinInfo)
+						.OnPinTypeChanged_Lambda(PinInfoChanged)
+						.Schema(GetDefault<UPropertyBagSchema>())
+						.bAllowArrays(true)
+						.TypeTreeFilter(ETypeTreeFilter::None)
+						.Font(IDetailLayoutBuilder::GetDetailFont())
+				];
+		}
+
+		return SNullWidget::NullWidget;
+	}
+
+	TSharedPtr<SAddVariablesDialog::FVariableToAddEntry> Entry;
+	TWeakPtr<SAddVariablesDialog> WeakDialog;
+};
+
+TSharedRef<ITableRow> SAddVariablesDialog::HandleGenerateRow(TSharedRef<FVariableToAddEntry> InEntry, const TSharedRef<STableViewBase>& InOwnerTable)
+{
+	return SNew(SVariableToAdd, InOwnerTable, InEntry, SharedThis(this));
+}
+
+bool SAddVariablesDialog::ShowModal(TArray<FVariableToAdd>& OutParameters)
+{
+	FSlateApplication::Get().AddModalWindow(SharedThis(this), FGlobalTabmanager::Get()->GetRootWindow());
+
+	if(bOKPressed)
+	{
+		bool bHasValid = false;
+		for(TSharedRef<FVariableToAddEntry>& Entry : Entries)
+		{
+			if(Entry->IsValid())
+			{
+				OutParameters.Add(*Entry);
+				bHasValid = true;
+			}
+		}
+		return bHasValid;
+	}
+	return false;
+}
+
+TSharedRef<SWidget> SAddVariablesDialog::HandleGetAddVariableMenuContent(TSharedPtr<FVariableToAddEntry> InEntry)
+{
+	using namespace AddVariablesDialog;
+
+	UToolMenus* ToolMenus = UToolMenus::Get();
+
+	UAddVariableDialogMenuContext* MenuContext = NewObject<UAddVariableDialogMenuContext>();
+	MenuContext->AddVariablesDialog = SharedThis(this);
+	MenuContext->Entry = InEntry;
+	return ToolMenus->GenerateWidget(SelectLibraryMenuName, FToolMenuContext(MenuContext));
+}
+
+}
+
+#undef LOCTEXT_NAMESPACE
\ No newline at end of file
diff --git a/Engine/Plugins/Experimental/AnimNext/Source/AnimNextEditor/Private/Variables/SAddVariablesDialog.h b/Engine/Plugins/Experimental/AnimNext/Source/AnimNextEditor/Private/Variables/SAddVariablesDialog.h
new file mode 100644
--- /dev/null
+++ b/Engine/Plugins/Experimental/AnimNext/Source/AnimNextEditor/Private/Variables/SAddVariablesDialog.h
@@ -0,0 +1,107 @@
+// Copyright Epic Games, Inc. All Rights Reserved.
+
+#pragma once
+
+#include "Param/ParamType.h"
+#include "Widgets/SWindow.h"
+#include "AssetRegistry/AssetData.h"
+#include "Widgets/Views/SListView.h"
+
+class SWrapBox;
+class UAnimNextModule_EditorData;
+
+namespace UE::AnimNext::Editor
+{
+
+struct FVariableToAdd
+{
+	FVariableToAdd() = default;
+
+	FVariableToAdd(const FAnimNextParamType& InType, FName InName)
+		: Type(InType)
+		, Name(InName)
+	{}
+
+	bool IsValid() const
+	{
+		return Name != NAME_None && Type.IsValid(); 
+	}
+
+	bool IsValid(FText& OutReason) const;
+
+	// Type
+	FAnimNextParamType Type;
+
+	// Name for variable
+	FName Name;
+};
+
+// Result of a filter operation via FOnFilterVariableType
+enum class EFilterVariableResult : int32
+{
+	Include,
+	Exclude
+};
+
+// Delegate called to filter variables by type for display to the user
+using FOnFilterVariableType = TDelegate<EFilterVariableResult(const FAnimNextParamType& /*InType*/)>;
+
+class SAddVariablesDialog : public SWindow
+{
+public:
+	SLATE_BEGIN_ARGS(SAddVariablesDialog)
+		: _AllowMultiple(true)
+	{}
+
+	/** Whether we allow multiple variables to be added or just one at a time */
+	SLATE_ARGUMENT(bool, AllowMultiple)
+
+	/** Delegate called to filter variables by type for display to the user */
+	SLATE_EVENT(FOnFilterVariableType, OnFilterVariableType)
+
+	/** Initial variable type to use */
+	SLATE_ARGUMENT(FAnimNextParamType, InitialParamType)
+
+	SLATE_END_ARGS()
+
+	void Construct(const FArguments& InArgs, const FAssetData& InAsset);
+
+	bool ShowModal(TArray<FVariableToAdd>& OutVariables);
+
+private:
+	virtual FReply OnKeyDown(const FGeometry& MyGeometry, const FKeyEvent& InKeyEvent) override;
+
+	void AddEntry(const FAnimNextParamType& InParamType = FAnimNextParamType());
+
+	void RefreshEntries();
+
+	struct FVariableToAddEntry : FVariableToAdd
+	{
+		FVariableToAddEntry() = default;
+
+		FVariableToAddEntry(const FAnimNextParamType& InType, FName InName)
+			: FVariableToAdd(InType, InName)
+		{}
+
+		bool bIsNew = true;
+	};
+
+	TSharedRef<ITableRow> HandleGenerateRow(TSharedRef<FVariableToAddEntry> InEntry, const TSharedRef<STableViewBase>& InOwnerTable);
+
+	TSharedRef<SWidget> HandleGetAddVariableMenuContent(TSharedPtr<FVariableToAddEntry> InEntry);
+	
+private:
+	friend class SVariableToAdd;
+
+	TSharedPtr<SListView<TSharedRef<FVariableToAddEntry>>> EntriesList;
+
+	TArray<TSharedRef<FVariableToAddEntry>> Entries;
+
+	FOnFilterVariableType OnFilterVariableType;
+
+	FAssetData Asset;
+
+	bool bOKPressed = false;
+};
+
+}
diff --git a/Engine/Plugins/Experimental/AnimNext/Source/AnimNextEditor/Private/Variables/SVariablesView.cpp b/Engine/Plugins/Experimental/AnimNext/Source/AnimNextEditor/Private/Variables/SVariablesView.cpp
new file mode 100644
--- /dev/null
+++ b/Engine/Plugins/Experimental/AnimNext/Source/AnimNextEditor/Private/Variables/SVariablesView.cpp
@@ -0,0 +1,330 @@
+﻿// Copyright Epic Games, Inc. All Rights Reserved.
+
+#include "SVariablesView.h"
+
+#include "AnimNextAssetWorkspaceAssetUserData.h"
+#include "AnimNextRigVMAsset.h"
+#include "ContentBrowserModule.h"
+#include "IContentBrowserSingleton.h"
+#include "SAssetDropTarget.h"
+#include "SceneOutlinerPublicTypes.h"
+#include "SceneOutlinerSourceControlColumn.h"
+#include "ScopedTransaction.h"
+#include "SPositiveActionButton.h"
+#include "SSceneOutliner.h"
+#include "Framework/Commands/GenericCommands.h"
+#include "Outliner/VariablesOutlinerColumns.h"
+#include "Outliner/VariablesOutlinerMode.h"
+#include "Outliner/VariablesOutlinerEntryItem.h"
+#include "Variables/SAddVariablesDialog.h"
+#include "AnimNextRigVMAsset.h"
+#include "IWorkspaceEditor.h"
+#include "WorkspaceAssetRegistryInfo.h"
+#include "Outliner/VariablesOutlinerAssetItem.h"
+
+#define LOCTEXT_NAMESPACE "SVariablesView"
+
+namespace UE::AnimNext::Editor
+{
+
+const FLazyName VariablesTabName("VariablesTab");
+
+void SVariablesOutliner::Construct(const FArguments& InArgs, const FSceneOutlinerInitializationOptions& InitOptions)
+{
+	AddVariablesButton = SNew(SPositiveActionButton)
+			.OnClicked(this, &SVariablesOutliner::HandleAddVariablesClicked)
+			.Icon(FAppStyle::Get().GetBrush("Icons.Plus"))
+			.Text(LOCTEXT("AddVariablesButton", "Add"))
+			.ToolTipText(LOCTEXT("AddVariablesButtonTooltip", "Adds variables to assets.\nIf multiple assets are selected, then variables will be added to each.\nIf no assets are selected and there are multiple assets, variables will be added to all assets."));
+
+	SSceneOutliner::Construct(InArgs, InitOptions);
+}
+
+void SVariablesOutliner::CustomAddToToolbar(TSharedPtr<class SHorizontalBox> Toolbar)
+{
+	Toolbar->AddSlot()
+	.VAlign(VAlign_Center)
+	.AutoWidth()
+	.Padding(4.f, 0.f, 0.f, 0.f)
+	[
+		AddVariablesButton.ToSharedRef()
+	];
+}
+
+FReply SVariablesOutliner::HandleAddVariablesClicked()
+{
+	TArray<UAnimNextRigVMAssetEditorData*> AssetsToAddTo;
+	if(Assets.Num() == 1)
+	{
+		UAnimNextRigVMAsset* Asset = Assets[0].Get();
+		if(Asset == nullptr)
+		{
+			FReply::Unhandled();
+		}
+
+		UAnimNextRigVMAssetEditorData* EditorData = UncookedOnly::FUtils::GetEditorData<UAnimNextRigVMAssetEditorData>(Asset);
+		if(EditorData == nullptr)
+		{
+			FReply::Unhandled();
+		}
+
+		AssetsToAddTo.Add(EditorData);
+	}
+	else
+	{
+		// Add to multiple selected
+		TArray<FSceneOutlinerTreeItemPtr> SelectedItems = GetSelectedItems();
+		for(const FSceneOutlinerTreeItemPtr& Item : SelectedItems)
+		{
+			const FVariablesOutlinerAssetItem* AssetItem = Item->CastTo<FVariablesOutlinerAssetItem>();
+			if (AssetItem == nullptr)
+			{
+				continue;
+			}
+
+			UAnimNextRigVMAsset* Asset = AssetItem->SoftAsset.Get();
+			if(Asset == nullptr)
+			{
+				continue;
+			}
+
+			UAnimNextRigVMAssetEditorData* EditorData = UncookedOnly::FUtils::GetEditorData<UAnimNextRigVMAssetEditorData>(Asset);
+			if(EditorData == nullptr)
+			{
+				continue;
+			}
+			
+			AssetsToAddTo.Add(EditorData);
+		}
+
+		// No selected asset items, so add to all asset items
+		if(AssetsToAddTo.Num() == 0)
+		{
+			for(const TSoftObjectPtr<UAnimNextRigVMAsset>& SoftAsset : Assets)
+			{
+				UAnimNextRigVMAsset* Asset = SoftAsset.Get();
+				if(Asset == nullptr)
+				{
+					continue;
+				}
+
+				UAnimNextRigVMAssetEditorData* EditorData = UncookedOnly::FUtils::GetEditorData<UAnimNextRigVMAssetEditorData>(Asset);
+				if(EditorData == nullptr)
+				{
+					continue;
+				}
+
+				AssetsToAddTo.Add(EditorData);
+			}
+		}
+	}
+
+	if(AssetsToAddTo.Num() == 0)
+	{
+		return FReply::Unhandled();
+	}
+
+	TSharedRef<SAddVariablesDialog> AddVariablesDialog =
+		SNew(SAddVariablesDialog, AssetsToAddTo.Num() == 1 ? FAssetData(UncookedOnly::FUtils::GetAsset<UAnimNextRigVMAsset>(AssetsToAddTo[0])) : FAssetData());
+
+	TArray<FVariableToAdd> VariablesToAdd;
+	if(AddVariablesDialog->ShowModal(VariablesToAdd))
+	{
+		if(VariablesToAdd.Num() > 0)
+		{
+			FScopedTransaction Transaction(FText::Format(LOCTEXT("AddVariablesFormat", "Add {0}|plural(one=variable, other=variables)"), VariablesToAdd.Num() * AssetsToAddTo.Num()));
+			for(UAnimNextRigVMAssetEditorData* EditorData : AssetsToAddTo)
+			{
+				for (const FVariableToAdd& VariableToAdd : VariablesToAdd)
+				{
+					if(EditorData->FindEntry(VariableToAdd.Name) == nullptr)
+					{
+						EditorData->AddVariable(VariableToAdd.Name, VariableToAdd.Type);
+					}
+				}
+			}
+		}
+	}
+
+	return FReply::Handled();
+}
+
+void SVariablesOutliner::SetAssets(TConstArrayView<TSoftObjectPtr<UAnimNextRigVMAsset>> InAssets)
+{
+	for(const TSoftObjectPtr<UAnimNextRigVMAsset>& CurrentSoftAsset : Assets)
+	{
+		if(UAnimNextRigVMAsset* CurrentAsset = CurrentSoftAsset.Get())
+		{
+			UAnimNextRigVMAssetEditorData* EditorData = UncookedOnly::FUtils::GetEditorData<UAnimNextRigVMAssetEditorData>(CurrentAsset);
+			if(EditorData != nullptr)
+			{
+				EditorData->ModifiedDelegate.RemoveAll(this);
+			}
+		}
+	}
+
+	Assets = InAssets;
+
+	for(const TSoftObjectPtr<UAnimNextRigVMAsset>& NewSoftAsset : Assets)
+	{
+		if(UAnimNextRigVMAsset* NewAsset = NewSoftAsset.Get())
+		{
+			UAnimNextRigVMAssetEditorData* EditorData = UncookedOnly::FUtils::GetEditorData<UAnimNextRigVMAssetEditorData>(NewAsset);
+			if(EditorData != nullptr)
+			{
+				EditorData->ModifiedDelegate.AddSP(this, &SVariablesOutliner::OnEditorDataModified);
+			}
+		}
+	}
+
+	FullRefresh();
+}
+
+void SVariablesOutliner::HandleAssetLoaded(const FSoftObjectPath& InSoftObjectPath, UAnimNextRigVMAsset* InAsset)
+{
+	if(Assets.Contains(TSoftObjectPtr<UAnimNextRigVMAsset>(InAsset)))
+	{
+		// Bind for any modification callbacks
+		UAnimNextRigVMAssetEditorData* EditorData = UncookedOnly::FUtils::GetEditorData<UAnimNextRigVMAssetEditorData>(InAsset);
+		if(EditorData != nullptr)
+		{
+			EditorData->ModifiedDelegate.AddSP(this, &SVariablesOutliner::OnEditorDataModified);
+		}
+
+		FullRefresh();
+	}
+}
+
+void SVariablesOutliner::OnEditorDataModified(UAnimNextRigVMAssetEditorData* InEditorData, EAnimNextEditorDataNotifType InType, UObject* InSubject)
+{
+	ensure(Assets.Contains(TSoftObjectPtr<UAnimNextRigVMAsset>(UncookedOnly::FUtils::GetAsset<UAnimNextRigVMAsset>(InEditorData))));
+
+	switch(InType)
+	{
+	case EAnimNextEditorDataNotifType::EntryAdded:
+	case EAnimNextEditorDataNotifType::EntryRemoved:
+		FullRefresh();
+		break;
+	default:
+		break;
+	}
+}
+
+bool SVariablesOutliner::IsEnabled() const
+{
+	return Assets.Num() > 0;
+}
+
+void SVariablesView::Construct(const FArguments& InArgs, TSharedRef<UE::Workspace::IWorkspaceEditor> InWorkspaceEditor)
+{
+	InWorkspaceEditor->OnOutlinerSelectionChanged().AddSP(this, &SVariablesView::HandleWorkspaceOutlinerSelectionChanged);
+
+	FSceneOutlinerInitializationOptions InitOptions;
+	InitOptions.OutlinerIdentifier = TEXT("AnimNextVariablesOutliner");
+	InitOptions.bShowHeaderRow = true;
+	InitOptions.ColumnMap.Add(FSceneOutlinerBuiltInColumnTypes::Label(), FSceneOutlinerColumnInfo(ESceneOutlinerColumnVisibility::Visible, 0, FCreateSceneOutlinerColumn(), false, 0.5f));
+	InitOptions.ColumnMap.Add(FVariablesOutlinerTypeColumn::GetID(), FSceneOutlinerColumnInfo(ESceneOutlinerColumnVisibility::Visible, 10, FCreateSceneOutlinerColumn::CreateLambda([](ISceneOutliner& InSceneOutliner) { return MakeShared<FVariablesOutlinerTypeColumn>(InSceneOutliner); }), false));
+	InitOptions.ColumnMap.Add(FVariablesOutlinerValueColumn::GetID(), FSceneOutlinerColumnInfo(ESceneOutlinerColumnVisibility::Visible, 20, FCreateSceneOutlinerColumn::CreateLambda([](ISceneOutliner& InSceneOutliner) { return MakeShared<FVariablesOutlinerValueColumn>(InSceneOutliner); }), true, 0.5f));
+	InitOptions.ColumnMap.Add(FSceneOutlinerSourceControlColumn::GetID(), FSceneOutlinerColumnInfo(ESceneOutlinerColumnVisibility::Visible, 30, FCreateSceneOutlinerColumn::CreateLambda([](ISceneOutliner& InSceneOutliner) { return MakeShared<FSceneOutlinerSourceControlColumn>(InSceneOutliner); }), true));
+	InitOptions.ColumnMap.Add(FVariablesOutlinerAccessSpecifierColumn::GetID(), FSceneOutlinerColumnInfo(ESceneOutlinerColumnVisibility::Visible, 40, FCreateSceneOutlinerColumn::CreateLambda([](ISceneOutliner& InSceneOutliner) { return MakeShared<FVariablesOutlinerAccessSpecifierColumn>(InSceneOutliner); }), false));
+	InitOptions.ModeFactory = FCreateSceneOutlinerMode::CreateLambda([this, WeakWorkspaceEditor=InWorkspaceEditor.ToWeakPtr()](SSceneOutliner* InOutliner) { return new FVariablesOutlinerMode(static_cast<SVariablesOutliner*>(InOutliner), WeakWorkspaceEditor.Pin().ToSharedRef()); });
+
+	VariablesOutliner = SNew(SVariablesOutliner, InitOptions);
+	VariablesOutliner->SetEnabled(MakeAttributeSP(VariablesOutliner.Get(), &SVariablesOutliner::IsEnabled));
+
+	TArray<FWorkspaceOutlinerItemExport> SelectedExports;
+	InWorkspaceEditor->GetOutlinerSelection(SelectedExports);
+	HandleWorkspaceOutlinerSelectionChanged(SelectedExports);
+
+	ChildSlot
+	[
+		VariablesOutliner.ToSharedRef()
+	];
+}
+
+void SVariablesView::HandleWorkspaceOutlinerSelectionChanged(TConstArrayView<FWorkspaceOutlinerItemExport> InExports)
+{
+	TArray<TSoftObjectPtr<UAnimNextRigVMAsset>> Assets;
+	TArray<FSoftObjectPath> AssetsToAsyncLoad;
+	for(const FWorkspaceOutlinerItemExport& Export : InExports)
+	{
+		if(Export.GetAssetPath().IsAsset() && Export.GetData().GetScriptStruct()->IsChildOf(FAnimNextRigVMAssetOutlinerData::StaticStruct()))
+		{
+			Assets.Add(TSoftObjectPtr<UAnimNextRigVMAsset>(Export.GetAssetPath()));
+			UAnimNextRigVMAsset* ResolvedAsset = Cast<UAnimNextRigVMAsset>(Export.GetAssetPath().ResolveObject());
+			if(ResolvedAsset == nullptr)
+			{
+				AssetsToAsyncLoad.Add(Export.GetAssetPath());
+			}
+		}
+		else if(Export.GetData().GetScriptStruct()->IsChildOf(FAnimNextAssetEntryOutlinerData::StaticStruct()))
+		{
+			// TODO: As we are not showing references in the workspace yet, we are just traversing the outer here and showing the containing asset's
+			// variables. Eventually we would want to recurse up to the root of the export hierarchy, which will entail changing the signature of
+			// FOnOutlinerSelectionChanged to include parent items somehow, or adding parent query functionality to the workspace API
+			const FAnimNextAssetEntryOutlinerData& EntryData = Export.GetData().Get<FAnimNextAssetEntryOutlinerData>();
+			if(EntryData.Entry)
+			{
+				if(UAnimNextRigVMAsset* Asset = EntryData.Entry->GetTypedOuter<UAnimNextRigVMAsset>())
+				{
+					Assets.Add(TSoftObjectPtr<UAnimNextRigVMAsset>(Asset));
+				}
+			}
+		}
+	}
+
+	VariablesOutliner->SetAssets(Assets);
+
+	// Try async load any missing assets
+	for(const FSoftObjectPath& AssetPath : AssetsToAsyncLoad)
+	{
+		AssetPath.LoadAsync(FLoadSoftObjectPathAsyncDelegate::CreateLambda([WeakVariablesOutliner = TWeakPtr<SVariablesOutliner>(VariablesOutliner)](const FSoftObjectPath& InSoftObjectPath, UObject* InObject)
+		{
+			UAnimNextRigVMAsset* Asset = Cast<UAnimNextRigVMAsset>(InObject);
+			if(Asset == nullptr)
+			{
+				return;
+			}
+
+			TSharedPtr<SVariablesOutliner> PinnedVariablesOutliner = WeakVariablesOutliner.Pin();
+			if(!PinnedVariablesOutliner.IsValid())
+			{
+				return;
+			}
+
+			PinnedVariablesOutliner->HandleAssetLoaded(InSoftObjectPath, Asset);
+		}));
+	}
+}
+
+FAnimNextVariablesTabSummoner::FAnimNextVariablesTabSummoner(TSharedPtr<UE::Workspace::IWorkspaceEditor> InHostingApp)
+	: FWorkflowTabFactory(VariablesTabName, StaticCastSharedPtr<FAssetEditorToolkit>(InHostingApp))
+{
+	TabLabel = LOCTEXT("AnimNextVariablesTabLabel", "Variables");
+	TabIcon = FSlateIcon("EditorStyle", "LevelEditor.Tabs.Outliner");
+	ViewMenuDescription = LOCTEXT("AnimNextVariablesTabMenuDescription", "Variables");
+	ViewMenuTooltip = LOCTEXT("AnimNextVariablesTabToolTip", "Shows the Variables tab.");
+	bIsSingleton = true;
+
+	VariablesView = SNew(SVariablesView, InHostingApp.ToSharedRef());
+
+	TArray<FWorkspaceOutlinerItemExport> Exports;
+	if(InHostingApp->GetOutlinerSelection(Exports))
+	{
+		VariablesView->HandleWorkspaceOutlinerSelectionChanged(Exports);
+	}
+}
+
+TSharedRef<SWidget> FAnimNextVariablesTabSummoner::CreateTabBody(const FWorkflowTabSpawnInfo& Info) const
+{
+	return VariablesView.ToSharedRef();
+}
+
+FText FAnimNextVariablesTabSummoner::GetTabToolTipText(const FWorkflowTabSpawnInfo& Info) const
+{
+	return ViewMenuTooltip;
+}
+
+}
+
+#undef LOCTEXT_NAMESPACE
\ No newline at end of file
diff --git a/Engine/Plugins/Experimental/AnimNext/Source/AnimNextEditor/Private/Variables/SVariablesView.h b/Engine/Plugins/Experimental/AnimNext/Source/AnimNextEditor/Private/Variables/SVariablesView.h
new file mode 100644
--- /dev/null
+++ b/Engine/Plugins/Experimental/AnimNext/Source/AnimNextEditor/Private/Variables/SVariablesView.h
@@ -0,0 +1,86 @@
+﻿// Copyright Epic Games, Inc. All Rights Reserved.
+
+#pragma once
+
+#include "ContentBrowserDelegates.h"
+#include "SSceneOutliner.h"
+#include "Widgets/SCompoundWidget.h"
+#include "WorkflowOrientedApp/WorkflowTabFactory.h"
+
+enum class EAnimNextEditorDataNotifType : uint8;
+class UAnimNextRigVMAsset;
+struct FWorkspaceOutlinerItemExport;
+class SPositiveActionButton;
+class UAnimNextRigVMAssetEditorData;
+
+namespace UE::Workspace
+{
+	class IWorkspaceEditor;
+}
+
+namespace UE::AnimNext::Editor
+{
+	class SVariablesOutliner;
+}
+
+namespace UE::AnimNext::Editor
+{
+
+extern const FLazyName VariablesTabName;
+
+class SVariablesOutliner : public SSceneOutliner
+{
+public:
+	void Construct(const FArguments& InArgs, const FSceneOutlinerInitializationOptions& InitOptions);
+
+	virtual void CustomAddToToolbar(TSharedPtr<SHorizontalBox> Toolbar) override;
+
+	FReply HandleAddVariablesClicked();
+
+	void SetAssets(TConstArrayView<TSoftObjectPtr<UAnimNextRigVMAsset>> InAssets);
+
+	void OnEditorDataModified(UAnimNextRigVMAssetEditorData* InEditorData, EAnimNextEditorDataNotifType InType, UObject* InSubject);
+
+	void HandleAssetLoaded(const FSoftObjectPath& InSoftObjectPath, UAnimNextRigVMAsset* InAsset);
+	
+	bool IsEnabled() const;
+private:
+	friend class FVariablesOutlinerMode;
+	friend class FVariablesOutlinerHierarchy;
+
+	TArray<TSoftObjectPtr<UAnimNextRigVMAsset>> Assets;
+	TSharedPtr<SPositiveActionButton> AddVariablesButton;
+};
+
+class SVariablesView : public SCompoundWidget
+{
+	SLATE_BEGIN_ARGS(SVariablesView) {}
+	SLATE_END_ARGS()
+
+	void Construct(const FArguments& InArgs, TSharedRef<UE::Workspace::IWorkspaceEditor> InWorkspaceEditor);
+
+private:
+	void HandleWorkspaceOutlinerSelectionChanged(TConstArrayView<FWorkspaceOutlinerItemExport> InExports);
+
+private:
+	friend struct FAnimNextVariablesTabSummoner;
+	
+	TSharedPtr<SVariablesOutliner> VariablesOutliner;
+};
+
+struct FAnimNextVariablesTabSummoner : public FWorkflowTabFactory
+{
+public:
+	FAnimNextVariablesTabSummoner(TSharedPtr<UE::Workspace::IWorkspaceEditor> InHostingApp);
+
+private:
+	// FWorkflowTabFactory interface
+	virtual TSharedRef<SWidget> CreateTabBody(const FWorkflowTabSpawnInfo& Info) const override;
+	virtual FText GetTabToolTipText(const FWorkflowTabSpawnInfo& Info) const override;
+
+	// The widget this tab spawner wraps
+	TSharedPtr<SVariablesView> VariablesView;
+};
+
+
+};
\ No newline at end of file
diff --git a/Engine/Plugins/Experimental/AnimNext/Source/AnimNextEditor/Private/Workspace/AnimNextWorkspaceSchema.cpp b/Engine/Plugins/Experimental/AnimNext/Source/AnimNextEditor/Private/Workspace/AnimNextWorkspaceSchema.cpp
--- a/Engine/Plugins/Experimental/AnimNext/Source/AnimNextEditor/Private/Workspace/AnimNextWorkspaceSchema.cpp
+++ b/Engine/Plugins/Experimental/AnimNext/Source/AnimNextEditor/Private/Workspace/AnimNextWorkspaceSchema.cpp
@@ -3,8 +3,8 @@
 #include "AnimNextWorkspaceSchema.h"
 
 #include "AnimNextEditorModule.h"
+#include "Graph/AnimNextAnimationGraph.h"
 #include "Module/AnimNextModule.h"
-#include "Scheduler/AnimNextSchedule.h"
 
 #define LOCTEXT_NAMESPACE "AnimNextWorkspaceSchema"
 
diff --git a/Engine/Plugins/Experimental/AnimNext/Source/AnimNextEditor/Public/IAnimNextEditorModule.h b/Engine/Plugins/Experimental/AnimNext/Source/AnimNextEditor/Public/IAnimNextEditorModule.h
--- a/Engine/Plugins/Experimental/AnimNext/Source/AnimNextEditor/Public/IAnimNextEditorModule.h
+++ b/Engine/Plugins/Experimental/AnimNext/Source/AnimNextEditor/Public/IAnimNextEditorModule.h
@@ -11,15 +11,11 @@ struct FTopLevelAssetPath;
 namespace UE::AnimNext::Editor
 {
 
-struct FParameterPickerArgs;
+struct FVariablePickerArgs;
 
 class IAnimNextEditorModule : public IModuleInterface
 {
 public:
-	// Create a parameter picker
-	// @param  InArgs          Arguments used for configuring the picker
-	virtual TSharedRef<SWidget> CreateParameterPicker(const FParameterPickerArgs& InArgs) = 0;
-
 	// Register a valid fragment type name to be used with parameter UOLs
 	// @param InLocatorFragmentEditorName The name of the locator fragment editor
 	virtual void RegisterLocatorFragmentEditorType(FName InLocatorFragmentEditorName) = 0;
