AnimNext: Module refactor

diff --git a/Engine/Plugins/Experimental/AnimNext/Source/AnimNextUncookedOnly/Private/AnimNextRigVMAssetEditorData.cpp b/Engine/Plugins/Experimental/AnimNext/Source/AnimNextUncookedOnly/Private/AnimNextRigVMAssetEditorData.cpp
--- a/Engine/Plugins/Experimental/AnimNext/Source/AnimNextUncookedOnly/Private/AnimNextRigVMAssetEditorData.cpp
+++ b/Engine/Plugins/Experimental/AnimNext/Source/AnimNextUncookedOnly/Private/AnimNextRigVMAssetEditorData.cpp
@@ -2,35 +2,52 @@
 
 #include "AnimNextRigVMAssetEditorData.h"
 
+#include "AnimNextController.h"
+#include "AnimNextEdGraph.h"
+#include "AnimNextEdGraphSchema.h"
 #include "AnimNextRigVMAsset.h"
-#include "AnimNextRigVMAssetEntry.h"
+#include "Entries/AnimNextRigVMAssetEntry.h"
 #include "AnimNextRigVMAssetSchema.h"
-#include "Module/AnimNextModuleWorkspaceAssetUserData.h"
+#include "AnimNextAssetWorkspaceAssetUserData.h"
 #include "ControlRigDefines.h"
 #include "ExternalPackageHelper.h"
 #include "IAnimNextRigVMGraphInterface.h"
-#include "IAnimNextRigVMParameterInterface.h"
+#include "IWorkspaceEditor.h"
+#include "IWorkspaceEditorModule.h"
+#include "Variables/IAnimNextRigVMVariableInterface.h"
 #include "UncookedOnlyUtils.h"
-#include "AnimNextEdGraph.h"
-#include "AnimNextEdGraphSchema.h"
+#include "Entries/AnimNextAnimationGraphEntry.h"
+#include "Entries/AnimNextEventGraphEntry.h"
+#include "Entries/AnimNextVariableEntry.h"
+#include "Graph/AnimNextAnimationGraphSchema.h"
+#include "Graph/RigUnit_AnimNextBeginExecution.h"
 #include "Misc/TransactionObjectEvent.h"
+#include "Module/AnimNextEventGraphSchema.h"
+#include "Module/AnimNextModule_EditorData.h"
+#include "Module/RigUnit_AnimNextModuleEvents.h"
 #include "Param/AnimNextTag.h"
 #include "RigVMModel/RigVMFunctionLibrary.h"
 #include "RigVMModel/RigVMNotifications.h"
+#include "RigVMModel/Nodes/RigVMAggregateNode.h"
 #include "RigVMModel/Nodes/RigVMCollapseNode.h"
 #include "UObject/AssetRegistryTagsContext.h"
 
-void UAnimNextRigVMAssetEditorData::BroadcastModified()
+TAutoConsoleVariable<bool> CVarDumpProgrammaticGraphs(
+	TEXT("AnimNext.DumpProgrammaticGraphs"),
+	false,
+	TEXT("When true the transient programmatic graphs will be automatically opened for any that are generated."));
+
+void UAnimNextRigVMAssetEditorData::BroadcastModified(EAnimNextEditorDataNotifType InType, UObject* InSubject)
 {
 	RecompileVM();
 
 	if(!bSuspendEditorDataNotifications)
 	{
-		ModifiedDelegate.Broadcast(this);
+		ModifiedDelegate.Broadcast(this, InType, InSubject);
 	}
 }
 
-void UAnimNextRigVMAssetEditorData::ReportError(const TCHAR* InMessage) const
+void UAnimNextRigVMAssetEditorData::ReportError(const TCHAR* InMessage)
 {
 	FScriptExceptionHandler::Get().HandleException(ELogVerbosity::Error, InMessage, TEXT(""));
 }
@@ -139,9 +156,9 @@ void UAnimNextRigVMAssetEditorData::Initialize(bool bRecompileVM)
 
 	if (IInterface_AssetUserData* OuterUserData = Cast<IInterface_AssetUserData>(GetOuter()))
 	{
-		if(!OuterUserData->HasAssetUserDataOfClass(UAnimNextModuleWorkspaceAssetUserData::StaticClass()))
+		if(!OuterUserData->HasAssetUserDataOfClass(UAnimNextAssetWorkspaceAssetUserData::StaticClass()))
 		{
-			OuterUserData->AddAssetUserDataOfClass(UAnimNextModuleWorkspaceAssetUserData::StaticClass());
+			OuterUserData->AddAssetUserDataOfClass(UAnimNextAssetWorkspaceAssetUserData::StaticClass());
 		}
 	}
 }
@@ -183,7 +200,7 @@ void UAnimNextRigVMAssetEditorData::PostTransacted(const FTransactionObjectEvent
 
 	if (TransactionEvent.GetEventType() == ETransactionObjectEventType::UndoRedo)
 	{
-		BroadcastModified();
+		BroadcastModified(EAnimNextEditorDataNotifType::PropertyChanged, this);
 	}
 }
 
@@ -199,12 +216,12 @@ void UAnimNextRigVMAssetEditorData::GetAssetRegistryTags(FAssetRegistryTagsConte
 	// We may not have compiled yet, so cache exports if we havent already
 	if(!CachedExports.IsSet())
 	{
-		CachedExports = FAnimNextParameterProviderAssetRegistryExports();
-		UE::AnimNext::UncookedOnly::FUtils::GetAssetParameters(this, CachedExports.GetValue());
+		CachedExports = FAnimNextAssetRegistryExports();
+		UE::AnimNext::UncookedOnly::FUtils::GetAssetVariables(this, CachedExports.GetValue());
 	}
 
 	FString TagValue;
-	FAnimNextParameterProviderAssetRegistryExports::StaticStruct()->ExportText(TagValue, &CachedExports.GetValue(), nullptr, nullptr, PPF_None, nullptr);
+	FAnimNextAssetRegistryExports::StaticStruct()->ExportText(TagValue, &CachedExports.GetValue(), nullptr, nullptr, PPF_None, nullptr);
 	Context.AddTag(FAssetRegistryTag(UE::AnimNext::ExportsAnimNextAssetRegistryTag, TagValue, FAssetRegistryTag::TT_Hidden));
 }
 
@@ -414,7 +431,19 @@ void UAnimNextRigVMAssetEditorData::HandleConfigureRigVMController(const FRigVMC
 
 	TWeakObjectPtr<UAnimNextRigVMAssetEditorData> WeakThis(this);
 
-	// this delegate is used by the controller to retrieve the current bytecode of the VM
+	InControllerToConfigure->GetExternalVariablesDelegate.BindLambda([](URigVMGraph* InGraph) -> TArray<FRigVMExternalVariable> {
+		if (InGraph)
+		{
+			if(URigVMHost* RigVMHost = InGraph->GetTypedOuter<URigVMHost>())
+			{
+				return RigVMHost->GetExternalVariables();
+			}
+		}
+		return TArray<FRigVMExternalVariable>();
+	});
+	
+	// this delegate is used by the controller
+	// to retrieve the current bytecode of the VM
 	InControllerToConfigure->GetCurrentByteCodeDelegate.BindLambda([WeakThis]() -> const FRigVMByteCode*
 	{
 		if (WeakThis.IsValid())
@@ -554,6 +583,138 @@ TObjectPtr<URigVMGraph> UAnimNextRigVMAssetEditorData::CreateContainedGraphModel
 	return Model;
 }
 
+void UAnimNextRigVMAssetEditorData::RecompileVM()
+{
+	using namespace UE::AnimNext::UncookedOnly;
+
+	if (bIsCompiling)
+	{
+		return;
+	}
+
+	TGuardValue<bool> CompilingGuard(bIsCompiling, true);
+
+	UAnimNextRigVMAsset* Asset = FUtils::GetAsset<UAnimNextRigVMAsset>(this);
+
+	CachedExports = FAnimNextAssetRegistryExports();
+	FUtils::GetAssetVariables(this, CachedExports.GetValue());
+
+	bErrorsDuringCompilation = false;
+
+	RigGraphDisplaySettings.MinMicroSeconds = RigGraphDisplaySettings.LastMinMicroSeconds = DBL_MAX;
+	RigGraphDisplaySettings.MaxMicroSeconds = RigGraphDisplaySettings.LastMaxMicroSeconds = (double)INDEX_NONE;
+
+	TArray<URigVMGraph*> ProgrammaticGraphs;
+	{
+		TGuardValue<bool> ReentrantGuardSelf(bSuspendModelNotificationsForSelf, true);
+		TGuardValue<bool> ReentrantGuardOthers(RigVMClient.bSuspendModelNotificationsForOthers, true);
+
+		VMCompileSettings.SetExecuteContextStruct(FAnimNextExecuteContext::StaticStruct());
+		FRigVMCompileSettings Settings = (bCompileInDebugMode) ? FRigVMCompileSettings::Fast(VMCompileSettings.GetExecuteContextStruct()) : VMCompileSettings;
+
+		FMessageLog("AnimNextCompilerResults").NewPage(FText::FromName(Asset->GetFName()));
+		Settings.ASTSettings.ReportDelegate.BindLambda([](EMessageSeverity::Type InType, UObject* InObject, const FString& InString)
+		{
+			FMessageLog("AnimNextCompilerResults").Message(InType, FText::FromString(InString));
+		});
+
+		FUtils::RecreateVM(Asset);
+
+		FUtils::CompileVariables(Asset);
+
+		GetProgrammaticGraphs(Settings, ProgrammaticGraphs);
+		for(URigVMGraph* ProgrammaticGraph : ProgrammaticGraphs)
+		{
+			check(ProgrammaticGraph != nullptr);
+		}
+
+		Asset->VMRuntimeSettings = VMRuntimeSettings;
+
+		FRigVMClient* VMClient = GetRigVMClient();
+
+		TArray<URigVMGraph*> AllGraphs = VMClient->GetAllModels(false, false);
+		AllGraphs.Append(ProgrammaticGraphs);
+
+		if(AllGraphs.Num() == 0)
+		{
+			return;
+		}
+
+		UAnimNextController* Controller = CastChecked<UAnimNextController>(VMClient->GetOrCreateController(AllGraphs[0]));
+
+		URigVMCompiler* Compiler = URigVMCompiler::StaticClass()->GetDefaultObject<URigVMCompiler>();
+		Compiler->Compile(Settings, AllGraphs, Controller, Asset->VM, Asset->ExtendedExecuteContext, Asset->GetExternalVariables(), &PinToOperandMap);
+
+		// Initialize right away, in packaged builds we initialize during PostLoad
+		Asset->VM->Initialize(Asset->ExtendedExecuteContext);
+		Asset->GenerateUserDefinedDependenciesData(Asset->ExtendedExecuteContext);
+
+		// Notable difference with vanilla RigVM host behavior - we init the VM here at the moment as we only have one 'instance'
+		Asset->InitializeVM(FRigUnit_AnimNextBeginExecution::EventName);
+
+		if (bErrorsDuringCompilation)
+		{
+			if(Settings.SurpressErrors)
+			{
+				Settings.Reportf(EMessageSeverity::Info, Asset, TEXT("Compilation Errors may be suppressed for AnimNext asset: %s. See VM Compile Settings for more Details"), *Asset->GetName());
+			}
+		}
+
+		bVMRecompilationRequired = false;
+		if(Asset->VM)
+		{
+			RigVMCompiledEvent.Broadcast(Asset, Asset->VM, Asset->ExtendedExecuteContext);
+		}
+
+		FAnimNextAssetRegistryExports Exports;
+		FUtils::GetAssetVariables(this, Exports);
+	}
+
+#if WITH_EDITOR
+	// Display programmatic graphs
+	if(CVarDumpProgrammaticGraphs.GetValueOnGameThread())
+	{
+		UE::Workspace::IWorkspaceEditorModule& WorkspaceEditorModule = FModuleManager::LoadModuleChecked<UE::Workspace::IWorkspaceEditorModule>("WorkspaceEditor");
+		if(UE::Workspace::IWorkspaceEditor* WorkspaceEditor = WorkspaceEditorModule.OpenWorkspaceForObject(Asset, UE::Workspace::EOpenWorkspaceMethod::Default))
+		{
+			TArray<UObject*> Graphs;
+			for(URigVMGraph* ProgrammaticGraph : ProgrammaticGraphs)
+			{
+				// Some explanation needed here!
+				// URigVMEdGraph caches its underlying model internally in GetModel depending on its outer if it is no attached to a RigVMClient
+				// So here we rename the graph into the transient package so we dont get any notifications
+				ProgrammaticGraph->Rename(nullptr, GetTransientPackage(), REN_ForceNoResetLoaders | REN_DoNotDirty | REN_DontCreateRedirectors | REN_NonTransactional);
+
+				// then create the graph (transient so it outers to the RigVMGraph)
+				URigVMEdGraph* EdGraph = CastChecked<URigVMEdGraph>(CreateEdGraph(ProgrammaticGraph, true));
+
+				// Then cache the model
+				EdGraph->GetModel();
+				Graphs.Add(EdGraph);
+
+				// Now rename into this asset again to be able to correctly create a controller (needed to view the graph and interact with it)
+				ProgrammaticGraph->Rename(nullptr, this, REN_ForceNoResetLoaders | REN_DoNotDirty | REN_DontCreateRedirectors | REN_NonTransactional);
+				URigVMController* TempController = GetOrCreateController(ProgrammaticGraph);
+
+				// Resend notifications to rebuild the EdGraph
+				TempController->ResendAllNotifications();
+			}
+
+			WorkspaceEditor->OpenObjects(Graphs);
+		}
+	}
+#endif
+
+#if WITH_EDITOR
+//	RefreshBreakpoints(EditorData);
+#endif
+
+	if (IAssetRegistry* AssetRegistry = IAssetRegistry::Get())
+	{
+		AssetRegistry->AssetUpdateTags(Asset, EAssetRegistryTagsCaller::Fast);
+	}
+}
+
 void UAnimNextRigVMAssetEditorData::RecompileVMIfRequired()
 {
 	if (bVMRecompilationRequired)
@@ -673,10 +834,21 @@ void UAnimNextRigVMAssetEditorData::HandleModifiedEvent(ERigVMGraphNotifType InN
 			RequestAutoVMRecompilation();	// We need to rebuild our metadata when a default value changes
 			break;
 		}
+	case ERigVMGraphNotifType::PinAdded:
+		{
+			if (URigVMPin* Pin = Cast<URigVMPin>(InSubject))
+			{
+				if (Pin->IsTraitPin())
+				{
+					RequestAutoVMRecompilation();
+				}
+			}
+			break;
+		}
 	}
 	
 	// if the notification still has to be sent...
-	if (bNotifForOthersPending && !bSuspendModelNotificationsForOthers)
+	if (bNotifForOthersPending && !RigVMClient.bSuspendModelNotificationsForOthers)
 	{
 		if (RigVMGraphModifiedEvent.IsBound())
 		{
@@ -774,7 +946,7 @@ bool UAnimNextRigVMAssetEditorData::RemoveEntry(UAnimNextRigVMAssetEntry* InEntr
 	// This will cause any external package to be removed when saved
 	EntryToRemove->MarkAsGarbage();
 
-	BroadcastModified();
+	BroadcastModified(EAnimNextEditorDataNotifType::EntryRemoved, this);
 
 	return bResult;
 }
@@ -784,7 +956,7 @@ bool UAnimNextRigVMAssetLibrary::RemoveEntries(UAnimNextRigVMAsset* InAsset, con
 	return UE::AnimNext::UncookedOnly::FUtils::GetEditorData(InAsset)->RemoveEntries(InEntries, bSetupUndoRedo, bPrintPythonCommand);
 }
 
-bool UAnimNextRigVMAssetEditorData::RemoveEntries(const TArray<UAnimNextRigVMAssetEntry*>& InEntries, bool bSetupUndoRedo, bool bPrintPythonCommand)
+bool UAnimNextRigVMAssetEditorData::RemoveEntries(TConstArrayView<UAnimNextRigVMAssetEntry*> InEntries, bool bSetupUndoRedo, bool bPrintPythonCommand)
 {
 	bool bResult = false;
 	{
@@ -796,7 +968,7 @@ bool UAnimNextRigVMAssetEditorData::RemoveEntries(const TArray<UAnimNextRigVMAss
 		}
 	}
 
-	BroadcastModified();
+	BroadcastModified(EAnimNextEditorDataNotifType::EntryRemoved, this);
 
 	return bResult;
 }
@@ -849,6 +1021,477 @@ UAnimNextRigVMAssetEntry* UAnimNextRigVMAssetEditorData::FindEntryForRigVMEdGrap
 	return nullptr;
 }
 
+void UAnimNextRigVMAssetEditorData::CreateEdGraphForCollapseNode(URigVMCollapseNode* InNode, bool bForce)
+{
+	check(InNode);
+	URigVMGraph* CollapseNodeGraph = InNode->GetGraph();
+	check(CollapseNodeGraph);
+
+	if (bForce)
+	{
+		RemoveEdGraphForCollapseNode(InNode, false);
+	}
+
+	// For Function node
+	if (InNode->GetGraph()->IsA<URigVMFunctionLibrary>())
+	{
+		if (URigVMGraph* ContainedGraph = InNode->GetContainedGraph())
+		{
+			bool bFunctionGraphExists = false;
+			for (UEdGraph* FunctionGraph : FunctionEdGraphs)
+			{
+				if (URigVMEdGraph* RigFunctionGraph = Cast<URigVMEdGraph>(FunctionGraph))
+				{
+					if (RigFunctionGraph->ModelNodePath == ContainedGraph->GetNodePath())
+					{
+						bFunctionGraphExists = true;
+						break;
+					}
+				}
+			}
+
+			if (!bFunctionGraphExists)
+			{
+				const FName SubGraphName = RigVMClient.GetUniqueName(this, *InNode->GetName());
+				// create a sub graph
+				UAnimNextEdGraph* RigFunctionGraph = NewObject<UAnimNextEdGraph>(this, SubGraphName, RF_Transactional);
+				RigFunctionGraph->Schema = UAnimNextEdGraphSchema::StaticClass();
+				RigFunctionGraph->bAllowRenaming = true;
+				RigFunctionGraph->bEditable = true;
+				RigFunctionGraph->bAllowDeletion = true;
+				RigFunctionGraph->ModelNodePath = ContainedGraph->GetNodePath();
+				RigFunctionGraph->bIsFunctionDefinition = true;
+
+				RigFunctionGraph->Initialize(this);
+
+				FunctionEdGraphs.Add(RigFunctionGraph);
+
+				RigVMClient.GetOrCreateController(ContainedGraph)->ResendAllNotifications();
+			}
+		}
+	}
+	// --- For Collapse nodes ---
+	else if (URigVMEdGraph* RigEdGraph = Cast<URigVMEdGraph>(GetEditorObjectForRigVMGraph(InNode->GetGraph())))
+	{
+		if (URigVMGraph* ContainedGraph = InNode->GetContainedGraph())
+		{
+			bool bSubGraphExists = false;
+
+			const FString ContainedGraphNodePath = ContainedGraph->GetNodePath();
+			for (UEdGraph* SubGraph : RigEdGraph->SubGraphs)
+			{
+				if (UAnimNextEdGraph* SubRigGraph = Cast<UAnimNextEdGraph>(SubGraph))
+				{
+					if (SubRigGraph->ModelNodePath == ContainedGraphNodePath)
+					{
+						bSubGraphExists = true;
+						break;
+					}
+				}
+			}
+
+			if (!bSubGraphExists)
+			{
+				bool bEditable = true;
+				if (InNode->IsA<URigVMAggregateNode>())
+				{
+					bEditable = false;
+				}
+
+				UObject* Outer = FindEntryForRigVMGraph(CollapseNodeGraph->GetRootGraph());
+				if (Outer == nullptr)
+				{
+					Outer = this; // function library graph has no entry
+				}
+
+				const FName SubGraphName = RigVMClient.GetUniqueName(Outer, *InNode->GetEditorSubGraphName());
+				// create a sub graph, no need to set external package if outer is an Entry
+				UAnimNextEdGraph* SubRigGraph = NewObject<UAnimNextEdGraph>(Outer, SubGraphName, RF_Transactional);
+				SubRigGraph->Schema = UAnimNextEdGraphSchema::StaticClass();
+				SubRigGraph->bAllowRenaming = 1;
+				SubRigGraph->bEditable = bEditable;
+				SubRigGraph->bAllowDeletion = 1;
+				SubRigGraph->ModelNodePath = ContainedGraphNodePath;
+				SubRigGraph->bIsFunctionDefinition = false;
+
+				RigEdGraph->SubGraphs.Add(SubRigGraph);
+
+				SubRigGraph->Initialize(this);
+
+				GetOrCreateController(ContainedGraph)->ResendAllNotifications();
+			}
+		}
+	}
+}
+
+void UAnimNextRigVMAssetEditorData::RemoveEdGraphForCollapseNode(URigVMCollapseNode* InNode, bool bNotify)
+{
+	check(InNode);
+
+	if (InNode->GetGraph()->IsA<URigVMFunctionLibrary>())
+	{
+		if (URigVMGraph* ContainedGraph = InNode->GetContainedGraph())
+		{
+			for (UEdGraph* FunctionGraph : FunctionEdGraphs)
+			{
+				if (URigVMEdGraph* RigFunctionGraph = Cast<URigVMEdGraph>(FunctionGraph))
+				{
+					if (RigFunctionGraph->ModelNodePath == ContainedGraph->GetNodePath())
+					{
+						if (URigVMController* SubController = GetController(ContainedGraph))
+						{
+							SubController->OnModified().RemoveAll(RigFunctionGraph);
+						}
+
+						if (RigVMGraphModifiedEvent.IsBound() && bNotify)
+						{
+							RigVMGraphModifiedEvent.Broadcast(ERigVMGraphNotifType::NodeRemoved, InNode->GetGraph(), InNode);
+						}
+
+						FunctionEdGraphs.Remove(RigFunctionGraph);
+						RigFunctionGraph->Rename(nullptr, GetTransientPackage(), REN_ForceNoResetLoaders | REN_DontCreateRedirectors);
+						RigFunctionGraph->MarkAsGarbage();
+						break;
+					}
+				}
+			}
+		}
+	}
+	else if (URigVMEdGraph* RigGraph = Cast<URigVMEdGraph>(GetEditorObjectForRigVMGraph(InNode->GetGraph())))
+	{
+		if (URigVMGraph* ContainedGraph = InNode->GetContainedGraph())
+		{
+			for (UEdGraph* SubGraph : RigGraph->SubGraphs)
+			{
+				if (URigVMEdGraph* SubRigGraph = Cast<URigVMEdGraph>(SubGraph))
+				{
+					if (SubRigGraph->ModelNodePath == ContainedGraph->GetNodePath())
+					{
+						if (URigVMController* SubController = GetController(ContainedGraph))
+						{
+							SubController->OnModified().RemoveAll(SubRigGraph);
+						}
+
+						if (RigVMGraphModifiedEvent.IsBound() && bNotify)
+						{
+							RigVMGraphModifiedEvent.Broadcast(ERigVMGraphNotifType::NodeRemoved, InNode->GetGraph(), InNode);
+						}
+
+						RigGraph->SubGraphs.Remove(SubRigGraph);
+						SubRigGraph->Rename(nullptr, GetTransientPackage(), REN_ForceNoResetLoaders | REN_DontCreateRedirectors);
+						SubRigGraph->MarkAsGarbage();
+						break;
+					}
+				}
+			}
+		}
+	}
+}
+
+UEdGraph* UAnimNextRigVMAssetEditorData::CreateEdGraph(URigVMGraph* InRigVMGraph, bool bForce)
+{
+	check(InRigVMGraph);
+
+	if(InRigVMGraph->IsA<URigVMFunctionLibrary>())
+	{
+		return nullptr;
+	}
+
+	const bool bIsTransient = InRigVMGraph->HasAnyFlags(RF_Transient);
+	IAnimNextRigVMGraphInterface* Entry = Cast<IAnimNextRigVMGraphInterface>(FindEntryForRigVMGraph(InRigVMGraph));
+	if(Entry == nullptr && !bIsTransient)
+	{
+		// Not found, we could be adding a new entry, in which case the graph wont be assigned yet
+		check(Entries.Num() > 0);
+		check(Cast<IAnimNextRigVMGraphInterface>(Entries.Last()) != nullptr);
+		check(Cast<IAnimNextRigVMGraphInterface>(Entries.Last())->GetRigVMGraph() == nullptr);
+		Entry = Cast<IAnimNextRigVMGraphInterface>(FindEntryForRigVMGraph(nullptr));
+	}
+
+	if(Entry == nullptr && !bIsTransient)
+	{
+		return nullptr;
+	}
+	
+	if(bForce)
+	{
+		RemoveEdGraph(InRigVMGraph);
+	}
+
+	UObject* Outer = nullptr;
+	EObjectFlags Flags = RF_NoFlags;
+	if(!bIsTransient)
+	{
+		Outer = CastChecked<UObject>(Entry);
+		Flags = RF_Transactional;
+	}
+	else
+	{
+		// This outer is to allow URigVMEdGraph::GetModel to retrieve the graph in 'preview' scenarios 
+		Outer = InRigVMGraph;
+		Flags = RF_Transient;
+	}
+
+	const FName GraphName = Entry != nullptr ? RigVMClient.GetUniqueName(Outer, Entry->GetGraphName()) : NAME_None;
+	UAnimNextEdGraph* RigFunctionGraph = NewObject<UAnimNextEdGraph>(Outer, GraphName, Flags);
+	RigFunctionGraph->Schema = UAnimNextEdGraphSchema::StaticClass();
+	RigFunctionGraph->bAllowDeletion = true;
+	RigFunctionGraph->bIsFunctionDefinition = false;
+	RigFunctionGraph->ModelNodePath = InRigVMGraph->GetNodePath();
+	RigFunctionGraph->Initialize(this);
+
+	if(!bIsTransient)
+	{
+		Entry->SetEdGraph(RigFunctionGraph);
+		if(Entry->GetRigVMGraph() == nullptr)
+		{
+			Entry->SetRigVMGraph(InRigVMGraph);
+		}
+		else
+		{
+			check(Entry->GetRigVMGraph() == InRigVMGraph);
+		}
+	}
+
+	return RigFunctionGraph;
+}
+
+bool UAnimNextRigVMAssetEditorData::RemoveEdGraph(URigVMGraph* InModel)
+{
+	if(IAnimNextRigVMGraphInterface* Entry = Cast<IAnimNextRigVMGraphInterface>(FindEntryForRigVMGraph(InModel)))
+	{
+		RigVMClient.DestroyObject(Entry->GetEdGraph());
+		Entry->SetEdGraph(nullptr);
+		return true;
+	}
+	return false;
+}
+
+UAnimNextVariableEntry* UAnimNextRigVMAssetLibrary::AddVariable(UAnimNextModule* InModule, FName InName, EPropertyBagPropertyType InValueType,
+	EPropertyBagContainerType InContainerType, const UObject* InValueTypeObject, const FString& InDefaultValue, bool bSetupUndoRedo, bool bPrintPythonCommand)
+{
+	return UE::AnimNext::UncookedOnly::FUtils::GetEditorData(InModule)->AddVariable(InName, FAnimNextParamType(InValueType, InContainerType, InValueTypeObject), InDefaultValue, bSetupUndoRedo, bPrintPythonCommand);
+}
+
+UAnimNextVariableEntry* UAnimNextRigVMAssetEditorData::AddVariable(FName InName, FAnimNextParamType InType, const FString& InDefaultValue, bool bSetupUndoRedo, bool bPrintPythonCommand)
+{
+	if(InName == NAME_None)
+	{
+		ReportError(TEXT("UAnimNextRigVMAssetEditorData::AddVariable: Invalid variable name supplied."));
+		return nullptr;
+	}
+
+	if(!GetEntryClasses().Contains(UAnimNextVariableEntry::StaticClass()) || !CanAddNewEntry(UAnimNextVariableEntry::StaticClass()))
+	{
+		ReportError(TEXT("UAnimNextRigVMAssetEditorData::AddVariable: Cannot add a variable to this asset - entry is not allowed."));
+		return nullptr;
+	}
+
+	// Check for duplicate name
+	FName NewParameterName = InName;
+	auto DuplicateNamePredicate = [&NewParameterName](const UAnimNextRigVMAssetEntry* InEntry)
+	{
+		return InEntry->GetEntryName() == NewParameterName;
+	};
+
+	bool bAlreadyExists = Entries.ContainsByPredicate(DuplicateNamePredicate);
+	int32 NameNumber = InName.GetNumber() + 1;
+	while(bAlreadyExists)
+	{
+		NewParameterName = FName(InName, NameNumber++);
+		bAlreadyExists = Entries.ContainsByPredicate(DuplicateNamePredicate);
+	}
+
+	UAnimNextVariableEntry* NewEntry = CreateNewSubEntry<UAnimNextVariableEntry>(this);
+	{
+		TGuardValue<bool> DisableEditorDataNotifications(bSuspendEditorDataNotifications, true);
+		TGuardValue<bool> DisableAutoCompile(bAutoRecompileVM, false);
+
+		NewEntry->SetVariableName(NewParameterName, false);
+		NewEntry->SetType(InType, false);
+		if(InDefaultValue.Len() > 0)
+		{
+			NewEntry->SetDefaultValue(InDefaultValue, false);
+		}
+	}
+
+	if(bSetupUndoRedo)
+	{
+		NewEntry->Modify();
+		Modify();
+	}
+
+	Entries.Add(NewEntry);
+
+	BroadcastModified(EAnimNextEditorDataNotifType::EntryAdded, NewEntry);
+
+	return NewEntry;
+}
+
+UAnimNextEventGraphEntry* UAnimNextRigVMAssetLibrary::AddEventGraph(UAnimNextModule* InModule, FName InName, UScriptStruct* InEventStruct, bool bSetupUndoRedo, bool bPrintPythonCommand)
+{
+	return UE::AnimNext::UncookedOnly::FUtils::GetEditorData(InModule)->AddEventGraph(InName, InEventStruct, bSetupUndoRedo, bPrintPythonCommand);
+}
+
+UAnimNextEventGraphEntry* UAnimNextRigVMAssetEditorData::AddEventGraph(FName InName, UScriptStruct* InEventStruct, bool bSetupUndoRedo, bool bPrintPythonCommand)
+{
+	if(InName == NAME_None)
+	{
+		ReportError(TEXT("UAnimNextRigVMAssetEditorData::AddEventGraph: Invalid graph name supplied."));
+		return nullptr;
+	}
+
+	if(InEventStruct == nullptr || !InEventStruct->IsChildOf(FRigVMStruct::StaticStruct()))
+	{
+		ReportError(TEXT("UAnimNextRigVMAssetEditorData::AddEventGraph: Invalid event struct name supplied."));
+		return nullptr;
+	}
+
+	if(!GetEntryClasses().Contains(UAnimNextEventGraphEntry::StaticClass()) || !CanAddNewEntry(UAnimNextEventGraphEntry::StaticClass()))
+	{
+		ReportError(TEXT("UAnimNextRigVMAssetEditorData::AddEventGraph: Cannot add an event graph to this asset - entry is not allowed."));
+		return nullptr;
+	}
+
+	// Check for duplicate name
+	FName NewGraphName = InName;
+	auto DuplicateNamePredicate = [&NewGraphName](const UAnimNextRigVMAssetEntry* InEntry)
+	{
+		return InEntry->GetEntryName() == NewGraphName;
+	};
+
+	bool bAlreadyExists = Entries.ContainsByPredicate(DuplicateNamePredicate);
+	int32 NameNumber = InName.GetNumber() + 1;
+	while(bAlreadyExists)
+	{
+		NewGraphName = FName(InName, NameNumber++);
+		bAlreadyExists =  Entries.ContainsByPredicate(DuplicateNamePredicate);
+	}
+
+	UAnimNextEventGraphEntry* NewEntry = CreateNewSubEntry<UAnimNextEventGraphEntry>(this);
+	NewEntry->GraphName = NewGraphName;
+
+	if(bSetupUndoRedo)
+	{
+		NewEntry->Modify();
+		Modify();
+	}
+
+	Entries.Add(NewEntry);
+
+	// Add new graph
+	{
+		TGuardValue<bool> EnablePythonPrint(bSuspendPythonMessagesForRigVMClient, !bPrintPythonCommand);
+		TGuardValue<bool> DisableAutoCompile(bAutoRecompileVM, false);
+		// Editor data has to be the graph outer, or RigVM unique name generator will not work
+		URigVMGraph* NewRigVMGraphModel = RigVMClient.CreateModel(URigVMGraph::StaticClass()->GetFName(), UAnimNextEventGraphSchema::StaticClass(), bSetupUndoRedo, this);
+		if (ensure(NewRigVMGraphModel))
+		{
+			// Then, to avoid the graph losing ref due to external package, set the same package as the Entry
+			if (!NewRigVMGraphModel->HasAnyFlags(RF_Transient))
+			{
+				NewRigVMGraphModel->SetExternalPackage(CastChecked<UObject>(NewEntry)->GetExternalPackage());
+			}
+			ensure(NewRigVMGraphModel);
+			NewEntry->Graph = NewRigVMGraphModel;
+
+			RefreshExternalModels();
+			RigVMClient.AddModel(NewRigVMGraphModel, true);
+			URigVMController* Controller = RigVMClient.GetController(NewRigVMGraphModel);
+			UE::AnimNext::UncookedOnly::FUtils::SetupEventGraph(Controller, InEventStruct);
+		}
+	}
+
+	BroadcastModified(EAnimNextEditorDataNotifType::EntryAdded, NewEntry);
+
+	return NewEntry;
+}
+
+UAnimNextAnimationGraphEntry* UAnimNextRigVMAssetLibrary::AddAnimationGraph(UAnimNextModule* InModule, FName InName, bool bSetupUndoRedo, bool bPrintPythonCommand)
+{
+	return UE::AnimNext::UncookedOnly::FUtils::GetEditorData(InModule)->AddAnimationGraph(InName, bSetupUndoRedo, bPrintPythonCommand);
+}
+
+UAnimNextAnimationGraphEntry* UAnimNextRigVMAssetEditorData::AddAnimationGraph(FName InName, bool bSetupUndoRedo, bool bPrintPythonCommand)
+{
+	if(InName == NAME_None)
+	{
+		ReportError(TEXT("UAnimNextRigVMAssetEditorData::AddAnimationGraph: Invalid graph name supplied."));
+		return nullptr;
+	}
+
+	if(!GetEntryClasses().Contains(UAnimNextAnimationGraphEntry::StaticClass()) || !CanAddNewEntry(UAnimNextAnimationGraphEntry::StaticClass()))
+	{
+		ReportError(TEXT("UAnimNextRigVMAssetEditorData::AddAnimationGraph: Cannot add an animation graph to this asset - entry is not allowed."));
+		return nullptr;
+	}
+
+	// Check for duplicate name
+	FName NewGraphName = InName;
+	auto DuplicateNamePredicate = [&NewGraphName](const UAnimNextRigVMAssetEntry* InEntry)
+	{
+		return InEntry->GetEntryName() == NewGraphName;
+	};
+
+	bool bAlreadyExists = Entries.ContainsByPredicate(DuplicateNamePredicate);
+	int32 NameNumber = InName.GetNumber() + 1;
+	while(bAlreadyExists)
+	{
+		NewGraphName = FName(InName, NameNumber++);
+		bAlreadyExists =  Entries.ContainsByPredicate(DuplicateNamePredicate);
+	}
+
+	UAnimNextAnimationGraphEntry* NewEntry = CreateNewSubEntry<UAnimNextAnimationGraphEntry>(this);
+	NewEntry->GraphName = NewGraphName;
+
+	if(bSetupUndoRedo)
+	{
+		NewEntry->Modify();
+		Modify();
+	}
+
+	Entries.Add(NewEntry);
+
+	// Add new graph
+	{
+		TGuardValue<bool> EnablePythonPrint(bSuspendPythonMessagesForRigVMClient, !bPrintPythonCommand);
+		TGuardValue<bool> DisableAutoCompile(bAutoRecompileVM, false);
+		// Editor data has to be the graph outer, or RigVM unique name generator will not work
+		URigVMGraph* NewRigVMGraphModel = RigVMClient.CreateModel(URigVMGraph::StaticClass()->GetFName(), UAnimNextAnimationGraphSchema::StaticClass(), bSetupUndoRedo, this);
+		if (ensure(NewRigVMGraphModel))
+		{
+			// Then, to avoid the graph losing ref due to external package, set the same package as the Entry
+			if (!NewRigVMGraphModel->HasAnyFlags(RF_Transient))
+			{
+				NewRigVMGraphModel->SetExternalPackage(CastChecked<UObject>(NewEntry)->GetExternalPackage());
+			}
+			ensure(NewRigVMGraphModel);
+			NewEntry->Graph = NewRigVMGraphModel;
+
+			RefreshExternalModels();
+			RigVMClient.AddModel(NewRigVMGraphModel, true);
+			URigVMController* Controller = RigVMClient.GetController(NewRigVMGraphModel);
+			UE::AnimNext::UncookedOnly::FUtils::SetupAnimGraph(NewEntry, Controller);
+		}
+	}
+
+	BroadcastModified(EAnimNextEditorDataNotifType::EntryAdded, NewEntry);
+
+	return NewEntry;
+}
+
+bool UAnimNextRigVMAssetEditorData::HasPublicVariables() const
+{
+	for(UAnimNextRigVMAssetEntry* Entry : Entries)
+	{
+		if(UAnimNextVariableEntry* Variable = Cast<UAnimNextVariableEntry>(Entry))
+		{
+			return Variable->GetExportAccessSpecifier() == EAnimNextExportAccessSpecifier::Public;
+		}
+	}
+	return false;
+}
+
 void UAnimNextRigVMAssetEditorData::RefreshExternalModels()
 {
 	GraphModels.Reset();
@@ -863,5 +1506,4 @@ void UAnimNextRigVMAssetEditorData::RefreshExternalModels()
 			}
 		}
 	}
-
 }
diff --git a/Engine/Plugins/Experimental/AnimNext/Source/AnimNextUncookedOnly/Private/AnimNextUncookedOnlyModule.cpp b/Engine/Plugins/Experimental/AnimNext/Source/AnimNextUncookedOnly/Private/AnimNextUncookedOnlyModule.cpp
--- a/Engine/Plugins/Experimental/AnimNext/Source/AnimNextUncookedOnly/Private/AnimNextUncookedOnlyModule.cpp
+++ b/Engine/Plugins/Experimental/AnimNext/Source/AnimNextUncookedOnly/Private/AnimNextUncookedOnlyModule.cpp
@@ -2,16 +2,13 @@
 
 #include "AnimNextUncookedOnlyModule.h"
 
-#include "Module/AnimNextModuleWorkspaceAssetUserData.h"
+#include "AnimNextAssetWorkspaceAssetUserData.h"
 #include "UncookedOnlyUtils.h"
 #include "Engine/Blueprint.h"
 #include "Modules/ModuleManager.h"
-#include "Param/AnimNextParamUniversalObjectLocator.h"
-#include "Param/IParameterSourceType.h"
-#include "Param/ObjectProxyType.h"
-#include "Scheduler/AnimNextSchedule.h"
 #include "UObject/AssetRegistryTagsContext.h"
 #include "MessageLogModule.h"
+#include "Variables/UniversalObjectLocatorBindingType.h"
 
 #define LOCTEXT_NAMESPACE "AnimNextUncookedOnlyModule"
 
@@ -20,54 +17,7 @@ namespace UE::AnimNext::UncookedOnly
 
 void FModule::StartupModule()
 {
-	// TEMP: Bind the compilation function for schedules
-	UAnimNextSchedule::CompileFunction = [](UAnimNextSchedule* InSchedule)
-	{
-		FUtils::CompileSchedule(InSchedule);
-	};
-
-	// TEMP: Bind the asset registry tags function for schedules
-	UAnimNextSchedule::GetAssetRegistryTagsFunction = [](const UAnimNextSchedule* InSchedule, FAssetRegistryTagsContext Context)
-	{
-		FAnimNextParameterProviderAssetRegistryExports Exports;
-        {
-		    FUtils::GetScheduleParameters(InSchedule, Exports);
-		    
-		    FString TagValue;
-		    FAnimNextParameterProviderAssetRegistryExports::StaticStruct()->ExportText(TagValue, &Exports, nullptr, nullptr, PPF_None, nullptr);
-		    Context.AddTag(UObject::FAssetRegistryTag(UE::AnimNext::ExportsAnimNextAssetRegistryTag, TagValue, UObject::FAssetRegistryTag::TT_Hidden));
-        }
-
-		FWorkspaceOutlinerItemExports OutlinerExports;
-		{
-			FWorkspaceOutlinerItemExport& RootAssetExport = OutlinerExports.Exports.Add_GetRef(FWorkspaceOutlinerItemExport(InSchedule->GetFName(), InSchedule));
-			RootAssetExport.GetData().InitializeAsScriptStruct(FAnimNextSchedulerData::StaticStruct());
-			FAnimNextSchedulerData& AssetData = RootAssetExport.GetData().GetMutable<FAnimNextSchedulerData>();
-
-			{
-				FString TagValue;
-				FWorkspaceOutlinerItemExports::StaticStruct()->ExportText(TagValue, &OutlinerExports, nullptr, nullptr, PPF_None, nullptr);
-				Context.AddTag(UObject::FAssetRegistryTag(UE::Workspace::ExportsWorkspaceItemsRegistryTag, TagValue, UObject::FAssetRegistryTag::TT_Hidden));
-			}
-		}
-	};
-	
-	// Ensure that any BP components that we care about contribute to the parameter pool
-	OnGetExtraObjectTagsHandle = UObject::FAssetRegistryTag::OnGetExtraObjectTagsWithContext.AddLambda([](FAssetRegistryTagsContext Context)
-	{
-		const UObject* InObject = Context.GetObject();
-		if(const UBlueprint* Blueprint = Cast<UBlueprint>(InObject))
-		{
-			FAnimNextParameterProviderAssetRegistryExports Exports;
-			FUtils::GetBlueprintParameters(Blueprint, Exports);
-
-			FString TagValue;
-			FAnimNextParameterProviderAssetRegistryExports::StaticStruct()->ExportText(TagValue, &Exports, nullptr, nullptr, PPF_None, nullptr);
-			Context.AddTag(UObject::FAssetRegistryTag(UE::AnimNext::ExportsAnimNextAssetRegistryTag, TagValue, UObject::FAssetRegistryTag::TT_Hidden));
-		}
-	});
-
-	RegisterParameterSourceType(FAnimNextParamUniversalObjectLocator::StaticStruct(), MakeShared<FObjectProxyType>());
+	RegisterVariableBindingType("/Script/AnimNextUncookedOnly.AnimNextUniversalObjectLocatorBindingData", MakeShared<FUniversalObjectLocatorBindingType>());
 
 	// Register the compilation log (hidden from the main log set, it is displayed in the workspace editor)
 	FMessageLogModule& MessageLogModule = FModuleManager::LoadModuleChecked<FMessageLogModule>("MessageLog");
@@ -84,27 +34,22 @@ void FModule::ShutdownModule()
 		MessageLogModule->UnregisterLogListing("AnimNextCompilerResults");
 	}
 
-	UnregisterParameterSourceType(FAnimNextParamUniversalObjectLocator::StaticStruct());
-
-	UAnimNextSchedule::GetAssetRegistryTagsFunction = nullptr;
-	UAnimNextSchedule::CompileFunction = nullptr;
-
-	UObject::FAssetRegistryTag::OnGetExtraObjectTagsWithContext.Remove(OnGetExtraObjectTagsHandle);
+	UnregisterVariableBindingType("/Script/AnimNextUncookedOnly.AnimNextUniversalObjectLocatorBindingData");
 }
 
-void FModule::RegisterParameterSourceType(const UScriptStruct* InInstanceIdStruct, TSharedPtr<IParameterSourceType> InType)
+void FModule::RegisterVariableBindingType(FName InStructName, TSharedPtr<IVariableBindingType> InType)
 {
-	ParameterSourceTypes.Add(InInstanceIdStruct, InType);
+	VariableBindingTypes.Add(InStructName, InType);
 }
 
-void FModule::UnregisterParameterSourceType(const UScriptStruct* InInstanceIdStruct)
+void FModule::UnregisterVariableBindingType(FName InStructName)
 {
-	ParameterSourceTypes.Remove(InInstanceIdStruct);
+	VariableBindingTypes.Remove(InStructName);
 }
 
-TSharedPtr<IParameterSourceType> FModule::FindParameterSourceType(const UScriptStruct* InInstanceIdStruct) const
+TSharedPtr<IVariableBindingType> FModule::FindVariableBindingType(const UScriptStruct* InStruct) const
 {
-	return ParameterSourceTypes.FindRef(InInstanceIdStruct);
+	return VariableBindingTypes.FindRef(*InStruct->GetPathName());
 }
 
 }
diff --git a/Engine/Plugins/Experimental/AnimNext/Source/AnimNextUncookedOnly/Private/AnimNextUncookedOnlyModule.h b/Engine/Plugins/Experimental/AnimNext/Source/AnimNextUncookedOnly/Private/AnimNextUncookedOnlyModule.h
--- a/Engine/Plugins/Experimental/AnimNext/Source/AnimNextUncookedOnly/Private/AnimNextUncookedOnlyModule.h
+++ b/Engine/Plugins/Experimental/AnimNext/Source/AnimNextUncookedOnly/Private/AnimNextUncookedOnlyModule.h
@@ -21,12 +21,11 @@ private:
 	virtual void ShutdownModule() override;
 
 	// IAnimNextUncookedOnlyModule interface
-	virtual void RegisterParameterSourceType(const UScriptStruct* InInstanceIdStruct, TSharedPtr<IParameterSourceType> InType) override;
-	virtual void UnregisterParameterSourceType(const UScriptStruct* InInstanceIdStruct) override;
-	virtual TSharedPtr<IParameterSourceType> FindParameterSourceType(const UScriptStruct* InInstanceIdStruct) const override;
+	virtual void RegisterVariableBindingType(FName InStructName, TSharedPtr<IVariableBindingType> InType) override;
+	virtual void UnregisterVariableBindingType(FName InStructName) override;
+	virtual TSharedPtr<IVariableBindingType> FindVariableBindingType(const UScriptStruct* InInstanceIdStruct) const override;
 
-	TMap<const UScriptStruct*, TSharedPtr<IParameterSourceType>> ParameterSourceTypes;
-	FDelegateHandle OnGetExtraObjectTagsHandle;
+	TMap<FName, TSharedPtr<IVariableBindingType>> VariableBindingTypes;
 
 	friend struct FUtils;
 };
diff --git a/Engine/Plugins/Experimental/AnimNext/Source/AnimNextUncookedOnly/Private/Entries/AnimNextRigVMAssetEntry.cpp b/Engine/Plugins/Experimental/AnimNext/Source/AnimNextUncookedOnly/Private/Entries/AnimNextRigVMAssetEntry.cpp
new file mode 100644
--- /dev/null
+++ b/Engine/Plugins/Experimental/AnimNext/Source/AnimNextUncookedOnly/Private/Entries/AnimNextRigVMAssetEntry.cpp
@@ -0,0 +1,36 @@
+// Copyright Epic Games, Inc. All Rights Reserved.
+
+#include "Entries/AnimNextRigVMAssetEntry.h"
+#include "AnimNextRigVMAssetEditorData.h"
+
+void UAnimNextRigVMAssetEntry::Initialize(UAnimNextRigVMAssetEditorData* InEditorData)
+{
+	InEditorData->RigVMGraphModifiedEvent.RemoveAll(this);
+	InEditorData->RigVMGraphModifiedEvent.AddUObject(this, &UAnimNextRigVMAssetEntry::HandleRigVMGraphModifiedEvent);
+}
+
+bool UAnimNextRigVMAssetEntry::IsAsset() const
+{
+	// Entries are considered assets to allow using the asset logic for save dialogs, etc.
+	// Also, they return true even if pending kill, in order to show up as deleted in these dialogs.
+	return IsPackageExternal() && !GetPackage()->HasAnyFlags(RF_Transient) && !HasAnyFlags(RF_Transient | RF_ClassDefaultObject);
+}
+
+#if WITH_EDITOR
+
+void UAnimNextRigVMAssetEntry::PostEditChangeProperty(FPropertyChangedEvent& PropertyChangedEvent)
+{
+	Super::PostEditChangeProperty(PropertyChangedEvent);
+
+	BroadcastModified(EAnimNextEditorDataNotifType::PropertyChanged);
+}
+
+#endif
+
+void UAnimNextRigVMAssetEntry::BroadcastModified(EAnimNextEditorDataNotifType InType)
+{
+	if(UAnimNextRigVMAssetEditorData* EditorData = Cast<UAnimNextRigVMAssetEditorData>(GetOuter()))
+	{
+		EditorData->BroadcastModified(InType, this);
+	}
+}
\ No newline at end of file
diff --git a/Engine/Plugins/Experimental/AnimNext/Source/AnimNextUncookedOnly/Private/Entries/AnimNextVariableEntry.cpp b/Engine/Plugins/Experimental/AnimNext/Source/AnimNextUncookedOnly/Private/Entries/AnimNextVariableEntry.cpp
new file mode 100644
--- /dev/null
+++ b/Engine/Plugins/Experimental/AnimNext/Source/AnimNextUncookedOnly/Private/Entries/AnimNextVariableEntry.cpp
@@ -0,0 +1,150 @@
+// Copyright Epic Games, Inc. All Rights Reserved.
+
+#include "Entries/AnimNextVariableEntry.h"
+
+#include "AnimNextRigVMAssetEditorData.h"
+#include "UncookedOnlyUtils.h"
+#include "Module/AnimNextModule.h"
+#include "Variables/AnimNextUniversalObjectLocatorBindingData.h"
+
+const FLazyName IAnimNextRigVMVariableInterface::ValueName("Value");
+
+UAnimNextVariableEntry::UAnimNextVariableEntry()
+{
+	// TEMP: This should be selectable
+	Binding.BindingData.InitializeAs<FAnimNextUniversalObjectLocatorBindingData>();
+}
+
+FAnimNextParamType UAnimNextVariableEntry::GetExportType() const
+{
+	return GetType();
+}
+
+FName UAnimNextVariableEntry::GetExportName() const
+{
+	return GetVariableName();
+}
+
+EAnimNextExportAccessSpecifier UAnimNextVariableEntry::GetExportAccessSpecifier() const
+{
+	return Access;
+}
+
+void UAnimNextVariableEntry::SetExportAccessSpecifier(EAnimNextExportAccessSpecifier InAccessSpecifier, bool bSetupUndoRedo)
+{
+	if(bSetupUndoRedo)
+	{
+		Modify();
+	};
+
+	Access = InAccessSpecifier;
+
+	BroadcastModified(EAnimNextEditorDataNotifType::EntryAccessSpecifierChanged);
+}
+
+FAnimNextParamType UAnimNextVariableEntry::GetType() const
+{
+	return Type;
+}
+
+FName UAnimNextVariableEntry::GetEntryName() const
+{
+	return ParameterName;
+}
+
+bool UAnimNextVariableEntry::SetType(const FAnimNextParamType& InType, bool bSetupUndoRedo)
+{
+	if(bSetupUndoRedo)
+	{
+		Modify();
+	}
+
+	Type = InType;
+
+	DefaultValue.Reset();
+	DefaultValue.AddProperties({ FPropertyBagPropertyDesc(IAnimNextRigVMVariableInterface::ValueName, Type.GetContainerType(), Type.GetValueType(), Type.GetValueTypeObject()) });
+
+	BroadcastModified(EAnimNextEditorDataNotifType::VariableTypeChanged);
+
+	return true;
+}
+
+bool UAnimNextVariableEntry::SetDefaultValue(const FString& InDefaultValue, bool bSetupUndoRedo)
+{
+	if(bSetupUndoRedo)
+	{
+		Modify();
+	}
+	
+	return DefaultValue.SetValueSerializedString(IAnimNextRigVMVariableInterface::ValueName, InDefaultValue) == EPropertyBagResult::Success;
+}
+
+FName UAnimNextVariableEntry::GetVariableName() const
+{
+	return ParameterName;
+}
+
+void UAnimNextVariableEntry::SetVariableName(FName InName, bool bSetupUndoRedo)
+{
+	SetEntryName(InName, bSetupUndoRedo);
+}
+
+const FInstancedPropertyBag& UAnimNextVariableEntry::GetPropertyBag() const
+{
+	return DefaultValue;
+}
+
+const FAnimNextVariableBinding& UAnimNextVariableEntry::GetBinding() const
+{
+	return Binding;
+}
+
+void UAnimNextVariableEntry::SetEntryName(FName InName, bool bSetupUndoRedo)
+{
+	if(bSetupUndoRedo)
+	{
+		Modify();
+	}
+
+	ParameterName = InName;
+	BroadcastModified(EAnimNextEditorDataNotifType::EntryRenamed);
+}
+
+FText UAnimNextVariableEntry::GetDisplayName() const
+{
+	return FText::FromName(ParameterName);
+}
+
+FText UAnimNextVariableEntry::GetDisplayNameTooltip() const
+{
+	return FText::FromString(Comment);
+}
+
+void UAnimNextVariableEntry::Serialize(FArchive& Ar)
+{
+	Super::Serialize(Ar);
+	Ar.UsingCustomVersion(FFortniteMainBranchObjectVersion::GUID);
+}
+
+void UAnimNextVariableEntry::PostLoad()
+{
+	Super::PostLoad();
+	if(GetLinkerCustomVersion(FFortniteMainBranchObjectVersion::GUID) < FFortniteMainBranchObjectVersion::AnimNextModuleRefactor)
+	{
+		// Add a property for this type
+		DefaultValue.Reset();
+		DefaultValue.AddProperties({ FPropertyBagPropertyDesc(IAnimNextRigVMVariableInterface::ValueName, Type.GetContainerType(), Type.GetValueType(), Type.GetValueTypeObject()) });
+
+		// Copy any default value from the module's defaults, if found
+		UAnimNextRigVMAsset* Asset = GetTypedOuter<UAnimNextRigVMAsset>();
+		FString FullName = Asset->GetPathName() + TEXT(":") + GetVariableName().ToString();
+		const FPropertyBagPropertyDesc* OldDesc = Asset->VariableDefaults.GetPropertyBagStruct() ? Asset->VariableDefaults.GetPropertyBagStruct()->FindPropertyDescByName(*FullName) : nullptr;
+		const FPropertyBagPropertyDesc* NewDesc = DefaultValue.GetPropertyBagStruct() ? DefaultValue.GetPropertyBagStruct()->FindPropertyDescByName(IAnimNextRigVMVariableInterface::ValueName) : nullptr;
+		if(OldDesc && NewDesc)
+		{
+			const void* OldValue = OldDesc->CachedProperty->ContainerPtrToValuePtr<void>(Asset->VariableDefaults.GetValue().GetMemory());
+			void* NewValue = NewDesc->CachedProperty->ContainerPtrToValuePtr<void>(DefaultValue.GetMutableValue().GetMemory());
+			NewDesc->CachedProperty->CopyCompleteValue(NewValue, OldValue);
+		}
+	}
+}
\ No newline at end of file
diff --git a/Engine/Plugins/Experimental/AnimNext/Source/AnimNextUncookedOnly/Private/Module/AnimNextModule_EditorData.cpp b/Engine/Plugins/Experimental/AnimNext/Source/AnimNextUncookedOnly/Private/Module/AnimNextModule_EditorData.cpp
--- a/Engine/Plugins/Experimental/AnimNext/Source/AnimNextUncookedOnly/Private/Module/AnimNextModule_EditorData.cpp
+++ b/Engine/Plugins/Experimental/AnimNext/Source/AnimNextUncookedOnly/Private/Module/AnimNextModule_EditorData.cpp
@@ -5,22 +5,31 @@
 #include "ExternalPackageHelper.h"
 #include "UncookedOnlyUtils.h"
 #include "AssetRegistry/IAssetRegistry.h"
-#include "Curves/CurveFloat.h"
 #include "Module/AnimNextModule.h"
-#include "Graph/AnimNextModule_AnimationGraph.h"
-#include "AnimNextEdGraphSchema.h"
-#include "Module/AnimNextModule_EventGraph.h"
+#include "Entries/AnimNextAnimationGraphEntry.h"
+#include "Entries/AnimNextEventGraphEntry.h"
 #include "AnimNextEventGraphSchema.h"
-#include "Module/AnimNextModule_Parameter.h"
+#include "Entries/AnimNextVariableEntry.h"
 #include "Graph/AnimNextAnimationGraphSchema.h"
-#include "RigVMModel/RigVMFunctionLibrary.h"
-#include "RigVMModel/RigVMNotifications.h"
-#include "RigVMModel/Nodes/RigVMCollapseNode.h"
 #include "RigVMModel/Nodes/RigVMAggregateNode.h"
 #include "UObject/AssetRegistryTagsContext.h"
 #include "UObject/FortniteMainBranchObjectVersion.h"
 #include "UObject/LinkerLoad.h"
 
+void UAnimNextModule_EditorData::RecompileVM()
+{
+	Super::RecompileVM();
+
+	if (bIsCompiling)
+	{
+		return;
+	}
+
+#if WITH_EDITOR
+	UAnimNextModule::OnModuleCompiled().Broadcast(UE::AnimNext::UncookedOnly::FUtils::GetAsset<UAnimNextModule>(this));
+#endif
+}
+
 void UAnimNextModule_EditorData::PostLoad()
 {
 	Super::PostLoad();
@@ -56,7 +65,7 @@ void UAnimNextModule_EditorData::PostLoad()
 			Graph->SetExecuteContextStruct(FAnimNextExecuteContext::StaticStruct());
 			if(UAnimNextRigVMAssetEntry* FoundEntry = FindEntryForRigVMGraph(Graph))
 			{
-				if(FoundEntry->IsA(UAnimNextModule_AnimationGraph::StaticClass()))
+				if(FoundEntry->IsA(UAnimNextAnimationGraphEntry::StaticClass()))
 				{
 					Graph->SetSchemaClass(UAnimNextAnimationGraphSchema::StaticClass());
 				}
@@ -87,11 +96,11 @@ void UAnimNextModule_EditorData::PostLoad()
 			{
 				if(UAnimNextRigVMAssetEntry* FoundEntry = FindEntryForRigVMGraph(FoundRigVMGraph))
 				{
-					if(UAnimNextModule_AnimationGraph* AnimationGraphEntry = Cast<UAnimNextModule_AnimationGraph>(FoundEntry))
+					if(UAnimNextAnimationGraphEntry* AnimationGraphEntry = Cast<UAnimNextAnimationGraphEntry>(FoundEntry))
 					{
 						AnimationGraphEntry->EdGraph = Graph;
 					}
-					else if(UAnimNextModule_EventGraph* EventGraphEntry = Cast<UAnimNextModule_EventGraph>(FoundEntry))
+					else if(UAnimNextEventGraphEntry* EventGraphEntry = Cast<UAnimNextEventGraphEntry>(FoundEntry))
 					{
 						EventGraphEntry->EdGraph = Graph;
 					}
@@ -109,7 +118,7 @@ void UAnimNextModule_EditorData::PostLoad()
 			bool bFound = false;
 			for(UAnimNextRigVMAssetEntry* Entry : Entries)
 			{
-				if(UAnimNextModule_EventGraph* EventGraphEntry = Cast<UAnimNextModule_EventGraph>(Entry))
+				if(UAnimNextEventGraphEntry* EventGraphEntry = Cast<UAnimNextEventGraphEntry>(Entry))
 				{
 					if(DefaultModel == static_cast<IAnimNextRigVMGraphInterface*>(EventGraphEntry)->GetRigVMGraph())
 					{
@@ -141,11 +150,11 @@ void UAnimNextModule_EditorData::PostLoad()
 		// Force older assets to all have public symbols so they work as-is. Newer assets need user intervention as entries default to private
 		for(UAnimNextRigVMAssetEntry* Entry : Entries)
 		{
-			if(UAnimNextModule_AnimationGraph* AnimationGraphEntry = Cast<UAnimNextModule_AnimationGraph>(Entry))
+			if(UAnimNextAnimationGraphEntry* AnimationGraphEntry = Cast<UAnimNextAnimationGraphEntry>(Entry))
 			{
 				AnimationGraphEntry->Access = EAnimNextExportAccessSpecifier::Public;
 			}
-			else if(UAnimNextModule_Parameter* ParameterEntry = Cast<UAnimNextModule_Parameter>(Entry))
+			else if(UAnimNextVariableEntry* ParameterEntry = Cast<UAnimNextVariableEntry>(Entry))
 			{
 				ParameterEntry->Access = EAnimNextExportAccessSpecifier::Public;
 			}
@@ -153,448 +162,26 @@ void UAnimNextModule_EditorData::PostLoad()
 	}
 }
 
-void UAnimNextModule_EditorData::RecompileVM()
-{
-	CachedExports = FAnimNextParameterProviderAssetRegistryExports();
-	UE::AnimNext::UncookedOnly::FUtils::GetAssetParameters(this, CachedExports.GetValue());
-	UE::AnimNext::UncookedOnly::FUtils::Compile(GetTypedOuter<UAnimNextModule>());
-	
-	if (IAssetRegistry* AssetRegistry = IAssetRegistry::Get())
-	{
-		AssetRegistry->AssetUpdateTags(GetTypedOuter<UAnimNextModule>(), EAssetRegistryTagsCaller::Fast);
-	}
-}
-
-void UAnimNextModule_EditorData::HandleModifiedEvent(ERigVMGraphNotifType InNotifType, URigVMGraph* InGraph, UObject* InSubject)
-{
-	switch(InNotifType)
-	{
-	case ERigVMGraphNotifType::PinAdded:
-		{
-			if (URigVMPin* Pin = Cast<URigVMPin>(InSubject))
-			{
-				if (Pin->IsTraitPin())
-				{
-					RequestAutoVMRecompilation();
-				}
-			}
-			break;
-		}
-	}
-
-	Super::HandleModifiedEvent(InNotifType, InGraph, InSubject);
-}
-
 TConstArrayView<TSubclassOf<UAnimNextRigVMAssetEntry>> UAnimNextModule_EditorData::GetEntryClasses() const
 {
 	static const TSubclassOf<UAnimNextRigVMAssetEntry> Classes[] =
 	{
-		UAnimNextModule_AnimationGraph::StaticClass(),
-		UAnimNextModule_EventGraph::StaticClass(),
-		UAnimNextModule_Parameter::StaticClass(),
+		UAnimNextAnimationGraphEntry::StaticClass(),	// TEMP: Remove this once we have ported all old data
+		UAnimNextEventGraphEntry::StaticClass(),
+		UAnimNextVariableEntry::StaticClass(),
 	};
 	
 	return Classes;
 }
 
-void UAnimNextModule_EditorData::CreateEdGraphForCollapseNode(URigVMCollapseNode* InNode, bool bForce)
+void UAnimNextModule_EditorData::GetProgrammaticGraphs(const FRigVMCompileSettings& InSettings, TArray<URigVMGraph*>& OutGraphs)
 {
-	check(InNode);
-	URigVMGraph* CollapseNodeGraph = InNode->GetGraph();
-	check(CollapseNodeGraph);
-
-	if (bForce)
-	{
-		RemoveEdGraphForCollapseNode(InNode, false);
-	}
+	using namespace UE::AnimNext::UncookedOnly;
 
-	// For Function node
-	if (InNode->GetGraph()->IsA<URigVMFunctionLibrary>())
+	URigVMGraph* BindingsGraph = FUtils::CompileVariableBindings(InSettings, FUtils::GetAsset<UAnimNextModule>(this));
+	if(BindingsGraph)
 	{
-		if (URigVMGraph* ContainedGraph = InNode->GetContainedGraph())
-		{
-			bool bFunctionGraphExists = false;
-			for (UEdGraph* FunctionGraph : FunctionEdGraphs)
-			{
-				if (URigVMEdGraph* RigFunctionGraph = Cast<URigVMEdGraph>(FunctionGraph))
-				{
-					if (RigFunctionGraph->ModelNodePath == ContainedGraph->GetNodePath())
-					{
-						bFunctionGraphExists = true;
-						break;
-					}
-				}
-			}
-
-			if (!bFunctionGraphExists)
-			{
-				const FName SubGraphName = RigVMClient.GetUniqueName(this, *InNode->GetName());
-				// create a sub graph
-				UAnimNextEdGraph* RigFunctionGraph = NewObject<UAnimNextEdGraph>(this, SubGraphName, RF_Transactional);
-				RigFunctionGraph->Schema = UAnimNextEdGraphSchema::StaticClass();
-				RigFunctionGraph->bAllowRenaming = true;
-				RigFunctionGraph->bEditable = true;
-				RigFunctionGraph->bAllowDeletion = true;
-				RigFunctionGraph->ModelNodePath = ContainedGraph->GetNodePath();
-				RigFunctionGraph->bIsFunctionDefinition = true;
-
-				RigFunctionGraph->Initialize(this);
-
-				FunctionEdGraphs.Add(RigFunctionGraph);
-
-				RigVMClient.GetOrCreateController(ContainedGraph)->ResendAllNotifications();
-			}
-		}
-	}
-	// --- For Collapse nodes ---
-	else if (URigVMEdGraph* RigEdGraph = Cast<URigVMEdGraph>(GetEditorObjectForRigVMGraph(InNode->GetGraph())))
-	{
-		if (URigVMGraph* ContainedGraph = InNode->GetContainedGraph())
-		{
-			bool bSubGraphExists = false;
-
-			const FString ContainedGraphNodePath = ContainedGraph->GetNodePath();
-			for (UEdGraph* SubGraph : RigEdGraph->SubGraphs)
-			{
-				if (UAnimNextEdGraph* SubRigGraph = Cast<UAnimNextEdGraph>(SubGraph))
-				{
-					if (SubRigGraph->ModelNodePath == ContainedGraphNodePath)
-					{
-						bSubGraphExists = true;
-						break;
-					}
-				}
-			}
-
-			if (!bSubGraphExists)
-			{
-				bool bEditable = true;
-				if (InNode->IsA<URigVMAggregateNode>())
-				{
-					bEditable = false;
-				}
-
-				UObject* Outer = FindEntryForRigVMGraph(CollapseNodeGraph->GetRootGraph());
-				if (Outer == nullptr)
-				{
-					Outer = this; // function library graph has no entry
-				}
-
-				const FName SubGraphName = RigVMClient.GetUniqueName(Outer, *InNode->GetEditorSubGraphName());
-				// create a sub graph, no need to set external package if outer is an Entry
-				UAnimNextEdGraph* SubRigGraph = NewObject<UAnimNextEdGraph>(Outer, SubGraphName, RF_Transactional);
-				SubRigGraph->Schema = UAnimNextEdGraphSchema::StaticClass();
-				SubRigGraph->bAllowRenaming = 1;
-				SubRigGraph->bEditable = bEditable;
-				SubRigGraph->bAllowDeletion = 1;
-				SubRigGraph->ModelNodePath = ContainedGraphNodePath;
-				SubRigGraph->bIsFunctionDefinition = false;
-
-				RigEdGraph->SubGraphs.Add(SubRigGraph);
-
-				SubRigGraph->Initialize(this);
-
-				GetOrCreateController(ContainedGraph)->ResendAllNotifications();
-			}
-		}
+		OutGraphs.Add(BindingsGraph);
 	}
 }
 
-void UAnimNextModule_EditorData::RemoveEdGraphForCollapseNode(URigVMCollapseNode* InNode, bool bNotify)
-{
-	check(InNode);
-
-	if (InNode->GetGraph()->IsA<URigVMFunctionLibrary>())
-	{
-		if (URigVMGraph* ContainedGraph = InNode->GetContainedGraph())
-		{
-			for (UEdGraph* FunctionGraph : FunctionEdGraphs)
-			{
-				if (URigVMEdGraph* RigFunctionGraph = Cast<URigVMEdGraph>(FunctionGraph))
-				{
-					if (RigFunctionGraph->ModelNodePath == ContainedGraph->GetNodePath())
-					{
-						if (URigVMController* SubController = GetController(ContainedGraph))
-						{
-							SubController->OnModified().RemoveAll(RigFunctionGraph);
-						}
-
-						if (RigVMGraphModifiedEvent.IsBound() && bNotify)
-						{
-							RigVMGraphModifiedEvent.Broadcast(ERigVMGraphNotifType::NodeRemoved, InNode->GetGraph(), InNode);
-						}
-
-						FunctionEdGraphs.Remove(RigFunctionGraph);
-						RigFunctionGraph->Rename(nullptr, GetTransientPackage(), REN_ForceNoResetLoaders | REN_DontCreateRedirectors);
-						RigFunctionGraph->MarkAsGarbage();
-						break;
-					}
-				}
-			}
-		}
-	}
-	else if (URigVMEdGraph* RigGraph = Cast<URigVMEdGraph>(GetEditorObjectForRigVMGraph(InNode->GetGraph())))
-	{
-		if (URigVMGraph* ContainedGraph = InNode->GetContainedGraph())
-		{
-			for (UEdGraph* SubGraph : RigGraph->SubGraphs)
-			{
-				if (URigVMEdGraph* SubRigGraph = Cast<URigVMEdGraph>(SubGraph))
-				{
-					if (SubRigGraph->ModelNodePath == ContainedGraph->GetNodePath())
-					{
-						if (URigVMController* SubController = GetController(ContainedGraph))
-						{
-							SubController->OnModified().RemoveAll(SubRigGraph);
-						}
-
-						if (RigVMGraphModifiedEvent.IsBound() && bNotify)
-						{
-							RigVMGraphModifiedEvent.Broadcast(ERigVMGraphNotifType::NodeRemoved, InNode->GetGraph(), InNode);
-						}
-
-						RigGraph->SubGraphs.Remove(SubRigGraph);
-						SubRigGraph->Rename(nullptr, GetTransientPackage(), REN_ForceNoResetLoaders | REN_DontCreateRedirectors);
-						SubRigGraph->MarkAsGarbage();
-						break;
-					}
-				}
-			}
-		}
-	}
-}
-
-UEdGraph* UAnimNextModule_EditorData::CreateEdGraph(URigVMGraph* InRigVMGraph, bool bForce)
-{
-	check(InRigVMGraph);
-
-	if(InRigVMGraph->IsA<URigVMFunctionLibrary>())
-	{
-		return nullptr;
-	}
-
-	IAnimNextRigVMGraphInterface* Entry = Cast<IAnimNextRigVMGraphInterface>(FindEntryForRigVMGraph(InRigVMGraph));
-	if(Entry == nullptr)
-	{
-		// Not found, we could be adding a new entry, in which case the graph wont be assigned yet
-		check(Entries.Num() > 0);
-		check(Cast<IAnimNextRigVMGraphInterface>(Entries.Last()) != nullptr);
-		check(Cast<IAnimNextRigVMGraphInterface>(Entries.Last())->GetRigVMGraph() == nullptr);
-		Entry = Cast<IAnimNextRigVMGraphInterface>(FindEntryForRigVMGraph(nullptr));
-	}
-
-	if(Entry == nullptr)
-	{
-		return nullptr;
-	}
-	
-	if(bForce)
-	{
-		RemoveEdGraph(InRigVMGraph);
-	}
-
-	const FName GraphName = RigVMClient.GetUniqueName(CastChecked<UObject>(Entry), Entry->GetGraphName());
-	UAnimNextEdGraph* RigFunctionGraph = NewObject<UAnimNextEdGraph>(CastChecked<UObject>(Entry), GraphName, RF_Transactional);
-	RigFunctionGraph->Schema = UAnimNextEdGraphSchema::StaticClass();
-
-	RigFunctionGraph->Schema = UAnimNextEdGraphSchema::StaticClass();
-	RigFunctionGraph->bAllowDeletion = true;
-	RigFunctionGraph->bIsFunctionDefinition = false;
-	RigFunctionGraph->ModelNodePath = InRigVMGraph->GetNodePath();
-	RigFunctionGraph->Initialize(this);
-
-	Entry->SetEdGraph(RigFunctionGraph);
-	if(Entry->GetRigVMGraph() == nullptr)
-	{
-		Entry->SetRigVMGraph(InRigVMGraph);
-	}
-	else
-	{
-		check(Entry->GetRigVMGraph() == InRigVMGraph);
-	}
-
-	return RigFunctionGraph;
-}
-
-bool UAnimNextModule_EditorData::RemoveEdGraph(URigVMGraph* InModel)
-{
-	if(UAnimNextModule_AnimationGraph* Entry = Cast<UAnimNextModule_AnimationGraph>(FindEntryForRigVMGraph(InModel)))
-	{
-		RigVMClient.DestroyObject(Entry->EdGraph);
-		Entry->EdGraph = nullptr;
-		return true;
-	}
-	return false;
-}
-
-UAnimNextModule_Parameter* UAnimNextModuleLibrary::AddParameter(UAnimNextModule* InModule, FName InName, EPropertyBagPropertyType InValueType,
-	EPropertyBagContainerType InContainerType, const UObject* InValueTypeObject, bool bSetupUndoRedo, bool bPrintPythonCommand)
-{
-	return UE::AnimNext::UncookedOnly::FUtils::GetEditorData(InModule)->AddParameter(InName, FAnimNextParamType(InValueType, InContainerType, InValueTypeObject), bSetupUndoRedo, bPrintPythonCommand);
-}
-
-UAnimNextModule_Parameter* UAnimNextModule_EditorData::AddParameter(FName InName, FAnimNextParamType InType, bool bSetupUndoRedo, bool bPrintPythonCommand)
-{
-	if(InName == NAME_None)
-	{
-		ReportError(TEXT("UAnimNextModule_EditorData::AddParameter: Invalid parameter name supplied."));
-		return nullptr;
-	}
-
-	// Check for duplicate parameter
-	const bool bAlreadyExists = Entries.ContainsByPredicate([InName](const UAnimNextRigVMAssetEntry* InEntry)
-	{
-		if(const UAnimNextModule_Parameter* Parameter = Cast<UAnimNextModule_Parameter>(InEntry))
-		{
-			return Parameter->ParameterName == InName;
-		}
-		return false;
-	});
-
-	if(bAlreadyExists)
-	{
-		ReportError(TEXT("UAnimNextModule_EditorData::AddParameter: A parameter already exists for the supplied parameter name."));
-		return nullptr;
-	}
-
-	UAnimNextModule_Parameter* NewEntry = CreateNewSubEntry<UAnimNextModule_Parameter>(this);
-	NewEntry->ParameterName = InName;
-	NewEntry->Type = InType;
-
-	if(bSetupUndoRedo)
-	{
-		NewEntry->Modify();
-		Modify();
-	}
-	
-	Entries.Add(NewEntry);
-
-	BroadcastModified();
-
-	return NewEntry;
-}
-
-UAnimNextModule_EventGraph* UAnimNextModuleLibrary::AddEventGraph(UAnimNextModule* InModule, FName InName, bool bSetupUndoRedo, bool bPrintPythonCommand)
-{
-	return UE::AnimNext::UncookedOnly::FUtils::GetEditorData(InModule)->AddEventGraph(InName, bSetupUndoRedo, bPrintPythonCommand);
-}
-
-UAnimNextModule_EventGraph* UAnimNextModule_EditorData::AddEventGraph(FName InName, bool bSetupUndoRedo, bool bPrintPythonCommand)
-{
-	if(InName == NAME_None)
-	{
-		ReportError(TEXT("UAnimNextModule_EditorData::AddEventGraph: Invalid graph name supplied."));
-		return nullptr;
-	}
-
-	// Check for duplicate name
-	FName NewGraphName = InName;
-	auto DuplicateNamePredicate = [&NewGraphName](const UAnimNextRigVMAssetEntry* InEntry)
-	{
-		return InEntry->GetEntryName() == NewGraphName;
-	};
-
-	bool bAlreadyExists = Entries.ContainsByPredicate(DuplicateNamePredicate);
-	int32 NameNumber = InName.GetNumber() + 1;
-	while(bAlreadyExists)
-	{
-		NewGraphName = FName(InName, NameNumber++);
-		bAlreadyExists =  Entries.ContainsByPredicate(DuplicateNamePredicate);
-	}
-
-	UAnimNextModule_EventGraph* NewEntry = CreateNewSubEntry<UAnimNextModule_EventGraph>(this);
-	NewEntry->GraphName = NewGraphName;
-
-	if(bSetupUndoRedo)
-	{
-		NewEntry->Modify();
-		Modify();
-	}
-
-	Entries.Add(NewEntry);
-
-	// Add new graph
-	{
-		TGuardValue<bool> EnablePythonPrint(bSuspendPythonMessagesForRigVMClient, !bPrintPythonCommand);
-		TGuardValue<bool> DisableAutoCompile(bAutoRecompileVM, false);
-		URigVMGraph* NewGraph = RigVMClient.AddModel(URigVMGraph::StaticClass()->GetFName(), UAnimNextEventGraphSchema::StaticClass(), bSetupUndoRedo);
-		ensure(NewGraph);
-		NewEntry->Graph = NewGraph;
-
-		URigVMController* Controller = RigVMClient.GetController(NewGraph);
-		UE::AnimNext::UncookedOnly::FUtils::SetupEventGraph(Controller);
-	}
-
-	BroadcastModified();
-
-	return NewEntry;
-}
-
-UAnimNextModule_AnimationGraph* UAnimNextModuleLibrary::AddAnimationGraph(UAnimNextModule* InModule, FName InName, bool bSetupUndoRedo, bool bPrintPythonCommand)
-{
-	return UE::AnimNext::UncookedOnly::FUtils::GetEditorData(InModule)->AddAnimationGraph(InName, bSetupUndoRedo, bPrintPythonCommand);
-}
-
-UAnimNextModule_AnimationGraph* UAnimNextModule_EditorData::AddAnimationGraph(FName InName, bool bSetupUndoRedo, bool bPrintPythonCommand)
-{
-	if(InName == NAME_None)
-	{
-		ReportError(TEXT("UAnimNextModule_EditorData::AddAnimationGraph: Invalid graph name supplied."));
-		return nullptr;
-	}
-
-	// Check for duplicate name
-	FName NewGraphName = InName;
-	auto DuplicateNamePredicate = [&NewGraphName](const UAnimNextRigVMAssetEntry* InEntry)
-	{
-		return InEntry->GetEntryName() == NewGraphName;
-	};
-
-	bool bAlreadyExists = Entries.ContainsByPredicate(DuplicateNamePredicate);
-	int32 NameNumber = InName.GetNumber() + 1;
-	while(bAlreadyExists)
-	{
-		NewGraphName = FName(InName, NameNumber++);
-		bAlreadyExists =  Entries.ContainsByPredicate(DuplicateNamePredicate);
-	}
-
-	UAnimNextModule_AnimationGraph* NewEntry = CreateNewSubEntry<UAnimNextModule_AnimationGraph>(this);
-	NewEntry->GraphName = NewGraphName;
-
-	if(bSetupUndoRedo)
-	{
-		NewEntry->Modify();
-		Modify();
-	}
-	
-	Entries.Add(NewEntry);
-
-	// Add new graph
-	{
-		TGuardValue<bool> EnablePythonPrint(bSuspendPythonMessagesForRigVMClient, !bPrintPythonCommand);
-		TGuardValue<bool> DisableAutoCompile(bAutoRecompileVM, false);
-
-		// Editor data has to be the graph outer, or RigVM unique name generator will not work
-		URigVMGraph* NewRigVMGraphModel = RigVMClient.CreateModel(URigVMGraph::StaticClass()->GetFName(), UAnimNextAnimationGraphSchema::StaticClass(), bSetupUndoRedo, this);
-		if (ensure(NewRigVMGraphModel))
-		{
-			// Then, to avoid the graph losing ref due to external package, set the same package as the Entry
-			if (!NewRigVMGraphModel->HasAnyFlags(RF_Transient))
-			{
-				NewRigVMGraphModel->SetExternalPackage(CastChecked<UObject>(NewEntry)->GetExternalPackage());
-			}
-
-			NewEntry->Graph = NewRigVMGraphModel;
-		
-			RefreshExternalModels();
-			RigVMClient.AddModel(NewRigVMGraphModel, true);
-
-			URigVMController* Controller = RigVMClient.GetController(NewRigVMGraphModel);
-			UE::AnimNext::UncookedOnly::FUtils::SetupAnimGraph(NewEntry, Controller);
-		}
-	}
-
-	BroadcastModified();
-
-	return NewEntry;
-}
diff --git a/Engine/Plugins/Experimental/AnimNext/Source/AnimNextUncookedOnly/Private/Module/AnimNextModule_EventGraphSchema.cpp b/Engine/Plugins/Experimental/AnimNext/Source/AnimNextUncookedOnly/Private/Module/AnimNextModule_EventGraphSchema.cpp
--- a/Engine/Plugins/Experimental/AnimNext/Source/AnimNextUncookedOnly/Private/Module/AnimNextModule_EventGraphSchema.cpp
+++ b/Engine/Plugins/Experimental/AnimNext/Source/AnimNextUncookedOnly/Private/Module/AnimNextModule_EventGraphSchema.cpp
@@ -3,7 +3,8 @@
 #include "AnimNextEventGraphSchema.h"
 
 #include "AnimNextExecuteContext.h"
-#include "Param/RigUnit_AnimNextParameterBase.h"
+#include "Graph/RigUnit_AnimNextBase.h"
+#include "Graph/RigUnit_AnimNextTraitStack.h"
 
 bool UAnimNextEventGraphSchema::SupportsUnitFunction(URigVMController* InController, const FRigVMFunction* InUnitFunction) const
 {
@@ -11,10 +12,13 @@ bool UAnimNextEventGraphSchema::SupportsUnitFunction(URigVMController* InControl
 	{
 		if(FunctionExecuteContextStruct == FAnimNextExecuteContext::StaticStruct())
 		{
-			// Only allow nodes that are children of FRigUnit_AnimNextParameterBase
 			if(InUnitFunction->Struct)
 			{
-				return InUnitFunction->Struct->IsChildOf(FRigUnit_AnimNextParameterBase::StaticStruct());
+				if(InUnitFunction->Struct->IsChildOf(FRigUnit_AnimNextTraitStack::StaticStruct()))
+				{
+					return false;
+				}
+				return InUnitFunction->Struct->IsChildOf(FRigUnit_AnimNextBase::StaticStruct());
 			}
 		}
 	}
