AnimNext: Module refactor

diff --git a/Engine/Plugins/Experimental/AnimNext/Source/AnimNextUncookedOnly/Private/UncookedOnlyUtils.cpp b/Engine/Plugins/Experimental/AnimNext/Source/AnimNextUncookedOnly/Private/UncookedOnlyUtils.cpp
--- a/Engine/Plugins/Experimental/AnimNext/Source/AnimNextUncookedOnly/Private/UncookedOnlyUtils.cpp
+++ b/Engine/Plugins/Experimental/AnimNext/Source/AnimNextUncookedOnly/Private/UncookedOnlyUtils.cpp
@@ -2,946 +2,241 @@
 
 #include "UncookedOnlyUtils.h"
 
-#include "Scheduler/Scheduler.h"
 #include "K2Node_CallFunction.h"
 #include "AssetRegistry/AssetRegistryModule.h"
 #include "AssetRegistry/IAssetRegistry.h"
-#include "Component/AnimNextComponentParameter.h"
 #include "Module/AnimNextModule.h"
-#include "Module/AnimNextModule_Controller.h"
+#include "AnimNextController.h"
 #include "Module/AnimNextModule_EditorData.h"
-#include "Graph/RigUnit_AnimNextTraitStack.h"
 #include "Graph/RigUnit_AnimNextGraphRoot.h"
-#include "Graph/RigUnit_AnimNextGraphEvaluator.h"
-#include "Graph/RigUnit_AnimNextShimRoot.h"
-#include "Graph/RigDecorator_AnimNextCppTrait.h"
-#include "Param/RigUnit_AnimNextParameterBeginExecution.h"
-#include "Param/RigVMDispatch_GetParameter.h"
-#include "Param/RigVMDispatch_SetLayerParameter.h"
-#include "IAnimNextRigVMParameterInterface.h"
-#include "TraitCore/TraitReader.h"
-#include "TraitCore/TraitWriter.h"
-#include "TraitCore/NodeTemplateBuilder.h"
-#include "TraitCore/TraitRegistry.h"
-#include "TraitCore/Trait.h"
-#include "Graph/RigUnit_AnimNextBeginExecution.h"
 #include "Serialization/MemoryReader.h"
-#include "RigVMRuntimeDataRegistry.h"
-#include "Component/AnimNextComponent.h"
-#include "Engine/SCS_Node.h"
-#include "Engine/SimpleConstructionScript.h"
-#include "Param/AnimNextParam.h"
-#include "Param/AnimNextTag.h"
-#include "Param/RigVMDispatch_GetLayerParameter.h"
-#include "RigVMCompiler/RigVMCompiler.h"
 #include "RigVMCore/RigVM.h"
-#include "Scheduler/AnimNextSchedule.h"
-#include "Scheduler/AnimNextSchedulePort.h"
 #include "AnimNextRigVMAsset.h"
 #include "AnimNextRigVMAssetEditorData.h"
-#include "AnimNextRigVMAssetEntry.h"
+#include "Entries/AnimNextRigVMAssetEntry.h"
 #include "AnimNextUncookedOnlyModule.h"
 #include "IAnimNextRigVMExportInterface.h"
 #include "GameFramework/Character.h"
 #include "Graph/AnimNextGraphEntryPoint.h"
 #include "Graph/AnimNextAnimationGraphSchema.h"
 #include "Logging/StructuredLog.h"
-#include "Param/AnimNextParamInstanceIdentifier.h"
-#include "Param/IParameterSourceType.h"
-#include "Param/RigVMDispatch_GetScopedParameter.h"
-#include "Module/AnimNextModuleWorkspaceAssetUserData.h"
+#include "AnimNextAssetWorkspaceAssetUserData.h"
+#include "Misc/EnumerateRange.h"
+#include "Module/RigUnit_AnimNextModuleEvents.h"
+#include "Variables/IVariableBindingType.h"
+#include "Variables/RigUnit_CopyModuleProxyVariables.h"
 #include "WorkspaceAssetRegistryInfo.h"
 
 #define LOCTEXT_NAMESPACE "AnimNextUncookedOnlyUtils"
 
 namespace UE::AnimNext::UncookedOnly
 {
-namespace Private
-{
-	// Represents a trait entry on a node
-	struct FTraitEntryMapping
-	{
-		// The RigVM node that hosts this RigVM decorator
-		const URigVMNode* DecoratorStackNode = nullptr;
-
-		// The RigVM decorator pin on our host node
-		const URigVMPin* DecoratorEntryPin = nullptr;
-
-		// The AnimNext trait
-		const FTrait* Trait = nullptr;
-
-		// A map from latent property names to their corresponding RigVM memory handle index
-		TMap<FName, uint16> LatentPropertyNameToIndexMap;
-
-		FTraitEntryMapping(const URigVMNode* InDecoratorStackNode, const URigVMPin* InDecoratorEntryPin, const FTrait* InTrait)
-			: DecoratorStackNode(InDecoratorStackNode)
-			, DecoratorEntryPin(InDecoratorEntryPin)
-			, Trait(InTrait)
-		{}
-	};
-
-	// Represents a node that contains a trait list
-	struct FTraitStackMapping
-	{
-		// The RigVM node that hosts the RigVM decorators
-		const URigVMNode* DecoratorStackNode = nullptr;
-
-		// The trait list on this node
-		TArray<FTraitEntryMapping> TraitEntries;
-
-		// The node handle assigned to this RigVM node
-		FNodeHandle TraitStackNodeHandle;
-
-		explicit FTraitStackMapping(const URigVMNode* InDecoratorStackNode)
-			: DecoratorStackNode(InDecoratorStackNode)
-		{}
-	};
-
-	struct FTraitGraph
-	{
-		FName EntryPoint;
-		URigVMNode* RootNode;
-		TArray<FTraitStackMapping> TraitStackNodes;
 
-		explicit FTraitGraph(const UAnimNextModule* InModule, URigVMNode* InRootNode)
-			: RootNode(InRootNode)
-		{
-			TStringBuilder<256> StringBuilder;
-			StringBuilder.Append(InModule->GetPathName());
-			StringBuilder.Append(TEXT(":"));
-			StringBuilder.Append(InRootNode->FindPin(GET_MEMBER_NAME_STRING_CHECKED(FRigUnit_AnimNextGraphRoot, EntryPoint))->GetDefaultValue());
-			EntryPoint = FName(StringBuilder.ToView());
-		}
-	};
-
-	template<typename TraitAction>
-	void ForEachTraitInStack(const URigVMNode* DecoratorStackNode, const TraitAction& Action)
+void FUtils::RecreateVM(UAnimNextRigVMAsset* InAsset)
+{
+	if (InAsset->VM == nullptr)
 	{
-		const TArray<URigVMPin*>& Pins = DecoratorStackNode->GetPins();
-		for (URigVMPin* Pin : Pins)
-		{
-			if (!Pin->IsTraitPin())
-			{
-				continue;	// Not a decorator pin
-			}
-
-			if (Pin->GetScriptStruct() == FRigDecorator_AnimNextCppDecorator::StaticStruct())
-			{
-				TSharedPtr<FStructOnScope> DecoratorScope = Pin->GetTraitInstance();
-				FRigDecorator_AnimNextCppDecorator* VMDecorator = (FRigDecorator_AnimNextCppDecorator*)DecoratorScope->GetStructMemory();
-
-				if (const FTrait* Trait = VMDecorator->GetTrait())
-				{
-					Action(DecoratorStackNode, Pin, Trait);
-				}
-			}
-		}
+		InAsset->VM = NewObject<URigVM>(InAsset, TEXT("VM"), RF_NoFlags);
 	}
+	InAsset->VM->Reset(InAsset->ExtendedExecuteContext);
+	InAsset->RigVM = InAsset->VM; // Local serialization
+}
 
-	TArray<FTraitUID> GetTraitUIDs(const URigVMNode* DecoratorStackNode)
-	{
-		TArray<FTraitUID> Traits;
 
-		ForEachTraitInStack(DecoratorStackNode,
-			[&Traits](const URigVMNode* DecoratorStackNode, const URigVMPin* DecoratorPin, const FTrait* Trait)
-			{
-				Traits.Add(Trait->GetTraitUID());
-			});
+void FUtils::CompileVariables(UAnimNextRigVMAsset* InAsset)
+{
+	check(InAsset);
 
-		return Traits;
-	}
+	UAnimNextRigVMAssetEditorData* EditorData = GetEditorData(InAsset);
 
-	FNodeHandle RegisterTraitNodeTemplate(FTraitWriter& TraitWriter, const URigVMNode* DecoratorStackNode)
+	struct FStructEntryInfo
 	{
-		const TArray<FTraitUID> TraitUIDs = GetTraitUIDs(DecoratorStackNode);
-
-		TArray<uint8> NodeTemplateBuffer;
-		const FNodeTemplate* NodeTemplate = FNodeTemplateBuilder::BuildNodeTemplate(TraitUIDs, NodeTemplateBuffer);
+		FName Name;
+		FAnimNextParamType Type;
+		EAnimNextExportAccessSpecifier AccessSpecifier;
+		TConstArrayView<uint8> Value;
+	};
 
-		return TraitWriter.RegisterNode(*NodeTemplate);
-	}
+	TArray<FStructEntryInfo> StructEntryInfos;
+	StructEntryInfos.Reserve(EditorData->Entries.Num());
 
-	FString GetTraitProperty(const FTraitStackMapping& TraitStack, uint32 TraitIndex, FName PropertyName, const TArray<FTraitStackMapping>& TraitStackNodes)
+	// Gather all parameters in this asset
+	int32 NumPublicVariables = 0;
+	for(const UAnimNextRigVMAssetEntry* Entry : EditorData->Entries)
 	{
-		const TArray<URigVMPin*>& Pins = TraitStack.TraitEntries[TraitIndex].DecoratorEntryPin->GetSubPins();
-		for (const URigVMPin* Pin : Pins)
+		const IAnimNextRigVMExportInterface* Export = Cast<IAnimNextRigVMExportInterface>(Entry);
+		const IAnimNextRigVMVariableInterface* Variable = Cast<IAnimNextRigVMVariableInterface>(Entry);
+		if(Export && Variable)
 		{
-			if (Pin->GetDirection() != ERigVMPinDirection::Input)
-			{
-				continue;	// We only look for input pins
-			}
-
-			if (Pin->GetFName() == PropertyName)
+			const FAnimNextParamType& Type = Export->GetExportType();
+			ensure(Type.IsValid());
+			const FName Name = Export->GetExportName();
+			const EAnimNextExportAccessSpecifier AccessSpecifier = Export->GetExportAccessSpecifier();
+			if(AccessSpecifier == EAnimNextExportAccessSpecifier::Public)
 			{
-				if (Pin->GetCPPTypeObject() == FAnimNextTraitHandle::StaticStruct())
-				{
-					// Trait handle pins don't have a value, just an optional link
-					const TArray<URigVMLink*>& PinLinks = Pin->GetLinks();
-					if (!PinLinks.IsEmpty())
-					{
-						// Something is connected to us, find the corresponding node handle so that we can encode it as our property value
-						check(PinLinks.Num() == 1);
-
-						const URigVMNode* SourceNode = PinLinks[0]->GetSourceNode();
-
-						FNodeHandle SourceNodeHandle;
-						int32 SourceTraitIndex = INDEX_NONE;
-
-						const FTraitStackMapping* SourceTraitStack = TraitStackNodes.FindByPredicate([SourceNode](const FTraitStackMapping& Mapping) { return Mapping.DecoratorStackNode == SourceNode; });
-						if (SourceTraitStack != nullptr)
-						{
-							SourceNodeHandle = SourceTraitStack->TraitStackNodeHandle;
-
-							// If the source pin is null, we are a node where the result pin lives on the stack node instead of a decorator sub-pin
-							// If this is the case, we bind to the first trait index since we only allowed a single base trait per stack
-							// Otherwise we lookup the trait index we are linked to
-							const URigVMPin* SourceDecoratorPin = PinLinks[0]->GetSourcePin()->GetParentPin();
-							SourceTraitIndex = SourceDecoratorPin != nullptr ? SourceTraitStack->DecoratorStackNode->GetTraitPins().IndexOfByKey(SourceDecoratorPin) : 0;
-						}
-
-						if (SourceNodeHandle.IsValid())
-						{
-							check(SourceTraitIndex != INDEX_NONE);
-
-							const FAnimNextTraitHandle TraitHandle(SourceNodeHandle, SourceTraitIndex);
-							const FAnimNextTraitHandle DefaultTraitHandle;
-
-							// We need an instance of a trait handle property to be able to serialize it into text, grab it from the root
-							const FProperty* Property = FRigUnit_AnimNextGraphRoot::StaticStruct()->FindPropertyByName(GET_MEMBER_NAME_STRING_CHECKED(FRigUnit_AnimNextGraphRoot, Result));
-
-							FString PropertyValue;
-							Property->ExportText_Direct(PropertyValue, &TraitHandle, &DefaultTraitHandle, nullptr, PPF_SerializedAsImportText);
-
-							return PropertyValue;
-						}
-					}
-
-					// This handle pin isn't connected
-					return FString();
-				}
-
-				// A regular property pin
-				return Pin->GetDefaultValue();
+				NumPublicVariables++;
 			}
-		}
 
-		// Unknown property
-		return FString();
-	}
-
-	uint16 GetTraitLatentPropertyIndex(const FTraitStackMapping& TraitStack, uint32 TraitIndex, FName PropertyName)
-	{
-		const FTraitEntryMapping& Entry = TraitStack.TraitEntries[TraitIndex];
-		if (const uint16* RigVMIndex = Entry.LatentPropertyNameToIndexMap.Find(PropertyName))
-		{
-			return *RigVMIndex;
-		}
-
-		return MAX_uint16;
-	}
-
-	void WriteTraitProperties(FTraitWriter& TraitWriter, const FTraitStackMapping& Mapping, const TArray<FTraitStackMapping>& TraitStackNodes)
-	{
-		TraitWriter.WriteNode(Mapping.TraitStackNodeHandle,
-			[&Mapping, &TraitStackNodes](uint32 TraitIndex, FName PropertyName)
-			{
-				return GetTraitProperty(Mapping, TraitIndex, PropertyName, TraitStackNodes);
-			},
-			[&Mapping](uint32 TraitIndex, FName PropertyName)
-			{
-				return GetTraitLatentPropertyIndex(Mapping, TraitIndex, PropertyName);
-			});
-	}
-
-	URigVMUnitNode* FindRootNode(const TArray<URigVMNode*>& VMNodes)
-	{
-		for (URigVMNode* VMNode : VMNodes)
-		{
-			if (URigVMUnitNode* VMUnitNode = Cast<URigVMUnitNode>(VMNode))
-			{
-				const UScriptStruct* ScriptStruct = VMUnitNode->GetScriptStruct();
-				if (ScriptStruct == FRigUnit_AnimNextGraphRoot::StaticStruct())
+			StructEntryInfos.Add(
 				{
-					return VMUnitNode;
-				}
-			}
+					Name,
+					FAnimNextParamType(Type.GetValueType(),  Type.GetContainerType(), Type.GetValueTypeObject()),
+					AccessSpecifier,
+					TConstArrayView<uint8>(Variable->GetValuePtr(), Type.GetSize())
+				});
 		}
-
-		return nullptr;
 	}
 
-	void AddMissingInputLinks(const URigVMPin* DecoratorPin, URigVMController* VMController)
+	// Sort public entries first & then by size, largest first, for better packing
+	static_assert(EAnimNextExportAccessSpecifier::Private < EAnimNextExportAccessSpecifier::Public, "Private must be less than Public as parameters are sorted internally according to this assumption");
+	StructEntryInfos.Sort([](const FStructEntryInfo& InLHS, const FStructEntryInfo& InRHS)
 	{
-		const TArray<URigVMPin*>& Pins = DecoratorPin->GetSubPins();
-		for (URigVMPin* Pin : Pins)
+		if(InLHS.AccessSpecifier != InRHS.AccessSpecifier)
 		{
-			const ERigVMPinDirection PinDirection = Pin->GetDirection();
-			if (PinDirection != ERigVMPinDirection::Input && PinDirection != ERigVMPinDirection::Hidden)
-			{
-				continue;	// We only look for hidden or input pins
-			}
-
-			if (Pin->GetCPPTypeObject() != FAnimNextTraitHandle::StaticStruct())
-			{
-				continue;	// We only look for trait handle pins
-			}
-
-			const TArray<URigVMLink*>& PinLinks = Pin->GetLinks();
-			if (!PinLinks.IsEmpty())
-			{
-				continue;	// This pin already has a link, all good
-			}
-
-			// Add a dummy node that will output a reference pose to ensure every link is valid.
-			// RigVM doesn't let us link two decorators on a same node together or linking a child back to a parent
-			// as this would create a cycle in the RigVM graph. The AnimNext graph traits do support it
-			// and so perhaps we could have a merging pass later on to remove useless dummy nodes like this.
-
-			URigVMUnitNode* VMReferencePoseNode = VMController->AddUnitNode(FRigUnit_AnimNextTraitStack::StaticStruct(), FRigVMStruct::ExecuteName, FVector2D(0.0f, 0.0f), FString(), false);
-			check(VMReferencePoseNode != nullptr);
-
-			const UScriptStruct* CppDecoratorStruct = FRigDecorator_AnimNextCppDecorator::StaticStruct();
-
-			FString DefaultValue;
-			{
-				const UE::AnimNext::FTraitUID ReferencePoseTraitUID(0x7508ab89);	// Trait header is private, reference by UID directly
-				const FTrait* Trait = FTraitRegistry::Get().Find(ReferencePoseTraitUID);
-				check(Trait != nullptr);
-
-				const FRigDecorator_AnimNextCppDecorator DefaultCppDecoratorStructInstance;
-				FRigDecorator_AnimNextCppDecorator CppDecoratorStructInstance;
-				CppDecoratorStructInstance.DecoratorSharedDataStruct = Trait->GetTraitSharedDataStruct();
-
-				const FProperty* Prop = FAnimNextCppDecoratorWrapper::StaticStruct()->FindPropertyByName(GET_MEMBER_NAME_STRING_CHECKED(FAnimNextCppDecoratorWrapper, CppDecorator));
-				check(Prop != nullptr);
-
-				Prop->ExportText_Direct(DefaultValue, &CppDecoratorStructInstance, &DefaultCppDecoratorStructInstance, nullptr, PPF_SerializedAsImportText);
-			}
-
-			const FName ReferencePoseDecoratorName = VMController->AddTrait(VMReferencePoseNode->GetFName(), *CppDecoratorStruct->GetPathName(), TEXT("ReferencePose"), DefaultValue, INDEX_NONE, false, false);
-			check(!ReferencePoseDecoratorName.IsNone());
-
-			URigVMPin* OutputPin = VMReferencePoseNode->FindPin(GET_MEMBER_NAME_STRING_CHECKED(FRigUnit_AnimNextTraitStack, Result));
-			check(OutputPin != nullptr);
-
-			ensure(VMController->AddLink(OutputPin, Pin, false));
+			return InLHS.AccessSpecifier > InRHS.AccessSpecifier;
 		}
-	}
-
-	void AddMissingInputLinks(const URigVMGraph* VMGraph, URigVMController* VMController)
-	{
-		const TArray<URigVMNode*> VMNodes = VMGraph->GetNodes();	// Copy since we might add new nodes
-		for (URigVMNode* VMNode : VMNodes)
+		else
 		{
-			if (const URigVMUnitNode* VMUnitNode = Cast<URigVMUnitNode>(VMNode))
-			{
-				const UScriptStruct* ScriptStruct = VMUnitNode->GetScriptStruct();
-				if (ScriptStruct != FRigUnit_AnimNextTraitStack::StaticStruct())
-				{
-					continue;	// Skip non-trait nodes
-				}
-
-				ForEachTraitInStack(VMNode,
-					[VMController](const URigVMNode* DecoratorStackNode, const URigVMPin* DecoratorPin, const FTrait* Trait)
-					{
-						AddMissingInputLinks(DecoratorPin, VMController);
-					});
-			}
+			return InLHS.Type.GetSize() > InRHS.Type.GetSize();
 		}
-	}
+	});
 
-	FTraitGraph CollectGraphInfo(const UAnimNextModule* InModule, const URigVMGraph* VMGraph, URigVMController* VMController)
+	if(StructEntryInfos.Num() > 0)
 	{
-		const TArray<URigVMNode*>& VMNodes = VMGraph->GetNodes();
-		URigVMUnitNode* VMRootNode = FindRootNode(VMNodes);
+		// Build PropertyDescs and values to batch-create the property bag
+		TArray<FPropertyBagPropertyDesc> PropertyDescs;
+		PropertyDescs.Reserve(StructEntryInfos.Num());
+		TArray<TConstArrayView<uint8>> Values;
+		Values.Reserve(StructEntryInfos.Num());
 
-		if (VMRootNode == nullptr)
+		for (TEnumerateRef<const FStructEntryInfo> StructEntryInfo : EnumerateRange(StructEntryInfos))
 		{
-			// Root node wasn't found, add it, we'll need it to compile
-			VMRootNode = VMController->AddUnitNode(FRigUnit_AnimNextGraphRoot::StaticStruct(), FRigUnit_AnimNextGraphRoot::EventName, FVector2D(0.0f, 0.0f), FString(), false);
+			PropertyDescs.Emplace(StructEntryInfo->Name, StructEntryInfo->Type.ContainerType, StructEntryInfo->Type.ValueType, StructEntryInfo->Type.ValueTypeObject);
+			Values.Add(StructEntryInfo->Value);
 		}
 
-		// Make sure we don't have empty input pins
-		AddMissingInputLinks(VMGraph, VMController);
-
-		FTraitGraph TraitGraph(InModule, VMRootNode);
+		// Create new property bags and migrate
+		EPropertyBagResult Result = InAsset->VariableDefaults.ReplaceAllPropertiesAndValues(PropertyDescs, Values);
+		check(Result == EPropertyBagResult::Success);
 
-		TArray<const URigVMNode*> NodesToVisit;
-		NodesToVisit.Add(VMRootNode);
-
-		while (NodesToVisit.Num() != 0)
+		if(NumPublicVariables > 0)
 		{
-			const URigVMNode* VMNode = NodesToVisit[0];
-			NodesToVisit.RemoveAt(0);
-
-			if (const URigVMUnitNode* VMUnitNode = Cast<URigVMUnitNode>(VMNode))
-			{
-				const UScriptStruct* ScriptStruct = VMUnitNode->GetScriptStruct();
-				if (ScriptStruct == FRigUnit_AnimNextTraitStack::StaticStruct())
-				{
-					FTraitStackMapping Mapping(VMNode);
-					ForEachTraitInStack(VMNode,
-						[&Mapping](const URigVMNode* DecoratorStackNode, const URigVMPin* DecoratorPin, const FTrait* Trait)
-						{
-							Mapping.TraitEntries.Add(FTraitEntryMapping(DecoratorStackNode, DecoratorPin, Trait));
-						});
-
-					TraitGraph.TraitStackNodes.Add(MoveTemp(Mapping));
-				}
-			}
-
-			const TArray<URigVMNode*> SourceNodes = VMNode->GetLinkedSourceNodes();
-			NodesToVisit.Append(SourceNodes);
+			TConstArrayView<FPropertyBagPropertyDesc> PublicPropertyDescs(PropertyDescs.GetData(), NumPublicVariables);
+			TConstArrayView<TConstArrayView<uint8>> PublicValues(Values.GetData(), NumPublicVariables);
+			Result = InAsset->PublicVariableDefaults.ReplaceAllPropertiesAndValues(PublicPropertyDescs, PublicValues);
+			check(Result == EPropertyBagResult::Success);
 		}
-
-		if (TraitGraph.TraitStackNodes.IsEmpty())
+		else
 		{
-			// If the graph is empty, add a dummy node that just pushes a reference pose
-			URigVMUnitNode* VMNode = VMController->AddUnitNode(FRigUnit_AnimNextTraitStack::StaticStruct(), FRigVMStruct::ExecuteName, FVector2D(0.0f, 0.0f), FString(), false);
-
-			const UScriptStruct* CppDecoratorStruct = FRigDecorator_AnimNextCppDecorator::StaticStruct();
-
-			FString DefaultValue;
-			{
-				const UE::AnimNext::FTraitUID ReferencePoseTraitUID(0x7508ab89);	// Trait header is private, reference by UID directly
-				const FTrait* Trait = FTraitRegistry::Get().Find(ReferencePoseTraitUID);
-				check(Trait != nullptr);
-
-				const FRigDecorator_AnimNextCppDecorator DefaultCppDecoratorStructInstance;
-				FRigDecorator_AnimNextCppDecorator CppDecoratorStructInstance;
-				CppDecoratorStructInstance.DecoratorSharedDataStruct = Trait->GetTraitSharedDataStruct();
-
-				const FProperty* Prop = FAnimNextCppDecoratorWrapper::StaticStruct()->FindPropertyByName(GET_MEMBER_NAME_STRING_CHECKED(FAnimNextCppDecoratorWrapper, CppDecorator));
-				check(Prop != nullptr);
-
-				Prop->ExportText_Direct(DefaultValue, &CppDecoratorStructInstance, &DefaultCppDecoratorStructInstance, nullptr, PPF_SerializedAsImportText);
-			}
-
-			VMController->AddTrait(VMNode->GetFName(), *CppDecoratorStruct->GetPathName(), TEXT("ReferencePose"), DefaultValue, INDEX_NONE, false, false);
-
-			FTraitStackMapping Mapping(VMNode);
-			ForEachTraitInStack(VMNode,
-				[&Mapping](const URigVMNode* DecoratorStackNode, const URigVMPin* DecoratorPin, const FTrait* Trait)
-				{
-					Mapping.TraitEntries.Add(FTraitEntryMapping(DecoratorStackNode, DecoratorPin, Trait));
-				});
-
-			TraitGraph.TraitStackNodes.Add(MoveTemp(Mapping));
+			InAsset->PublicVariableDefaults.Reset();
 		}
 
-		return TraitGraph;
+		// Rebuild external variables
+		InAsset->VM->SetExternalVariableDefs(InAsset->GetExternalVariablesImpl(false));
 	}
-
-	void CollectLatentPins(TArray<FTraitStackMapping>& TraitStackNodes, FRigVMPinInfoArray& OutLatentPins, TMap<FName, URigVMPin*>& OutLatentPinMapping)
-	{
-		for (FTraitStackMapping& TraitStack : TraitStackNodes)
-		{
-			for (FTraitEntryMapping& TraitEntry : TraitStack.TraitEntries)
-			{
-				for (URigVMPin* Pin : TraitEntry.DecoratorEntryPin->GetSubPins())
-				{
-					if (Pin->IsLazy() && !Pin->GetLinks().IsEmpty())
-					{
-						// This pin has something linked to it, it is a latent pin
-						check(OutLatentPins.Num() < ((1 << 16) - 1));	// We reserve MAX_uint16 as an invalid value and we must fit on 15 bits when packed
-						TraitEntry.LatentPropertyNameToIndexMap.Add(Pin->GetFName(), (uint16)OutLatentPins.Num());
-
-						const FName LatentPinName(TEXT("LatentPin"), OutLatentPins.Num());	// Create unique latent pin names
-
-						FRigVMPinInfo PinInfo;
-						PinInfo.Name = LatentPinName;
-						PinInfo.TypeIndex = Pin->GetTypeIndex();
-
-						// All our programmatic pins are lazy inputs
-						PinInfo.Direction = ERigVMPinDirection::Input;
-						PinInfo.bIsLazy = true;
-
-						OutLatentPins.Pins.Emplace(PinInfo);
-
-						const TArray<URigVMLink*>& PinLinks = Pin->GetLinks();
-						check(PinLinks.Num() == 1);
-
-						OutLatentPinMapping.Add(LatentPinName, PinLinks[0]->GetSourcePin());
-					}
-				}
-			}
-		}
-	}
-
-	FAnimNextGraphEvaluatorExecuteDefinition GetGraphEvaluatorExecuteMethod(const FRigVMPinInfoArray& LatentPins)
+	else
 	{
-		const uint32 LatentPinListHash = GetTypeHash(LatentPins);
-		if (const FAnimNextGraphEvaluatorExecuteDefinition* ExecuteDefinition = FRigUnit_AnimNextGraphEvaluator::FindExecuteMethod(LatentPinListHash))
-		{
-			return *ExecuteDefinition;
-		}
-
-		const FRigVMRegistry& Registry = FRigVMRegistry::Get();
-
-		// Generate a new method for this argument list
-		FAnimNextGraphEvaluatorExecuteDefinition ExecuteDefinition;
-		ExecuteDefinition.Hash = LatentPinListHash;
-		ExecuteDefinition.MethodName = FString::Printf(TEXT("Execute_%X"), LatentPinListHash);
-		ExecuteDefinition.Arguments.Reserve(LatentPins.Num());
-
-		for (const FRigVMPinInfo& Pin : LatentPins)
-		{
-			const FRigVMTemplateArgumentType& TypeArg = Registry.GetType(Pin.TypeIndex);
-
-			FAnimNextGraphEvaluatorExecuteArgument Argument;
-			Argument.Name = Pin.Name.ToString();
-			Argument.CPPType = TypeArg.CPPType.ToString();
-
-			ExecuteDefinition.Arguments.Add(Argument);
-		}
-
-		FRigUnit_AnimNextGraphEvaluator::RegisterExecuteMethod(ExecuteDefinition);
-
-		return ExecuteDefinition;
+		InAsset->VariableDefaults.Reset();
+		InAsset->PublicVariableDefaults.Reset();
+		InAsset->VM->ClearExternalVariables(InAsset->ExtendedExecuteContext);
 	}
 }
 
-void FUtils::Compile(UAnimNextModule* InModule)
-{
-	check(InModule);
-
-	FMessageLog("AnimNextCompilerResults").NewPage(FText::FromName(InModule->GetFName()));
-
-	CompileStruct(InModule);
-	CompileVM(InModule);
-}
-
-void FUtils::CompileVM(UAnimNextModule* InModule)
+URigVMGraph* FUtils::CompileVariableBindings(const FRigVMCompileSettings& InSettings, UAnimNextRigVMAsset* InAsset)
 {
-	check(InModule);
-
-	UAnimNextModule_EditorData* EditorData = GetEditorData(InModule);
-
-	if (EditorData->bIsCompiling)
-	{
-		return;
-	}
-
-	TGuardValue<bool> CompilingGuard(EditorData->bIsCompiling, true);
-	
-	// Before we re-compile a graph, we need to release and live instances since we need the metadata we are about to replace
-	// to call trait destructors etc
-	InModule->FreezeGraphInstances();
-
-	EditorData->bErrorsDuringCompilation = false;
-
-	EditorData->RigGraphDisplaySettings.MinMicroSeconds = EditorData->RigGraphDisplaySettings.LastMinMicroSeconds = DBL_MAX;
-	EditorData->RigGraphDisplaySettings.MaxMicroSeconds = EditorData->RigGraphDisplaySettings.LastMaxMicroSeconds = (double)INDEX_NONE;
-
-	TGuardValue<bool> ReentrantGuardSelf(EditorData->bSuspendModelNotificationsForSelf, true);
-	TGuardValue<bool> ReentrantGuardOthers(EditorData->bSuspendModelNotificationsForOthers, true);
-
-	RecreateVM(InModule);
-
-	InModule->VMRuntimeSettings = EditorData->VMRuntimeSettings;
-	InModule->EntryPoints.Empty();
-	InModule->ResolvedRootTraitHandles.Empty();
-	InModule->ResolvedEntryPoints.Empty();
-	InModule->ExecuteDefinition = FAnimNextGraphEvaluatorExecuteDefinition();
-	InModule->SharedDataBuffer.Empty();
-	InModule->GraphReferencedObjects.Empty();
-	InModule->RequiredParametersHash = 0;
-	InModule->RequiredParameters.Empty();
-
-	FRigVMClient* VMClient = EditorData->GetRigVMClient();
-	URigVMGraph* VMRootGraph = VMClient->GetDefaultModel();
-
-	if(VMRootGraph == nullptr)
-	{
-		return;
-	}
-
-	TArray<URigVMGraph*> TempGraphs;
-	for(const URigVMGraph* SourceGraph : VMClient->GetAllModels(false, false))
-	{
-		// We use a temporary graph models to build our final graphs that we'll compile
-		URigVMGraph* TempGraph = CastChecked<URigVMGraph>(StaticDuplicateObject(SourceGraph, GetTransientPackage(), NAME_None, RF_Transient));
-		TempGraph->SetFlags(RF_Transient);
-		TempGraphs.Add(TempGraph);
-	}
-
-	if(TempGraphs.Num() == 0)
-	{
-		return;
-	}
-
-	UAnimNextModule_Controller* TempController = CastChecked<UAnimNextModule_Controller>(VMClient->GetOrCreateController(TempGraphs[0]));
-
-	FTraitWriter TraitWriter;
+	check(InAsset);
 
-	FRigVMPinInfoArray LatentPins;
-	TMap<FName, URigVMPin*> LatentPinMapping;
-	TArray<Private::FTraitGraph> TraitGraphs;
+	FModule& Module = FModuleManager::LoadModuleChecked<FModule>("AnimNextUncookedOnly");
+	UAnimNextRigVMAssetEditorData* EditorData = GetEditorData(InAsset);
+	TMap<IVariableBindingType*, TArray<IVariableBindingType::FBindingGraphInput>> BindingGroups;
 
-	// Build entry points and extract their required latent pins
-	for(const URigVMGraph* TempGraph : TempGraphs)
+	for(const UAnimNextRigVMAssetEntry* Entry : EditorData->Entries)
 	{
-		if(TempGraph->GetSchemaClass() == UAnimNextAnimationGraphSchema::StaticClass())
+		const IAnimNextRigVMVariableInterface* Variable = Cast<IAnimNextRigVMVariableInterface>(Entry);
+		if(Variable == nullptr)
 		{
-			// Gather our trait stacks
-			Private::FTraitGraph& TraitGraph = TraitGraphs.Add_GetRef(Private::CollectGraphInfo(InModule, TempGraph, TempController->GetControllerForGraph(TempGraph)));
-			check(!TraitGraph.TraitStackNodes.IsEmpty());
-
-			FAnimNextGraphEntryPoint& EntryPoint = InModule->EntryPoints.AddDefaulted_GetRef();
-			EntryPoint.EntryPointName = TraitGraph.EntryPoint;
-
-			// Extract latent pins for this graph
-			Private::CollectLatentPins(TraitGraph.TraitStackNodes, LatentPins, LatentPinMapping);
-
-			// Iterate over every trait stack and register our node templates
-			for (Private::FTraitStackMapping& NodeMapping : TraitGraph.TraitStackNodes)
-			{
-				NodeMapping.TraitStackNodeHandle = Private::RegisterTraitNodeTemplate(TraitWriter, NodeMapping.DecoratorStackNode);
-			}
-
-			// Find our root node handle, if we have any stack nodes, the first one is our root stack
-			if (TraitGraph.TraitStackNodes.Num() != 0)
-			{
-				EntryPoint.RootTraitHandle = FAnimNextEntryPointHandle(TraitGraph.TraitStackNodes[0].TraitStackNodeHandle);
-			}
+			continue;
 		}
-	}
-
-	// Remove our old root nodes
-	for (Private::FTraitGraph& TraitGraph : TraitGraphs)
-	{
-		URigVMController* GraphController = TempController->GetControllerForGraph(TraitGraph.RootNode->GetGraph());
-		GraphController->RemoveNode(TraitGraph.RootNode, false, false);
-	}
-
-	if(LatentPins.Num() > 0)
-	{
-		// We need a unique method name to match our unique argument list
-		InModule->ExecuteDefinition = Private::GetGraphEvaluatorExecuteMethod(LatentPins);
 
-		// Add our runtime shim root node
-		URigVMUnitNode* TempShimRootNode = TempController->AddUnitNode(FRigUnit_AnimNextShimRoot::StaticStruct(), FRigUnit_AnimNextShimRoot::EventName, FVector2D::ZeroVector, FString(), false);
-		URigVMUnitNode* GraphEvaluatorNode = TempController->AddUnitNodeWithPins(FRigUnit_AnimNextGraphEvaluator::StaticStruct(), LatentPins, *InModule->ExecuteDefinition.MethodName, FVector2D::ZeroVector, FString(), false);
-
-		// Link our shim and evaluator nodes together using the execution context
-		TempController->AddLink(
-			TempShimRootNode->FindPin(GET_MEMBER_NAME_STRING_CHECKED(FRigUnit_AnimNextShimRoot, ExecuteContext)),
-			GraphEvaluatorNode->FindPin(GET_MEMBER_NAME_STRING_CHECKED(FRigUnit_AnimNextGraphEvaluator, ExecuteContext)),
-			false);
-
-		// Link our latent pins
-		for (const FRigVMPinInfo& LatentPin : LatentPins)
+		const FAnimNextVariableBinding& Binding = Variable->GetBinding();
+		if(!Binding.IsValid())
 		{
-			TempController->AddLink(
-				LatentPinMapping[LatentPin.Name],
-				GraphEvaluatorNode->FindPin(LatentPin.Name.ToString()),
-				false);
+			continue;
 		}
-	}
-
-	// Write our node shared data
-	TraitWriter.BeginNodeWriting();
 
-	for(Private::FTraitGraph& TraitGraph : TraitGraphs)
-	{
-		for (const Private::FTraitStackMapping& NodeMapping : TraitGraph.TraitStackNodes)
+		TSharedPtr<IVariableBindingType> BindingType = Module.FindVariableBindingType(Binding.BindingData.GetScriptStruct());
+		if(!BindingType.IsValid())
 		{
-			Private::WriteTraitProperties(TraitWriter, NodeMapping, TraitGraph.TraitStackNodes);
+			continue;
 		}
-	}
-
-	TraitWriter.EndNodeWriting();
 
-	// Cache our compiled metadata
-	InModule->SharedDataArchiveBuffer = TraitWriter.GetGraphSharedData();
-	InModule->GraphReferencedObjects = TraitWriter.GetGraphReferencedObjects();
-
-	// Populate our runtime metadata
-	InModule->LoadFromArchiveBuffer(InModule->SharedDataArchiveBuffer);
-
-	URigVMCompiler* Compiler = URigVMCompiler::StaticClass()->GetDefaultObject<URigVMCompiler>();
-	EditorData->VMCompileSettings.SetExecuteContextStruct(FAnimNextExecuteContext::StaticStruct());
-	FRigVMCompileSettings Settings = (EditorData->bCompileInDebugMode) ? FRigVMCompileSettings::Fast(EditorData->VMCompileSettings.GetExecuteContextStruct()) : EditorData->VMCompileSettings;
-	Settings.ASTSettings.bSetupTraits = false; // disable the default implementation of decorators for now
-	Settings.ASTSettings.ReportDelegate.BindLambda([InModule](EMessageSeverity::Type InType, UObject* InObject, const FString& InString)
-	{
-		FMessageLog("AnimNextCompilerResults").Message(InType, FText::FromString(InString));
-	});
-
-	Compiler->Compile(Settings, TempGraphs, TempController, InModule->VM, InModule->ExtendedExecuteContext, TArray<FRigVMExternalVariable>(), & EditorData->PinToOperandMap);
-
-	// Initialize right away, in packaged builds we initialize during PostLoad
-	InModule->VM->Initialize(InModule->ExtendedExecuteContext);
-	InModule->GenerateUserDefinedDependenciesData(InModule->ExtendedExecuteContext);
-
-	// Notable difference with vanilla RigVM host behavior - we init the VM here at the moment as we only have one 'instance'
-	InModule->InitializeVM(FRigUnit_AnimNextBeginExecution::EventName);
-
-	if (EditorData->bErrorsDuringCompilation)
-	{
-		if(Settings.SurpressErrors)
-		{
-			Settings.Reportf(EMessageSeverity::Info, InModule, TEXT("Compilation Errors may be suppressed for AnimNext Interface Graph: %s. See VM Compile Settings for more Details"), *InModule->GetName());
-		}
+		TArray<IVariableBindingType::FBindingGraphInput>& Group = BindingGroups.FindOrAdd(BindingType.Get());
+		FRigVMTemplateArgumentType RigVMArg = Variable->GetType().ToRigVMTemplateArgument();
+		Group.Add({ Variable->GetVariableName(), RigVMArg.CPPType.ToString(), RigVMArg.CPPTypeObject, Binding.BindingData });
 	}
 
-	EditorData->bVMRecompilationRequired = false;
-	if(InModule->VM)
+	const bool bHasBindings = BindingGroups.Num() > 0;
+	const bool bHasPublicVariablesToCopy = EditorData->IsA<UAnimNextModule_EditorData>() && EditorData->HasPublicVariables();
+	if(!bHasBindings && !bHasPublicVariablesToCopy)
 	{
-		EditorData->RigVMCompiledEvent.Broadcast(InModule, InModule->VM, InModule->ExtendedExecuteContext);
+		// Nothing to do here
+		return nullptr;
 	}
 
-	for(URigVMGraph* TempGraph : TempGraphs)
-	{
-		VMClient->RemoveController(TempGraph);
-	}
+	URigVMGraph* BindingGraph = NewObject<URigVMGraph>(EditorData, NAME_None, RF_Transient);
 
-	// Now that the graph has been re-compiled, re-allocate the previous live instances
-	InModule->ThawGraphInstances();
-	
-	FAnimNextParameterProviderAssetRegistryExports Exports;
-	GetAssetParameters(EditorData, Exports);
-
-	for(FAnimNextParameterAssetRegistryExportEntry& Entry : Exports.Parameters)
+	FRigVMClient* VMClient = EditorData->GetRigVMClient();
+	URigVMController* Controller = VMClient->GetOrCreateController(BindingGraph);
+	URigVMNode* ExcecuteBindingsNode = Controller->AddUnitNode(FRigUnit_AnimNextExecuteBindings::StaticStruct(), FRigUnit_AnimNextExecuteBindings::GetMethodName(), FVector2D::ZeroVector, FString(), false);
+	if(ExcecuteBindingsNode == nullptr)
 	{
-		// Required parameters are those that are read in this asset but not declared in this asset as state
-		if(EnumHasAllFlags(Entry.GetFlags(), EAnimNextParameterFlags::Read) && !EnumHasAnyFlags(Entry.GetFlags(), EAnimNextParameterFlags::Declared))
-		{
-			InModule->RequiredParameters.Emplace(Entry.Name, Entry.Type, Entry.InstanceId);
-		}
+		InSettings.ReportError(TEXT("Could not spawn Execute Bindings node"));
+		return nullptr;
 	}
-
-	InModule->RequiredParametersHash = SortAndHashParameters(InModule->RequiredParameters);
-
-#if WITH_EDITOR
-//	RefreshBreakpoints(EditorData);
-#endif
-}
-
-void FUtils::RecreateVM(UAnimNextModule* InModule)
-{
-	if (InModule->VM == nullptr)
+	URigVMPin* ExecuteBindingsExecPin = ExcecuteBindingsNode->FindPin(FRigVMStruct::ExecuteContextName.ToString());
+	if(ExecuteBindingsExecPin == nullptr)
 	{
-		InModule->VM = NewObject<URigVM>(InModule, TEXT("VM"), RF_NoFlags);
+		InSettings.ReportError(TEXT("Could not find execute pin on Execute Bindings node"));
+		return nullptr;
 	}
-	InModule->VM->Reset(InModule->ExtendedExecuteContext);
-	InModule->RigVM = InModule->VM; // Local serialization
-}
-
-UAnimNextModule_EditorData* FUtils::GetEditorData(const UAnimNextModule* InModule)
-{
-	check(InModule);
-	
-	return CastChecked<UAnimNextModule_EditorData>(InModule->EditorData);
-}
-
-UAnimNextModule* FUtils::GetGraph(const UAnimNextModule_EditorData* InEditorData)
-{
-	check(InEditorData);
-
-	return CastChecked<UAnimNextModule>(InEditorData->GetOuter());
-}
+	URigVMPin* ExecPin = ExecuteBindingsExecPin;
 
-FParamTypeHandle FUtils::GetParameterHandleFromPin(const FEdGraphPinType& InPinType)
-{
-	FAnimNextParamType::EValueType ValueType = FAnimNextParamType::EValueType::None;
-	FAnimNextParamType::EContainerType ContainerType = FAnimNextParamType::EContainerType::None;
-	UObject* ValueTypeObject = nullptr;
-
-	if (InPinType.PinCategory == UEdGraphSchema_K2::PC_Boolean)
-	{
-		ValueType = FAnimNextParamType::EValueType::Bool;
-	}
-	else if (InPinType.PinCategory == UEdGraphSchema_K2::PC_Byte)
-	{
-		ValueType = FAnimNextParamType::EValueType::Byte;
-	}
-	else if (InPinType.PinCategory == UEdGraphSchema_K2::PC_Int)
-	{
-		ValueType = FAnimNextParamType::EValueType::Int32;
-	}
-	else if (InPinType.PinCategory == UEdGraphSchema_K2::PC_Int64)
-	{
-		ValueType = FAnimNextParamType::EValueType::Int64;
-	}
-	else if (InPinType.PinCategory == UEdGraphSchema_K2::PC_Real)
+	if(bHasPublicVariablesToCopy)
 	{
-		if (InPinType.PinSubCategory == UEdGraphSchema_K2::PC_Float)
+		URigVMNode* CopyProxyVariablesNode = Controller->AddUnitNode(FRigUnit_CopyModuleProxyVariables::StaticStruct(), FRigUnit_CopyModuleProxyVariables::GetMethodName(), FVector2D(200, 0.0f), FString(), false);
+		if(CopyProxyVariablesNode == nullptr)
 		{
-			ValueType = FAnimNextParamType::EValueType::Float;
+			InSettings.ReportError(TEXT("Could not spawn Copy Module Proxy Variables node"));
+			return nullptr;
 		}
-		else if (InPinType.PinSubCategory == UEdGraphSchema_K2::PC_Double)
+		URigVMPin* CopyProxyVariablesExecPin = CopyProxyVariablesNode->FindPin(FRigVMStruct::ExecuteContextName.ToString());
+		if(ExecPin == nullptr)
 		{
-			ValueType = FAnimNextParamType::EValueType::Double;
+			InSettings.ReportError(TEXT("Could not find execute pin on Execute Bindings node"));
+			return nullptr;
 		}
-		else
+		bool bLinkAdded = Controller->AddLink(ExecuteBindingsExecPin, CopyProxyVariablesExecPin, false);
+		if(!bLinkAdded)
 		{
-			ensure(false);	// Reals should be either floats or doubles
+			InSettings.ReportError(TEXT("Could not link Execute Bindings node"));
+			return nullptr;
 		}
-	}
-	else if (InPinType.PinCategory == UEdGraphSchema_K2::PC_Float)
-	{
-		ValueType = FAnimNextParamType::EValueType::Float;
-	}
-	else if (InPinType.PinCategory == UEdGraphSchema_K2::PC_Double)
-	{
-		ValueType = FAnimNextParamType::EValueType::Double;
-	}
-	else if (InPinType.PinCategory == UEdGraphSchema_K2::PC_Name)
-	{
-		ValueType = FAnimNextParamType::EValueType::Name;
-	}
-	else if (InPinType.PinCategory == UEdGraphSchema_K2::PC_String)
-	{
-		ValueType = FAnimNextParamType::EValueType::String;
-	}
-	else if (InPinType.PinCategory == UEdGraphSchema_K2::PC_Text)
-	{
-		ValueType = FAnimNextParamType::EValueType::Text;
-	}
-	else if (InPinType.PinCategory == UEdGraphSchema_K2::PC_Enum)
-	{
-		ValueType = FAnimNextParamType::EValueType::Enum;
-		ValueTypeObject = InPinType.PinSubCategoryObject.Get();
-		ensure(ValueTypeObject);
-	}
-	else if (InPinType.PinCategory == UEdGraphSchema_K2::PC_Struct)
-	{
-		ValueType = FAnimNextParamType::EValueType::Struct;
-		ValueTypeObject = Cast<UScriptStruct>(InPinType.PinSubCategoryObject.Get());
-		ensure(ValueTypeObject);
-	}
-	else if (InPinType.PinCategory == UEdGraphSchema_K2::PC_Object)
-	{
-		ValueType = FAnimNextParamType::EValueType::Object;
-		ValueTypeObject = Cast<UClass>(InPinType.PinSubCategoryObject.Get());
-		ensure(ValueTypeObject);
-	}
-	else if (InPinType.PinCategory == UEdGraphSchema_K2::PC_SoftObject)
-	{
-		ValueType = FAnimNextParamType::EValueType::SoftObject;
-		ValueTypeObject = Cast<UClass>(InPinType.PinSubCategoryObject.Get());
-		ensure(ValueTypeObject);
-	}
-	else if (InPinType.PinCategory == UEdGraphSchema_K2::PC_SoftClass)
-	{
-		ValueType = FAnimNextParamType::EValueType::SoftClass;
-		ValueTypeObject = Cast<UClass>(InPinType.PinSubCategoryObject.Get());
-		ensure(ValueTypeObject);
+		ExecPin = CopyProxyVariablesExecPin;
 	}
 
-	if(InPinType.ContainerType == EPinContainerType::Array)
-	{
-		ContainerType = FAnimNextParamType::EContainerType::Array;
-	}
-	else if(InPinType.ContainerType == EPinContainerType::Set)
-	{
-		ensureMsgf(false, TEXT("Set pins are not yet supported"));
-	}
-	if(InPinType.ContainerType == EPinContainerType::Map)
-	{
-		ensureMsgf(false, TEXT("Map pins are not yet supported"));
-	}
-	
-	return FAnimNextParamType(ValueType, ContainerType, ValueTypeObject).GetHandle();
-}
+	IVariableBindingType::FBindingGraphFragmentArgs Args;
+	Args.Controller = Controller;
+	Args.BindingGraph = BindingGraph;
+	Args.ExecTail = ExecPin;
 
-void FUtils::CompileStruct(UAnimNextModule* InModule)
-{
-	check(InModule);
-
-	UAnimNextModule_EditorData* EditorData = GetEditorData(InModule);
-	if(EditorData->bIsCompiling)
-	{
-		return;
-	}
-
-	TGuardValue<bool> CompilingGuard(EditorData->bIsCompiling, true);
-
-	struct FStructEntryInfo
-	{
-		FName Name;
-		FAnimNextParamType Type;
-		EAnimNextExportAccessSpecifier AccessSpecifier;
-	};
-
-	TArray<FStructEntryInfo> StructEntryInfos;
-	StructEntryInfos.Reserve(EditorData->Entries.Num());
-
-	// Gather all parameters in this asset
-	for(const UAnimNextRigVMAssetEntry* Entry : EditorData->Entries)
+	FVector2D Location(0.0f, 0.0f);
+	for(const TPair<IVariableBindingType*, TArray<IVariableBindingType::FBindingGraphInput>>& BindingGroupPair : BindingGroups)
 	{
-		const IAnimNextRigVMExportInterface* Export = Cast<IAnimNextRigVMExportInterface>(Entry);
-		const IAnimNextRigVMParameterInterface* Parameter = Cast<IAnimNextRigVMParameterInterface>(Entry);
-		if(Export && Parameter)
-		{
-			const FAnimNextParamType& Type = Export->GetExportType();
-			ensure(Type.IsValid());
-			const FName Name = Export->GetExportName();
-			const EAnimNextExportAccessSpecifier AccessSpecifier = Export->GetExportAccessSpecifier();
-
-			StructEntryInfos.Add( { Name, FAnimNextParamType(Type.GetValueType(),  Type.GetContainerType(), Type.GetValueTypeObject()), AccessSpecifier } );
-		}
+		Args.Inputs = BindingGroupPair.Value;
+		BindingGroupPair.Key->BuildBindingGraphFragment(InSettings, Args, ExecPin, Location);
 	}
 
-	// Sort private entries first & then by size, largest first, for better packing
-	static_assert(EAnimNextExportAccessSpecifier::Private < EAnimNextExportAccessSpecifier::Public, "Private must be less than Public as parameters are sorted internally according to this assumption");
-	StructEntryInfos.Sort([](const FStructEntryInfo& InLHS, const FStructEntryInfo& InRHS)
-	{
-		if(InLHS.AccessSpecifier < InRHS.AccessSpecifier)
-		{
-			return true;
-		}
-		else
-		{
-			return InLHS.Type.GetSize() > InRHS.Type.GetSize();
-		}
-	});
-
-	if(StructEntryInfos.Num() > 0)
-	{
-		// Build PropertyDescs to batch-create the property bag
-		TArray<FPropertyBagPropertyDesc> PropertyDescs;
-		PropertyDescs.Reserve(StructEntryInfos.Num());
-
-		InModule->DefaultState.PublicParameterStartIndex = INDEX_NONE;
-
-		for(int32 EntryIndex = 0; EntryIndex < StructEntryInfos.Num(); ++EntryIndex)
-		{
-			const FStructEntryInfo& StructEntryInfo = StructEntryInfos[EntryIndex];
-			// Find the first parameter that is public and record it
-			if(StructEntryInfo.AccessSpecifier == EAnimNextExportAccessSpecifier::Public)
-			{
-				InModule->DefaultState.PublicParameterStartIndex = EntryIndex;
-			}
-			PropertyDescs.Emplace(StructEntryInfo.Name, StructEntryInfo.Type.ContainerType, StructEntryInfo.Type.ValueType, StructEntryInfo.Type.ValueTypeObject);
-		}
-
-		// Find any existing IDs for old properties with name-matching
-		// TODO: linear search - we could cache the name->GUID lookup in editor to accelerate this.
-		for(FPropertyBagPropertyDesc& NewDesc : PropertyDescs)
-		{
-			if(InModule->DefaultState.State.GetPropertyBagStruct())
-			{
-				for(const FPropertyBagPropertyDesc& ExistingDesc : InModule->DefaultState.State.GetPropertyBagStruct()->GetPropertyDescs())
-				{
-					if(ExistingDesc.Name == NewDesc.Name)
-					{
-						NewDesc.ID = ExistingDesc.ID;
-						break;
-					}
-				}
-			}
-		}
-
-		// Create new property bag and migrate
-		const UPropertyBag* NewBagStruct = UPropertyBag::GetOrCreateFromDescs(PropertyDescs);
-		InModule->DefaultState.State.MigrateToNewBagStruct(NewBagStruct);
-	}
-	else
-	{
-		InModule->DefaultState.Reset();
-	}
+	return BindingGraph;
 }
 
 UAnimNextRigVMAsset* FUtils::GetAsset(UAnimNextRigVMAssetEditorData* InEditorData)
@@ -956,18 +251,6 @@ UAnimNextRigVMAssetEditorData* FUtils::GetEditorData(UAnimNextRigVMAsset* InAsse
 	return CastChecked<UAnimNextRigVMAssetEditorData>(InAsset->EditorData);
 }
 
-FInstancedPropertyBag* FUtils::GetPropertyBag(UAnimNextModule* InModule)
-{
-	FInstancedPropertyBag* InstancedPropertyBag = &InModule->DefaultState.State;
-
-	return InstancedPropertyBag;
-}
-
-FParamTypeHandle FUtils::GetParamTypeHandleFromPinType(const FEdGraphPinType& InPinType)
-{
-	return GetParamTypeFromPinType(InPinType).GetHandle();
-}
-
 FAnimNextParamType FUtils::GetParamTypeFromPinType(const FEdGraphPinType& InPinType)
 {
 	FAnimNextParamType::EValueType ValueType = FAnimNextParamType::EValueType::None;
@@ -1070,11 +353,6 @@ FAnimNextParamType FUtils::GetParamTypeFromPinType(const FEdGraphPinType& InPinT
 	return FAnimNextParamType(ValueType, ContainerType, ValueTypeObject);
 }
 
-FEdGraphPinType FUtils::GetPinTypeFromParamTypeHandle(const FParamTypeHandle& InParamTypeHandle)
-{
-	return GetPinTypeFromParamType(InParamTypeHandle.GetType());
-}
-
 FEdGraphPinType FUtils::GetPinTypeFromParamType(const FAnimNextParamType& InParamType)
 {
 	FEdGraphPinType PinType;
@@ -1154,11 +432,6 @@ FEdGraphPinType FUtils::GetPinTypeFromParamType(const FAnimNextParamType& InPara
 	return PinType;
 }
 
-FRigVMTemplateArgumentType FUtils::GetRigVMArgTypeFromParamTypeHandle(const FParamTypeHandle& InParamTypeHandle)
-{
-	return GetRigVMArgTypeFromParamType(InParamTypeHandle.GetType());
-}
-
 FRigVMTemplateArgumentType FUtils::GetRigVMArgTypeFromParamType(const FAnimNextParamType& InParamType)
 {
 	FRigVMTemplateArgumentType ArgType;
@@ -1245,108 +518,45 @@ void FUtils::SetupAnimGraph(UAnimNextRigVMAssetEntry* InEntry, URigVMController*
 	// Clear the graph
 	InController->RemoveNodes(InController->GetGraph()->GetNodes());
 
-	// Add root node
-	URigVMUnitNode* MainEntryPointNode = InController->AddUnitNode(FRigUnit_AnimNextGraphRoot::StaticStruct(), FRigUnit_AnimNextGraphRoot::EventName, FVector2D(-400.0f, 0.0f), FString(), false);
-	check(MainEntryPointNode);
-	URigVMPin* BeginExecutePin = MainEntryPointNode->FindPin(GET_MEMBER_NAME_STRING_CHECKED(FRigUnit_AnimNextGraphRoot, Result));
-	check(BeginExecutePin);
-	check(BeginExecutePin->GetDirection() == ERigVMPinDirection::Input);
-
-	URigVMPin* EntryPointPin = MainEntryPointNode->FindPin(GET_MEMBER_NAME_STRING_CHECKED(FRigUnit_AnimNextGraphRoot, EntryPoint));
-	check(EntryPointPin);
-	check(EntryPointPin->GetDirection() == ERigVMPinDirection::Hidden);
-	InController->SetPinDefaultValue(EntryPointPin->GetPinPath(), InEntry->GetEntryName().ToString());
-}
-
-void FUtils::SetupEventGraph(URigVMController* InController)
-{
-	// Clear the graph
-	InController->RemoveNodes(InController->GetGraph()->GetNodes());
-
-	// Add entry point
-	InController->AddUnitNode(FRigUnit_AnimNextParameterBeginExecution::StaticStruct(), FRigVMStruct::ExecuteName, FVector2D(-200.0f, 0.0f), FString(), false);
-}
-
-FName FUtils::GetParameterNameFromQualifiedName(FName InName)
-{
-	const FSoftObjectPath SoftObjectPath(InName.ToString());
-	return FName(*SoftObjectPath.GetSubPathString());
-}
-
-FName FUtils::GetQualifiedName(UAnimNextRigVMAsset* InAsset, FName InBaseName)
-{
-	if(InAsset)
-	{
-		TStringBuilder<256> StringBuilder;
-		StringBuilder.Append(InAsset->GetPathName());
-		StringBuilder.Append(TEXT(":"));
-		InBaseName.AppendString(StringBuilder);
-
-		return FName(StringBuilder.ToView());
-	}
-	return InBaseName;
-}
-
-FText FUtils::GetParameterDisplayNameText(FName InParameterName, const TInstancedStruct<FAnimNextParamInstanceIdentifier>& InInstanceId)
-{
-	if(InInstanceId.IsValid())
-	{
-		FText ParameterText;
-		FModule& Module = FModuleManager::GetModuleChecked<FModule>("AnimNextUncookedOnly");
-		if(TSharedPtr<IParameterSourceType> SourceType = Module.FindParameterSourceType(InInstanceId.GetScriptStruct()))
-		{
-			ParameterText = SourceType->GetDisplayText(InInstanceId);
-		}
-
-		TStringBuilder<256> StringBuilder;
-		if(!ParameterText.IsEmpty())
-		{
-			StringBuilder.Append(ParameterText.ToString());
-			StringBuilder.Append(TEXT("."));
-		}
+	// Add root node
+	URigVMUnitNode* MainEntryPointNode = InController->AddUnitNode(FRigUnit_AnimNextGraphRoot::StaticStruct(), FRigUnit_AnimNextGraphRoot::EventName, FVector2D(-400.0f, 0.0f), FString(), false);
+	if(MainEntryPointNode == nullptr)
+	{
+		return;
+	}
 
-		GetParameterNameFromQualifiedName(InParameterName).AppendString(StringBuilder);
-		return FText::FromStringView(StringBuilder);
+	URigVMPin* BeginExecutePin = MainEntryPointNode->FindPin(GET_MEMBER_NAME_STRING_CHECKED(FRigUnit_AnimNextGraphRoot, Result));
+	if(BeginExecutePin == nullptr)
+	{
+		return;
 	}
-	else
+
+	URigVMPin* EntryPointPin = MainEntryPointNode->FindPin(GET_MEMBER_NAME_STRING_CHECKED(FRigUnit_AnimNextGraphRoot, EntryPoint));
+	if(EntryPointPin == nullptr)
 	{
-		if(InParameterName.IsNone())
-		{
-			return FText::FromName(InParameterName);
-		}
-		else
-		{
-			const FSoftObjectPath SoftObjectPath(InParameterName.ToString());
-			return FText::Format(LOCTEXT("ParameterNameDisplayFormat", "{0}.{1}"), FText::FromString(SoftObjectPath.GetAssetName()), FText::FromString(SoftObjectPath.GetSubPathString()));
-		}
+		return;
 	}
+
+	InController->SetPinDefaultValue(EntryPointPin->GetPinPath(), InEntry->GetEntryName().ToString());
 }
 
-FText FUtils::GetParameterTooltipText(FName InParameterName, const TInstancedStruct<FAnimNextParamInstanceIdentifier>& InInstanceId)
+void FUtils::SetupEventGraph(URigVMController* InController, UScriptStruct* InEventStruct)
 {
-	FTextBuilder TextBuilder;
-	TextBuilder.AppendLine(FText::Format(LOCTEXT("ParameterNameTooltipFormat", "Parameter: {0}"), FText::FromString(InParameterName.ToString())));
-
-	if(InInstanceId.IsValid())
-	{
-		FModule& Module = FModuleManager::GetModuleChecked<FModule>("AnimNextUncookedOnly");
-		if(TSharedPtr<IParameterSourceType> SourceType = Module.FindParameterSourceType(InInstanceId.GetScriptStruct()))
-		{
-			TextBuilder.AppendLine(SourceType->GetTooltipText(InInstanceId));
-		}
-	}
+	// Clear the graph
+	InController->RemoveNodes(InController->GetGraph()->GetNodes());
 
-	return TextBuilder.ToText();
+	// Add entry point
+	InController->AddUnitNode(InEventStruct, FRigVMStruct::ExecuteName, FVector2D(-200.0f, 0.0f), FString(), false);
 }
 
 FAnimNextParamType FUtils::GetParameterTypeFromName(FName InName)
 {
 	// Query the asset registry for other params
-	TMap<FAssetData, FAnimNextParameterProviderAssetRegistryExports> ExportMap;
-	GetExportedParametersFromAssetRegistry(ExportMap);
-	for(const TPair<FAssetData, FAnimNextParameterProviderAssetRegistryExports>& ExportPair : ExportMap)
+	TMap<FAssetData, FAnimNextAssetRegistryExports> ExportMap;
+	GetExportedVariablesFromAssetRegistry(ExportMap);
+	for(const TPair<FAssetData, FAnimNextAssetRegistryExports>& ExportPair : ExportMap)
 	{
-		for(const FAnimNextParameterAssetRegistryExportEntry& Parameter : ExportPair.Value.Parameters)
+		for(const FAnimNextAssetRegistryExportedVariable& Parameter : ExportPair.Value.Variables)
 		{
 			if(Parameter.Name == InName)
 			{
@@ -1358,13 +568,13 @@ FAnimNextParamType FUtils::GetParameterTypeFromName(FName InName)
 	return FAnimNextParamType();
 }
 
-bool FUtils::GetExportedParametersForAsset(const FAssetData& InAsset, FAnimNextParameterProviderAssetRegistryExports& OutExports)
+bool FUtils::GetExportedVariablesForAsset(const FAssetData& InAsset, FAnimNextAssetRegistryExports& OutExports)
 {
 	const FString TagValue = InAsset.GetTagValueRef<FString>(UE::AnimNext::ExportsAnimNextAssetRegistryTag);
-	return FAnimNextParameterProviderAssetRegistryExports::StaticStruct()->ImportText(*TagValue, &OutExports, nullptr, PPF_None, nullptr, FAnimNextParameterProviderAssetRegistryExports::StaticStruct()->GetName()) != nullptr;
+	return FAnimNextAssetRegistryExports::StaticStruct()->ImportText(*TagValue, &OutExports, nullptr, PPF_None, nullptr, FAnimNextAssetRegistryExports::StaticStruct()->GetName()) != nullptr;
 }
 
-bool FUtils::GetExportedParametersFromAssetRegistry(TMap<FAssetData, FAnimNextParameterProviderAssetRegistryExports>& OutExports)
+bool FUtils::GetExportedVariablesFromAssetRegistry(TMap<FAssetData, FAnimNextAssetRegistryExports>& OutExports)
 {
 	TArray<FAssetData> AssetData;
 	IAssetRegistry::GetChecked().GetAssetsByTags({UE::AnimNext::ExportsAnimNextAssetRegistryTag}, AssetData);
@@ -1372,8 +582,8 @@ bool FUtils::GetExportedParametersFromAssetRegistry(TMap<FAssetData, FAnimNextPa
 	for (const FAssetData& Asset : AssetData)
 	{
 		const FString TagValue = Asset.GetTagValueRef<FString>(UE::AnimNext::ExportsAnimNextAssetRegistryTag);
-		FAnimNextParameterProviderAssetRegistryExports AssetExports;
-		if (FAnimNextParameterProviderAssetRegistryExports::StaticStruct()->ImportText(*TagValue, &AssetExports, nullptr, PPF_None, nullptr, FAnimNextParameterProviderAssetRegistryExports::StaticStruct()->GetName()) != nullptr)
+		FAnimNextAssetRegistryExports AssetExports;
+		if (FAnimNextAssetRegistryExports::StaticStruct()->ImportText(*TagValue, &AssetExports, nullptr, PPF_None, nullptr, FAnimNextAssetRegistryExports::StaticStruct()->GetName()) != nullptr)
 		{
 			OutExports.Add(Asset, MoveTemp(AssetExports));
 		}
@@ -1382,9 +592,9 @@ bool FUtils::GetExportedParametersFromAssetRegistry(TMap<FAssetData, FAnimNextPa
 	return OutExports.Num() > 0;
 }
 
-static void AddParamToSet(const FAnimNextParameterAssetRegistryExportEntry& InNewParam, TSet<FAnimNextParameterAssetRegistryExportEntry>& OutExports)
+static void AddParamToSet(const FAnimNextAssetRegistryExportedVariable& InNewParam, TSet<FAnimNextAssetRegistryExportedVariable>& OutExports)
 {
-	if(FAnimNextParameterAssetRegistryExportEntry* ExistingEntry = OutExports.Find(InNewParam))
+	if(FAnimNextAssetRegistryExportedVariable* ExistingEntry = OutExports.Find(InNewParam))
 	{
 		if(ExistingEntry->Type != InNewParam.Type)
 		{
@@ -1398,305 +608,63 @@ static void AddParamToSet(const FAnimNextParameterAssetRegistryExportEntry& InNe
 	}
 }
 
-void FUtils::GetAssetParameters(const UAnimNextRigVMAssetEditorData* EditorData, FAnimNextParameterProviderAssetRegistryExports& OutExports)
+void FUtils::GetAssetVariables(const UAnimNextRigVMAssetEditorData* EditorData, FAnimNextAssetRegistryExports& OutExports)
 {
-	OutExports.Parameters.Reset();
-	OutExports.Parameters.Reserve(EditorData->Entries.Num());
+	OutExports.Variables.Reset();
+	OutExports.Variables.Reserve(EditorData->Entries.Num());
 
-	TSet<FAnimNextParameterAssetRegistryExportEntry> ExportSet;
-	GetAssetParameters(EditorData, ExportSet);
-	OutExports.Parameters = ExportSet.Array();
+	TSet<FAnimNextAssetRegistryExportedVariable> ExportSet;
+	GetAssetVariables(EditorData, ExportSet);
+	OutExports.Variables = ExportSet.Array();
 }
 
-void FUtils::GetAssetParameters(const UAnimNextRigVMAssetEditorData* EditorData, TSet<FAnimNextParameterAssetRegistryExportEntry>& OutExports)
+void FUtils::GetAssetVariables(const UAnimNextRigVMAssetEditorData* EditorData, TSet<FAnimNextAssetRegistryExportedVariable>& OutExports)
 {
 	for(const UAnimNextRigVMAssetEntry* Entry : EditorData->Entries)
 	{
 		if(const IAnimNextRigVMExportInterface* ExportInterface = Cast<IAnimNextRigVMExportInterface>(Entry))
 		{
-			EAnimNextParameterFlags Flags = EAnimNextParameterFlags::Declared;
+			EAnimNextExportedVariableFlags Flags = EAnimNextExportedVariableFlags::Declared;
 			if(ExportInterface->GetExportAccessSpecifier() == EAnimNextExportAccessSpecifier::Public)
 			{
-				Flags |= EAnimNextParameterFlags::Public;
-				FAnimNextParameterAssetRegistryExportEntry NewParam(ExportInterface->GetExportName(), TInstancedStruct<FAnimNextParamInstanceIdentifier>(), ExportInterface->GetExportType(), Flags);
+				Flags |= EAnimNextExportedVariableFlags::Public;
+				FAnimNextAssetRegistryExportedVariable NewParam(ExportInterface->GetExportName(), ExportInterface->GetExportType(), Flags);
 				AddParamToSet(NewParam, OutExports);
 			}
 		}
-		if(const IAnimNextRigVMGraphInterface* GraphInterface = Cast<IAnimNextRigVMGraphInterface>(Entry))
-		{
-			GetGraphParameters(GraphInterface->GetRigVMGraph(), OutExports);
-		}
-	}
-}
-
-void FUtils::GetGraphParameters(const URigVMGraph* Graph, FAnimNextParameterProviderAssetRegistryExports& OutExports)
-{
-	TSet<FAnimNextParameterAssetRegistryExportEntry> ExportSet;
-	GetGraphParameters(Graph, ExportSet);
-	OutExports.Parameters = ExportSet.Array();
-}
-
-void FUtils::GetGraphParameters(const URigVMGraph* Graph, TSet<FAnimNextParameterAssetRegistryExportEntry>& OutExports)
-{
-	if (Graph == nullptr)
-	{
-		return;
-	}
-
-	const TArray<URigVMNode*>& Nodes = Graph->GetNodes();
-	for (URigVMNode* Node : Nodes)
-	{
-		if (const URigVMTemplateNode* TemplateNode = Cast<URigVMTemplateNode>(Node))
-		{
-			const FRigVMDispatchFactory* GetParameterFactory = FRigVMRegistry::Get().FindOrAddDispatchFactory(FRigVMDispatch_GetParameter::StaticStruct());
-			const FName GetParameterNotation = GetParameterFactory->GetTemplate()->GetNotation();
-
-			const FRigVMDispatchFactory* GetScopedParameterFactory = FRigVMRegistry::Get().FindOrAddDispatchFactory(FRigVMDispatch_GetScopedParameter::StaticStruct());
-			const FName GetScopedParameterNotation = GetScopedParameterFactory->GetTemplate()->GetNotation();
-			
-			const FRigVMDispatchFactory* GetLayerParameterFactory = FRigVMRegistry::Get().FindOrAddDispatchFactory(FRigVMDispatch_GetLayerParameter::StaticStruct());
-			const FName GetLayerParameterNotation = GetLayerParameterFactory->GetTemplate()->GetNotation();
-
-			const FRigVMDispatchFactory* SetLayerParameterFactory = FRigVMRegistry::Get().FindOrAddDispatchFactory(FRigVMDispatch_SetLayerParameter::StaticStruct());
-			const FName SetLayerParameterNotation = SetLayerParameterFactory->GetTemplate()->GetNotation();
-
-			const bool bIsScopedParameter = TemplateNode->GetNotation() == GetScopedParameterNotation; 
-			const bool bReadParameter = bIsScopedParameter || TemplateNode->GetNotation() == GetParameterNotation || TemplateNode->GetNotation() == GetLayerParameterNotation;
-			const bool bWriteParameter = TemplateNode->GetNotation() == SetLayerParameterNotation;
-			const bool bUsesRuntimeStruct = bIsScopedParameter;
-			const bool bUsesName = TemplateNode->GetNotation() == GetLayerParameterNotation || TemplateNode->GetNotation() == SetLayerParameterNotation;
-			const bool bIsParameterNode = bIsScopedParameter || bReadParameter || bWriteParameter;
-
-			const URigVMPin* ParameterPin = TemplateNode->FindPin(FRigVMDispatch_GetParameter::ParameterName.ToString());
-			if (bIsParameterNode && ParameterPin != nullptr)
-			{
-				const FString PinDefaultValue = ParameterPin->GetDefaultValue();
-				if(!PinDefaultValue.IsEmpty())
-				{
-					FAnimNextEditorParam PinParam;
-					if(bUsesRuntimeStruct)
-					{
-						FAnimNextParam AnimNextParam;
-						FAnimNextParam::StaticStruct()->ImportText(*PinDefaultValue, &AnimNextParam, nullptr, PPF_None, nullptr, FAnimNextParam::StaticStruct()->GetName());
-						PinParam = FAnimNextEditorParam(AnimNextParam);
-					}
-					else if(bUsesName)
-					{
-						FName ParameterName = *PinDefaultValue;
-						const URigVMPin* ValuePin = TemplateNode->FindPin(FRigVMDispatch_GetLayerParameter::ValueName.ToString());
-						if (ValuePin != nullptr)
-						{
-							const FAnimNextParamType ParamType = FAnimNextParamType::FromRigVMTemplateArgument(FRigVMTemplateArgumentType(FName(*ValuePin->GetCPPType()), ValuePin->GetCPPTypeObject()));
-							PinParam = FAnimNextEditorParam(ParameterName, ParamType, TInstancedStruct<FAnimNextParamInstanceIdentifier>());
-						}
-					}
-					else
-					{
-						FAnimNextEditorParam::StaticStruct()->ImportText(*PinDefaultValue, &PinParam, nullptr, PPF_None, nullptr, FAnimNextEditorParam::StaticStruct()->GetName());
-					}
-
-					if(PinParam.Type.IsValid())
-					{
-						EAnimNextParameterFlags Flags = EAnimNextParameterFlags::NoFlags;
-						if (bReadParameter)
-						{
-							Flags |= EAnimNextParameterFlags::Read;
-						}
-
-						if (bWriteParameter)
-						{
-							Flags |= EAnimNextParameterFlags::Write;
-						}
-
-						FAnimNextParameterAssetRegistryExportEntry NewEntry(PinParam.Name, PinParam.InstanceId, PinParam.Type, Flags);
-						AddParamToSet(NewEntry, OutExports);
-					}
-				}
-			}
-		}
 	}
 }
 
-void FUtils::GetScheduleParameters(const UAnimNextSchedule* InSchedule, FAnimNextParameterProviderAssetRegistryExports& OutExports)
-{
-	TSet<FAnimNextParameterAssetRegistryExportEntry> ExportSet;
-	GetScheduleParameters(InSchedule, ExportSet);
-	OutExports.Parameters = ExportSet.Array();
-}
-
-void FUtils::GetScheduleParameters(const UAnimNextSchedule* InSchedule, TSet<FAnimNextParameterAssetRegistryExportEntry>& OutExports)
+void FUtils::GetAssetOutlinerItems(const UAnimNextRigVMAssetEditorData* EditorData, FWorkspaceOutlinerItemExports& OutExports)
 {
-	for(UAnimNextScheduleEntry* Entry : InSchedule->Entries)
+	FWorkspaceOutlinerItemExport AssetIdentifier = FWorkspaceOutlinerItemExport(EditorData->GetOuter()->GetFName(), EditorData->GetOuter());
+	for(UAnimNextRigVMAssetEntry* Entry : EditorData->Entries)
 	{
-		if (UAnimNextScheduleEntry_Port* PortEntry = Cast<UAnimNextScheduleEntry_Port>(Entry))
+		if(IAnimNextRigVMGraphInterface* GraphInterface = Cast<IAnimNextRigVMGraphInterface>(Entry))
 		{
-			if(PortEntry->Port)
+			if (Entry->IsHiddenInOutliner())
 			{
-				UAnimNextSchedulePort* CDO = PortEntry->Port->GetDefaultObject<UAnimNextSchedulePort>();
-				TConstArrayView<FAnimNextEditorParam> RequiredParameters = CDO->GetRequiredParameters();
-				for(const FAnimNextEditorParam& RequiredParameter : RequiredParameters)
+				if (URigVMEdGraph* RigVMEdGraph = GraphInterface->GetEdGraph())
 				{
-					if(!RequiredParameter.Name.IsNone() && RequiredParameter.Type.IsValid())
-					{
-						FAnimNextParameterAssetRegistryExportEntry NewEntry(RequiredParameter.Name, RequiredParameter.InstanceId, RequiredParameter.Type, EAnimNextParameterFlags::Read);
-						AddParamToSet(NewEntry, OutExports);
-					}
+					CreateSubGraphsOutlinerItemsRecursive(EditorData, OutExports, AssetIdentifier, RigVMEdGraph);
 				}
 			}
-		}
-		else if (UAnimNextScheduleEntry_AnimNextGraph* GraphEntry = Cast<UAnimNextScheduleEntry_AnimNextGraph>(Entry))
-		{
-			if(GraphEntry->DynamicGraph.IsValid())
-			{
-				FAnimNextParameterAssetRegistryExportEntry NewEntry(GraphEntry->DynamicGraph.Name, GraphEntry->DynamicGraph.InstanceId, GraphEntry->DynamicGraph.Type, EAnimNextParameterFlags::Read);
-				AddParamToSet(NewEntry, OutExports);
-			}
-			for(const FAnimNextEditorParam& RequiredParameter : GraphEntry->RequiredParameters)
-			{
-				FAnimNextParameterAssetRegistryExportEntry NewEntry(RequiredParameter.Name, RequiredParameter.InstanceId, RequiredParameter.Type, EAnimNextParameterFlags::Read);
-				AddParamToSet(NewEntry, OutExports);
-			}
-		}
-		else if (UAnimNextScheduleEntry_ExternalTask* ExternalTaskEntry = Cast<UAnimNextScheduleEntry_ExternalTask>(Entry))
-		{
-			if(ExternalTaskEntry->ExternalTask.IsValid())
-			{
-				FAnimNextParameterAssetRegistryExportEntry NewEntry(ExternalTaskEntry->ExternalTask.Name, ExternalTaskEntry->ExternalTask.InstanceId, ExternalTaskEntry->ExternalTask.Type, EAnimNextParameterFlags::Read);
-				AddParamToSet(NewEntry, OutExports);
-			}
-		}
-		else if (UAnimNextScheduleEntry_ParamScope* ParamScopeTaskEntry = Cast<UAnimNextScheduleEntry_ParamScope>(Entry))
-		{
-			if(ParamScopeTaskEntry->Scope.IsValid())
-			{
-				FAnimNextParameterAssetRegistryExportEntry NewEntry(ParamScopeTaskEntry->Scope.Name, ParamScopeTaskEntry->Scope.InstanceId, ParamScopeTaskEntry->Scope.Type, EAnimNextParameterFlags::Read);
-				AddParamToSet(NewEntry, OutExports);
-			}
-		}
-	}
-}
-
-void FUtils::GetBlueprintParameters(const UBlueprint* InBlueprint, FAnimNextParameterProviderAssetRegistryExports& OutExports)
-{
-	TSet<FAnimNextParameterAssetRegistryExportEntry> ExportSet;
-	GetBlueprintParameters(InBlueprint, ExportSet);
-	OutExports.Parameters = ExportSet.Array();
-}
-
-void FUtils::GetBlueprintParameters(const UBlueprint* InBlueprint, TSet<FAnimNextParameterAssetRegistryExportEntry>& OutExports)
-{
-	// Add 'static' params held on components
-	if(InBlueprint->SimpleConstructionScript)
-	{
-		for(USCS_Node* SCSNode : InBlueprint->SimpleConstructionScript->GetAllNodes())
-		{
-			if(UAnimNextComponent* Template = Cast<UAnimNextComponent>(SCSNode->ComponentTemplate))
+			else
 			{
-				for(UAnimNextComponentParameter* Parameter : Template->Parameters)
-				{
-					if(!Parameter->Scope.IsNone())
-					{
-						FAnimNextParameterAssetRegistryExportEntry NewEntry(Parameter->Scope, TInstancedStruct<FAnimNextParamInstanceIdentifier>(), FAnimNextParamType::GetType<FAnimNextScope>(), EAnimNextParameterFlags::Read);
-						AddParamToSet(NewEntry, OutExports);
-					}
-
-					FName Name;
-					const FProperty* Property = nullptr;
-					Parameter->GetParamInfo(Name, Property);
-					if(!Name.IsNone())
-					{
-						FAnimNextParamType Type = FParamTypeHandle::FromProperty(Property).GetType();
-						check(Type.IsValid());
-						FAnimNextParameterAssetRegistryExportEntry NewEntry(Name, TInstancedStruct<FAnimNextParamInstanceIdentifier>(), Type, EAnimNextParameterFlags::Write);
-						AddParamToSet(NewEntry, OutExports);
-					}
-				}
-			}
-		}
-	}
+				FWorkspaceOutlinerItemExport& Export = OutExports.Exports.Add_GetRef(FWorkspaceOutlinerItemExport(Entry->GetEntryName(), AssetIdentifier));
 
-	// Add any dynamic params held on graph nodes
-	const UFunction* SetParameterInScopeFunc = UAnimNextComponent::StaticClass()->FindFunctionByName(GET_FUNCTION_NAME_CHECKED(UAnimNextComponent, SetParameterInScope));
-
-	check(SetParameterInScopeFunc);
-	TArray<UEdGraph*> AllGraphs;
-	InBlueprint->GetAllGraphs(AllGraphs);
-	for(UEdGraph* EdGraph : AllGraphs)
-	{
-		TArray<UK2Node_CallFunction*> FunctionNodes;
-		EdGraph->GetNodesOfClass(FunctionNodes);
-		for(UK2Node_CallFunction* FunctionNode : FunctionNodes)
-		{
-			UFunction* Function = FunctionNode->FunctionReference.ResolveMember<UFunction>(FunctionNode->GetBlueprintClassFromNode());
-			if(Function == SetParameterInScopeFunc)
-			{
-				UEdGraphPin* ScopePin = FunctionNode->FindPinChecked(TEXT("Scope"));
-				FName ScopeName(*ScopePin->GetDefaultAsString());
-				if(!ScopeName.IsNone())
-				{
-					FAnimNextParameterAssetRegistryExportEntry NewEntry(ScopeName, TInstancedStruct<FAnimNextParamInstanceIdentifier>(), FAnimNextParamType::GetType<FAnimNextScope>(), EAnimNextParameterFlags::Read);
-					AddParamToSet(NewEntry, OutExports);
-				}
+				Export.GetData().InitializeAsScriptStruct(FAnimNextGraphOutlinerData::StaticStruct());
+				FAnimNextGraphOutlinerData& GraphData = Export.GetData().GetMutable<FAnimNextGraphOutlinerData>();
+				GraphData.Entry = Entry;
+				GraphData.GraphInterface = GraphInterface->_getUObject();
 
-				UEdGraphPin* NamePin = FunctionNode->FindPinChecked(TEXT("Name"));
-				FName ParamName(*NamePin->GetDefaultAsString());
-				if(!ParamName.IsNone())
+				if (URigVMEdGraph* RigVMEdGraph = GraphInterface->GetEdGraph())
 				{
-					UEdGraphPin* ValuePin = FunctionNode->FindPinChecked(TEXT("Value"));
-					FAnimNextParamType Type = UncookedOnly::FUtils::GetParamTypeFromPinType(ValuePin->PinType);
-					if(Type.IsValid())
-					{
-						FAnimNextParameterAssetRegistryExportEntry NewEntry(ParamName, TInstancedStruct<FAnimNextParamInstanceIdentifier>(), Type, EAnimNextParameterFlags::Write);
-						AddParamToSet(NewEntry, OutExports);
-					}
+					CreateSubGraphsOutlinerItemsRecursive(EditorData, OutExports, Export, RigVMEdGraph);
 				}
 			}
 		}
 	}
-}
-
-void FUtils::GetAssetOutlinerItems(const UAnimNextRigVMAssetEditorData* EditorData, FWorkspaceOutlinerItemExports& OutExports)
-{
-	FWorkspaceOutlinerItemExport AssetIdentifier = FWorkspaceOutlinerItemExport(EditorData->GetOuter()->GetFName(), EditorData->GetOuter());
-	
-	constexpr bool bExportParametersAsOutlinerItems = false;
-	if (bExportParametersAsOutlinerItems)
-	{
-		TSet<FName> ParameterNames;
-		
-		FAnimNextParameterProviderAssetRegistryExports GraphExports;
-		UE::AnimNext::UncookedOnly::FUtils::GetAssetParameters(EditorData, GraphExports);
-		for ( const FAnimNextParameterAssetRegistryExportEntry& Entry : GraphExports.Parameters)
-		{
-			if (!ParameterNames.Contains(Entry.Name))
-			{
-				FWorkspaceOutlinerItemExport& ParameterExport = OutExports.Exports.Add_GetRef(FWorkspaceOutlinerItemExport(Entry.Name, AssetIdentifier));	
-
-				ParameterExport.GetData().InitializeAsScriptStruct(FAnimNextParameterOutlinerData::StaticStruct());
-				FAnimNextParameterOutlinerData& AssetData = ParameterExport.GetData().GetMutable<FAnimNextParameterOutlinerData>();
-				AssetData.Type = Entry.Type;
-
-				ParameterNames.Add(Entry.Name);
-			}
-		}
-	}
-	
-	for(const UAnimNextRigVMAssetEntry* Entry : EditorData->Entries)
-	{
-		if(const IAnimNextRigVMGraphInterface* GraphInterface = Cast<IAnimNextRigVMGraphInterface>(Entry))
-		{
-			FWorkspaceOutlinerItemExport& Export = OutExports.Exports.Add_GetRef(FWorkspaceOutlinerItemExport(Entry->GetEntryName(), AssetIdentifier));
-
-			Export.GetData().InitializeAsScriptStruct(FAnimNextGraphOutlinerData::StaticStruct());
-			FAnimNextGraphOutlinerData& GraphData = Export.GetData().GetMutable<FAnimNextGraphOutlinerData>();
-			GraphData.GraphInterface = GraphInterface->_getUObject();
-
-			if (URigVMEdGraph* RigVMEdGraph = GraphInterface->GetEdGraph())
-			{
-				CreateSubGraphsOutlinerItemsRecursive(EditorData, OutExports, Export, RigVMEdGraph);
-			}
-		}
-	}
 
 	CreateFunctionLibraryOutlinerItemsRecursive(EditorData, OutExports, AssetIdentifier, EditorData->GetRigVMGraphFunctionStore()->PublicFunctions, EditorData->GetRigVMGraphFunctionStore()->PrivateFunctions);
 }
@@ -1777,608 +745,6 @@ void FUtils::CreateFunctionsOutlinerItemsRecursive(const UAnimNextRigVMAssetEdit
 	}
 }
 
-void FUtils::CompileSchedule(UAnimNextSchedule* InSchedule)
-{
-	using namespace UE::AnimNext;
-
-	FMessageLog Log("AnimNextCompilerResults");
-	Log.NewPage(FText::FromName(InSchedule->GetFName()));
-
-	InSchedule->Instructions.Empty();
-	InSchedule->GraphTasks.Empty();
-	InSchedule->Ports.Empty();
-	InSchedule->ExternalTasks.Empty();
-	InSchedule->ParamScopeEntryTasks.Empty();
-	InSchedule->ParamScopeExitTasks.Empty();
-	InSchedule->ExternalParamTasks.Empty();
-	InSchedule->IntermediatesData.Reset();
-	InSchedule->NumParameterScopes = 0;
-	InSchedule->NumTickFunctions = 0;
-
-	EAnimNextScheduleScheduleOpcode LastOpCode = EAnimNextScheduleScheduleOpcode::None;
-
-	auto Emit = [InSchedule, &LastOpCode](EAnimNextScheduleScheduleOpcode InOpCode, int32 InOperand = 0)
-	{
-		FAnimNextScheduleInstruction Instruction;
-		Instruction.Opcode = InOpCode;
-		Instruction.Operand = InOperand;
-		InSchedule->Instructions.Add(Instruction);
-
-		LastOpCode = InOpCode;
-	};
-
-	auto EmitPrerequisite = [InSchedule, &Emit, &LastOpCode]()
-	{
-		switch (LastOpCode)
-		{
-		case EAnimNextScheduleScheduleOpcode::RunGraphTask:
-			Emit(EAnimNextScheduleScheduleOpcode::PrerequisiteTask, InSchedule->NumTickFunctions - 1);
-			break;
-		case EAnimNextScheduleScheduleOpcode::BeginRunExternalTask:
-			Emit(EAnimNextScheduleScheduleOpcode::PrerequisiteBeginExternalTask, InSchedule->NumTickFunctions - 1);
-			break;
-		case EAnimNextScheduleScheduleOpcode::EndRunExternalTask:
-			Emit(EAnimNextScheduleScheduleOpcode::PrerequisiteEndExternalTask, InSchedule->NumTickFunctions - 1);
-			break;
-		case EAnimNextScheduleScheduleOpcode::RunParamScopeEntry:
-			Emit(EAnimNextScheduleScheduleOpcode::PrerequisiteScopeEntry, InSchedule->NumTickFunctions - 1);
-			break;
-		case EAnimNextScheduleScheduleOpcode::RunParamScopeExit:
-			Emit(EAnimNextScheduleScheduleOpcode::PrerequisiteScopeExit, InSchedule->NumTickFunctions - 1);
-			break;
-		case EAnimNextScheduleScheduleOpcode::RunExternalParamTask:
-			Emit(EAnimNextScheduleScheduleOpcode::PrerequisiteExternalParamTask, InSchedule->NumTickFunctions - 1);
-			break;
-		case EAnimNextScheduleScheduleOpcode::None:
-			break;
-		default:
-			checkNoEntry();
-			break;
-		}
-	};
-
-	// MAX_uint32 means 'global scope' in this context
-	uint32 ParentScopeIndex = MAX_uint32;
-
-	TArray<FAnimNextScheduleEntryTerm> IntermediateTerms;
-	TMap<FName, uint32> IntermediateMap;
-
-	TFunction<void(TArrayView<TObjectPtr<UAnimNextScheduleEntry>>)> EmitEntries;
-
-	// Iterate over all entries, recursing into scopes
-	EmitEntries = [InSchedule, &EmitEntries, &Emit, &EmitPrerequisite, &ParentScopeIndex, &IntermediateTerms, &IntermediateMap, &Log](TArrayView<TObjectPtr<UAnimNextScheduleEntry>> InEntries)
-	{
-		for (int32 EntryIndex = 0; EntryIndex < InEntries.Num(); ++EntryIndex)
-		{
-			UAnimNextScheduleEntry* Entry = InEntries[EntryIndex];
-
-			auto CheckTermDirectionCompatibility = [&Log](FName InName, EScheduleTermDirection InExistingDirection, EScheduleTermDirection InNewDirection)
-			{
-				switch(InExistingDirection)
-				{
-				case EScheduleTermDirection::Input:
-					// Input before output: error
-					Log.Error(FText::Format(LOCTEXT("TermInputError", "Term '{0}' was used as an input before it was output"), FText::FromName(InName)));
-					return false;
-				case EScheduleTermDirection::Output:
-					return true;
-				}
-
-				return false;
-			};
-
-			if (UAnimNextScheduleEntry_Port* PortEntry = Cast<UAnimNextScheduleEntry_Port>(Entry))
-			{
-				bool bValid = true;
-
-				if(PortEntry->Port == nullptr)
-				{
-					Log.Error(LOCTEXT("InvalidPortError", "Invalid port class found"));
-					bValid = false;
-				}
-				else
-				{
-					UAnimNextSchedulePort* CDO = PortEntry->Port->GetDefaultObject<UAnimNextSchedulePort>();
-					check(CDO);
-
-					TConstArrayView<FScheduleTerm> Terms = CDO->GetTerms();
-					if(PortEntry->Terms.Num() != Terms.Num())
-					{
-						Log.Error(FText::Format(LOCTEXT("PortIncorrectTermCountError", "Incorrect term count for port: {0}"), FText::AsNumber(PortEntry->Terms.Num())));
-						bValid = false;
-					}
-
-					for(int32 TermIndex = 0; TermIndex < PortEntry->Terms.Num(); ++TermIndex)
-					{
-						FName TermName = PortEntry->Terms[TermIndex].Name;
-						if(!PortEntry->Terms[TermIndex].Type.IsValid())
-						{
-							Log.Error(FText::Format(LOCTEXT("PortIncorrectTermTypeError", "Invalid type when processing port term, ignored: '{0}'"), FText::FromName(TermName)));
-							bValid = false;
-						}
-						else
-						{
-							const uint32* ExistingIntermediateIndexPtr = IntermediateMap.Find(TermName);
-							if(ExistingIntermediateIndexPtr != nullptr)
-							{
-								const FAnimNextScheduleEntryTerm& IntermediateTerm = IntermediateTerms[*ExistingIntermediateIndexPtr];
-								if(IntermediateTerm.Type != Terms[TermIndex].GetType())
-								{
-									Log.Error(FText::Format(LOCTEXT("PortMismatchedTermTypeError", "Mismatched types when processing port term, ignored: '{0}'"), FText::FromName(TermName)));
-									bValid = false;
-								}
-
-								if(!CheckTermDirectionCompatibility(TermName, IntermediateTerm.Direction, Terms[TermIndex].Direction))
-								{
-									bValid = false;
-								}
-							}
-						}
-					}
-				}
-
-				if(bValid)
-				{
-					EmitPrerequisite();
-
-					FAnimNextSchedulePortTask PortTask;
-					PortTask.TaskIndex = InSchedule->Ports.Num();
-					PortTask.ParamScopeIndex = ParentScopeIndex;
-					PortTask.Port = PortEntry->Port;
-
-					for(int32 TermIndex = 0; TermIndex < PortEntry->Terms.Num(); ++TermIndex)
-					{
-						FName TermName = PortEntry->Terms[TermIndex].Name;
-						const uint32* ExistingIntermediateIndexPtr = IntermediateMap.Find(TermName);
-						if(ExistingIntermediateIndexPtr == nullptr)
-						{
-							uint32 IntermediateIndex = IntermediateTerms.Emplace(TermName, PortEntry->Terms[TermIndex].Type, PortEntry->Terms[TermIndex].Direction);
-							IntermediateMap.Add(TermName, IntermediateIndex);
-							PortTask.Terms.Add(IntermediateIndex);
-						}
-						else
-						{
-							PortTask.Terms.Add(*ExistingIntermediateIndexPtr);
-						}
-					}
-
-					int32 PortIndex = InSchedule->Ports.Add(PortTask);
-
-					InSchedule->NumTickFunctions++;
-
-					Emit(EAnimNextScheduleScheduleOpcode::RunPort, PortIndex);
-				}
-			}
-			else if (UAnimNextScheduleEntry_AnimNextGraph* GraphEntry = Cast<UAnimNextScheduleEntry_AnimNextGraph>(Entry))
-			{
-				bool bValid = true;
-
-				if(GraphEntry->Module == nullptr && !GraphEntry->DynamicGraph.IsValid())
-				{
-					Log.Error(LOCTEXT("InvalidGraphOrParameterError", "Invalid graph or invalid parameter supplied in graph task"));
-					bValid = false;
-				}
-				else if(GraphEntry->Module != nullptr)
-				{
-					TConstArrayView<FScheduleTerm> Terms = GraphEntry->Module->GetTerms();
-					if(GraphEntry->Terms.Num() != Terms.Num())
-					{
-						Log.Error(FText::Format(LOCTEXT("GraphIncorrectTermCountError", "Incorrect term count for graph: {0}"), FText::AsNumber(GraphEntry->Terms.Num())));
-						bValid = false;
-					}
-					else
-					{
-						// Validate graph terms match schedule-expected terms
-						for(int32 TermIndex = 0; TermIndex < GraphEntry->Terms.Num(); ++TermIndex)
-						{
-							FName TermName = GraphEntry->Terms[TermIndex].Name;
-							if(Terms[TermIndex].Direction != GraphEntry->Terms[TermIndex].Direction)
-							{
-								Log.Error(FText::Format(LOCTEXT("MismatchedDirectionInGraphTermError", "Mismatched direction when processing graph term, ignored: '{0}'"), FText::FromName(TermName)));
-								bValid = false;
-							}
-							
-							if(Terms[TermIndex].GetType() != GraphEntry->Terms[TermIndex].Type)
-							{
-								Log.Error(FText::Format(LOCTEXT("MismatchedTypesInGraphTermError", "Mismatched types when processing graph term, ignored: '{0}'"), FText::FromName(TermName)));
-								bValid = false;
-							}
-						}
-					}
-				}
-
-				// We must have a reference pose
-				if (!GraphEntry->ReferencePose.IsValid())
-				{
-					Log.Error(LOCTEXT("InvalidRefPoseError", "Invalid reference pose supplied to graph task"));
-					bValid = false;
-				}
-
-				// Validate terms and check against priors
-				for(int32 TermIndex = 0; TermIndex < GraphEntry->Terms.Num(); ++TermIndex)
-				{
-					FName TermName = GraphEntry->Terms[TermIndex].Name;
-					if(!GraphEntry->Terms[TermIndex].Type.IsValid())
-					{
-						Log.Error(FText::Format(LOCTEXT("InvalidTypeInGraphTermError", "Invalid type when processing graph term, ignored: '{0}'"), FText::FromName(TermName)));
-						bValid = false;
-					}
-					else
-					{
-						const uint32* ExistingIntermediateIndexPtr = IntermediateMap.Find(TermName);
-						if(ExistingIntermediateIndexPtr != nullptr)
-						{
-							const FAnimNextScheduleEntryTerm& IntermediateTerm = IntermediateTerms[*ExistingIntermediateIndexPtr];
-							if(IntermediateTerm.Type != GraphEntry->Terms[TermIndex].Type)
-							{
-								Log.Error(FText::Format(LOCTEXT("MismatchedTypeInGraphTermError", "Mismatched types when processing graph term, ignored: '{0}'"), FText::FromName(TermName)));
-								bValid = false;
-							}
-
-							if(!CheckTermDirectionCompatibility(TermName, IntermediateTerm.Direction, GraphEntry->Terms[TermIndex].Direction))
-							{
-								bValid = false;
-							}
-						}
-					}
-				}
-
-				if(bValid)
-				{
-					EmitPrerequisite();
-
-					FAnimNextScheduleGraphTask GraphTask;
-					GraphTask.TaskIndex = InSchedule->GraphTasks.Num();
-					GraphTask.ParamScopeIndex = InSchedule->NumParameterScopes++;
-					GraphTask.ParamParentScopeIndex = ParentScopeIndex;
-					GraphTask.EntryPoint = FAnimNextParam(GraphEntry->EntryPoint);
-					GraphTask.Module = GraphEntry->Module;
-					GraphTask.DynamicModule = FAnimNextParam(GraphEntry->DynamicGraph);
-					GraphTask.ReferencePose = FAnimNextParam(GraphEntry->ReferencePose);
-					GraphTask.LOD = FAnimNextParam(GraphEntry->LOD);
-					if(GraphEntry->Module == nullptr && GraphEntry->DynamicGraph.IsValid())
-					{
-						Algo::Transform(GraphEntry->RequiredParameters, GraphTask.SuppliedParameters, [](const FAnimNextEditorParam& InParam){ return FAnimNextParam(InParam); });
-						GraphTask.SuppliedParametersHash = SortAndHashParameters(GraphTask.SuppliedParameters);
-					}
-
-					for(int32 TermIndex = 0; TermIndex < GraphEntry->Terms.Num(); ++TermIndex)
-					{
-						FName TermName = GraphEntry->Terms[TermIndex].Name;
-						const uint32* ExistingIntermediateIndexPtr = IntermediateMap.Find(TermName);
-						if(ExistingIntermediateIndexPtr == nullptr)
-						{
-							uint32 IntermediateIndex = IntermediateTerms.Emplace(TermName, GraphEntry->Terms[TermIndex].Type, GraphEntry->Terms[TermIndex].Direction);
-							IntermediateMap.Add(TermName, IntermediateIndex);
-							GraphTask.Terms.Add(IntermediateIndex);
-						}
-						else
-						{
-							GraphTask.Terms.Add(*ExistingIntermediateIndexPtr);
-						}
-					}
-
-					int32 TaskIndex = InSchedule->GraphTasks.Add(GraphTask);
-
-					InSchedule->NumTickFunctions++;
-
-					Emit(EAnimNextScheduleScheduleOpcode::RunGraphTask, TaskIndex);
-				}
-			}
-			else if (UAnimNextScheduleEntry_ExternalTask* ExternalTaskEntry = Cast<UAnimNextScheduleEntry_ExternalTask>(Entry))
-			{
-				EmitPrerequisite();
-
-				FAnimNextScheduleExternalTask ExternalTask;
-				ExternalTask.TaskIndex = InSchedule->ExternalTasks.Num();
-				ExternalTask.ParamScopeIndex = InSchedule->NumParameterScopes++;
-				ExternalTask.ParamParentScopeIndex = ParentScopeIndex;
-				ExternalTask.ExternalTask = FAnimNextParam(ExternalTaskEntry->ExternalTask);
-				int32 ExternalTaskIndex = InSchedule->ExternalTasks.Add(ExternalTask);
-
-				// Emit the external task
-				Emit(EAnimNextScheduleScheduleOpcode::BeginRunExternalTask, ExternalTaskIndex);
-				InSchedule->NumTickFunctions++;
-
-				EmitPrerequisite();
-
-				Emit(EAnimNextScheduleScheduleOpcode::EndRunExternalTask, ExternalTaskIndex);
-				InSchedule->NumTickFunctions++;
-			}
-			else if (UAnimNextScheduleEntry_ParamScope* ParamScopeTaskEntry = Cast<UAnimNextScheduleEntry_ParamScope>(Entry))
-			{
-				EmitPrerequisite();
-
-				FAnimNextScheduleParamScopeEntryTask ParamScopeEntryTask;
-				ParamScopeEntryTask.TaskIndex = InSchedule->ParamScopeEntryTasks.Num();
-				const uint32 ParamScopeIndex = InSchedule->NumParameterScopes++;
-				ParamScopeEntryTask.ParamScopeIndex = ParamScopeIndex;
-				ParamScopeEntryTask.ParamParentScopeIndex = ParentScopeIndex;
-				ParamScopeEntryTask.TickFunctionIndex = InSchedule->NumTickFunctions;
-				ParamScopeEntryTask.Scope = FAnimNextParam(ParamScopeTaskEntry->Scope);
-				ParamScopeEntryTask.Parameters = ParamScopeTaskEntry->Parameters;
-				int32 ParamScopeTaskEntryIndex = InSchedule->ParamScopeEntryTasks.Add(ParamScopeEntryTask);
-
-				Emit(EAnimNextScheduleScheduleOpcode::RunParamScopeEntry, ParamScopeTaskEntryIndex);
-				InSchedule->NumTickFunctions++;
-
-				// Enter new scope
-				uint32 PreviousParentScope = ParentScopeIndex;
-				ParentScopeIndex = ParamScopeIndex;
-
-				// Emit the subentries
-				EmitEntries(ParamScopeTaskEntry->SubEntries);
-
-				// Exit scope
-				ParentScopeIndex = PreviousParentScope;
-
-				EmitPrerequisite();
-
-				FAnimNextScheduleParamScopeExitTask ParamScopeExitTask;
-				ParamScopeExitTask.TaskIndex = InSchedule->ParamScopeExitTasks.Num();
-				ParamScopeExitTask.ParamScopeIndex = ParamScopeIndex;
-				ParamScopeExitTask.Scope = FAnimNextParam(ParamScopeTaskEntry->Scope);
-				int32 ParamScopeExitTaskIndex = InSchedule->ParamScopeExitTasks.Add(ParamScopeExitTask);
-
-				Emit(EAnimNextScheduleScheduleOpcode::RunParamScopeExit, ParamScopeExitTaskIndex);
-				InSchedule->NumTickFunctions++;
-			}
-			else if(UAnimNextScheduleEntry_ExternalParams* ExternalParamsTaskEntry = Cast<UAnimNextScheduleEntry_ExternalParams>(Entry))
-			{
-				EmitPrerequisite();
-
-				FAnimNextScheduleExternalParamTask ExternalParamTask;
-				ExternalParamTask.TaskIndex = InSchedule->ExternalParamTasks.Num();
-				ExternalParamTask.ParameterSources = ExternalParamsTaskEntry->ParameterSources;
-				ExternalParamTask.bThreadSafe = ExternalParamsTaskEntry->bThreadSafe;
-				int32 ExternalParamTaskEntryIndex = InSchedule->ExternalParamTasks.Add(ExternalParamTask);
-
-				Emit(EAnimNextScheduleScheduleOpcode::RunExternalParamTask, ExternalParamTaskEntryIndex);
-				InSchedule->NumTickFunctions++;
-			}
-		}
-	};
-
-	auto GenerateExternalParameters = [&InSchedule](TArray<TObjectPtr<UAnimNextScheduleEntry>>& InEntries)
-	{
-		struct FParameterTracker
-		{
-			TInstancedStruct<FAnimNextParamInstanceIdentifier> InstanceId;
-			TArray<TObjectPtr<UAnimNextScheduleEntry>>* BestContainer = nullptr;
-			TSet<FName> ThreadSafeParameters;
-			TSet<FName> NonThreadSafeParameters;
-			uint32 BestDistance = MAX_uint32;
-			uint32 BestArrayIndex = MAX_uint32;
-		};
-
-		FModule& Module = FModuleManager::Get().LoadModuleChecked<FModule>("AnimNextUncookedOnly");
-
-		// We need to find the task that is 'earliest' in the DAG for each external parameter, so we track that with this map
-		TMap<FName, FParameterTracker> TrackerMap;
-		auto TrackExternalParameters = [&InSchedule, &TrackerMap, &Module](TConstArrayView<FAnimNextEditorParam> InParameters, uint32 InDistance, UAnimNextScheduleEntry* InEntry, TArray<TObjectPtr<UAnimNextScheduleEntry>>* InContainer, int32 InArrayIndex, bool bInUpdateDependent)
-		{
-			using namespace UE::UniversalObjectLocator;
-
-			bool bHasExternal = false;
-			for(const FAnimNextEditorParam& Parameter : InParameters)
-			{
-				// Only add if the parameter name is 'external' (i.e. it has a valid instance ID)
-				if(Parameter.InstanceId.IsValid() && Parameter.InstanceId.Get().IsValid())
-				{
-					FParameterTracker& Tracker = TrackerMap.FindOrAdd(Parameter.InstanceId.Get().ToName());
-					Tracker.InstanceId = Parameter.InstanceId;
-
-					// Only track array index if this param is update dependent
-					if(bInUpdateDependent && InDistance < Tracker.BestDistance)
-					{
-						Tracker.BestDistance = InDistance;
-						Tracker.BestContainer = InContainer;
-						Tracker.BestArrayIndex = InArrayIndex == 0 ? 0 : InArrayIndex - 1;
-					}
-
-					FParameterSourceInfo Info[1];
-					if(TSharedPtr<IParameterSourceType> SourceType = Module.FindParameterSourceType(Parameter.InstanceId.GetScriptStruct()))
-					{
-						SourceType->FindParameterInfo(Parameter.InstanceId, { Parameter.Name }, Info);
-					}
-
-					// Note that if the above call to FindParameterSourceType or FindParameterInfo fails, this will still default to non-thread safe
-					if(Info[0].bThreadSafe)
-					{
-						Tracker.ThreadSafeParameters.Add(Parameter.Name);
-					}
-					else
-					{
-						Tracker.NonThreadSafeParameters.Add(Parameter.Name);
-					}
-
-					bHasExternal = true;
-				}
-			}
-			return bHasExternal;
-		};
-
-		TFunction<void(TArray<TObjectPtr<UAnimNextScheduleEntry>>&, uint32)> PopulateExternalParameters;
-		PopulateExternalParameters = [&PopulateExternalParameters, &TrackExternalParameters](TArray<TObjectPtr<UAnimNextScheduleEntry>>& InEntries, uint32 InDistance)
-		{
-			// First populate internal parameters for all those tasks entries that reference them
-			int32 ArrayIndex = 0;
-			for(UAnimNextScheduleEntry* Entry : InEntries)
-			{
-				if (UAnimNextScheduleEntry_AnimNextGraph* GraphEntry = Cast<UAnimNextScheduleEntry_AnimNextGraph>(Entry))
-				{
-					if(GraphEntry->Module)
-					{
-						FAnimNextParameterProviderAssetRegistryExports Exports;
-						if(UncookedOnly::FUtils::GetExportedParametersForAsset(FAssetData(GraphEntry->Module), Exports))
-						{
-							TArray<FAnimNextEditorParam> RequiredParameters;
-							RequiredParameters.Reserve(Exports.Parameters.Num());
-							for(const FAnimNextParameterAssetRegistryExportEntry& ExportedParameter : Exports.Parameters)
-							{
-								RequiredParameters.Emplace(ExportedParameter.Name, ExportedParameter.Type, ExportedParameter.InstanceId);
-							}
-							TrackExternalParameters(RequiredParameters, InDistance, Entry, &InEntries, ArrayIndex, true);
-						}
-					}
-
-					TrackExternalParameters(GraphEntry->RequiredParameters, InDistance, Entry, &InEntries, ArrayIndex, true);
-
-					TrackExternalParameters({ GraphEntry->DynamicGraph, GraphEntry->EntryPoint, GraphEntry->ReferencePose, GraphEntry->LOD }, InDistance, Entry, &InEntries, ArrayIndex, true);
-				}
-				else if (UAnimNextScheduleEntry_Port* PortEntry = Cast<UAnimNextScheduleEntry_Port>(Entry))
-				{
-					if(PortEntry->Port)
-					{
-						UAnimNextSchedulePort* CDO = PortEntry->Port->GetDefaultObject<UAnimNextSchedulePort>();
-						TConstArrayView<FAnimNextEditorParam> RequiredParameters = CDO->GetRequiredParameters();
-						TrackExternalParameters(RequiredParameters, InDistance, Entry, &InEntries, ArrayIndex, true);
-					}
-				}
-				else if (UAnimNextScheduleEntry_ExternalTask* ExternalTaskEntry = Cast<UAnimNextScheduleEntry_ExternalTask>(Entry))
-				{
-					// Note: external task params are not update dependent as this would cause external param updates to occur before tick functions
-					TrackExternalParameters({ ExternalTaskEntry->ExternalTask }, InDistance, Entry, &InEntries, ArrayIndex, false);
-				}
-				else if (UAnimNextScheduleEntry_ParamScope* ParamScopeTaskEntry = Cast<UAnimNextScheduleEntry_ParamScope>(Entry))
-				{
-					for(UAnimNextModule* Parameters : ParamScopeTaskEntry->Parameters)
-					{
-						if(Parameters)
-						{
-							FAnimNextParameterProviderAssetRegistryExports Exports;
-							if(UncookedOnly::FUtils::GetExportedParametersForAsset(FAssetData(Parameters), Exports))
-							{
-								TArray<FAnimNextEditorParam> RequiredParameters;
-								RequiredParameters.Reserve(Exports.Parameters.Num());
-								for(const FAnimNextParameterAssetRegistryExportEntry& ExportedParameter : Exports.Parameters)
-								{
-									RequiredParameters.Emplace(ExportedParameter.Name, ExportedParameter.Type, ExportedParameter.InstanceId);
-								}
-								TrackExternalParameters(RequiredParameters, InDistance, Entry, &InEntries, ArrayIndex, true);
-							}
-						}
-					}
-
-					// Recurse into sub-entries
-					PopulateExternalParameters(ParamScopeTaskEntry->SubEntries, InDistance);
-				}
-
-				InDistance++;
-				ArrayIndex++;
-			}
-		};
-
-		// First populate internal parameter lists and tracker map
-		PopulateExternalParameters(InEntries, 0);
-
-		// Unique location: container, index and thread safe flag
-		using FInsertionLocation = TTuple<TArray<TObjectPtr<UAnimNextScheduleEntry>>*, uint32, bool>;
-
-		// Build sources that need to run (thread-safe or not) at each index
-		TMap<FInsertionLocation, TArray<FAnimNextScheduleExternalParameterSource>> InsertionMap;
-		for(const TPair<FName, FParameterTracker>& TrackedParameterPair : TrackerMap) //-V1078
-		{
-			// If an index/container was not set up, then the update of the set of parameters is not a pre-requisite of a task, so we just insert the task at the
-			// start of the schedule as the only requirement is that the parameters exist 
-			TArray<TObjectPtr<UAnimNextScheduleEntry>>* Container = TrackedParameterPair.Value.BestContainer != nullptr ? TrackedParameterPair.Value.BestContainer : &InEntries;
-			uint32 ArrayIndex = TrackedParameterPair.Value.BestArrayIndex != MAX_uint32 ? TrackedParameterPair.Value.BestArrayIndex : 0;
-			
-			if(TrackedParameterPair.Value.ThreadSafeParameters.Num() > 0)
-			{
-				TArray<FAnimNextScheduleExternalParameterSource>& ParameterSources = InsertionMap.FindOrAdd({ Container, ArrayIndex, true });
-				FAnimNextScheduleExternalParameterSource& NewSource = ParameterSources.AddDefaulted_GetRef();
-				NewSource.InstanceId = TrackedParameterPair.Value.InstanceId;
-				NewSource.Parameters = TrackedParameterPair.Value.ThreadSafeParameters.Array();
-			}
-
-			if(TrackedParameterPair.Value.NonThreadSafeParameters.Num() > 0)
-			{
-				TArray<FAnimNextScheduleExternalParameterSource>& ParameterSources = InsertionMap.FindOrAdd({ Container, ArrayIndex, false });
-				FAnimNextScheduleExternalParameterSource& NewSource = ParameterSources.AddDefaulted_GetRef();
-				NewSource.InstanceId = TrackedParameterPair.Value.InstanceId;
-				NewSource.Parameters = TrackedParameterPair.Value.NonThreadSafeParameters.Array();
-			}
-		}
-
-		// Sort sources map by insertion index
-		InsertionMap.KeySort([](const FInsertionLocation& InLHS, const FInsertionLocation& InRHS)
-		{
-			return InLHS.Get<1>() > InRHS.Get<1>();
-		});
-
-		// Now insert a task to fetch the parameters at the recorded index/container
-		// NOTE: here we just insert the task before the earliest usage of the external parameter source we found, but in the case of a full DAG
-		// schedule, we would need to add a prerequisite for ALL tasks that use the external parameters
-		for(const TPair<FInsertionLocation, TArray<FAnimNextScheduleExternalParameterSource>>& InsertionLocationPair : InsertionMap)
-		{
-			// Add a new external param entry
-			UAnimNextScheduleEntry_ExternalParams* NewParamEntry = NewObject<UAnimNextScheduleEntry_ExternalParams>();
-			NewParamEntry->bThreadSafe = InsertionLocationPair.Key.Get<2>();
-			NewParamEntry->ParameterSources = InsertionLocationPair.Value;
-			InsertionLocationPair.Key.Get<0>()->Insert(NewParamEntry, InsertionLocationPair.Key.Get<1>());
-		}
-	};
-
-	// Duplicate the entries, we are going to rewrite them
-	TArray<TObjectPtr<UAnimNextScheduleEntry>> NewEntries;
-	for(UAnimNextScheduleEntry* Entry : InSchedule->Entries)
-	{
-		if(Entry)
-		{
-			NewEntries.Add(CastChecked<UAnimNextScheduleEntry>(StaticDuplicateObject(Entry, GetTransientPackage())));
-		}
-	}
-
-	// Push required parameters up scopes
-	GenerateExternalParameters(NewEntries);
-
-	// Emit the schedule 'bytecode'
-	EmitEntries(NewEntries);
-
-	Emit(EAnimNextScheduleScheduleOpcode::Exit);
-
-	// Process intermediates
-	if(IntermediateMap.Num() > 0)
-	{
-		check(IntermediateMap.Num() == IntermediateTerms.Num());
-		
-		TArray<FPropertyBagPropertyDesc> PropertyDescs;
-		PropertyDescs.Reserve(IntermediateTerms.Num());
-		 
-		for(const TPair<FName, uint32>& IntermediatePair : IntermediateMap)
-		{
-			const FAnimNextParamType& IntermediateType = IntermediateTerms[IntermediatePair.Value].Type;
-			check(IntermediateType.IsValid());
-			PropertyDescs.Emplace(IntermediatePair.Key, IntermediateType.GetContainerType(), IntermediateType.GetValueType(), IntermediateType.GetValueTypeObject());
-		}
-
-		InSchedule->IntermediatesData.AddProperties(PropertyDescs);
-	}
-
-	FScheduler::OnScheduleCompiled(InSchedule);
-
-	InSchedule->CompiledEvent.Broadcast();
-}
-
-uint64 FUtils::SortAndHashParameters(TArray<FAnimNextParam>& InParameters)
-{
-	Algo::Sort(InParameters, [](const FAnimNextParam& InLHS, const FAnimNextParam& InRHS)
-	{
-		return InLHS.Name.LexicalLess(InRHS.Name);
-	});
-
-	uint64 Hash = 0;
-	for(const FAnimNextParam& Parameter : InParameters)
-	{
-		FString ExportedString;
-		FAnimNextParam::StaticStruct()->ExportText(ExportedString, &Parameter, nullptr, nullptr, PPF_None, nullptr);
-		Hash = CityHash64WithSeed(reinterpret_cast<const char*>(*ExportedString), ExportedString.Len() * sizeof(TCHAR), Hash);
-	}
-
-	return Hash;
-}
-
 const FText& FUtils::GetFunctionLibraryDisplayName()
 {
 	static const FText FunctionLibraryName = LOCTEXT("WorkspaceFunctionLibraryName", "Function Library");
diff --git a/Engine/Plugins/Experimental/AnimNext/Source/AnimNextUncookedOnly/Private/Variables/UniversalObjectLocatorBindingType.h b/Engine/Plugins/Experimental/AnimNext/Source/AnimNextUncookedOnly/Private/Variables/UniversalObjectLocatorBindingType.h
new file mode 100644
--- /dev/null
+++ b/Engine/Plugins/Experimental/AnimNext/Source/AnimNextUncookedOnly/Private/Variables/UniversalObjectLocatorBindingType.h
@@ -0,0 +1,41 @@
+﻿// Copyright Epic Games, Inc. All Rights Reserved.
+
+#pragma once
+
+#include "CoreMinimal.h"
+#include "Variables/IVariableBindingType.h"
+
+namespace UE::AnimNext::UncookedOnly
+{
+	struct FClassProxy;
+}
+
+namespace UE::AnimNext::UncookedOnly
+{
+
+// Provides information about object proxy parameter sources
+class FUniversalObjectLocatorBindingType : public IVariableBindingType
+{
+public:
+	FUniversalObjectLocatorBindingType();
+	virtual ~FUniversalObjectLocatorBindingType() override;
+
+private:
+	// IVariableBindingType interface
+	virtual TSharedRef<SWidget> CreateEditWidget(const TSharedRef<IPropertyHandle>& InPropertyHandle, const FAnimNextParamType& InType) const override;
+	virtual FText GetDisplayText(TConstStructView<FAnimNextVariableBindingData> InBindingData) const override;
+	virtual FText GetTooltipText(TConstStructView<FAnimNextVariableBindingData> InBindingData) const override;
+	virtual void BuildBindingGraphFragment(const FRigVMCompileSettings& InSettings, const FBindingGraphFragmentArgs& InArgs, URigVMPin*& OutExecTail, FVector2D& OutLocation) const override;
+
+	static const UClass* GetClass(TConstStructView<FAnimNextVariableBindingData> InBindingData);
+
+	TSharedRef<FClassProxy> FindOrCreateClassProxy(const UClass* InClass) const; 
+
+	// Cached map of classes -> proxy
+	mutable TMap<TObjectKey<UClass>, TSharedPtr<FClassProxy>> ClassMap;
+
+	// Hook for object reinstancing
+	FDelegateHandle OnObjectsReinstancedHandle;
+};
+
+}
diff --git a/Engine/Plugins/Experimental/AnimNext/Source/AnimNextUncookedOnly/Public/IAnimNextUncookedOnlyModule.h b/Engine/Plugins/Experimental/AnimNext/Source/AnimNextUncookedOnly/Public/IAnimNextUncookedOnlyModule.h
--- a/Engine/Plugins/Experimental/AnimNext/Source/AnimNextUncookedOnly/Public/IAnimNextUncookedOnlyModule.h
+++ b/Engine/Plugins/Experimental/AnimNext/Source/AnimNextUncookedOnly/Public/IAnimNextUncookedOnlyModule.h
@@ -4,14 +4,12 @@
 
 #include "CoreMinimal.h"
 #include "Modules/ModuleInterface.h"
-#include "Param/AnimNextParamInstanceIdentifier.h"
 
 template<typename T> struct TInstancedStruct;
 
 namespace UE::AnimNext::UncookedOnly
 {
-	struct FParameterSourceInfo;
-	class IParameterSourceType;
+	class IVariableBindingType;
 }
 
 namespace UE::AnimNext::UncookedOnly
@@ -20,18 +18,18 @@ namespace UE::AnimNext::UncookedOnly
 class IAnimNextUncookedOnlyModule : public IModuleInterface
 {
 public:
-	// Register a parameter source type that can be used to query information about parameter sources
-	// @param   InInstanceIdStruct     The name of struct type for this parameter source's instance ID (must be a child struct of FAnimNextParamInstanceIdentifier)
-	// @param   InFactory              The type to register
-	virtual void RegisterParameterSourceType(const UScriptStruct* InInstanceIdStruct, TSharedPtr<IParameterSourceType> InType) = 0;
+	// Register a variable binding type, used to query information and process variable bindings
+	// @param   InStructName     The full path name of the struct type for this variable binding's data (must be a child struct of FAnimNextVariableBindingData)
+	// @param   InType           The type to register
+	virtual void RegisterVariableBindingType(FName InStructName, TSharedPtr<IVariableBindingType> InType) = 0;
 	
-	// Unregister a parameter source type previously passed to RegisterParameterSourceType
-	// @param   InInstanceIdStruct     The name of struct type for this parameter source's instance ID
-	virtual void UnregisterParameterSourceType(const UScriptStruct* InInstanceIdStruct) = 0;
+	// Unregister a variable binding type previously passed to RegisterVariableBindingType
+	// @param   InStructName     The full path name of the struct type for this variable binding's data
+	virtual void UnregisterVariableBindingType(FName InStructName) = 0;
 
-	// Find a parameter source type previously passed to RegisterParameterSourceType
-	// @param   InInstanceIdStruct     The name of struct type for this parameter source's instance ID
-	virtual TSharedPtr<IParameterSourceType> FindParameterSourceType(const UScriptStruct* InInstanceIdStruct) const = 0;
+	// Find a variable binding type previously passed to RegisterVariableBindingType
+	// @param   InStruct         The struct type for this variable binding's data (must be a child struct of FAnimNextVariableBindingData)
+	virtual TSharedPtr<IVariableBindingType> FindVariableBindingType(const UScriptStruct* InInstanceIdStruct) const = 0;
 };
 
 }
