AnimNext: Reinstate workspace outliner menu functionality

Merged out the call to FAnimNextAssetItemDetails::RegisterToolMenuExtensions()

diff --git a/Engine/Plugins/Experimental/AnimNext/Source/AnimNextEditor/Private/AnimNextEditorModule.cpp b/Engine/Plugins/Experimental/AnimNext/Source/AnimNextEditor/Private/AnimNextEditorModule.cpp
--- a/Engine/Plugins/Experimental/AnimNext/Source/AnimNextEditor/Private/AnimNextEditorModule.cpp
+++ b/Engine/Plugins/Experimental/AnimNext/Source/AnimNextEditor/Private/AnimNextEditorModule.cpp
@@ -231,6 +231,7 @@ void FAnimNextEditorModule::StartupModule()
 	WorkspaceModule.RegisterWorkspaceItemDetails(Workspace::FOutlinerItemDetailsId(FAnimNextAnimationGraphOutlinerData::StaticStruct()->GetFName()), AssetItemDetails);
 
 	FAnimNextGraphItemDetails::RegisterToolMenuExtensions();
+	FAnimNextAssetItemDetails::RegisterToolMenuExtensions();
 
 	SupportedAssetClasses.Append(
 		{
