AnimNext: Fix duplicate localization key

diff --git a/Engine/Plugins/Experimental/AnimNext/Source/AnimNextEditor/Private/Common/AnimNextEntryAssetDefinitions.h b/Engine/Plugins/Experimental/AnimNext/Source/AnimNextEditor/Private/Common/AnimNextEntryAssetDefinitions.h
--- a/Engine/Plugins/Experimental/AnimNext/Source/AnimNextEditor/Private/Common/AnimNextEntryAssetDefinitions.h
+++ b/Engine/Plugins/Experimental/AnimNext/Source/AnimNextEditor/Private/Common/AnimNextEntryAssetDefinitions.h
@@ -30,7 +30,7 @@ class UAssetDefinition_AnimNextAnimationGraphEntry : public UAssetDefinitionDefa
 
 public:
 	// UAssetDefinition interface
-	virtual FText GetAssetDisplayName() const override { return LOCTEXT("AnimNextAnimationGraph", "Animation Graph"); }
+	virtual FText GetAssetDisplayName() const override { return LOCTEXT("AnimNextAnimationGraphEntry", "Animation Graph"); }
 	virtual FLinearColor GetAssetColor() const override { return FLinearColor(FColor(128,64,64)); }
 	virtual TSoftClassPtr<UObject> GetAssetClass() const override { return UAnimNextAnimationGraphEntry::StaticClass(); }
 	virtual FText GetObjectDisplayNameText(UObject* Object) const override;
