AnimNext: Fix crash when selecting a trait stack with nested arrays of structs

Changes general property changed callback over to per-sub-property handle callbacks. The callback sets the value on the pin by converting the property path to a pin path on-demand.

diff --git a/Engine/Plugins/Experimental/AnimNext/Source/AnimNextEditor/Private/Graph/AnimNextEdGraphNodeCustomization.cpp b/Engine/Plugins/Experimental/AnimNext/Source/AnimNextEditor/Private/Graph/AnimNextEdGraphNodeCustomization.cpp
--- a/Engine/Plugins/Experimental/AnimNext/Source/AnimNextEditor/Private/Graph/AnimNextEdGraphNodeCustomization.cpp
+++ b/Engine/Plugins/Experimental/AnimNext/Source/AnimNextEditor/Private/Graph/AnimNextEdGraphNodeCustomization.cpp
@@ -123,14 +123,14 @@ void FAnimNextEdGraphNodeCustomization::GenerateTraitData(UAnimNextEdGraphNode*
 					FTraitStackDetailsData* TraitData = nullptr;
 					if (UScriptStruct* TraitSharedInstanceData = Trait->GetTraitSharedDataStruct())
 					{
-						const TSharedPtr<FCategoryDetailsData>* TraitDataPtr = CategoryDetailsData.FindByPredicate([TraitSharedInstanceData](const TSharedPtr<FCategoryDetailsData>& InItem)
+						const TSharedPtr<FCategoryDetailsData>* TraitDataPtr = CategoryDetailsData.FindByPredicate([TraitSharedInstanceData, TraitPin](const TSharedPtr<FCategoryDetailsData>& InItem)
 							{
-								return InItem->Type == FCategoryDetailsData::EType::TraitStack && InItem->Name == TraitSharedInstanceData->GetFName();
+								return InItem->Type == FCategoryDetailsData::EType::TraitStack && InItem->Name == TraitPin->GetFName();
 							});
 
 						if (TraitDataPtr == nullptr)
 						{
-							TraitDataPtr = &CategoryDetailsData.Add_GetRef(MakeShared<FTraitStackDetailsData>(TraitSharedInstanceData->GetFName()));
+							TraitDataPtr = &CategoryDetailsData.Add_GetRef(MakeShared<FTraitStackDetailsData>(TraitPin->GetFName()));
 						}
 
 						TraitData = static_cast<FTraitStackDetailsData*>(TraitDataPtr->Get());
@@ -247,7 +247,7 @@ void FAnimNextEdGraphNodeCustomization::PopulateCategory(IDetailLayoutBuilder& D
 	const FString TraitDisplayName = *TraitData->ScopedSharedDataInstances[0]->GetStruct()->GetDisplayNameText().ToString();
 	const FName CategoryName = (TraitData->EdGraphNodes.Num() == 1)
 		? *TraitDisplayName
-		: *FString::Printf(TEXT("%s (% d)"), *TraitDisplayName, TraitData->EdGraphNodes.Num());
+		: *FString::Printf(TEXT("%s (%d)"), *TraitDisplayName, TraitData->EdGraphNodes.Num());
 
 	IDetailCategoryBuilder& ParameterCategory = DetailBuilder.EditCategory(CategoryName, FText::GetEmpty(), ECategoryPriority::Important);
 	
@@ -256,57 +256,98 @@ void FAnimNextEdGraphNodeCustomization::PopulateCategory(IDetailLayoutBuilder& D
 	AddPropertyParams.HideRootObjectNode(true);
 
 	IDetailPropertyRow* DetailPropertyRow = ParameterCategory.AddExternalStructureProperty(MakeShared<FStructOnScopeStructureDataProvider>(TraitData->ScopedSharedDataInstances), NAME_None, EPropertyLocation::Default, AddPropertyParams);
-	if (TSharedPtr<IPropertyHandle> PropertyHandle = DetailPropertyRow->GetPropertyHandle(); PropertyHandle.IsValid())
+	TSharedPtr<IPropertyHandle> PropertyHandle = DetailPropertyRow->GetPropertyHandle();
+	if (!PropertyHandle.IsValid())
 	{
-		const TWeakPtr<FTraitStackDetailsData> TraitDataWeak = TraitData.ToWeakPtr();
+		return;
+	}
+
+	TraitData->RootPropertyHandle = PropertyHandle;
 
-		const auto UpdatePinDefaultValue = [TraitDataWeak](const FPropertyChangedEvent& InEvent)
+	auto VisitChildren = [](const TSharedRef<IPropertyHandle>& InPropertyHandle, TFunctionRef<void(const TSharedRef<IPropertyHandle>&)> InPredicate)
+	{
+		auto VisitChildrenRecursive = [](const TSharedRef<IPropertyHandle>& InPropertyHandle, auto& InVisitChildren, TFunctionRef<void(const TSharedRef<IPropertyHandle>&)> InPredicate)
+		{
+			uint32 NumChildHandles = 0;
+			if(InPropertyHandle->GetNumChildren(NumChildHandles) != FPropertyAccess::Success)
 			{
-				if (const TSharedPtr<FTraitStackDetailsData> TraitData = TraitDataWeak.Pin())
+				return;
+			}
+
+			for(uint32 ChildHandleIndex = 0; ChildHandleIndex < NumChildHandles; ++ChildHandleIndex)
+			{
+				TSharedPtr<IPropertyHandle> Handle = InPropertyHandle->GetChildHandle(ChildHandleIndex);
+				if(Handle.IsValid())
 				{
-					// Avoid VM recompilation for each Set Default Value
-					FRigVMControllerCompileBracketScope CompileScope(TraitData->EdGraphNodes[0]->GetController());
+					InPredicate(Handle.ToSharedRef());
+					InVisitChildren(Handle.ToSharedRef(), InVisitChildren, InPredicate);
+				}
+			}
+		};
 
-					const int32 NumTraitInstances = TraitData->ScopedSharedDataInstances.Num();
-					for (int32 InstanceIndex = 0; InstanceIndex < NumTraitInstances; InstanceIndex++)
-					{
-						TWeakObjectPtr<UAnimNextEdGraphNode>& EdGraphNode = TraitData->EdGraphNodes[InstanceIndex];
-						if (EdGraphNode.IsValid())
-						{
-							TSharedPtr<FStructOnScope>& ScopedSharedData = TraitData->ScopedSharedDataInstances[InstanceIndex];
-
-							const bool bIsContainer = InEvent.Property->GetOwnerProperty()->IsA<FArrayProperty>()
-								|| InEvent.Property->GetOwnerProperty()->IsA<FMapProperty>()
-								|| InEvent.Property->GetOwnerProperty()->IsA<FSetProperty>();
-							
-							// For some reason, sub properties of a container does not come with the correct struct offsets, so getting the container property in that case
-							FProperty* Property = bIsContainer ? InEvent.Property->GetOwnerProperty() : InEvent.Property;
-							
-							// Extract the value from the property and assign it to the Pin as a default value (via Schema)
-							const uint8* StructMemberMemoryPtr = Property->ContainerPtrToValuePtr<uint8>(ScopedSharedData->GetStructMemory());
-							const FString ValueStr = FRigVMStruct::ExportToFullyQualifiedText(Property, StructMemberMemoryPtr, true);
-						
-							for (UEdGraphPin* EdGraphPin : EdGraphNode->Pins)
-							{
-								// Find the EdGraphPin that corresponds to the Property
-								const FString ModelPinName = FString::Printf(TEXT(".%s"), *Property->GetFName().ToString());
-								if (EdGraphPin->GetFName().ToString().EndsWith(ModelPinName, ESearchCase::CaseSensitive))
-								{
-									if (URigVMPin* ModelPin = EdGraphNode->FindModelPinFromGraphPin(EdGraphPin))
-									{
-										// Set the default value using the Controller
-										EdGraphNode->GetController()->SetPinDefaultValue(ModelPin->GetPinPath(), ValueStr);
-									}
-									break;
-								}
-							}
-						}
-					}
+		VisitChildrenRecursive(InPropertyHandle, VisitChildrenRecursive, InPredicate);
+	};
+
+	VisitChildren(PropertyHandle.ToSharedRef(), [&TraitData](const TSharedRef<IPropertyHandle>& InPropertyHandle)
+	{
+		TraitData->PropertyHandles.Add(InPropertyHandle);
+
+		const TWeakPtr<FTraitStackDetailsData> WeakTraitData = TraitData.ToWeakPtr();
+		const TWeakPtr<IPropertyHandle> WeakPropertyHandle = InPropertyHandle.ToWeakPtr();
+		InPropertyHandle->SetOnPropertyValueChanged(FSimpleDelegate::CreateLambda([WeakTraitData, WeakPropertyHandle]()
+		{
+			TSharedPtr<IPropertyHandle> PinnedPropertyHandle = WeakPropertyHandle.Pin(); 
+			if(!PinnedPropertyHandle.IsValid())
+			{
+				return;
+			}
+
+			TSharedPtr<FTraitStackDetailsData> PinnedTraitData = WeakTraitData.Pin(); 
+			if(!PinnedTraitData.IsValid())
+			{
+				return;
+			}
+
+			// Avoid VM recompilation for each Set Default Value
+			URigVMController* Controller = PinnedTraitData->EdGraphNodes[0]->GetController();
+			FRigVMControllerCompileBracketScope CompileScope(Controller);
+
+			PinnedPropertyHandle->EnumerateConstRawData([&PinnedTraitData, &PinnedPropertyHandle, Controller](const void* InRawData, const int32 InDataIndex, const int32 InNumDatas)
+			{
+				check(InNumDatas == PinnedTraitData->ScopedSharedDataInstances.Num());
+				check(InNumDatas == PinnedTraitData->EdGraphNodes.Num());
+				check(PinnedTraitData->ScopedSharedDataInstances.IsValidIndex(InDataIndex));
+
+				TWeakObjectPtr<UAnimNextEdGraphNode>& EdGraphNode = PinnedTraitData->EdGraphNodes[InDataIndex];
+				if (!EdGraphNode.IsValid())
+				{
+					return true;
 				}
-			};
 
-		PropertyHandle->SetOnChildPropertyValueChangedWithData(TDelegate<void(const FPropertyChangedEvent&)>::CreateLambda(UpdatePinDefaultValue));
-	}
+				const FProperty* Property = PinnedPropertyHandle->GetProperty();
+				const FString ValueStr = FRigVMStruct::ExportToFullyQualifiedText(Property, (uint8*)InRawData, true);
+
+				// Transform property path into pin path
+				FString PropertyPath(PinnedPropertyHandle->GetPropertyPath());
+
+				// Replace delimiter
+				FString PinPath = PropertyPath.Replace(TEXT("->"), TEXT("."));
+				// Replace array element delimiter
+				PinPath = PinPath.Replace(TEXT("["), TEXT("."));
+				PinPath = PinPath.Replace(TEXT("]."), TEXT("."));
+				PinPath = PinPath.Replace(TEXT("]"), TEXT(""));
+				// Replace struct root with trait root
+				const FString RootPropertyPath(PinnedTraitData->RootPropertyHandle->GetPropertyPath());
+				PinPath = PinPath.Replace(*RootPropertyPath, *PinnedTraitData->Name.ToString());
+				// Prefix node path
+				const FString NodePath = EdGraphNode->GetModelNodePath();
+				PinPath = NodePath + TEXT(".") + PinPath;
+				Controller->SetPinDefaultValue(PinPath, ValueStr);
+
+				return true;
+			});
+		}));
+	});
 }
 
 void FAnimNextEdGraphNodeCustomization::PopulateCategory(IDetailLayoutBuilder& DetailBuilder, const TSharedPtr<FRigVMNodeDetailsData>& RigVMTypeData)
diff --git a/Engine/Plugins/Experimental/AnimNext/Source/AnimNextEditor/Private/Graph/AnimNextEdGraphNodeCustomization.h b/Engine/Plugins/Experimental/AnimNext/Source/AnimNextEditor/Private/Graph/AnimNextEdGraphNodeCustomization.h
--- a/Engine/Plugins/Experimental/AnimNext/Source/AnimNextEditor/Private/Graph/AnimNextEdGraphNodeCustomization.h
+++ b/Engine/Plugins/Experimental/AnimNext/Source/AnimNextEditor/Private/Graph/AnimNextEdGraphNodeCustomization.h
@@ -74,6 +74,8 @@ protected:
 		}
 
 		TArray<TSharedPtr<FStructOnScope>> ScopedSharedDataInstances;
+		TSharedPtr<IPropertyHandle> RootPropertyHandle;
+		TArray<TSharedRef<IPropertyHandle>> PropertyHandles;
 	};
 
 	struct FRigVMNodeDetailsData : FCategoryDetailsData
