Fix handling of AnimNext latent properties
The latent property header is looked up using the base trait template

diff --git a/Engine/Plugins/Experimental/AnimNext/Source/AnimNext/Internal/TraitCore/TraitTemplate.h b/Engine/Plugins/Experimental/AnimNext/Source/AnimNext/Internal/TraitCore/TraitTemplate.h
--- a/Engine/Plugins/Experimental/AnimNext/Source/AnimNext/Internal/TraitCore/TraitTemplate.h
+++ b/Engine/Plugins/Experimental/AnimNext/Source/AnimNext/Internal/TraitCore/TraitTemplate.h
@@ -67,14 +67,17 @@ namespace UE::AnimNext
 		// Returns the offset into the shared data where the descriptor begins, relative to the root of the node's shared data
 		uint32 GetNodeSharedOffset() const noexcept { return NodeSharedOffset; }
 
-		// Returns the offset into the shared data where the latent property handles begin (or 0 if we have no latent properties), relative to the root of the node's shared data
+		// Returns the offset into the shared data where the latent property handles begin, relative to the root of the node's shared data
+		// The base trait has an instance of FLatentPropertiesHeader preceding its latent handles of type FLatentPropertyHandle
+		// Additive traits just have a list of FLatentPropertyHandle
+		// Note that the offset might not be zero even if no latent properties live on this trait
 		uint32 GetNodeSharedLatentPropertyHandlesOffset() const noexcept { return NodeSharedLatentPropertyHandlesOffset; }
 
 		// Returns the offset into the instance data where the descriptor begins, relative to the root of the node's instance data
 		uint32 GetNodeInstanceOffset() const noexcept { return NodeInstanceOffset; }
 
 		// Returns whether or not we have latent properties on this trait
-		bool HasLatentProperties() const noexcept { return NodeSharedLatentPropertyHandlesOffset != 0; }
+		bool HasLatentProperties() const noexcept { return NumLatentProperties != 0; }
 
 		// Returns a pointer to the specified trait description on the current node
 		FAnimNextTraitSharedData* GetTraitDescription(FNodeDescription& NodeDescription) const noexcept
@@ -91,25 +94,27 @@ namespace UE::AnimNext
 		// Returns a reference to the specified trait latent properties header on the current node
 		FLatentPropertiesHeader& GetTraitLatentPropertiesHeader(FNodeDescription& NodeDescription) const noexcept
 		{
-			return *reinterpret_cast<FLatentPropertiesHeader*>(reinterpret_cast<uint8*>(&NodeDescription) + GetNodeSharedLatentPropertyHandlesOffset());
+			check(GetMode() == ETraitMode::Base);	// Only the base trait has a header
+			return *reinterpret_cast<FLatentPropertiesHeader*>(reinterpret_cast<uint8*>(&NodeDescription) + GetNodeSharedLatentPropertyHandlesOffset() - sizeof(FLatentPropertiesHeader));
 		}
 
 		// Returns a reference to the specified trait latent properties header on the current node
 		const FLatentPropertiesHeader& GetTraitLatentPropertiesHeader(const FNodeDescription& NodeDescription) const noexcept
 		{
-			return *reinterpret_cast<const FLatentPropertiesHeader*>(reinterpret_cast<const uint8*>(&NodeDescription) + GetNodeSharedLatentPropertyHandlesOffset());
+			check(GetMode() == ETraitMode::Base);	// Only the base trait has a header
+			return *reinterpret_cast<const FLatentPropertiesHeader*>(reinterpret_cast<const uint8*>(&NodeDescription) + GetNodeSharedLatentPropertyHandlesOffset() - sizeof(FLatentPropertiesHeader));
 		}
 
 		// Returns a pointer to the specified trait latent property handles on the current node
 		FLatentPropertyHandle* GetTraitLatentPropertyHandles(FNodeDescription& NodeDescription) const noexcept
 		{
-			return reinterpret_cast<FLatentPropertyHandle*>(reinterpret_cast<uint8*>(&NodeDescription) + GetNodeSharedLatentPropertyHandlesOffset() + sizeof(FLatentPropertiesHeader));
+			return reinterpret_cast<FLatentPropertyHandle*>(reinterpret_cast<uint8*>(&NodeDescription) + GetNodeSharedLatentPropertyHandlesOffset());
 		}
 
 		// Returns a pointer to the specified trait latent property handles on the current node
 		const FLatentPropertyHandle* GetTraitLatentPropertyHandles(const FNodeDescription& NodeDescription) const noexcept
 		{
-			return reinterpret_cast<const FLatentPropertyHandle*>(reinterpret_cast<const uint8*>(&NodeDescription) + GetNodeSharedLatentPropertyHandlesOffset() + sizeof(FLatentPropertiesHeader));
+			return reinterpret_cast<const FLatentPropertyHandle*>(reinterpret_cast<const uint8*>(&NodeDescription) + GetNodeSharedLatentPropertyHandlesOffset());
 		}
 
 		// Returns a pointer to the specified trait instance on the current node
diff --git a/Engine/Plugins/Experimental/AnimNext/Source/AnimNext/Private/TraitCore/NodeTemplate.cpp b/Engine/Plugins/Experimental/AnimNext/Source/AnimNext/Private/TraitCore/NodeTemplate.cpp
--- a/Engine/Plugins/Experimental/AnimNext/Source/AnimNext/Private/TraitCore/NodeTemplate.cpp
+++ b/Engine/Plugins/Experimental/AnimNext/Source/AnimNext/Private/TraitCore/NodeTemplate.cpp
@@ -121,11 +121,10 @@ namespace UE::AnimNext
 
 				// Save our latent pins offset (if we have any)
 				NumLatentProperties = Trait->GetNumLatentTraitProperties();
-				if (NumLatentProperties != 0)
-				{
-					TraitSharedLatentPropertyHandlesOffset = SharedLatentPropertyHandlesOffset;
-					SharedLatentPropertyHandlesOffset += NumLatentProperties * sizeof(FLatentPropertyHandle);
-				}
+
+				// The handle offset points to the first handle, if we are a base trait, our header precedes it
+				TraitSharedLatentPropertyHandlesOffset = SharedLatentPropertyHandlesOffset;
+				SharedLatentPropertyHandlesOffset += NumLatentProperties * sizeof(FLatentPropertyHandle);
 			}
 
 			check(NumSubStackLatentProperties <= MAX_uint16);
diff --git a/Engine/Plugins/Experimental/AnimNext/Source/AnimNext/Private/TraitCore/TraitReader.cpp b/Engine/Plugins/Experimental/AnimNext/Source/AnimNext/Private/TraitCore/TraitReader.cpp
--- a/Engine/Plugins/Experimental/AnimNext/Source/AnimNext/Private/TraitCore/TraitReader.cpp
+++ b/Engine/Plugins/Experimental/AnimNext/Source/AnimNext/Private/TraitCore/TraitReader.cpp
@@ -138,11 +138,22 @@ namespace UE::AnimNext
 					continue;	// Nothing to do
 				}
 
-				FLatentPropertiesHeader& LatentHeader = TraitTemplate.GetTraitLatentPropertiesHeader(*NodeDesc);
+				const uint32 BaseTraitIndex = TraitIndex - TraitTemplate.GetTraitIndex();
+				const FTraitTemplate& BaseTraitTemplate = TraitTemplates[BaseTraitIndex];
+
+				// Latent header is looked up using the base trait
+				FLatentPropertiesHeader& LatentHeader = BaseTraitTemplate.GetTraitLatentPropertiesHeader(*NodeDesc);
 				FLatentPropertyHandle* LatentHandles = TraitTemplate.GetTraitLatentPropertyHandles(*NodeDesc);
 
-				bool bHasValidLatentProperties = false;
-				bool bCanAllPropertiesFreeze = true;
+				// If we are a base trait, reset our header
+				if (TraitTemplate.GetMode() == ETraitMode::Base)
+				{
+					LatentHeader.bHasValidLatentProperties = false;
+					LatentHeader.bCanAllPropertiesFreeze = true;
+				}
+
+				bool bHasValidLatentProperties = !!LatentHeader.bHasValidLatentProperties;
+				bool bCanAllPropertiesFreeze = !!LatentHeader.bCanAllPropertiesFreeze;
 
 				const FTraitRegistryHandle TraitHandle = TraitTemplate.GetRegistryHandle();
 				const FTrait* Trait = TraitRegistry.Find(TraitHandle);
