Fixed missing calls to Super::AddReferencedObjects

#rnx

diff --git a/Engine/Plugins/Experimental/AnimNext/Source/AnimNext/Private/Component/AnimNextWorldSubsystem.cpp b/Engine/Plugins/Experimental/AnimNext/Source/AnimNext/Private/Component/AnimNextWorldSubsystem.cpp
--- a/Engine/Plugins/Experimental/AnimNext/Source/AnimNext/Private/Component/AnimNextWorldSubsystem.cpp
+++ b/Engine/Plugins/Experimental/AnimNext/Source/AnimNext/Private/Component/AnimNextWorldSubsystem.cpp
@@ -1,4 +1,4 @@
-﻿// Copyright Epic Games, Inc. All Rights Reserved.
+// Copyright Epic Games, Inc. All Rights Reserved.
 
 #include "Component/AnimNextWorldSubsystem.h"
 
@@ -54,6 +54,8 @@ UAnimNextWorldSubsystem::~UAnimNextWorldSubsystem()
 
 void UAnimNextWorldSubsystem::AddReferencedObjects(UObject* InThis, FReferenceCollector& Collector)
 {
+	Super::AddReferencedObjects(InThis, Collector);
+
 	UAnimNextWorldSubsystem* This = CastChecked<UAnimNextWorldSubsystem>(InThis);
 	for(FAnimNextModuleInstance& Instance : This->Instances)
 	{
diff --git a/Engine/Source/Runtime/Landscape/Private/LandscapeInfoMap.cpp b/Engine/Source/Runtime/Landscape/Private/LandscapeInfoMap.cpp
--- a/Engine/Source/Runtime/Landscape/Private/LandscapeInfoMap.cpp
+++ b/Engine/Source/Runtime/Landscape/Private/LandscapeInfoMap.cpp
@@ -41,6 +41,8 @@ void ULandscapeInfoMap::BeginDestroy()
 
 void ULandscapeInfoMap::AddReferencedObjects(UObject* InThis, FReferenceCollector& Collector)
 {
+	Super::AddReferencedObjects(InThis, Collector);
+
 	ULandscapeInfoMap* This = CastChecked<ULandscapeInfoMap>(InThis);
 	Collector.AddReferencedObjects(This->Map, This);
 }
