Ensure latent property header is reserved and initialized even if no latent properties are present

diff --git a/Engine/Plugins/Experimental/AnimNext/Source/AnimNext/Private/TraitCore/NodeTemplate.cpp b/Engine/Plugins/Experimental/AnimNext/Source/AnimNext/Private/TraitCore/NodeTemplate.cpp
--- a/Engine/Plugins/Experimental/AnimNext/Source/AnimNext/Private/TraitCore/NodeTemplate.cpp
+++ b/Engine/Plugins/Experimental/AnimNext/Source/AnimNext/Private/TraitCore/NodeTemplate.cpp
@@ -104,7 +104,7 @@ namespace UE::AnimNext
 
 				// Base traits include the list of all latent property handles in its shared data
 				// Latent property offsets will point into that list
-				if (TraitTemplate.GetMode() == ETraitMode::Base && NumSubStackLatentProperties != 0)
+				if (TraitTemplate.GetMode() == ETraitMode::Base)
 				{
 					// Align our handles
 					SharedDataOffset = Align(SharedDataOffset, alignof(FLatentPropertiesHeader));
