AnimNext: Data interfaces

Adds a new asset type: UAnimNextDataInterface. Reparents UAnimNextModule and UAnimNextAnimationGraph to be children of this asset.
UAnimNextDataInterface can only contain public variables. It can also 'implement' other data interfaces, which causes theuir variables to be insantiated within it. Collisions in variable naming/types are handled at edit time and by the compiler.
When a sub-graph is instantiated (either in a module or in a subgraph trait), then variables can be auto-bound to the host, which causes the inner variable's memory to point to that of the host, removing the need to 'pipe' parameters or copy them. Piping can still occur via trait-based pin piping as was introduced in the module refactor (however I've still not added a trait to do this in animation graphs yet, only modules).
Variable compilation keeps public variables together in one block as before, but introduces sorting via data interface as well as size, to ensure that data interface variables are addressable by index & count.
Added CustomizeNewAssetEntry override point used to setup new items. Data interfaces use this to force new variables to public.
Fixed various issues with modification event routing

Updated AnimNextChooser plugin
- Split chooser rig unit into 3: an abstract base and an implementation that uses specific RigVMExecuteContexts so that additional data is available to Chooser columns
- Uses the new variable access API

Updated PlayAnim trait
- Implemented binding of variables to event payload members via a struct adapter class

Updated graph schemas to correctly allow/disallow anim next execute context nodes

Fixed const-issues when usting GetAsset and GetEditorData

#rnx

diff --git a/Engine/Plugins/Experimental/AnimNext/Source/AnimNext/Internal/AnimNextExecuteContext.h b/Engine/Plugins/Experimental/AnimNext/Source/AnimNext/Internal/AnimNextExecuteContext.h
--- a/Engine/Plugins/Experimental/AnimNext/Source/AnimNext/Internal/AnimNextExecuteContext.h
+++ b/Engine/Plugins/Experimental/AnimNext/Source/AnimNext/Internal/AnimNextExecuteContext.h
@@ -6,9 +6,12 @@
 #include "RigVMCore/RigVMExecuteContext.h"
 #include "Module/AnimNextModuleContextData.h"
 #include "Graph/AnimNextGraphContextData.h"
+#include "Graph/AnimNextGraphInstance.h"
 #include "Misc/TVariant.h"
+#include "Module/AnimNextModuleInstance.h"
 #include "AnimNextExecuteContext.generated.h"
 
+struct FAnimNextDataInterfaceInstance;
 struct FAnimNextGraphInstance;
 
 namespace UE::AnimNext
@@ -39,6 +42,16 @@ struct FAnimNextExecuteContext : public FRigVMExecuteContext
 		return ContextData.Get<ContextType>();
 	}
 
+	// Get the current data interface instance (module or graph) that is executing 
+	const FAnimNextDataInterfaceInstance& GetInstance() const
+	{
+		if(ContextData.IsType<FAnimNextModuleContextData>())
+		{
+			return ContextData.Get<FAnimNextModuleContextData>().GetModuleInstance();
+		}
+		return ContextData.Get<FAnimNextGraphContextData>().GetGraphInstance();
+	}
+
 protected:
 	// Setup the context data to the specified type
 	template<typename ContextType, typename... ArgsType>
@@ -56,7 +69,7 @@ protected:
 
 	// All possible known variants of our context data. IF we ever want this to be extensible, this can be converted into an FInstancedStruct
 	TVariant<FAnimNextModuleContextData, FAnimNextGraphContextData> ContextData;
-	
+
 	friend struct UE::AnimNext::FModuleEventTickFunction;
 	friend struct FAnimNextGraphInstance;
 	friend class UAnimNextAnimationGraph;
diff --git a/Engine/Plugins/Experimental/AnimNext/Source/AnimNext/Internal/AnimNextRigVMAsset.h b/Engine/Plugins/Experimental/AnimNext/Source/AnimNext/Internal/AnimNextRigVMAsset.h
--- a/Engine/Plugins/Experimental/AnimNext/Source/AnimNext/Internal/AnimNextRigVMAsset.h
+++ b/Engine/Plugins/Experimental/AnimNext/Source/AnimNext/Internal/AnimNextRigVMAsset.h
@@ -14,6 +14,7 @@ struct FAnimNextModuleInstance;
 namespace UE::AnimNext
 {
 	struct FProxyVariablesContext;
+	struct FModuleEventTickFunction;
 }
 
 namespace UE::AnimNext::UncookedOnly
@@ -31,6 +32,7 @@ protected:
 	friend class UAnimNextVariableEntry;
 	friend class UAnimNextRigVMAssetEditorData;
 	friend struct FAnimNextModuleInstance;
+	friend struct UE::AnimNext::FModuleEventTickFunction;
 	friend struct UE::AnimNext::FProxyVariablesContext;
 	friend struct UE::AnimNext::UncookedOnly::FUtils;
 
diff --git a/Engine/Plugins/Experimental/AnimNext/Source/AnimNext/Internal/DataInterface/AnimNextDataInterface.h b/Engine/Plugins/Experimental/AnimNext/Source/AnimNext/Internal/DataInterface/AnimNextDataInterface.h
new file mode 100644
--- /dev/null
+++ b/Engine/Plugins/Experimental/AnimNext/Source/AnimNext/Internal/DataInterface/AnimNextDataInterface.h
@@ -0,0 +1,66 @@
+﻿// Copyright Epic Games, Inc. All Rights Reserved.
+
+#pragma once
+
+#include "CoreMinimal.h"
+#include "AnimNextRigVMAsset.h"
+#include "AnimNextDataInterface.generated.h"
+
+class UAnimNextDataInterfaceFactory;
+
+namespace UE::AnimNext::Tests
+{
+	class FDataInterfaceCompile;
+}
+
+namespace UE::AnimNext::UncookedOnly
+{
+	struct FUtils;
+}
+
+// Information about an implemented interface
+USTRUCT()
+struct FAnimNextImplementedDataInterface
+{
+	GENERATED_BODY()
+
+	UPROPERTY()
+	TObjectPtr<const UAnimNextDataInterface> DataInterface;
+
+	// Index of the first variable that implements the interface
+	UPROPERTY()
+	int32 VariableIndex = INDEX_NONE;
+
+	// Number of variables that implement the interface
+	UPROPERTY()
+	int32 NumVariables = 0;
+
+	// Whether to automatically bind this interface to any host data interface 
+	UPROPERTY()
+	bool bAutoBindToHost = false;
+};
+
+// Data interfaces provide a set of named data that is shared between AnimNext assets and used for communication between assets and functional units
+UCLASS(BlueprintType)
+class ANIMNEXT_API UAnimNextDataInterface : public UAnimNextRigVMAsset
+{
+	GENERATED_BODY()
+
+public:
+	UAnimNextDataInterface(const FObjectInitializer& ObjectInitializer);
+
+	// Get all the implemented interfaces
+	TConstArrayView<FAnimNextImplementedDataInterface> GetImplementedInterfaces() const;
+
+	// Find an implemented interface
+	const FAnimNextImplementedDataInterface* FindImplementedInterface(const UAnimNextDataInterface* InDataInterface) const;
+
+private:
+	friend class UAnimNextDataInterfaceFactory;
+	friend class UE::AnimNext::Tests::FDataInterfaceCompile;
+	friend struct UE::AnimNext::UncookedOnly::FUtils;
+
+	// Information about implemented interfaces. Note this includes the 'self' interface (first), if any public variables are specified.
+	UPROPERTY()
+	TArray<FAnimNextImplementedDataInterface> ImplementedInterfaces;
+};
diff --git a/Engine/Plugins/Experimental/AnimNext/Source/AnimNext/Internal/Graph/AnimNextAnimationGraph.h b/Engine/Plugins/Experimental/AnimNext/Source/AnimNext/Internal/Graph/AnimNextAnimationGraph.h
--- a/Engine/Plugins/Experimental/AnimNext/Source/AnimNext/Internal/Graph/AnimNextAnimationGraph.h
+++ b/Engine/Plugins/Experimental/AnimNext/Source/AnimNext/Internal/Graph/AnimNextAnimationGraph.h
@@ -4,7 +4,7 @@
 
 #include "CoreMinimal.h"
 #include "Graph/AnimNextGraphEntryPoint.h"
-#include "AnimNextRigVMAsset.h"
+#include "DataInterface/AnimNextDataInterface.h"
 #include "Graph/RigUnit_AnimNextGraphRoot.h"
 #include "RigVMCore/RigVM.h"
 #include "TraitCore/TraitPtr.h"
@@ -56,7 +56,7 @@ namespace UE::AnimNext::Graph
 
 // A user-created collection of animation logic & data
 UCLASS(BlueprintType)
-class ANIMNEXT_API UAnimNextAnimationGraph : public UAnimNextRigVMAsset
+class ANIMNEXT_API UAnimNextAnimationGraph : public UAnimNextDataInterface
 {
 	GENERATED_BODY()
 
diff --git a/Engine/Plugins/Experimental/AnimNext/Source/AnimNext/Internal/Graph/RigVMTrait_AnimNextPublicVariables.h b/Engine/Plugins/Experimental/AnimNext/Source/AnimNext/Internal/Graph/RigVMTrait_AnimNextPublicVariables.h
--- a/Engine/Plugins/Experimental/AnimNext/Source/AnimNext/Internal/Graph/RigVMTrait_AnimNextPublicVariables.h
+++ b/Engine/Plugins/Experimental/AnimNext/Source/AnimNext/Internal/Graph/RigVMTrait_AnimNextPublicVariables.h
@@ -2,10 +2,12 @@
 
 #pragma once
 
+#include "DataInterface/AnimNextDataInterfaceHost.h"
 #include "RigVMCore/RigVMTrait.h"
 #include "RigVMTrait_AnimNextPublicVariables.generated.h"
 
 class UAnimNextRigVMAsset;
+class FRigVMTraitScope;
 
 // Represents public variables of an asset via a trait 
 USTRUCT(BlueprintType)
@@ -13,9 +15,9 @@ struct ANIMNEXT_API FRigVMTrait_AnimNextPublicVariables : public FRigVMTrait
 {
 	GENERATED_BODY()
 
-	// The asset that any programmatic pins will be derived from
+	// The data interface that any programmatic pins will be derived from
 	UPROPERTY(meta = (Hidden))
-	TObjectPtr<UAnimNextRigVMAsset> Asset = nullptr;
+	TObjectPtr<UAnimNextDataInterface> Asset = nullptr;
 
 	// Variable names that are exposed
 	UPROPERTY(meta = (Hidden))
@@ -28,3 +30,23 @@ struct ANIMNEXT_API FRigVMTrait_AnimNextPublicVariables : public FRigVMTrait
 	virtual bool ShouldCreatePinForProperty(const FProperty* InProperty) const override;
 #endif
 };
+
+namespace UE::AnimNext
+{
+
+struct FPublicVariablesTraitToDataInterfaceHostAdapter : public IDataInterfaceHost
+{
+	FPublicVariablesTraitToDataInterfaceHostAdapter(const FRigVMTrait_AnimNextPublicVariables& InTrait, const FRigVMTraitScope& InTraitScope)
+		: Trait(InTrait)
+		, TraitScope(InTraitScope)
+	{}
+
+	// IDataInterfaceHost interface
+	virtual const UAnimNextDataInterface* GetDataInterface() const override;
+	virtual uint8* GetMemoryForVariable(int32 InVariableIndex, FName InVariableName, const FProperty* InVariableProperty) const override;
+
+	const FRigVMTrait_AnimNextPublicVariables& Trait;
+	const FRigVMTraitScope& TraitScope;
+};
+
+}
\ No newline at end of file
diff --git a/Engine/Plugins/Experimental/AnimNext/Source/AnimNext/Internal/Module/AnimNextModule.h b/Engine/Plugins/Experimental/AnimNext/Source/AnimNext/Internal/Module/AnimNextModule.h
--- a/Engine/Plugins/Experimental/AnimNext/Source/AnimNext/Internal/Module/AnimNextModule.h
+++ b/Engine/Plugins/Experimental/AnimNext/Source/AnimNext/Internal/Module/AnimNextModule.h
@@ -4,7 +4,7 @@
 
 #include "CoreMinimal.h"
 #include "Graph/AnimNextGraphEntryPoint.h"
-#include "AnimNextRigVMAsset.h"
+#include "DataInterface/AnimNextDataInterface.h"
 #include "RigVMCore/RigVM.h"
 #include "TraitCore/TraitPtr.h"
 #include "RigVMHost.h"
@@ -60,7 +60,7 @@ enum class EAnimNextModuleInitMethod : uint8
 
 // Root asset represented by a component when instantiated
 UCLASS(BlueprintType)
-class ANIMNEXT_API UAnimNextModule : public UAnimNextRigVMAsset
+class ANIMNEXT_API UAnimNextModule : public UAnimNextDataInterface
 {
 	GENERATED_BODY()
 
diff --git a/Engine/Plugins/Experimental/AnimNext/Source/AnimNext/Internal/Module/AnimNextModuleInstance.h b/Engine/Plugins/Experimental/AnimNext/Source/AnimNext/Internal/Module/AnimNextModuleInstance.h
--- a/Engine/Plugins/Experimental/AnimNext/Source/AnimNext/Internal/Module/AnimNextModuleInstance.h
+++ b/Engine/Plugins/Experimental/AnimNext/Source/AnimNext/Internal/Module/AnimNextModuleInstance.h
@@ -6,6 +6,7 @@
 #include "AnimNextModule.h"
 #include "ModuleHandle.h"
 #include "Component/AnimNextPublicVariablesProxy.h"
+#include "DataInterface/AnimNextDataInterfaceInstance.h"
 #include "Module/ModuleTickFunction.h"
 #include "TraitCore/TraitEvent.h"
 #include "TraitCore/TraitEventList.h"
@@ -45,7 +46,7 @@ namespace UE::AnimNext::Private
 
 // Root memory owner of a parameterized schedule 
 USTRUCT()
-struct FAnimNextModuleInstance
+struct FAnimNextModuleInstance : public FAnimNextDataInterfaceInstance
 {
 	GENERATED_BODY()
 
@@ -83,8 +84,8 @@ struct FAnimNextModuleInstance
 	void OnModuleCompiled();
 #endif
 
-	UPROPERTY(Transient)
-	TObjectPtr<UAnimNextModule> Module = nullptr;
+	// Get the module that this instance represents
+	ANIMNEXT_API const UAnimNextModule* GetModule() const;
 
 	// Object this entry is bound to
 	UPROPERTY(Transient)
@@ -108,13 +109,8 @@ struct FAnimNextModuleInstance
 	// Lock to ensure event list actions are thread safe
 	FRWLock EventListLock;
 
-	// Execute context for the RigVM instance to execute with
-	FRigVMExtendedExecuteContext ExtendedExecuteContext;
-
-	// Variables state
-	FInstancedPropertyBag Variables;
-
 	// Proxy public variables
+	UPROPERTY(Transient)
 	FAnimNextPublicVariablesProxy PublicVariablesProxy;
 
 	enum class ERunState : uint8
diff --git a/Engine/Plugins/Experimental/AnimNext/Source/AnimNext/Private/Component/AnimNextComponentWorldSubsystem.cpp b/Engine/Plugins/Experimental/AnimNext/Source/AnimNext/Private/Component/AnimNextComponentWorldSubsystem.cpp
--- a/Engine/Plugins/Experimental/AnimNext/Source/AnimNext/Private/Component/AnimNextComponentWorldSubsystem.cpp
+++ b/Engine/Plugins/Experimental/AnimNext/Source/AnimNext/Private/Component/AnimNextComponentWorldSubsystem.cpp
@@ -49,7 +49,7 @@ void UAnimNextComponentWorldSubsystem::OnModuleCompiled(UAnimNextModule* InModul
 
 	for(FAnimNextModuleInstance& Instance : Instances)
 	{
-		if(Instance.Module == InModule)
+		if(Instance.GetModule() == InModule)
 		{
 			UAnimNextComponent* AnimNextComponent = CastChecked<UAnimNextComponent>(Instance.Object);
 			AnimNextComponent->OnModuleCompiled();
diff --git a/Engine/Plugins/Experimental/AnimNext/Source/AnimNext/Private/Component/AnimNextWorldSubsystem.cpp b/Engine/Plugins/Experimental/AnimNext/Source/AnimNext/Private/Component/AnimNextWorldSubsystem.cpp
--- a/Engine/Plugins/Experimental/AnimNext/Source/AnimNext/Private/Component/AnimNextWorldSubsystem.cpp
+++ b/Engine/Plugins/Experimental/AnimNext/Source/AnimNext/Private/Component/AnimNextWorldSubsystem.cpp
@@ -59,8 +59,7 @@ void UAnimNextWorldSubsystem::AddReferencedObjects(UObject* InThis, FReferenceCo
 	UAnimNextWorldSubsystem* This = CastChecked<UAnimNextWorldSubsystem>(InThis);
 	for(FAnimNextModuleInstance& Instance : This->Instances)
 	{
-		Collector.AddReferencedObject(Instance.Module);
-		Collector.AddReferencedObject(Instance.Object);
+		Collector.AddPropertyReferences(FAnimNextModuleInstance::StaticStruct(), &Instance, InThis);
 	}
 }
 
@@ -211,7 +210,7 @@ void UAnimNextWorldSubsystem::QueueTaskHandle(UE::AnimNext::FModuleHandle InOutH
 		}
 		else
 		{
-			UE_LOGFMT(LogAnimation, Warning, "QueueTask: Could not find event '{EventName}' in module '{ModuleName}'", InModuleEventName, Instance.Module->GetName());
+			UE_LOGFMT(LogAnimation, Warning, "QueueTask: Could not find event '{EventName}' in module '{ModuleName}'", InModuleEventName, Instance.GetDataInterfaceName());
 		}
 	}
 }
@@ -225,7 +224,7 @@ void UAnimNextWorldSubsystem::OnModuleCompiled(UAnimNextModule* InModule)
 
 	for(FAnimNextModuleInstance& Instance : Instances)
 	{
-		if(Instance.Module == InModule)
+		if(Instance.GetModule() == InModule)
 		{
 			Instance.OnModuleCompiled();
 		}
diff --git a/Engine/Plugins/Experimental/AnimNext/Source/AnimNext/Private/DataInterface/AnimNextDataInterface.cpp b/Engine/Plugins/Experimental/AnimNext/Source/AnimNext/Private/DataInterface/AnimNextDataInterface.cpp
new file mode 100644
--- /dev/null
+++ b/Engine/Plugins/Experimental/AnimNext/Source/AnimNext/Private/DataInterface/AnimNextDataInterface.cpp
@@ -0,0 +1,28 @@
+﻿// Copyright Epic Games, Inc. All Rights Reserved.
+
+#include "DataInterface/AnimNextDataInterface.h"
+#include "AnimNextExecuteContext.h"
+
+#if WITH_EDITOR	
+#include "Engine/ExternalAssetDependencyGatherer.h"
+REGISTER_ASSETDEPENDENCY_GATHERER(FExternalAssetDependencyGatherer, UAnimNextDataInterface);
+#endif // WITH_EDITOR
+
+UAnimNextDataInterface::UAnimNextDataInterface(const FObjectInitializer& ObjectInitializer)
+	: Super(ObjectInitializer)
+{
+	ExtendedExecuteContext.SetContextPublicDataStruct(FAnimNextExecuteContext::StaticStruct());
+}
+
+TConstArrayView<FAnimNextImplementedDataInterface> UAnimNextDataInterface::GetImplementedInterfaces() const
+{
+	return ImplementedInterfaces;
+}
+
+const FAnimNextImplementedDataInterface* UAnimNextDataInterface::FindImplementedInterface(const UAnimNextDataInterface* InDataInterface) const
+{
+	return ImplementedInterfaces.FindByPredicate([InDataInterface](const FAnimNextImplementedDataInterface& InImplementedDataInterface)
+	{
+		return InImplementedDataInterface.DataInterface == InDataInterface;
+	});
+}
\ No newline at end of file
diff --git a/Engine/Plugins/Experimental/AnimNext/Source/AnimNext/Private/DataInterface/AnimNextDataInterfaceInstance.cpp b/Engine/Plugins/Experimental/AnimNext/Source/AnimNext/Private/DataInterface/AnimNextDataInterfaceInstance.cpp
new file mode 100644
--- /dev/null
+++ b/Engine/Plugins/Experimental/AnimNext/Source/AnimNext/Private/DataInterface/AnimNextDataInterfaceInstance.cpp
@@ -0,0 +1,621 @@
+﻿// Copyright Epic Games, Inc. All Rights Reserved.
+
+#include "DataInterface/AnimNextDataInterfaceInstance.h"
+#include "Logging/StructuredLog.h"
+#include "Param/ParamUtils.h"
+
+namespace UE::AnimNext::Private
+{
+	// Utility functions for type conversion, adapted from UE::StructUtils::Private in PropertyBag.cpp
+
+	bool CanCastTo(const UStruct* From, const UStruct* To)
+	{
+		return From != nullptr && To != nullptr && From->IsChildOf(To);
+	}
+
+	EPropertyBagResult GetPropertyAsInt64(const FProperty* InProperty, EPropertyBagPropertyType InValueType, const void* Address, int64& OutValue)
+	{
+		check(Address != nullptr);
+		check(InProperty != nullptr);
+
+		switch(InValueType)
+		{
+		case EPropertyBagPropertyType::Bool:
+			{
+				const FBoolProperty* Property = CastFieldChecked<FBoolProperty>(InProperty);
+				OutValue = Property->GetPropertyValue(Address) ? 1 : 0;
+				return EPropertyBagResult::Success;
+			}
+		case EPropertyBagPropertyType::Byte:
+			{
+				const FByteProperty* Property = CastFieldChecked<FByteProperty>(InProperty);
+				OutValue = Property->GetPropertyValue(Address);
+				return EPropertyBagResult::Success;
+			}
+		case EPropertyBagPropertyType::Int32:
+			{
+				const FIntProperty* Property = CastFieldChecked<FIntProperty>(InProperty);
+				OutValue = Property->GetPropertyValue(Address);
+				return EPropertyBagResult::Success;
+			}
+		case EPropertyBagPropertyType::UInt32:
+			{
+				const FUInt32Property* Property = CastFieldChecked<FUInt32Property>(InProperty);
+				OutValue = static_cast<uint32>(Property->GetPropertyValue(Address));
+				return EPropertyBagResult::Success;
+			}
+		case EPropertyBagPropertyType::Int64:
+			{
+				const FInt64Property* Property = CastFieldChecked<FInt64Property>(InProperty);
+				OutValue = Property->GetPropertyValue(Address);
+				return EPropertyBagResult::Success;
+			}
+		case EPropertyBagPropertyType::UInt64:
+			{
+				const FUInt64Property* Property = CastFieldChecked<FUInt64Property>(InProperty);
+				OutValue = static_cast<int64>(Property->GetPropertyValue(Address));
+				return EPropertyBagResult::Success;
+			}
+		case EPropertyBagPropertyType::Float:
+			{
+				const FFloatProperty* Property = CastFieldChecked<FFloatProperty>(InProperty);
+				OutValue = static_cast<int64>(Property->GetPropertyValue(Address));
+				return EPropertyBagResult::Success;
+			}
+		case EPropertyBagPropertyType::Double:
+			{
+				const FDoubleProperty* Property = CastFieldChecked<FDoubleProperty>(InProperty);
+				OutValue = static_cast<int64>(Property->GetPropertyValue(Address));
+				return EPropertyBagResult::Success;
+			}
+		case EPropertyBagPropertyType::Enum:
+			{
+				const FEnumProperty* EnumProperty = CastFieldChecked<FEnumProperty>(InProperty);
+				const FNumericProperty* UnderlyingProperty = EnumProperty->GetUnderlyingProperty();
+				check(UnderlyingProperty);
+				OutValue = UnderlyingProperty->GetSignedIntPropertyValue(Address);
+				return EPropertyBagResult::Success;
+			}
+		default:
+			return EPropertyBagResult::TypeMismatch;
+		}
+	}
+
+	EPropertyBagResult GetPropertyAsUInt64(const FProperty* InProperty, EPropertyBagPropertyType InValueType, const void* Address, uint64& OutValue)
+	{
+		check(Address != nullptr);
+		check(InProperty != nullptr);
+
+		switch (InValueType)
+		{
+		case EPropertyBagPropertyType::Bool:
+		{
+			const FBoolProperty* Property = CastFieldChecked<FBoolProperty>(InProperty);
+			OutValue = Property->GetPropertyValue(Address) ? 1 : 0;
+			return EPropertyBagResult::Success;
+		}
+		case EPropertyBagPropertyType::Byte:
+		{
+			const FByteProperty* Property = CastFieldChecked<FByteProperty>(InProperty);
+			OutValue = Property->GetPropertyValue(Address);
+			return EPropertyBagResult::Success;
+		}
+		case EPropertyBagPropertyType::Int32:
+		{
+			const FIntProperty* Property = CastFieldChecked<FIntProperty>(InProperty);
+			OutValue = static_cast<uint32>(Property->GetPropertyValue(Address));
+			return EPropertyBagResult::Success;
+		}
+		case EPropertyBagPropertyType::UInt32:
+		{
+			const FUInt32Property* Property = CastFieldChecked<FUInt32Property>(InProperty);
+			OutValue = Property->GetPropertyValue(Address);
+			return EPropertyBagResult::Success;
+		}
+		case EPropertyBagPropertyType::Int64:
+		{
+			const FInt64Property* Property = CastFieldChecked<FInt64Property>(InProperty);
+			OutValue = static_cast<uint64>(Property->GetPropertyValue(Address));
+			return EPropertyBagResult::Success;
+		}
+		case EPropertyBagPropertyType::UInt64:
+		{
+			const FUInt64Property* Property = CastFieldChecked<FUInt64Property>(InProperty);
+			OutValue = Property->GetPropertyValue(Address);
+			return EPropertyBagResult::Success;
+		}
+		case EPropertyBagPropertyType::Float:
+		{
+			const FFloatProperty* Property = CastFieldChecked<FFloatProperty>(InProperty);
+			OutValue = static_cast<uint64>(Property->GetPropertyValue(Address));
+			return EPropertyBagResult::Success;
+		}
+		case EPropertyBagPropertyType::Double:
+		{
+			const FDoubleProperty* Property = CastFieldChecked<FDoubleProperty>(InProperty);
+			OutValue = static_cast<uint64>(Property->GetPropertyValue(Address));
+			return EPropertyBagResult::Success;
+		}
+		case EPropertyBagPropertyType::Enum:
+		{
+			const FEnumProperty* EnumProperty = CastFieldChecked<FEnumProperty>(InProperty);
+			const FNumericProperty* UnderlyingProperty = EnumProperty->GetUnderlyingProperty();
+			check(UnderlyingProperty);
+			OutValue = UnderlyingProperty->GetUnsignedIntPropertyValue(Address);
+			return EPropertyBagResult::Success;
+		}
+		default:
+			return EPropertyBagResult::TypeMismatch;
+		}
+	}
+
+	EPropertyBagResult GetPropertyAsDouble(const FProperty* InProperty, EPropertyBagPropertyType InValueType, const void* Address, double& OutValue)
+	{
+		check(Address != nullptr);
+		check(InProperty != nullptr);
+
+		switch(InValueType)
+		{
+		case EPropertyBagPropertyType::Bool:
+			{
+				const FBoolProperty* Property = CastFieldChecked<FBoolProperty>(InProperty);
+				OutValue = Property->GetPropertyValue(Address) ? 1.0 : 0.0;
+				return EPropertyBagResult::Success;
+			}
+		case EPropertyBagPropertyType::Byte:
+			{
+				const FByteProperty* Property = CastFieldChecked<FByteProperty>(InProperty);
+				OutValue = Property->GetPropertyValue(Address);
+				return EPropertyBagResult::Success;
+			}
+		case EPropertyBagPropertyType::Int32:
+			{
+				const FIntProperty* Property = CastFieldChecked<FIntProperty>(InProperty);
+				OutValue = Property->GetPropertyValue(Address);
+				return EPropertyBagResult::Success;
+			}
+		case EPropertyBagPropertyType::UInt32:
+			{
+				const FUInt32Property* Property = CastFieldChecked<FUInt32Property>(InProperty);
+				OutValue = Property->GetPropertyValue(Address);
+				return EPropertyBagResult::Success;
+			}
+		case EPropertyBagPropertyType::Int64:
+			{
+				const FInt64Property* Property = CastFieldChecked<FInt64Property>(InProperty);
+				OutValue = static_cast<double>(Property->GetPropertyValue(Address));
+				return EPropertyBagResult::Success;
+			}
+		case EPropertyBagPropertyType::UInt64:
+			{
+				const FUInt64Property* Property = CastFieldChecked<FUInt64Property>(InProperty);
+				OutValue = static_cast<double>(Property->GetPropertyValue(Address));
+				return EPropertyBagResult::Success;
+			}
+		case EPropertyBagPropertyType::Float:
+			{
+				const FFloatProperty* Property = CastFieldChecked<FFloatProperty>(InProperty);
+				OutValue = Property->GetPropertyValue(Address);
+				return EPropertyBagResult::Success;
+			}
+		case EPropertyBagPropertyType::Double:
+			{
+				const FDoubleProperty* Property = CastFieldChecked<FDoubleProperty>(InProperty);
+				OutValue = Property->GetPropertyValue(Address);
+				return EPropertyBagResult::Success;
+			}
+		case EPropertyBagPropertyType::Enum:
+			{
+				const FEnumProperty* EnumProperty = CastFieldChecked<FEnumProperty>(InProperty);
+				const FNumericProperty* UnderlyingProperty = EnumProperty->GetUnderlyingProperty();
+				check(UnderlyingProperty);
+				OutValue = static_cast<double>(UnderlyingProperty->GetSignedIntPropertyValue(Address));
+				return EPropertyBagResult::Success;
+			}
+		default:
+			return EPropertyBagResult::TypeMismatch;
+		}
+	}
+
+	// Generic property getter. Used for FName, FString, FText. 
+	template<typename T, typename PropT>
+	EPropertyBagResult GetPropertyValue(const FProperty* InProperty, EPropertyBagPropertyType InValueType, const void* Address, T& OutValue)
+	{
+		check(Address != nullptr);
+		check(InProperty != nullptr);
+
+		if (!InProperty->IsA<PropT>())
+		{
+			return EPropertyBagResult::TypeMismatch;
+		}
+		
+		const PropT* Property = CastFieldChecked<PropT>(InProperty);
+		OutValue = Property->GetPropertyValue(Address);
+
+		return EPropertyBagResult::Success;
+	}
+
+	EPropertyBagResult GetPropertyValueAsEnum(const FProperty* InProperty, EPropertyBagPropertyType InValueType, const void* Address, const UEnum* RequestedEnum, uint8& OutValue)
+	{
+		check(Address != nullptr);
+		check(InProperty != nullptr);
+
+		if (InValueType != EPropertyBagPropertyType::Enum)
+		{
+			return EPropertyBagResult::TypeMismatch;
+		}
+
+		const FEnumProperty* EnumProperty = CastFieldChecked<FEnumProperty>(InProperty);
+		const FNumericProperty* UnderlyingProperty = EnumProperty->GetUnderlyingProperty();
+		check(UnderlyingProperty);
+	
+		if (RequestedEnum != EnumProperty->GetEnum())
+		{
+			return EPropertyBagResult::TypeMismatch;
+		}
+	
+		OutValue = static_cast<uint8>(UnderlyingProperty->GetUnsignedIntPropertyValue(Address));
+
+		return EPropertyBagResult::Success;
+	}
+
+	EPropertyBagResult GetPropertyValueAsStruct(const FProperty* InProperty, EPropertyBagPropertyType InValueType, const void* Address, const UScriptStruct* RequestedStruct, uint8* OutValue)
+	{
+		check(Address != nullptr);
+		check(InProperty != nullptr);
+
+		if (InValueType != EPropertyBagPropertyType::Struct)
+		{
+			return EPropertyBagResult::TypeMismatch;
+		}
+
+		const FStructProperty* StructProperty = CastFieldChecked<FStructProperty>(InProperty);
+		check(StructProperty->Struct);
+
+		if (RequestedStruct != nullptr && CanCastTo(StructProperty->Struct, RequestedStruct) == false)
+		{
+			return EPropertyBagResult::TypeMismatch;
+		}
+
+		StructProperty->CopyCompleteValue(OutValue, Address);
+
+		return EPropertyBagResult::Success;
+	}
+
+	EPropertyBagResult GetPropertyValueAsObject(const FProperty* InProperty, EPropertyBagPropertyType InValueType, const void* Address, const UClass* RequestedClass,  UObject*& OutValue)
+	{
+		check(Address != nullptr);
+		check(InProperty != nullptr);
+
+		if (InValueType != EPropertyBagPropertyType::Object
+			&& InValueType != EPropertyBagPropertyType::SoftObject
+			&& InValueType != EPropertyBagPropertyType::Class
+			&& InValueType != EPropertyBagPropertyType::SoftClass)
+		{
+			return EPropertyBagResult::TypeMismatch;
+		}
+
+		const FObjectPropertyBase* ObjectProperty = CastFieldChecked<FObjectPropertyBase>(InProperty);
+		check(ObjectProperty->PropertyClass);
+
+		if (RequestedClass != nullptr && CanCastTo(ObjectProperty->PropertyClass, RequestedClass) == false)
+		{
+			return EPropertyBagResult::TypeMismatch;
+		}
+	
+		OutValue = ObjectProperty->GetObjectPropertyValue(Address);
+
+		return EPropertyBagResult::Success;
+	}
+
+	EPropertyBagResult GetPropertyValueAsSoftPath(const FProperty* InProperty, EPropertyBagPropertyType InValueType, const void* Address, FSoftObjectPath& OutValue)
+	{
+		check(Address != nullptr);
+		check(InProperty != nullptr);
+
+		if (InValueType != EPropertyBagPropertyType::SoftObject
+			&& InValueType != EPropertyBagPropertyType::SoftClass)
+		{
+			return EPropertyBagResult::TypeMismatch;
+		}
+
+		const FSoftObjectProperty* SoftObjectProperty = CastFieldChecked<FSoftObjectProperty>(InProperty);
+		check(SoftObjectProperty->PropertyClass);
+
+		OutValue = SoftObjectProperty->GetPropertyValue(Address).ToSoftObjectPath();
+
+		return EPropertyBagResult::Success;
+	}
+
+	EPropertyBagResult GetValueBool(const FProperty* InProperty, EPropertyBagPropertyType InValueType, const void* Address, bool& OutValue)
+	{
+		int64 ReturnValue = 0;
+		const EPropertyBagResult Result = GetPropertyAsInt64(InProperty, InValueType, Address, ReturnValue);
+		if (Result == EPropertyBagResult::Success)
+		{
+			OutValue = ReturnValue != 0;
+		}
+		return Result;
+	}
+
+	EPropertyBagResult GetValueByte(const FProperty* InProperty, EPropertyBagPropertyType InValueType, const void* Address, uint8& OutValue)
+	{
+		uint64 ReturnValue = 0;
+		const EPropertyBagResult Result = GetPropertyAsUInt64(InProperty, InValueType, Address, ReturnValue);
+		if (Result == EPropertyBagResult::Success)
+		{
+			OutValue = static_cast<uint8>(ReturnValue);
+		}
+		return Result;
+	}
+
+	EPropertyBagResult GetValueInt32(const FProperty* InProperty, EPropertyBagPropertyType InValueType, const void* Address, int32& OutValue)
+	{
+		int64 ReturnValue = 0;
+		const EPropertyBagResult Result = GetPropertyAsInt64(InProperty, InValueType, Address, ReturnValue);
+		if (Result == EPropertyBagResult::Success)
+		{
+			OutValue = static_cast<int32>(ReturnValue);
+		}
+		return Result;
+	}
+
+	EPropertyBagResult GetValueUInt32(const FProperty* InProperty, EPropertyBagPropertyType InValueType, const void* Address, uint32& OutValue)
+	{
+		uint64 ReturnValue = 0;
+		const EPropertyBagResult Result = GetPropertyAsUInt64(InProperty, InValueType, Address, ReturnValue);
+		if (Result == EPropertyBagResult::Success)
+		{
+			OutValue = static_cast<uint32>(ReturnValue);
+		}
+		return Result;
+	}
+
+	EPropertyBagResult GetValueInt64(const FProperty* InProperty, EPropertyBagPropertyType InValueType, const void* Address, int64& OutValue)
+	{
+		return GetPropertyAsInt64(InProperty, InValueType, Address, OutValue);
+	}
+
+	EPropertyBagResult GetValueUInt64(const FProperty* InProperty, EPropertyBagPropertyType InValueType, const void* Address, uint64& OutValue)
+	{
+		return GetPropertyAsUInt64(InProperty, InValueType, Address, OutValue);
+	}
+
+	EPropertyBagResult GetValueFloat(const FProperty* InProperty, EPropertyBagPropertyType InValueType, const void* Address, float& OutValue)
+	{
+		double ReturnValue = 0;
+		const EPropertyBagResult Result = GetPropertyAsDouble(InProperty, InValueType, Address, ReturnValue);
+		if (Result == EPropertyBagResult::Success)
+		{
+			OutValue = static_cast<float>(ReturnValue);
+		}
+		return Result;
+	}
+
+	EPropertyBagResult GetValueDouble(const FProperty* InProperty, EPropertyBagPropertyType InValueType, const void* Address, double& OutValue)
+	{
+		return GetPropertyAsDouble(InProperty, InValueType, Address, OutValue);
+	}
+
+	EPropertyBagResult GetValueName(const FProperty* InProperty, EPropertyBagPropertyType InValueType, const void* Address, FName& OutValue)
+	{
+		return GetPropertyValue<FName, FNameProperty>(InProperty, InValueType, Address, OutValue);
+	}
+
+	EPropertyBagResult GetValueString(const FProperty* InProperty, EPropertyBagPropertyType InValueType, const void* Address, FString& OutValue)
+	{
+		return GetPropertyValue<FString, FStrProperty>(InProperty, InValueType, Address, OutValue);
+	}
+
+	EPropertyBagResult GetValueText(const FProperty* InProperty, EPropertyBagPropertyType InValueType, const void* Address, FText& OutValue)
+	{
+		return GetPropertyValue<FText, FTextProperty>(InProperty, InValueType, Address, OutValue);
+	}
+
+	EPropertyBagResult GetValueEnum(const FProperty* InProperty, EPropertyBagPropertyType InValueType, const void* Address, const UEnum* RequestedEnum, uint8& OutValue)
+	{
+		return GetPropertyValueAsEnum(InProperty, InValueType, Address, RequestedEnum, OutValue);
+	}
+
+	EPropertyBagResult GetValueStruct(const FProperty* InProperty, EPropertyBagPropertyType InValueType, const void* Address, const UScriptStruct* RequestedStruct, uint8* OutValue)
+	{
+		return GetPropertyValueAsStruct(InProperty, InValueType, Address, RequestedStruct, OutValue);
+	}
+
+	EPropertyBagResult GetValueObject(const FProperty* InProperty, EPropertyBagPropertyType InValueType, const void* Address, const UClass* RequestedClass, UObject*& OutValue)
+	{
+		return GetPropertyValueAsObject(InProperty, InValueType, Address, RequestedClass, OutValue);
+	}
+
+	EPropertyBagResult GetValueClass(const FProperty* InProperty, EPropertyBagPropertyType InValueType, const void* Address, const UClass* RequestedClass, UClass*& OutValue)
+	{
+		UObject* ReturnValue = nullptr;
+		const EPropertyBagResult Result = GetPropertyValueAsObject(InProperty, InValueType, Address, nullptr, ReturnValue);
+		if (Result != EPropertyBagResult::Success)
+		{
+			return Result;
+		}
+		UClass* Class = Cast<UClass>(ReturnValue);
+		if (Class == nullptr && ReturnValue != nullptr)
+		{
+			return EPropertyBagResult::TypeMismatch;
+		}
+		if(CanCastTo(Class, RequestedClass) == false)
+		{
+			return EPropertyBagResult::TypeMismatch;
+		}
+		OutValue = Class;
+		return Result;
+	}
+
+	EPropertyBagResult GetValueSoftPath(const FProperty* InProperty, EPropertyBagPropertyType InValueType, const void* Address, FSoftObjectPath& OutValue)
+	{
+		return GetPropertyValueAsSoftPath(InProperty, InValueType, Address, OutValue);
+	}
+
+	EPropertyBagResult GetVariableFromMismatchedValueType(const FProperty* InProperty, const FAnimNextParamType& InSrcType, const FAnimNextParamType& InDestType, const uint8* InAddress, uint8* OutResult)
+	{
+		check(InSrcType != InDestType);	// Function assumes that types are mismatched
+
+		switch(InDestType.GetValueType())
+		{
+		case EPropertyBagPropertyType::Bool:
+			return GetValueBool(InProperty, InSrcType.GetValueType(), InAddress, *reinterpret_cast<bool*>(OutResult));
+		case EPropertyBagPropertyType::Byte:
+			return GetValueByte(InProperty, InSrcType.GetValueType(), InAddress, *OutResult);
+		case EPropertyBagPropertyType::Int32:
+			return GetValueInt32(InProperty, InSrcType.GetValueType(), InAddress, *reinterpret_cast<int32*>(OutResult));
+		case EPropertyBagPropertyType::Int64:
+			return GetValueInt64(InProperty, InSrcType.GetValueType(), InAddress, *reinterpret_cast<int64*>(OutResult));
+		case EPropertyBagPropertyType::Float:
+			return GetValueFloat(InProperty, InSrcType.GetValueType(), InAddress, *reinterpret_cast<float*>(OutResult));
+		case EPropertyBagPropertyType::Double:
+			return GetValueDouble(InProperty, InSrcType.GetValueType(), InAddress, *reinterpret_cast<double*>(OutResult));
+		case EPropertyBagPropertyType::Name:
+			return GetValueName(InProperty, InSrcType.GetValueType(), InAddress, *reinterpret_cast<FName*>(OutResult));
+		case EPropertyBagPropertyType::String:
+			return GetValueString(InProperty, InSrcType.GetValueType(), InAddress, *reinterpret_cast<FString*>(OutResult));
+		case EPropertyBagPropertyType::Text:
+			return GetValueText(InProperty, InSrcType.GetValueType(), InAddress, *reinterpret_cast<FText*>(OutResult));
+		case EPropertyBagPropertyType::Enum:
+			return GetValueEnum(InProperty, InSrcType.GetValueType(), InAddress, CastChecked<UEnum>(InDestType.GetValueTypeObject()), *OutResult);
+		case EPropertyBagPropertyType::Struct:
+			return GetValueStruct(InProperty, InSrcType.GetValueType(), InAddress, CastChecked<UScriptStruct>(InDestType.GetValueTypeObject()), OutResult);
+		case EPropertyBagPropertyType::Object:
+			return GetValueObject(InProperty, InSrcType.GetValueType(), InAddress, CastChecked<UClass>(InDestType.GetValueTypeObject()), *reinterpret_cast<UObject**>(OutResult));
+		case EPropertyBagPropertyType::SoftObject:
+			return GetValueSoftPath(InProperty, InSrcType.GetValueType(), InAddress, *reinterpret_cast<FSoftObjectPath*>(OutResult));
+		case EPropertyBagPropertyType::Class:
+			return GetValueClass(InProperty, InSrcType.GetValueType(), InAddress, CastChecked<UClass>(InDestType.GetValueTypeObject()), *reinterpret_cast<UClass**>(OutResult));
+		case EPropertyBagPropertyType::SoftClass:
+			return GetValueSoftPath(InProperty, InSrcType.GetValueType(), InAddress, *reinterpret_cast<FSoftObjectPath*>(OutResult));
+		case EPropertyBagPropertyType::UInt32:
+			return GetValueUInt32(InProperty, InSrcType.GetValueType(), InAddress, *reinterpret_cast<uint32*>(OutResult));
+		case EPropertyBagPropertyType::UInt64:
+			return GetValueUInt64(InProperty, InSrcType.GetValueType(), InAddress, *reinterpret_cast<uint64*>(OutResult));
+		default:
+			return EPropertyBagResult::TypeMismatch;
+		}
+	}
+
+	EPropertyBagResult GetVariableFromMismatchedArrayType(const FProperty* InProperty, const FAnimNextParamType& InSrcType, const FAnimNextParamType& InDestType, const uint8* InAddress, uint8* OutResult)
+	{
+		const FArrayProperty* ArrayProperty = CastFieldChecked<FArrayProperty>(InProperty);
+		const FProperty* ElementProperty = ArrayProperty->Inner;
+		FAnimNextParamType SrcValueType(InSrcType.GetValueType(), FAnimNextParamType::EContainerType::None, InSrcType.GetValueTypeObject());
+		FAnimNextParamType DestValueType(InDestType.GetValueType(), FAnimNextParamType::EContainerType::None, InDestType.GetValueTypeObject());
+		check(SrcValueType != DestValueType);	// Function assumes that types are mismatched
+
+		const TScriptArray<FHeapAllocator>* SrcArray = reinterpret_cast<const TScriptArray<FHeapAllocator>*>(InAddress);
+		TScriptArray<FHeapAllocator>* DestArray = reinterpret_cast<TScriptArray<FHeapAllocator>*>(OutResult);
+
+		const int32 NumElements = SrcArray->Num();
+		const size_t SrcValueTypeSize = InSrcType.GetValueTypeSize();
+		const size_t DestValueTypeSize = InDestType.GetValueTypeSize();
+		const size_t DestValueTypeAlignment = InDestType.GetValueTypeSize();
+
+		// Reallocate dest array
+		DestArray->SetNumUninitialized(NumElements, DestValueTypeSize, DestValueTypeAlignment);
+
+		// Perform per-element conversion
+		bool bSucceeded = true;
+		const uint8* SrcElement = static_cast<const uint8*>(SrcArray->GetData());
+		uint8* DestElement = static_cast<uint8*>(DestArray->GetData());
+		for (int32 Index = 0; Index < NumElements; ++Index)
+		{
+			bSucceeded &= (GetVariableFromMismatchedValueType(ElementProperty, SrcValueType, DestValueType, SrcElement, DestElement) == EPropertyBagResult::Success);
+			SrcElement += SrcValueTypeSize;
+			DestElement += DestValueTypeSize;
+		}
+
+		return bSucceeded ? EPropertyBagResult::Success : EPropertyBagResult::TypeMismatch;
+	}
+
+	EPropertyBagResult GetVariableFromMismatchedType(const FPropertyBagPropertyDesc* InDesc, const FAnimNextParamType& InSrcType, const FAnimNextParamType& InDestType, const uint8* InAddress, uint8* OutResult)
+	{
+		switch(InSrcType.GetContainerType())
+		{
+		case EPropertyBagContainerType::None:
+			if(InDestType.GetContainerType() == EPropertyBagContainerType::None)
+			{
+				return GetVariableFromMismatchedValueType(InDesc->CachedProperty, InSrcType, InDestType, InAddress, OutResult);
+			}
+			return EPropertyBagResult::TypeMismatch;
+		case EPropertyBagContainerType::Array:
+			if(InDestType.GetContainerType() == EPropertyBagContainerType::Array)
+			{
+				return GetVariableFromMismatchedArrayType(InDesc->CachedProperty, InSrcType, InDestType, InAddress, OutResult);
+			}
+			return EPropertyBagResult::TypeMismatch;
+		default:
+			return EPropertyBagResult::TypeMismatch;
+		}
+	}
+}
+
+EPropertyBagResult FAnimNextDataInterfaceInstance::GetVariableInternal(FName InVariableName, const FAnimNextParamType& InType, TArrayView<uint8> OutResult) const
+{
+	using namespace UE::AnimNext;
+	
+	const UPropertyBag* PropertyBag = Variables.GetPropertyBagStruct();
+	if(PropertyBag == nullptr)
+	{
+		return EPropertyBagResult::PropertyNotFound;
+	}
+
+	const FPropertyBagPropertyDesc* Desc = PropertyBag->FindPropertyDescByName(InVariableName);
+	if(Desc == nullptr)
+	{
+		return EPropertyBagResult::PropertyNotFound;
+	}
+
+	const int32 DescIndex = Desc - &PropertyBag->GetPropertyDescs()[0];
+	const uint8* Memory = ExtendedExecuteContext.ExternalVariableRuntimeData[DescIndex].Memory;
+
+	const FAnimNextParamType InternalType(Desc->ValueType, Desc->ContainerTypes.GetFirstContainerType(), Desc->ValueTypeObject);
+	if(InternalType != InType)
+	{
+		return Private::GetVariableFromMismatchedType(Desc, InternalType, InType, Memory, OutResult.GetData());
+	}
+
+	Desc->CachedProperty->CopyCompleteValue(OutResult.GetData(), Memory);
+	return EPropertyBagResult::Success;
+}
+
+uint8* FAnimNextDataInterfaceInstance::GetMemoryForVariable(int32 InVariableIndex, FName InVariableName, const FProperty* InVariableProperty) const
+{
+	const UPropertyBag* PropertyBag = Variables.GetPropertyBagStruct();
+	if(PropertyBag == nullptr)
+	{
+		return nullptr;
+	}
+
+	if(!PropertyBag->GetPropertyDescs().IsValidIndex(InVariableIndex))
+	{
+		UE_LOGFMT(LogAnimation, Error, "FAnimNextDataInterfaceInstance::GetMemoryForVariable: Variable index {Index} out of range ({Range})", InVariableIndex, PropertyBag->GetPropertyDescs().Num());
+		return nullptr;
+	}
+
+	if(!ExtendedExecuteContext.ExternalVariableRuntimeData.IsValidIndex(InVariableIndex))
+	{
+		UE_LOGFMT(LogAnimation, Error, "FAnimNextDataInterfaceInstance::GetMemoryForVariable: External variable index {Index} out of range ({Range})", InVariableIndex, ExtendedExecuteContext.ExternalVariableRuntimeData.Num());
+		return nullptr;
+	}
+
+	check(ExtendedExecuteContext.ExternalVariableRuntimeData.Num() == PropertyBag->GetPropertyDescs().Num());
+
+	const FPropertyBagPropertyDesc& Desc = PropertyBag->GetPropertyDescs()[InVariableIndex];
+	if(Desc.Name != InVariableName)
+	{
+		UE_LOGFMT(LogAnimation, Error, "FAnimNextDataInterfaceInstance::GetMemoryForVariable: Mismatched variable names: {Name} vs {OtherName} in '{Host}'", Desc.Name, InVariableName, GetDataInterfaceName());
+		return nullptr;
+	}
+
+	if(Desc.CachedProperty->GetClass() != InVariableProperty->GetClass())
+	{
+		UE_LOGFMT(LogAnimation, Error, "FAnimNextDataInterfaceInstance::GetMemoryForVariable: Mismatched variable types: {Name}:{Type} vs {OtherType} in '{Host}'", Desc.Name, Desc.CachedProperty->GetFName(), InVariableProperty->GetFName(), GetDataInterfaceName());
+		return nullptr;
+	}
+
+	return ExtendedExecuteContext.ExternalVariableRuntimeData[InVariableIndex].Memory;
+}
diff --git a/Engine/Plugins/Experimental/AnimNext/Source/AnimNext/Private/DataInterface/DataInterfaceStructAdapter.cpp b/Engine/Plugins/Experimental/AnimNext/Source/AnimNext/Private/DataInterface/DataInterfaceStructAdapter.cpp
new file mode 100644
--- /dev/null
+++ b/Engine/Plugins/Experimental/AnimNext/Source/AnimNext/Private/DataInterface/DataInterfaceStructAdapter.cpp
@@ -0,0 +1,32 @@
+﻿// Copyright Epic Games, Inc. All Rights Reserved.
+
+#include "DataInterface/DataInterfaceStructAdapter.h"
+
+#include "Logging/StructuredLog.h"
+
+namespace UE::AnimNext
+{
+
+const UAnimNextDataInterface* FDataInterfaceStructAdapter::GetDataInterface() const
+{
+	return DataInterface;
+}
+
+uint8* FDataInterfaceStructAdapter::GetMemoryForVariable(int32 InVariableIndex, FName InVariableName, const FProperty* InVariableProperty) const
+{
+	const FProperty* Property = StructView.GetScriptStruct()->FindPropertyByName(InVariableName);
+	if(Property == nullptr)
+	{
+		return nullptr;
+	}
+
+	if(Property->GetClass() != InVariableProperty->GetClass())
+	{
+		UE_LOGFMT(LogAnimation, Error, "FDataInterfaceStructAdapter::GetMemoryForVariable: Mismatched variable types: {Name}:{Type} vs {OtherType} in '{Host}'", InVariableName, Property->GetFName(), InVariableProperty->GetFName(), GetDataInterfaceName());
+		return nullptr;
+	}
+
+	return Property->ContainerPtrToValuePtr<uint8>(StructView.GetMemory());
+}
+
+}
diff --git a/Engine/Plugins/Experimental/AnimNext/Source/AnimNext/Private/Graph/AnimNextAnimationGraph.cpp b/Engine/Plugins/Experimental/AnimNext/Source/AnimNext/Private/Graph/AnimNextAnimationGraph.cpp
--- a/Engine/Plugins/Experimental/AnimNext/Source/AnimNext/Private/Graph/AnimNextAnimationGraph.cpp
+++ b/Engine/Plugins/Experimental/AnimNext/Source/AnimNext/Private/Graph/AnimNextAnimationGraph.cpp
@@ -57,7 +57,7 @@ void UAnimNextAnimationGraph::AllocateInstanceImpl(FAnimNextModuleInstance* InMo
 	{
 		TSharedPtr<FAnimNextGraphInstance> InstanceImpl = MakeShared<FAnimNextGraphInstance>();
 
-		InstanceImpl->AnimationGraph = this;
+		InstanceImpl->DataInterface = this;
 		InstanceImpl->ModuleInstance = InModuleInstance;
 		InstanceImpl->ParentGraphInstance = ParentGraphInstance;
 		InstanceImpl->EntryPoint = EntryPoint;
diff --git a/Engine/Plugins/Experimental/AnimNext/Source/AnimNext/Private/Graph/AnimNextGraphInstance.cpp b/Engine/Plugins/Experimental/AnimNext/Source/AnimNext/Private/Graph/AnimNextGraphInstance.cpp
--- a/Engine/Plugins/Experimental/AnimNext/Source/AnimNext/Private/Graph/AnimNextGraphInstance.cpp
+++ b/Engine/Plugins/Experimental/AnimNext/Source/AnimNext/Private/Graph/AnimNextGraphInstance.cpp
@@ -3,6 +3,7 @@
 #include "Graph/AnimNextGraphInstance.h"
 
 #include "AnimNextStats.h"
+#include "DataInterface/AnimNextDataInterfaceHost.h"
 #include "Graph/AnimNextAnimationGraph.h"
 #include "TraitCore/ExecutionContext.h"
 #include "Graph/GC_GraphInstanceComponent.h"
@@ -49,7 +50,7 @@ void FAnimNextGraphInstance::Release()
 	RootGraphInstance = nullptr;
 	ExtendedExecuteContext.Reset();
 	Components.Empty();
-	AnimationGraph = nullptr;
+	DataInterface = nullptr;
 }
 
 bool FAnimNextGraphInstance::IsValid() const
@@ -59,7 +60,7 @@ bool FAnimNextGraphInstance::IsValid() const
 
 const UAnimNextAnimationGraph* FAnimNextGraphInstance::GetAnimationGraph() const
 {
-	return AnimationGraph;
+	return CastChecked<UAnimNextAnimationGraph>(DataInterface);
 }
 
 FName FAnimNextGraphInstance::GetEntryPoint() const
@@ -89,12 +90,12 @@ FAnimNextGraphInstance* FAnimNextGraphInstance::GetRootGraphInstance() const
 
 bool FAnimNextGraphInstance::UsesAnimationGraph(const UAnimNextAnimationGraph* InAnimationGraph) const
 {
-	return AnimationGraph == InAnimationGraph;
+	return GetAnimationGraph() == InAnimationGraph;
 }
 
 bool FAnimNextGraphInstance::UsesEntryPoint(FName InEntryPoint) const
 {
-	if(AnimationGraph != nullptr)
+	if(const UAnimNextAnimationGraph* AnimationGraph = GetAnimationGraph())
 	{
 		if(InEntryPoint == NAME_None)
 		{
@@ -154,9 +155,14 @@ void FAnimNextGraphInstance::Update()
 	bHasUpdatedOnce = true;
 }
 
-FRigVMExtendedExecuteContext& FAnimNextGraphInstance::GetExtendedExecuteContext()
+FAnimNextDataInterfaceInstance* FAnimNextGraphInstance::GetHost() const
 {
-	return ExtendedExecuteContext;
+	if(ParentGraphInstance != nullptr)
+	{
+		return ParentGraphInstance;
+	}
+
+	return ModuleInstance;
 }
 
 void FAnimNextGraphInstance::ExecuteLatentPins(const TConstArrayView<UE::AnimNext::FLatentPropertyHandle>& LatentHandles, void* DestinationBasePtr, bool bIsFrozen)
@@ -168,7 +174,7 @@ void FAnimNextGraphInstance::ExecuteLatentPins(const TConstArrayView<UE::AnimNex
 		return;
 	}
 
-	if (URigVM* VM = AnimationGraph->VM)
+	if (URigVM* VM = GetAnimationGraph()->RigVM)
 	{
 		FAnimNextExecuteContext& AnimNextContext = ExtendedExecuteContext.GetPublicDataSafe<FAnimNextExecuteContext>();
 		AnimNextContext.SetContextData<FAnimNextGraphContextData>(ModuleInstance, this, LatentHandles, DestinationBasePtr, bIsFrozen);
@@ -191,22 +197,21 @@ void FAnimNextGraphInstance::Freeze()
 	GraphInstancePtr.Reset();
 	ExtendedExecuteContext.Reset();
 	Components.Empty();
+	PublicVariablesState = PublicVariablesState == EPublicVariablesState::Bound ? EPublicVariablesState::Unbound : EPublicVariablesState::None;
 	bHasUpdatedOnce = false;
 }
 
 void FAnimNextGraphInstance::Thaw()
 {
-	if (const UAnimNextAnimationGraph* AnimationGraphPtr = AnimationGraph)
+	if (const UAnimNextAnimationGraph* AnimationGraph = GetAnimationGraph())
 	{
-		Variables.MigrateToNewBagInstance(AnimationGraphPtr->VariableDefaults);
-
-		ExtendedExecuteContext = AnimationGraphPtr->ExtendedExecuteContext;
+		Variables.MigrateToNewBagInstance(AnimationGraph->VariableDefaults);
 
-		RebindPublicVariables();
+		ExtendedExecuteContext = AnimationGraph->ExtendedExecuteContext;
 
 		{
 			UE::AnimNext::FExecutionContext Context(*this);
-			if(const FAnimNextTraitHandle* FoundHandle = AnimationGraphPtr->ResolvedRootTraitHandles.Find(EntryPoint))
+			if(const FAnimNextTraitHandle* FoundHandle = AnimationGraph->ResolvedRootTraitHandles.Find(EntryPoint))
 			{
 				GraphInstancePtr = Context.AllocateNodeInstance(*this, *FoundHandle);
 			}
@@ -220,55 +225,72 @@ void FAnimNextGraphInstance::Thaw()
 	}
 }
 
-void FAnimNextGraphInstance::RebindPublicVariables()
+void FAnimNextGraphInstance::OnModuleCompiled(UAnimNextModule* InModule)
 {
-	if(PublicVariablesState != EPublicVariablesState::Bound)
+	// If we are hosted directly by a module, invalidate and mark our bindings as needing update. They will be lazily re-bound the next time we run.
+	if(ModuleInstance && InModule == ModuleInstance->GetModule() && ParentGraphInstance == nullptr)
 	{
-		return;
+		UnbindPublicVariables();
 	}
+}
 
-	// Setup external variables memory ptrs manually as we dont follow the pattern of owning multiple URigVMHosts like control rig.
-	// InitializeVM() is called, but only sets up handles for the defaults in the module, not for an instance
-	const int32 NumVariables = Variables.GetNumPropertiesInBag();
-	TArray<FRigVMExternalVariableRuntimeData> ExternalVariableRuntimeData;
-	ExternalVariableRuntimeData.Reserve(NumVariables);
-	TConstArrayView<FPropertyBagPropertyDesc> Descs = Variables.GetPropertyBagStruct()->GetPropertyDescs();
-	uint8* BasePtr = Variables.GetMutableValue().GetMemory();
-	for(int32 VariableIndex = 0; VariableIndex < NumVariables; ++VariableIndex)
-	{
-		ExternalVariableRuntimeData.Emplace(Descs[VariableIndex].CachedProperty->ContainerPtrToValuePtr<uint8>(BasePtr));
-	}
-	ExtendedExecuteContext.ExternalVariableRuntimeData = MoveTemp(ExternalVariableRuntimeData);
+#endif
 
-	// Re-apply any bindings from our host
-	for(const FCachedVariableBinding& CachedBinding : CachedVariableBindings)
+// Helper for binding to both FAnimNextDataInterfaceInstance and IDataInterfaceHost without an abstraction between the two
+template<typename HostType>
+bool FAnimNextGraphInstance::BindToHostHelper(const HostType& InHost, bool bInAutoBind)
+{
+	bool bPublicVariablesBound = false;
+
+	const UAnimNextAnimationGraph* AnimationGraph = GetAnimationGraph();
+	const UPropertyBag* PropertyBag = Variables.GetPropertyBagStruct();
+	for(const FAnimNextImplementedDataInterface& ImplementedInterface : AnimationGraph->GetImplementedInterfaces())
 	{
-		// TODO: remove this linear search in FindPropertyDescByName with a hash table?
-		if(const FPropertyBagPropertyDesc* Desc = Variables.FindPropertyDescByName(CachedBinding.VariableName))
+		if(bInAutoBind && !ImplementedInterface.bAutoBindToHost)
 		{
-			int32 VariableIndex = Desc - Variables.GetPropertyBagStruct()->GetPropertyDescs().GetData();
-			ExtendedExecuteContext.ExternalVariableRuntimeData[VariableIndex].Memory = CachedBinding.Memory;
+			continue;
 		}
-	}
 
-	// Now reinitialize the 'instance', cache memory handles etc. in the context
-	AnimationGraph->VM->InitializeInstance(ExtendedExecuteContext);
-}
+		const FAnimNextImplementedDataInterface* HostImplementedInterface = InHost.GetDataInterface()->FindImplementedInterface(ImplementedInterface.DataInterface);
+		if(HostImplementedInterface == nullptr)
+		{
+			// Host does not implement this interface, so skip
+			continue;
+		}
 
-void FAnimNextGraphInstance::OnModuleCompiled(UAnimNextModule* InModule)
-{
-	// If we are hosted directly by a module, invalidate and mark our bindings as needing update 
-	if(ModuleInstance && InModule == ModuleInstance->Module && ParentGraphInstance == nullptr)
-	{
-		UnbindPublicVariables();
+		if(HostImplementedInterface->NumVariables != ImplementedInterface.NumVariables)
+		{
+			UE_LOGFMT(LogAnimation, Error, "BindToHost: Mismatched interface variables: '{Name}' ({Count}) vs Host '{Host}' ({HostCount})", AnimationGraph->GetFName(), ImplementedInterface.NumVariables, InHost.GetDataInterfaceName(), HostImplementedInterface->NumVariables);
+			continue;
+		}
+
+		int32 VariableIndex = ImplementedInterface.VariableIndex;
+		int32 HostVariableIndex = HostImplementedInterface->VariableIndex;
+		const int32 EndVariableIndex = ImplementedInterface.VariableIndex + ImplementedInterface.NumVariables;
+		for(; VariableIndex < EndVariableIndex; ++VariableIndex, ++HostVariableIndex)
+		{
+			const FPropertyBagPropertyDesc& Desc = PropertyBag->GetPropertyDescs()[VariableIndex];
+
+			uint8* HostMemory = InHost.GetMemoryForVariable(HostVariableIndex, Desc.Name, Desc.CachedProperty);
+			if(HostMemory == nullptr)
+			{
+				continue;
+			}
+
+			ExtendedExecuteContext.ExternalVariableRuntimeData[VariableIndex].Memory = HostMemory;
+			bPublicVariablesBound = true;
+		}
 	}
-}
 
-#endif
+	return bPublicVariablesBound;
+}
 
-void FAnimNextGraphInstance::BindPublicVariables(TConstArrayView<FRigVMTraitScope> InTraitScopes)
+void FAnimNextGraphInstance::BindPublicVariables(TConstArrayView<UE::AnimNext::IDataInterfaceHost*> InHosts)
 {
-	if(AnimationGraph == nullptr)
+	using namespace UE::AnimNext;
+
+	const UAnimNextAnimationGraph* AnimationGraph = GetAnimationGraph();
+	if (AnimationGraph == nullptr)
 	{
 		return;
 	}
@@ -278,56 +300,33 @@ void FAnimNextGraphInstance::BindPublicVariables(TConstArrayView<FRigVMTraitScop
 		return;
 	}
 
-#if WITH_EDITORONLY_DATA
-	CachedVariableBindings.Reset();
-#endif
-
-	bool bPublicVariablesBound = false;
-	for(const FRigVMTraitScope& TraitScope : InTraitScopes)
+	const UPropertyBag* PropertyBag = Variables.GetPropertyBagStruct();
+	if(PropertyBag == nullptr)
 	{
-		const FRigVMTrait_AnimNextPublicVariables* VariablesTrait = TraitScope.GetTrait<FRigVMTrait_AnimNextPublicVariables>();
-		if(VariablesTrait == nullptr)
-		{
-			continue;
-		}
-
-		for(int32 TraitVariableIndex = 0; TraitVariableIndex < VariablesTrait->VariableNames.Num(); ++TraitVariableIndex)
-		{
-			FName VariableName = VariablesTrait->VariableNames[TraitVariableIndex];
-			// TODO: remove this linear search in FindPropertyDescByName with a hash table?
-			// TODO: multiple traits with the same name can collide at the moment, we should validate against this is at the compiler level to avoid issues
-			const FPropertyBagPropertyDesc* Desc = Variables.FindPropertyDescByName(VariableName);
-			if(Desc == nullptr)
-			{
-				continue;
-			}
+		// Nothing to bind
+		PublicVariablesState = EPublicVariablesState::None;
+		return;
+	}
 
-			if(!TraitScope.GetAdditionalMemoryHandles().IsValidIndex(TraitVariableIndex))
-			{
-				continue;
-			}
+	bool bPublicVariablesBound = false;
 
-			const FRigVMMemoryHandle& MemoryHandle = TraitScope.GetAdditionalMemoryHandles()[TraitVariableIndex];
-			int32 VariableIndex = Desc - Variables.GetPropertyBagStruct()->GetPropertyDescs().GetData();
-			if(Desc->CachedProperty->GetClass() != MemoryHandle.GetProperty()->GetClass())
-			{
-				UE_LOGFMT(LogAnimation, Warning, "Mismatched variable types when binding AnimNext graph: {Name}:{Type} vs {OtherType}", Desc->Name, Desc->CachedProperty->GetFName(), MemoryHandle.GetProperty()->GetFName());
-				continue;
-			}
+	// First apply any automatic bindings to this instance's host
+	if(FAnimNextDataInterfaceInstance* InstanceHost = GetHost())
+	{
+		bPublicVariablesBound |= BindToHostHelper(*InstanceHost, true);
+	}
 
-			uint8* Memory = const_cast<uint8*>(MemoryHandle.GetData());
-			ExtendedExecuteContext.ExternalVariableRuntimeData[VariableIndex].Memory = Memory;
-#if WITH_EDITORONLY_DATA
-			CachedVariableBindings.Add({ VariableName, Memory });
-#endif
-			bPublicVariablesBound = true;
-		}
+	// Next bind to any supplied host interfaces (bInAutoBind = false)
+	for(const IDataInterfaceHost* HostInterface : InHosts)
+	{
+		check(HostInterface);
+		bPublicVariablesBound |= BindToHostHelper(*HostInterface, false);
 	}
 
 	if(bPublicVariablesBound)
 	{
 		// Re-initialize memory handles
-		AnimationGraph->VM->InitializeInstance(ExtendedExecuteContext, /* bCopyMemory = */false);
+		AnimationGraph->RigVM->InitializeInstance(ExtendedExecuteContext, /* bCopyMemory = */false);
 	}
 
 	PublicVariablesState = EPublicVariablesState::Bound;
@@ -335,14 +334,16 @@ void FAnimNextGraphInstance::BindPublicVariables(TConstArrayView<FRigVMTraitScop
 
 void FAnimNextGraphInstance::UnbindPublicVariables()
 {
-	if(PublicVariablesState != EPublicVariablesState::Bound)
+	const UAnimNextAnimationGraph* AnimationGraph = GetAnimationGraph();
+	if (AnimationGraph == nullptr)
 	{
 		return;
 	}
 
-#if WITH_EDITORONLY_DATA
-	CachedVariableBindings.Reset();
-#endif
+	if(PublicVariablesState != EPublicVariablesState::Bound)
+	{
+		return;
+	}
 
 	// Reset external variable ptrs to point to internal public vars
 	const int32 NumVariables = Variables.GetNumPropertiesInBag();
@@ -354,7 +355,7 @@ void FAnimNextGraphInstance::UnbindPublicVariables()
 	}
 
 	// Re-initialize memory handles
-	AnimationGraph->VM->InitializeInstance(ExtendedExecuteContext, /* bCopyMemory = */false);
+	AnimationGraph->RigVM->InitializeInstance(ExtendedExecuteContext, /* bCopyMemory = */false);
 
 	PublicVariablesState = EPublicVariablesState::Unbound;
 }
diff --git a/Engine/Plugins/Experimental/AnimNext/Source/AnimNext/Private/Graph/AnimNextGraphInstancePtr.cpp b/Engine/Plugins/Experimental/AnimNext/Source/AnimNext/Private/Graph/AnimNextGraphInstancePtr.cpp
--- a/Engine/Plugins/Experimental/AnimNext/Source/AnimNext/Private/Graph/AnimNextGraphInstancePtr.cpp
+++ b/Engine/Plugins/Experimental/AnimNext/Source/AnimNext/Private/Graph/AnimNextGraphInstancePtr.cpp
@@ -56,6 +56,8 @@ void FAnimNextGraphInstancePtr::AddStructReferencedObjects(FReferenceCollector&
 {
 	if (Impl)
 	{
+		FAnimNextGraphInstance* ImplPtr = Impl.Get();
+		Collector.AddPropertyReferences(FAnimNextGraphInstance::StaticStruct(), ImplPtr);
 		Impl->AddStructReferencedObjects(Collector);
 	}
 }
@@ -95,8 +97,8 @@ bool FAnimNextGraphInstancePtr::RequiresPublicVariableBinding() const
 	return Impl->PublicVariablesState == FAnimNextGraphInstance::EPublicVariablesState::Unbound;
 }
 
-void FAnimNextGraphInstancePtr::BindPublicVariables(TConstArrayView<FRigVMTraitScope> InTraitScopes) const
+void FAnimNextGraphInstancePtr::BindPublicVariables(TConstArrayView<UE::AnimNext::IDataInterfaceHost*> InHosts) const
 {
 	check(Impl);
-	return Impl->BindPublicVariables(InTraitScopes);
+	return Impl->BindPublicVariables(InHosts);
 }
diff --git a/Engine/Plugins/Experimental/AnimNext/Source/AnimNext/Private/Graph/RigUnit_AnimNextRunAnimationGraph.cpp b/Engine/Plugins/Experimental/AnimNext/Source/AnimNext/Private/Graph/RigUnit_AnimNextRunAnimationGraph.cpp
--- a/Engine/Plugins/Experimental/AnimNext/Source/AnimNext/Private/Graph/RigUnit_AnimNextRunAnimationGraph.cpp
+++ b/Engine/Plugins/Experimental/AnimNext/Source/AnimNext/Private/Graph/RigUnit_AnimNextRunAnimationGraph.cpp
@@ -7,6 +7,7 @@
 #include "IAnimNextModuleInterface.h"
 #include "Module/AnimNextModuleInstance.h"
 #include "Graph/AnimNextAnimationGraph.h"
+#include "Graph/RigVMTrait_AnimNextPublicVariables.h"
 #include "TraitCore/TraitEventList.h"
 
 #include UE_INLINE_GENERATED_CPP_BY_NAME(RigUnit_AnimNextRunAnimationGraph)
@@ -52,7 +53,23 @@ FRigUnit_AnimNextRunAnimationGraph_Execute()
 	if(Instance.RequiresPublicVariableBinding())
 	{
 		// Bind any public variables we find
-		Instance.BindPublicVariables(ExecuteContext.GetTraits());
+		TArray<FPublicVariablesTraitToDataInterfaceHostAdapter, TInlineAllocator<4>> Adapters;
+		Adapters.Reserve(ExecuteContext.GetTraits().Num());
+		TArray<IDataInterfaceHost*, TInlineAllocator<4>> Hosts;
+		Hosts.Reserve(ExecuteContext.GetTraits().Num());
+		for(const FRigVMTraitScope& TraitScope : ExecuteContext.GetTraits())
+		{
+			const FRigVMTrait_AnimNextPublicVariables* VariablesTrait = TraitScope.GetTrait<FRigVMTrait_AnimNextPublicVariables>();
+			if(VariablesTrait == nullptr)
+			{
+				continue;
+			}
+
+			FPublicVariablesTraitToDataInterfaceHostAdapter& Adapter = Adapters.Emplace_GetRef(*VariablesTrait, TraitScope);
+			Hosts.Add(&Adapter);
+		}
+
+		Instance.BindPublicVariables(Hosts);
 	}
 
 	const UE::AnimNext::FReferencePose& RefPose = ReferencePose.ReferencePose.GetRef<UE::AnimNext::FReferencePose>();
diff --git a/Engine/Plugins/Experimental/AnimNext/Source/AnimNext/Private/Graph/RigVMTrait_AnimNextPublicVariables.cpp b/Engine/Plugins/Experimental/AnimNext/Source/AnimNext/Private/Graph/RigVMTrait_AnimNextPublicVariables.cpp
--- a/Engine/Plugins/Experimental/AnimNext/Source/AnimNext/Private/Graph/RigVMTrait_AnimNextPublicVariables.cpp
+++ b/Engine/Plugins/Experimental/AnimNext/Source/AnimNext/Private/Graph/RigVMTrait_AnimNextPublicVariables.cpp
@@ -2,6 +2,7 @@
 
 #include "Graph/RigVMTrait_AnimNextPublicVariables.h"
 #include "AnimNextRigVMAsset.h"
+#include "Logging/StructuredLog.h"
 #if WITH_EDITOR
 #include "RigVMModel/RigVMController.h"
 #include "RigVMModel/RigVMPin.h"
@@ -58,4 +59,44 @@ bool FRigVMTrait_AnimNextPublicVariables::ShouldCreatePinForProperty(const FProp
 		InProperty->GetFName() == GET_MEMBER_NAME_CHECKED(FRigVMTrait_AnimNextPublicVariables, VariableNames) ||
 		VariableNames.Contains(InProperty->GetFName());
 }
-#endif
\ No newline at end of file
+#endif
+
+namespace UE::AnimNext
+{
+
+const UAnimNextDataInterface* FPublicVariablesTraitToDataInterfaceHostAdapter::GetDataInterface() const
+{
+	return Trait.Asset;
+}
+
+uint8* FPublicVariablesTraitToDataInterfaceHostAdapter::GetMemoryForVariable(int32 InVariableIndex, FName InVariableName, const FProperty* InVariableProperty) const
+{
+	// Note we dont use InVariableIndex here as we may not have bound all variables in an interface
+
+	int32 TraitVariableIndex = Trait.VariableNames.Find(InVariableName);
+	if(TraitVariableIndex == INDEX_NONE)
+	{
+		// Variable not bound here
+		return nullptr;
+	}
+
+	if(!TraitScope.GetAdditionalMemoryHandles().IsValidIndex(TraitVariableIndex))
+	{
+		// Memory handle is out of bounds
+		// If this ensure fires then we have a mismatch between the variable names and the compiled memory handles, indicating a bug with the
+		// compilation of trait additional memory handles (programmatic pins)
+		ensure(false);
+		return nullptr;
+	}
+
+	const FRigVMMemoryHandle& MemoryHandle = TraitScope.GetAdditionalMemoryHandles()[TraitVariableIndex];
+	if(InVariableProperty->GetClass() != MemoryHandle.GetProperty()->GetClass())
+	{
+		UE_LOGFMT(LogAnimation, Error, "FPublicVariablesTraitToDataInterfaceHostAdapter::GetMemoryForVariable: Mismatched variable types: {Name}:{Type} vs {OtherType}", InVariableName, InVariableProperty->GetFName(), MemoryHandle.GetProperty()->GetFName());
+		return nullptr;
+	}
+
+	return const_cast<uint8*>(MemoryHandle.GetData());
+}
+
+}
\ No newline at end of file
diff --git a/Engine/Plugins/Experimental/AnimNext/Source/AnimNext/Private/Module/AnimNextModuleInstance.cpp b/Engine/Plugins/Experimental/AnimNext/Source/AnimNext/Private/Module/AnimNextModuleInstance.cpp
--- a/Engine/Plugins/Experimental/AnimNext/Source/AnimNext/Private/Module/AnimNextModuleInstance.cpp
+++ b/Engine/Plugins/Experimental/AnimNext/Source/AnimNext/Private/Module/AnimNextModuleInstance.cpp
@@ -18,11 +18,11 @@ FAnimNextModuleInstance::FAnimNextModuleInstance(
 		UAnimNextModule* InModule,
 		UObject* InObject,
 		EAnimNextModuleInitMethod InInitMethod)
-	: Module(InModule)
-	, Object(InObject)
+	: Object(InObject)
 	, RunState(ERunState::None)
 	, InitMethod(InInitMethod)
 {
+	DataInterface = InModule;
 }
 
 FAnimNextModuleInstance::~FAnimNextModuleInstance()
@@ -30,7 +30,7 @@ FAnimNextModuleInstance::~FAnimNextModuleInstance()
 	ResetBindingsAndInstanceData();
 
 	Object = nullptr;
-	Module = nullptr;
+	DataInterface = nullptr;
 	Handle.Reset();
 }
 
@@ -75,7 +75,7 @@ void CacheAllModuleEvents()
 }
 
 // Gets information about the module events that are implemented by the supplied VM
-static TConstArrayView<FImplementedModuleEvent> GetImplementedModuleEvents(URigVM* VM)
+static TConstArrayView<FImplementedModuleEvent> GetImplementedModuleEvents(const URigVM* VM)
 {
 	check(IsInGameThread());	// This function cannot be run concurrently because of static usage
 	check(GAllModuleEvents.Num() > 0);	// Call CacheAllModuleEvents before this function
@@ -105,6 +105,7 @@ void FAnimNextModuleInstance::Initialize()
 	check(IsInGameThread());
 
 	check(Object)
+	const UAnimNextModule* Module = GetModule();
 	check(Module);
 	check(Handle.IsValid());
 
@@ -112,7 +113,7 @@ void FAnimNextModuleInstance::Initialize()
 	bIsEditor = World && World->WorldType == EWorldType::Editor;
 
 	// Get all the module events from the VM entry points, sorted by phase
-	URigVM* VM = Module->GetVM();
+	URigVM* VM = Module->RigVM;
 	TConstArrayView<Private::FImplementedModuleEvent> ImplementedModuleEvents = Private::GetImplementedModuleEvents(VM);
 
 	// Setup tick function graph using module events
@@ -247,7 +248,7 @@ void FAnimNextModuleInstance::Invalidate()
 
 	ResetBindingsAndInstanceData();
 
-	Module = nullptr;
+	DataInterface = nullptr;
 	Object = nullptr;
 	Handle.Reset();
 }
@@ -335,7 +336,7 @@ void FAnimNextModuleInstance::CopyProxyVariables()
 			TConstArrayView<FPropertyBagPropertyDesc> PublicProxyDescs = PublicVariablesProxy.Data.GetPropertyBagStruct()->GetPropertyDescs();
 			const uint8* SourceContainerPtr = PublicVariablesProxy.Data.GetValue().GetMemory();
 			uint8* TargetContainerPtr = Variables.GetMutableValue().GetMemory();
-			for (TBitArray<>::FConstIterator It(PublicVariablesProxy.DirtyFlags); It; ++It)
+			for (TConstSetBitIterator<> It(PublicVariablesProxy.DirtyFlags); It; ++It)
 			{
 				const int32 Index = It.GetIndex();
 				const FProperty* SourceProperty = PublicProxyDescs[Index].CachedProperty;
@@ -351,6 +352,11 @@ void FAnimNextModuleInstance::CopyProxyVariables()
 	}
 }
 
+const UAnimNextModule* FAnimNextModuleInstance::GetModule() const
+{
+	return CastChecked<UAnimNextModule>(DataInterface);
+}
+
 #if WITH_EDITOR
 void FAnimNextModuleInstance::OnModuleCompiled()
 {
diff --git a/Engine/Plugins/Experimental/AnimNext/Source/AnimNext/Private/Module/ModuleTickFunction.cpp b/Engine/Plugins/Experimental/AnimNext/Source/AnimNext/Private/Module/ModuleTickFunction.cpp
--- a/Engine/Plugins/Experimental/AnimNext/Source/AnimNext/Private/Module/ModuleTickFunction.cpp
+++ b/Engine/Plugins/Experimental/AnimNext/Source/AnimNext/Private/Module/ModuleTickFunction.cpp
@@ -135,9 +135,9 @@ void FModuleEventTickFunction::Run(float InDeltaTime)
 		Function.GetValue()(FModuleTaskContext(*ModuleInstance));
 	}
 
-	if (URigVM* VM = ModuleInstance->Module->GetVM())
+	if (URigVM* VM = ModuleInstance->GetModule()->RigVM)
 	{
-		FRigVMExtendedExecuteContext& Context = ModuleInstance->ExtendedExecuteContext;
+		FRigVMExtendedExecuteContext& Context = ModuleInstance->GetExtendedExecuteContext();
 		check(Context.VMHash == VM->GetVMHash());
 
 		FAnimNextExecuteContext& AnimNextContext = Context.GetPublicDataSafe<FAnimNextExecuteContext>();
diff --git a/Engine/Plugins/Experimental/AnimNext/Source/AnimNext/Private/Param/ParamType.cpp b/Engine/Plugins/Experimental/AnimNext/Source/AnimNext/Private/Param/ParamType.cpp
--- a/Engine/Plugins/Experimental/AnimNext/Source/AnimNext/Private/Param/ParamType.cpp
+++ b/Engine/Plugins/Experimental/AnimNext/Source/AnimNext/Private/Param/ParamType.cpp
@@ -309,6 +309,9 @@ bool FAnimNextParamType::IsValidObject() const
 			const UObject* ResolvedObject = ValueTypeObject.Get();
 			return ResolvedObject && ResolvedObject->IsA(UClass::StaticClass());
 		}
+	case EValueType::UInt32:
+	case EValueType::UInt64:
+		return false;
 	}
 }
 
@@ -373,6 +376,10 @@ size_t FAnimNextParamType::GetValueTypeSize() const
 		return sizeof(TSubclassOf<UObject>);
 	case EValueType::SoftClass:
 		return sizeof(TSoftClassPtr<UObject>);
+	case EValueType::UInt32:
+		return sizeof(uint32);
+	case EValueType::UInt64:
+		return sizeof(uint64);
 	default:
 		break;
 	}
@@ -442,6 +449,10 @@ size_t FAnimNextParamType::GetValueTypeAlignment() const
 		return alignof(TSubclassOf<UObject>);
 	case EValueType::SoftClass:
 		return alignof(TSoftClassPtr<UObject>);
+	case EValueType::UInt32:
+		return alignof(uint32);
+	case EValueType::UInt64:
+		return alignof(uint64);
 	default:
 		break;
 	}
@@ -556,6 +567,12 @@ void FAnimNextParamType::ToString(FStringBuilderBase& InStringBuilder) const
 				InStringBuilder.Append(TEXT("Error: TSoftClassPtr of Unknown Class"));
 			}
 			break;
+		case EValueType::UInt32:
+			InStringBuilder.Append(TEXT("uint32"));
+			break;
+		case EValueType::UInt64:
+			InStringBuilder.Append(TEXT("uint64"));
+			break;
 		default:
 			InStringBuilder.Append(TEXT("Error: Unknown value type"));
 			break;
diff --git a/Engine/Plugins/Experimental/AnimNext/Source/AnimNext/Public/DataInterface/AnimNextDataInterfaceHost.h b/Engine/Plugins/Experimental/AnimNext/Source/AnimNext/Public/DataInterface/AnimNextDataInterfaceHost.h
new file mode 100644
--- /dev/null
+++ b/Engine/Plugins/Experimental/AnimNext/Source/AnimNext/Public/DataInterface/AnimNextDataInterfaceHost.h
@@ -0,0 +1,41 @@
+﻿// Copyright Epic Games, Inc. All Rights Reserved.
+
+#pragma once
+
+#include "CoreMinimal.h"
+#include "DataInterface/AnimNextDataInterface.h"
+
+struct FAnimNextGraphInstance;
+
+namespace UE::AnimNext
+{
+
+// Adapter allowing external systems to implement a data interface
+// Used at instantiation time (and in editor, re-binding after compilation time) to bind directly to a host's memory for an interfaces variables
+class ANIMNEXT_API IDataInterfaceHost
+{
+public:
+	virtual ~IDataInterfaceHost() = default;
+
+protected:
+	// Safely get the name of the data interface that this host provides
+	FName GetDataInterfaceName() const
+	{
+		const UAnimNextDataInterface* DataInterface = GetDataInterface();
+		return DataInterface != nullptr ? DataInterface->GetFName() : NAME_None;
+	}
+
+private:
+	// Get the data interface asset that this host represents
+	virtual const UAnimNextDataInterface* GetDataInterface() const = 0;
+
+	// Get the memory for the supplied variable, at the specified index
+	// @param    InVariableIndex    The index into the data interface of the variable
+	// @param    InVariableName     The name of the variable
+	// @param    InVariableProperty The property of the variable
+	virtual uint8* GetMemoryForVariable(int32 InVariableIndex, FName InVariableName, const FProperty* InVariableProperty) const = 0;
+
+	friend struct ::FAnimNextGraphInstance;
+};
+
+}
\ No newline at end of file
diff --git a/Engine/Plugins/Experimental/AnimNext/Source/AnimNext/Public/DataInterface/AnimNextDataInterfaceInstance.h b/Engine/Plugins/Experimental/AnimNext/Source/AnimNext/Public/DataInterface/AnimNextDataInterfaceInstance.h
new file mode 100644
--- /dev/null
+++ b/Engine/Plugins/Experimental/AnimNext/Source/AnimNext/Public/DataInterface/AnimNextDataInterfaceInstance.h
@@ -0,0 +1,77 @@
+﻿// Copyright Epic Games, Inc. All Rights Reserved.
+
+#pragma once
+
+#include "CoreMinimal.h"
+#include "DataInterface/AnimNextDataInterface.h"
+#include "Param/ParamType.h"
+#include "RigVMCore/RigVMExecuteContext.h"
+#include "StructUtils/PropertyBag.h"
+#include "AnimNextDataInterfaceInstance.generated.h"
+
+class UAnimNextDataInterface;
+
+// Base struct for data interface-derived instances
+USTRUCT()
+struct ANIMNEXT_API FAnimNextDataInterfaceInstance
+{
+	GENERATED_BODY()
+
+	// Get the data interface asset that this instance represents
+	const UAnimNextDataInterface* GetDataInterface() const
+	{
+		return DataInterface;
+	}
+
+	// Safely get the name of the data interface that this host provides
+	FName GetDataInterfaceName() const
+	{
+		return DataInterface ? DataInterface->GetFName() : NAME_None;
+	}
+
+	// Get the property bag that holds external variables for this instance
+	const FInstancedPropertyBag& GetVariables() const
+	{
+		return Variables;
+	}
+
+	// Get the RigVM extended execute context
+	FRigVMExtendedExecuteContext& GetExtendedExecuteContext()
+	{
+		return ExtendedExecuteContext;
+	}
+
+	// Helper function used for bindings
+	// Get the memory for the supplied variable, at the specified index
+	// @param    InVariableIndex    The index into the data interface of the variable
+	// @param    InVariableName     The name of the variable
+	// @param    InVariableProperty The property of the variable
+	uint8* GetMemoryForVariable(int32 InVariableIndex, FName InVariableName, const FProperty* InVariableProperty) const;
+
+	// Get a variable's value given its name.
+	// @param	InVariableName		The name of the variable to get he value of
+	// @param	OutResult			Result that will be filled if no errors occur
+	// @return nullptr if the variable is not present
+	template<typename ValueType>
+	EPropertyBagResult GetVariable(FName InVariableName, ValueType& OutResult) const
+	{
+		return GetVariableInternal(InVariableName, FAnimNextParamType::GetType<ValueType>(), TArrayView<uint8>(reinterpret_cast<uint8*>(&OutResult), sizeof(ValueType)));
+	}
+
+private:
+	// Helper function for GetVariable
+	EPropertyBagResult GetVariableInternal(FName InVariableName, const FAnimNextParamType& InType, TArrayView<uint8> OutResult) const;
+
+protected:
+	// Hard reference to the asset used to create this instance to ensure we can release it safely
+	UPROPERTY(Transient)
+	TObjectPtr<const UAnimNextDataInterface> DataInterface;
+
+	// User variables used to operate the graph
+	UPROPERTY(Transient)
+	FInstancedPropertyBag Variables;
+
+	// Extended execute context instance for this graph instance, we own it
+	UPROPERTY(Transient)
+	FRigVMExtendedExecuteContext ExtendedExecuteContext;
+};
diff --git a/Engine/Plugins/Experimental/AnimNext/Source/AnimNext/Public/DataInterface/DataInterfaceStructAdapter.h b/Engine/Plugins/Experimental/AnimNext/Source/AnimNext/Public/DataInterface/DataInterfaceStructAdapter.h
new file mode 100644
--- /dev/null
+++ b/Engine/Plugins/Experimental/AnimNext/Source/AnimNext/Public/DataInterface/DataInterfaceStructAdapter.h
@@ -0,0 +1,32 @@
+﻿// Copyright Epic Games, Inc. All Rights Reserved.
+
+#pragma once
+
+#include "CoreMinimal.h"
+#include "DataInterface/AnimNextDataInterfaceHost.h"
+
+namespace UE::AnimNext
+{
+
+// Allows a struct to host a data interface's variables
+struct ANIMNEXT_API FDataInterfaceStructAdapter : public IDataInterfaceHost
+{
+public:
+	explicit FDataInterfaceStructAdapter(const UAnimNextDataInterface* InDataInterface, FStructView InStructView)
+		: DataInterface(InDataInterface)
+		, StructView(InStructView)
+	{}
+
+private:
+	// IDataInterfaceHost interface
+	virtual const UAnimNextDataInterface* GetDataInterface() const override;
+	virtual uint8* GetMemoryForVariable(int32 InVariableIndex, FName InVariableName, const FProperty* InVariableProperty) const override;
+
+	// The data interface we wrap
+	const UAnimNextDataInterface* DataInterface;
+
+	// The struct we host
+	FStructView StructView;
+};
+
+}
\ No newline at end of file
diff --git a/Engine/Plugins/Experimental/AnimNext/Source/AnimNext/Public/Graph/AnimNextGraphInstance.h b/Engine/Plugins/Experimental/AnimNext/Source/AnimNext/Public/Graph/AnimNextGraphInstance.h
--- a/Engine/Plugins/Experimental/AnimNext/Source/AnimNext/Public/Graph/AnimNextGraphInstance.h
+++ b/Engine/Plugins/Experimental/AnimNext/Source/AnimNext/Public/Graph/AnimNextGraphInstance.h
@@ -3,6 +3,7 @@
 #pragma once
 
 #include "CoreMinimal.h"
+#include "DataInterface/AnimNextDataInterfaceInstance.h"
 #include "HAL/CriticalSection.h"
 #include "TraitCore/TraitEvent.h"
 #include "TraitCore/TraitEventList.h"
@@ -20,10 +21,11 @@ class UAnimNextModule;
 struct FAnimNextModuleInstance;
 class FRigVMTraitScope;
 struct FRigVMExtendedExecuteContext;
+struct FAnimNextExecuteContext;
 
 namespace UE::AnimNext
 {
-	class IParameterSource;
+	class IDataInterfaceHost;
 	struct FExecutionContext;
 	struct FGraphInstanceComponent;
 	struct FLatentPropertyHandle;
@@ -36,7 +38,7 @@ using GraphInstanceComponentMapType = TMap<FName, TSharedPtr<UE::AnimNext::FGrap
 // This struct uses UE reflection because we wish for the GC to keep the graph
 // alive while we own a reference to it. It is not intended to be serialized on disk with a live instance.
 USTRUCT()
-struct ANIMNEXT_API FAnimNextGraphInstance
+struct ANIMNEXT_API FAnimNextGraphInstance : public FAnimNextDataInterfaceInstance
 {
 	GENERATED_BODY()
 
@@ -107,8 +109,8 @@ struct ANIMNEXT_API FAnimNextGraphInstance
 	// Called each time the graph updates
 	void Update();
 
-	// Get the extended execute context that we own
-	FRigVMExtendedExecuteContext& GetExtendedExecuteContext();
+	// Get the hosting instance, if any, that owns us
+	FAnimNextDataInterfaceInstance* GetHost() const;
 
 private:
 	// Returns a pointer to the specified component, or nullptr if not found
@@ -131,21 +133,17 @@ private:
 
 	// Hook into module compilation
 	void OnModuleCompiled(UAnimNextModule* InModule);
-
-	// Rebind any public variables when this instance is recompiled
-	void RebindPublicVariables();
 #endif
 
+	template<typename HostType>
+	bool BindToHostHelper(const HostType& InHost, bool bInAutoBind);
+
 	// Bind the variables in the supplied traits in scope to their respective public variables, so they point at host memory
-	void BindPublicVariables(TConstArrayView<FRigVMTraitScope> InTraitScopes);
+	void BindPublicVariables(TConstArrayView<UE::AnimNext::IDataInterfaceHost*> InHosts);
 
 	// Unbind any public variables that were pointing at host memory and re-point them at the internal defaults
 	void UnbindPublicVariables();
 
-	// Hard reference to the graph used to create this instance to ensure we can release it safely
-	UPROPERTY()
-	TObjectPtr<const UAnimNextAnimationGraph> AnimationGraph;
-
 	// The entry point in Graph that this instance corresponds to 
 	FName EntryPoint;
 
@@ -161,9 +159,6 @@ private:
 	// The root graph instance that owns us and the components
 	FAnimNextGraphInstance* RootGraphInstance = nullptr;
 
-	// User variables used to operate the graph
-	FInstancedPropertyBag Variables;
-
 #if WITH_EDITORONLY_DATA
 	struct FCachedVariableBinding
 	{
@@ -175,10 +170,6 @@ private:
 	TArray<FCachedVariableBinding> CachedVariableBindings;
 #endif
 
-	// Extended execute context instance for this graph instance, we own it
-	UPROPERTY()
-	FRigVMExtendedExecuteContext ExtendedExecuteContext;
-
 	// Graph instance components that persist from update to update
 	GraphInstanceComponentMapType Components;
 
diff --git a/Engine/Plugins/Experimental/AnimNext/Source/AnimNext/Public/Graph/AnimNextGraphInstancePtr.h b/Engine/Plugins/Experimental/AnimNext/Source/AnimNext/Public/Graph/AnimNextGraphInstancePtr.h
--- a/Engine/Plugins/Experimental/AnimNext/Source/AnimNext/Public/Graph/AnimNextGraphInstancePtr.h
+++ b/Engine/Plugins/Experimental/AnimNext/Source/AnimNext/Public/Graph/AnimNextGraphInstancePtr.h
@@ -16,11 +16,13 @@ struct FRigUnit_AnimNextRunAnimationGraph;
 class UAnimNextAnimationGraph;
 class FRigVMTraitScope;
 struct FRigVMExtendedExecuteContext;
+struct FAnimNextExecuteContext;
 
 namespace UE::AnimNext
 {
 	struct FExecutionContext;
 	struct FGraphInstanceComponent;
+	class IDataInterfaceHost;
 }
 
 using GraphInstanceComponentMapType = TMap<FName, TSharedPtr<UE::AnimNext::FGraphInstanceComponent>>;
@@ -90,6 +92,9 @@ struct ANIMNEXT_API FAnimNextGraphInstancePtr
 	// Whether public variables require a binding
 	bool RequiresPublicVariableBinding() const;
 
+	// Bind the variables in the supplied traits in scope to their respective public variables
+	void BindPublicVariables(TConstArrayView<UE::AnimNext::IDataInterfaceHost*> InHosts) const;
+
 private:
 	// Returns a pointer to the specified component, or nullptr if not found
 	UE::AnimNext::FGraphInstanceComponent* TryGetComponent(int32 ComponentNameHash, FName ComponentName) const;
@@ -97,9 +102,6 @@ private:
 	// Adds the specified component and returns a reference to it
 	UE::AnimNext::FGraphInstanceComponent& AddComponent(int32 ComponentNameHash, FName ComponentName, TSharedPtr<UE::AnimNext::FGraphInstanceComponent>&& Component);
 
-	// Bind the variables in the supplied traits in scope to their respective public variables
-	void BindPublicVariables(TConstArrayView<FRigVMTraitScope> InTraitScopes) const;
-
 	// Indirection to hide implementation details and to fix the graph instance into a single memory location
 	TSharedPtr<FAnimNextGraphInstance> Impl;
 
diff --git a/Engine/Plugins/Experimental/AnimNext/Source/AnimNext/Public/Param/ParamType.h b/Engine/Plugins/Experimental/AnimNext/Source/AnimNext/Public/Param/ParamType.h
--- a/Engine/Plugins/Experimental/AnimNext/Source/AnimNext/Public/Param/ParamType.h
+++ b/Engine/Plugins/Experimental/AnimNext/Source/AnimNext/Public/Param/ParamType.h
@@ -97,7 +97,7 @@ private:
 		{
 			ParameterType.ValueType = EValueType::Byte;
 		}
-		else if constexpr (std::is_same_v<ParamType, int32>)
+		else if constexpr (std::is_same_v<ParamType, int32> || std::is_same_v<ParamType, int>)
 		{
 			ParameterType.ValueType = EValueType::Int32;
 		}
@@ -315,7 +315,7 @@ public:
 	{
 		const bool bHasValidValueType = ValueType != EValueType::None;
 		const bool bHasValidContainerType = (ContainerType == EContainerType::None) || (bHasValidValueType && ContainerType != EContainerType::None);
-		const bool bHasValidObjectType = (ValueType < EValueType::Enum) || (ValueType >= EValueType::Enum && ValueType <= EValueType::SoftClass && IsValidObject());
+		const bool bHasValidObjectType = ((ValueType < EValueType::Enum) || ValueType == EValueType::UInt32 || ValueType == EValueType::UInt64) || (ValueType >= EValueType::Enum && ValueType <= EValueType::SoftClass && IsValidObject());
 		return bHasValidValueType && bHasValidContainerType && bHasValidObjectType;
 	}
 
diff --git a/Engine/Plugins/Experimental/AnimNext/Source/AnimNextEditor/Internal/DataInterface/AnimNextDataInterfaceFactory.h b/Engine/Plugins/Experimental/AnimNext/Source/AnimNextEditor/Internal/DataInterface/AnimNextDataInterfaceFactory.h
new file mode 100644
--- /dev/null
+++ b/Engine/Plugins/Experimental/AnimNext/Source/AnimNextEditor/Internal/DataInterface/AnimNextDataInterfaceFactory.h
@@ -0,0 +1,22 @@
+﻿// Copyright Epic Games, Inc. All Rights Reserved.
+
+#pragma once
+
+#include "CoreMinimal.h"
+#include "Factories/Factory.h"
+#include "AnimNextDataInterfaceFactory.generated.h"
+
+UCLASS(MinimalAPI)
+class UAnimNextDataInterfaceFactory : public UFactory
+{
+	GENERATED_BODY()
+
+	UAnimNextDataInterfaceFactory();
+	
+	virtual bool ConfigureProperties() override;
+	virtual UObject* FactoryCreateNew(UClass* Class, UObject* InParent, FName Name, EObjectFlags Flags, UObject* Context, FFeedbackContext* Warn, FName CallingContext) override;
+	virtual UObject* FactoryCreateNew(UClass* Class, UObject* InParent, FName Name, EObjectFlags Flags, UObject* Context, FFeedbackContext* Warn) override
+	{
+		return FactoryCreateNew(Class, InParent, Name, Flags, Context, Warn, NAME_None);
+	}
+};
\ No newline at end of file
diff --git a/Engine/Plugins/Experimental/AnimNext/Source/AnimNextEditor/Private/AnimNextEditorModule.cpp b/Engine/Plugins/Experimental/AnimNext/Source/AnimNextEditor/Private/AnimNextEditorModule.cpp
--- a/Engine/Plugins/Experimental/AnimNext/Source/AnimNextEditor/Private/AnimNextEditorModule.cpp
+++ b/Engine/Plugins/Experimental/AnimNext/Source/AnimNextEditor/Private/AnimNextEditorModule.cpp
@@ -46,8 +46,11 @@
 #include "Graph/AnimNextAnimationGraph.h"
 #include "Graph/AnimNextAnimationGraph_EditorData.h"
 #include "Common/AnimNextAssetItemDetails.h"
+#include "DataInterface/AnimNextDataInterface.h"
 #include "Module/RigUnit_AnimNextModuleEvents.h"
 #include "Variables/SVariablesView.h"
+#include "Variables/VariableOverrideCommands.h"
+#include "Variables/VariableProxyCustomization.h"
 #include "Workspace/AnimNextWorkspaceSchema.h"
 
 #define LOCTEXT_NAMESPACE "AnimNextEditorModule"
@@ -57,6 +60,8 @@ namespace UE::AnimNext::Editor
 
 void FAnimNextEditorModule::StartupModule()
 {
+	FVariableOverrideCommands::Register();
+
 	// Register settings for user editing
 	ISettingsModule& SettingsModule = FModuleManager::Get().LoadModuleChecked<ISettingsModule>("Settings");
 	SettingsModule.RegisterSettings("Editor", "General", "AnimNext",
@@ -78,6 +83,9 @@ void FAnimNextEditorModule::StartupModule()
 	PropertyModule.RegisterCustomClassLayout("AnimNextVariableEntry", 
 		FOnGetDetailCustomizationInstance::CreateLambda([] { return MakeShared<FVariableCustomization>(); }));
 
+	PropertyModule.RegisterCustomClassLayout("AnimNextVariableEntryProxy", 
+		FOnGetDetailCustomizationInstance::CreateLambda([] { return MakeShared<FVariableProxyCustomization>(); }));
+	
 	AnimNextGraphPanelNodeFactory = MakeShared<FAnimNextGraphPanelNodeFactory>();
 	FEdGraphUtilities::RegisterVisualNodeFactory(AnimNextGraphPanelNodeFactory);
 
@@ -147,19 +155,17 @@ void FAnimNextEditorModule::StartupModule()
 			.OnClicked_Lambda([EditorData, Asset]()
 			{
 				TSharedRef<SAddVariablesDialog> AddVariablesDialog =
-					SNew(SAddVariablesDialog, FAssetData(EditorData));
+					SNew(SAddVariablesDialog, TArray<UAnimNextRigVMAssetEditorData*>({ EditorData }));
 
 				TArray<FVariableToAdd> VariablesToAdd;
-				if(AddVariablesDialog->ShowModal(VariablesToAdd))
+				TArray<FDataInterfaceToAdd> DataInterfacesToAdd;
+				if(AddVariablesDialog->ShowModal(VariablesToAdd, DataInterfacesToAdd))
 				{
-					if(VariablesToAdd.Num() > 0)
+					FScopedTransaction Transaction(LOCTEXT("AddVariables", "Add variable(s)"));
+					for (const FVariableToAdd& VariableToAdd : VariablesToAdd)
 					{
-						FScopedTransaction Transaction(LOCTEXT("AddVariable", "Add variable"));
-						for (const FVariableToAdd& VariableToAdd : VariablesToAdd)
-						{
-							check(EditorData->FindEntry(VariableToAdd.Name) == nullptr);
-							EditorData->AddVariable(VariableToAdd.Name, VariableToAdd.Type);
-						}
+						check(EditorData->FindEntry(VariableToAdd.Name) == nullptr);
+						EditorData->AddVariable(VariableToAdd.Name, VariableToAdd.Type);
 					}
 				}
 				return FReply::Handled();
@@ -236,7 +242,8 @@ void FAnimNextEditorModule::StartupModule()
 	SupportedAssetClasses.Append(
 		{
 			UAnimNextAnimationGraph::StaticClass()->GetClassPathName(),
-			UAnimNextModule::StaticClass()->GetClassPathName()
+			UAnimNextModule::StaticClass()->GetClassPathName(),
+			UAnimNextDataInterface::StaticClass()->GetClassPathName()
 		});
 }
 
@@ -248,6 +255,7 @@ void FAnimNextEditorModule::ShutdownModule()
 		PropertyModule.UnregisterCustomPropertyTypeLayout("AnimNextParamType");
 		PropertyModule.UnregisterCustomPropertyTypeLayout("AnimNextVariableBinding");
 		PropertyModule.UnregisterCustomClassLayout("AnimNextVariableEntry");
+		PropertyModule.UnregisterCustomClassLayout("AnimNextVariableEntryProxy");
 	}
 
 	FEdGraphUtilities::UnregisterVisualNodeFactory(AnimNextGraphPanelNodeFactory);
diff --git a/Engine/Plugins/Experimental/AnimNext/Source/AnimNextEditor/Private/Common/AnimNextAssetItemDetails.cpp b/Engine/Plugins/Experimental/AnimNext/Source/AnimNextEditor/Private/Common/AnimNextAssetItemDetails.cpp
--- a/Engine/Plugins/Experimental/AnimNext/Source/AnimNextEditor/Private/Common/AnimNextAssetItemDetails.cpp
+++ b/Engine/Plugins/Experimental/AnimNext/Source/AnimNextEditor/Private/Common/AnimNextAssetItemDetails.cpp
@@ -14,6 +14,7 @@
 #include "ScopedTransaction.h"
 #include "RigVMModel/RigVMClient.h"
 #include "ToolMenus.h"
+#include "Entries/AnimNextDataInterfaceEntry.h"
 #include "Framework/Commands/GenericCommands.h"
 #include "Module/AnimNextModule_EditorData.h"
 #include "Module/RigUnit_AnimNextModuleEvents.h"
@@ -154,19 +155,41 @@ void FAnimNextAssetItemDetails::RegisterToolMenuExtensions()
 							if(SubEntryClass == UAnimNextVariableEntry::StaticClass())
 							{
 								TSharedRef<SAddVariablesDialog> AddVariablesDialog =
-									SNew(SAddVariablesDialog, FAssetData(OutlinerData.Asset));
+									SNew(SAddVariablesDialog, TArray<UAnimNextRigVMAssetEditorData*>({ UncookedOnly::FUtils::GetEditorData<UAnimNextRigVMAssetEditorData>(OutlinerData.Asset.Get()) }));
 
 								TArray<FVariableToAdd> VariablesToAdd;
-								if(AddVariablesDialog->ShowModal(VariablesToAdd))
+								TArray<FDataInterfaceToAdd> DataInterfacesToAdd;
+								if(AddVariablesDialog->ShowModal(VariablesToAdd, DataInterfacesToAdd))
 								{
-									if(VariablesToAdd.Num() > 0)
+									FScopedTransaction Transaction(LOCTEXT("AddVariables", "Add Variable(s)"));
+									for (const FVariableToAdd& VariableToAdd : VariablesToAdd)
 									{
-										FScopedTransaction Transaction(LOCTEXT("AddVariable", "Add Variable"));
-										for (const FVariableToAdd& VariableToAdd : VariablesToAdd)
-										{
-											check(EditorData->FindEntry(VariableToAdd.Name) == nullptr);
-											EditorData->AddVariable(VariableToAdd.Name, VariableToAdd.Type);
-										}
+										EditorData->AddVariable(VariableToAdd.Name, VariableToAdd.Type);
+									}
+									for (const FDataInterfaceToAdd& DataInterfaceToAdd : DataInterfacesToAdd)
+									{
+										EditorData->AddDataInterface(DataInterfaceToAdd.DataInterface);
+									}
+								}
+							}
+							else if(SubEntryClass == UAnimNextDataInterfaceEntry::StaticClass())
+							{
+								TSharedRef<SAddVariablesDialog> AddVariablesDialog =
+									SNew(SAddVariablesDialog, TArray<UAnimNextRigVMAssetEditorData*>({ UncookedOnly::FUtils::GetEditorData<UAnimNextRigVMAssetEditorData>(OutlinerData.Asset.Get()) }))
+									.ShouldAddInitialVariable(false);
+
+								TArray<FVariableToAdd> VariablesToAdd;
+								TArray<FDataInterfaceToAdd> DataInterfacesToAdd;
+								if(AddVariablesDialog->ShowModal(VariablesToAdd, DataInterfacesToAdd))
+								{
+									FScopedTransaction Transaction(LOCTEXT("AddVariables", "Add Variable(s)"));
+									for (const FVariableToAdd& VariableToAdd : VariablesToAdd)
+									{
+										EditorData->AddVariable(VariableToAdd.Name, VariableToAdd.Type);
+									}
+									for (const FDataInterfaceToAdd& DataInterfaceToAdd : DataInterfacesToAdd)
+									{
+										EditorData->AddDataInterface(DataInterfaceToAdd.DataInterface);
 									}
 								}
 							}
diff --git a/Engine/Plugins/Experimental/AnimNext/Source/AnimNextEditor/Private/Common/AnimNextEntryAssetDefinitions.cpp b/Engine/Plugins/Experimental/AnimNext/Source/AnimNextEditor/Private/Common/AnimNextEntryAssetDefinitions.cpp
--- a/Engine/Plugins/Experimental/AnimNext/Source/AnimNextEditor/Private/Common/AnimNextEntryAssetDefinitions.cpp
+++ b/Engine/Plugins/Experimental/AnimNext/Source/AnimNextEditor/Private/Common/AnimNextEntryAssetDefinitions.cpp
@@ -18,8 +18,14 @@ FText UAssetDefinition_AnimNextAnimationGraphEntry::GetObjectDisplayNameText(UOb
 
 FText UAssetDefinition_AnimNextEventGraphEntry::GetObjectDisplayNameText(UObject* Object) const
 {
-	UAnimNextRigVMAssetEntry* Parameter = CastChecked<UAnimNextRigVMAssetEntry>(Object);
-	return Parameter->GetDisplayName();
+	UAnimNextRigVMAssetEntry* Variable = CastChecked<UAnimNextRigVMAssetEntry>(Object);
+	return Variable->GetDisplayName();
+}
+
+FText UAssetDefinition_AnimNextDataInterfaceEntry::GetObjectDisplayNameText(UObject* Object) const
+{
+	UAnimNextDataInterfaceEntry* DataInterface = CastChecked<UAnimNextDataInterfaceEntry>(Object);
+	return DataInterface->GetDisplayName();
 }
 
 #undef LOCTEXT_NAMESPACE
\ No newline at end of file
diff --git a/Engine/Plugins/Experimental/AnimNext/Source/AnimNextEditor/Private/Common/AnimNextEntryAssetDefinitions.h b/Engine/Plugins/Experimental/AnimNext/Source/AnimNextEditor/Private/Common/AnimNextEntryAssetDefinitions.h
--- a/Engine/Plugins/Experimental/AnimNext/Source/AnimNextEditor/Private/Common/AnimNextEntryAssetDefinitions.h
+++ b/Engine/Plugins/Experimental/AnimNext/Source/AnimNextEditor/Private/Common/AnimNextEntryAssetDefinitions.h
@@ -6,6 +6,7 @@
 #include "Entries/AnimNextEventGraphEntry.h"
 #include "Entries/AnimNextVariableEntry.h"
 #include "Entries/AnimNextAnimationGraphEntry.h"
+#include "Entries/AnimNextDataInterfaceEntry.h"
 #include "AnimNextEntryAssetDefinitions.generated.h"
 
 #define LOCTEXT_NAMESPACE "AnimNextAssetDefinitions"
@@ -43,10 +44,23 @@ class UAssetDefinition_AnimNextEventGraphEntry : public UAssetDefinitionDefault
 
 public:
 	// UAssetDefinition interface
-	virtual FText GetAssetDisplayName() const override { return LOCTEXT("AnimNextEventGraph", "Event Graph"); }
+	virtual FText GetAssetDisplayName() const override { return LOCTEXT("AnimNextEventGraphEntry", "Event Graph"); }
 	virtual FLinearColor GetAssetColor() const override { return FLinearColor(FColor(128,64,64)); }
 	virtual TSoftClassPtr<UObject> GetAssetClass() const override { return UAnimNextEventGraphEntry::StaticClass(); }
 	virtual FText GetObjectDisplayNameText(UObject* Object) const override;
 };
 
+UCLASS()
+class UAssetDefinition_AnimNextDataInterfaceEntry : public UAssetDefinitionDefault
+{
+	GENERATED_BODY()
+
+public:
+	// UAssetDefinition interface
+	virtual FText GetAssetDisplayName() const override { return LOCTEXT("AnimNextDataInterfaceEntry", "Data Interface"); }
+	virtual FLinearColor GetAssetColor() const override { return FLinearColor(FColor(128,64,64)); }
+	virtual TSoftClassPtr<UObject> GetAssetClass() const override { return UAnimNextDataInterfaceEntry::StaticClass(); }
+	virtual FText GetObjectDisplayNameText(UObject* Object) const override;
+};
+
 #undef LOCTEXT_NAMESPACE
\ No newline at end of file
diff --git a/Engine/Plugins/Experimental/AnimNext/Source/AnimNextEditor/Private/Common/SRigVMAssetView.cpp b/Engine/Plugins/Experimental/AnimNext/Source/AnimNextEditor/Private/Common/SRigVMAssetView.cpp
--- a/Engine/Plugins/Experimental/AnimNext/Source/AnimNextEditor/Private/Common/SRigVMAssetView.cpp
+++ b/Engine/Plugins/Experimental/AnimNext/Source/AnimNextEditor/Private/Common/SRigVMAssetView.cpp
@@ -442,6 +442,7 @@ void SRigVMAssetView::HandleAssetModified(UAnimNextRigVMAssetEditorData* InEdito
 
 	switch (InType)
 	{
+	case EAnimNextEditorDataNotifType::UndoRedo:
 	case EAnimNextEditorDataNotifType::EntryAdded:
 	case EAnimNextEditorDataNotifType::EntryRemoved:
 		RequestRefresh();
@@ -594,7 +595,7 @@ class SRigVMAssetViewRow : public SMultiColumnTableRow<TSharedRef<FRigVMAssetVie
 					AssetEntry->MarkPackageDirty();
 
 					// Ensure that default values get picked up and forwarded to compiler
-					EditorData->BroadcastModified(EAnimNextEditorDataNotifType::PropertyChanged, AssetEntry);
+					EditorData->BroadcastModified(EAnimNextEditorDataNotifType::VariableDefaultValueChanged, AssetEntry);
 				}
 			}
 		}
diff --git a/Engine/Plugins/Experimental/AnimNext/Source/AnimNextEditor/Private/DataInterface/AnimNextDataInterfaceAssetDefinition.cpp b/Engine/Plugins/Experimental/AnimNext/Source/AnimNextEditor/Private/DataInterface/AnimNextDataInterfaceAssetDefinition.cpp
new file mode 100644
--- /dev/null
+++ b/Engine/Plugins/Experimental/AnimNext/Source/AnimNextEditor/Private/DataInterface/AnimNextDataInterfaceAssetDefinition.cpp
@@ -0,0 +1,23 @@
+﻿// Copyright Epic Games, Inc. All Rights Reserved.
+
+#include "AnimNextDataInterfaceAssetDefinition.h"
+#include "ContentBrowserMenuContexts.h"
+#include "IWorkspaceEditorModule.h"
+#include "Workspace/AnimNextWorkspaceFactory.h"
+
+#define LOCTEXT_NAMESPACE "AnimNextAssetDefinitions"
+
+EAssetCommandResult UAssetDefinition_AnimNextDataInterface::OpenAssets(const FAssetOpenArgs& OpenArgs) const
+{
+	using namespace UE::Workspace;
+
+	for (UAnimNextDataInterface* Asset : OpenArgs.LoadObjects<UAnimNextDataInterface>())
+	{
+		IWorkspaceEditorModule& WorkspaceEditorModule = FModuleManager::Get().LoadModuleChecked<IWorkspaceEditorModule>("WorkspaceEditor");
+		WorkspaceEditorModule.OpenWorkspaceForObject(Asset, EOpenWorkspaceMethod::Default, UAnimNextWorkspaceFactory::StaticClass());
+	}
+
+	return EAssetCommandResult::Handled;
+}
+
+#undef LOCTEXT_NAMESPACE
\ No newline at end of file
diff --git a/Engine/Plugins/Experimental/AnimNext/Source/AnimNextEditor/Private/DataInterface/AnimNextDataInterfaceAssetDefinition.h b/Engine/Plugins/Experimental/AnimNext/Source/AnimNextEditor/Private/DataInterface/AnimNextDataInterfaceAssetDefinition.h
new file mode 100644
--- /dev/null
+++ b/Engine/Plugins/Experimental/AnimNext/Source/AnimNextEditor/Private/DataInterface/AnimNextDataInterfaceAssetDefinition.h
@@ -0,0 +1,30 @@
+﻿// Copyright Epic Games, Inc. All Rights Reserved.
+
+#pragma once
+
+#include "DataInterface/AnimNextDataInterface.h"
+#include "AssetDefinitionDefault.h"
+#include "AnimNextDataInterfaceAssetDefinition.generated.h"
+
+#define LOCTEXT_NAMESPACE "AnimNextAssetDefinitions"
+
+UCLASS()
+class UAssetDefinition_AnimNextDataInterface : public UAssetDefinitionDefault
+{
+	GENERATED_BODY()
+
+public:
+	// UAssetDefinition interface
+	virtual FText GetAssetDisplayName() const override { return LOCTEXT("AnimNextDataInterface", "AnimNext Data Interface"); }
+	virtual FLinearColor GetAssetColor() const override { return FLinearColor(FColor(128,64,32)); }
+	virtual TSoftClassPtr<UObject> GetAssetClass() const override { return UAnimNextDataInterface::StaticClass(); }
+	virtual EAssetCommandResult OpenAssets(const FAssetOpenArgs& OpenArgs) const override;
+	virtual TConstArrayView<FAssetCategoryPath> GetAssetCategories() const override
+	{
+		static const auto Categories = { FAssetCategoryPath(EAssetCategoryPaths::Animation, LOCTEXT("AnimNextSubMenu", "AnimNext")) };
+		return Categories;
+	}
+	virtual bool ShouldSaveExternalPackages() const override { return true; }
+};
+
+#undef LOCTEXT_NAMESPACE
\ No newline at end of file
diff --git a/Engine/Plugins/Experimental/AnimNext/Source/AnimNextEditor/Private/DataInterface/AnimNextDataInterfaceFactory.cpp b/Engine/Plugins/Experimental/AnimNext/Source/AnimNextEditor/Private/DataInterface/AnimNextDataInterfaceFactory.cpp
new file mode 100644
--- /dev/null
+++ b/Engine/Plugins/Experimental/AnimNext/Source/AnimNextEditor/Private/DataInterface/AnimNextDataInterfaceFactory.cpp
@@ -0,0 +1,40 @@
+﻿// Copyright Epic Games, Inc. All Rights Reserved.
+
+#include "DataInterface/AnimNextDataInterfaceFactory.h"
+#include "DataInterface/AnimNextDataInterface.h"
+#include "DataInterface/AnimNextDataInterface_EditorData.h"
+#include "Module/RigUnit_AnimNextModuleEvents.h"
+
+UAnimNextDataInterfaceFactory::UAnimNextDataInterfaceFactory()
+{
+	bCreateNew = true;
+	bEditAfterNew = true;
+	SupportedClass = UAnimNextDataInterface::StaticClass();
+}
+
+bool UAnimNextDataInterfaceFactory::ConfigureProperties()
+{
+	return true;
+}
+
+UObject* UAnimNextDataInterfaceFactory::FactoryCreateNew(UClass* Class, UObject* InParent, FName Name, EObjectFlags Flags, UObject* Context, FFeedbackContext* Warn, FName CallingContext)
+{
+	EObjectFlags FlagsToUse = Flags | RF_Public | RF_Standalone | RF_Transactional | RF_LoadCompleted;
+	if(InParent == GetTransientPackage())
+	{
+		FlagsToUse &= ~RF_Standalone;
+	}
+
+	UAnimNextDataInterface* NewDataInterface = NewObject<UAnimNextDataInterface>(InParent, Class, Name, FlagsToUse);
+
+	// Create internal editor data
+	UAnimNextDataInterface_EditorData* EditorData = NewObject<UAnimNextDataInterface_EditorData>(NewDataInterface, TEXT("EditorData"), RF_Transactional);
+	NewDataInterface->EditorData = EditorData;
+	EditorData->Initialize(/*bRecompileVM*/false);
+
+	// Compile the initial skeleton
+	EditorData->RecompileVM();
+	check(!EditorData->bErrorsDuringCompilation);
+
+	return NewDataInterface;
+}
\ No newline at end of file
