AnimNext: Data interfaces

diff --git a/Engine/Plugins/Experimental/AnimNext/Source/AnimNextEditor/Private/EditorUtils.cpp b/Engine/Plugins/Experimental/AnimNext/Source/AnimNextEditor/Private/EditorUtils.cpp
--- a/Engine/Plugins/Experimental/AnimNext/Source/AnimNextEditor/Private/EditorUtils.cpp
+++ b/Engine/Plugins/Experimental/AnimNext/Source/AnimNextEditor/Private/EditorUtils.cpp
@@ -314,51 +314,6 @@ void FUtils::GetFilteredVariableTypeTree(TArray<TSharedPtr<UEdGraphSchema_K2::FP
 	}
 };
 
-
-FName FUtils::GetNewParameterName(FName InBaseName, const FAssetData& InAssetData, TArrayView<FName> InExistingNames)
-{
-	auto NameExists = [&InExistingNames, &InAssetData](FName InName)
-	{
-		for(FName AdditionalName : InExistingNames)
-		{
-			if(AdditionalName == InName)
-			{
-				return true;
-			}
-		}
-
-		if(DoesParameterNameExistInAsset(InName, InAssetData))
-		{
-			return true;
-		}
-		
-		return false;
-	};
-
-	if(!NameExists(InBaseName))
-	{
-		// Early out - name is valid
-		return InBaseName;
-	}
-
-	int32 PostFixIndex = 0;
-	TStringBuilder<128> StringBuilder;
-	while(true)
-	{
-		StringBuilder.Reset();
-		InBaseName.GetDisplayNameEntry()->AppendNameToString(StringBuilder);
-		StringBuilder.Appendf(TEXT("_%d"), PostFixIndex++);
-
-		FName TestName(StringBuilder.ToString()); 
-		if(!NameExists(TestName))
-		{
-			return TestName;
-		}
-	}
-
-	return NAME_None;
-}
-
 bool FUtils::IsValidParameterNameString(FStringView InStringView, FText& OutErrorText)
 {
 	// See if this can be represented as an FName
@@ -404,14 +359,6 @@ bool FUtils::IsValidParameterName(const FName InName, FText& OutErrorText)
 	return true;
 }
 
-bool FUtils::DoesParameterNameExistInAsset(const FName InName, const FAssetData& InAsset)
-{
-	FAnimNextAssetRegistryExports Exports;
-	UncookedOnly::FUtils::GetExportedVariablesForAsset(InAsset, Exports);
-	return Exports.Variables.ContainsByPredicate([InName](const FAnimNextAssetRegistryExportedVariable& Entry) { return Entry.Name == InName; });
-}
-
-
 }
 
 #undef LOCTEXT_NAMESPACE
\ No newline at end of file
diff --git a/Engine/Plugins/Experimental/AnimNext/Source/AnimNextEditor/Private/EditorUtils.h b/Engine/Plugins/Experimental/AnimNext/Source/AnimNextEditor/Private/EditorUtils.h
--- a/Engine/Plugins/Experimental/AnimNext/Source/AnimNextEditor/Private/EditorUtils.h
+++ b/Engine/Plugins/Experimental/AnimNext/Source/AnimNextEditor/Private/EditorUtils.h
@@ -30,13 +30,9 @@ struct FUtils
 
 	static void GetFilteredVariableTypeTree(TArray<TSharedPtr<UEdGraphSchema_K2::FPinTypeTreeInfo>>& TypeTree, ETypeTreeFilter TypeTreeFilter);
 
-	static FName GetNewParameterName(FName InBaseName, const FAssetData& InAssetData, TArrayView<FName> InExistingNames);
-
 	static bool IsValidParameterNameString(FStringView InStringView, FText& OutErrorText);
 
 	static bool IsValidParameterName(const FName InName, FText& OutErrorText);
-
-	static bool DoesParameterNameExistInAsset(const FName InName, const FAssetData& InAsset);
 };
 
 }
diff --git a/Engine/Plugins/Experimental/AnimNext/Source/AnimNextEditor/Private/Variables/AddVariableDialogMenuContext.h b/Engine/Plugins/Experimental/AnimNext/Source/AnimNextEditor/Private/Variables/AddVariableDialogMenuContext.h
--- a/Engine/Plugins/Experimental/AnimNext/Source/AnimNextEditor/Private/Variables/AddVariableDialogMenuContext.h
+++ b/Engine/Plugins/Experimental/AnimNext/Source/AnimNextEditor/Private/Variables/AddVariableDialogMenuContext.h
@@ -2,6 +2,7 @@
 
 #pragma once
 
+#include "SAddVariablesDialog.h"
 #include "AddVariableDialogMenuContext.generated.h"
 
 namespace UE::AnimNext::Editor
@@ -21,5 +22,5 @@ class UAddVariableDialogMenuContext : public UObject
 	TWeakPtr<UE::AnimNext::Editor::SAddVariablesDialog> AddVariablesDialog;
 
 	// The entry that we are editing
-	TWeakPtr<UE::AnimNext::Editor::FVariableToAdd> Entry;
+	TWeakPtr<UE::AnimNext::Editor::SAddVariablesDialog::FEntry> Entry;
 };
\ No newline at end of file
diff --git a/Engine/Plugins/Experimental/AnimNext/Source/AnimNextEditor/Private/Variables/AnimNextVariableEntryProxy.h b/Engine/Plugins/Experimental/AnimNext/Source/AnimNextEditor/Private/Variables/AnimNextVariableEntryProxy.h
new file mode 100644
--- /dev/null
+++ b/Engine/Plugins/Experimental/AnimNext/Source/AnimNextEditor/Private/Variables/AnimNextVariableEntryProxy.h
@@ -0,0 +1,31 @@
+﻿// Copyright Epic Games, Inc. All Rights Reserved.
+
+#pragma once
+
+#include "CoreMinimal.h"
+#include "AnimNextVariableEntryProxy.generated.h"
+
+class UAnimNextVariableEntry;
+class UAnimNextDataInterfaceEntry;
+
+namespace UE::AnimNext::Editor
+{
+	class FVariablesOutlinerMode;
+	class FVariableProxyCustomization;
+}
+
+// Editor-only proxy object used for editing variable entries in the context of per-interface overrides
+UCLASS()
+class UAnimNextVariableEntryProxy : public UObject
+{
+	GENERATED_BODY()
+
+	friend class UE::AnimNext::Editor::FVariableProxyCustomization;
+	friend class UE::AnimNext::Editor::FVariablesOutlinerMode;
+
+	UPROPERTY()
+	TObjectPtr<UAnimNextVariableEntry> VariableEntry;
+
+	UPROPERTY()
+	TObjectPtr<UAnimNextDataInterfaceEntry> DataInterfaceEntry;
+};
\ No newline at end of file
diff --git a/Engine/Plugins/Experimental/AnimNext/Source/AnimNextEditor/Private/Variables/Outliner/VariablesOutlinerColumns.cpp b/Engine/Plugins/Experimental/AnimNext/Source/AnimNextEditor/Private/Variables/Outliner/VariablesOutlinerColumns.cpp
--- a/Engine/Plugins/Experimental/AnimNext/Source/AnimNextEditor/Private/Variables/Outliner/VariablesOutlinerColumns.cpp
+++ b/Engine/Plugins/Experimental/AnimNext/Source/AnimNextEditor/Private/Variables/Outliner/VariablesOutlinerColumns.cpp
@@ -15,11 +15,14 @@
 #include "SPinTypeSelector.h"
 #include "UncookedOnlyUtils.h"
 #include "VariablesOutlinerEntryItem.h"
+#include "Entries/AnimNextDataInterfaceEntry.h"
 #include "Entries/AnimNextRigVMAssetEntry.h"
+#include "Entries/AnimNextVariableEntry.h"
 #include "Widgets/Images/SImage.h"
 #include "UObject/Package.h"
 #include "Widgets/Input/SButton.h"
 #include "Widgets/Text/SInlineEditableTextBlock.h"
+#include "Variables/SVariableOverride.h"
 
 #define LOCTEXT_NAMESPACE "VariablesOutlinerColumns"
 
@@ -68,6 +71,7 @@ const TSharedRef<SWidget> FVariablesOutlinerTypeColumn::ConstructRowWidget(FScen
 		SNew(SBox)
 		.HAlign(HAlign_Left)
 		.VAlign(VAlign_Center)
+		.IsEnabled(!TreeItem->WeakDataInterfaceEntry.IsValid())
 		[
 			SNew(SPinTypeSelector, FGetPinTypeTree::CreateStatic(&Editor::FUtils::GetFilteredVariableTypeTree))
 			.TargetPinType_Lambda([WeakEntry = TreeItem->WeakEntry]()
@@ -113,12 +117,12 @@ SHeaderRow::FColumn::FArguments FVariablesOutlinerValueColumn::ConstructHeaderRo
 		.FillWidth(1.0f)
 		.HAlignHeader(HAlign_Left)
 		.VAlignHeader(VAlign_Center)
-		.HAlignCell(HAlign_Left)
+		.HAlignCell(HAlign_Fill)
 		.VAlignCell(VAlign_Center)
 		[
 			SNew(SBox) 
 			.VAlign(VAlign_Center)
-			.HAlign(HAlign_Left)
+			.HAlign(HAlign_Fill)
 			[
 				SNew(STextBlock)
 				.Text(LOCTEXT("ValueLabel", "Value"))
@@ -136,10 +140,95 @@ class SVariablesOutlinerValue : public SCompoundWidget, public FNotifyHook
 	{
 		WeakTreeItem = StaticCastSharedRef<FVariablesOutlinerEntryItem>(InTreeItem.AsShared());
 
+		if(UAnimNextVariableEntry* VariableEntry = Cast<UAnimNextVariableEntry>(InTreeItem.WeakEntry.Get()))
+		{
+			if(UAnimNextRigVMAssetEditorData* EditorData = VariableEntry->GetTypedOuter<UAnimNextRigVMAssetEditorData>())
+			{
+				EditorData->ModifiedDelegate.AddSP(this, &SVariablesOutlinerValue::HandleModified);
+			}
+		}
+		if(UAnimNextDataInterfaceEntry* DataInterfaceEntry = Cast<UAnimNextDataInterfaceEntry>(InTreeItem.WeakDataInterfaceEntry.Get()))
+		{
+			if(UAnimNextRigVMAssetEditorData* EditorData = DataInterfaceEntry->GetTypedOuter<UAnimNextRigVMAssetEditorData>())
+			{
+				EditorData->ModifiedDelegate.AddSP(this, &SVariablesOutlinerValue::HandleModified);
+			}
+		}
+
+		ChildSlot
+		[
+			SAssignNew(WidgetContainer, SBox)
+		];
+
+		BuildValueWidget();
+	}
+
+	void BuildValueWidget()
+	{
+		TSharedPtr<FVariablesOutlinerEntryItem> TreeItem = WeakTreeItem.Pin();
+		if(!TreeItem.IsValid())
+		{
+			return;
+		}
+
+		UAnimNextVariableEntry* VariableEntry = TreeItem->WeakEntry.Get();
+		if(VariableEntry == nullptr)
+		{
+			return;
+		}
+
 		TSharedPtr<SWidget> ValueWidget = SNullWidget::NullWidget;
-		IAnimNextRigVMVariableInterface* Variable = CastChecked<IAnimNextRigVMVariableInterface>(InTreeItem.WeakEntry.Get());
-		FInstancedPropertyBag& PropertyBag = Variable->GetMutablePropertyBag();
-		const FPropertyBagPropertyDesc* PropertyDesc = PropertyBag.FindPropertyDescByName(IAnimNextRigVMVariableInterface::ValueName);
+		const FName VariableName = VariableEntry->GetEntryName();
+		FInstancedPropertyBag* PropertyBag = nullptr;
+		FName ValueName;
+		bool bCanOverride = false;
+		UAnimNextDataInterfaceEntry* DataInterfaceEntry = TreeItem->WeakDataInterfaceEntry.Get();
+		if(DataInterfaceEntry != nullptr)
+		{
+			// Update for override status
+			bCanOverride = true;
+			OverrideStatus = DataInterfaceEntry->FindValueOverridePropertyBagRecursive(VariableName, PropertyBag);
+			switch(OverrideStatus)
+			{
+			case EAnimNextDataInterfaceValueOverrideStatus::NotOverridden:
+				{
+					// Use the internal property bag (copying from the source) and the default name
+					InternalPropertyBag = VariableEntry->GetPropertyBag();
+					PropertyBag = &InternalPropertyBag;
+					ValueName = IAnimNextRigVMVariableInterface::ValueName;
+					break;
+				}
+			case EAnimNextDataInterfaceValueOverrideStatus::OverriddenInThisAsset:
+				{
+					// Use the found property bag and the variable name
+					check(PropertyBag);
+					ValueName = VariableName;
+					break;
+				}
+			case EAnimNextDataInterfaceValueOverrideStatus::OverriddenInParentAsset:
+				{
+					// Use the internal property bag (copying just the value from the overriding asset's property bag) and the default name
+					check(PropertyBag);
+					const FPropertyBagPropertyDesc* Desc = PropertyBag->FindPropertyDescByName(VariableName);
+					check(Desc);
+					const uint8* DataPtr = Desc->CachedProperty->ContainerPtrToValuePtr<uint8>(PropertyBag->GetValue().GetMemory());
+					TConstArrayView<uint8> Value(DataPtr, Desc->CachedProperty->GetElementSize());
+					InternalPropertyBag.ReplaceAllPropertiesAndValues(TConstArrayView<FPropertyBagPropertyDesc>(Desc, 1), TConstArrayView<TConstArrayView<uint8>>(&Value, 1));
+					PropertyBag = &InternalPropertyBag;
+					ValueName = VariableName;
+					break;
+				}
+			}
+		}
+		else
+		{
+			// No data interface present, use the variable's internal property bag
+			bCanOverride = false;
+			PropertyBag = &VariableEntry->GetMutablePropertyBag();
+			ValueName = IAnimNextRigVMVariableInterface::ValueName;
+		}
+
+		const FPropertyBagPropertyDesc* PropertyDesc = PropertyBag ? PropertyBag->FindPropertyDescByName(ValueName) : nullptr;
 		if (PropertyDesc != nullptr)
 		{
 			if (PropertyDesc->ContainerTypes.IsEmpty()) // avoid trying to inline containers
@@ -147,27 +236,37 @@ class SVariablesOutlinerValue : public SCompoundWidget, public FNotifyHook
 				FSinglePropertyParams SinglePropertyArgs;
 				SinglePropertyArgs.NamePlacement = EPropertyNamePlacement::Hidden;
 				SinglePropertyArgs.NotifyHook = this;
+				SinglePropertyArgs.bHideResetToDefault = true;
 
 				FPropertyEditorModule& PropertyEditorModule = FModuleManager::Get().LoadModuleChecked<FPropertyEditorModule>("PropertyEditor");
 				
-				const TSharedPtr<ISinglePropertyView> SingleStructPropertyView = PropertyEditorModule.CreateSingleProperty(MakeShared<FInstancePropertyBagStructureDataProvider>(PropertyBag), IAnimNextRigVMVariableInterface::ValueName, SinglePropertyArgs);
+				const TSharedPtr<ISinglePropertyView> SingleStructPropertyView = PropertyEditorModule.CreateSingleProperty(MakeShared<FInstancePropertyBagStructureDataProvider>(*PropertyBag), ValueName, SinglePropertyArgs);
 				if (SingleStructPropertyView.IsValid())
 				{
-					ValueWidget = SingleStructPropertyView;
+					ValueWidget = SNew(SHorizontalBox)
+						+SHorizontalBox::Slot()
+						.FillWidth(1.0f)
+						[
+							SNew(SBox)
+							.HAlign(HAlign_Left)
+							[
+								SingleStructPropertyView.ToSharedRef()
+							]
+						]
+						+SHorizontalBox::Slot()
+						.AutoWidth()
+						.VAlign(VAlign_Center)
+						.HAlign(HAlign_Center)
+						[
+							SNew(SVariableOverride, DataInterfaceEntry, VariableName)
+							.Visibility(bCanOverride ? EVisibility::Visible : EVisibility::Collapsed)
+							.OverrideStatus(OverrideStatus)
+						];
 				}
 			}
 		}
 
-		ChildSlot
-		[
-			SNew(SHorizontalBox)
-			+SHorizontalBox::Slot()
-			.HAlign(HAlign_Center)
-			.AutoWidth()
-			[
-				ValueWidget.ToSharedRef()
-			]
-		];
+		WidgetContainer->SetContent(ValueWidget.ToSharedRef());
 	}
 
 	// FNotifyHook interface
@@ -178,27 +277,71 @@ class SVariablesOutlinerValue : public SCompoundWidget, public FNotifyHook
 		{
 			return;
 		}
-		
-		UAnimNextRigVMAssetEntry* Entry = CastChecked<UAnimNextRigVMAssetEntry>(TreeItem->WeakEntry.Get());
-		if(Entry == nullptr)
+
+		UAnimNextVariableEntry* VariableEntry = Cast<UAnimNextVariableEntry>(TreeItem->WeakEntry.Get());
+		if(VariableEntry == nullptr)
 		{
 			return;
 		}
 
-		UAnimNextRigVMAssetEditorData* EditorData = Entry->GetTypedOuter<UAnimNextRigVMAssetEditorData>();
+		UAnimNextRigVMAssetEditorData* EditorData = VariableEntry->GetTypedOuter<UAnimNextRigVMAssetEditorData>();
 		if (EditorData == nullptr)
 		{
 			return;
 		}
 
-		// Needed to show the changed sign in the table when we modify the PropertyBag
-		Entry->MarkPackageDirty();
+		UAnimNextDataInterfaceEntry* DataInterfaceEntry = Cast<UAnimNextDataInterfaceEntry>(TreeItem->WeakDataInterfaceEntry.Get());
+		if(DataInterfaceEntry != nullptr)
+		{
+			const FName VariableName = VariableEntry->GetVariableName();
+			if(!TreeItem->WeakDataInterfaceEntry->HasValueOverride(VariableName))
+			{
+				// No value override yet, so we copy from the internal property bag
+				check(InternalPropertyBag.IsValid() && InternalPropertyBag.GetPropertyBagStruct()->GetPropertyDescs().Num() == 1);
+				const FProperty* Property = InternalPropertyBag.GetPropertyBagStruct()->GetPropertyDescs()[0].CachedProperty; 
+				const uint8* DataPtr = Property->ContainerPtrToValuePtr<uint8>(InternalPropertyBag.GetValue().GetMemory());
+				TreeItem->WeakDataInterfaceEntry->SetValueOverride(VariableName, VariableEntry->GetType(), TConstArrayView<uint8>(DataPtr, Property->GetElementSize()));
+				ensure(TreeItem->WeakDataInterfaceEntry->HasValueOverride(VariableName));
+			}
+
+			DataInterfaceEntry->MarkPackageDirty();
+			DataInterfaceEntry->BroadcastModified(EAnimNextEditorDataNotifType::VariableDefaultValueChanged);
+		}
+		else
+		{
+			VariableEntry->MarkPackageDirty();
+			VariableEntry->BroadcastModified(EAnimNextEditorDataNotifType::VariableDefaultValueChanged);
+		}
+	}
+
+	void HandleModified(UAnimNextRigVMAssetEditorData* InEditorData, EAnimNextEditorDataNotifType InType, UObject* InSubject)
+	{
+		if( InType != EAnimNextEditorDataNotifType::VariableDefaultValueChanged &&
+			InType != EAnimNextEditorDataNotifType::UndoRedo)
+		{
+			return;
+		}
+		
+		TSharedPtr<FVariablesOutlinerEntryItem> TreeItem = WeakTreeItem.Pin();
+		if(!TreeItem.IsValid())
+		{
+			return;
+		}
+
+		UAnimNextVariableEntry* VariableEntry = Cast<UAnimNextVariableEntry>(TreeItem->WeakEntry.Get());
+		UAnimNextDataInterfaceEntry* DataInterfaceEntry = Cast<UAnimNextDataInterfaceEntry>(TreeItem->WeakDataInterfaceEntry.Get());
+		if(VariableEntry != InSubject && DataInterfaceEntry != InSubject)
+		{
+			return;
+		}
 
-		// Ensure that default values get picked up and forwarded to compiler
-		EditorData->BroadcastModified(EAnimNextEditorDataNotifType::PropertyChanged, Entry);
+		BuildValueWidget();
 	}
 
 	TWeakPtr<FVariablesOutlinerEntryItem> WeakTreeItem;
+	TSharedPtr<SBox> WidgetContainer;
+	FInstancedPropertyBag InternalPropertyBag;
+	EAnimNextDataInterfaceValueOverrideStatus OverrideStatus = EAnimNextDataInterfaceValueOverrideStatus::NotOverridden;
 };
 
 const TSharedRef<SWidget> FVariablesOutlinerValueColumn::ConstructRowWidget(FSceneOutlinerTreeItemRef Item, const STableRow<FSceneOutlinerTreeItemPtr>& Row)
@@ -259,6 +402,7 @@ class SVariablesOutlinerAccessSpecifier : public SCompoundWidget, public FNotify
 		ChildSlot
 		[
 			SNew(SBox)
+			.IsEnabled(!InTreeItem.WeakDataInterfaceEntry.IsValid())
 			.WidthOverride(16.0f)
 			.HeightOverride(16.0f)
 			.VAlign(VAlign_Center)
diff --git a/Engine/Plugins/Experimental/AnimNext/Source/AnimNextEditor/Private/Variables/Outliner/VariablesOutlinerDataInterfaceItem.cpp b/Engine/Plugins/Experimental/AnimNext/Source/AnimNextEditor/Private/Variables/Outliner/VariablesOutlinerDataInterfaceItem.cpp
new file mode 100644
--- /dev/null
+++ b/Engine/Plugins/Experimental/AnimNext/Source/AnimNextEditor/Private/Variables/Outliner/VariablesOutlinerDataInterfaceItem.cpp
@@ -0,0 +1,213 @@
+﻿// Copyright Epic Games, Inc. All Rights Reserved.
+
+#include "VariablesOutlinerDataInterfaceItem.h"
+
+#include "ISceneOutliner.h"
+#include "Styling/SlateColor.h"
+#include "Misc/Attribute.h"
+#include "Misc/Optional.h"
+#include "SceneOutlinerStandaloneTypes.h"
+#include "AnimNextEditorModule.h"
+#include "AnimNextRigVMAssetEditorData.h"
+#include "DataInterface/AnimNextDataInterface.h"
+#include "Entries/AnimNextDataInterfaceEntry.h"
+#include "Styling/SlateIconFinder.h"
+#include "Styling/StarshipCoreStyle.h"
+#include "UObject/Package.h"
+#include "Widgets/SCompoundWidget.h"
+#include "Widgets/DeclarativeSyntaxSupport.h"
+#include "Widgets/Images/SImage.h"
+#include "Widgets/Text/SInlineEditableTextBlock.h"
+#include "Widgets/Text/STextBlock.h"
+#include "Widgets/Views/STableRow.h"
+#include "Widgets/Views/SListView.h"
+
+#define LOCTEXT_NAMESPACE "VariablesOutlinerTreeItem"
+
+namespace UE::AnimNext::Editor
+{
+
+const FSceneOutlinerTreeItemType FVariablesOutlinerDataInterfaceItem::Type;
+
+class SVariablesOutlinerDataInterfaceLabel : FSceneOutlinerCommonLabelData, public SCompoundWidget
+{
+	SLATE_BEGIN_ARGS(SVariablesOutlinerDataInterfaceLabel) {}
+	SLATE_END_ARGS()
+
+	void Construct(const FArguments& InArgs, FVariablesOutlinerDataInterfaceItem& InTreeItem, ISceneOutliner& SceneOutliner, const STableRow<FSceneOutlinerTreeItemPtr>& InRow)
+	{
+		WeakSceneOutliner = StaticCastSharedRef<ISceneOutliner>(SceneOutliner.AsShared());
+		TreeItem = StaticCastSharedRef<FVariablesOutlinerDataInterfaceItem>(InTreeItem.AsShared());
+
+		FText AssetName;
+		if(UAnimNextDataInterfaceEntry* AssetEntry = InTreeItem.WeakEntry.Get())
+		{
+			AssetName = FText::FromString(AssetEntry->GetDataInterfacePath().ToString());
+		}
+		else
+		{
+			AssetName = LOCTEXT("UnknownAssetName", "Unknown Asset");
+		}
+		
+		ChildSlot
+		[
+			SNew(SHorizontalBox)
+			.ToolTipText(FText::Format(LOCTEXT("ImportedVariablesFormat", "Imported variables from '{0}'"), AssetName))
+			+SHorizontalBox::Slot()
+			.AutoWidth()
+			.VAlign(VAlign_Center)
+			.Padding(2.0f, 2.0f)
+			[
+				SNew(SImage)
+				.Image(FAppStyle::Get().GetBrush("Icons.Import"))
+			]
+			+SHorizontalBox::Slot()
+			.AutoWidth()
+			.VAlign(VAlign_Center)
+			.Padding(0.0f, 2.0f)
+			[
+				SNew(STextBlock)
+				.Font(FStyleFonts::Get().NormalBold)
+				.Text(this, &SVariablesOutlinerDataInterfaceLabel::GetDisplayText)
+				.HighlightText(SceneOutliner.GetFilterHighlightText())
+				.ColorAndOpacity(this, &SVariablesOutlinerDataInterfaceLabel::GetForegroundColor)
+			]
+			+SHorizontalBox::Slot()
+			.VAlign(VAlign_Center)
+			.Padding(2.0f, 0.0f, 2.0f, 3.0f)
+			.AutoWidth()
+			[
+				SNew(SImage)
+				.Visibility(this, &SVariablesOutlinerDataInterfaceLabel::GetDirtyImageVisibility)
+				.ToolTipText(this, &SVariablesOutlinerDataInterfaceLabel::GetDirtyTooltipText)
+				.Image(this, &SVariablesOutlinerDataInterfaceLabel::GetDirtyImageBrush)
+			]
+		];
+	}
+
+	FText GetDirtyTooltipText() const
+	{
+		if (const TSharedPtr<FVariablesOutlinerDataInterfaceItem> Item = TreeItem.Pin())
+		{
+			FTextBuilder TextBuilder;
+			TextBuilder.AppendLine(LOCTEXT("ModifiedTooltip", "Modified"));
+
+			if(UAnimNextDataInterfaceEntry* AssetEntry = Item->WeakEntry.Get())
+			{
+				const UPackage* ExternalPackage = AssetEntry->GetExternalPackage();
+				if(ExternalPackage && ExternalPackage->IsDirty())
+				{
+					TextBuilder.AppendLine(FText::FromName(ExternalPackage->GetFName()));
+				}
+			}
+
+			return TextBuilder.ToText();
+		}
+		return FText::GetEmpty();
+	}
+
+	const FSlateBrush* GetDirtyImageBrush() const
+	{
+		if (const TSharedPtr<FVariablesOutlinerDataInterfaceItem> Item = TreeItem.Pin())
+		{
+			if(UAnimNextDataInterfaceEntry* AssetEntry = Item->WeakEntry.Get())
+			{
+				bool bIsDirty = false;
+				const UPackage* ExternalPackage = AssetEntry->GetExternalPackage();
+				if(ExternalPackage && ExternalPackage->IsDirty())
+				{
+					bIsDirty = true;
+				}
+
+				return bIsDirty ? FAppStyle::GetBrush("Icons.DirtyBadge") : nullptr;
+			}
+		}
+		return nullptr;
+	}
+
+	EVisibility GetDirtyImageVisibility() const
+	{
+		if (const TSharedPtr<FVariablesOutlinerDataInterfaceItem> Item = TreeItem.Pin())
+		{
+			if(UAnimNextDataInterfaceEntry* AssetEntry = Item->WeakEntry.Get())
+			{
+				bool bIsDirty = false;
+				const UPackage* Package = AssetEntry->GetPackage();
+				check(Package);
+				if(Package->IsDirty())
+				{
+					bIsDirty = true;
+				}
+
+				return bIsDirty ? EVisibility::Visible : EVisibility::Collapsed;
+			}
+		}
+		return EVisibility::Collapsed;
+	}
+	
+	FText GetDisplayText() const
+	{
+		if (const TSharedPtr<FVariablesOutlinerDataInterfaceItem> Item = TreeItem.Pin())
+		{
+			return FText::FromString(Item->GetDisplayString());
+		}
+		return FText();
+	}
+
+	virtual FSlateColor GetForegroundColor() const override
+	{
+		const TOptional<FLinearColor> BaseColor = FSceneOutlinerCommonLabelData::GetForegroundColor(*TreeItem.Pin());
+		return BaseColor.IsSet() ? BaseColor.GetValue() : FSlateColor::UseForeground();
+	}
+
+	TWeakPtr<FVariablesOutlinerDataInterfaceItem> TreeItem;
+	TSharedPtr<SInlineEditableTextBlock> TextBlock;
+};
+
+FVariablesOutlinerDataInterfaceItem::FVariablesOutlinerDataInterfaceItem(UAnimNextDataInterfaceEntry* InEntry)
+	: ISceneOutlinerTreeItem(FVariablesOutlinerDataInterfaceItem::Type)
+	, WeakEntry(InEntry)
+{
+}
+
+bool FVariablesOutlinerDataInterfaceItem::IsValid() const
+{
+	return WeakEntry.Get() != nullptr;
+}
+
+FSceneOutlinerTreeItemID FVariablesOutlinerDataInterfaceItem::GetID() const
+{
+	TSoftObjectPtr<UAnimNextRigVMAssetEntry> SoftObjectPtr(WeakEntry.Get());
+	return GetTypeHash(SoftObjectPtr);
+}
+
+FString FVariablesOutlinerDataInterfaceItem::GetDisplayString() const
+{
+	UAnimNextRigVMAssetEntry* Entry = WeakEntry.Get();
+	if(Entry == nullptr)
+	{
+		return FString();
+	}
+
+	return Entry->GetDisplayName().ToString();
+}
+
+TSharedRef<SWidget> FVariablesOutlinerDataInterfaceItem::GenerateLabelWidget(ISceneOutliner& Outliner, const STableRow<FSceneOutlinerTreeItemPtr>& InRow)
+{
+	return SNew(SVariablesOutlinerDataInterfaceLabel, *this, Outliner, InRow);
+}
+
+FString FVariablesOutlinerDataInterfaceItem::GetPackageName() const
+{
+	UAnimNextRigVMAssetEntry* Entry = WeakEntry.Get();
+	if(Entry == nullptr)
+	{
+		return ISceneOutlinerTreeItem::GetPackageName();
+	}
+
+	return Entry->GetPackage()->GetName();
+}
+
+}
+
+#undef LOCTEXT_NAMESPACE // "VariablesOutlinerTreeItem"
\ No newline at end of file
diff --git a/Engine/Plugins/Experimental/AnimNext/Source/AnimNextEditor/Private/Variables/Outliner/VariablesOutlinerDataInterfaceItem.h b/Engine/Plugins/Experimental/AnimNext/Source/AnimNextEditor/Private/Variables/Outliner/VariablesOutlinerDataInterfaceItem.h
new file mode 100644
--- /dev/null
+++ b/Engine/Plugins/Experimental/AnimNext/Source/AnimNextEditor/Private/Variables/Outliner/VariablesOutlinerDataInterfaceItem.h
@@ -0,0 +1,33 @@
+﻿// Copyright Epic Games, Inc. All Rights Reserved.
+
+#pragma once
+
+#include "ISceneOutlinerTreeItem.h"
+
+class UAnimNextDataInterfaceEntry;
+
+namespace UE::AnimNext::Editor
+{
+
+class IWorkspaceOutlinerItemDetails;
+
+struct FVariablesOutlinerDataInterfaceItem : ISceneOutlinerTreeItem
+{
+	static const FSceneOutlinerTreeItemType Type;
+
+	FVariablesOutlinerDataInterfaceItem(UAnimNextDataInterfaceEntry* InEntry);
+
+	// Begin ISceneOutlinerTreeItem overrides
+	virtual bool IsValid() const override;
+	virtual FSceneOutlinerTreeItemID GetID() const override;
+	virtual FString GetDisplayString() const override;
+	virtual bool CanInteract() const override { return true; }
+	virtual TSharedRef<SWidget> GenerateLabelWidget(ISceneOutliner& Outliner, const STableRow<FSceneOutlinerTreeItemPtr>& InRow) override;
+	virtual FString GetPackageName() const override;
+	// End ISceneOutlinerTreeItem overrides
+
+	// Ptr to the underlying entry
+	TWeakObjectPtr<UAnimNextDataInterfaceEntry> WeakEntry;
+};
+
+}
diff --git a/Engine/Plugins/Experimental/AnimNext/Source/AnimNextEditor/Private/Variables/Outliner/VariablesOutlinerEntryItem.cpp b/Engine/Plugins/Experimental/AnimNext/Source/AnimNextEditor/Private/Variables/Outliner/VariablesOutlinerEntryItem.cpp
--- a/Engine/Plugins/Experimental/AnimNext/Source/AnimNextEditor/Private/Variables/Outliner/VariablesOutlinerEntryItem.cpp
+++ b/Engine/Plugins/Experimental/AnimNext/Source/AnimNextEditor/Private/Variables/Outliner/VariablesOutlinerEntryItem.cpp
@@ -23,6 +23,7 @@
 #include "Widgets/Views/STableRow.h"
 #include "Widgets/Views/SListView.h"
 #include "EdGraphSchema_K2.h"
+#include "Entries/AnimNextVariableEntry.h"
 
 #define LOCTEXT_NAMESPACE "VariablesOutlinerTreeItem"
 
@@ -75,7 +76,7 @@ class SVariablesOutlinerEntrylabel : FSceneOutlinerCommonLabelData, public SComp
 			FTextBuilder TextBuilder;
 			TextBuilder.AppendLine(LOCTEXT("ModifiedTooltip", "Modified"));
 
-			if(UAnimNextRigVMAssetEntry* AssetEntry = Item->WeakEntry.Get())
+			if(UAnimNextVariableEntry* AssetEntry = Item->WeakEntry.Get())
 			{
 				const UPackage* ExternalPackage = AssetEntry->GetExternalPackage();
 				if(ExternalPackage && ExternalPackage->IsDirty())
@@ -93,7 +94,7 @@ class SVariablesOutlinerEntrylabel : FSceneOutlinerCommonLabelData, public SComp
 	{
 		if (const TSharedPtr<FVariablesOutlinerEntryItem> Item = TreeItem.Pin())
 		{
-			if(UAnimNextRigVMAssetEntry* AssetEntry = Item->WeakEntry.Get())
+			if(UAnimNextVariableEntry* AssetEntry = Item->WeakEntry.Get())
 			{
 				bool bIsDirty = false;
 				const UPackage* ExternalPackage = AssetEntry->GetExternalPackage();
@@ -147,7 +148,7 @@ class SVariablesOutlinerEntrylabel : FSceneOutlinerCommonLabelData, public SComp
 	TSharedPtr<SInlineEditableTextBlock> TextBlock;
 };
 
-FVariablesOutlinerEntryItem::FVariablesOutlinerEntryItem(UAnimNextRigVMAssetEntry* InEntry)
+FVariablesOutlinerEntryItem::FVariablesOutlinerEntryItem(UAnimNextVariableEntry* InEntry)
 	: ISceneOutlinerTreeItem(FVariablesOutlinerEntryItem::Type)
 	, WeakEntry(InEntry)
 {
@@ -160,7 +161,7 @@ bool FVariablesOutlinerEntryItem::IsValid() const
 
 FSceneOutlinerTreeItemID FVariablesOutlinerEntryItem::GetID() const
 {
-	return GetTypeHash(WeakEntry);
+	return HashCombine(GetTypeHash(WeakEntry), GetTypeHash(WeakDataInterfaceEntry));
 }
 
 FString FVariablesOutlinerEntryItem::GetDisplayString() const
@@ -206,7 +207,7 @@ void FVariablesOutlinerEntryItem::Rename(const FText& InNewName)
 
 bool FVariablesOutlinerEntryItem::ValidateName(const FText& InNewName, FText& OutErrorMessage) const
 {
-	UAnimNextRigVMAssetEntry* Entry = WeakEntry.Get();
+	UAnimNextVariableEntry* Entry = WeakEntry.Get();
 	if(Entry == nullptr)
 	{
 		OutErrorMessage = LOCTEXT("InvalidVariableError", "Variable is invalid");
diff --git a/Engine/Plugins/Experimental/AnimNext/Source/AnimNextEditor/Private/Variables/Outliner/VariablesOutlinerEntryItem.h b/Engine/Plugins/Experimental/AnimNext/Source/AnimNextEditor/Private/Variables/Outliner/VariablesOutlinerEntryItem.h
--- a/Engine/Plugins/Experimental/AnimNext/Source/AnimNextEditor/Private/Variables/Outliner/VariablesOutlinerEntryItem.h
+++ b/Engine/Plugins/Experimental/AnimNext/Source/AnimNextEditor/Private/Variables/Outliner/VariablesOutlinerEntryItem.h
@@ -3,19 +3,19 @@
 #pragma once
 
 #include "ISceneOutlinerTreeItem.h"
+#include "Variables/AnimNextVariableEntryProxy.h"
 
-class UAnimNextRigVMAssetEntry;
+class UAnimNextDataInterfaceEntry;
+class UAnimNextVariableEntry;
 
 namespace UE::AnimNext::Editor
 {
 
-class IWorkspaceOutlinerItemDetails;
-
 struct FVariablesOutlinerEntryItem : ISceneOutlinerTreeItem
 {
 	static const FSceneOutlinerTreeItemType Type;
 
-	FVariablesOutlinerEntryItem(UAnimNextRigVMAssetEntry* InEntry);
+	FVariablesOutlinerEntryItem(UAnimNextVariableEntry* InEntry);
 
 	// Begin ISceneOutlinerTreeItem overrides
 	virtual bool IsValid() const override;
@@ -33,7 +33,13 @@ struct FVariablesOutlinerEntryItem : ISceneOutlinerTreeItem
 	bool ValidateName(const FText& InNewName, FText& OutErrorMessage) const;
 
 	// Ptr to the underlying entry
-	TWeakObjectPtr<UAnimNextRigVMAssetEntry> WeakEntry;
+	TWeakObjectPtr<UAnimNextVariableEntry> WeakEntry;
+
+	// The data interface entry this entry is from, if any
+	TWeakObjectPtr<UAnimNextDataInterfaceEntry> WeakDataInterfaceEntry;
+
+	// Proxy entry used for details panel editing of variables in implemented data interfaces
+	TStrongObjectPtr<UAnimNextVariableEntryProxy> ProxyEntry;
 };
 
 }
diff --git a/Engine/Plugins/Experimental/AnimNext/Source/AnimNextEditor/Private/Variables/Outliner/VariablesOutlinerHierarchy.cpp b/Engine/Plugins/Experimental/AnimNext/Source/AnimNextEditor/Private/Variables/Outliner/VariablesOutlinerHierarchy.cpp
--- a/Engine/Plugins/Experimental/AnimNext/Source/AnimNextEditor/Private/Variables/Outliner/VariablesOutlinerHierarchy.cpp
+++ b/Engine/Plugins/Experimental/AnimNext/Source/AnimNextEditor/Private/Variables/Outliner/VariablesOutlinerHierarchy.cpp
@@ -5,13 +5,16 @@
 #include "AnimNextRigVMAsset.h"
 #include "ISceneOutlinerMode.h"
 #include "AnimNextRigVMAssetEditorData.h"
-#include "Variables/IAnimNextRigVMVariableInterface.h"
 #include "VariablesOutlinerAssetItem.h"
+#include "VariablesOutlinerDataInterfaceItem.h"
 #include "VariablesOutlinerMode.h"
 #include "WorkspaceAssetRegistryInfo.h"
 #include "VariablesOutlinerEntryItem.h"
+#include "Entries/AnimNextDataInterfaceEntry.h"
 #include "Entries/AnimNextRigVMAssetEntry.h"
+#include "Entries/AnimNextVariableEntry.h"
 #include "Variables/SVariablesView.h"
+#include "DataInterface/AnimNextDataInterface.h"
 
 namespace UE::AnimNext::Editor
 {
@@ -43,43 +46,114 @@ void FVariablesOutlinerHierarchy::CreateItems(TArray<FSceneOutlinerTreeItemPtr>&
 			OutItems.Add(Item);
 		}
 
-		EditorData->ForEachEntryOfType<IAnimNextRigVMVariableInterface>([this, &OutItems](IAnimNextRigVMVariableInterface* InVariable)
+		EditorData->ForEachEntryOfType<UAnimNextVariableEntry>([this, &OutItems](UAnimNextVariableEntry* InVariable)
 		{
-			if (FSceneOutlinerTreeItemPtr Item = Mode->CreateItemFor<FVariablesOutlinerEntryItem>(CastChecked<UAnimNextRigVMAssetEntry>(InVariable)))
+			if (FSceneOutlinerTreeItemPtr Item = Mode->CreateItemFor<FVariablesOutlinerEntryItem>(CastChecked<UAnimNextVariableEntry>(InVariable)))
 			{
 				OutItems.Add(Item);
 			}
 			return true;
 		});
+
+		EditorData->ForEachEntryOfType<UAnimNextDataInterfaceEntry>([this, &OutItems](UAnimNextDataInterfaceEntry* InInterfaceEntry)
+		{
+			if (FSceneOutlinerTreeItemPtr Item = Mode->CreateItemFor<FVariablesOutlinerDataInterfaceItem>(CastChecked<UAnimNextDataInterfaceEntry>(InInterfaceEntry)))
+			{
+				OutItems.Add(Item);
+			}
+
+			auto AddDataInterface = [this, &OutItems](UAnimNextDataInterfaceEntry* InDataInterfaceEntry, UAnimNextDataInterfaceEntry* InRootDataInterfaceEntry, auto& InAddDataInterface) -> void
+			{
+				UAnimNextRigVMAssetEditorData* InterfaceEditorData = UncookedOnly::FUtils::GetEditorData<UAnimNextRigVMAssetEditorData>(InDataInterfaceEntry->DataInterface.Get());
+				InterfaceEditorData->ForEachEntryOfType<UAnimNextVariableEntry>([this, &OutItems, InRootDataInterfaceEntry](UAnimNextVariableEntry* InVariable)
+				{
+					if(InVariable->GetExportAccessSpecifier() == EAnimNextExportAccessSpecifier::Public)
+					{
+						if (FSceneOutlinerTreeItemPtr Item = Mode->CreateItemFor<FVariablesOutlinerEntryItem>(CastChecked<UAnimNextVariableEntry>(InVariable)))
+						{
+							FVariablesOutlinerEntryItem* EntryItem = Item->CastTo<FVariablesOutlinerEntryItem>();
+							EntryItem->WeakDataInterfaceEntry = InRootDataInterfaceEntry;
+							OutItems.Add(Item);
+						}
+					}
+					return true;
+				});
+
+				InterfaceEditorData->ForEachEntryOfType<UAnimNextDataInterfaceEntry>([this, &OutItems, &InAddDataInterface, InRootDataInterfaceEntry](UAnimNextDataInterfaceEntry* InSubInterfaceEntry)
+				{
+					InAddDataInterface(InSubInterfaceEntry, InRootDataInterfaceEntry, InAddDataInterface);
+					return true;
+				});
+			};
+
+			AddDataInterface(InInterfaceEntry, InInterfaceEntry, AddDataInterface);
+
+			return true;
+		});
 	}
 }
 
 FSceneOutlinerTreeItemPtr FVariablesOutlinerHierarchy::FindOrCreateParentItem(const ISceneOutlinerTreeItem& Item, const TMap<FSceneOutlinerTreeItemID, FSceneOutlinerTreeItemPtr>& Items, bool bCreate)
 {
-	const FVariablesOutlinerEntryItem* TreeItem = Item.CastTo<FVariablesOutlinerEntryItem>();
-	if (TreeItem == nullptr)
+	if (const FVariablesOutlinerEntryItem* EntryItem = Item.CastTo<FVariablesOutlinerEntryItem>())
 	{
-		return nullptr;
-	}
+		UAnimNextDataInterfaceEntry* DataInterfaceEntry = EntryItem->WeakDataInterfaceEntry.Get();
+		if(DataInterfaceEntry != nullptr)
+		{
+			// Add as part of an imported data interface, so use that as parent
+			TSoftObjectPtr<UAnimNextRigVMAssetEntry> SoftObjectPtr(DataInterfaceEntry); 
+			uint32 Hash = GetTypeHash(SoftObjectPtr);
 
-	UAnimNextRigVMAssetEntry* Entry = TreeItem->WeakEntry.Get();
-	if (Entry == nullptr)
-	{
-		return nullptr;
+			if (const FSceneOutlinerTreeItemPtr* ParentItem = Items.Find(Hash))
+			{
+				return *ParentItem;
+			}
+		}
+		else
+		{
+			UAnimNextRigVMAssetEntry* Entry = EntryItem->WeakEntry.Get();
+			if (Entry == nullptr)
+			{
+				return nullptr;
+			}
+
+			UAnimNextRigVMAsset* Asset = Entry->GetTypedOuter<UAnimNextRigVMAsset>();
+			if (Entry == nullptr)
+			{
+				return nullptr;
+			}
+
+			TSoftObjectPtr<UAnimNextRigVMAsset> SoftObjectPtr(Asset);
+			uint32 Hash = GetTypeHash(SoftObjectPtr);
+
+			if (const FSceneOutlinerTreeItemPtr* ParentItem = Items.Find(Hash))
+			{
+				return *ParentItem;
+			}
+		}
 	}
 
-	UAnimNextRigVMAsset* Asset = Entry->GetTypedOuter<UAnimNextRigVMAsset>();
-	if (Entry == nullptr)
+	if (const FVariablesOutlinerDataInterfaceItem* InterfaceItem = Item.CastTo<FVariablesOutlinerDataInterfaceItem>())
 	{
-		return nullptr;
-	}
+		UAnimNextRigVMAssetEntry* Entry = InterfaceItem->WeakEntry.Get();
+		if (Entry == nullptr)
+		{
+			return nullptr;
+		}
 
-	TSoftObjectPtr<UAnimNextRigVMAsset> SoftObjectPtr(Asset);
-	uint32 Hash = GetTypeHash(SoftObjectPtr);
+		UAnimNextRigVMAsset* Asset = Entry->GetTypedOuter<UAnimNextRigVMAsset>();
+		if (Entry == nullptr)
+		{
+			return nullptr;
+		}
 
-	if (const FSceneOutlinerTreeItemPtr* ParentItem = Items.Find(Hash))
-	{
-		return *ParentItem;
+		TSoftObjectPtr<UAnimNextRigVMAsset> SoftObjectPtr(Asset);
+		uint32 Hash = GetTypeHash(SoftObjectPtr);
+
+		if (const FSceneOutlinerTreeItemPtr* ParentItem = Items.Find(Hash))
+		{
+			return *ParentItem;
+		}
 	}
 
 	return nullptr;
diff --git a/Engine/Plugins/Experimental/AnimNext/Source/AnimNextEditor/Private/Variables/Outliner/VariablesOutlinerMode.cpp b/Engine/Plugins/Experimental/AnimNext/Source/AnimNextEditor/Private/Variables/Outliner/VariablesOutlinerMode.cpp
--- a/Engine/Plugins/Experimental/AnimNext/Source/AnimNextEditor/Private/Variables/Outliner/VariablesOutlinerMode.cpp
+++ b/Engine/Plugins/Experimental/AnimNext/Source/AnimNextEditor/Private/Variables/Outliner/VariablesOutlinerMode.cpp
@@ -20,10 +20,13 @@
 #include "Variables/AnimNextVariableItemMenuContext.h"
 #include "AnimNextRigVMAsset.h"
 #include "VariablesOutlinerAssetItem.h"
+#include "VariablesOutlinerDataInterfaceItem.h"
 #include "VariablesOutlinerDragDrop.h"
 #include "Common/GraphEditorSchemaActions.h"
 #include "Entries/AnimNextVariableEntry.h"
 #include "Variables/SVariablesView.h"
+#include "Entries/AnimNextDataInterfaceEntry.h"
+#include "Variables/AnimNextVariableEntryProxy.h"
 
 #define LOCTEXT_NAMESPACE "VariablesOutlinerMode"
 
@@ -155,11 +158,32 @@ void FVariablesOutlinerMode::HandleItemSelection(const FSceneOutlinerItemSelecti
 		EntriesToShow.Reserve(SelectedItems.Num());
 		for(const FSceneOutlinerTreeItemPtr& Item : SelectedItems)
 		{
-			if (const FVariablesOutlinerEntryItem* VariablesItem = Item->CastTo<FVariablesOutlinerEntryItem>())
+			if (FVariablesOutlinerEntryItem* VariablesItem = Item->CastTo<FVariablesOutlinerEntryItem>())
 			{
-				if(UAnimNextRigVMAssetEntry* Entry = VariablesItem->WeakEntry.Get())
+				if(UAnimNextVariableEntry* VariableEntry = VariablesItem->WeakEntry.Get())
 				{
-					EntriesToShow.Add(Entry);
+					if(UAnimNextDataInterfaceEntry* DataInterfaceEntry = VariablesItem->WeakDataInterfaceEntry.Get())
+					{
+						// Create proxy object to display in the details panel
+						if(!VariablesItem->ProxyEntry.IsValid())
+						{
+							VariablesItem->ProxyEntry = TStrongObjectPtr(NewObject<UAnimNextVariableEntryProxy>(GetTransientPackage(), NAME_None, RF_Transient));
+						}
+						VariablesItem->ProxyEntry.Get()->VariableEntry = VariableEntry;
+						VariablesItem->ProxyEntry.Get()->DataInterfaceEntry = DataInterfaceEntry;
+						EntriesToShow.Add(VariablesItem->ProxyEntry.Get());
+					}
+					else
+					{
+						EntriesToShow.Add(VariableEntry);
+					}
+				}
+			}
+			else if(const FVariablesOutlinerDataInterfaceItem* DataInterfaceItem = Item->CastTo<FVariablesOutlinerDataInterfaceItem>())
+			{
+				if(UAnimNextDataInterfaceEntry* DataInterfaceEntry = DataInterfaceItem->WeakEntry.Get())
+				{
+					EntriesToShow.Add(DataInterfaceEntry);
 				}
 			}
 		}
@@ -255,20 +279,39 @@ void FVariablesOutlinerMode::Delete()
 	{
 		if (const FVariablesOutlinerEntryItem* VariablesItem = Item->CastTo<FVariablesOutlinerEntryItem>())
 		{
-			UAnimNextRigVMAssetEntry* Entry = VariablesItem->WeakEntry.Get();
-			if(Entry == nullptr)
+			UAnimNextVariableEntry* VariableEntry = VariablesItem->WeakEntry.Get();
+			UAnimNextDataInterfaceEntry* DataInterfaceEntry = VariablesItem->WeakDataInterfaceEntry.Get();
+			if(VariableEntry == nullptr || DataInterfaceEntry != nullptr)	// Cant delete variables in other data interfaces
+			{
+				continue;
+			}
+
+			UAnimNextRigVMAssetEditorData* EditorData = VariableEntry->GetTypedOuter<UAnimNextRigVMAssetEditorData>();
+			if(EditorData == nullptr)
+			{
+				continue;
+			}
+
+			TArray<UAnimNextRigVMAssetEntry*>& EntriesToDelete = EntriesToDeletePerAsset.FindOrAdd(EditorData);
+			EntriesToDelete.Add(VariableEntry);
+			NumEntries++;
+		}
+		else if (const FVariablesOutlinerDataInterfaceItem* DataInterfaceItem = Item->CastTo<FVariablesOutlinerDataInterfaceItem>())
+		{
+			UAnimNextDataInterfaceEntry* DataInterfaceEntry = DataInterfaceItem->WeakEntry.Get();
+			if(DataInterfaceEntry == nullptr)
 			{
 				continue;
 			}
 
-			UAnimNextRigVMAssetEditorData* EditorData = Entry->GetTypedOuter<UAnimNextRigVMAssetEditorData>();
+			UAnimNextRigVMAssetEditorData* EditorData = DataInterfaceEntry->GetTypedOuter<UAnimNextRigVMAssetEditorData>();
 			if(EditorData == nullptr)
 			{
 				continue;
 			}
 
 			TArray<UAnimNextRigVMAssetEntry*>& EntriesToDelete = EntriesToDeletePerAsset.FindOrAdd(EditorData);
-			EntriesToDelete.Add(Entry);
+			EntriesToDelete.Add(DataInterfaceEntry);
 			NumEntries++;
 		}
 	}
@@ -319,7 +362,7 @@ TSharedPtr<FDragDropOperation> FVariablesOutlinerMode::CreateDragDropOperation(c
 	{
 		if(ExternalVariable.Name == VariableName)
 		{
-			const bool bIsGetter = MouseEvent.IsControlDown();
+			const bool bIsGetter = !(MouseEvent.IsShiftDown() || MouseEvent.IsAltDown());
 			TSharedPtr<FAnimNextSchemaAction_Variable> Action = MakeShared<FAnimNextSchemaAction_Variable>(ExternalVariable, bIsGetter);
 			return FVariableDragDropOp::New(Action);
 		}
diff --git a/Engine/Plugins/Experimental/AnimNext/Source/AnimNextEditor/Private/Variables/SAddVariablesDialog.cpp b/Engine/Plugins/Experimental/AnimNext/Source/AnimNextEditor/Private/Variables/SAddVariablesDialog.cpp
--- a/Engine/Plugins/Experimental/AnimNext/Source/AnimNextEditor/Private/Variables/SAddVariablesDialog.cpp
+++ b/Engine/Plugins/Experimental/AnimNext/Source/AnimNextEditor/Private/Variables/SAddVariablesDialog.cpp
@@ -3,6 +3,7 @@
 #include "SAddVariablesDialog.h"
 
 #include "AddVariableDialogMenuContext.h"
+#include "AnimNextRigVMAsset.h"
 #include "AnimNextVariableSettings.h"
 #include "AssetToolsModule.h"
 #include "ContentBrowserModule.h"
@@ -11,6 +12,7 @@
 #include "EditorUtils.h"
 #include "UncookedOnlyUtils.h"
 #include "PropertyBagDetails.h"
+#include "PropertyCustomizationHelpers.h"
 #include "SPinTypeSelector.h"
 #include "Framework/Application/SlateApplication.h"
 #include "Framework/Docking/TabManager.h"
@@ -21,7 +23,13 @@
 #include "SSimpleButton.h"
 #include "SSimpleComboButton.h"
 #include "ToolMenus.h"
+#include "DataInterface/AnimNextDataInterface.h"
+#include "DataInterface/AnimNextDataInterface_EditorData.h"
+#include "Entries/AnimNextDataInterfaceEntry.h"
+#include "Entries/AnimNextRigVMAssetEntry.h"
+#include "Entries/AnimNextVariableEntry.h"
 #include "String/ParseTokens.h"
+#include "Widgets/Input/SEditableTextBox.h"
 
 #define LOCTEXT_NAMESPACE "SAddVariablesDialog"
 
@@ -35,30 +43,93 @@ static FName Column_Type(TEXT("Type"));
 static FName SelectLibraryMenuName(TEXT("AnimNext.AddVariablesDialog.SelectedLibraryMenu"));
 }
 
-bool FVariableToAdd::IsValid(FText& OutReason) const
+bool SAddVariablesDialog::FVariableToAddEntry::IsValid(FText& OutReason) const
 {
 	if(Name == NAME_None)
 	{
 		OutReason = LOCTEXT("InvalidVariableName", "Invalid Variable Name");
+		return false;
 	}
 
 	if(!Type.IsValid())
 	{
 		OutReason = LOCTEXT("InvalidVariableType", "Invalid Variable Type");
+		return false;
 	}
-	
+
+	TSharedPtr<SAddVariablesDialog> PinnedDialog = Dialog.Pin();
+	if(PinnedDialog.IsValid())
+	{
+		TArray<FName> PendingNames;
+		PinnedDialog->GetPendingNames(PendingNames);
+		for(FName PendingName : PendingNames)
+		{
+			if(PendingName == Name)
+			{
+				OutReason = LOCTEXT("DuplicateVariableName", "Duplicate Variable Name");
+				return false;
+			}
+		}
+
+		for(const TSharedRef<FEntry>& Entry : PinnedDialog->RootEntries)
+		{
+			if(&Entry.Get() == this)
+			{
+				continue;
+			}
+			
+			switch(Entry->EntryType)
+			{
+			case EEntryType::Variable:
+				if(StaticCastSharedRef<FVariableToAddEntry>(Entry)->Name == Name)
+				{
+					OutReason = LOCTEXT("DuplicateVariableName", "Duplicate Variable Name");
+					return false;
+				}
+				break;
+			case EEntryType::DataInterface:
+				for(const TSharedRef<FEntry>& SubEntry : Entry->Children)
+				{
+					if(&SubEntry.Get() == this)
+					{
+						continue;
+					}
+
+					check(SubEntry->EntryType == EEntryType::Variable);
+					if(StaticCastSharedRef<FVariableToAddEntry>(SubEntry)->Name == Name)
+					{
+						OutReason = LOCTEXT("DuplicateVariableName", "Duplicate Variable Name");
+						return false;
+					}
+				}
+				break;
+			}
+		}
+	}
+
+	return true; 
+}
+
+bool SAddVariablesDialog::FDataInterfaceToAddEntry::IsValid(FText& OutReason) const
+{
+	if(DataInterface == nullptr)
+	{
+		OutReason = LOCTEXT("InvalidDataInterface", "Invalid Data Interface");
+		return false;
+	}
+
 	return true; 
 }
 
-void SAddVariablesDialog::Construct(const FArguments& InArgs, const FAssetData& InAsset)
+void SAddVariablesDialog::Construct(const FArguments& InArgs, const TArray<UAnimNextRigVMAssetEditorData*>& InAssetEditorDatas)
 {
 	using namespace AddVariablesDialog;
 
 	OnFilterVariableType = InArgs._OnFilterVariableType;
-	Asset = InAsset;
+	AssetEditorDatas = InAssetEditorDatas;
 
 	SWindow::Construct(SWindow::FArguments()
-		.Title(LOCTEXT("WindowTitle", "Add Parameters"))
+		.Title(LOCTEXT("WindowTitle", "Add Variables"))
 		.SizingRule(ESizingRule::UserSized)
 		.ClientSize(InArgs._AllowMultiple ? FVector2D(500.f, 500.f) : FVector2D(500.f, 100.f))
 		.SupportsMaximize(false)
@@ -73,33 +144,84 @@ void SAddVariablesDialog::Construct(const FArguments& InArgs, const FAssetData&
 				.HAlign(HAlign_Left)
 				.Padding(0.0f, 5.0f)
 				[
-					SNew(SSimpleButton)
-					.Visibility(InArgs._AllowMultiple ? EVisibility::Visible : EVisibility::Collapsed)
-					.Text(LOCTEXT("AddButton", "Add"))
-					.ToolTipText(LOCTEXT("AddButtonTooltip", "Queue a new parameter for adding. New parameters will re-use the settings from the last queued parameter."))
-					.Icon(FAppStyle::Get().GetBrush("Icons.Plus"))
-					.OnClicked_Lambda([this]()
-					{
-						AddEntry();
-						return FReply::Handled();
-					})
+					SNew(SHorizontalBox)
+					+SHorizontalBox::Slot()
+					.AutoWidth()
+					.Padding(0.3f, 0.0f)
+					[
+						SNew(SSimpleButton)
+						.Visibility(InArgs._AllowMultiple ? EVisibility::Visible : EVisibility::Collapsed)
+						.Text(LOCTEXT("AddVariableButton", "Add Variable"))
+						.ToolTipText(LOCTEXT("AddVariableButtonTooltip", "Queue a new variable for adding. New variables will re-use the settings from the last queued variable."))
+						.Icon(FAppStyle::Get().GetBrush("Icons.Plus"))
+						.OnClicked_Lambda([this]()
+						{
+							AddEntry();
+							return FReply::Handled();
+						})
+					]
+					+SHorizontalBox::Slot()
+					.AutoWidth()
+					.Padding(0.3f, 0.0f)
+					[
+						SNew(SSimpleComboButton)
+						.Visibility(InArgs._AllowMultiple ? EVisibility::Visible : EVisibility::Collapsed)
+						.Text(LOCTEXT("AddDataInterfaceButton", "Add Data Interface"))
+						.ToolTipText(LOCTEXT("AddDataInterfaceButtonTooltip", "Select a new data interface for adding."))
+						.Icon(FAppStyle::Get().GetBrush("Icons.Plus"))
+						.OnGetMenuContent_Lambda([this]()
+						{
+							FContentBrowserModule& ContentBrowserModule = FModuleManager::Get().LoadModuleChecked<FContentBrowserModule>(TEXT("ContentBrowser"));
+
+							FAssetPickerConfig AssetPickerConfig;
+							AssetPickerConfig.Filter.bRecursiveClasses = true;
+							AssetPickerConfig.Filter.ClassPaths.Add(UAnimNextDataInterface::StaticClass()->GetClassPathName());
+							AssetPickerConfig.InitialAssetViewType = EAssetViewType::List;
+							AssetPickerConfig.OnAssetSelected = FOnAssetSelected::CreateLambda([this](const FAssetData& InAssetData)
+							{
+								FSlateApplication::Get().DismissAllMenus();
+								if(UAnimNextDataInterface* DataInterface = Cast<UAnimNextDataInterface>(InAssetData.GetAsset()))
+								{
+									AddDataInterface(DataInterface);
+								}
+							});
+							AssetPickerConfig.OnShouldFilterAsset = FOnShouldFilterAsset::CreateLambda([](const FAssetData& InAssetData)
+							{
+								FAnimNextAssetRegistryExports Exports;
+								UncookedOnly::FUtils::GetExportedVariablesForAsset(InAssetData, Exports);
+								if(Exports.Variables.Num() == 0)
+								{
+									return true;
+								}
+								return false;
+							});
+
+							return SNew(SBox)
+								.WidthOverride(300.0f)
+								.HeightOverride(400.0f)
+								[
+									ContentBrowserModule.Get().CreateAssetPicker(AssetPickerConfig)
+								];
+						})
+					]
 				]
 				+SVerticalBox::Slot()
 				.FillHeight(1.0f)
 				[
-					SAssignNew(EntriesList, SListView<TSharedRef<FVariableToAddEntry>>)
-					.ListItemsSource(&Entries)
+					SAssignNew(EntriesTree, STreeView<TSharedRef<FEntry>>)
+					.TreeItemsSource(&RootEntries)
 					.OnGenerateRow(this, &SAddVariablesDialog::HandleGenerateRow)
+					.OnGetChildren(this, &SAddVariablesDialog::HandleGetChildren)
 					.HeaderRow(
 						SNew(SHeaderRow)
 						+SHeaderRow::Column(Column_Name)
 						.DefaultLabel(LOCTEXT("NameColumnHeader", "Name"))
-						.ToolTipText(LOCTEXT("NameColumnHeaderTooltip", "The name of the new parameter"))
+						.ToolTipText(LOCTEXT("NameColumnHeaderTooltip", "The name of the new variable"))
 						.FillWidth(0.25f)
 
 						+SHeaderRow::Column(Column_Type)
 						.DefaultLabel(LOCTEXT("TypeColumnHeader", "Type"))
-						.ToolTipText(LOCTEXT("TypeColumnHeaderTooltip", "The type of the new parameter"))
+						.ToolTipText(LOCTEXT("TypeColumnHeaderTooltip", "The type of the new variable"))
 						.FillWidth(0.25f)
 					)
 				]
@@ -118,33 +240,19 @@ void SAddVariablesDialog::Construct(const FArguments& InArgs, const FAssetData&
 						.ButtonStyle(&FAppStyle::Get().GetWidgetStyle<FButtonStyle>("PrimaryButton"))
 						.IsEnabled_Lambda([this]()
 						{
-							// Check each entry to see if the button can be pressed
-							for(TSharedRef<FVariableToAdd> Entry : Entries)
-							{
-								if(!Entry->IsValid())
-								{
-									return false;
-								}
-							}
-
-							return true;
-						})
-						.Text_Lambda([this]()
-						{
-							return FText::Format(LOCTEXT("AddParametersButtonFormat", "Add {0} {0}|plural(one=Parameter,other=Parameters)"), FText::AsNumber(Entries.Num()));
+							return bCanCreateVariables;
 						})
+						.Text(LOCTEXT("AddVariablesButtonFormat", "Add Variable(s)"))
 						.ToolTipText_Lambda([this]()
 						{
-							// Check each entry to see if the button can be pressed
-							for(TSharedRef<FVariableToAdd> Entry : Entries)
+							if(bCanCreateVariables)
 							{
-								FText Reason;
-								if(!Entry->IsValid(Reason))
-								{
-									return FText::Format(LOCTEXT("AddParametersButtonTooltip_InvalidEntry", "A parameter to add is not valid: {0}"), Reason);
-								}
+								return LOCTEXT("AddVariablesButtonTooltip", "Add the selected variables to the current graph");
+							}
+							else
+							{
+								return FText::Format(LOCTEXT("AddVariablesButtonTooltip_InvalidEntry", "A variable to add is not valid: {0}"), CreateErrorMessage);
 							}
-							return LOCTEXT("AddParametersButtonTooltip", "Add the selected parameters to the current graph");
 						})
 						.OnClicked_Lambda([this]()
 						{
@@ -159,7 +267,7 @@ void SAddVariablesDialog::Construct(const FArguments& InArgs, const FAssetData&
 						.HAlign(HAlign_Center)
 						.ButtonStyle(&FAppStyle::Get().GetWidgetStyle<FButtonStyle>("Button"))
 						.Text(LOCTEXT("CancelButton", "Cancel"))
-						.ToolTipText(LOCTEXT("CancelButtonTooltip", "Cancel adding new parameters"))
+						.ToolTipText(LOCTEXT("CancelButtonTooltip", "Cancel adding new variables"))
 						.OnClicked_Lambda([this]()
 						{
 							RequestDestroyWindow();
@@ -170,8 +278,11 @@ void SAddVariablesDialog::Construct(const FArguments& InArgs, const FAssetData&
 			]
 		]);
 
-	// Add an initial item
-	AddEntry(InArgs._InitialParamType);
+	if(InArgs._ShouldAddInitialVariable)
+	{
+		// Add an initial item
+		AddEntry(InArgs._InitialParamType);
+	}
 }
 
 FReply SAddVariablesDialog::OnKeyDown(const FGeometry& MyGeometry, const FKeyEvent& InKeyEvent)
@@ -181,42 +292,216 @@ FReply SAddVariablesDialog::OnKeyDown(const FGeometry& MyGeometry, const FKeyEve
 		RequestDestroyWindow();
 		return FReply::Handled();
 	}
+	else if(InKeyEvent.GetKey() == EKeys::Delete)
+	{
+		DeleteSelectedItems();
+		return FReply::Handled();
+	}
 
 	return FReply::Unhandled();
 }
 
+void SAddVariablesDialog::DeleteSelectedItems()
+{
+	TArray<TSharedRef<FEntry>> SelectedItems;
+	EntriesTree->GetSelectedItems(SelectedItems);
+	for(const TSharedRef<FEntry>& SelectedItem : SelectedItems)
+	{
+		RootEntries.Remove(SelectedItem);
+	}
+
+	RefreshEntries();
+}
+
+static FName GetNewVariableName(FName InBaseName, TArrayView<FName> InExistingNames)
+{
+	auto NameExists = [&InExistingNames](FName InName)
+	{
+		for(FName AdditionalName : InExistingNames)
+		{
+			if(AdditionalName == InName)
+			{
+				return true;
+			}
+		}
+
+		return false;
+	};
+
+	if(!NameExists(InBaseName))
+	{
+		// Early out - name is valid
+		return InBaseName;
+	}
+
+	int32 PostFixIndex = 0;
+	TStringBuilder<128> StringBuilder;
+	while(true)
+	{
+		StringBuilder.Reset();
+		InBaseName.GetDisplayNameEntry()->AppendNameToString(StringBuilder);
+		StringBuilder.Appendf(TEXT("_%d"), PostFixIndex++);
+
+		FName TestName(StringBuilder.ToString()); 
+		if(!NameExists(TestName))
+		{
+			return TestName;
+		}
+	}
+
+	return NAME_None;
+}
+
+void SAddVariablesDialog::GetPendingNamesRecursive(UAnimNextRigVMAssetEditorData* InEditorData, TArray<FName>& OutPendingNames) const
+{
+	for(UAnimNextRigVMAssetEntry* Entry : InEditorData->Entries)
+	{
+		if(UAnimNextVariableEntry* VariableEntry = Cast<UAnimNextVariableEntry>(Entry))
+		{
+			if(VariableEntry->GetExportAccessSpecifier() == EAnimNextExportAccessSpecifier::Public)
+			{
+				OutPendingNames.Add(Entry->GetEntryName());
+			}
+		}
+		else if(UAnimNextDataInterfaceEntry* DataInterfaceEntry = Cast<UAnimNextDataInterfaceEntry>(Entry))
+		{
+			if(DataInterfaceEntry->DataInterface)
+			{
+				UAnimNextDataInterface_EditorData* EditorData = UE::AnimNext::UncookedOnly::FUtils::GetEditorData<UAnimNextDataInterface_EditorData>(DataInterfaceEntry->DataInterface.Get());
+				GetPendingNamesRecursive(EditorData, OutPendingNames);
+			}
+		}
+	}
+}
+
+void SAddVariablesDialog::GetPendingNames(TArray<FName>& OutPendingNames) const
+{
+	for(UAnimNextRigVMAssetEditorData* EditorData : AssetEditorDatas)
+	{
+		GetPendingNamesRecursive(EditorData, OutPendingNames);
+	}
+}
+
 void SAddVariablesDialog::AddEntry(const FAnimNextParamType& InParamType)
 {
 	const UAnimNextVariableSettings* Settings = GetDefault<UAnimNextVariableSettings>();
-	
+
 	TArray<FName> PendingNames;
-	PendingNames.Reserve(Entries.Num());
-	for(const TSharedRef<FVariableToAddEntry>& QueuedAdd : Entries)
+	PendingNames.Reserve(RootEntries.Num());
+	for(const TSharedRef<FEntry>& QueuedAdd : RootEntries)
 	{
-		PendingNames.Add(QueuedAdd->Name);
+		switch(QueuedAdd->EntryType)
+		{
+		case EEntryType::Variable:
+			PendingNames.Add(StaticCastSharedRef<FVariableToAddEntry>(QueuedAdd)->Name);
+			break;
+		case EEntryType::DataInterface:
+			for(const TSharedRef<FEntry>& SubEntry : QueuedAdd->Children)
+			{
+				check(SubEntry->EntryType == EEntryType::Variable);
+				PendingNames.Add(StaticCastSharedRef<FVariableToAddEntry>(SubEntry)->Name);
+			}
+			break;
+		}
 	}
-	FName ParameterName = FUtils::GetNewParameterName(Settings->GetLastVariableName(), Asset, PendingNames);
-	Entries.Add(MakeShared<FVariableToAddEntry>(InParamType.IsValid() ? InParamType : Settings->GetLastVariableType(), ParameterName));
+
+	GetPendingNames(PendingNames);
+
+	FName VariableName = GetNewVariableName(Settings->GetLastVariableName(), PendingNames);
+	RootEntries.Add(MakeShared<FVariableToAddEntry>(InParamType.IsValid() ? InParamType : Settings->GetLastVariableType(), VariableName, SharedThis(this)));
 
 	RefreshEntries();
 }
 
+void SAddVariablesDialog::AddDataInterface(UAnimNextDataInterface* InDataInterface)
+{
+	
+	TSharedRef<FDataInterfaceToAddEntry> NewEntry = MakeShared<FDataInterfaceToAddEntry>(InDataInterface, SharedThis(this));
+
+	auto AddVariable = [this, &NewEntry](UAnimNextVariableEntry* InVariableEntry)
+	{
+		if(InVariableEntry->GetExportAccessSpecifier() == EAnimNextExportAccessSpecifier::Public)
+		{
+			TSharedRef<FVariableToAddEntry> NewSubEntry = MakeShared<FVariableToAddEntry>(InVariableEntry->GetType(), InVariableEntry->GetEntryName(), SharedThis(this));
+			NewSubEntry->Parent = NewEntry;
+			NewSubEntry->bIsNew = false;
+			NewEntry->Children.Add(NewSubEntry);
+		}
+	};
+
+	auto AddDataInterfaceInternal = [this, &AddVariable](UAnimNextDataInterface* InDataInterfaceToAdd, auto& InAddDataInterfaceInternal) -> void
+	{
+		// Add the child entries
+		UAnimNextDataInterface_EditorData* EditorData = UncookedOnly::FUtils::GetEditorData<UAnimNextDataInterface_EditorData>(InDataInterfaceToAdd);
+		for(UAnimNextRigVMAssetEntry* AssetEntry : EditorData->Entries)
+		{
+			if(UAnimNextVariableEntry* VariableEntry = Cast<UAnimNextVariableEntry>(AssetEntry))
+			{
+				AddVariable(VariableEntry);
+			}
+			else if(UAnimNextDataInterfaceEntry* DataInterfaceEntry = Cast<UAnimNextDataInterfaceEntry>(AssetEntry))
+			{
+				if(DataInterfaceEntry->GetDataInterface())
+				{
+					InAddDataInterfaceInternal(DataInterfaceEntry->GetDataInterface(), InAddDataInterfaceInternal);
+				}
+			}
+		}
+	};
+
+	// Add the selected data interface and its recursive dependents
+	AddDataInterfaceInternal(InDataInterface, AddDataInterfaceInternal);
+
+	if(NewEntry->Children.Num() > 0)
+	{
+		RootEntries.Add(NewEntry);
+		RefreshEntries();
+	}
+}
+
+void SAddVariablesDialog::RefreshCanCreate()
+{
+	bCanCreateVariables = true;
+	for(TSharedRef<FEntry> Entry : RootEntries)
+	{
+		FText Reason;
+		if(!Entry->IsValid(Reason))
+		{
+			CreateErrorMessage = Reason;
+			bCanCreateVariables = false;
+			return;
+		}
+
+		for(const TSharedRef<FEntry>& SubEntry : Entry->Children)
+		{
+			if(!SubEntry->IsValid(Reason))
+			{
+				CreateErrorMessage = Reason;
+				bCanCreateVariables = false;
+				return;
+			}
+		}
+	}
+}
+
 void SAddVariablesDialog::RefreshEntries()
 {
-	EntriesList->RequestListRefresh();
+	EntriesTree->RequestTreeRefresh();
+
+	RefreshCanCreate();
 }
 
-class SVariableToAdd : public SMultiColumnTableRow<TSharedRef<SAddVariablesDialog::FVariableToAddEntry>>
+class SVariableToAdd : public SMultiColumnTableRow<TSharedRef<SAddVariablesDialog::FEntry>>
 {
 	SLATE_BEGIN_ARGS(SVariableToAdd) {}
 	SLATE_END_ARGS()
 
-	void Construct(const FArguments& InArgs, const TSharedRef<STableViewBase>& InOwnerTableView, TSharedRef<SAddVariablesDialog::FVariableToAddEntry> InEntry, TSharedRef<SAddVariablesDialog> InDialog)
+	void Construct(const FArguments& InArgs, const TSharedRef<STableViewBase>& InOwnerTableView, TSharedRef<SAddVariablesDialog::FEntry> InEntry, TSharedRef<SAddVariablesDialog> InDialog)
 	{
 		Entry = InEntry;
 		WeakDialog = InDialog;
 		
-		SMultiColumnTableRow<TSharedRef<SAddVariablesDialog::FVariableToAddEntry>>::Construct( SMultiColumnTableRow<TSharedRef<SAddVariablesDialog::FVariableToAddEntry>>::FArguments(), InOwnerTableView);
+		SMultiColumnTableRow<TSharedRef<SAddVariablesDialog::FEntry>>::Construct( SMultiColumnTableRow<TSharedRef<SAddVariablesDialog::FEntry>>::FArguments(), InOwnerTableView);
 	}
 
 	virtual TSharedRef<SWidget> GenerateWidgetForColumn(const FName& InColumnName) override
@@ -225,171 +510,246 @@ class SVariableToAdd : public SMultiColumnTableRow<TSharedRef<SAddVariablesDialo
 
 		if(InColumnName == Column_Name)
 		{
-			TSharedPtr<SInlineEditableTextBlock> EditableText;
-			TSharedRef<SWidget> Widget =
-				SNew(SBox)
-				.HAlign(HAlign_Left)
-				.VAlign(VAlign_Center)
-				[
-					SAssignNew(EditableText, SInlineEditableTextBlock)
-					.Font(IDetailLayoutBuilder::GetDetailFont())
-					.IsSelected(this, &SVariableToAdd::IsSelectedExclusively)
-					.ToolTipText(LOCTEXT("NameTooltip", "The name of the new variable"))
-					.Text_Lambda([this]()
-					{
-						return FText::FromName(Entry->Name);
-					})
-					.OnTextCommitted_Lambda([this](const FText& InText, ETextCommit::Type InCommitType)
+			TSharedPtr<SWidget> EntryWidget = SNullWidget::NullWidget;
+			if(Entry->EntryType == SAddVariablesDialog::EEntryType::Variable)
+			{
+				EditableTextBox = SAssignNew(EntryWidget, SEditableTextBox)
+				.Font(IDetailLayoutBuilder::GetDetailFont())
+				.IsReadOnly_Lambda([this]()
+				{
+					// Cant rename entries from data interfaces (i.e. with parents)
+					check(Entry->EntryType == SAddVariablesDialog::EEntryType::Variable);
+					return Entry->Parent.IsValid();
+				})
+				.ToolTipText_Lambda([this]()
+				{
+					if(!CurrentError.IsEmpty())
 					{
-						Entry->Name = *InText.ToString();
-
-						UAnimNextVariableSettings* Settings = GetMutableDefault<UAnimNextVariableSettings>();
-						Settings->SetLastVariableName(Entry->Name);
-					})
-					.OnVerifyTextChanged_Lambda([this](const FText& InNewText, FText& OutErrorText)
+						return CurrentError;
+					}
+					return LOCTEXT("NameTooltip", "The name of the new variable");
+				})
+				.Text_Lambda([this]()
+				{
+					check(Entry->EntryType == SAddVariablesDialog::EEntryType::Variable);
+					return FText::FromName(StaticCastSharedPtr<SAddVariablesDialog::FVariableToAddEntry>(Entry)->Name);
+				})
+				.OnTextCommitted_Lambda([this](const FText& InText, ETextCommit::Type InCommitType)
+				{
+					check(Entry->EntryType == SAddVariablesDialog::EEntryType::Variable);
+					TSharedPtr<SAddVariablesDialog::FVariableToAddEntry> VariableEntry = StaticCastSharedPtr<SAddVariablesDialog::FVariableToAddEntry>(Entry);
+					VariableEntry->Name = *InText.ToString();
+
+					UAnimNextVariableSettings* Settings = GetMutableDefault<UAnimNextVariableSettings>();
+					Settings->SetLastVariableName(VariableEntry->Name);
+					RefreshErrors();
+					WeakDialog.Pin()->RefreshCanCreate();
+				});
+			}
+			else if(Entry->EntryType == SAddVariablesDialog::EEntryType::DataInterface)
+			{
+				SAssignNew(EntryWidget, STextBlock)
+				.Font(IDetailLayoutBuilder::GetDetailFontBold())
+				.Text_Lambda([this]()
+				{
+					TSharedPtr<SAddVariablesDialog::FDataInterfaceToAddEntry> InterfaceEntry = StaticCastSharedPtr<SAddVariablesDialog::FDataInterfaceToAddEntry>(Entry);
+					FName Name = InterfaceEntry->DataInterface ? InterfaceEntry->DataInterface->GetFName() : NAME_None; 
+					return FText::FromName(Name);
+				});
+				
+				if(Entry->bIsNew)
+				{
+					TWeakPtr<STreeView<TSharedRef<SAddVariablesDialog::FEntry>>> WeakTreeView = WeakDialog.Pin()->EntriesTree;
+					EntryWidget->RegisterActiveTimer(1/60.0f, FWidgetActiveTimerDelegate::CreateSPLambda(WeakDialog.Pin()->EntriesTree.Get(), [WeakTreeView, WeakEntry = TWeakPtr<SAddVariablesDialog::FEntry>(Entry)](double, float)
 					{
-						const FString NewString = InNewText.ToString();
-
-						if(!FUtils::IsValidParameterNameString(NewString, OutErrorText))
+						TSharedPtr<STreeView<TSharedRef<SAddVariablesDialog::FEntry>>> PinnedTreeView = WeakTreeView.Pin();
+						TSharedPtr<SAddVariablesDialog::FEntry> PinnedEntry = WeakEntry.Pin();
+						if(PinnedTreeView.IsValid() && PinnedEntry.IsValid())
 						{
-							return false;
+							PinnedTreeView->SetItemExpansion(PinnedEntry.ToSharedRef(), true);
 						}
+						return EActiveTimerReturnType::Stop;
+					}));
 
-						if(TSharedPtr<SAddVariablesDialog> Dialog = WeakDialog.Pin())
-						{
-							const FName Name(*NewString);
-							if(FUtils::DoesParameterNameExistInAsset(Name, Dialog->Asset))
-							{
-								OutErrorText = LOCTEXT("Error_NameExists", "This name already exists in the project");
-								return false;
-							}
-
-							return true;
-						}
+					Entry->bIsNew = false;
+				}
+			}
 
-						return false;
-					})
+			TSharedRef<SWidget> Widget =
+				SNew(SHorizontalBox)
+				+SHorizontalBox::Slot()
+				.AutoWidth()
+				[
+					SNew(SExpanderArrow, SharedThis(this))
+				]
+				+SHorizontalBox::Slot()
+				.AutoWidth()
+				[
+					SNew(SBox)
+					.HAlign(HAlign_Left)
+					.VAlign(VAlign_Center)
+					[
+						EntryWidget.ToSharedRef()
+					]
 				];
 
-			if(Entry->bIsNew)
-			{
-				EditableText->RegisterActiveTimer(1/60.0f, FWidgetActiveTimerDelegate::CreateSPLambda(EditableText.Get(), [WeakEditableText = TWeakPtr<SInlineEditableTextBlock>(EditableText)](double, float)
-				{
-					if(TSharedPtr<SInlineEditableTextBlock> PinnedEditableText = WeakEditableText.Pin())
-					{
-						PinnedEditableText->EnterEditingMode();
-					}
-					return EActiveTimerReturnType::Stop;
-				}));
-	
-				Entry->bIsNew = false;
-			}
-
 			return Widget;
 		}
 		else if(InColumnName == Column_Type)
 		{
-			auto GetPinInfo = [this]()
-			{
-				return UncookedOnly::FUtils::GetPinTypeFromParamType(Entry->Type);
-			};
-
-			auto PinInfoChanged = [this](const FEdGraphPinType& PinType)
-			{
-				Entry->Type = UncookedOnly::FUtils::GetParamTypeFromPinType(PinType);
-
-				UAnimNextVariableSettings* Settings = GetMutableDefault<UAnimNextVariableSettings>();
-				Settings->SetLastVariableType(Entry->Type);
-			};
-			
-			auto GetFilteredVariableTypeTree = [this](TArray<TSharedPtr<UEdGraphSchema_K2::FPinTypeTreeInfo>>& TypeTree, ETypeTreeFilter TypeTreeFilter)
+			if(Entry->EntryType == SAddVariablesDialog::EEntryType::Variable)
 			{
-				FUtils::GetFilteredVariableTypeTree(TypeTree, TypeTreeFilter);
+				auto GetPinInfo = [this]()
+				{
+					check(Entry->EntryType == SAddVariablesDialog::EEntryType::Variable);
+					return UncookedOnly::FUtils::GetPinTypeFromParamType(StaticCastSharedPtr<SAddVariablesDialog::FVariableToAddEntry>(Entry)->Type);
+				};
 
-				if(TSharedPtr<SAddVariablesDialog> Dialog = WeakDialog.Pin())
+				auto PinInfoChanged = [this](const FEdGraphPinType& PinType)
+				{
+					check(Entry->EntryType == SAddVariablesDialog::EEntryType::Variable);
+					TSharedPtr<SAddVariablesDialog::FVariableToAddEntry> VariableEntry = StaticCastSharedPtr<SAddVariablesDialog::FVariableToAddEntry>(Entry);
+					VariableEntry->Type = UncookedOnly::FUtils::GetParamTypeFromPinType(PinType);
+
+					UAnimNextVariableSettings* Settings = GetMutableDefault<UAnimNextVariableSettings>();
+					Settings->SetLastVariableType(VariableEntry->Type);
+				};
+				
+				auto GetFilteredVariableTypeTree = [this](TArray<TSharedPtr<UEdGraphSchema_K2::FPinTypeTreeInfo>>& TypeTree, ETypeTreeFilter TypeTreeFilter)
 				{
-					if(Dialog->OnFilterVariableType.IsBound())
+					FUtils::GetFilteredVariableTypeTree(TypeTree, TypeTreeFilter);
+
+					if(TSharedPtr<SAddVariablesDialog> Dialog = WeakDialog.Pin())
 					{
-						auto IsPinTypeAllowed = [&Dialog](const FEdGraphPinType& InType)
+						if(Dialog->OnFilterVariableType.IsBound())
 						{
-							FAnimNextParamType Type = UncookedOnly::FUtils::GetParamTypeFromPinType(InType);
-							if(Type.IsValid())
+							auto IsPinTypeAllowed = [&Dialog](const FEdGraphPinType& InType)
 							{
-								return Dialog->OnFilterVariableType.Execute(Type) == EFilterVariableResult::Include;
-							}
-							return false;
-						};
-
-						// Additionally filter by allowed types
-						for (int32 Index = 0; Index < TypeTree.Num(); )
-						{
-							TSharedPtr<UEdGraphSchema_K2::FPinTypeTreeInfo>& PinType = TypeTree[Index];
+								FAnimNextParamType Type = UncookedOnly::FUtils::GetParamTypeFromPinType(InType);
+								if(Type.IsValid())
+								{
+									return Dialog->OnFilterVariableType.Execute(Type) == EFilterVariableResult::Include;
+								}
+								return false;
+							};
 
-							if (PinType->Children.Num() == 0 && !IsPinTypeAllowed(PinType->GetPinType(/*bForceLoadSubCategoryObject*/false)))
+							// Additionally filter by allowed types
+							for (int32 Index = 0; Index < TypeTree.Num(); )
 							{
-								TypeTree.RemoveAt(Index);
-								continue;
-							}
+								TSharedPtr<UEdGraphSchema_K2::FPinTypeTreeInfo>& PinType = TypeTree[Index];
 
-							for (int32 ChildIndex = 0; ChildIndex < PinType->Children.Num(); )
-							{
-								TSharedPtr<UEdGraphSchema_K2::FPinTypeTreeInfo> Child = PinType->Children[ChildIndex];
-								if (Child.IsValid())
+								if (PinType->Children.Num() == 0 && !IsPinTypeAllowed(PinType->GetPinType(/*bForceLoadSubCategoryObject*/false)))
 								{
-									if (!IsPinTypeAllowed(Child->GetPinType(/*bForceLoadSubCategoryObject*/false)))
+									TypeTree.RemoveAt(Index);
+									continue;
+								}
+
+								for (int32 ChildIndex = 0; ChildIndex < PinType->Children.Num(); )
+								{
+									TSharedPtr<UEdGraphSchema_K2::FPinTypeTreeInfo> Child = PinType->Children[ChildIndex];
+									if (Child.IsValid())
 									{
-										PinType->Children.RemoveAt(ChildIndex);
-										continue;
+										if (!IsPinTypeAllowed(Child->GetPinType(/*bForceLoadSubCategoryObject*/false)))
+										{
+											PinType->Children.RemoveAt(ChildIndex);
+											continue;
+										}
 									}
+									++ChildIndex;
 								}
-								++ChildIndex;
-							}
 
-							++Index;
+								++Index;
+							}
 						}
 					}
-				}
-			};
-
-			return
-				SNew(SBox)
-				.HAlign(HAlign_Left)
-				.VAlign(VAlign_Center)
-				[
-					SNew(SPinTypeSelector, FGetPinTypeTree::CreateLambda(GetFilteredVariableTypeTree))
-						.TargetPinType_Lambda(GetPinInfo)
-						.OnPinTypeChanged_Lambda(PinInfoChanged)
-						.Schema(GetDefault<UPropertyBagSchema>())
-						.bAllowArrays(true)
-						.TypeTreeFilter(ETypeTreeFilter::None)
-						.Font(IDetailLayoutBuilder::GetDetailFont())
-				];
+				};
+					
+				return
+					SNew(SBox)
+					.HAlign(HAlign_Left)
+					.VAlign(VAlign_Center)
+					[
+						SNew(SPinTypeSelector, FGetPinTypeTree::CreateLambda(GetFilteredVariableTypeTree))
+							.IsEnabled_Lambda([this]()
+							{
+								// Cant rename entries from data interfaces (i.e. with parents)
+								check(Entry->EntryType == SAddVariablesDialog::EEntryType::Variable);
+								return !Entry->Parent.IsValid();
+							})
+							.TargetPinType_Lambda(GetPinInfo)
+							.OnPinTypeChanged_Lambda(PinInfoChanged)
+							.Schema(GetDefault<UPropertyBagSchema>())
+							.bAllowArrays(true)
+							.TypeTreeFilter(ETypeTreeFilter::None)
+							.Font(IDetailLayoutBuilder::GetDetailFont())
+					];
+			}
 		}
 
 		return SNullWidget::NullWidget;
 	}
 
-	TSharedPtr<SAddVariablesDialog::FVariableToAddEntry> Entry;
+	void RefreshErrors()
+	{
+		if(EditableTextBox.IsValid())
+		{
+			FText Reason;
+			if(!Entry->IsValid(Reason))
+			{
+				EditableTextBox->SetError(Reason);
+				CurrentError = Reason;
+			}
+			else
+			{
+				EditableTextBox->SetError(FText::GetEmpty());
+				CurrentError = FText::GetEmpty();
+			}
+		}
+		else
+		{
+			CurrentError = FText::GetEmpty();
+		}
+	}
+
+	FText CurrentError;
+	TSharedPtr<SEditableTextBox> EditableTextBox;
+	TSharedPtr<SAddVariablesDialog::FEntry> Entry;
 	TWeakPtr<SAddVariablesDialog> WeakDialog;
 };
 
-TSharedRef<ITableRow> SAddVariablesDialog::HandleGenerateRow(TSharedRef<FVariableToAddEntry> InEntry, const TSharedRef<STableViewBase>& InOwnerTable)
+TSharedRef<ITableRow> SAddVariablesDialog::HandleGenerateRow(TSharedRef<FEntry> InEntry, const TSharedRef<STableViewBase>& InOwnerTable)
 {
-	return SNew(SVariableToAdd, InOwnerTable, InEntry, SharedThis(this));
+	TSharedRef<SVariableToAdd> RowWidget = SNew(SVariableToAdd, InOwnerTable, InEntry, SharedThis(this));
+	RowWidget->RefreshErrors();
+	return RowWidget;
 }
 
-bool SAddVariablesDialog::ShowModal(TArray<FVariableToAdd>& OutParameters)
+void SAddVariablesDialog::HandleGetChildren(TSharedRef<FEntry> InEntry, TArray<TSharedRef<FEntry>>& OutChildren)
+{
+	OutChildren = InEntry->Children;
+}
+
+bool SAddVariablesDialog::ShowModal(TArray<FVariableToAdd>& OutVariables, TArray<FDataInterfaceToAdd>& OutDataInterfaces)
 {
 	FSlateApplication::Get().AddModalWindow(SharedThis(this), FGlobalTabmanager::Get()->GetRootWindow());
 
 	if(bOKPressed)
 	{
 		bool bHasValid = false;
-		for(TSharedRef<FVariableToAddEntry>& Entry : Entries)
+		for(TSharedRef<FEntry>& Entry : RootEntries)
 		{
-			if(Entry->IsValid())
+			FText Reason;
+			if(Entry->IsValid(Reason))
 			{
-				OutParameters.Add(*Entry);
+				switch(Entry->EntryType)
+				{
+				case EEntryType::Variable:
+					OutVariables.Add(*StaticCastSharedRef<FVariableToAddEntry>(Entry));
+					break;
+				case EEntryType::DataInterface:
+					OutDataInterfaces.Add(*StaticCastSharedRef<FDataInterfaceToAddEntry>(Entry));
+					break;
+				}
 				bHasValid = true;
 			}
 		}
@@ -398,7 +758,7 @@ bool SAddVariablesDialog::ShowModal(TArray<FVariableToAdd>& OutParameters)
 	return false;
 }
 
-TSharedRef<SWidget> SAddVariablesDialog::HandleGetAddVariableMenuContent(TSharedPtr<FVariableToAddEntry> InEntry)
+TSharedRef<SWidget> SAddVariablesDialog::HandleGetAddVariableMenuContent(TSharedPtr<FEntry> InEntry)
 {
 	using namespace AddVariablesDialog;
 
diff --git a/Engine/Plugins/Experimental/AnimNext/Source/AnimNextEditor/Private/Variables/SAddVariablesDialog.h b/Engine/Plugins/Experimental/AnimNext/Source/AnimNextEditor/Private/Variables/SAddVariablesDialog.h
--- a/Engine/Plugins/Experimental/AnimNext/Source/AnimNextEditor/Private/Variables/SAddVariablesDialog.h
+++ b/Engine/Plugins/Experimental/AnimNext/Source/AnimNextEditor/Private/Variables/SAddVariablesDialog.h
@@ -4,31 +4,32 @@
 
 #include "Param/ParamType.h"
 #include "Widgets/SWindow.h"
-#include "AssetRegistry/AssetData.h"
-#include "Widgets/Views/SListView.h"
+#include "Widgets/Views/STreeView.h"
 
+class UAnimNextDataInterface;
 class SWrapBox;
 class UAnimNextModule_EditorData;
+class UAnimNextRigVMAssetEditorData;
 
 namespace UE::AnimNext::Editor
 {
 
-struct FVariableToAdd
+struct FDataInterfaceToAdd
 {
-	FVariableToAdd() = default;
+	FDataInterfaceToAdd(UAnimNextDataInterface* InDataInterface)
+		: DataInterface(InDataInterface)
+	{}
 
+	UAnimNextDataInterface* DataInterface;
+};
+
+struct FVariableToAdd
+{
 	FVariableToAdd(const FAnimNextParamType& InType, FName InName)
 		: Type(InType)
 		, Name(InName)
 	{}
 
-	bool IsValid() const
-	{
-		return Name != NAME_None && Type.IsValid(); 
-	}
-
-	bool IsValid(FText& OutReason) const;
-
 	// Type
 	FAnimNextParamType Type;
 
@@ -62,44 +63,98 @@ public:
 	/** Initial variable type to use */
 	SLATE_ARGUMENT(FAnimNextParamType, InitialParamType)
 
+	/** Whether we should add an initial variable */
+	SLATE_ARGUMENT(bool, ShouldAddInitialVariable)
+
 	SLATE_END_ARGS()
 
-	void Construct(const FArguments& InArgs, const FAssetData& InAsset);
+	void Construct(const FArguments& InArgs, const TArray<UAnimNextRigVMAssetEditorData*>& InAssetEditorDatas);
+
+	bool ShowModal(TArray<FVariableToAdd>& OutVariables, TArray<FDataInterfaceToAdd>& OutDataInterfaces);
+
+	enum class EEntryType : uint8
+	{
+		Variable,
+		DataInterface,
+	};
+
+	struct FEntry
+	{
+		FEntry(EEntryType InEntryType, const TSharedRef<SAddVariablesDialog>& InDialog)
+			: Dialog(InDialog)
+			, EntryType(InEntryType)
+		{}
 
-	bool ShowModal(TArray<FVariableToAdd>& OutVariables);
+		virtual ~FEntry() = default;
+
+		virtual bool IsValid(FText& OutReason) const = 0;
+
+		TWeakPtr<SAddVariablesDialog> Dialog;
+		TWeakPtr<FEntry> Parent;
+		TArray<TSharedRef<FEntry>> Children;
+		EEntryType EntryType;
+		bool bIsNew = true;
+	};
 
 private:
 	virtual FReply OnKeyDown(const FGeometry& MyGeometry, const FKeyEvent& InKeyEvent) override;
 
+	void AddDataInterface(UAnimNextDataInterface* InDataInterface);
+
 	void AddEntry(const FAnimNextParamType& InParamType = FAnimNextParamType());
 
 	void RefreshEntries();
 
-	struct FVariableToAddEntry : FVariableToAdd
+	void RefreshCanCreate();
+
+	void DeleteSelectedItems();
+
+	struct FDataInterfaceToAddEntry : FEntry, FDataInterfaceToAdd
 	{
-		FVariableToAddEntry() = default;
+		FDataInterfaceToAddEntry(UAnimNextDataInterface* InDataInterface, const TSharedRef<SAddVariablesDialog>& InDialog)
+			: FEntry(EEntryType::DataInterface, InDialog)
+			, FDataInterfaceToAdd(InDataInterface)
+		{
+		}
+
+		virtual bool IsValid(FText& OutReason) const override;
+	};
 
-		FVariableToAddEntry(const FAnimNextParamType& InType, FName InName)
-			: FVariableToAdd(InType, InName)
+	struct FVariableToAddEntry : FEntry, FVariableToAdd
+	{
+		FVariableToAddEntry(const FAnimNextParamType& InType, FName InName, const TSharedRef<SAddVariablesDialog>& InDialog)
+			: FEntry(EEntryType::Variable, InDialog)
+			, FVariableToAdd(InType, InName)
 		{}
 
-		bool bIsNew = true;
+		virtual bool IsValid(FText& OutReason) const override;
 	};
 
-	TSharedRef<ITableRow> HandleGenerateRow(TSharedRef<FVariableToAddEntry> InEntry, const TSharedRef<STableViewBase>& InOwnerTable);
+	TSharedRef<ITableRow> HandleGenerateRow(TSharedRef<FEntry> InEntry, const TSharedRef<STableViewBase>& InOwnerTable);
 
-	TSharedRef<SWidget> HandleGetAddVariableMenuContent(TSharedPtr<FVariableToAddEntry> InEntry);
+	void HandleGetChildren(TSharedRef<FEntry> InEntry, TArray<TSharedRef<FEntry>>& OutChildren);
 	
+	TSharedRef<SWidget> HandleGetAddVariableMenuContent(TSharedPtr<FEntry> InEntry);
+
+	void GetPendingNamesRecursive(UAnimNextRigVMAssetEditorData* InEditorData, TArray<FName>& OutPendingNames) const;
+
+	void GetPendingNames(TArray<FName>& OutPendingNames) const;
+
 private:
 	friend class SVariableToAdd;
 
-	TSharedPtr<SListView<TSharedRef<FVariableToAddEntry>>> EntriesList;
+	TSharedPtr<STreeView<TSharedRef<FEntry>>> EntriesTree;
 
-	TArray<TSharedRef<FVariableToAddEntry>> Entries;
+	// Root entries of the tree
+	TArray<TSharedRef<FEntry>> RootEntries;
 
 	FOnFilterVariableType OnFilterVariableType;
 
-	FAssetData Asset;
+	TArray<UAnimNextRigVMAssetEditorData*> AssetEditorDatas;
+
+	FText CreateErrorMessage;
+
+	bool bCanCreateVariables = false;
 
 	bool bOKPressed = false;
 };
diff --git a/Engine/Plugins/Experimental/AnimNext/Source/AnimNextEditor/Private/Variables/SVariableOverride.cpp b/Engine/Plugins/Experimental/AnimNext/Source/AnimNextEditor/Private/Variables/SVariableOverride.cpp
new file mode 100644
--- /dev/null
+++ b/Engine/Plugins/Experimental/AnimNext/Source/AnimNextEditor/Private/Variables/SVariableOverride.cpp
@@ -0,0 +1,221 @@
+﻿// Copyright Epic Games, Inc. All Rights Reserved.
+
+#include "SVariableOverride.h"
+
+#include "ScopedTransaction.h"
+#include "ToolMenus.h"
+#include "VariableOverrideCommands.h"
+#include "Framework/Application/SlateApplication.h"
+#include "Styling/StyleColors.h"
+
+#define LOCTEXT_NAMESPACE "SVariableOverride"
+
+namespace UE::AnimNext::Editor
+{
+
+static FLazyName VariableOverrideMenuName("AnimNextVariableOverrides");
+
+void SVariableOverride::Construct(const FArguments& InArgs, UAnimNextDataInterfaceEntry* InDataInterfaceEntry, FName InVariableName)
+{
+	WeakDataInterfaceEntry = InDataInterfaceEntry;
+	VariableName = InVariableName;
+	OverrideStatus = InArgs._OverrideStatus;
+
+	SImage::Construct(
+	SImage::FArguments()
+		.ColorAndOpacity(this, &SVariableOverride::GetColor)
+		.Image(this, &SVariableOverride::GetBrush)
+	);
+
+	SetToolTipText(MakeAttributeSP(this, &SVariableOverride::GetToolTipText));
+}
+
+void SVariableOverride::SetupMenu()
+{
+	const FVariableOverrideCommands& Commands = FVariableOverrideCommands::Get();
+	TSharedRef<FUICommandList> ActionList = MakeShared<FUICommandList>();
+	UToolMenu* ToolMenu = UToolMenus::Get()->RegisterMenu(VariableOverrideMenuName,NAME_None, EMultiBoxType::Menu, false);
+	ToolMenu->bShouldCloseWindowAfterMenuSelection = true;
+	ToolMenu->bCloseSelfOnly = true;
+
+	ActionList->MapAction(
+		Commands.OverrideVariable,
+		FExecuteAction::CreateSP(this, &SVariableOverride::OverrideVariable),
+		FCanExecuteAction::CreateSP(this, &SVariableOverride::CanOverrideVariable),
+		EUIActionRepeatMode::RepeatDisabled);
+	
+	ActionList->MapAction(
+		Commands.ClearOverride,
+		FExecuteAction::CreateSP(this, &SVariableOverride::ClearOverride),
+		FCanExecuteAction::CreateSP(this, &SVariableOverride::CanClearOverride),
+		EUIActionRepeatMode::RepeatDisabled);
+
+	ActionList->MapAction(
+		Commands.ResetPropertyToDefault,
+		FExecuteAction::CreateSP(this, &SVariableOverride::ResetToDefault),
+		FCanExecuteAction::CreateSP(this, &SVariableOverride::CanResetToDefault),
+		FIsActionChecked(),
+		FIsActionButtonVisible::CreateSP(this, &SVariableOverride::IsResetToDefaultVisible),
+		EUIActionRepeatMode::RepeatDisabled);
+
+	static const FName OverrideSectionName = TEXT("Overrides");
+	FToolMenuSection& OverrideSection = ToolMenu->AddSection(OverrideSectionName);
+	OverrideSection.AddMenuEntryWithCommandList(Commands.OverrideVariable, ActionList);
+	OverrideSection.AddMenuEntryWithCommandList(Commands.ResetPropertyToDefault, ActionList);
+	OverrideSection.AddMenuEntryWithCommandList(Commands.ClearOverride, ActionList);
+}
+
+FReply SVariableOverride::OnMouseButtonDown(const FGeometry& MyGeometry, const FPointerEvent& MouseEvent)
+{
+	if (MouseEvent.GetEffectingButton() != EKeys::LeftMouseButton)
+	{
+		return FReply::Unhandled();
+	}
+
+	FSlateApplication::Get().PushMenu(
+		SharedThis(this),
+		FWidgetPath(),
+		GetMenuContent(),
+		FSlateApplication::Get().GetCursorPos(),
+		FPopupTransitionEffect(FPopupTransitionEffect::TypeInPopup));
+
+	return FReply::Handled();
+}
+
+const FSlateBrush* SVariableOverride::GetBrush() const
+{
+	const FSlateBrush* Image = nullptr;
+	const FSlateBrush* HoveredImage = nullptr;
+
+	switch(OverrideStatus.Get())
+	{
+	case EAnimNextDataInterfaceValueOverrideStatus::NotOverridden:
+		Image = FAppStyle::GetBrush( "DetailsView.OverrideNone");
+		HoveredImage = FAppStyle::GetBrush( "DetailsView.OverrideNone.Hovered");
+		break;
+	case EAnimNextDataInterfaceValueOverrideStatus::OverriddenInThisAsset:
+		Image = FAppStyle::GetBrush( "DetailsView.OverrideHere");
+		HoveredImage = FAppStyle::GetBrush( "DetailsView.OverrideHere.Hovered");
+		break;
+	case EAnimNextDataInterfaceValueOverrideStatus::OverriddenInParentAsset:
+		Image = FAppStyle::GetBrush( "DetailsView.OverrideInherited");
+		HoveredImage = FAppStyle::GetBrush( "DetailsView.OverrideInherited.Hovered");
+		break;
+	}
+
+	return IsHovered() ? HoveredImage : Image;
+}
+
+FSlateColor SVariableOverride::GetColor() const
+{
+	return IsHovered() ? FStyleColors::White : FStyleColors::Foreground;
+}
+
+FText SVariableOverride::GetToolTipText() const
+{
+	switch(OverrideStatus.Get())
+	{
+	case EAnimNextDataInterfaceValueOverrideStatus::NotOverridden:
+		return LOCTEXT("NotOverriddenTooltip", "This variable is not overriden");
+	case EAnimNextDataInterfaceValueOverrideStatus::OverriddenInThisAsset:
+		return LOCTEXT("OverriddenInThisAssetTooltip", "This variable is overridden in this asset.");
+	case EAnimNextDataInterfaceValueOverrideStatus::OverriddenInParentAsset:
+		return LOCTEXT("OverriddenInParentAssetTooltip", "This variable is overridden in a parent asset.");
+	}
+
+	return FText::GetEmpty();
+}
+
+TSharedRef<SWidget> SVariableOverride::GetMenuContent()
+{
+	SetupMenu();
+
+	UToolMenu* ToolMenu = UToolMenus::Get()->FindMenu(VariableOverrideMenuName);
+	check(ToolMenu);
+	return UToolMenus::Get()->GenerateWidget(ToolMenu);
+}
+
+void SVariableOverride::OverrideVariable()
+{
+	UAnimNextDataInterfaceEntry* DataInterfaceEntry = WeakDataInterfaceEntry.Get();
+	if(DataInterfaceEntry == nullptr)
+	{
+		return;
+	}
+
+	FScopedTransaction Transaction(LOCTEXT("OverrideValue", "Override Value"));
+	DataInterfaceEntry->SetValueOverrideToDefault(VariableName);
+}
+
+bool SVariableOverride::CanOverrideVariable() const
+{
+	UAnimNextDataInterfaceEntry* DataInterfaceEntry = WeakDataInterfaceEntry.Get();
+	if(DataInterfaceEntry == nullptr)
+	{
+		return false;
+	}
+
+	return !DataInterfaceEntry->HasValueOverride(VariableName);
+}
+
+void SVariableOverride::ClearOverride()
+{
+	UAnimNextDataInterfaceEntry* DataInterfaceEntry = WeakDataInterfaceEntry.Get();
+	if(DataInterfaceEntry == nullptr)
+	{
+		return;
+	}
+
+	FScopedTransaction Transaction(LOCTEXT("ClearOverrideValue", "Clear Override Value"));
+	DataInterfaceEntry->ClearValueOverride(VariableName);
+}
+
+bool SVariableOverride::CanClearOverride() const
+{
+	UAnimNextDataInterfaceEntry* DataInterfaceEntry = WeakDataInterfaceEntry.Get();
+	if(DataInterfaceEntry == nullptr)
+	{
+		return false;
+	}
+
+	return DataInterfaceEntry->HasValueOverride(VariableName);
+}
+
+void SVariableOverride::ResetToDefault()
+{
+	UAnimNextDataInterfaceEntry* DataInterfaceEntry = WeakDataInterfaceEntry.Get();
+	if(DataInterfaceEntry == nullptr)
+	{
+		return;
+	}
+
+	FScopedTransaction Transaction(LOCTEXT("ResetToDefault", "Reset to default"));
+	DataInterfaceEntry->ClearValueOverride(VariableName);
+	DataInterfaceEntry->SetValueOverrideToDefault(VariableName);
+}
+
+bool SVariableOverride::CanResetToDefault() const
+{
+	UAnimNextDataInterfaceEntry* DataInterfaceEntry = WeakDataInterfaceEntry.Get();
+	if(DataInterfaceEntry == nullptr)
+	{
+		return false;
+	}
+
+	return DataInterfaceEntry->HasValueOverride(VariableName);
+}
+
+bool SVariableOverride::IsResetToDefaultVisible() const
+{
+	UAnimNextDataInterfaceEntry* DataInterfaceEntry = WeakDataInterfaceEntry.Get();
+	if(DataInterfaceEntry == nullptr)
+	{
+		return false;
+	}
+
+	return DataInterfaceEntry->HasValueOverrideNotMatchingDefault(VariableName);
+}
+
+}
+
+#undef LOCTEXT_NAMESPACE
\ No newline at end of file
diff --git a/Engine/Plugins/Experimental/AnimNext/Source/AnimNextEditor/Private/Variables/SVariableOverride.h b/Engine/Plugins/Experimental/AnimNext/Source/AnimNextEditor/Private/Variables/SVariableOverride.h
new file mode 100644
--- /dev/null
+++ b/Engine/Plugins/Experimental/AnimNext/Source/AnimNextEditor/Private/Variables/SVariableOverride.h
@@ -0,0 +1,53 @@
+﻿// Copyright Epic Games, Inc. All Rights Reserved.
+
+#pragma once
+
+#include "Entries/AnimNextDataInterfaceEntry.h"
+#include "Widgets/Images/SImage.h"
+
+namespace UE::AnimNext::Editor
+{
+
+// Widget used to display and manipulate variable overrides
+class SVariableOverride : public SImage
+{
+	SLATE_BEGIN_ARGS(SVariableOverride) {}
+
+	SLATE_ATTRIBUTE(EAnimNextDataInterfaceValueOverrideStatus, OverrideStatus)
+
+	SLATE_END_ARGS()
+
+	void Construct(const FArguments& InArgs, UAnimNextDataInterfaceEntry* InDataInterfaceEntry, FName InVariableName);
+
+	void SetupMenu();
+
+	virtual FReply OnMouseButtonDown(const FGeometry& MyGeometry, const FPointerEvent& MouseEvent) override;
+
+	TSharedRef<SWidget> GetMenuContent();
+
+	const FSlateBrush* GetBrush() const;
+
+	FSlateColor GetColor() const;
+
+	FText GetToolTipText() const;
+
+	void OverrideVariable();
+
+	bool CanOverrideVariable() const;
+
+	void ClearOverride();
+
+	bool CanClearOverride() const;
+
+	void ResetToDefault();
+
+	bool CanResetToDefault() const;
+
+	bool IsResetToDefaultVisible() const;
+
+	TAttribute<EAnimNextDataInterfaceValueOverrideStatus> OverrideStatus;
+	TWeakObjectPtr<UAnimNextDataInterfaceEntry> WeakDataInterfaceEntry;
+	FName VariableName;
+};
+
+}
diff --git a/Engine/Plugins/Experimental/AnimNext/Source/AnimNextEditor/Private/Variables/SVariablesView.cpp b/Engine/Plugins/Experimental/AnimNext/Source/AnimNextEditor/Private/Variables/SVariablesView.cpp
--- a/Engine/Plugins/Experimental/AnimNext/Source/AnimNextEditor/Private/Variables/SVariablesView.cpp
+++ b/Engine/Plugins/Experimental/AnimNext/Source/AnimNextEditor/Private/Variables/SVariablesView.cpp
@@ -125,23 +125,23 @@ FReply SVariablesOutliner::HandleAddVariablesClicked()
 	}
 
 	TSharedRef<SAddVariablesDialog> AddVariablesDialog =
-		SNew(SAddVariablesDialog, AssetsToAddTo.Num() == 1 ? FAssetData(UncookedOnly::FUtils::GetAsset<UAnimNextRigVMAsset>(AssetsToAddTo[0])) : FAssetData());
+		SNew(SAddVariablesDialog, AssetsToAddTo);
 
 	TArray<FVariableToAdd> VariablesToAdd;
-	if(AddVariablesDialog->ShowModal(VariablesToAdd))
+	TArray<FDataInterfaceToAdd> DataInterfacesToAdd;
+	if(AddVariablesDialog->ShowModal(VariablesToAdd, DataInterfacesToAdd))
 	{
-		if(VariablesToAdd.Num() > 0)
+		FScopedTransaction Transaction(FText::Format(LOCTEXT("AddVariablesFormat", "Add {0}|plural(one=variable, other=variables)"), (DataInterfacesToAdd.Num() + VariablesToAdd.Num()) * AssetsToAddTo.Num()));
+		for(UAnimNextRigVMAssetEditorData* EditorData : AssetsToAddTo)
 		{
-			FScopedTransaction Transaction(FText::Format(LOCTEXT("AddVariablesFormat", "Add {0}|plural(one=variable, other=variables)"), VariablesToAdd.Num() * AssetsToAddTo.Num()));
-			for(UAnimNextRigVMAssetEditorData* EditorData : AssetsToAddTo)
+			for (const FVariableToAdd& VariableToAdd : VariablesToAdd)
 			{
-				for (const FVariableToAdd& VariableToAdd : VariablesToAdd)
-				{
-					if(EditorData->FindEntry(VariableToAdd.Name) == nullptr)
-					{
-						EditorData->AddVariable(VariableToAdd.Name, VariableToAdd.Type);
-					}
-				}
+				EditorData->AddVariable(VariableToAdd.Name, VariableToAdd.Type);
+			}
+
+			for (const FDataInterfaceToAdd& DataInterfaceToAdd : DataInterfacesToAdd)
+			{
+				EditorData->AddDataInterface(DataInterfaceToAdd.DataInterface);
 			}
 		}
 	}
@@ -201,6 +201,7 @@ void SVariablesOutliner::OnEditorDataModified(UAnimNextRigVMAssetEditorData* InE
 
 	switch(InType)
 	{
+	case EAnimNextEditorDataNotifType::UndoRedo:
 	case EAnimNextEditorDataNotifType::EntryAdded:
 	case EAnimNextEditorDataNotifType::EntryRemoved:
 		FullRefresh();
diff --git a/Engine/Plugins/Experimental/AnimNext/Source/AnimNextEditor/Private/Variables/VariableBindingPropertyCustomization.cpp b/Engine/Plugins/Experimental/AnimNext/Source/AnimNextEditor/Private/Variables/VariableBindingPropertyCustomization.cpp
--- a/Engine/Plugins/Experimental/AnimNext/Source/AnimNextEditor/Private/Variables/VariableBindingPropertyCustomization.cpp
+++ b/Engine/Plugins/Experimental/AnimNext/Source/AnimNextEditor/Private/Variables/VariableBindingPropertyCustomization.cpp
@@ -7,6 +7,7 @@
 #include "EditorUtils.h"
 #include "IAnimNextUncookedOnlyModule.h"
 #include "PropertyHandle.h"
+#include "SInstancedStructPicker.h"
 #include "UncookedOnlyUtils.h"
 #include "Entries/AnimNextVariableEntry.h"
 #include "Modules/ModuleManager.h"
@@ -59,48 +60,67 @@ void FVariableBindingPropertyCustomization::CustomizeHeader(TSharedRef<IProperty
 	]
 	.ValueContent()
 	[
-		SAssignNew(ValueWidget, SComboButton)
-		.ToolTipText_Lambda([this]()
-		{
-			return TooltipText;
-		})
-		.MenuContent()
-		[
-			CreateBindingWidget()
-		]
-		.ButtonContent()
+		SAssignNew(ContainerWidget, SHorizontalBox)
+		+SHorizontalBox::Slot()
+		.Padding(0.0f, 0.0f, 5.0f, 0.0f)
+		.AutoWidth()
 		[
-			SNew(SHorizontalBox)
-			+SHorizontalBox::Slot()
-			.AutoWidth()
-			.VAlign(VAlign_Center)
-			.HAlign(HAlign_Center)
-			.Padding(0.0f, 2.0f, 2.0f, 2.0f)
+			SAssignNew(ValueWidget, SComboButton)
+			.Visibility_Lambda([this]()
+			{
+				return bShowBindingSelector ? EVisibility::Visible : EVisibility::Collapsed;
+			})
+			.ToolTipText_Lambda([this]()
+			{
+				return TooltipText;
+			})
+			.MenuContent()
 			[
-				SNew(SImage)
-				.Image_Lambda([this]()
-				{
-					return Icon;
-				})
-				.ColorAndOpacity_Lambda([this]()
-				{
-					return IconColor;
-				})
+				CreateBindingWidget()
 			]
-			+SHorizontalBox::Slot()
-			.AutoWidth()
-			.VAlign(VAlign_Center)
-			.HAlign(HAlign_Left)
-			.Padding(0.0f, 2.0f, 0.0f, 2.0f)
+			.ButtonContent()
 			[
-				SNew(STextBlock)
-				.TextStyle(&FCoreStyle::Get().GetWidgetStyle<FTextBlockStyle>("SmallText"))
-				.Text_Lambda([this]()
-				{
-					return NameText;
-				})
+				SNew(SHorizontalBox)
+				+SHorizontalBox::Slot()
+				.AutoWidth()
+				.VAlign(VAlign_Center)
+				.HAlign(HAlign_Center)
+				.Padding(0.0f, 2.0f, 2.0f, 2.0f)
+				[
+					SNew(SImage)
+					.Image_Lambda([this]()
+					{
+						return Icon;
+					})
+					.ColorAndOpacity_Lambda([this]()
+					{
+						return IconColor;
+					})
+				]
+				+SHorizontalBox::Slot()
+				.AutoWidth()
+				.VAlign(VAlign_Center)
+				.HAlign(HAlign_Left)
+				.Padding(0.0f, 2.0f, 0.0f, 2.0f)
+				[
+					SNew(STextBlock)
+					.TextStyle(&FCoreStyle::Get().GetWidgetStyle<FTextBlockStyle>("SmallText"))
+					.Text_Lambda([this]()
+					{
+						return NameText;
+					})
+				]
 			]
 		]
+		+SHorizontalBox::Slot()
+		.AutoWidth()
+		[
+			SNew(SInstancedStructPicker, BindingDataHandle, InCustomizationUtils.GetPropertyUtilities())
+			.OnStructPicked_Lambda([this](const UScriptStruct* InStruct)
+			{
+				RequestRefresh();
+			})
+		]
 	];
 
 	RequestRefresh();
@@ -120,7 +140,7 @@ void FVariableBindingPropertyCustomization::RequestRefresh()
 	if(!bRefreshRequested)
 	{
 		bRefreshRequested = true;
-		ValueWidget->RegisterActiveTimer(0.0f, FWidgetActiveTimerDelegate::CreateLambda([this](double /*InCurrentTime*/, float /*InDeltaTime*/)
+		ContainerWidget->RegisterActiveTimer(0.0f, FWidgetActiveTimerDelegate::CreateLambda([this](double /*InCurrentTime*/, float /*InDeltaTime*/)
 		{
 			bRefreshRequested = false;
 			Refresh();
@@ -131,12 +151,23 @@ void FVariableBindingPropertyCustomization::RequestRefresh()
 
 void FVariableBindingPropertyCustomization::Refresh()
 {
+	bShowBindingSelector = false;
+
+	if(!Type.IsValid())
+	{
+		NameText = LOCTEXT("MultipleTypes", "Multiple Types");
+		TooltipText = LOCTEXT("MultipleTypes", "Multiple Types");
+		bShowBindingSelector = true;
+		return;
+	}
+
 	FEdGraphPinType PinType  = UncookedOnly::FUtils::GetPinTypeFromParamType(Type);
 	Icon = FBlueprintEditorUtils::GetIconFromPin(PinType, true);
 	IconColor = GetDefault<UEdGraphSchema_K2>()->GetPinTypeColor(PinType);
 
+	TOptional<const UScriptStruct*> CommonBindingStruct;
 	TOptional<TInstancedStruct<FAnimNextVariableBindingData>> CommonBindingData;
-	BindingDataHandle->EnumerateConstRawData([this, &CommonBindingData](const void* RawData, const int32 DataIndex, const int32 NumDatas)
+	BindingDataHandle->EnumerateConstRawData([this, &CommonBindingData, &CommonBindingStruct](const void* RawData, const int32 DataIndex, const int32 NumDatas)
 	{
 		const TInstancedStruct<FAnimNextVariableBindingData>& BindingData = *static_cast<const TInstancedStruct<FAnimNextVariableBindingData>*>(RawData);
 		if(!CommonBindingData.IsSet())
@@ -147,18 +178,42 @@ void FVariableBindingPropertyCustomization::Refresh()
 		{
 			// No common binding
 			CommonBindingData = TInstancedStruct<FAnimNextVariableBindingData>();
-			NameText = LOCTEXT("MutipleValues", "Multiple Values");
-			TooltipText = LOCTEXT("MutipleValues", "Multiple Values");
-			return false;
+		}
+
+		if(!CommonBindingStruct.IsSet())
+		{
+			CommonBindingStruct = BindingData.GetScriptStruct();
+		}
+		else if(CommonBindingStruct.GetValue() != BindingData.GetScriptStruct())
+		{
+			// No common struct
+			CommonBindingStruct = nullptr;
 		}
 		return true;
 	});
 
+	const UScriptStruct* BindingStruct = CommonBindingStruct.IsSet() ? CommonBindingStruct.GetValue() : nullptr;
 	TInstancedStruct<FAnimNextVariableBindingData> BindingData = CommonBindingData.IsSet() ? CommonBindingData.GetValue() : TInstancedStruct<FAnimNextVariableBindingData>();
 	if(BindingData.IsValid())
 	{
+		// Common values
 		NameText = GetBindingDisplayNameText(BindingData);
 		TooltipText = GetBindingTooltipText(BindingData);
+		bShowBindingSelector = true;
+	}
+	else if(BindingStruct != nullptr)
+	{
+		// No common values, but common struct, allow selection
+		NameText = LOCTEXT("MultipleValues", "Multiple Values");
+		TooltipText = LOCTEXT("MultipleValues", "Multiple Values");
+		bShowBindingSelector = true;
+	}
+	else
+	{
+		// No common struct, dont allow binding selection 
+		NameText = LOCTEXT("MultipleValues", "Multiple Values");
+		TooltipText = LOCTEXT("MultipleValues", "Multiple Values");
+		bShowBindingSelector = false;
 	}
 }
 
@@ -206,6 +261,11 @@ FText FVariableBindingPropertyCustomization::GetBindingTooltipText(TConstStructV
 
 TSharedRef<SWidget> FVariableBindingPropertyCustomization::CreateBindingWidget() const
 {
+	if(!Type.IsValid())
+	{
+		return SNullWidget::NullWidget;
+	}
+
 	TOptional<const UScriptStruct*> CommonBindingStruct;
 	BindingDataHandle->EnumerateConstRawData([&CommonBindingStruct](const void* RawData, const int32 DataIndex, const int32 NumDatas)
 	{
diff --git a/Engine/Plugins/Experimental/AnimNext/Source/AnimNextEditor/Private/Variables/VariableBindingPropertyCustomization.h b/Engine/Plugins/Experimental/AnimNext/Source/AnimNextEditor/Private/Variables/VariableBindingPropertyCustomization.h
--- a/Engine/Plugins/Experimental/AnimNext/Source/AnimNextEditor/Private/Variables/VariableBindingPropertyCustomization.h
+++ b/Engine/Plugins/Experimental/AnimNext/Source/AnimNextEditor/Private/Variables/VariableBindingPropertyCustomization.h
@@ -45,7 +45,9 @@ private:
 	const FSlateBrush* Icon = nullptr;
 	FSlateColor IconColor = FLinearColor::Gray;
 	TSharedPtr<SWidget> ValueWidget;
+	TSharedPtr<SWidget> ContainerWidget;
 	bool bRefreshRequested = false;
+	bool bShowBindingSelector = false;
 };
 
 }
diff --git a/Engine/Plugins/Experimental/AnimNext/Source/AnimNextEditor/Private/Variables/VariableCustomization.cpp b/Engine/Plugins/Experimental/AnimNext/Source/AnimNextEditor/Private/Variables/VariableCustomization.cpp
--- a/Engine/Plugins/Experimental/AnimNext/Source/AnimNextEditor/Private/Variables/VariableCustomization.cpp
+++ b/Engine/Plugins/Experimental/AnimNext/Source/AnimNextEditor/Private/Variables/VariableCustomization.cpp
@@ -27,7 +27,7 @@ void FVariableCustomization::CustomizeDetails(IDetailLayoutBuilder& DetailBuilde
 		return;
 	}
 
-	// Only allow bindings on module variables
+	
 	for(TWeakObjectPtr<UObject> WeakObject : Objects)
 	{
 		UAnimNextVariableEntry* Variable = Cast<UAnimNextVariableEntry>(WeakObject.Get());
@@ -36,12 +36,14 @@ void FVariableCustomization::CustomizeDetails(IDetailLayoutBuilder& DetailBuilde
 			continue;
 		}
 
+
 		UAnimNextRigVMAsset* Asset = Variable->GetTypedOuter<UAnimNextRigVMAsset>();
-		if(!Asset->IsA<UAnimNextModule>())
+
+		// Disable access specifier switching specifically for data interfaces
+		if(UAnimNextDataInterface* DataInterface = ExactCast<UAnimNextDataInterface>(Asset))
 		{
-			TSharedPtr<IPropertyHandle> BindingProperty = DetailBuilder.GetProperty(GET_MEMBER_NAME_CHECKED(UAnimNextVariableEntry, Binding));
-			BindingProperty->MarkHiddenByCustomization();
-			break;
+			TSharedPtr<IPropertyHandle> AccessProperty = DetailBuilder.GetProperty(GET_MEMBER_NAME_CHECKED(UAnimNextVariableEntry, Access));
+			AccessProperty->MarkHiddenByCustomization();
 		}
 	}
 
