AnimNext: Data interfaces

diff --git a/Engine/Plugins/Experimental/AnimNext/Source/AnimNextEditor/Private/Variables/VariableProxyCustomization.cpp b/Engine/Plugins/Experimental/AnimNext/Source/AnimNextEditor/Private/Variables/VariableProxyCustomization.cpp
new file mode 100644
--- /dev/null
+++ b/Engine/Plugins/Experimental/AnimNext/Source/AnimNextEditor/Private/Variables/VariableProxyCustomization.cpp
@@ -0,0 +1,172 @@
+﻿// Copyright Epic Games, Inc. All Rights Reserved.
+
+#include "Variables/VariableProxyCustomization.h"
+
+#include "AnimNextRigVMAssetEditorData.h"
+#include "AnimNextVariableEntryProxy.h"
+#include "DetailBuilderTypes.h"
+#include "DetailCategoryBuilder.h"
+#include "DetailLayoutBuilder.h"
+#include "DetailWidgetRow.h"
+#include "InstancedPropertyBagStructureDataProvider.h"
+#include "SVariableOverride.h"
+#include "Entries/AnimNextDataInterfaceEntry.h"
+#include "Entries/AnimNextVariableEntry.h"
+#include "Modules/ModuleManager.h"
+#include "StructUtils/PropertyBag.h"
+
+#define LOCTEXT_NAMESPACE "VariableProxyCustomization"
+
+namespace UE::AnimNext::Editor
+{
+
+void FVariableProxyCustomization::CustomizeDetails(IDetailLayoutBuilder& DetailBuilder)
+{
+	TArray<TWeakObjectPtr<UObject>> Objects;
+	DetailBuilder.GetObjectsBeingCustomized(Objects);
+
+	// Don't customize default value if we have multi-selection
+	if(Objects.Num() != 1)
+	{
+		return;
+	}
+
+	if (UAnimNextVariableEntryProxy* VariableProxy = CastChecked<UAnimNextVariableEntryProxy>(Objects[0].Get()))
+	{
+		IDetailCategoryBuilder& DefaultValueCategory = DetailBuilder.EditCategory(TEXT("DefaultValue"), LOCTEXT("DefaultValue", "Default Value"), ECategoryPriority::Default);
+
+		UAnimNextVariableEntry* VariableEntry = VariableProxy->VariableEntry;
+		UAnimNextDataInterfaceEntry* DataInterfaceEntry = VariableProxy->DataInterfaceEntry;
+
+		UAnimNextRigVMAssetEditorData* VariableEditorData = VariableEntry->GetTypedOuter<UAnimNextRigVMAssetEditorData>();
+		if(VariableEditorData)
+		{
+			VariableEditorData->ModifiedDelegate.AddSP(this, &FVariableProxyCustomization::HandleModified);
+		}
+
+		UAnimNextRigVMAssetEditorData* DataInterfaceEditorData = DataInterfaceEntry->GetTypedOuter<UAnimNextRigVMAssetEditorData>();
+		if(DataInterfaceEditorData)
+		{
+			DataInterfaceEditorData->ModifiedDelegate.AddSP(this, &FVariableProxyCustomization::HandleModified);
+		}
+
+		const FName VariableName = VariableEntry->GetEntryName();
+		FName ValueName;
+		FInstancedPropertyBag* PropertyBag = nullptr;
+		EAnimNextDataInterfaceValueOverrideStatus OverrideStatus = DataInterfaceEntry->FindValueOverridePropertyBagRecursive(VariableName, PropertyBag);
+		switch(OverrideStatus)
+		{
+		case EAnimNextDataInterfaceValueOverrideStatus::NotOverridden:
+			{
+				// Use the internal property bag (copying from the source) and the default name
+				InternalPropertyBag = VariableEntry->GetPropertyBag();
+				PropertyBag = &InternalPropertyBag;
+				ValueName = IAnimNextRigVMVariableInterface::ValueName;
+				break;
+			}
+		case EAnimNextDataInterfaceValueOverrideStatus::OverriddenInThisAsset:
+			{
+				// Use the found property bag and the variable name
+				check(PropertyBag);
+				ValueName = VariableName;
+				break;
+			}
+		case EAnimNextDataInterfaceValueOverrideStatus::OverriddenInParentAsset:
+			{
+				// Use the internal property bag (copying just the value from the overriding asset's property bag) and the default name
+				check(PropertyBag);
+				const FPropertyBagPropertyDesc* Desc = PropertyBag->FindPropertyDescByName(VariableName);
+				check(Desc);
+				const uint8* DataPtr = Desc->CachedProperty->ContainerPtrToValuePtr<uint8>(PropertyBag->GetValue().GetMemory());
+				TConstArrayView<uint8> Value(DataPtr, Desc->CachedProperty->GetElementSize());
+				InternalPropertyBag.ReplaceAllPropertiesAndValues(TConstArrayView<FPropertyBagPropertyDesc>(Desc, 1), TConstArrayView<TConstArrayView<uint8>>(&Value, 1));
+				PropertyBag = &InternalPropertyBag;
+				ValueName = VariableName;
+				break;
+			}
+		}
+
+		IDetailPropertyRow* DetailPropertyRow = DefaultValueCategory.AddExternalStructureProperty(MakeShared<FInstancePropertyBagStructureDataProvider>(*PropertyBag), ValueName);
+		PropertyHandle = DetailPropertyRow->GetPropertyHandle();
+		
+		if (!PropertyHandle.IsValid())
+		{
+			return;
+		}
+
+		// Hide reset to default & replace with our own
+		DetailPropertyRow->OverrideResetToDefault(FResetToDefaultOverride::Create(TAttribute<bool>(false)));
+		PropertyHandle->SetPropertyDisplayName(FText::FromName(VariableEntry->GetEntryName()));
+		TSharedPtr<SWidget> NameWidget, ValueWidget;
+		DetailPropertyRow->GetDefaultWidgets(NameWidget, ValueWidget);
+		DetailPropertyRow->CustomWidget(true)
+		.NameContent()
+		[
+			NameWidget.IsValid() ? NameWidget.ToSharedRef() : SNullWidget::NullWidget
+		]
+		.ValueContent()
+		[
+			ValueWidget.IsValid() ? ValueWidget.ToSharedRef() : SNullWidget::NullWidget
+		]
+		.ResetToDefaultContent()
+		[
+			SNew(SVariableOverride, DataInterfaceEntry, VariableName)
+			.OverrideStatus(OverrideStatus)
+		];
+
+		
+
+		const auto OnPropertyValueChange = [this, WeakVariableEntry = TWeakObjectPtr<UAnimNextVariableEntry>(VariableEntry), WeakDataInterfaceEntry = TWeakObjectPtr<UAnimNextDataInterfaceEntry>(DataInterfaceEntry)](const FPropertyChangedEvent& InEvent)
+		{
+			UAnimNextVariableEntry* PinnedVariableEntry = WeakVariableEntry.Get();
+			UAnimNextDataInterfaceEntry* PinnedDataInterfaceEntry = WeakDataInterfaceEntry.Get();
+			if (PinnedVariableEntry == nullptr || PinnedDataInterfaceEntry == nullptr)
+			{
+				return;
+			}
+
+			const FName VariableName = PinnedVariableEntry->GetVariableName();
+			if(!PinnedDataInterfaceEntry->HasValueOverride(VariableName))
+			{
+				// No value override yet, so we copy from the internal property bag
+				check(InternalPropertyBag.IsValid() && InternalPropertyBag.GetPropertyBagStruct()->GetPropertyDescs().Num() == 1);
+				const FProperty* Property = InternalPropertyBag.GetPropertyBagStruct()->GetPropertyDescs()[0].CachedProperty; 
+				const uint8* DataPtr = Property->ContainerPtrToValuePtr<uint8>(InternalPropertyBag.GetValue().GetMemory());
+				PinnedDataInterfaceEntry->SetValueOverride(VariableName, PinnedVariableEntry->GetType(), TConstArrayView<uint8>(DataPtr, Property->GetElementSize()));
+				ensure(PinnedDataInterfaceEntry->GetValueOverrideStatusRecursive(VariableName) == EAnimNextDataInterfaceValueOverrideStatus::OverriddenInThisAsset);
+			}
+
+			PinnedDataInterfaceEntry->MarkPackageDirty();
+			PinnedDataInterfaceEntry->BroadcastModified(EAnimNextEditorDataNotifType::VariableDefaultValueChanged);
+		};
+
+		PropertyHandle->SetOnPropertyValueChangedWithData(TDelegate<void(const FPropertyChangedEvent&)>::CreateLambda(OnPropertyValueChange));
+		PropertyHandle->SetOnChildPropertyValueChangedWithData(TDelegate<void(const FPropertyChangedEvent&)>::CreateLambda(OnPropertyValueChange));
+
+	}
+}
+
+void FVariableProxyCustomization::CustomizeDetails(const TSharedPtr<IDetailLayoutBuilder>& DetailBuilder) 
+{
+	WeakDetailBuilder = DetailBuilder;
+	CustomizeDetails(*DetailBuilder);
+}
+
+void FVariableProxyCustomization::HandleModified(UAnimNextRigVMAssetEditorData* InEditorData, EAnimNextEditorDataNotifType InType, UObject* InSubject)
+{
+	if( InType != EAnimNextEditorDataNotifType::VariableDefaultValueChanged &&
+		InType != EAnimNextEditorDataNotifType::UndoRedo)
+	{
+		return;
+	}
+
+	if(TSharedPtr<IDetailLayoutBuilder> DetailBuilder = WeakDetailBuilder.Pin())
+	{
+		// Rebuild this customization to point at the correct data
+		DetailBuilder->ForceRefreshDetails();
+	}
+}
+
+}
+
+#undef LOCTEXT_NAMESPACE
diff --git a/Engine/Plugins/Experimental/AnimNext/Source/AnimNextEditor/Private/Variables/VariableProxyCustomization.h b/Engine/Plugins/Experimental/AnimNext/Source/AnimNextEditor/Private/Variables/VariableProxyCustomization.h
new file mode 100644
--- /dev/null
+++ b/Engine/Plugins/Experimental/AnimNext/Source/AnimNextEditor/Private/Variables/VariableProxyCustomization.h
@@ -0,0 +1,32 @@
+﻿// Copyright Epic Games, Inc. All Rights Reserved.
+
+#pragma once
+
+#include "IDetailCustomization.h"
+#include "IPropertyTypeCustomization.h"
+#include "PropertyEditorModule.h"
+#include "StructUtils/PropertyBag.h"
+
+class UAnimNextRigVMAssetEditorData;
+enum class EAnimNextEditorDataNotifType : uint8;
+
+namespace UE::AnimNext::Editor
+{
+
+class FVariableProxyCustomization : public IDetailCustomization
+{
+private:
+	// IDetailCustomization interface
+	virtual void CustomizeDetails(IDetailLayoutBuilder& DetailBuilder) override;
+	virtual void CustomizeDetails(const TSharedPtr<IDetailLayoutBuilder>& DetailBuilder) override;
+
+	void HandleModified(UAnimNextRigVMAssetEditorData* InEditorData, EAnimNextEditorDataNotifType InType, UObject* InSubject);
+
+	void HandleGlobalRowExtension(const FOnGenerateGlobalRowExtensionArgs& InArgs, TArray<FPropertyRowExtensionButton>& OutExtensions);
+
+	TSharedPtr<IPropertyHandle> PropertyHandle;
+	FInstancedPropertyBag InternalPropertyBag;
+	TWeakPtr<IDetailLayoutBuilder> WeakDetailBuilder;
+};
+
+}
diff --git a/Engine/Plugins/Experimental/AnimNext/Source/AnimNextTestSuite/Private/AnimNextDataInterfaceTest.cpp b/Engine/Plugins/Experimental/AnimNext/Source/AnimNextTestSuite/Private/AnimNextDataInterfaceTest.cpp
new file mode 100644
--- /dev/null
+++ b/Engine/Plugins/Experimental/AnimNext/Source/AnimNextTestSuite/Private/AnimNextDataInterfaceTest.cpp
@@ -0,0 +1,180 @@
+﻿// Copyright Epic Games, Inc. All Rights Reserved.
+
+#include "CoreMinimal.h"
+#include "AnimNextRigVMAssetEditorData.h"
+#include "AnimNextTest.h"
+#include "UncookedOnlyUtils.h"
+#include "DataInterface/AnimNextDataInterface.h"
+#include "Misc/AutomationTest.h"
+#include "DataInterface/AnimNextDataInterfaceFactory.h"
+#include "DataInterface/AnimNextDataInterface_EditorData.h"
+#include "Entries/AnimNextDataInterfaceEntry.h"
+#include "Entries/AnimNextVariableEntry.h"
+#include "Graph/AnimNextAnimationGraph.h"
+#include "Graph/AnimNextAnimationGraphFactory.h"
+
+// AnimNext Data Interface Tests
+
+#if WITH_DEV_AUTOMATION_TESTS && WITH_EDITOR
+
+namespace UE::AnimNext::Tests
+{
+
+IMPLEMENT_SIMPLE_AUTOMATION_TEST(FDataInterfaceCompile, "Animation.AnimNext.DataInterface.Compile", EAutomationTestFlags::EditorContext | EAutomationTestFlags::EngineFilter)
+
+bool FDataInterfaceCompile::RunTest(const FString& InParameters)
+{
+	using namespace UE::AnimNext;
+
+	UFactory* Factory = NewObject<UFactory>(GetTransientPackage(), UAnimNextDataInterfaceFactory::StaticClass());
+	UAnimNextDataInterface* A = CastChecked<UAnimNextDataInterface>(Factory->FactoryCreateNew(UAnimNextDataInterface::StaticClass(), GetTransientPackage(), TEXT("TestDataInterfaceA"), RF_Transient, nullptr, nullptr, NAME_None));
+	UE_RETURN_ON_ERROR(A != nullptr, "FDataInterfaceCompile -> Failed to create asset");
+	UAnimNextDataInterface* B = CastChecked<UAnimNextDataInterface>(Factory->FactoryCreateNew(UAnimNextDataInterface::StaticClass(), GetTransientPackage(), TEXT("TestDataInterfaceB"), RF_Transient, nullptr, nullptr, NAME_None));
+	UE_RETURN_ON_ERROR(B != nullptr, "FDataInterfaceCompile -> Failed to create asset");
+	UAnimNextDataInterface* C = CastChecked<UAnimNextDataInterface>(Factory->FactoryCreateNew(UAnimNextDataInterface::StaticClass(), GetTransientPackage(), TEXT("TestDataInterfaceC"), RF_Transient, nullptr, nullptr, NAME_None));
+	UE_RETURN_ON_ERROR(C != nullptr, "FDataInterfaceCompile -> Failed to create asset");
+	UAnimNextDataInterface* D = CastChecked<UAnimNextDataInterface>(Factory->FactoryCreateNew(UAnimNextDataInterface::StaticClass(), GetTransientPackage(), TEXT("TestDataInterfaceD"), RF_Transient, nullptr, nullptr, NAME_None));
+	UE_RETURN_ON_ERROR(D != nullptr, "FDataInterfaceCompile -> Failed to create asset");
+
+	auto CompileAndTest = [A, B, C, D](TFunctionRef<void()> InTestFunction)
+	{
+		UncookedOnly::FUtils::CompileVariables(A);
+		UncookedOnly::FUtils::CompileVariables(B);
+		UncookedOnly::FUtils::CompileVariables(C);
+		UncookedOnly::FUtils::CompileVariables(D);
+
+		InTestFunction();
+	};
+
+	auto ClearVariables = [A, B, C, D]()
+	{
+		UAnimNextRigVMAssetLibrary::RemoveAllEntries(A, false, false);
+		UAnimNextRigVMAssetLibrary::RemoveAllEntries(B, false, false);
+		UAnimNextRigVMAssetLibrary::RemoveAllEntries(C, false, false);
+		UAnimNextRigVMAssetLibrary::RemoveAllEntries(D, false, false);
+	};
+
+	auto RunTest = [this, &ClearVariables](TFunctionRef<void()> InTestFunction)
+	{
+		InTestFunction();
+		ClearVariables();
+	};
+
+	auto AddPublicVariable = [this](UAnimNextDataInterface* InDataInterface, FName InName, const auto& InValue)
+	{
+		using ValueType = std::remove_reference_t<decltype(InValue)>;
+		FAnimNextParamType Type = FAnimNextParamType::GetType<ValueType>();
+		UAnimNextVariableEntry* VariableEntry = UAnimNextRigVMAssetLibrary::AddVariable(InDataInterface, InName, Type.GetValueType(), Type.GetContainerType(), Type.GetValueTypeObject(), TEXT(""), false, false);
+		UE_RETURN_ON_ERROR(VariableEntry != nullptr, "FDataInterfaceCompile::AddPublicVariable -> Failed to create variable");
+		VariableEntry->SetExportAccessSpecifier(EAnimNextExportAccessSpecifier::Public, false);
+		UE_RETURN_ON_ERROR(VariableEntry->SetDefaultValue(InValue, false), "FDataInterfaceCompile::AddPublicVariable -> Failed to set variable default value");
+		return true;
+	};
+
+	auto OverrideVariable = [this](UAnimNextDataInterface* InDataInterface, FName InName, const auto& InValue)
+	{
+		UAnimNextDataInterface_EditorData* EditorData = UncookedOnly::FUtils::GetEditorData<UAnimNextDataInterface_EditorData>(InDataInterface);
+		bool bSuccessfulOverride = false;
+		EditorData->ForEachEntryOfType<UAnimNextDataInterfaceEntry>([this, &InName, &InValue, &bSuccessfulOverride](UAnimNextDataInterfaceEntry* InEntry)
+		{
+			bSuccessfulOverride = InEntry->SetValueOverride(InName, InValue, false);
+			return !bSuccessfulOverride;
+		});
+
+		UE_RETURN_ON_ERROR(bSuccessfulOverride, "FDataInterfaceCompile::OverrideVariable -> Failed to override default value");
+
+		return false;
+	};
+
+	auto Implement = [this](UAnimNextDataInterface* InA, UAnimNextDataInterface* InB)
+	{
+		UAnimNextDataInterfaceEntry* DataInterfaceEntry = UAnimNextRigVMAssetLibrary::AddDataInterface(InA, InB, false, false);
+		UE_RETURN_ON_ERROR(DataInterfaceEntry != nullptr, "FDataInterfaceCompile::Implement -> Failed to add data interface");
+		return true;
+	};
+
+	auto CheckValue = [this](UAnimNextDataInterface* InDataInterface, FName InName, const auto& InValue)
+	{
+		const FPropertyBagPropertyDesc* Desc = InDataInterface->VariableDefaults.FindPropertyDescByName(InName);
+		UE_RETURN_ON_ERROR(Desc != nullptr, "FDataInterfaceCompile::CheckValue -> Failed to find value");
+		using ValueType = std::remove_reference_t<decltype(InValue)>;
+		FAnimNextParamType DesiredType = FAnimNextParamType::GetType<ValueType>();
+		FAnimNextParamType FoundType = FAnimNextParamType(Desc->ValueType, Desc->ContainerTypes.GetFirstContainerType(), Desc->ValueTypeObject);
+		UE_RETURN_ON_ERROR(DesiredType == FoundType, "FDataInterfaceCompile::CheckValue -> Type was incorrect");
+		check(Desc->CachedProperty);
+		const uint8* ValuePtr = Desc->CachedProperty->ContainerPtrToValuePtr<uint8>(InDataInterface->VariableDefaults.GetValue().GetMemory());
+		UE_RETURN_ON_ERROR(Desc->CachedProperty->Identical(ValuePtr, &InValue), "FDataInterfaceCompile::CheckValue -> Values were not equal");
+		return true;
+	};
+
+	// Add a variable, dont override it, check its value
+	RunTest([&]()
+	{
+		AddPublicVariable(A, "A", 1);
+		CompileAndTest([&]()
+		{
+			CheckValue(A, "A", 1);
+		});
+	});
+
+	// Add a variable, override it, check its value
+	RunTest([&]()
+	{
+		AddPublicVariable(A, "A", 1);
+		Implement(B, A);
+		OverrideVariable(B, "A", 2);
+		CompileAndTest([&]()
+		{
+			CheckValue(A, "A", 1);
+			CheckValue(B, "A", 2);
+		});
+	});
+
+	// Add a variable, override it twice, check its value at each step
+	RunTest([&]()
+	{
+		AddPublicVariable(A, "A", 1);
+		Implement(B, A);
+		Implement(C, A);
+		OverrideVariable(B, "A", 2);
+		OverrideVariable(C, "A", 3);
+		CompileAndTest([&]()
+		{
+			CheckValue(A, "A", 1);
+			CheckValue(B, "A", 2);
+			CheckValue(C, "A", 3);
+		});
+	});
+
+	// Diamond inheritance
+	//   A=1
+	//  /   \
+	// B=2  C=3
+	//  \   /
+	//   D=3
+	RunTest([&]()
+	{
+		AddPublicVariable(A, "A", 1);
+		Implement(B, A);
+		Implement(C, A);
+		Implement(D, B);
+		Implement(D, C);
+		OverrideVariable(B, "A", 2);
+		OverrideVariable(C, "A", 3);
+		CompileAndTest([&]()
+		{
+			CheckValue(A, "A", 1);
+			CheckValue(B, "A", 2);
+			CheckValue(C, "A", 3);
+			CheckValue(D, "A", 3);
+		});
+	});
+	
+	FUtils::CleanupAfterTests();
+
+	return true;
+}
+
+}
+
+#endif	// WITH_DEV_AUTOMATION_TESTS && WITH_EDITOR
\ No newline at end of file
diff --git a/Engine/Plugins/Experimental/AnimNext/Source/AnimNextUncookedOnly/Internal/AnimNextAssetWorkspaceAssetUserData.h b/Engine/Plugins/Experimental/AnimNext/Source/AnimNextUncookedOnly/Internal/AnimNextAssetWorkspaceAssetUserData.h
--- a/Engine/Plugins/Experimental/AnimNext/Source/AnimNextUncookedOnly/Internal/AnimNextAssetWorkspaceAssetUserData.h
+++ b/Engine/Plugins/Experimental/AnimNext/Source/AnimNextUncookedOnly/Internal/AnimNextAssetWorkspaceAssetUserData.h
@@ -9,6 +9,7 @@
 #include "RigVMModel/Nodes/RigVMFunctionReferenceNode.h"
 #include "Graph/AnimNextAnimationGraph.h"
 #include "Module/AnimNextModule.h"
+#include "DataInterface/AnimNextDataInterface.h"
 
 #include "AnimNextAssetWorkspaceAssetUserData.generated.h"
 
@@ -51,6 +52,19 @@ struct FAnimNextAnimationGraphOutlinerData : public FAnimNextRigVMAssetOutlinerD
 	}
 };
 
+USTRUCT()
+struct FAnimNextDataInterfaceOutlinerData : public FAnimNextRigVMAssetOutlinerData
+{
+	GENERATED_BODY()
+
+	FAnimNextDataInterfaceOutlinerData() = default;
+
+	UAnimNextDataInterface* GetDataInteface() const
+	{
+		return Cast<UAnimNextDataInterface>(Asset);
+	}
+};
+
 // Base struct used to identify asset sub-entries
 USTRUCT()
 struct FAnimNextAssetEntryOutlinerData : public FWorkspaceOutlinerItemData
diff --git a/Engine/Plugins/Experimental/AnimNext/Source/AnimNextUncookedOnly/Internal/AnimNextRigVMAssetEditorData.h b/Engine/Plugins/Experimental/AnimNext/Source/AnimNextUncookedOnly/Internal/AnimNextRigVMAssetEditorData.h
--- a/Engine/Plugins/Experimental/AnimNext/Source/AnimNextUncookedOnly/Internal/AnimNextRigVMAssetEditorData.h
+++ b/Engine/Plugins/Experimental/AnimNext/Source/AnimNextUncookedOnly/Internal/AnimNextRigVMAssetEditorData.h
@@ -14,6 +14,7 @@ enum class ERigVMGraphNotifType : uint8;
 class UAnimNextRigVMAssetEntry;
 class UAnimNextRigVMAssetEditorData;
 class UAnimNextEdGraph;
+class UAnimNextDataInterfaceEntry;
 
 namespace UE::AnimNext::UncookedOnly
 {
@@ -37,6 +38,8 @@ namespace UE::AnimNext::Editor
 	class FVariablesOutlinerHierarchy;
 	class SVariablesOutlinerValue;
 	class SVariablesOutliner;
+	class SAddVariablesDialog;
+	class FVariableProxyCustomization;
 }
 
 namespace UE::AnimNext::Tests
@@ -44,6 +47,7 @@ namespace UE::AnimNext::Tests
 	class FEditor_Graphs;
 	class FEditor_Variables;
 	class FVariables;
+	class FDataInterfaceCompile;
 }
 
 enum class EAnimNextEditorDataNotifType : uint8
@@ -54,6 +58,9 @@ enum class EAnimNextEditorDataNotifType : uint8
 	EntryRenamed,	// An entry has been renamed (Subject == UAnimNextRigVMAssetEntry)
 	EntryAccessSpecifierChanged,	// An entry access specifier has been changed (Subject == UAnimNextRigVMAssetEntry)
 	VariableTypeChanged,	// A variable entry type changed (Subject == UAnimNextVariableEntry)
+	UndoRedo,		// Transaction was performed (Subject == UObject)
+	VariableDefaultValueChanged,	// A variable entry default value changed (Subject == UAnimNextVariableEntry)
+	VariableBindingChanged,	// A variable entry binding changed (Subject == UAnimNextVariableEntry)
 };
 
 namespace UE::AnimNext::UncookedOnly
@@ -71,29 +78,38 @@ class UAnimNextRigVMAssetLibrary : public UBlueprintFunctionLibrary
 {
 	GENERATED_BODY()
 
+public:
 	/** Finds an entry in an AnimNext asset */
 	UFUNCTION(BlueprintCallable, Category = "AnimNext|Entries", meta=(ScriptMethod))
-	static UAnimNextRigVMAssetEntry* FindEntry(UAnimNextRigVMAsset* InAsset, FName InName);
+	static ANIMNEXTUNCOOKEDONLY_API UAnimNextRigVMAssetEntry* FindEntry(UAnimNextRigVMAsset* InAsset, FName InName);
 
 	/** Removes an entry from an AnimNext asset */
 	UFUNCTION(BlueprintCallable, Category = "AnimNext|Entries", meta=(ScriptMethod))
-	static bool RemoveEntry(UAnimNextRigVMAsset* InAsset, UAnimNextRigVMAssetEntry* InEntry, bool bSetupUndoRedo = true, bool bPrintPythonCommand = true);
+	static ANIMNEXTUNCOOKEDONLY_API bool RemoveEntry(UAnimNextRigVMAsset* InAsset, UAnimNextRigVMAssetEntry* InEntry, bool bSetupUndoRedo = true, bool bPrintPythonCommand = true);
 
 	/** Removes multiple entries from an AnimNext asset */
 	UFUNCTION(BlueprintCallable, Category = "AnimNext|Entries", meta=(ScriptMethod))
-	static bool RemoveEntries(UAnimNextRigVMAsset* InAsset, const TArray<UAnimNextRigVMAssetEntry*>& InEntries, bool bSetupUndoRedo = true, bool bPrintPythonCommand = true);
+	static ANIMNEXTUNCOOKEDONLY_API bool RemoveEntries(UAnimNextRigVMAsset* InAsset, const TArray<UAnimNextRigVMAssetEntry*>& InEntries, bool bSetupUndoRedo = true, bool bPrintPythonCommand = true);
+
+	/** Removes all entries from an AnimNext asset */
+	UFUNCTION(BlueprintCallable, Category = "AnimNext|Entries", meta=(ScriptMethod))
+	static ANIMNEXTUNCOOKEDONLY_API bool RemoveAllEntries(UAnimNextRigVMAsset* InAsset, bool bSetupUndoRedo = true, bool bPrintPythonCommand = true);
 
 	/** Adds an animation graph to an AnimNext asset */
 	UFUNCTION(BlueprintCallable, Category = "AnimNext|Entries", meta=(ScriptMethod))
-	static UAnimNextAnimationGraphEntry* AddAnimationGraph(UAnimNextModule* InModule, FName InName, bool bSetupUndoRedo = true, bool bPrintPythonCommand = true);
+	static ANIMNEXTUNCOOKEDONLY_API UAnimNextAnimationGraphEntry* AddAnimationGraph(UAnimNextRigVMAsset* InAsset, FName InName, bool bSetupUndoRedo = true, bool bPrintPythonCommand = true);
 
 	/** Adds a parameter to an AnimNext asset */
 	UFUNCTION(BlueprintCallable, Category = "AnimNext|Entries", meta=(ScriptMethod))
-	static UAnimNextVariableEntry* AddVariable(UAnimNextModule* InModule, FName InName, EPropertyBagPropertyType InValueType, EPropertyBagContainerType InContainerType = EPropertyBagContainerType::None, const UObject* InValueTypeObject = nullptr, const FString& InDefaultValue = TEXT(""), bool bSetupUndoRedo = true, bool bPrintPythonCommand = true);
+	static ANIMNEXTUNCOOKEDONLY_API UAnimNextVariableEntry* AddVariable(UAnimNextRigVMAsset* InAsset, FName InName, EPropertyBagPropertyType InValueType, EPropertyBagContainerType InContainerType = EPropertyBagContainerType::None, const UObject* InValueTypeObject = nullptr, const FString& InDefaultValue = TEXT(""), bool bSetupUndoRedo = true, bool bPrintPythonCommand = true);
 
 	/** Adds an event graph to an AnimNext asset */
 	UFUNCTION(BlueprintCallable, Category = "AnimNext|Entries", meta=(ScriptMethod))
-	static UAnimNextEventGraphEntry* AddEventGraph(UAnimNextModule* InModule, FName InName, UScriptStruct* InEventStruct, bool bSetupUndoRedo = true, bool bPrintPythonCommand = true);
+	static ANIMNEXTUNCOOKEDONLY_API UAnimNextEventGraphEntry* AddEventGraph(UAnimNextRigVMAsset* InAsset, FName InName, UScriptStruct* InEventStruct, bool bSetupUndoRedo = true, bool bPrintPythonCommand = true);
+
+	/** Adds a data interface to an AnimNext asset */
+	UFUNCTION(BlueprintCallable, Category = "AnimNext|Entries", meta=(ScriptMethod))
+	static ANIMNEXTUNCOOKEDONLY_API UAnimNextDataInterfaceEntry* AddDataInterface(UAnimNextRigVMAsset* InAsset, UAnimNextDataInterface* InDataInterface, bool bSetupUndoRedo = true, bool bPrintPythonCommand = true);
 };
 
 /* Base class for all AnimNext editor data objects that use RigVM */
@@ -111,7 +127,10 @@ public:
 
 	/** Adds an event graph to this asset */
 	ANIMNEXTUNCOOKEDONLY_API UAnimNextEventGraphEntry* AddEventGraph(FName InName, UScriptStruct* InEventStruct, bool bSetupUndoRedo = true, bool bPrintPythonCommand = true);
-	
+
+	/** Adds a data interface to this asset */
+	ANIMNEXTUNCOOKEDONLY_API UAnimNextDataInterfaceEntry* AddDataInterface(UAnimNextDataInterface* InDataInterface, bool bSetupUndoRedo = true, bool bPrintPythonCommand = true);
+
 	// Report an error to the user, typically used for scripting APIs
 	ANIMNEXTUNCOOKEDONLY_API static void ReportError(const TCHAR* InMessage);
 
@@ -128,6 +147,7 @@ protected:
 	friend class UE::AnimNext::Tests::FEditor_Graphs;
 	friend class UE::AnimNext::Tests::FEditor_Variables;
 	friend class UE::AnimNext::Tests::FVariables;
+	friend class UE::AnimNext::Tests::FDataInterfaceCompile;
 	friend class UE::AnimNext::Editor::FVariableCustomization;
 	friend class UE::AnimNext::Editor::FAnimNextAssetItemDetails;
 	friend class UE::AnimNext::Editor::FAnimNextGraphItemDetails;
@@ -137,6 +157,9 @@ protected:
 	friend class UE::AnimNext::Editor::SVariablesOutlinerValue;
 	friend class UE::AnimNext::Editor::SVariablesOutliner;
 	friend class UAnimNextModuleWorkspaceAssetUserData;
+	friend class UE::AnimNext::Editor::SAddVariablesDialog;
+	friend class UE::AnimNext::Editor::FVariableProxyCustomization;
+	friend class UAnimNextDataInterfaceEntry;
 
 	// UObject interface
 	virtual void Serialize(FArchive& Ar) override;
@@ -232,6 +255,9 @@ protected:
 	// Allows this asset to generate graphs to be injected at compilation time
 	virtual void GetProgrammaticGraphs(const FRigVMCompileSettings& InSettings, TArray<URigVMGraph*>& OutGraphs) {}
 
+	// Customization point for derived types to transform new asset entries
+	virtual void CustomizeNewAssetEntry(UAnimNextRigVMAssetEntry* InNewEntry) const {}
+
 	// Helper for creating new sub-entries. Sets package flags and outers appropriately 
 	static UObject* CreateNewSubEntry(UAnimNextRigVMAssetEditorData* InEditorData, TSubclassOf<UObject> InClass);
 
@@ -322,6 +348,10 @@ protected:
 	// @return true if any items were removed
 	ANIMNEXTUNCOOKEDONLY_API bool RemoveEntries(TConstArrayView<UAnimNextRigVMAssetEntry*> InEntries, bool bSetupUndoRedo = true, bool bPrintPythonCommand = true);
 
+	// Remove all entries from the asset
+	// @return true if any items were removed
+	ANIMNEXTUNCOOKEDONLY_API bool RemoveAllEntries(bool bSetupUndoRedo = true, bool bPrintPythonCommand = true);
+
 	ANIMNEXTUNCOOKEDONLY_API void BroadcastModified(EAnimNextEditorDataNotifType InType, UObject* InSubject);
 
 	void ReconstructAllNodes();
diff --git a/Engine/Plugins/Experimental/AnimNext/Source/AnimNextUncookedOnly/Internal/DataInterface/AnimNextDataInterface_EditorData.h b/Engine/Plugins/Experimental/AnimNext/Source/AnimNextUncookedOnly/Internal/DataInterface/AnimNextDataInterface_EditorData.h
new file mode 100644
--- /dev/null
+++ b/Engine/Plugins/Experimental/AnimNext/Source/AnimNextUncookedOnly/Internal/DataInterface/AnimNextDataInterface_EditorData.h
@@ -0,0 +1,27 @@
+﻿// Copyright Epic Games, Inc. All Rights Reserved.
+
+#pragma once
+
+#include "CoreMinimal.h"
+#include "AnimNextController.h"
+#include "AnimNextExecuteContext.h"
+#include "AnimNextRigVMAssetEditorData.h"
+#include "AnimNextDataInterface_EditorData.generated.h"
+
+class UAnimNextDataInterfaceFactory;
+
+/** Editor data for AnimNext data interfaces */
+UCLASS(MinimalAPI)
+class UAnimNextDataInterface_EditorData : public UAnimNextRigVMAssetEditorData
+{
+	GENERATED_BODY()
+
+protected:
+	// UAnimNextRigVMAssetEditorData interface
+	virtual TSubclassOf<URigVMController> GetControllerClass() const override { return UAnimNextController::StaticClass(); }
+	virtual UScriptStruct* GetExecuteContextStruct() const override { return FAnimNextExecuteContext::StaticStruct(); }
+	virtual TConstArrayView<TSubclassOf<UAnimNextRigVMAssetEntry>> GetEntryClasses() const override;
+	virtual void CustomizeNewAssetEntry(UAnimNextRigVMAssetEntry* InNewEntry) const override;
+
+	friend class UAnimNextDataInterfaceFactory;
+};
diff --git a/Engine/Plugins/Experimental/AnimNext/Source/AnimNextUncookedOnly/Internal/Entries/AnimNextDataInterfaceEntry.h b/Engine/Plugins/Experimental/AnimNext/Source/AnimNextUncookedOnly/Internal/Entries/AnimNextDataInterfaceEntry.h
new file mode 100644
--- /dev/null
+++ b/Engine/Plugins/Experimental/AnimNext/Source/AnimNextUncookedOnly/Internal/Entries/AnimNextDataInterfaceEntry.h
@@ -0,0 +1,169 @@
+﻿// Copyright Epic Games, Inc. All Rights Reserved.
+
+#pragma once
+
+#include "CoreMinimal.h"
+#include "AnimNextRigVMAssetEntry.h"
+#include "Param/ParamType.h"
+#include "StructUtils/PropertyBag.h"
+#include "AnimNextDataInterfaceEntry.generated.h"
+
+class UAnimNextDataInterface;
+class UAnimNextRigVMAssetEditorData;
+class UAssetDefinition_AnimNextDataInterfaceEntry;
+
+namespace UE::AnimNext::Editor
+{
+	class FVariablesOutlinerHierarchy;
+	struct FVariablesOutlinerDataInterfaceItem;
+	class SVariablesOutlinerDataInterfaceLabel;
+	class SVariablesOutlinerValue;
+	class FVariableProxyCustomization;
+	class SAddVariablesDialog;
+	class SVariableOverride;
+}
+
+namespace UE::AnimNext::UncookedOnly
+{
+	struct FUtils;
+}
+
+namespace UE::AnimNext::Tests
+{
+	class FDataInterfaceCompile;
+}
+
+UENUM()
+enum class EAnimNextDataInterfaceAutomaticBindingMode : uint8
+{
+	// No automatic binding will be performed
+	NoBinding,
+
+	// Public variables that exist on shared data interfaces on this asset and its host will be bound together if they share an interface
+	BindSharedInterfaces,
+};
+
+// Enum describing how a variable value is overriden
+UENUM()
+enum class EAnimNextDataInterfaceValueOverrideStatus
+{
+	// No override present in the implementation hierarchy
+	NotOverridden,
+
+	// Override present in this asset 
+	OverriddenInThisAsset,
+
+	// Override present in a parent asset
+	OverriddenInParentAsset,
+};
+
+UCLASS(MinimalAPI, Category = "Data Interfaces", DisplayName = "Data Interface")
+class UAnimNextDataInterfaceEntry : public UAnimNextRigVMAssetEntry
+{
+	GENERATED_BODY()
+
+	friend class UAnimNextRigVMAssetEditorData;
+	friend class UAssetDefinition_AnimNextDataInterfaceEntry;
+	friend class UE::AnimNext::Editor::FVariablesOutlinerHierarchy;
+	friend struct UE::AnimNext::Editor::FVariablesOutlinerDataInterfaceItem;
+	friend struct UE::AnimNext::UncookedOnly::FUtils;
+	friend class UE::AnimNext::Tests::FDataInterfaceCompile;
+	friend class UE::AnimNext::Editor::SVariablesOutlinerDataInterfaceLabel;
+	friend class UE::AnimNext::Editor::SVariablesOutlinerValue;
+	friend class UE::AnimNext::Editor::FVariableProxyCustomization;
+	friend class UE::AnimNext::Editor::SAddVariablesDialog;
+	friend class UE::AnimNext::Editor::SVariableOverride;
+
+	// UAnimNextRigVMAssetEntry interface
+	virtual void Initialize(UAnimNextRigVMAssetEditorData* InEditorData) override;
+	virtual FName GetEntryName() const override;
+	virtual void SetEntryName(FName InName, bool bSetupUndoRedo = true) override {}
+	virtual FText GetDisplayName() const override;
+	virtual FText GetDisplayNameTooltip() const override;
+
+	// Set the data interface that this entry represents
+	ANIMNEXTUNCOOKEDONLY_API void SetDataInterface(UAnimNextDataInterface* InDataInterface, bool bSetupUndoRedo = true);
+
+	// Get the data interface that this entry represents
+	ANIMNEXTUNCOOKEDONLY_API UAnimNextDataInterface* GetDataInterface() const;
+
+	// Get the path to the data interface that this entry represents
+	ANIMNEXTUNCOOKEDONLY_API FSoftObjectPath GetDataInterfacePath() const;
+
+	// Set an overriden value for the specified named variable, using the current value as the initial override
+	// @return true if the value was set successfully
+	ANIMNEXTUNCOOKEDONLY_API bool SetValueOverrideToDefault(FName InName, bool bSetupUndoRedo = true);
+
+	// Set an overriden value for the specified name and type
+	// @return true if the value was set successfully
+	ANIMNEXTUNCOOKEDONLY_API bool SetValueOverride(FName InName, const FAnimNextParamType& InType, TConstArrayView<uint8> InValue, bool bSetupUndoRedo = true);
+
+	// Set an overriden value for the specified name and type
+	// @return true if the value was set successfully
+	template<typename ValueType>
+	bool SetValueOverride(FName InName, const ValueType& InValue, bool bSetupUndoRedo = true)
+	{
+		return SetValueOverride(InName, FAnimNextParamType::GetType<ValueType>(), TConstArrayView<uint8>((uint8*)&InValue, sizeof(ValueType)), bSetupUndoRedo);
+	}
+
+	// Clear the overriden value for the specified name in this entry
+	// @return true if the value was cleared successfully
+	ANIMNEXTUNCOOKEDONLY_API bool ClearValueOverride(FName InName, bool bSetupUndoRedo = true);
+
+	// Get the overriden value for the specified name, if any exist in this entry
+	// @return true if the override exists
+	ANIMNEXTUNCOOKEDONLY_API bool GetValueOverride(FName InName, const FProperty*& OutProperty, TConstArrayView<uint8>& OutValue) const;
+	ANIMNEXTUNCOOKEDONLY_API bool GetValueOverride(FName InName, FAnimNextParamType& OutType, const FProperty*& OutProperty, TConstArrayView<uint8>& OutValue) const;
+
+	// Get the overriden value for the specified name & type, if any exist in the implementation hierarchy. Note: does not return default values, only overrides
+	// @return the status of the override
+	ANIMNEXTUNCOOKEDONLY_API EAnimNextDataInterfaceValueOverrideStatus FindValueOverrideRecursive(FName InName, const FProperty*& OutProperty, TConstArrayView<uint8>& OutValue) const;
+	ANIMNEXTUNCOOKEDONLY_API EAnimNextDataInterfaceValueOverrideStatus FindValueOverrideRecursive(FName InName, FAnimNextParamType& OutType, const FProperty*& OutProperty, TConstArrayView<uint8>& OutValue) const;
+
+	// Get whether this entry contains an override value for the specified named variable
+	// @return true if this entry contains an override for the specified variable
+	ANIMNEXTUNCOOKEDONLY_API bool HasValueOverride(FName InName) const;
+	ANIMNEXTUNCOOKEDONLY_API bool HasValueOverride(FName InName, FAnimNextParamType& OutType) const;
+
+	// Get the value before this data interface 'layer'. Value could be the base value, or overriden by any values in-between in the implementation
+	// hierarchy, but any overrides in this entry are skipped.
+	// @return true if the value was found
+	ANIMNEXTUNCOOKEDONLY_API bool GetDefaultValueRecursive(FName InName, const FProperty*& OutProperty, TConstArrayView<uint8>& OutValue) const;
+	
+	// Get whether this entry contains an override value for the specified named variable that differs from the inherited default (recursive call)
+	// @return true if this entry contains a default override for the specified variable
+	ANIMNEXTUNCOOKEDONLY_API bool HasValueOverrideNotMatchingDefault(FName InName) const;
+
+	// Get whether this entry contains an override value for the specified named variable, or if any overrides exist in the implementation hierarchy
+	// @return the status of the override
+	ANIMNEXTUNCOOKEDONLY_API EAnimNextDataInterfaceValueOverrideStatus GetValueOverrideStatusRecursive(FName InName) const;
+
+	// Get the property bag that contains the value override for the specified named variable
+	ANIMNEXTUNCOOKEDONLY_API EAnimNextDataInterfaceValueOverrideStatus FindValueOverridePropertyBagRecursive(FName InName, FInstancedPropertyBag*& OutPropertyBag) const;
+
+	// Get the property bag that contains the value overrides in this entry
+	FInstancedPropertyBag& GetValueOverridePropertyBag() { return ValueOverrides; }
+
+	// Recompiles this asset when the linked data interface is modified
+	void HandleDataInterfaceModified(UAnimNextRigVMAssetEditorData* InEditorData, EAnimNextEditorDataNotifType InType, UObject* InSubject);
+
+	// Internal helper functions used to find overrides
+	const UAnimNextDataInterfaceEntry* FindOverrideRecursiveHelper(TFunctionRef<bool(const UAnimNextDataInterfaceEntry*)> InPredicate) const;
+	EAnimNextDataInterfaceValueOverrideStatus FindOverrideStatusRecursiveHelper(TFunctionRef<bool(const UAnimNextDataInterfaceEntry*)> InPredicate) const;
+
+	/** The implemented interface */
+	UPROPERTY(VisibleAnywhere, Category = "Data Interface")
+	TObjectPtr<UAnimNextDataInterface> DataInterface;
+
+	/** Soft reference to the Data Interface for error reporting */
+	UPROPERTY()
+	FSoftObjectPath DataInterfacePath;
+
+	/** Property bag for overriden values */
+	UPROPERTY()
+	FInstancedPropertyBag ValueOverrides;
+
+	/** How to automatically bind to the hosting graph or module */
+	UPROPERTY(EditAnywhere, Category = "Data Interface")
+	EAnimNextDataInterfaceAutomaticBindingMode AutomaticBinding = EAnimNextDataInterfaceAutomaticBindingMode::BindSharedInterfaces;
+};
diff --git a/Engine/Plugins/Experimental/AnimNext/Source/AnimNextUncookedOnly/Internal/Entries/AnimNextRigVMAssetEntry.h b/Engine/Plugins/Experimental/AnimNext/Source/AnimNextUncookedOnly/Internal/Entries/AnimNextRigVMAssetEntry.h
--- a/Engine/Plugins/Experimental/AnimNext/Source/AnimNextUncookedOnly/Internal/Entries/AnimNextRigVMAssetEntry.h
+++ b/Engine/Plugins/Experimental/AnimNext/Source/AnimNextUncookedOnly/Internal/Entries/AnimNextRigVMAssetEntry.h
@@ -53,6 +53,7 @@ public:
 	virtual bool IsAsset() const override;
 #if WITH_EDITOR
 	virtual void PostEditChangeProperty(FPropertyChangedEvent& PropertyChangedEvent) override;
+	virtual void PostTransacted(const FTransactionObjectEvent& TransactionEvent) override;
 #endif
 
 protected:
diff --git a/Engine/Plugins/Experimental/AnimNext/Source/AnimNextUncookedOnly/Internal/Entries/AnimNextVariableEntry.h b/Engine/Plugins/Experimental/AnimNext/Source/AnimNextUncookedOnly/Internal/Entries/AnimNextVariableEntry.h
--- a/Engine/Plugins/Experimental/AnimNext/Source/AnimNextUncookedOnly/Internal/Entries/AnimNextVariableEntry.h
+++ b/Engine/Plugins/Experimental/AnimNext/Source/AnimNextUncookedOnly/Internal/Entries/AnimNextVariableEntry.h
@@ -8,12 +8,16 @@
 #include "Variables/IAnimNextRigVMVariableInterface.h"
 #include "Param/ParamType.h"
 #include "Kismet/BlueprintFunctionLibrary.h"
+#include "Logging/StructuredLog.h"
 #include "Variables/AnimNextVariableBinding.h"
 #include "AnimNextVariableEntry.generated.h"
 
 class UAnimNextRigVMAssetEditorData;
 class UAnimNextModule_EditorData;
+class UAnimNextDataInterface_EditorData;
 class UAnimNextVariableBinding;
+class UAnimNextDataInterfaceEntry;
+class FAnimationAnimNextRuntimeTest_Variables;
 
 namespace UE::AnimNext::Editor
 {
@@ -21,11 +25,21 @@ namespace UE::AnimNext::Editor
 	class FVariableBindingPropertyCustomization;
 	class FUniversalObjectLocatorBindingCustomization;
 	class FVariablesOutlinerMode;
+	class SAddVariablesDialog;
+	class FVariablesOutlinerHierarchy;
+	class SVariablesOutlinerValue;
+	class FVariableProxyCustomization;
 }
 
 namespace UE::AnimNext::Tests
 {
 	class FEditor_Variables;
+	class FDataInterfaceCompile;
+}
+
+namespace UE::AnimNext::UncookedOnly
+{
+	struct FUtils;
 }
 
 UCLASS(MinimalAPI, Category = "Variables", DisplayName = "Variable")
@@ -35,11 +49,20 @@ class UAnimNextVariableEntry : public UAnimNextRigVMAssetEntry, public IAnimNext
 
 	friend class UAnimNextRigVMAssetEditorData;
 	friend class UAnimNextModule_EditorData;
+	friend class UAnimNextDataInterface_EditorData;
+	friend class UAnimNextDataInterfaceEntry;
 	friend class UE::AnimNext::Tests::FEditor_Variables;
+	friend class UE::AnimNext::Tests::FDataInterfaceCompile;
+	friend class FAnimationAnimNextRuntimeTest_Variables;
 	friend class UE::AnimNext::Editor::FVariableCustomization;
 	friend class UE::AnimNext::Editor::FVariableBindingPropertyCustomization;
 	friend class UE::AnimNext::Editor::FUniversalObjectLocatorBindingCustomization;
 	friend class UE::AnimNext::Editor::FVariablesOutlinerMode;
+	friend class UE::AnimNext::Editor::SAddVariablesDialog;
+	friend class UE::AnimNext::Editor::SVariablesOutlinerValue;
+	friend class UE::AnimNext::Editor::FVariablesOutlinerHierarchy;
+	friend struct UE::AnimNext::UncookedOnly::FUtils;
+	friend class UE::AnimNext::Editor::FVariableProxyCustomization;
 
 	UAnimNextVariableEntry();
 
@@ -64,9 +87,34 @@ class UAnimNextVariableEntry : public UAnimNextRigVMAssetEntry, public IAnimNext
 	virtual bool SetType(const FAnimNextParamType& InType, bool bSetupUndoRedo = true) override;
 	virtual FName GetVariableName() const override;
 	virtual void SetVariableName(FName InName, bool bSetupUndoRedo = true) override;
-	virtual bool SetDefaultValue(const FString& InDefaultValue, bool bSetupUndoRedo = true) override;
+	virtual bool SetDefaultValue(TConstArrayView<uint8> InValue, bool bSetupUndoRedo = true) override;
+	virtual bool SetDefaultValueFromString(const FString& InDefaultValue, bool bSetupUndoRedo = true) override;
 	virtual const FInstancedPropertyBag& GetPropertyBag() const override;
-	virtual const FAnimNextVariableBinding& GetBinding() const override;
+	virtual bool GetDefaultValue(const FProperty*& OutProperty, TConstArrayView<uint8>& OutValue) const override;
+	virtual void SetBindingType(UScriptStruct* InBindingTypeStruct, bool bSetupUndoRedo = true) override;
+	virtual TConstStructView<FAnimNextVariableBindingData> GetBinding() const override;
+
+	// Set the default value
+	template<typename ValueType>
+	bool SetDefaultValue(const ValueType& InValue, bool bSetupUndoRedo = true)
+	{
+		const FPropertyBagPropertyDesc* Desc = DefaultValue.FindPropertyDescByName(IAnimNextRigVMVariableInterface::ValueName);
+		if(Desc == nullptr)
+		{
+			UE_LOGFMT(LogAnimation, Error, "UAnimNextVariableEntry::SetDefaultValue: Could not find default value in property bag");
+			return false;
+		}
+
+		FAnimNextParamType SuppliedType = FAnimNextParamType::GetType<ValueType>();
+		FAnimNextParamType InternalType = FAnimNextParamType(Desc->ValueType, Desc->ContainerTypes.GetFirstContainerType(), Desc->ValueTypeObject);
+		if(SuppliedType != InternalType)
+		{
+			UE_LOGFMT(LogAnimation, Error, "UAnimNextVariableEntry::SetDefaultValue: Mismatched types");
+			return false;
+		}
+
+		return SetDefaultValue(TConstArrayView<uint8>((uint8*)&InValue, sizeof(ValueType)), bSetupUndoRedo);
+	}
 
 	/** Access specifier - whether the variable is visible external to this asset */
 	UPROPERTY(EditAnywhere, Category = "Variable")
diff --git a/Engine/Plugins/Experimental/AnimNext/Source/AnimNextUncookedOnly/Internal/Graph/AnimNextAnimationGraphSchema.h b/Engine/Plugins/Experimental/AnimNext/Source/AnimNextUncookedOnly/Internal/Graph/AnimNextAnimationGraphSchema.h
--- a/Engine/Plugins/Experimental/AnimNext/Source/AnimNextUncookedOnly/Internal/Graph/AnimNextAnimationGraphSchema.h
+++ b/Engine/Plugins/Experimental/AnimNext/Source/AnimNextUncookedOnly/Internal/Graph/AnimNextAnimationGraphSchema.h
@@ -11,7 +11,5 @@ class UAnimNextAnimationGraphSchema : public UAnimNextRigVMAssetSchema
 {
 	GENERATED_BODY()
 
-	// URigVMSchema interface
-	virtual bool SupportsUnitFunction(URigVMController* InController, const FRigVMFunction* InUnitFunction) const override;
-	virtual bool SupportsDispatchFactory(URigVMController* InController, const FRigVMDispatchFactory* InDispatchFactory) const override;
+	// Schema reserved for future filtering of units etc.
 };
diff --git a/Engine/Plugins/Experimental/AnimNext/Source/AnimNextUncookedOnly/Internal/Graph/AnimNextAnimationGraph_EditorData.h b/Engine/Plugins/Experimental/AnimNext/Source/AnimNextUncookedOnly/Internal/Graph/AnimNextAnimationGraph_EditorData.h
--- a/Engine/Plugins/Experimental/AnimNext/Source/AnimNextUncookedOnly/Internal/Graph/AnimNextAnimationGraph_EditorData.h
+++ b/Engine/Plugins/Experimental/AnimNext/Source/AnimNextUncookedOnly/Internal/Graph/AnimNextAnimationGraph_EditorData.h
@@ -4,15 +4,11 @@
 
 #include "CoreMinimal.h"
 #include "AnimNextController.h"
-#include "RigVMModel/RigVMGraph.h"
-#include "AnimNextEdGraph.h"
-#include "AnimNextRigVMAssetEditorData.h"
+#include "DataInterface/AnimNextDataInterface_EditorData.h"
 #include "AnimNextExecuteContext.h"
-#include "Kismet/BlueprintFunctionLibrary.h"
+
 #include "AnimNextAnimationGraph_EditorData.generated.h"
 
-class UAnimNextAnimationGraph;
-enum class ERigVMGraphNotifType : uint8;
 class FAnimationAnimNextRuntimeTest_GraphAddTrait;
 class FAnimationAnimNextRuntimeTest_GraphExecute;
 class FAnimationAnimNextRuntimeTest_GraphExecuteLatent;
@@ -31,7 +27,7 @@ namespace UE::AnimNext::Editor
 
 /** Editor data for AnimNext animation graphs */
 UCLASS(MinimalAPI)
-class UAnimNextAnimationGraph_EditorData : public UAnimNextRigVMAssetEditorData
+class UAnimNextAnimationGraph_EditorData : public UAnimNextDataInterface_EditorData
 {
 	GENERATED_BODY()
 
diff --git a/Engine/Plugins/Experimental/AnimNext/Source/AnimNextUncookedOnly/Internal/Module/AnimNextModule_EditorData.h b/Engine/Plugins/Experimental/AnimNext/Source/AnimNextUncookedOnly/Internal/Module/AnimNextModule_EditorData.h
--- a/Engine/Plugins/Experimental/AnimNext/Source/AnimNextUncookedOnly/Internal/Module/AnimNextModule_EditorData.h
+++ b/Engine/Plugins/Experimental/AnimNext/Source/AnimNextUncookedOnly/Internal/Module/AnimNextModule_EditorData.h
@@ -6,9 +6,8 @@
 #include "AnimNextController.h"
 #include "RigVMModel/RigVMGraph.h"
 #include "AnimNextEdGraph.h"
-#include "AnimNextRigVMAssetEditorData.h"
+#include "DataInterface/AnimNextDataInterface_EditorData.h"
 #include "AnimNextExecuteContext.h"
-#include "Kismet/BlueprintFunctionLibrary.h"
 #include "AnimNextModule_EditorData.generated.h"
 
 class UAnimNextModule;
@@ -32,7 +31,7 @@ namespace UE::AnimNext::Editor
 
 /** Editor data for AnimNext modules */
 UCLASS(MinimalAPI)
-class UAnimNextModule_EditorData : public UAnimNextRigVMAssetEditorData
+class UAnimNextModule_EditorData : public UAnimNextDataInterface_EditorData
 {
 	GENERATED_BODY()
 
@@ -59,6 +58,7 @@ private:
 	virtual TConstArrayView<TSubclassOf<UAnimNextRigVMAssetEntry>> GetEntryClasses() const override;
 	virtual void RecompileVM() override;
 	virtual void GetProgrammaticGraphs(const FRigVMCompileSettings& InSettings, TArray<URigVMGraph*>& OutGraphs) override;
+	virtual void CustomizeNewAssetEntry(UAnimNextRigVMAssetEntry* InNewEntry) const override;
 
 private:
 	UPROPERTY()
diff --git a/Engine/Plugins/Experimental/AnimNext/Source/AnimNextUncookedOnly/Internal/UncookedOnlyUtils.h b/Engine/Plugins/Experimental/AnimNext/Source/AnimNextUncookedOnly/Internal/UncookedOnlyUtils.h
--- a/Engine/Plugins/Experimental/AnimNext/Source/AnimNextUncookedOnly/Internal/UncookedOnlyUtils.h
+++ b/Engine/Plugins/Experimental/AnimNext/Source/AnimNextUncookedOnly/Internal/UncookedOnlyUtils.h
@@ -101,7 +101,7 @@ struct ANIMNEXTUNCOOKEDONLY_API FUtils
 {
 	static void CompileVariables(UAnimNextRigVMAsset* InAsset);
 
-	static URigVMGraph* CompileVariableBindings(const FRigVMCompileSettings& InSettings, UAnimNextRigVMAsset* InAsset);
+	static void CompileVariableBindings(const FRigVMCompileSettings& InSettings, UAnimNextRigVMAsset* InAsset, TArray<URigVMGraph*>& OutGraphs);
 
 	static void RecreateVM(UAnimNextRigVMAsset* InAsset);
 
@@ -111,7 +111,8 @@ struct ANIMNEXTUNCOOKEDONLY_API FUtils
 	template<typename AssetType, typename EditorDataType>
 	static AssetType* GetAsset(EditorDataType* InEditorData)
 	{
-		return CastChecked<AssetType>(GetAsset(const_cast<EditorDataType*>(InEditorData)));
+		using NonConstEditorDataType = std::remove_const_t<EditorDataType>;
+		return CastChecked<AssetType>(GetAsset(const_cast<NonConstEditorDataType*>(InEditorData)));
 	}
 
 	// Get the corresponding editor data from an asset (casts the editor data appropriately)
@@ -120,7 +121,8 @@ struct ANIMNEXTUNCOOKEDONLY_API FUtils
 	template<typename EditorDataType, typename AssetType>
 	static EditorDataType* GetEditorData(AssetType* InAsset)
 	{
-		return CastChecked<EditorDataType>(GetEditorData(const_cast<AssetType*>(InAsset)));
+		using NonConstAssetType = std::remove_const_t<AssetType>;
+		return CastChecked<EditorDataType>(GetEditorData(const_cast<NonConstAssetType*>(InAsset)));
 	}
 
 	/**
diff --git a/Engine/Plugins/Experimental/AnimNext/Source/AnimNextUncookedOnly/Internal/Variables/AnimNextUniversalObjectLocatorBindingData.h b/Engine/Plugins/Experimental/AnimNext/Source/AnimNextUncookedOnly/Internal/Variables/AnimNextUniversalObjectLocatorBindingData.h
--- a/Engine/Plugins/Experimental/AnimNext/Source/AnimNextUncookedOnly/Internal/Variables/AnimNextUniversalObjectLocatorBindingData.h
+++ b/Engine/Plugins/Experimental/AnimNext/Source/AnimNextUncookedOnly/Internal/Variables/AnimNextUniversalObjectLocatorBindingData.h
@@ -17,8 +17,8 @@ enum class FAnimNextUniversalObjectLocatorBindingType : uint8
 	HoistedFunction,
 };
 
-// Wraps a universal object locator to identify external object instances to the parameter system
-USTRUCT()
+// Allows binding of module variables to gameplay data via Universal Object Locators
+USTRUCT(DisplayName = "Universal Object Locator")
 struct FAnimNextUniversalObjectLocatorBindingData : public FAnimNextVariableBindingData
 {
 	GENERATED_BODY()
diff --git a/Engine/Plugins/Experimental/AnimNext/Source/AnimNextUncookedOnly/Private/AnimNextEdGraphNode.cpp b/Engine/Plugins/Experimental/AnimNext/Source/AnimNextUncookedOnly/Private/AnimNextEdGraphNode.cpp
--- a/Engine/Plugins/Experimental/AnimNext/Source/AnimNextUncookedOnly/Private/AnimNextEdGraphNode.cpp
+++ b/Engine/Plugins/Experimental/AnimNext/Source/AnimNextUncookedOnly/Private/AnimNextEdGraphNode.cpp
@@ -193,7 +193,7 @@ void UAnimNextEdGraphNode::BuildExposeVariablesContextMenu(UToolMenu* SubMenu)
 	FContentBrowserModule& ContentBrowserModule = FModuleManager::Get().LoadModuleChecked<FContentBrowserModule>(TEXT("ContentBrowser"));
 
 	FAssetPickerConfig AssetPickerConfig;
-	AssetPickerConfig.Filter.ClassPaths.Add(UAnimNextAnimationGraph::StaticClass()->GetClassPathName());
+	AssetPickerConfig.Filter.ClassPaths.Add(UAnimNextDataInterface::StaticClass()->GetClassPathName());
 	AssetPickerConfig.InitialAssetViewType = EAssetViewType::List;
 	AssetPickerConfig.AssetShowWarningText = LOCTEXT("NoAssetsWithPublicVariablesMessage", "No animation graphs with public variables found");
 	AssetPickerConfig.OnAssetSelected = FOnAssetSelected::CreateLambda([VMController, VMNode](const FAssetData& InAssetData)
@@ -203,7 +203,7 @@ void UAnimNextEdGraphNode::BuildExposeVariablesContextMenu(UToolMenu* SubMenu)
 		FString DefaultValue;
 		FRigVMTrait_AnimNextPublicVariables DefaultTrait;
 		FRigVMTrait_AnimNextPublicVariables NewTrait;
-		UAnimNextRigVMAsset* Asset = CastChecked<UAnimNextRigVMAsset>(InAssetData.GetAsset());
+		UAnimNextDataInterface* Asset = CastChecked<UAnimNextDataInterface>(InAssetData.GetAsset());
 		NewTrait.Asset = Asset;
 		const FInstancedPropertyBag& PublicVariableDefaults = Asset->GetPublicVariableDefaults();
 		TConstArrayView<FPropertyBagPropertyDesc> Descs = PublicVariableDefaults.GetPropertyBagStruct()->GetPropertyDescs();
@@ -234,9 +234,6 @@ void UAnimNextEdGraphNode::BuildExposeVariablesContextMenu(UToolMenu* SubMenu)
 				}
 			}
 		}
-
-		// TODO: filter already-added traits
-
 		return true;
 	});
 
diff --git a/Engine/Plugins/Experimental/AnimNext/Source/AnimNextUncookedOnly/Private/AnimNextRigVMAssetEditorData.cpp b/Engine/Plugins/Experimental/AnimNext/Source/AnimNextUncookedOnly/Private/AnimNextRigVMAssetEditorData.cpp
--- a/Engine/Plugins/Experimental/AnimNext/Source/AnimNextUncookedOnly/Private/AnimNextRigVMAssetEditorData.cpp
+++ b/Engine/Plugins/Experimental/AnimNext/Source/AnimNextUncookedOnly/Private/AnimNextRigVMAssetEditorData.cpp
@@ -14,23 +14,23 @@
 #include "IAnimNextRigVMGraphInterface.h"
 #include "IWorkspaceEditor.h"
 #include "IWorkspaceEditorModule.h"
-#include "Variables/IAnimNextRigVMVariableInterface.h"
 #include "UncookedOnlyUtils.h"
+#include "DataInterface/AnimNextDataInterface_EditorData.h"
 #include "Entries/AnimNextAnimationGraphEntry.h"
+#include "Entries/AnimNextDataInterfaceEntry.h"
 #include "Entries/AnimNextEventGraphEntry.h"
 #include "Entries/AnimNextVariableEntry.h"
 #include "Graph/AnimNextAnimationGraphSchema.h"
 #include "Graph/RigUnit_AnimNextBeginExecution.h"
 #include "Misc/TransactionObjectEvent.h"
 #include "Module/AnimNextEventGraphSchema.h"
-#include "Module/AnimNextModule_EditorData.h"
 #include "Module/RigUnit_AnimNextModuleEvents.h"
-#include "Param/AnimNextTag.h"
 #include "RigVMModel/RigVMFunctionLibrary.h"
 #include "RigVMModel/RigVMNotifications.h"
 #include "RigVMModel/Nodes/RigVMAggregateNode.h"
 #include "RigVMModel/Nodes/RigVMCollapseNode.h"
 #include "UObject/AssetRegistryTagsContext.h"
+#include "DataInterface/AnimNextDataInterface_EditorData.h"
 
 TAutoConsoleVariable<bool> CVarDumpProgrammaticGraphs(
 	TEXT("AnimNext.DumpProgrammaticGraphs"),
@@ -200,7 +200,7 @@ void UAnimNextRigVMAssetEditorData::PostTransacted(const FTransactionObjectEvent
 
 	if (TransactionEvent.GetEventType() == ETransactionObjectEventType::UndoRedo)
 	{
-		BroadcastModified(EAnimNextEditorDataNotifType::PropertyChanged, this);
+		BroadcastModified(EAnimNextEditorDataNotifType::UndoRedo, this);
 	}
 }
 
@@ -978,6 +978,29 @@ bool UAnimNextRigVMAssetEditorData::RemoveEntries(TConstArrayView<UAnimNextRigVM
 	return bResult;
 }
 
+bool UAnimNextRigVMAssetLibrary::RemoveAllEntries(UAnimNextRigVMAsset* InAsset, bool bSetupUndoRedo, bool bPrintPythonCommand)
+{
+	return UE::AnimNext::UncookedOnly::FUtils::GetEditorData(InAsset)->RemoveAllEntries(bSetupUndoRedo, bPrintPythonCommand);
+}
+
+bool UAnimNextRigVMAssetEditorData::RemoveAllEntries(bool bSetupUndoRedo, bool bPrintPythonCommand)
+{
+	bool bResult = false;
+	{
+		TGuardValue<bool> DisableEditorDataNotifications(bSuspendEditorDataNotifications, true);
+		TGuardValue<bool> DisableAutoCompile(bAutoRecompileVM, false);
+		TArray<UAnimNextRigVMAssetEntry*> EntriesCopy = Entries; 
+		for(UAnimNextRigVMAssetEntry* Entry : EntriesCopy)
+		{
+			bResult |= RemoveEntry(Entry, bSetupUndoRedo, bPrintPythonCommand);
+		}
+	}
+
+	BroadcastModified(EAnimNextEditorDataNotifType::EntryRemoved, this);
+
+	return bResult;
+}
+
 UObject* UAnimNextRigVMAssetEditorData::CreateNewSubEntry(UAnimNextRigVMAssetEditorData* InEditorData, TSubclassOf<UObject> InClass)
 {
 	UObject* NewEntry = NewObject<UObject>(InEditorData, InClass.Get(), NAME_None, RF_Transactional);
@@ -1272,10 +1295,10 @@ bool UAnimNextRigVMAssetEditorData::RemoveEdGraph(URigVMGraph* InModel)
 	return false;
 }
 
-UAnimNextVariableEntry* UAnimNextRigVMAssetLibrary::AddVariable(UAnimNextModule* InModule, FName InName, EPropertyBagPropertyType InValueType,
+UAnimNextVariableEntry* UAnimNextRigVMAssetLibrary::AddVariable(UAnimNextRigVMAsset* InAsset, FName InName, EPropertyBagPropertyType InValueType,
 	EPropertyBagContainerType InContainerType, const UObject* InValueTypeObject, const FString& InDefaultValue, bool bSetupUndoRedo, bool bPrintPythonCommand)
 {
-	return UE::AnimNext::UncookedOnly::FUtils::GetEditorData(InModule)->AddVariable(InName, FAnimNextParamType(InValueType, InContainerType, InValueTypeObject), InDefaultValue, bSetupUndoRedo, bPrintPythonCommand);
+	return UE::AnimNext::UncookedOnly::FUtils::GetEditorData(InAsset)->AddVariable(InName, FAnimNextParamType(InValueType, InContainerType, InValueTypeObject), InDefaultValue, bSetupUndoRedo, bPrintPythonCommand);
 }
 
 UAnimNextVariableEntry* UAnimNextRigVMAssetEditorData::AddVariable(FName InName, FAnimNextParamType InType, const FString& InDefaultValue, bool bSetupUndoRedo, bool bPrintPythonCommand)
@@ -1316,8 +1339,10 @@ UAnimNextVariableEntry* UAnimNextRigVMAssetEditorData::AddVariable(FName InName,
 		NewEntry->SetType(InType, false);
 		if(InDefaultValue.Len() > 0)
 		{
-			NewEntry->SetDefaultValue(InDefaultValue, false);
+			NewEntry->SetDefaultValueFromString(InDefaultValue, false);
 		}
+
+		NewEntry->Initialize(this);
 	}
 
 	if(bSetupUndoRedo)
@@ -1328,14 +1353,16 @@ UAnimNextVariableEntry* UAnimNextRigVMAssetEditorData::AddVariable(FName InName,
 
 	Entries.Add(NewEntry);
 
+	CustomizeNewAssetEntry(NewEntry);
+
 	BroadcastModified(EAnimNextEditorDataNotifType::EntryAdded, NewEntry);
 
 	return NewEntry;
 }
 
-UAnimNextEventGraphEntry* UAnimNextRigVMAssetLibrary::AddEventGraph(UAnimNextModule* InModule, FName InName, UScriptStruct* InEventStruct, bool bSetupUndoRedo, bool bPrintPythonCommand)
+UAnimNextEventGraphEntry* UAnimNextRigVMAssetLibrary::AddEventGraph(UAnimNextRigVMAsset* InAsset, FName InName, UScriptStruct* InEventStruct, bool bSetupUndoRedo, bool bPrintPythonCommand)
 {
-	return UE::AnimNext::UncookedOnly::FUtils::GetEditorData(InModule)->AddEventGraph(InName, InEventStruct, bSetupUndoRedo, bPrintPythonCommand);
+	return UE::AnimNext::UncookedOnly::FUtils::GetEditorData(InAsset)->AddEventGraph(InName, InEventStruct, bSetupUndoRedo, bPrintPythonCommand);
 }
 
 UAnimNextEventGraphEntry* UAnimNextRigVMAssetEditorData::AddEventGraph(FName InName, UScriptStruct* InEventStruct, bool bSetupUndoRedo, bool bPrintPythonCommand)
@@ -1375,6 +1402,7 @@ UAnimNextEventGraphEntry* UAnimNextRigVMAssetEditorData::AddEventGraph(FName InN
 
 	UAnimNextEventGraphEntry* NewEntry = CreateNewSubEntry<UAnimNextEventGraphEntry>(this);
 	NewEntry->GraphName = NewGraphName;
+	NewEntry->Initialize(this);
 
 	if(bSetupUndoRedo)
 	{
@@ -1407,14 +1435,16 @@ UAnimNextEventGraphEntry* UAnimNextRigVMAssetEditorData::AddEventGraph(FName InN
 		}
 	}
 
+	CustomizeNewAssetEntry(NewEntry);
+
 	BroadcastModified(EAnimNextEditorDataNotifType::EntryAdded, NewEntry);
 
 	return NewEntry;
 }
 
-UAnimNextAnimationGraphEntry* UAnimNextRigVMAssetLibrary::AddAnimationGraph(UAnimNextModule* InModule, FName InName, bool bSetupUndoRedo, bool bPrintPythonCommand)
+UAnimNextAnimationGraphEntry* UAnimNextRigVMAssetLibrary::AddAnimationGraph(UAnimNextRigVMAsset* InAsset, FName InName, bool bSetupUndoRedo, bool bPrintPythonCommand)
 {
-	return UE::AnimNext::UncookedOnly::FUtils::GetEditorData(InModule)->AddAnimationGraph(InName, bSetupUndoRedo, bPrintPythonCommand);
+	return UE::AnimNext::UncookedOnly::FUtils::GetEditorData(InAsset)->AddAnimationGraph(InName, bSetupUndoRedo, bPrintPythonCommand);
 }
 
 UAnimNextAnimationGraphEntry* UAnimNextRigVMAssetEditorData::AddAnimationGraph(FName InName, bool bSetupUndoRedo, bool bPrintPythonCommand)
@@ -1448,6 +1478,7 @@ UAnimNextAnimationGraphEntry* UAnimNextRigVMAssetEditorData::AddAnimationGraph(F
 
 	UAnimNextAnimationGraphEntry* NewEntry = CreateNewSubEntry<UAnimNextAnimationGraphEntry>(this);
 	NewEntry->GraphName = NewGraphName;
+	NewEntry->Initialize(this);
 
 	if(bSetupUndoRedo)
 	{
@@ -1480,6 +1511,131 @@ UAnimNextAnimationGraphEntry* UAnimNextRigVMAssetEditorData::AddAnimationGraph(F
 		}
 	}
 
+	CustomizeNewAssetEntry(NewEntry);
+
+	BroadcastModified(EAnimNextEditorDataNotifType::EntryAdded, NewEntry);
+
+	return NewEntry;
+}
+
+UAnimNextDataInterfaceEntry* UAnimNextRigVMAssetLibrary::AddDataInterface(UAnimNextRigVMAsset* InAsset, UAnimNextDataInterface* InDataInterface, bool bSetupUndoRedo, bool bPrintPythonCommand)
+{
+	return UE::AnimNext::UncookedOnly::FUtils::GetEditorData(InAsset)->AddDataInterface(InDataInterface, bSetupUndoRedo, bPrintPythonCommand);
+}
+
+UAnimNextDataInterfaceEntry* UAnimNextRigVMAssetEditorData::AddDataInterface(UAnimNextDataInterface* InDataInterface, bool bSetupUndoRedo, bool bPrintPythonCommand)
+{
+	if(InDataInterface == nullptr)
+	{
+		ReportError(TEXT("UAnimNextRigVMAssetEditorData::AddDataInterface: Invalid data interface supplied."));
+		return nullptr;
+	}
+
+	if(!GetEntryClasses().Contains(UAnimNextDataInterfaceEntry::StaticClass()) || !CanAddNewEntry(UAnimNextDataInterfaceEntry::StaticClass()))
+	{
+		ReportError(TEXT("UAnimNextRigVMAssetEditorData::AddDataInterface: Cannot add a data interface to this asset - entry is not allowed."));
+		return nullptr;
+	}
+	
+	// Check if interface has any public members or if any of its parent interfaces do
+	UAnimNextDataInterface_EditorData* EditorData = UE::AnimNext::UncookedOnly::FUtils::GetEditorData<UAnimNextDataInterface_EditorData>(InDataInterface);
+	if(EditorData == nullptr)
+	{
+		ReportError(TEXT("UAnimNextRigVMAssetEditorData::AddDataInterface: Invalid data interface supplied - asset has no editor data."));
+		return nullptr;
+	}
+
+	// Check for circularity
+	auto CheckForCircularity = [this](UAnimNextDataInterface_EditorData* InEditorData, auto& InCheckForCircularity)
+	{
+		if(InEditorData == this)
+		{
+			return true;
+		}
+
+		for(UAnimNextRigVMAssetEntry* Entry : InEditorData->Entries)
+		{
+			if(UAnimNextDataInterfaceEntry* DataInterfaceEntry = Cast<UAnimNextDataInterfaceEntry>(Entry))
+			{
+				UAnimNextDataInterface* DataInterface = DataInterfaceEntry->GetDataInterface();
+				UAnimNextDataInterface_EditorData* EditorData = UE::AnimNext::UncookedOnly::FUtils::GetEditorData<UAnimNextDataInterface_EditorData>(DataInterface);
+				if(InCheckForCircularity(EditorData, InCheckForCircularity))
+				{
+					return true;
+				}
+			}
+		}
+
+		return false;
+	};
+
+	if(CheckForCircularity(EditorData, CheckForCircularity))
+	{
+		ReportError(TEXT("UAnimNextRigVMAssetEditorData::AddDataInterface: Circular reference detected."));
+		return nullptr;
+	}
+
+	auto CheckForPublicMembers = [](UAnimNextDataInterface_EditorData* InEditorData, auto& InCheckForPublicMembers)
+	{
+		for(UAnimNextRigVMAssetEntry* Entry : InEditorData->Entries)
+		{
+			if(UAnimNextVariableEntry* VariableEntry = Cast<UAnimNextVariableEntry>(Entry))
+			{
+				if(VariableEntry->GetExportAccessSpecifier() == EAnimNextExportAccessSpecifier::Public)
+				{
+					return true;
+				}
+			}
+			else if(UAnimNextDataInterfaceEntry* DataInterfaceEntry = Cast<UAnimNextDataInterfaceEntry>(Entry))
+			{
+				UAnimNextDataInterface* DataInterface = DataInterfaceEntry->GetDataInterface();
+				UAnimNextDataInterface_EditorData* EditorData = UE::AnimNext::UncookedOnly::FUtils::GetEditorData<UAnimNextDataInterface_EditorData>(DataInterface);
+				if(InCheckForPublicMembers(EditorData, InCheckForPublicMembers))
+				{
+					return true;
+				}
+			}
+		}
+
+		return false;
+	};
+
+	if(!CheckForPublicMembers(EditorData, CheckForPublicMembers))
+	{
+		ReportError(TEXT("UAnimNextRigVMAssetEditorData::AddDataInterface: No public variables found."));
+		return nullptr;
+	}
+
+	// Check for duplicate interface
+	auto DuplicatePredicate = [InDataInterface](const UAnimNextRigVMAssetEntry* InEntry)
+	{
+		if(const UAnimNextDataInterfaceEntry* InterfaceEntry = Cast<UAnimNextDataInterfaceEntry>(InEntry))
+		{
+			return InterfaceEntry->DataInterface == InDataInterface;
+		}
+		return false;
+	};
+
+	if(Entries.ContainsByPredicate(DuplicatePredicate))
+	{
+		ReportError(TEXT("UAnimNextRigVMAssetEditorData::AddDataInterface: Data interface already implemented."));
+		return nullptr;
+	}
+
+	UAnimNextDataInterfaceEntry* NewEntry = CreateNewSubEntry<UAnimNextDataInterfaceEntry>(this);
+	NewEntry->SetDataInterface(InDataInterface);
+	NewEntry->Initialize(this);
+
+	if(bSetupUndoRedo)
+	{
+		NewEntry->Modify();
+		Modify();
+	}
+
+	Entries.Add(NewEntry);
+
+	CustomizeNewAssetEntry(NewEntry);
+
 	BroadcastModified(EAnimNextEditorDataNotifType::EntryAdded, NewEntry);
 
 	return NewEntry;
@@ -1489,9 +1645,20 @@ bool UAnimNextRigVMAssetEditorData::HasPublicVariables() const
 {
 	for(UAnimNextRigVMAssetEntry* Entry : Entries)
 	{
-		if(UAnimNextVariableEntry* Variable = Cast<UAnimNextVariableEntry>(Entry))
+		if(UAnimNextVariableEntry* VariableEntry = Cast<UAnimNextVariableEntry>(Entry))
 		{
-			return Variable->GetExportAccessSpecifier() == EAnimNextExportAccessSpecifier::Public;
+			if(VariableEntry->GetExportAccessSpecifier() == EAnimNextExportAccessSpecifier::Public)
+			{
+				return true;
+			}
+		}
+		else if(UAnimNextDataInterfaceEntry* DataInterfaceEntry = Cast<UAnimNextDataInterfaceEntry>(Entry))
+		{
+			if(DataInterfaceEntry->DataInterface)
+			{
+				UAnimNextDataInterface_EditorData* EditorData = UE::AnimNext::UncookedOnly::FUtils::GetEditorData<UAnimNextDataInterface_EditorData>(DataInterfaceEntry->DataInterface.Get());
+				return EditorData->HasPublicVariables();
+			}
 		}
 	}
 	return false;
diff --git a/Engine/Plugins/Experimental/AnimNext/Source/AnimNextUncookedOnly/Private/DataInterface/AnimNextDataInterface_EditorData.cpp b/Engine/Plugins/Experimental/AnimNext/Source/AnimNextUncookedOnly/Private/DataInterface/AnimNextDataInterface_EditorData.cpp
new file mode 100644
--- /dev/null
+++ b/Engine/Plugins/Experimental/AnimNext/Source/AnimNextUncookedOnly/Private/DataInterface/AnimNextDataInterface_EditorData.cpp
@@ -0,0 +1,35 @@
+﻿// Copyright Epic Games, Inc. All Rights Reserved.
+
+#include "DataInterface/AnimNextDataInterface_EditorData.h"
+
+#include "Entries/AnimNextDataInterfaceEntry.h"
+#include "Entries/AnimNextVariableEntry.h"
+
+TConstArrayView<TSubclassOf<UAnimNextRigVMAssetEntry>> UAnimNextDataInterface_EditorData::GetEntryClasses() const
+{
+	static const TSubclassOf<UAnimNextRigVMAssetEntry> Classes[] =
+	{
+		UAnimNextVariableEntry::StaticClass(),
+		UAnimNextDataInterfaceEntry::StaticClass(),
+	};
+
+	return Classes;
+}
+
+void UAnimNextDataInterface_EditorData::CustomizeNewAssetEntry(UAnimNextRigVMAssetEntry* InNewEntry) const
+{
+	UAnimNextVariableEntry* VariableEntry = Cast<UAnimNextVariableEntry>(InNewEntry);
+	if(VariableEntry == nullptr)
+	{
+		return;
+	}
+	
+	const bool bIsDataInterface = ExactCast<UAnimNextDataInterface_EditorData>(this) != nullptr;
+	if(!bIsDataInterface)
+	{
+		return;
+	}
+
+	// Force all data interface variables to be public
+	VariableEntry->SetExportAccessSpecifier(EAnimNextExportAccessSpecifier::Public, false);
+}
\ No newline at end of file
diff --git a/Engine/Plugins/Experimental/AnimNext/Source/AnimNextUncookedOnly/Private/Entries/AnimNextDataInterfaceEntry.cpp b/Engine/Plugins/Experimental/AnimNext/Source/AnimNextUncookedOnly/Private/Entries/AnimNextDataInterfaceEntry.cpp
new file mode 100644
--- /dev/null
+++ b/Engine/Plugins/Experimental/AnimNext/Source/AnimNextUncookedOnly/Private/Entries/AnimNextDataInterfaceEntry.cpp
@@ -0,0 +1,366 @@
+﻿// Copyright Epic Games, Inc. All Rights Reserved.
+
+#include "Entries/AnimNextDataInterfaceEntry.h"
+
+#include "UncookedOnlyUtils.h"
+#include "DataInterface/AnimNextDataInterface.h"
+#include "DataInterface/AnimNextDataInterface_EditorData.h"
+#include "Entries/AnimNextVariableEntry.h"
+#include "Logging/StructuredLog.h"
+#include "Param/ParamType.h"
+
+#define LOCTEXT_NAMESPACE "AnimNextDataInterfaceEntry"
+
+void UAnimNextDataInterfaceEntry::Initialize(UAnimNextRigVMAssetEditorData* InEditorData)
+{
+	Super::Initialize(InEditorData);
+
+	if(DataInterface)
+	{
+		UAnimNextDataInterface_EditorData* DataInterfaceEditorData = UE::AnimNext::UncookedOnly::FUtils::GetEditorData<UAnimNextDataInterface_EditorData>(DataInterface.Get());
+		DataInterfaceEditorData->ModifiedDelegate.AddUObject(this, &UAnimNextDataInterfaceEntry::HandleDataInterfaceModified);
+	}
+}
+
+FName UAnimNextDataInterfaceEntry::GetEntryName() const
+{
+	return DataInterface ? DataInterface->GetFName() : NAME_None;
+}
+
+FText UAnimNextDataInterfaceEntry::GetDisplayName() const
+{
+	return DataInterface ? FText::FromName(DataInterface->GetFName()) : LOCTEXT("InvalidDataInterface", "Invalid Data Interface");
+}
+
+FText UAnimNextDataInterfaceEntry::GetDisplayNameTooltip() const
+{
+	return DataInterface ? FText::FromName(DataInterface->GetFName()) : LOCTEXT("InvalidDataInterfaceTooltip", "Invalid or deleted Data Interface");
+}
+
+void UAnimNextDataInterfaceEntry::SetDataInterface(UAnimNextDataInterface* InDataInterface, bool bSetupUndoRedo)
+{
+	check(InDataInterface != nullptr);
+
+	if(bSetupUndoRedo)
+	{
+		Modify();
+	}
+
+	DataInterface = InDataInterface;
+	DataInterfacePath = FSoftObjectPath(DataInterface);
+	ValueOverrides.Reset();
+}
+
+UAnimNextDataInterface* UAnimNextDataInterfaceEntry::GetDataInterface() const
+{
+	return DataInterface;
+}
+
+FSoftObjectPath UAnimNextDataInterfaceEntry::GetDataInterfacePath() const
+{
+	return DataInterfacePath;
+}
+
+bool UAnimNextDataInterfaceEntry::SetValueOverrideToDefault(FName InName, bool bSetupUndoRedo)
+{
+	const FProperty* DefaultValueProperty = nullptr;
+	TConstArrayView<uint8> DefaultValue;
+	if(!GetDefaultValueRecursive(InName, DefaultValueProperty, DefaultValue))
+	{
+		// No value, so cannot compare with default
+		UE_LOGFMT(LogAnimation, Error, "UAnimNextDataInterfaceEntry::SetValueOverrideToDefault: Could not find a default value for variable {Name}", InName);
+		return false;
+	}
+
+	return SetValueOverride(InName, FAnimNextParamType::FromProperty(DefaultValueProperty), DefaultValue, bSetupUndoRedo);
+}
+
+bool UAnimNextDataInterfaceEntry::SetValueOverride(FName InName, const FAnimNextParamType& InType, TConstArrayView<uint8> InValue, bool bSetupUndoRedo)
+{
+	check(InName != NAME_None);
+	check(InType.IsValid());
+	check(InValue.GetData() != nullptr);
+
+	if(bSetupUndoRedo)
+	{
+		Modify();
+	}
+
+	const FPropertyBagPropertyDesc* Desc = ValueOverrides.FindPropertyDescByName(InName);
+	if(Desc == nullptr)
+	{
+		ValueOverrides.AddContainerProperty(InName, InType.GetContainerType(), InType.GetValueType(), InType.GetValueTypeObject());
+		Desc = ValueOverrides.FindPropertyDescByName(InName);
+	}
+
+	if(Desc == nullptr)
+	{
+		UE_LOGFMT(LogAnimation, Error, "UAnimNextDataInterfaceEntry::SetValueOverride: Failed to add value override to property bag for {Name}", InName);
+		return false;
+	}
+
+	// Check type matches
+	FAnimNextParamType FoundType(Desc->ValueType, Desc->ContainerTypes.GetFirstContainerType(), Desc->ValueTypeObject);
+	if(FoundType != InType)
+	{
+		UE_LOGFMT(LogAnimation, Error, "UAnimNextDataInterfaceEntry::SetValueOverride: Failed to add value override of the correct type to property bag for {Name}", InName);
+		return false;
+	}
+	
+	check(Desc->CachedProperty);
+	if(Desc->CachedProperty->GetElementSize() != InValue.Num())
+	{
+		UE_LOGFMT(LogAnimation, Error, "UAnimNextDataInterfaceEntry::SetValueOverride: Mismatched buffer sizes (%d vs %d)", Desc->CachedProperty->GetElementSize(), InValue.Num());
+		return false;
+	}
+
+	uint8* DestPtr = Desc->CachedProperty->ContainerPtrToValuePtr<uint8>(ValueOverrides.GetMutableValue().GetMemory());
+	const uint8* SrcPtr = InValue.GetData();
+	Desc->CachedProperty->CopyCompleteValue(DestPtr, SrcPtr);
+
+	BroadcastModified(EAnimNextEditorDataNotifType::VariableDefaultValueChanged);
+
+	return true;
+}
+
+bool UAnimNextDataInterfaceEntry::ClearValueOverride(FName InName, bool bSetupUndoRedo)
+{
+	if(bSetupUndoRedo)
+	{
+		Modify();
+	}
+
+	const FPropertyBagPropertyDesc* Desc = ValueOverrides.FindPropertyDescByName(InName);
+	if(Desc == nullptr)
+	{
+		UE_LOGFMT(LogAnimation, Error, "Failed to clear value override in property bag");
+		return false;
+	}
+
+	ValueOverrides.RemovePropertyByName(InName);
+
+	BroadcastModified(EAnimNextEditorDataNotifType::VariableDefaultValueChanged);
+
+	return true;
+}
+
+const UAnimNextDataInterfaceEntry* UAnimNextDataInterfaceEntry::FindOverrideRecursiveHelper(TFunctionRef<bool(const UAnimNextDataInterfaceEntry*)> InPredicate) const
+{
+	if(DataInterface == nullptr)
+	{
+		return nullptr;
+	}
+
+	if(InPredicate(this))
+	{
+		return this;
+	}
+
+	UAnimNextDataInterface_EditorData* EditorData = UE::AnimNext::UncookedOnly::FUtils::GetEditorData<UAnimNextDataInterface_EditorData>(DataInterface.Get());
+	for(const UAnimNextRigVMAssetEntry* Entry : EditorData->Entries)
+	{
+		if(const UAnimNextDataInterfaceEntry* DataInterfaceEntry = Cast<UAnimNextDataInterfaceEntry>(Entry))
+		{
+			if(InPredicate(DataInterfaceEntry))
+			{
+				return DataInterfaceEntry;
+			}
+			else
+			{
+				return DataInterfaceEntry->FindOverrideRecursiveHelper(InPredicate);
+			}
+		}
+	}
+
+	return nullptr;
+}
+
+EAnimNextDataInterfaceValueOverrideStatus UAnimNextDataInterfaceEntry::FindOverrideStatusRecursiveHelper(TFunctionRef<bool(const UAnimNextDataInterfaceEntry*)> InPredicate) const
+{
+	const UAnimNextDataInterfaceEntry* OverridingEntry = FindOverrideRecursiveHelper(InPredicate);
+	if(OverridingEntry == this)
+	{
+		return EAnimNextDataInterfaceValueOverrideStatus::OverriddenInThisAsset;
+	}
+
+	if(OverridingEntry != nullptr)
+	{
+		return EAnimNextDataInterfaceValueOverrideStatus::OverriddenInParentAsset;
+	}
+
+	return EAnimNextDataInterfaceValueOverrideStatus::NotOverridden;
+}
+
+bool UAnimNextDataInterfaceEntry::GetValueOverride(FName InName, const FProperty*& OutProperty, TConstArrayView<uint8>& OutValue) const
+{
+	const FPropertyBagPropertyDesc* Desc = ValueOverrides.FindPropertyDescByName(InName);
+	if(Desc == nullptr)
+	{
+		return false;
+	}
+
+	check(Desc->CachedProperty);
+	OutProperty = Desc->CachedProperty;
+	OutValue = TConstArrayView<uint8>(Desc->CachedProperty->ContainerPtrToValuePtr<uint8>(ValueOverrides.GetValue().GetMemory()), Desc->CachedProperty->GetElementSize());
+	return true;
+}
+
+bool UAnimNextDataInterfaceEntry::GetValueOverride(FName InName, FAnimNextParamType& OutType, const FProperty*& OutProperty, TConstArrayView<uint8>& OutValue) const
+{
+	const FPropertyBagPropertyDesc* Desc = ValueOverrides.FindPropertyDescByName(InName);
+	if(Desc == nullptr)
+	{
+		return false;
+	}
+
+	check(Desc->CachedProperty);
+	OutType = FAnimNextParamType(Desc->ValueType, Desc->ContainerTypes.GetFirstContainerType(), Desc->ValueTypeObject);
+	OutProperty = Desc->CachedProperty;
+	OutValue = TConstArrayView<uint8>(Desc->CachedProperty->ContainerPtrToValuePtr<uint8>(ValueOverrides.GetValue().GetMemory()), Desc->CachedProperty->GetElementSize());
+	return true;
+}
+
+EAnimNextDataInterfaceValueOverrideStatus UAnimNextDataInterfaceEntry::FindValueOverrideRecursive(FName InName, const FProperty*& OutProperty, TConstArrayView<uint8>& OutValue) const
+{
+	auto CheckOverride = [&InName, &OutValue, &OutProperty](const UAnimNextDataInterfaceEntry* InInterfaceEntry) -> bool
+	{
+		return InInterfaceEntry->GetValueOverride(InName, OutProperty, OutValue);
+	};
+
+	return FindOverrideStatusRecursiveHelper(CheckOverride);
+}
+
+EAnimNextDataInterfaceValueOverrideStatus UAnimNextDataInterfaceEntry::FindValueOverrideRecursive(FName InName, FAnimNextParamType& OutType, const FProperty*& OutProperty, TConstArrayView<uint8>& OutValue) const
+{
+	auto CheckOverride = [&InName, &OutType, &OutValue, &OutProperty](const UAnimNextDataInterfaceEntry* InInterfaceEntry) -> bool
+	{
+		return InInterfaceEntry->GetValueOverride(InName, OutType, OutProperty, OutValue);
+	};
+
+	return FindOverrideStatusRecursiveHelper(CheckOverride);
+}
+
+bool UAnimNextDataInterfaceEntry::HasValueOverride(FName InName, FAnimNextParamType& OutType) const
+{
+	const FPropertyBagPropertyDesc* Desc = ValueOverrides.FindPropertyDescByName(InName);
+	if(Desc == nullptr)
+	{
+		return false;
+	}
+
+	OutType = FAnimNextParamType(Desc->ValueType, Desc->ContainerTypes.GetFirstContainerType(), Desc->ValueTypeObject);
+	return true;
+}
+
+bool UAnimNextDataInterfaceEntry::HasValueOverride(FName InName) const
+{
+	const FPropertyBagPropertyDesc* Desc = ValueOverrides.FindPropertyDescByName(InName);
+	return Desc != nullptr;
+}
+
+bool UAnimNextDataInterfaceEntry::GetDefaultValueRecursive(FName InName, const FProperty*& OutProperty, TConstArrayView<uint8>& OutValue) const
+{
+	auto CheckOverride = [this, &InName, &OutValue, &OutProperty](const UAnimNextDataInterfaceEntry* InInterfaceEntry) -> bool
+	{
+		// Skip 'this' when looking for overrides
+		bool bHasOverride = this != InInterfaceEntry && InInterfaceEntry->GetValueOverride(InName, OutProperty, OutValue);
+		if(!bHasOverride && InInterfaceEntry->DataInterface)
+		{
+			// No override, so see if this data interface holds a value
+			UAnimNextDataInterface_EditorData* EditorData = UE::AnimNext::UncookedOnly::FUtils::GetEditorData<UAnimNextDataInterface_EditorData>(InInterfaceEntry->DataInterface.Get());
+			if(UAnimNextVariableEntry* VariableEntry = Cast<UAnimNextVariableEntry>(EditorData->FindEntry(InName)))
+			{
+				VariableEntry->GetDefaultValue(OutProperty, OutValue);
+			}
+		}
+		return bHasOverride;
+	};
+
+	(void)FindOverrideStatusRecursiveHelper(CheckOverride);
+	return (!OutValue.IsEmpty() && OutProperty != nullptr);
+}
+
+bool UAnimNextDataInterfaceEntry::HasValueOverrideNotMatchingDefault(FName InName) const
+{
+	const FProperty* OverrideProperty = nullptr;
+	TConstArrayView<uint8> OverrideValue;
+	if(!GetValueOverride(InName, OverrideProperty, OverrideValue))
+	{
+		// No override, so cant have a default-matching override
+		return false;
+	}
+
+	check(OverrideProperty != nullptr);
+	check(!OverrideValue.IsEmpty());
+
+	const FProperty* BaseProperty = nullptr;
+	TConstArrayView<uint8> BaseValue;
+	if(!GetDefaultValueRecursive(InName, BaseProperty, BaseValue))
+	{
+		// No value, so cannot compare with default
+		return false;
+	}
+
+	check(BaseProperty != nullptr);
+	check(!BaseValue.IsEmpty());
+
+	if(BaseProperty->GetClass() != OverrideProperty->GetClass()) 
+	{
+		// Properties differ, cannot compare
+		// If the ensure hits here then we have somehow ended up with different types in implementing/base interfaces, so the workflow that got us to
+		// this point needs edge cases handling better
+		ensure(false);
+		return false;
+	}
+
+	return !BaseProperty->Identical(OverrideValue.GetData(), BaseValue.GetData());
+}
+
+EAnimNextDataInterfaceValueOverrideStatus UAnimNextDataInterfaceEntry::GetValueOverrideStatusRecursive(FName InName) const
+{
+	auto CheckOverride = [&InName](const UAnimNextDataInterfaceEntry* InInterfaceEntry) -> bool
+	{
+		return InInterfaceEntry->HasValueOverride(InName);
+	};
+
+	return FindOverrideStatusRecursiveHelper(CheckOverride);
+}
+
+EAnimNextDataInterfaceValueOverrideStatus UAnimNextDataInterfaceEntry::FindValueOverridePropertyBagRecursive(FName InName, FInstancedPropertyBag*& OutPropertyBag) const
+{
+	auto CheckOverride = [&InName, &OutPropertyBag](const UAnimNextDataInterfaceEntry* InInterfaceEntry) -> bool
+	{
+		const FPropertyBagPropertyDesc* Desc = InInterfaceEntry->ValueOverrides.FindPropertyDescByName(InName);
+		if(Desc == nullptr)
+		{
+			return false;
+		}
+
+		OutPropertyBag = const_cast<FInstancedPropertyBag*>(&InInterfaceEntry->ValueOverrides);
+		return true;
+	};
+
+	return FindOverrideStatusRecursiveHelper(CheckOverride);
+}
+
+void UAnimNextDataInterfaceEntry::HandleDataInterfaceModified(UAnimNextRigVMAssetEditorData* InEditorData, EAnimNextEditorDataNotifType InType, UObject* InSubject)
+{
+	switch(InType)
+	{
+	case EAnimNextEditorDataNotifType::UndoRedo:
+	case EAnimNextEditorDataNotifType::EntryAdded:
+	case EAnimNextEditorDataNotifType::EntryRemoved:
+	case EAnimNextEditorDataNotifType::EntryRenamed:
+	case EAnimNextEditorDataNotifType::EntryAccessSpecifierChanged:
+	case EAnimNextEditorDataNotifType::VariableDefaultValueChanged:
+		if(UAnimNextRigVMAssetEditorData* EditorData = GetTypedOuter<UAnimNextRigVMAssetEditorData>())
+		{
+			EditorData->RecompileVM();
+		}
+		break;
+	default:
+		break;
+	}
+}
+
+#undef LOCTEXT_NAMESPACE
\ No newline at end of file
diff --git a/Engine/Plugins/Experimental/AnimNext/Source/AnimNextUncookedOnly/Private/Entries/AnimNextRigVMAssetEntry.cpp b/Engine/Plugins/Experimental/AnimNext/Source/AnimNextUncookedOnly/Private/Entries/AnimNextRigVMAssetEntry.cpp
--- a/Engine/Plugins/Experimental/AnimNext/Source/AnimNextUncookedOnly/Private/Entries/AnimNextRigVMAssetEntry.cpp
+++ b/Engine/Plugins/Experimental/AnimNext/Source/AnimNextUncookedOnly/Private/Entries/AnimNextRigVMAssetEntry.cpp
@@ -2,6 +2,7 @@
 
 #include "Entries/AnimNextRigVMAssetEntry.h"
 #include "AnimNextRigVMAssetEditorData.h"
+#include "Misc/TransactionObjectEvent.h"
 
 void UAnimNextRigVMAssetEntry::Initialize(UAnimNextRigVMAssetEditorData* InEditorData)
 {
@@ -25,6 +26,16 @@ void UAnimNextRigVMAssetEntry::PostEditChangeProperty(FPropertyChangedEvent& Pro
 	BroadcastModified(EAnimNextEditorDataNotifType::PropertyChanged);
 }
 
+void UAnimNextRigVMAssetEntry::PostTransacted(const FTransactionObjectEvent& TransactionEvent)
+{
+	Super::PostTransacted(TransactionEvent);
+
+	if (TransactionEvent.GetEventType() == ETransactionObjectEventType::UndoRedo)
+	{
+		BroadcastModified(EAnimNextEditorDataNotifType::UndoRedo);
+	}
+}
+
 #endif
 
 void UAnimNextRigVMAssetEntry::BroadcastModified(EAnimNextEditorDataNotifType InType)
diff --git a/Engine/Plugins/Experimental/AnimNext/Source/AnimNextUncookedOnly/Private/Entries/AnimNextVariableEntry.cpp b/Engine/Plugins/Experimental/AnimNext/Source/AnimNextUncookedOnly/Private/Entries/AnimNextVariableEntry.cpp
--- a/Engine/Plugins/Experimental/AnimNext/Source/AnimNextUncookedOnly/Private/Entries/AnimNextVariableEntry.cpp
+++ b/Engine/Plugins/Experimental/AnimNext/Source/AnimNextUncookedOnly/Private/Entries/AnimNextVariableEntry.cpp
@@ -4,6 +4,7 @@
 
 #include "AnimNextRigVMAssetEditorData.h"
 #include "UncookedOnlyUtils.h"
+#include "Logging/StructuredLog.h"
 #include "Module/AnimNextModule.h"
 #include "Variables/AnimNextUniversalObjectLocatorBindingData.h"
 
@@ -69,14 +70,54 @@ bool UAnimNextVariableEntry::SetType(const FAnimNextParamType& InType, bool bSet
 	return true;
 }
 
-bool UAnimNextVariableEntry::SetDefaultValue(const FString& InDefaultValue, bool bSetupUndoRedo)
+bool UAnimNextVariableEntry::SetDefaultValue(TConstArrayView<uint8> InValue, bool bSetupUndoRedo)
 {
+	check(!InValue.IsEmpty());
+
 	if(bSetupUndoRedo)
 	{
 		Modify();
 	}
-	
-	return DefaultValue.SetValueSerializedString(IAnimNextRigVMVariableInterface::ValueName, InDefaultValue) == EPropertyBagResult::Success;
+
+	const FPropertyBagPropertyDesc* Desc = DefaultValue.FindPropertyDescByName(IAnimNextRigVMVariableInterface::ValueName);
+	if(Desc == nullptr)
+	{
+		UE_LOGFMT(LogAnimation, Error, "UAnimNextVariableEntry::SetDefaultValue: Could not find default value in property bag");
+		return false;
+	}
+
+	check(Desc->CachedProperty);
+	if(Desc->CachedProperty->GetElementSize() != InValue.Num())
+	{
+		UE_LOGFMT(LogAnimation, Error, "UAnimNextVariableEntry::SetDefaultValue: Mismatched buffer sizes");
+		return false;
+	}
+
+	uint8* DestPtr = Desc->CachedProperty->ContainerPtrToValuePtr<uint8>(DefaultValue.GetMutableValue().GetMemory());
+	const uint8* SrcPtr = InValue.GetData();
+	Desc->CachedProperty->CopyCompleteValue(DestPtr, SrcPtr);
+
+	BroadcastModified(EAnimNextEditorDataNotifType::VariableDefaultValueChanged);
+
+	return true;
+}
+
+bool UAnimNextVariableEntry::SetDefaultValueFromString(const FString& InDefaultValue, bool bSetupUndoRedo)
+{
+	if(bSetupUndoRedo)
+	{
+		Modify();
+	}
+
+	if(DefaultValue.SetValueSerializedString(IAnimNextRigVMVariableInterface::ValueName, InDefaultValue) != EPropertyBagResult::Success)
+	{
+		UE_LOGFMT(LogAnimation, Error, "UAnimNextVariableEntry::SetDefaultValueFromString: Could not set value from string");
+		return false;
+	}
+
+	BroadcastModified(EAnimNextEditorDataNotifType::VariableDefaultValueChanged);
+
+	return true;
 }
 
 FName UAnimNextVariableEntry::GetVariableName() const
@@ -94,9 +135,46 @@ const FInstancedPropertyBag& UAnimNextVariableEntry::GetPropertyBag() const
 	return DefaultValue;
 }
 
-const FAnimNextVariableBinding& UAnimNextVariableEntry::GetBinding() const
+bool UAnimNextVariableEntry::GetDefaultValue(const FProperty*& OutProperty, TConstArrayView<uint8>& OutValue) const
+{
+	const FPropertyBagPropertyDesc* Desc = DefaultValue.FindPropertyDescByName(IAnimNextRigVMVariableInterface::ValueName);
+	if(Desc == nullptr)
+	{
+		UE_LOGFMT(LogAnimation, Error, "UAnimNextVariableEntry::SetDefaultValue: Could not find default value in property bag");
+		return false;
+	}
+
+	check(Desc->CachedProperty);
+	OutProperty = Desc->CachedProperty;
+	const uint8* ValuePtr = Desc->CachedProperty->ContainerPtrToValuePtr<uint8>(DefaultValue.GetValue().GetMemory());
+	OutValue = TConstArrayView<uint8>(ValuePtr, Desc->CachedProperty->GetElementSize());
+	return true;
+}
+
+void UAnimNextVariableEntry::SetBindingType(UScriptStruct* InBindingTypeStruct, bool bSetupUndoRedo)
+{
+	check(InBindingTypeStruct == nullptr || InBindingTypeStruct->IsChildOf(FAnimNextVariableBindingData::StaticStruct()));
+
+	if(bSetupUndoRedo)
+	{
+		Modify();
+	}
+
+	if(InBindingTypeStruct)
+	{
+		Binding.BindingData.InitializeAsScriptStruct(InBindingTypeStruct);
+	}
+	else
+	{
+		Binding.BindingData.Reset();
+	}
+
+	BroadcastModified(EAnimNextEditorDataNotifType::VariableBindingChanged);
+}
+
+TConstStructView<FAnimNextVariableBindingData> UAnimNextVariableEntry::GetBinding() const
 {
-	return Binding;
+	return Binding.BindingData;
 }
 
 void UAnimNextVariableEntry::SetEntryName(FName InName, bool bSetupUndoRedo)
diff --git a/Engine/Plugins/Experimental/AnimNext/Source/AnimNextUncookedOnly/Private/Graph/AnimNextAnimationGraphSchema.cpp b/Engine/Plugins/Experimental/AnimNext/Source/AnimNextUncookedOnly/Private/Graph/AnimNextAnimationGraphSchema.cpp
--- a/Engine/Plugins/Experimental/AnimNext/Source/AnimNextUncookedOnly/Private/Graph/AnimNextAnimationGraphSchema.cpp
+++ b/Engine/Plugins/Experimental/AnimNext/Source/AnimNextUncookedOnly/Private/Graph/AnimNextAnimationGraphSchema.cpp
@@ -2,46 +2,3 @@
 
 #include "Graph/AnimNextAnimationGraphSchema.h"
 
-#include "AnimNextExecuteContext.h"
-#include "Graph/RigUnit_AnimNextBase.h"
-#include "Param/RigVMDispatch_GetParameter.h"
-#include "Param/RigVMDispatch_GetScopedParameter.h"
-
-bool UAnimNextAnimationGraphSchema::SupportsUnitFunction(URigVMController* InController, const FRigVMFunction* InUnitFunction) const
-{
-	if(InUnitFunction)
-	{
-		if(const UScriptStruct* FunctionExecuteContextStruct = InUnitFunction->GetExecuteContextStruct())
-		{
-			if(FunctionExecuteContextStruct == FAnimNextExecuteContext::StaticStruct())
-			{
-				// Only allow nodes that are children of FRigUnit_AnimNextBase
-				if(InUnitFunction->Struct)
-				{
-					return InUnitFunction->Struct->IsChildOf(FRigUnit_AnimNextBase::StaticStruct());
-				}
-			}
-		}
-	}
-
-	return Super::SupportsUnitFunction(InController, InUnitFunction);
-}
-
-bool UAnimNextAnimationGraphSchema::SupportsDispatchFactory(URigVMController* InController, const FRigVMDispatchFactory* InDispatchFactory) const
-{
-	if(InDispatchFactory)
-	{
-		if(const UScriptStruct* DispatchExecuteContextStruct = InDispatchFactory->GetExecuteContextStruct())
-		{
-			if(DispatchExecuteContextStruct == FAnimNextExecuteContext::StaticStruct())
-			{
-				// We only support the FRigVMDispatch_GetParameter/FRigVMDispatch_GetScopedParameter at the moment.
-				return
-				   (InDispatchFactory->GetScriptStruct() == FRigVMDispatch_GetParameter::StaticStruct() ||
-					InDispatchFactory->GetScriptStruct() == FRigVMDispatch_GetScopedParameter::StaticStruct());
-			}
-		}
-	}
-	
-	return Super::SupportsDispatchFactory(InController, InDispatchFactory);
-}
\ No newline at end of file
diff --git a/Engine/Plugins/Experimental/AnimNext/Source/AnimNextUncookedOnly/Private/Graph/AnimNextAnimationGraph_EditorData.cpp b/Engine/Plugins/Experimental/AnimNext/Source/AnimNextUncookedOnly/Private/Graph/AnimNextAnimationGraph_EditorData.cpp
--- a/Engine/Plugins/Experimental/AnimNext/Source/AnimNextUncookedOnly/Private/Graph/AnimNextAnimationGraph_EditorData.cpp
+++ b/Engine/Plugins/Experimental/AnimNext/Source/AnimNextUncookedOnly/Private/Graph/AnimNextAnimationGraph_EditorData.cpp
@@ -7,6 +7,7 @@
 #include "Curves/CurveFloat.h"
 #include "Module/AnimNextModule.h"
 #include "Entries/AnimNextAnimationGraphEntry.h"
+#include "Entries/AnimNextDataInterfaceEntry.h"
 #include "Graph/AnimNextAnimationGraph.h"
 #include "Entries/AnimNextEventGraphEntry.h"
 #include "Entries/AnimNextVariableEntry.h"
@@ -688,6 +689,7 @@ TConstArrayView<TSubclassOf<UAnimNextRigVMAssetEntry>> UAnimNextAnimationGraph_E
 		UAnimNextAnimationGraphEntry::StaticClass(),
 		UAnimNextEventGraphEntry::StaticClass(),	// TODO: remove when assets are reworked post-refactor
 		UAnimNextVariableEntry::StaticClass(),
+		UAnimNextDataInterfaceEntry::StaticClass(),
 	};
 
 	return Classes;
diff --git a/Engine/Plugins/Experimental/AnimNext/Source/AnimNextUncookedOnly/Private/Module/AnimNextModule_EditorData.cpp b/Engine/Plugins/Experimental/AnimNext/Source/AnimNextUncookedOnly/Private/Module/AnimNextModule_EditorData.cpp
--- a/Engine/Plugins/Experimental/AnimNext/Source/AnimNextUncookedOnly/Private/Module/AnimNextModule_EditorData.cpp
+++ b/Engine/Plugins/Experimental/AnimNext/Source/AnimNextUncookedOnly/Private/Module/AnimNextModule_EditorData.cpp
@@ -9,12 +9,14 @@
 #include "Entries/AnimNextAnimationGraphEntry.h"
 #include "Entries/AnimNextEventGraphEntry.h"
 #include "AnimNextEventGraphSchema.h"
+#include "Entries/AnimNextDataInterfaceEntry.h"
 #include "Entries/AnimNextVariableEntry.h"
 #include "Graph/AnimNextAnimationGraphSchema.h"
 #include "RigVMModel/Nodes/RigVMAggregateNode.h"
 #include "UObject/AssetRegistryTagsContext.h"
 #include "UObject/FortniteMainBranchObjectVersion.h"
 #include "UObject/LinkerLoad.h"
+#include "Variables/AnimNextUniversalObjectLocatorBindingData.h"
 
 void UAnimNextModule_EditorData::RecompileVM()
 {
@@ -169,6 +171,7 @@ TConstArrayView<TSubclassOf<UAnimNextRigVMAssetEntry>> UAnimNextModule_EditorDat
 		UAnimNextAnimationGraphEntry::StaticClass(),	// TEMP: Remove this once we have ported all old data
 		UAnimNextEventGraphEntry::StaticClass(),
 		UAnimNextVariableEntry::StaticClass(),
+		UAnimNextDataInterfaceEntry::StaticClass(),
 	};
 	
 	return Classes;
@@ -178,10 +181,18 @@ void UAnimNextModule_EditorData::GetProgrammaticGraphs(const FRigVMCompileSettin
 {
 	using namespace UE::AnimNext::UncookedOnly;
 
-	URigVMGraph* BindingsGraph = FUtils::CompileVariableBindings(InSettings, FUtils::GetAsset<UAnimNextModule>(this));
-	if(BindingsGraph)
+	FUtils::CompileVariableBindings(InSettings, FUtils::GetAsset<UAnimNextModule>(this), OutGraphs);
+}
+
+void UAnimNextModule_EditorData::CustomizeNewAssetEntry(UAnimNextRigVMAssetEntry* InNewEntry) const
+{
+	Super::CustomizeNewAssetEntry(InNewEntry);
+	
+	UAnimNextVariableEntry* VariableEntry = Cast<UAnimNextVariableEntry>(InNewEntry);
+	if(VariableEntry == nullptr)
 	{
-		OutGraphs.Add(BindingsGraph);
+		return;
 	}
-}
 
+	VariableEntry->SetBindingType(FAnimNextUniversalObjectLocatorBindingData::StaticStruct(), false);
+}
\ No newline at end of file
diff --git a/Engine/Plugins/Experimental/AnimNext/Source/AnimNextUncookedOnly/Private/Module/AnimNextModule_EventGraphSchema.cpp b/Engine/Plugins/Experimental/AnimNext/Source/AnimNextUncookedOnly/Private/Module/AnimNextModule_EventGraphSchema.cpp
--- a/Engine/Plugins/Experimental/AnimNext/Source/AnimNextUncookedOnly/Private/Module/AnimNextModule_EventGraphSchema.cpp
+++ b/Engine/Plugins/Experimental/AnimNext/Source/AnimNextUncookedOnly/Private/Module/AnimNextModule_EventGraphSchema.cpp
@@ -12,13 +12,10 @@ bool UAnimNextEventGraphSchema::SupportsUnitFunction(URigVMController* InControl
 	{
 		if(FunctionExecuteContextStruct == FAnimNextExecuteContext::StaticStruct())
 		{
-			if(InUnitFunction->Struct)
+			// Disallow trait stacks in event graphs 
+			if(InUnitFunction->Struct && InUnitFunction->Struct->IsChildOf(FRigUnit_AnimNextTraitStack::StaticStruct()))
 			{
-				if(InUnitFunction->Struct->IsChildOf(FRigUnit_AnimNextTraitStack::StaticStruct()))
-				{
-					return false;
-				}
-				return InUnitFunction->Struct->IsChildOf(FRigUnit_AnimNextBase::StaticStruct());
+				return false;
 			}
 		}
 	}
