AnimStateTree plugin
- Added AnimStateTree asset, deriving from AnimNextGraph - and functioning in a similar fashion
    * Embeds an actual StateTree used to drive the behaviour
    * Relies on programmatic root animation graph as its runtime integration into AnimNext
        + Currently consists out of a BlendStack trait and StateTreeTrait
- Added AnimNextStateTreeTrait (additive trait)
    * Responsible for funneling AnimGraphInstance variable data into the StateTree instance, and ticking its behaviour
- Added base types for AnimNextStateTree(Schema)
- Added AnimNextStateTreeContext struct used as a intermediate between AnimNext runtime and StateTree (tasks/runtime)
- Implemented simple Graph Instance task to push a specified AnimNextGraph onto the BlendStack

AnimNextEditor
- Moved LogListing/Tab FNames to Module interface
- Added ability to remove asset classes as well
- Added ability to retrieve IMessageLogListing used for AnimNext compilation tab
- Exposed UAnimNextWorkspaceFactory

Workspace Editor
- Added ability to customize outliner item color
- Added callback for handling selection of outliner item
- Added boilerplate for hosting editor modes inside of WorkspaceEditor
- Added RAII type to allow for multi widget global selection scopes (used by StateTree syncing all selection widgets)
- Added ability to add EditorMode ID to workspace asset document, if set is used to try and invoke the editor mode when the asset takes focus
- Added delegate to react to the focussed asset changing

AnimStateTree Editor
- Added StateTreeEditor host to be used in Workspace Editor
- Added AnimStateTree workspace exports

diff --git a/Engine/Plugins/Experimental/AnimNext/Source/AnimNext/Internal/AnimNextRigVMAsset.h b/Engine/Plugins/Experimental/AnimNext/Source/AnimNext/Internal/AnimNextRigVMAsset.h
--- a/Engine/Plugins/Experimental/AnimNext/Source/AnimNext/Internal/AnimNextRigVMAsset.h
+++ b/Engine/Plugins/Experimental/AnimNext/Source/AnimNext/Internal/AnimNextRigVMAsset.h
@@ -23,8 +23,8 @@ namespace UE::AnimNext::UncookedOnly
 }
 
 /** Base class for all AnimNext assets that can host RigVM logic */
-UCLASS(MinimalAPI, Abstract)
-class UAnimNextRigVMAsset : public URigVMHost
+UCLASS(Abstract)
+class ANIMNEXT_API UAnimNextRigVMAsset : public URigVMHost
 {
 	GENERATED_BODY()
 
diff --git a/Engine/Plugins/Experimental/AnimNext/Source/AnimNext/Public/Graph/AnimNextGraphInstance.h b/Engine/Plugins/Experimental/AnimNext/Source/AnimNext/Public/Graph/AnimNextGraphInstance.h
--- a/Engine/Plugins/Experimental/AnimNext/Source/AnimNext/Public/Graph/AnimNextGraphInstance.h
+++ b/Engine/Plugins/Experimental/AnimNext/Source/AnimNext/Public/Graph/AnimNextGraphInstance.h
@@ -112,6 +112,7 @@ struct ANIMNEXT_API FAnimNextGraphInstance : public FAnimNextDataInterfaceInstan
 	// Get the hosting instance, if any, that owns us
 	FAnimNextDataInterfaceInstance* GetHost() const;
 
+	FInstancedPropertyBag& GetVariables() { return Variables; }
 private:
 	// Returns a pointer to the specified component, or nullptr if not found
 	UE::AnimNext::FGraphInstanceComponent* TryGetComponent(int32 ComponentNameHash, FName ComponentName) const;
diff --git a/Engine/Plugins/Experimental/AnimNext/Source/AnimNextEditor/Internal/Workspace/AnimNextWorkspaceFactory.h b/Engine/Plugins/Experimental/AnimNext/Source/AnimNextEditor/Internal/Workspace/AnimNextWorkspaceFactory.h
new file mode 100644
--- /dev/null
+++ b/Engine/Plugins/Experimental/AnimNext/Source/AnimNextEditor/Internal/Workspace/AnimNextWorkspaceFactory.h
@@ -0,0 +1,21 @@
+﻿// Copyright Epic Games, Inc. All Rights Reserved.
+
+#pragma once
+
+#include "CoreMinimal.h"
+#include "WorkspaceFactory.h"
+#include "AnimNextWorkspaceFactory.generated.h"
+
+UCLASS(MinimalAPI, BlueprintType)
+class UAnimNextWorkspaceFactory : public UWorkspaceFactory
+{
+	GENERATED_BODY()
+
+	UAnimNextWorkspaceFactory();
+
+	// UFactory interface
+	virtual bool ConfigureProperties() override;
+	virtual FText GetDisplayName() const override;
+	virtual uint32 GetMenuCategories() const override;
+	virtual const TArray<FText>& GetMenuCategorySubMenus() const override;
+};
\ No newline at end of file
diff --git a/Engine/Plugins/Experimental/AnimNext/Source/AnimNextEditor/Private/AnimNextEditorModule.cpp b/Engine/Plugins/Experimental/AnimNext/Source/AnimNextEditor/Private/AnimNextEditorModule.cpp
--- a/Engine/Plugins/Experimental/AnimNext/Source/AnimNextEditor/Private/AnimNextEditorModule.cpp
+++ b/Engine/Plugins/Experimental/AnimNext/Source/AnimNextEditor/Private/AnimNextEditorModule.cpp
@@ -8,7 +8,6 @@
 #include "IWorkspaceEditor.h"
 #include "ScopedTransaction.h"
 #include "SSimpleButton.h"
-#include "SSimpleComboButton.h"
 #include "UncookedOnlyUtils.h"
 #include "Common/SRigVMAssetView.h"
 #include "Editor/RigVMEditorTools.h"
@@ -28,9 +27,9 @@
 #include "Common/SActionMenu.h"
 #include "Entries/AnimNextRigVMAssetEntry.h"
 #include "Editor/RigVMGraphDetailCustomization.h"
-#include "EditorUtils.h"
 #include "FileHelpers.h"
 #include "IUniversalObjectLocatorEditorModule.h"
+#include "MessageLogModule.h"
 #include "Param/AnimNextComponentLocatorEditor.h"
 #include "Param/AnimNextLocatorContext.h"
 #include "Param/ObjectCastLocatorEditor.h"
@@ -319,6 +318,14 @@ void FAnimNextEditorModule::AddWorkspaceSupportedAssetClass(const FTopLevelAsset
 	}	
 }
 
+void FAnimNextEditorModule::RemoveWorkspaceSupportedAssetClass(const FTopLevelAssetPath& InClassAssetPath)
+{
+	if (InClassAssetPath.IsValid())
+	{
+		SupportedAssetClasses.Remove(InClassAssetPath);
+	}	
+}
+
 void FAnimNextEditorModule::RegisterWorkspaceDocumentTypes(Workspace::IWorkspaceEditorModule& WorkspaceEditorModule)
 {
 	// --- AnimNextRigVMAsset ---
diff --git a/Engine/Plugins/Experimental/AnimNext/Source/AnimNextEditor/Private/AnimNextEditorModule.h b/Engine/Plugins/Experimental/AnimNext/Source/AnimNextEditor/Private/AnimNextEditorModule.h
--- a/Engine/Plugins/Experimental/AnimNext/Source/AnimNextEditor/Private/AnimNextEditorModule.h
+++ b/Engine/Plugins/Experimental/AnimNext/Source/AnimNextEditor/Private/AnimNextEditorModule.h
@@ -34,7 +34,7 @@ private:
 	virtual void RegisterLocatorFragmentEditorType(FName InLocatorFragmentEditorName) override;
 	virtual void UnregisterLocatorFragmentEditorType(FName InLocatorFragmentEditorName) override;
 	virtual void AddWorkspaceSupportedAssetClass(const FTopLevelAssetPath& InClassAssetPath) override;
-
+	virtual void RemoveWorkspaceSupportedAssetClass(const FTopLevelAssetPath& InClassAssetPath) override;
 private:
 	void RegisterWorkspaceDocumentTypes(Workspace::IWorkspaceEditorModule& WorkspaceEditorModule);
 	void UnregisterWorkspaceDocumentTypes();
diff --git a/Engine/Plugins/Experimental/AnimNext/Source/AnimNextEditor/Private/Graph/AnimNextCompilerResultsTabSummoner.cpp b/Engine/Plugins/Experimental/AnimNext/Source/AnimNextEditor/Private/Graph/AnimNextCompilerResultsTabSummoner.cpp
--- a/Engine/Plugins/Experimental/AnimNext/Source/AnimNextEditor/Private/Graph/AnimNextCompilerResultsTabSummoner.cpp
+++ b/Engine/Plugins/Experimental/AnimNext/Source/AnimNextEditor/Private/Graph/AnimNextCompilerResultsTabSummoner.cpp
@@ -1,6 +1,7 @@
 // Copyright Epic Games, Inc. All Rights Reserved.
 
 #include "Graph/AnimNextCompilerResultsTabSummoner.h"
+#include "IAnimNextEditorModule.h"
 #include "IWorkspaceEditor.h"
 #include "MessageLogModule.h"
 #include "Modules/ModuleManager.h"
@@ -10,8 +11,6 @@
 namespace UE::AnimNext::Editor
 {
 
-const FLazyName CompilerResultsTabName("CompilerResultsTab");
-
 // ***************************************************************************
 
 void SAnimNextCompilerResultsWidget::Construct(const FArguments& InArgs, const TWeakPtr<UE::Workspace::IWorkspaceEditor>& InWorkspaceEditorWeak)
@@ -36,9 +35,8 @@ void SAnimNextCompilerResultsWidget::CreateMessageLog(const TWeakPtr<UE::Workspa
 	if (const TSharedPtr<UE::Workspace::IWorkspaceEditor> WorkspaceEditorShared = InWorkspaceEditorWeak.Pin())
 	{
 		FMessageLogModule& MessageLogModule = FModuleManager::LoadModuleChecked<FMessageLogModule>("MessageLog");
-		const FName LogName("AnimNextCompilerResults");
-		check(MessageLogModule.IsRegisteredLogListing(LogName));
-		CompilerResultsListing = MessageLogModule.GetLogListing(LogName);
+		check(MessageLogModule.IsRegisteredLogListing(LogListingName));
+		CompilerResultsListing = MessageLogModule.GetLogListing(LogListingName);
 		CompilerResults = MessageLogModule.CreateLogListingWidget(CompilerResultsListing.ToSharedRef());
 	}
 }
@@ -46,7 +44,7 @@ void SAnimNextCompilerResultsWidget::CreateMessageLog(const TWeakPtr<UE::Workspa
 // ***************************************************************************
 
 FAnimNextCompilerResultsTabSummoner::FAnimNextCompilerResultsTabSummoner(TSharedPtr<UE::Workspace::IWorkspaceEditor> InHostingApp)
-	: FWorkflowTabFactory(CompilerResultsTabName, StaticCastSharedPtr<FAssetEditorToolkit>(InHostingApp))
+	: FWorkflowTabFactory(UE::AnimNext::Editor::CompilerResultsTabName, StaticCastSharedPtr<FAssetEditorToolkit>(InHostingApp))
 {
 	TabLabel = LOCTEXT("AnimNExtCompilerResultsTabLabel", "Compiler Results");
 	TabIcon = FSlateIcon("EditorStyle", "LevelEditor.Tabs.Outliner");
diff --git a/Engine/Plugins/Experimental/AnimNext/Source/AnimNextEditor/Private/Graph/AnimNextCompilerResultsTabSummoner.h b/Engine/Plugins/Experimental/AnimNext/Source/AnimNextEditor/Private/Graph/AnimNextCompilerResultsTabSummoner.h
--- a/Engine/Plugins/Experimental/AnimNext/Source/AnimNextEditor/Private/Graph/AnimNextCompilerResultsTabSummoner.h
+++ b/Engine/Plugins/Experimental/AnimNext/Source/AnimNextEditor/Private/Graph/AnimNextCompilerResultsTabSummoner.h
@@ -12,9 +12,6 @@ namespace UE::Workspace
 
 namespace UE::AnimNext::Editor
 {
-
-extern const FLazyName CompilerResultsTabName;
-
 class SAnimNextCompilerResultsWidget : public SCompoundWidget
 {
 public:
@@ -36,9 +33,6 @@ private:
 	TSharedPtr<class SWidget> CompilerResults;
 };
 
-
-
-
 struct FAnimNextCompilerResultsTabSummoner : public FWorkflowTabFactory
 {
 public:
diff --git a/Engine/Plugins/Experimental/AnimNext/Source/AnimNextEditor/Private/Graph/STraitEditorView.cpp b/Engine/Plugins/Experimental/AnimNext/Source/AnimNextEditor/Private/Graph/STraitEditorView.cpp
--- a/Engine/Plugins/Experimental/AnimNext/Source/AnimNextEditor/Private/Graph/STraitEditorView.cpp
+++ b/Engine/Plugins/Experimental/AnimNext/Source/AnimNextEditor/Private/Graph/STraitEditorView.cpp
@@ -11,6 +11,7 @@
 #include "TraitCore/TraitRegistry.h"
 #include "ObjectEditorUtils.h"
 #include "AnimNextEdGraphNode.h"
+#include "IAnimNextEditorModule.h"
 #include "STraitListView.h"
 #include "STraitStackView.h"
 #include "MessageLogModule.h"
@@ -298,7 +299,7 @@ void STraitEditorView::RefreshTraitStackTraitsStatus()
 					// Open tab to display errors
 					if (TSharedPtr<UE::Workspace::IWorkspaceEditor> WorkspaceEditor = WorkspaceEditorWeak.Pin())
 					{
-						WorkspaceEditor->GetTabManager()->TryInvokeTab(FTabId(CompilerResultsTabName));
+						WorkspaceEditor->GetTabManager()->TryInvokeTab(FTabId(UE::AnimNext::Editor::CompilerResultsTabName));
 					}
 				}
 			}
diff --git a/Engine/Plugins/Experimental/AnimNext/Source/AnimNextEditor/Private/Workspace/AnimNextWorkspaceFactory.cpp b/Engine/Plugins/Experimental/AnimNext/Source/AnimNextEditor/Private/Workspace/AnimNextWorkspaceFactory.cpp
--- a/Engine/Plugins/Experimental/AnimNext/Source/AnimNextEditor/Private/Workspace/AnimNextWorkspaceFactory.cpp
+++ b/Engine/Plugins/Experimental/AnimNext/Source/AnimNextEditor/Private/Workspace/AnimNextWorkspaceFactory.cpp
@@ -1,6 +1,6 @@
 ﻿// Copyright Epic Games, Inc. All Rights Reserved.
 
-#include "AnimNextWorkspaceFactory.h"
+#include "Workspace/AnimNextWorkspaceFactory.h"
 #include "AnimNextWorkspaceSchema.h"
 #include "AssetTypeCategories.h"
 
diff --git a/Engine/Plugins/Experimental/AnimNext/Source/AnimNextEditor/Private/Workspace/AnimNextWorkspaceFactory.h b/Engine/Plugins/Experimental/AnimNext/Source/AnimNextEditor/Private/Workspace/AnimNextWorkspaceFactory.h
deleted file mode 100644
diff --git a/Engine/Plugins/Experimental/AnimNext/Source/AnimNextEditor/Public/IAnimNextEditorModule.h b/Engine/Plugins/Experimental/AnimNext/Source/AnimNextEditor/Public/IAnimNextEditorModule.h
--- a/Engine/Plugins/Experimental/AnimNext/Source/AnimNextEditor/Public/IAnimNextEditorModule.h
+++ b/Engine/Plugins/Experimental/AnimNext/Source/AnimNextEditor/Public/IAnimNextEditorModule.h
@@ -6,11 +6,15 @@
 #include "Modules/ModuleManager.h"
 
 class SWidget;
+class IMessageLogListing;
 struct FTopLevelAssetPath;
 
 namespace UE::AnimNext::Editor
 {
 
+const FLazyName CompilerResultsTabName("CompilerResultsTab");
+const FLazyName LogListingName("AnimNextCompilerResults");
+
 struct FVariablePickerArgs;
 
 class IAnimNextEditorModule : public IModuleInterface
@@ -27,6 +31,10 @@ public:
 	// Add a UClass path to the set of classes which can be opened within an AnimNext Workspace
 	// @param InClassAssetPath Asset path for to-be-registered Class 
 	virtual void AddWorkspaceSupportedAssetClass(const FTopLevelAssetPath& InClassAssetPath) = 0;
+	
+	// Remove a UClass path to the set of classes which can be opened within an AnimNext Workspace
+	// @param InClassAssetPath Asset path for to-be-unregistered Class 
+	virtual void RemoveWorkspaceSupportedAssetClass(const FTopLevelAssetPath& InClassAssetPath) = 0;
 };
 
 }
diff --git a/Engine/Plugins/Experimental/AnimNext/Source/AnimNextUncookedOnly/Internal/AnimNextRigVMAssetEditorData.h b/Engine/Plugins/Experimental/AnimNext/Source/AnimNextUncookedOnly/Internal/AnimNextRigVMAssetEditorData.h
--- a/Engine/Plugins/Experimental/AnimNext/Source/AnimNextUncookedOnly/Internal/AnimNextRigVMAssetEditorData.h
+++ b/Engine/Plugins/Experimental/AnimNext/Source/AnimNextUncookedOnly/Internal/AnimNextRigVMAssetEditorData.h
@@ -113,26 +113,26 @@ public:
 };
 
 /* Base class for all AnimNext editor data objects that use RigVM */
-UCLASS(MinimalAPI, Abstract)
-class UAnimNextRigVMAssetEditorData : public UObject, public IRigVMClientHost, public IRigVMGraphFunctionHost, public IRigVMClientExternalModelHost
+UCLASS(Abstract)
+class ANIMNEXTUNCOOKEDONLY_API UAnimNextRigVMAssetEditorData : public UObject, public IRigVMClientHost, public IRigVMGraphFunctionHost, public IRigVMClientExternalModelHost
 {
 	GENERATED_BODY()
 
 public:
 	/** Adds an animation graph to this asset */
-	ANIMNEXTUNCOOKEDONLY_API UAnimNextAnimationGraphEntry* AddAnimationGraph(FName InName, bool bSetupUndoRedo = true, bool bPrintPythonCommand = true);
+	UAnimNextAnimationGraphEntry* AddAnimationGraph(FName InName, bool bSetupUndoRedo = true, bool bPrintPythonCommand = true);
 
 	/** Adds a parameter to this asset */
-	ANIMNEXTUNCOOKEDONLY_API UAnimNextVariableEntry* AddVariable(FName InName, FAnimNextParamType InType, const FString& InDefaultValue = TEXT(""), bool bSetupUndoRedo = true, bool bPrintPythonCommand = true);
+	UAnimNextVariableEntry* AddVariable(FName InName, FAnimNextParamType InType, const FString& InDefaultValue = TEXT(""), bool bSetupUndoRedo = true, bool bPrintPythonCommand = true);
 
 	/** Adds an event graph to this asset */
-	ANIMNEXTUNCOOKEDONLY_API UAnimNextEventGraphEntry* AddEventGraph(FName InName, UScriptStruct* InEventStruct, bool bSetupUndoRedo = true, bool bPrintPythonCommand = true);
+	UAnimNextEventGraphEntry* AddEventGraph(FName InName, UScriptStruct* InEventStruct, bool bSetupUndoRedo = true, bool bPrintPythonCommand = true);
 
 	/** Adds a data interface to this asset */
-	ANIMNEXTUNCOOKEDONLY_API UAnimNextDataInterfaceEntry* AddDataInterface(UAnimNextDataInterface* InDataInterface, bool bSetupUndoRedo = true, bool bPrintPythonCommand = true);
+	UAnimNextDataInterfaceEntry* AddDataInterface(UAnimNextDataInterface* InDataInterface, bool bSetupUndoRedo = true, bool bPrintPythonCommand = true);
 
 	// Report an error to the user, typically used for scripting APIs
-	ANIMNEXTUNCOOKEDONLY_API static void ReportError(const TCHAR* InMessage);
+	static void ReportError(const TCHAR* InMessage);
 
 protected:
 	friend class UE::AnimNext::Editor::SRigVMAssetView;
@@ -246,6 +246,9 @@ protected:
 	// Handle RigVM modification events
 	virtual void HandleModifiedEvent(ERigVMGraphNotifType InNotifType, URigVMGraph* InGraph, UObject* InSubject);
 
+	// Class to use when instantiating AssetUserData for the EditorData instance
+	virtual TSubclassOf<UAssetUserData> GetAssetUserDataClass() const;
+
 	// Get all the kinds of entry for this asset
 	virtual TConstArrayView<TSubclassOf<UAnimNextRigVMAssetEntry>> GetEntryClasses() const PURE_VIRTUAL(UAnimNextRigVMAssetEditorData::GetEntryClasses, return {};)
 
@@ -338,21 +341,21 @@ protected:
 	}
 	
 	// Find an entry by name
-	ANIMNEXTUNCOOKEDONLY_API UAnimNextRigVMAssetEntry* FindEntry(FName InName) const;
+	UAnimNextRigVMAssetEntry* FindEntry(FName InName) const;
 
 	// Remove an entry from the asset
 	// @return true if the item was removed
-	ANIMNEXTUNCOOKEDONLY_API bool RemoveEntry(UAnimNextRigVMAssetEntry* InEntry, bool bSetupUndoRedo = true, bool bPrintPythonCommand = true);
+	bool RemoveEntry(UAnimNextRigVMAssetEntry* InEntry, bool bSetupUndoRedo = true, bool bPrintPythonCommand = true);
 
 	// Remove a number of entries from the asset
 	// @return true if any items were removed
-	ANIMNEXTUNCOOKEDONLY_API bool RemoveEntries(TConstArrayView<UAnimNextRigVMAssetEntry*> InEntries, bool bSetupUndoRedo = true, bool bPrintPythonCommand = true);
+	bool RemoveEntries(TConstArrayView<UAnimNextRigVMAssetEntry*> InEntries, bool bSetupUndoRedo = true, bool bPrintPythonCommand = true);
 
 	// Remove all entries from the asset
 	// @return true if any items were removed
-	ANIMNEXTUNCOOKEDONLY_API bool RemoveAllEntries(bool bSetupUndoRedo = true, bool bPrintPythonCommand = true);
+	bool RemoveAllEntries(bool bSetupUndoRedo = true, bool bPrintPythonCommand = true);
 
-	ANIMNEXTUNCOOKEDONLY_API void BroadcastModified(EAnimNextEditorDataNotifType InType, UObject* InSubject);
+	void BroadcastModified(EAnimNextEditorDataNotifType InType, UObject* InSubject);
 
 	void ReconstructAllNodes();
 
@@ -424,7 +427,6 @@ protected:
 	// Collection of models gleaned from graphs
 	TArray<TObjectPtr<URigVMGraph>> GraphModels;
 
-
 	// Set of functions implemented for this graph
 	UPROPERTY()
 	TArray<TObjectPtr<URigVMEdGraph>> FunctionEdGraphs;
@@ -433,7 +435,6 @@ protected:
 	UPROPERTY()
 	TObjectPtr<UAnimNextEdGraph> FunctionLibraryEdGraph;
 
-
 	bool bAutoRecompileVM = true;
 	bool bErrorsDuringCompilation = false;
 	bool bSuspendModelNotificationsForSelf = false;
diff --git a/Engine/Plugins/Experimental/AnimNext/Source/AnimNextUncookedOnly/Internal/AnimNextRigVMAssetSchema.h b/Engine/Plugins/Experimental/AnimNext/Source/AnimNextUncookedOnly/Internal/AnimNextRigVMAssetSchema.h
new file mode 100644
--- /dev/null
+++ b/Engine/Plugins/Experimental/AnimNext/Source/AnimNextUncookedOnly/Internal/AnimNextRigVMAssetSchema.h
@@ -0,0 +1,22 @@
+﻿// Copyright Epic Games, Inc. All Rights Reserved.
+
+#pragma once
+
+#include "CoreMinimal.h"
+#include "RigVMModel/RigVMSchema.h"
+#include "AnimNextRigVMAssetSchema.generated.h"
+
+UCLASS()
+class UAnimNextRigVMAssetSchema : public URigVMSchema
+{
+	GENERATED_BODY()
+
+	virtual bool SupportsNodeLayouts(const URigVMGraph* InGraph) const override
+	{
+		return true;
+	}
+
+protected:
+	UAnimNextRigVMAssetSchema();
+};
+
diff --git a/Engine/Plugins/Experimental/AnimNext/Source/AnimNextUncookedOnly/Internal/DataInterface/AnimNextDataInterface_EditorData.h b/Engine/Plugins/Experimental/AnimNext/Source/AnimNextUncookedOnly/Internal/DataInterface/AnimNextDataInterface_EditorData.h
--- a/Engine/Plugins/Experimental/AnimNext/Source/AnimNextUncookedOnly/Internal/DataInterface/AnimNextDataInterface_EditorData.h
+++ b/Engine/Plugins/Experimental/AnimNext/Source/AnimNextUncookedOnly/Internal/DataInterface/AnimNextDataInterface_EditorData.h
@@ -11,8 +11,8 @@
 class UAnimNextDataInterfaceFactory;
 
 /** Editor data for AnimNext data interfaces */
-UCLASS(MinimalAPI)
-class UAnimNextDataInterface_EditorData : public UAnimNextRigVMAssetEditorData
+UCLASS()
+class ANIMNEXTUNCOOKEDONLY_API UAnimNextDataInterface_EditorData : public UAnimNextRigVMAssetEditorData
 {
 	GENERATED_BODY()
 
diff --git a/Engine/Plugins/Experimental/AnimNext/Source/AnimNextUncookedOnly/Internal/Graph/AnimNextAnimationGraphSchema.h b/Engine/Plugins/Experimental/AnimNext/Source/AnimNextUncookedOnly/Internal/Graph/AnimNextAnimationGraphSchema.h
--- a/Engine/Plugins/Experimental/AnimNext/Source/AnimNextUncookedOnly/Internal/Graph/AnimNextAnimationGraphSchema.h
+++ b/Engine/Plugins/Experimental/AnimNext/Source/AnimNextUncookedOnly/Internal/Graph/AnimNextAnimationGraphSchema.h
@@ -6,7 +6,7 @@
 #include "AnimNextRigVMAssetSchema.h"
 #include "AnimNextAnimationGraphSchema.generated.h"
 
-UCLASS()
+UCLASS(MinimalAPI)
 class UAnimNextAnimationGraphSchema : public UAnimNextRigVMAssetSchema
 {
 	GENERATED_BODY()
diff --git a/Engine/Plugins/Experimental/AnimNext/Source/AnimNextUncookedOnly/Internal/Graph/AnimNextAnimationGraph_EditorData.h b/Engine/Plugins/Experimental/AnimNext/Source/AnimNextUncookedOnly/Internal/Graph/AnimNextAnimationGraph_EditorData.h
--- a/Engine/Plugins/Experimental/AnimNext/Source/AnimNextUncookedOnly/Internal/Graph/AnimNextAnimationGraph_EditorData.h
+++ b/Engine/Plugins/Experimental/AnimNext/Source/AnimNextUncookedOnly/Internal/Graph/AnimNextAnimationGraph_EditorData.h
@@ -26,8 +26,8 @@ namespace UE::AnimNext::Editor
 }
 
 /** Editor data for AnimNext animation graphs */
-UCLASS(MinimalAPI)
-class UAnimNextAnimationGraph_EditorData : public UAnimNextDataInterface_EditorData
+UCLASS()
+class ANIMNEXTUNCOOKEDONLY_API UAnimNextAnimationGraph_EditorData : public UAnimNextDataInterface_EditorData
 {
 	GENERATED_BODY()
 
@@ -44,7 +44,7 @@ class UAnimNextAnimationGraph_EditorData : public UAnimNextDataInterface_EditorD
 	friend class FAnimationAnimNextRuntimeTest_GraphExecute;
 	friend class FAnimationAnimNextRuntimeTest_GraphExecuteLatent;
 
-private:
+protected:
 	// IRigVMClientHost interface
 	virtual void RecompileVM() override;
 
diff --git a/Engine/Plugins/Experimental/AnimNext/Source/AnimNextUncookedOnly/Internal/UncookedOnlyUtils.h b/Engine/Plugins/Experimental/AnimNext/Source/AnimNextUncookedOnly/Internal/UncookedOnlyUtils.h
--- a/Engine/Plugins/Experimental/AnimNext/Source/AnimNextUncookedOnly/Internal/UncookedOnlyUtils.h
+++ b/Engine/Plugins/Experimental/AnimNext/Source/AnimNextUncookedOnly/Internal/UncookedOnlyUtils.h
@@ -5,6 +5,7 @@
 #include "CoreMinimal.h"
 #include "IAnimNextRigVMGraphInterface.h"
 #include "AssetRegistry/AssetData.h"
+#include "Misc/Build.h"
 #include "Param/ParamType.h"
 #include "RigVMCore/RigVMTemplate.h"
 #include "StructUtils/InstancedStruct.h"
@@ -97,6 +98,8 @@ struct FAnimNextAssetRegistryExports
 namespace UE::AnimNext::UncookedOnly
 {
 
+extern TAutoConsoleVariable<bool> CVarDumpProgrammaticGraphs;
+
 struct ANIMNEXTUNCOOKEDONLY_API FUtils
 {
 	static void CompileVariables(UAnimNextRigVMAsset* InAsset);
@@ -144,7 +147,7 @@ struct ANIMNEXTUNCOOKEDONLY_API FUtils
 	static FRigVMTemplateArgumentType GetRigVMArgTypeFromParamType(const FAnimNextParamType& InParamType);
 
 	/** Set up a simple animation graph */
-	static void SetupAnimGraph(UAnimNextRigVMAssetEntry* InEntry, URigVMController* InController);
+	static void SetupAnimGraph(const FName EntryName, URigVMController* InController);
 	
 	/** Set up a simple event graph */
 	static void SetupEventGraph(URigVMController* InController, UScriptStruct* InEventStruct);
@@ -172,6 +175,10 @@ struct ANIMNEXTUNCOOKEDONLY_API FUtils
 
 	// Returns an user friendly name for the Function Library
 	static const FText& GetFunctionLibraryDisplayName();
+
+#if WITH_EDITOR
+	static void OpenProgrammaticGraphs(UAnimNextRigVMAssetEditorData* EditorData, const TArray<URigVMGraph*>& ProgrammaticGraphs);
+#endif // WITH_EDITOR
 };
 
 }
\ No newline at end of file
diff --git a/Engine/Plugins/Experimental/AnimNext/Source/AnimNextUncookedOnly/Private/AnimNextEdGraphNode.cpp b/Engine/Plugins/Experimental/AnimNext/Source/AnimNextUncookedOnly/Private/AnimNextEdGraphNode.cpp
--- a/Engine/Plugins/Experimental/AnimNext/Source/AnimNextUncookedOnly/Private/AnimNextEdGraphNode.cpp
+++ b/Engine/Plugins/Experimental/AnimNext/Source/AnimNextUncookedOnly/Private/AnimNextEdGraphNode.cpp
@@ -194,6 +194,7 @@ void UAnimNextEdGraphNode::BuildExposeVariablesContextMenu(UToolMenu* SubMenu)
 
 	FAssetPickerConfig AssetPickerConfig;
 	AssetPickerConfig.Filter.ClassPaths.Add(UAnimNextDataInterface::StaticClass()->GetClassPathName());
+	AssetPickerConfig.Filter.bRecursiveClasses = true;
 	AssetPickerConfig.InitialAssetViewType = EAssetViewType::List;
 	AssetPickerConfig.AssetShowWarningText = LOCTEXT("NoAssetsWithPublicVariablesMessage", "No animation graphs with public variables found");
 	AssetPickerConfig.OnAssetSelected = FOnAssetSelected::CreateLambda([VMController, VMNode](const FAssetData& InAssetData)
diff --git a/Engine/Plugins/Experimental/AnimNext/Source/AnimNextUncookedOnly/Private/AnimNextRigVMAssetEditorData.cpp b/Engine/Plugins/Experimental/AnimNext/Source/AnimNextUncookedOnly/Private/AnimNextRigVMAssetEditorData.cpp
--- a/Engine/Plugins/Experimental/AnimNext/Source/AnimNextUncookedOnly/Private/AnimNextRigVMAssetEditorData.cpp
+++ b/Engine/Plugins/Experimental/AnimNext/Source/AnimNextUncookedOnly/Private/AnimNextRigVMAssetEditorData.cpp
@@ -32,10 +32,6 @@
 #include "UObject/AssetRegistryTagsContext.h"
 #include "DataInterface/AnimNextDataInterface_EditorData.h"
 
-TAutoConsoleVariable<bool> CVarDumpProgrammaticGraphs(
-	TEXT("AnimNext.DumpProgrammaticGraphs"),
-	false,
-	TEXT("When true the transient programmatic graphs will be automatically opened for any that are generated."));
 
 void UAnimNextRigVMAssetEditorData::BroadcastModified(EAnimNextEditorDataNotifType InType, UObject* InSubject)
 {
@@ -156,9 +152,9 @@ void UAnimNextRigVMAssetEditorData::Initialize(bool bRecompileVM)
 
 	if (IInterface_AssetUserData* OuterUserData = Cast<IInterface_AssetUserData>(GetOuter()))
 	{
-		if(!OuterUserData->HasAssetUserDataOfClass(UAnimNextAssetWorkspaceAssetUserData::StaticClass()))
+		if(!OuterUserData->HasAssetUserDataOfClass(GetAssetUserDataClass()))
 		{
-			OuterUserData->AddAssetUserDataOfClass(UAnimNextAssetWorkspaceAssetUserData::StaticClass());
+			OuterUserData->AddAssetUserDataOfClass(GetAssetUserDataClass());
 		}
 	}
 }
@@ -679,34 +675,7 @@ void UAnimNextRigVMAssetEditorData::RecompileVM()
 	// Display programmatic graphs
 	if(CVarDumpProgrammaticGraphs.GetValueOnGameThread())
 	{
-		UE::Workspace::IWorkspaceEditorModule& WorkspaceEditorModule = FModuleManager::LoadModuleChecked<UE::Workspace::IWorkspaceEditorModule>("WorkspaceEditor");
-		if(UE::Workspace::IWorkspaceEditor* WorkspaceEditor = WorkspaceEditorModule.OpenWorkspaceForObject(Asset, UE::Workspace::EOpenWorkspaceMethod::Default))
-		{
-			TArray<UObject*> Graphs;
-			for(URigVMGraph* ProgrammaticGraph : ProgrammaticGraphs)
-			{
-				// Some explanation needed here!
-				// URigVMEdGraph caches its underlying model internally in GetModel depending on its outer if it is no attached to a RigVMClient
-				// So here we rename the graph into the transient package so we dont get any notifications
-				ProgrammaticGraph->Rename(nullptr, GetTransientPackage(), REN_ForceNoResetLoaders | REN_DoNotDirty | REN_DontCreateRedirectors | REN_NonTransactional);
-
-				// then create the graph (transient so it outers to the RigVMGraph)
-				URigVMEdGraph* EdGraph = CastChecked<URigVMEdGraph>(CreateEdGraph(ProgrammaticGraph, true));
-
-				// Then cache the model
-				EdGraph->GetModel();
-				Graphs.Add(EdGraph);
-
-				// Now rename into this asset again to be able to correctly create a controller (needed to view the graph and interact with it)
-				ProgrammaticGraph->Rename(nullptr, this, REN_ForceNoResetLoaders | REN_DoNotDirty | REN_DontCreateRedirectors | REN_NonTransactional);
-				URigVMController* TempController = GetOrCreateController(ProgrammaticGraph);
-
-				// Resend notifications to rebuild the EdGraph
-				TempController->ResendAllNotifications();
-			}
-
-			WorkspaceEditor->OpenObjects(Graphs);
-		}
+		FUtils::OpenProgrammaticGraphs(this, ProgrammaticGraphs);
 	}
 #endif
 
@@ -862,6 +831,11 @@ void UAnimNextRigVMAssetEditorData::HandleModifiedEvent(ERigVMGraphNotifType InN
 	}
 }
 
+TSubclassOf<UAssetUserData> UAnimNextRigVMAssetEditorData::GetAssetUserDataClass() const
+{
+	return UAnimNextAssetWorkspaceAssetUserData::StaticClass();
+}
+
 TArray<UEdGraph*> UAnimNextRigVMAssetEditorData::GetAllEdGraphs() const
 {
 	TArray<UEdGraph*> Graphs;
@@ -1507,7 +1481,7 @@ UAnimNextAnimationGraphEntry* UAnimNextRigVMAssetEditorData::AddAnimationGraph(F
 			RefreshExternalModels();
 			RigVMClient.AddModel(NewRigVMGraphModel, true);
 			URigVMController* Controller = RigVMClient.GetController(NewRigVMGraphModel);
-			UE::AnimNext::UncookedOnly::FUtils::SetupAnimGraph(NewEntry, Controller);
+			UE::AnimNext::UncookedOnly::FUtils::SetupAnimGraph(NewEntry->GetEntryName(), Controller);
 		}
 	}
 
diff --git a/Engine/Plugins/Experimental/AnimNext/Source/AnimNextUncookedOnly/Private/AnimNextRigVMAssetSchema.h b/Engine/Plugins/Experimental/AnimNext/Source/AnimNextUncookedOnly/Private/AnimNextRigVMAssetSchema.h
deleted file mode 100644
diff --git a/Engine/Plugins/Experimental/AnimNext/Source/AnimNextUncookedOnly/Private/Graph/AnimNextAnimationGraph_EditorData.cpp b/Engine/Plugins/Experimental/AnimNext/Source/AnimNextUncookedOnly/Private/Graph/AnimNextAnimationGraph_EditorData.cpp
--- a/Engine/Plugins/Experimental/AnimNext/Source/AnimNextUncookedOnly/Private/Graph/AnimNextAnimationGraph_EditorData.cpp
+++ b/Engine/Plugins/Experimental/AnimNext/Source/AnimNextUncookedOnly/Private/Graph/AnimNextAnimationGraph_EditorData.cpp
@@ -361,25 +361,11 @@ namespace UE::AnimNext::UncookedOnly::Private
 			// If the graph is empty, add a dummy node that just pushes a reference pose
 			URigVMUnitNode* VMNode = VMController->AddUnitNode(FRigUnit_AnimNextTraitStack::StaticStruct(), FRigVMStruct::ExecuteName, FVector2D(0.0f, 0.0f), FString(), false);
 
-			const UScriptStruct* CppDecoratorStruct = FRigDecorator_AnimNextCppDecorator::StaticStruct();
-
-			FString DefaultValue;
-			{
-				const UE::AnimNext::FTraitUID ReferencePoseTraitUID(0x7508ab89);	// Trait header is private, reference by UID directly
-				const FTrait* Trait = FTraitRegistry::Get().Find(ReferencePoseTraitUID);
-				check(Trait != nullptr);
-
-				const FRigDecorator_AnimNextCppDecorator DefaultCppDecoratorStructInstance;
-				FRigDecorator_AnimNextCppDecorator CppDecoratorStructInstance;
-				CppDecoratorStructInstance.DecoratorSharedDataStruct = Trait->GetTraitSharedDataStruct();
-
-				const FProperty* Prop = FAnimNextCppDecoratorWrapper::StaticStruct()->FindPropertyByName(GET_MEMBER_NAME_STRING_CHECKED(FAnimNextCppDecoratorWrapper, CppDecorator));
-				check(Prop != nullptr);
-
-				Prop->ExportText_Direct(DefaultValue, &CppDecoratorStructInstance, &DefaultCppDecoratorStructInstance, nullptr, PPF_SerializedAsImportText);
-			}
-
-			VMController->AddTrait(VMNode->GetFName(), *CppDecoratorStruct->GetPathName(), TEXT("ReferencePose"), DefaultValue, INDEX_NONE, false, false);
+			UAnimNextController* AnimNextController = CastChecked<UAnimNextController>(VMController);
+			constexpr UE::AnimNext::FTraitUID ReferencePoseTraitUID(0x7508ab89); // Trait header is private, reference by UID directly
+			const FName RigVMTraitName =  AnimNextController->AddTraitByName(VMNode->GetFName(), *UE::AnimNext::FTraitRegistry::Get().Find(ReferencePoseTraitUID)->GetTraitName(), INDEX_NONE);
+		
+			check(RigVMTraitName != NAME_None);
 
 			FTraitStackMapping Mapping(VMNode);
 			ForEachTraitInStack(VMNode,
@@ -488,189 +474,200 @@ void UAnimNextAnimationGraph_EditorData::RecompileVM()
 	RigGraphDisplaySettings.MinMicroSeconds = RigGraphDisplaySettings.LastMinMicroSeconds = DBL_MAX;
 	RigGraphDisplaySettings.MaxMicroSeconds = RigGraphDisplaySettings.LastMaxMicroSeconds = (double)INDEX_NONE;
 
-	TGuardValue<bool> ReentrantGuardSelf(bSuspendModelNotificationsForSelf, true);
-	TGuardValue<bool> ReentrantGuardOthers(RigVMClient.bSuspendModelNotificationsForOthers, true);
+	TArray<URigVMGraph*> ProgrammaticGraphs;
+	{
+		TGuardValue<bool> ReentrantGuardSelf(bSuspendModelNotificationsForSelf, true);
+		TGuardValue<bool> ReentrantGuardOthers(RigVMClient.bSuspendModelNotificationsForOthers, true);
 
-	VMCompileSettings.SetExecuteContextStruct(FAnimNextExecuteContext::StaticStruct());
-	FRigVMCompileSettings Settings = (bCompileInDebugMode) ? FRigVMCompileSettings::Fast(VMCompileSettings.GetExecuteContextStruct()) : VMCompileSettings;
-	Settings.ASTSettings.bSetupTraits = false; // disable the default implementation of decorators for now
+		VMCompileSettings.SetExecuteContextStruct(FAnimNextExecuteContext::StaticStruct());
+		FRigVMCompileSettings Settings = (bCompileInDebugMode) ? FRigVMCompileSettings::Fast(VMCompileSettings.GetExecuteContextStruct()) : VMCompileSettings;
+		Settings.ASTSettings.bSetupTraits = false; // disable the default implementation of decorators for now
 
-	FMessageLog("AnimNextCompilerResults").NewPage(FText::FromName(AnimationGraph->GetFName()));
-	Settings.ASTSettings.ReportDelegate.BindLambda([](EMessageSeverity::Type InType, UObject* InObject, const FString& InString)
-	{
-		FMessageLog("AnimNextCompilerResults").Message(InType, FText::FromString(InString));
-	});
+		FMessageLog("AnimNextCompilerResults").NewPage(FText::FromName(AnimationGraph->GetFName()));
+		Settings.ASTSettings.ReportDelegate.BindLambda([](EMessageSeverity::Type InType, UObject* InObject, const FString& InString)
+		{
+			FMessageLog("AnimNextCompilerResults").Message(InType, FText::FromString(InString));
+		});
 
-	FUtils::RecreateVM(AnimationGraph);
+		FUtils::RecreateVM(AnimationGraph);
 
-	FUtils::CompileVariables(AnimationGraph);
+		FUtils::CompileVariables(AnimationGraph);
 
-	AnimationGraph->VMRuntimeSettings = VMRuntimeSettings;
-	AnimationGraph->EntryPoints.Empty();
-	AnimationGraph->ResolvedRootTraitHandles.Empty();
-	AnimationGraph->ResolvedEntryPoints.Empty();
-	AnimationGraph->ExecuteDefinition = FAnimNextGraphEvaluatorExecuteDefinition();
-	AnimationGraph->SharedDataBuffer.Empty();
-	AnimationGraph->GraphReferencedObjects.Empty();
-	AnimationGraph->DefaultEntryPoint = NAME_None;
+		AnimationGraph->VMRuntimeSettings = VMRuntimeSettings;
+		AnimationGraph->EntryPoints.Empty();
+		AnimationGraph->ResolvedRootTraitHandles.Empty();
+		AnimationGraph->ResolvedEntryPoints.Empty();
+		AnimationGraph->ExecuteDefinition = FAnimNextGraphEvaluatorExecuteDefinition();
+		AnimationGraph->SharedDataBuffer.Empty();
+		AnimationGraph->GraphReferencedObjects.Empty();
+		AnimationGraph->DefaultEntryPoint = NAME_None;
 
-	FRigVMClient* VMClient = GetRigVMClient();
+		FRigVMClient* VMClient = GetRigVMClient();
 
-	TArray<URigVMGraph*> ProgrammaticGraphs;
-	GetProgrammaticGraphs(Settings, ProgrammaticGraphs);
-	for(URigVMGraph* ProgrammaticGraph : ProgrammaticGraphs)
-	{
-		check(ProgrammaticGraph != nullptr);
-	}
+		GetProgrammaticGraphs(Settings, ProgrammaticGraphs);
+		for(URigVMGraph* ProgrammaticGraph : ProgrammaticGraphs)
+		{
+			check(ProgrammaticGraph != nullptr);
+		}
 	
-	TArray<URigVMGraph*> AllGraphs = VMClient->GetAllModels(false, false);
-	AllGraphs.Append(ProgrammaticGraphs);
+		TArray<URigVMGraph*> AllGraphs = VMClient->GetAllModels(false, false);
+		AllGraphs.Append(ProgrammaticGraphs);
 
-	if(AllGraphs.Num() == 0)
-	{
-		return;
-	}
+		if(AllGraphs.Num() == 0)
+		{
+			return;
+		}
 
-	TArray<URigVMGraph*> TempGraphs;
-	for(const URigVMGraph* SourceGraph : AllGraphs)
-	{
-		// We use a temporary graph models to build our final graphs that we'll compile
-		URigVMGraph* TempGraph = CastChecked<URigVMGraph>(StaticDuplicateObject(SourceGraph, GetTransientPackage(), NAME_None, RF_Transient));
-		TempGraph->SetFlags(RF_Transient);
-		TempGraphs.Add(TempGraph);
-	}
+		TArray<URigVMGraph*> TempGraphs;
+		for(const URigVMGraph* SourceGraph : AllGraphs)
+		{
+			// We use a temporary graph models to build our final graphs that we'll compile
+			URigVMGraph* TempGraph = CastChecked<URigVMGraph>(StaticDuplicateObject(SourceGraph, GetTransientPackage(), NAME_None, RF_Transient));
+			TempGraph->SetFlags(RF_Transient);
+			TempGraphs.Add(TempGraph);
+		}
 
-	if(TempGraphs.Num() == 0)
-	{
-		return;
-	}
+		if(TempGraphs.Num() == 0)
+		{
+			return;
+		}
 
-	UAnimNextController* TempController = CastChecked<UAnimNextController>(VMClient->GetOrCreateController(TempGraphs[0]));
+		UAnimNextController* TempController = CastChecked<UAnimNextController>(VMClient->GetOrCreateController(TempGraphs[0]));
 
-	UE::AnimNext::FTraitWriter TraitWriter;
+		UE::AnimNext::FTraitWriter TraitWriter;
 
-	FRigVMPinInfoArray LatentPins;
-	TMap<FName, URigVMPin*> LatentPinMapping;
-	TArray<Private::FTraitGraph> TraitGraphs;
+		FRigVMPinInfoArray LatentPins;
+		TMap<FName, URigVMPin*> LatentPinMapping;
+		TArray<Private::FTraitGraph> TraitGraphs;
 
-	// Build entry points and extract their required latent pins
-	for(const URigVMGraph* TempGraph : TempGraphs)
-	{
-		if(TempGraph->GetSchemaClass() == UAnimNextAnimationGraphSchema::StaticClass())
+		// Build entry points and extract their required latent pins
+		for(const URigVMGraph* TempGraph : TempGraphs)
 		{
-			// Gather our trait stacks
-			Private::FTraitGraph& TraitGraph = TraitGraphs.Add_GetRef(Private::CollectGraphInfo(AnimationGraph, TempGraph, TempController->GetControllerForGraph(TempGraph)));
-			check(!TraitGraph.TraitStackNodes.IsEmpty());
+			if(TempGraph->GetSchemaClass() == UAnimNextAnimationGraphSchema::StaticClass())
+			{
+				// Gather our trait stacks
+				Private::FTraitGraph& TraitGraph = TraitGraphs.Add_GetRef(Private::CollectGraphInfo(AnimationGraph, TempGraph, TempController->GetControllerForGraph(TempGraph)));
+				check(!TraitGraph.TraitStackNodes.IsEmpty());
 
-			FAnimNextGraphEntryPoint& EntryPoint = AnimationGraph->EntryPoints.AddDefaulted_GetRef();
-			EntryPoint.EntryPointName = TraitGraph.EntryPoint;
+				FAnimNextGraphEntryPoint& EntryPoint = AnimationGraph->EntryPoints.AddDefaulted_GetRef();
+				EntryPoint.EntryPointName = TraitGraph.EntryPoint;
 
-			// Extract latent pins for this graph
-			Private::CollectLatentPins(TraitGraph.TraitStackNodes, LatentPins, LatentPinMapping);
+				// Extract latent pins for this graph
+				Private::CollectLatentPins(TraitGraph.TraitStackNodes, LatentPins, LatentPinMapping);
 
-			// Iterate over every trait stack and register our node templates
-			for (Private::FTraitStackMapping& NodeMapping : TraitGraph.TraitStackNodes)
-			{
-				NodeMapping.TraitStackNodeHandle = Private::RegisterTraitNodeTemplate(TraitWriter, NodeMapping.DecoratorStackNode);
-			}
+				// Iterate over every trait stack and register our node templates
+				for (Private::FTraitStackMapping& NodeMapping : TraitGraph.TraitStackNodes)
+				{
+					NodeMapping.TraitStackNodeHandle = Private::RegisterTraitNodeTemplate(TraitWriter, NodeMapping.DecoratorStackNode);
+				}
 
-			// Find our root node handle, if we have any stack nodes, the first one is our root stack
-			if (TraitGraph.TraitStackNodes.Num() != 0)
-			{
-				EntryPoint.RootTraitHandle = FAnimNextEntryPointHandle(TraitGraph.TraitStackNodes[0].TraitStackNodeHandle);
+				// Find our root node handle, if we have any stack nodes, the first one is our root stack
+				if (TraitGraph.TraitStackNodes.Num() != 0)
+				{
+					EntryPoint.RootTraitHandle = FAnimNextEntryPointHandle(TraitGraph.TraitStackNodes[0].TraitStackNodeHandle);
+				}
 			}
 		}
-	}
-
-	// Set default entry point
-	if(AnimationGraph->EntryPoints.Num() > 0)
-	{
-		AnimationGraph->DefaultEntryPoint = AnimationGraph->EntryPoints[0].EntryPointName;
-	}
 
-	// Remove our old root nodes
-	for (Private::FTraitGraph& TraitGraph : TraitGraphs)
-	{
-		URigVMController* GraphController = TempController->GetControllerForGraph(TraitGraph.RootNode->GetGraph());
-		GraphController->RemoveNode(TraitGraph.RootNode, false, false);
-	}
+		// Set default entry point
+		if(AnimationGraph->EntryPoints.Num() > 0)
+		{
+			AnimationGraph->DefaultEntryPoint = AnimationGraph->EntryPoints[0].EntryPointName;
+		}
 
-	if(LatentPins.Num() > 0)
-	{
-		// We need a unique method name to match our unique argument list
-		AnimationGraph->ExecuteDefinition = Private::GetGraphEvaluatorExecuteMethod(LatentPins);
+		// Remove our old root nodes
+		for (Private::FTraitGraph& TraitGraph : TraitGraphs)
+		{
+			URigVMController* GraphController = TempController->GetControllerForGraph(TraitGraph.RootNode->GetGraph());
+			GraphController->RemoveNode(TraitGraph.RootNode, false, false);
+		}
 
-		// Add our runtime shim root node
-		URigVMUnitNode* TempShimRootNode = TempController->AddUnitNode(FRigUnit_AnimNextShimRoot::StaticStruct(), FRigUnit_AnimNextShimRoot::EventName, FVector2D::ZeroVector, FString(), false);
-		URigVMUnitNode* GraphEvaluatorNode = TempController->AddUnitNodeWithPins(FRigUnit_AnimNextGraphEvaluator::StaticStruct(), LatentPins, *AnimationGraph->ExecuteDefinition.MethodName, FVector2D::ZeroVector, FString(), false);
+		if(LatentPins.Num() > 0)
+		{
+			// We need a unique method name to match our unique argument list
+			AnimationGraph->ExecuteDefinition = Private::GetGraphEvaluatorExecuteMethod(LatentPins);
 
-		// Link our shim and evaluator nodes together using the execution context
-		TempController->AddLink(
-			TempShimRootNode->FindPin(GET_MEMBER_NAME_STRING_CHECKED(FRigUnit_AnimNextShimRoot, ExecuteContext)),
-			GraphEvaluatorNode->FindPin(GET_MEMBER_NAME_STRING_CHECKED(FRigUnit_AnimNextGraphEvaluator, ExecuteContext)),
-			false);
+			// Add our runtime shim root node
+			URigVMUnitNode* TempShimRootNode = TempController->AddUnitNode(FRigUnit_AnimNextShimRoot::StaticStruct(), FRigUnit_AnimNextShimRoot::EventName, FVector2D::ZeroVector, FString(), false);
+			URigVMUnitNode* GraphEvaluatorNode = TempController->AddUnitNodeWithPins(FRigUnit_AnimNextGraphEvaluator::StaticStruct(), LatentPins, *AnimationGraph->ExecuteDefinition.MethodName, FVector2D::ZeroVector, FString(), false);
 
-		// Link our latent pins
-		for (const FRigVMPinInfo& LatentPin : LatentPins)
-		{
+			// Link our shim and evaluator nodes together using the execution context
 			TempController->AddLink(
-				LatentPinMapping[LatentPin.Name],
-				GraphEvaluatorNode->FindPin(LatentPin.Name.ToString()),
+				TempShimRootNode->FindPin(GET_MEMBER_NAME_STRING_CHECKED(FRigUnit_AnimNextShimRoot, ExecuteContext)),
+				GraphEvaluatorNode->FindPin(GET_MEMBER_NAME_STRING_CHECKED(FRigUnit_AnimNextGraphEvaluator, ExecuteContext)),
 				false);
+
+			// Link our latent pins
+			for (const FRigVMPinInfo& LatentPin : LatentPins)
+			{
+				TempController->AddLink(
+					LatentPinMapping[LatentPin.Name],
+					GraphEvaluatorNode->FindPin(LatentPin.Name.ToString()),
+					false);
+			}
 		}
-	}
 
-	// Write our node shared data
-	TraitWriter.BeginNodeWriting();
+		// Write our node shared data
+		TraitWriter.BeginNodeWriting();
 
-	for(Private::FTraitGraph& TraitGraph : TraitGraphs)
-	{
-		for (const Private::FTraitStackMapping& NodeMapping : TraitGraph.TraitStackNodes)
+		for(Private::FTraitGraph& TraitGraph : TraitGraphs)
 		{
-			Private::WriteTraitProperties(TraitWriter, NodeMapping, TraitGraph.TraitStackNodes);
+			for (const Private::FTraitStackMapping& NodeMapping : TraitGraph.TraitStackNodes)
+			{
+				Private::WriteTraitProperties(TraitWriter, NodeMapping, TraitGraph.TraitStackNodes);
+			}
 		}
-	}
 
-	TraitWriter.EndNodeWriting();
+		TraitWriter.EndNodeWriting();
 
-	// Cache our compiled metadata
-	AnimationGraph->SharedDataArchiveBuffer = TraitWriter.GetGraphSharedData();
-	AnimationGraph->GraphReferencedObjects = TraitWriter.GetGraphReferencedObjects();
+		// Cache our compiled metadata
+		AnimationGraph->SharedDataArchiveBuffer = TraitWriter.GetGraphSharedData();
+		AnimationGraph->GraphReferencedObjects = TraitWriter.GetGraphReferencedObjects();
 
-	// Populate our runtime metadata
-	AnimationGraph->LoadFromArchiveBuffer(AnimationGraph->SharedDataArchiveBuffer);
+		// Populate our runtime metadata
+		AnimationGraph->LoadFromArchiveBuffer(AnimationGraph->SharedDataArchiveBuffer);
 
-	URigVMCompiler* Compiler = URigVMCompiler::StaticClass()->GetDefaultObject<URigVMCompiler>();
-	Compiler->Compile(Settings, TempGraphs, TempController, AnimationGraph->VM, AnimationGraph->ExtendedExecuteContext, AnimationGraph->GetExternalVariables(), &PinToOperandMap);
+		URigVMCompiler* Compiler = URigVMCompiler::StaticClass()->GetDefaultObject<URigVMCompiler>();
+		Compiler->Compile(Settings, TempGraphs, TempController, AnimationGraph->VM, AnimationGraph->ExtendedExecuteContext, AnimationGraph->GetExternalVariables(), &PinToOperandMap);
 
-	// Initialize right away, in packaged builds we initialize during PostLoad
-	AnimationGraph->VM->Initialize(AnimationGraph->ExtendedExecuteContext);
-	AnimationGraph->GenerateUserDefinedDependenciesData(AnimationGraph->ExtendedExecuteContext);
+		// Initialize right away, in packaged builds we initialize during PostLoad
+		AnimationGraph->VM->Initialize(AnimationGraph->ExtendedExecuteContext);
+		AnimationGraph->GenerateUserDefinedDependenciesData(AnimationGraph->ExtendedExecuteContext);
 
-	// Notable difference with vanilla RigVM host behavior - we init the VM here at the moment as we only have one 'instance'
-	AnimationGraph->InitializeVM(FRigUnit_AnimNextBeginExecution::EventName);
+		// Notable difference with vanilla RigVM host behavior - we init the VM here at the moment as we only have one 'instance'
+		AnimationGraph->InitializeVM(FRigUnit_AnimNextBeginExecution::EventName);
 
-	if (bErrorsDuringCompilation)
-	{
-		if(Settings.SurpressErrors)
+		if (bErrorsDuringCompilation)
 		{
-			Settings.Reportf(EMessageSeverity::Info, AnimationGraph, TEXT("Compilation Errors may be suppressed for AnimNext asset: %s. See VM Compile Settings for more Details"), *AnimationGraph->GetName());
+			if(Settings.SurpressErrors)
+			{
+				Settings.Reportf(EMessageSeverity::Info, AnimationGraph, TEXT("Compilation Errors may be suppressed for AnimNext asset: %s. See VM Compile Settings for more Details"), *AnimationGraph->GetName());
+			}
 		}
-	}
 
-	bVMRecompilationRequired = false;
-	if(AnimationGraph->VM)
-	{
-		RigVMCompiledEvent.Broadcast(AnimationGraph, AnimationGraph->VM, AnimationGraph->ExtendedExecuteContext);
-	}
+		bVMRecompilationRequired = false;
+		if(AnimationGraph->VM)
+		{
+			RigVMCompiledEvent.Broadcast(AnimationGraph, AnimationGraph->VM, AnimationGraph->ExtendedExecuteContext);
+		}
+
+		for(URigVMGraph* TempGraph : TempGraphs)
+		{
+			VMClient->RemoveController(TempGraph);
+		}
 
-	for(URigVMGraph* TempGraph : TempGraphs)
+		// Now that the graph has been re-compiled, re-allocate the previous live instances
+		AnimationGraph->ThawGraphInstances();
+	}
+		
+#if WITH_EDITOR
+	// Display programmatic graphs
+	if(CVarDumpProgrammaticGraphs.GetValueOnGameThread())
 	{
-		VMClient->RemoveController(TempGraph);
+		FUtils::OpenProgrammaticGraphs(this, ProgrammaticGraphs);
 	}
+#endif
 
-	// Now that the graph has been re-compiled, re-allocate the previous live instances
-	AnimationGraph->ThawGraphInstances();
 
 #if WITH_EDITOR
 //	RefreshBreakpoints(EditorData);
diff --git a/Engine/Plugins/Experimental/AnimNext/Source/AnimNextUncookedOnly/Private/UncookedOnlyUtils.cpp b/Engine/Plugins/Experimental/AnimNext/Source/AnimNextUncookedOnly/Private/UncookedOnlyUtils.cpp
--- a/Engine/Plugins/Experimental/AnimNext/Source/AnimNextUncookedOnly/Private/UncookedOnlyUtils.cpp
+++ b/Engine/Plugins/Experimental/AnimNext/Source/AnimNextUncookedOnly/Private/UncookedOnlyUtils.cpp
@@ -21,6 +21,8 @@
 #include "Graph/AnimNextAnimationGraphSchema.h"
 #include "Logging/StructuredLog.h"
 #include "AnimNextAssetWorkspaceAssetUserData.h"
+#include "IWorkspaceEditor.h"
+#include "IWorkspaceEditorModule.h"
 #include "Misc/EnumerateRange.h"
 #include "Module/RigUnit_AnimNextModuleEvents.h"
 #include "Variables/IVariableBindingType.h"
@@ -35,6 +37,11 @@
 namespace UE::AnimNext::UncookedOnly
 {
 
+TAutoConsoleVariable<bool> CVarDumpProgrammaticGraphs(
+	TEXT("AnimNext.DumpProgrammaticGraphs"),
+	false,
+	TEXT("When true the transient programmatic graphs will be automatically opened for any that are generated."));
+
 void FUtils::RecreateVM(UAnimNextRigVMAsset* InAsset)
 {
 	if (InAsset->VM == nullptr)
@@ -329,13 +336,13 @@ void FUtils::CompileVariableBindings(const FRigVMCompileSettings& InSettings, UA
 
 	FRigVMClient* VMClient = EditorData->GetRigVMClient();
 	URigVMController* Controller = VMClient->GetOrCreateController(BindingGraph);
-	URigVMNode* ExcecuteBindingsNode = Controller->AddUnitNode(FRigUnit_AnimNextExecuteBindings::StaticStruct(), FRigUnit_AnimNextExecuteBindings::GetMethodName(), FVector2D::ZeroVector, FString(), false);
-	if(ExcecuteBindingsNode == nullptr)
+	URigVMNode* ExecuteBindingsNode = Controller->AddUnitNode(FRigUnit_AnimNextExecuteBindings::StaticStruct(), FRigUnit_AnimNextExecuteBindings::GetMethodName(), FVector2D::ZeroVector, FString(), false);
+	if(ExecuteBindingsNode == nullptr)
 	{
 		InSettings.ReportError(TEXT("Could not spawn Execute Bindings node"));
 		return;
 	}
-	URigVMPin* ExecuteBindingsExecPin = ExcecuteBindingsNode->FindPin(FRigVMStruct::ExecuteContextName.ToString());
+	URigVMPin* ExecuteBindingsExecPin = ExecuteBindingsNode->FindPin(FRigVMStruct::ExecuteContextName.ToString());
 	if(ExecuteBindingsExecPin == nullptr)
 	{
 		InSettings.ReportError(TEXT("Could not find execute pin on Execute Bindings node"));
@@ -656,7 +663,7 @@ FRigVMTemplateArgumentType FUtils::GetRigVMArgTypeFromParamType(const FAnimNextP
 	return ArgType;
 }
 
-void FUtils::SetupAnimGraph(UAnimNextRigVMAssetEntry* InEntry, URigVMController* InController)
+void FUtils::SetupAnimGraph(const FName EntryName, URigVMController* InController)
 {
 	// Clear the graph
 	InController->RemoveNodes(InController->GetGraph()->GetNodes());
@@ -680,7 +687,7 @@ void FUtils::SetupAnimGraph(UAnimNextRigVMAssetEntry* InEntry, URigVMController*
 		return;
 	}
 
-	InController->SetPinDefaultValue(EntryPointPin->GetPinPath(), InEntry->GetEntryName().ToString());
+	InController->SetPinDefaultValue(EntryPointPin->GetPinPath(), EntryName.ToString());
 }
 
 void FUtils::SetupEventGraph(URigVMController* InController, UScriptStruct* InEventStruct)
@@ -902,6 +909,40 @@ const FText& FUtils::GetFunctionLibraryDisplayName()
 	return FunctionLibraryName;
 }
 
+#if WITH_EDITOR
+void FUtils::OpenProgrammaticGraphs(UAnimNextRigVMAssetEditorData* EditorData, const TArray<URigVMGraph*>& ProgrammaticGraphs)
+{
+	UAnimNextRigVMAsset* OwningAsset = FUtils::GetAsset(EditorData);
+	UE::Workspace::IWorkspaceEditorModule& WorkspaceEditorModule = FModuleManager::LoadModuleChecked<UE::Workspace::IWorkspaceEditorModule>("WorkspaceEditor");
+	if(UE::Workspace::IWorkspaceEditor* WorkspaceEditor = WorkspaceEditorModule.OpenWorkspaceForObject(OwningAsset, UE::Workspace::EOpenWorkspaceMethod::Default))
+	{
+		TArray<UObject*> Graphs;
+		for(URigVMGraph* ProgrammaticGraph : ProgrammaticGraphs)
+		{
+			// Some explanation needed here!
+			// URigVMEdGraph caches its underlying model internally in GetModel depending on its outer if it is no attached to a RigVMClient
+			// So here we rename the graph into the transient package so we dont get any notifications
+			ProgrammaticGraph->Rename(nullptr, GetTransientPackage(), REN_ForceNoResetLoaders | REN_DoNotDirty | REN_DontCreateRedirectors | REN_NonTransactional);
+
+			// then create the graph (transient so it outers to the RigVMGraph)
+			URigVMEdGraph* EdGraph = CastChecked<URigVMEdGraph>(EditorData->CreateEdGraph(ProgrammaticGraph, true));
+
+			// Then cache the model
+			EdGraph->GetModel();
+			Graphs.Add(EdGraph);
+
+			// Now rename into this asset again to be able to correctly create a controller (needed to view the graph and interact with it)
+			ProgrammaticGraph->Rename(nullptr, EditorData, REN_ForceNoResetLoaders | REN_DoNotDirty | REN_DontCreateRedirectors | REN_NonTransactional);
+			URigVMController* ProgrammaticController = EditorData->GetOrCreateController(ProgrammaticGraph);
+
+			// Resend notifications to rebuild the EdGraph
+			ProgrammaticController->ResendAllNotifications();
+		}
+
+		WorkspaceEditor->OpenObjects(Graphs);
+	}
+}
+#endif // WITH_EDITOR
 }
 
 #undef LOCTEXT_NAMESPACE
\ No newline at end of file
diff --git a/Engine/Plugins/Experimental/AnimNextAnimGraph/Source/AnimNextAnimGraph/Internal/Traits/BlendStackTrait.h b/Engine/Plugins/Experimental/AnimNextAnimGraph/Source/AnimNextAnimGraph/Internal/Traits/BlendStackTrait.h
new file mode 100644
--- /dev/null
+++ b/Engine/Plugins/Experimental/AnimNextAnimGraph/Source/AnimNextAnimGraph/Internal/Traits/BlendStackTrait.h
@@ -0,0 +1,202 @@
+// Copyright Epic Games, Inc. All Rights Reserved.
+
+#pragma once
+
+#include "TraitCore/Trait.h"
+#include "TraitInterfaces/IDiscreteBlend.h"
+#include "TraitInterfaces/IGarbageCollection.h"
+#include "TraitInterfaces/IHierarchy.h"
+#include "TraitInterfaces/IUpdate.h"
+#include "Graph/AnimNextAnimationGraph.h"
+#include "Graph/AnimNextGraphInstancePtr.h"
+#include "TraitInterfaces/IBlendStack.h"
+#include "TraitInterfaces/ISmoothBlend.h"
+
+#include "BlendStackTrait.generated.h"
+
+USTRUCT(meta = (DisplayName = "Blend Stack Core"))
+struct FAnimNextBlendStackCoreTraitSharedData : public FAnimNextTraitSharedData
+{
+	GENERATED_BODY()
+
+	// This struct is empty but required so that we can give a nice display name to the trait
+};
+
+USTRUCT(meta = (DisplayName = "Blend Stack"))
+struct FAnimNextBlendStackTraitSharedData : public FAnimNextTraitSharedData
+{
+	GENERATED_BODY()
+
+	/** Asset to use as a sub-graph */
+	UPROPERTY(EditAnywhere, Category = "Default")
+	TObjectPtr<const UAnimNextAnimationGraph> AnimationGraph;
+
+	/** Blend time to use when sub-graph pin input changes */
+	UPROPERTY(EditAnywhere, Category = "Default")
+	float BlendTime = 0.2f;
+
+	/** Force a blend to the current graph, even if there's been no change in selection. */
+	UPROPERTY(EditAnywhere, Category = "Default")
+	bool bForceBlend = false;
+
+	// Latent pin support boilerplate
+	#define TRAIT_LATENT_PROPERTIES_ENUMERATOR(GeneratorMacro) \
+		GeneratorMacro(AnimationGraph) \
+		GeneratorMacro(BlendTime) \
+		GeneratorMacro(bForceBlend) \
+
+	GENERATE_TRAIT_LATENT_PROPERTIES(FAnimNextBlendStackTraitSharedData, TRAIT_LATENT_PROPERTIES_ENUMERATOR)
+	#undef TRAIT_LATENT_PROPERTIES_ENUMERATOR
+};
+
+USTRUCT(meta = (DisplayName = "Blend Stack Requester"))
+struct FAnimNextBlendStackRequesterTraitSharedData : public FAnimNextTraitSharedData
+{
+	GENERATED_BODY()
+
+	/** Asset to use as a sub-graph */
+	UPROPERTY(EditAnywhere, Category = "Default")
+	TObjectPtr<const UAnimNextAnimationGraph> AnimationGraph;
+
+	/** Blend time to use when sub-graph pin input changes */
+	UPROPERTY(EditAnywhere, Category = "Default")
+	float BlendTime = 0.2f;
+
+	/** Force a blend to the current graph, even if there's been no change in selection. */
+	UPROPERTY(EditAnywhere, Category = "Default")
+	bool bForceBlend = false;
+
+	// Latent pin support boilerplate
+	#define TRAIT_LATENT_PROPERTIES_ENUMERATOR(GeneratorMacro) \
+		GeneratorMacro(AnimationGraph) \
+		GeneratorMacro(BlendTime) \
+		GeneratorMacro(bForceBlend) \
+
+	GENERATE_TRAIT_LATENT_PROPERTIES(FAnimNextBlendStackRequesterTraitSharedData, TRAIT_LATENT_PROPERTIES_ENUMERATOR)
+	#undef TRAIT_LATENT_PROPERTIES_ENUMERATOR
+};
+
+namespace UE::AnimNext
+{
+	/**
+	 * FBlendStackCoreTrait
+	 * 
+	 * Used as base trait in state machine-like setups where state machine is the additive trait that pushes graphs and blend settings to the blend stack.
+	 */
+	struct FBlendStackCoreTrait : FBaseTrait, IUpdateTraversal, IHierarchy, IDiscreteBlend, ISmoothBlend, IGarbageCollection, IBlendStack
+	{
+		DECLARE_ANIM_TRAIT(FBlendStackCoreTrait, 0xda7b7f8d, FBaseTrait)
+
+		enum class EGraphState : uint8
+		{
+			Inactive,
+			Active,
+		};
+
+		struct FGraphState
+		{
+			// The blend settings to use
+			IBlendStack::FGraphRequest Request;
+
+			// The graph instance
+			FAnimNextGraphInstancePtr Instance;
+
+			// Our child handle
+			// Weak handle to the graph instance's root
+			FTraitPtr ChildPtr;
+
+			FBlendStackCoreTrait::EGraphState State;
+
+			// Whether or not this graph was previously relevant
+			bool bNewlyCreated = false;
+
+			void Initialize(const IBlendStack::FGraphRequest& GraphRequest);
+			void Terminate();
+		};
+
+		struct FInstanceData : FTrait::FInstanceData
+		{
+			// List of current graphs in the blend stack
+			TArray<FGraphState> ChildGraphs;
+
+			// The index of the currently active graph
+			// All other graphs are blending out
+			int32 CurrentlyActiveGraphIndex = INDEX_NONE;
+
+			void Construct(const FExecutionContext& Context, const FTraitBinding& Binding);
+			void Destruct(const FExecutionContext& Context, const FTraitBinding& Binding);
+		};
+
+		using FSharedData = FAnimNextBlendStackCoreTraitSharedData;
+
+		// Internal impl
+		static int32 FindFreeGraphIndexOrAdd(FInstanceData& InstanceData);
+
+		// IUpdateTraversal impl
+		virtual void QueueChildrenForTraversal(FUpdateTraversalContext& Context, const TTraitBinding<IUpdateTraversal>& Binding, const FTraitUpdateState& TraitState, FUpdateTraversalQueue& TraversalQueue) const override;
+
+		// IHierarchy impl
+		virtual uint32 GetNumChildren(const FExecutionContext& Context, const TTraitBinding<IHierarchy>& Binding) const override;
+		virtual void GetChildren(const FExecutionContext& Context, const TTraitBinding<IHierarchy>& Binding, FChildrenArray& Children) const override;
+
+		// IDiscreteBlend impl
+		virtual int32 GetBlendDestinationChildIndex(FExecutionContext& Context, const TTraitBinding<IDiscreteBlend>& Binding) const override;
+		virtual void OnBlendTransition(FExecutionContext& Context, const TTraitBinding<IDiscreteBlend>& Binding, int32 OldChildIndex, int32 NewChildIndex) const override;
+		virtual void OnBlendInitiated(FExecutionContext& Context, const TTraitBinding<IDiscreteBlend>& Binding, int32 ChildIndex) const override;
+		virtual void OnBlendTerminated(FExecutionContext& Context, const TTraitBinding<IDiscreteBlend>& Binding, int32 ChildIndex) const override;
+
+		//@TODO: Add inertial blend.
+		// ISmoothBlend impl
+		virtual float GetBlendTime(FExecutionContext& Context, const TTraitBinding<ISmoothBlend>& Binding, int32 ChildIndex) const override;
+		virtual EAlphaBlendOption GetBlendType(FExecutionContext& Context, const TTraitBinding<ISmoothBlend>& Binding, int32 ChildIndex) const override;
+		virtual UCurveFloat* GetCustomBlendCurve(FExecutionContext& Context, const TTraitBinding<ISmoothBlend>& Binding, int32 ChildIndex) const override;
+
+		// IGarbageCollection impl
+		virtual void AddReferencedObjects(const FExecutionContext& Context, const TTraitBinding<IGarbageCollection>& Binding, FReferenceCollector& Collector) const override;
+
+		// IBlendStack impl
+		virtual void PushGraph(FExecutionContext& Context, const TTraitBinding<IBlendStack>& Binding, const IBlendStack::FGraphRequest& GraphRequest, FAnimNextGraphInstancePtr& OutGraphInstance) const;
+		virtual void GetActiveGraphRequest(FExecutionContext& Context, const TTraitBinding<IBlendStack>& Binding, IBlendStack::FGraphRequest& OutRequest) const;
+
+#if WITH_EDITOR
+		virtual bool IsHidden() const override { return true; }
+#endif
+	};
+
+	/**
+	 * FBlendStackTrait
+	 * 
+	 * Standalone blend stack that exposes desired graph and blend time as pins. 
+	 * A new graph will be pushed every time the newest graph does not match the pin (or bForceBlend is true).
+	 */	
+	struct FBlendStackTrait : FBlendStackCoreTrait, IUpdate
+	{
+		DECLARE_ANIM_TRAIT(FBlendStackTrait, 0x46678513, FBlendStackCoreTrait)
+
+		using FSharedData = FAnimNextBlendStackTraitSharedData;
+
+		// IUpdate impl
+		virtual void PreUpdate(FUpdateTraversalContext& Context, const TTraitBinding<IUpdate>& Binding, const FTraitUpdateState& TraitState) const override;
+
+#if WITH_EDITOR
+		virtual bool IsHidden() const override { return false; }
+#endif
+	};
+
+	/**
+	 * FBlendStackRequesterTrait
+	 * 
+	 * Example additive trait used to push graphs to a Blend Stack base trait.
+	 * Same behavior as FBlendStackTrait, but as additive (i.e. FBlendStackCoreTrait base + FBlendStackRequesterTrait as additive).
+	 */	
+	struct FBlendStackRequesterTrait : FAdditiveTrait, IUpdate
+	{
+		DECLARE_ANIM_TRAIT(FBlendStackRequesterTrait, 0x3076bf31, FAdditiveTrait)
+
+		using FSharedData = FAnimNextBlendStackRequesterTraitSharedData;
+
+		// IUpdate impl
+		virtual void PreUpdate(FUpdateTraversalContext& Context, const TTraitBinding<IUpdate>& Binding, const FTraitUpdateState& TraitState) const override;
+		virtual void OnBecomeRelevant(FUpdateTraversalContext& Context, const TTraitBinding<IUpdate>& Binding, const FTraitUpdateState& TraitState) const override;
+	};
+}
diff --git a/Engine/Plugins/Experimental/AnimNextAnimGraph/Source/AnimNextAnimGraph/Private/Traits/BlendStackTrait.h b/Engine/Plugins/Experimental/AnimNextAnimGraph/Source/AnimNextAnimGraph/Private/Traits/BlendStackTrait.h
deleted file mode 100644
diff --git a/Engine/Plugins/Experimental/AnimNextStateTree/Source/AnimNextStateTree/AnimNextStateTree.Build.cs b/Engine/Plugins/Experimental/AnimNextStateTree/Source/AnimNextStateTree/AnimNextStateTree.Build.cs
--- a/Engine/Plugins/Experimental/AnimNextStateTree/Source/AnimNextStateTree/AnimNextStateTree.Build.cs
+++ b/Engine/Plugins/Experimental/AnimNextStateTree/Source/AnimNextStateTree/AnimNextStateTree.Build.cs
@@ -9,17 +9,19 @@ namespace UnrealBuildTool.Rules
 			PrivateDependencyModuleNames.AddRange(
 				new string[]
 				{
-					"AnimNext", 
 					"AnimNextAnimGraph",
 					"RigVM"
 				});
+			
 			PublicDependencyModuleNames.AddRange(
 				new string[]
 				{
 					"Core",
 					"CoreUObject",
 					"StateTreeModule",
-					"Engine"
+					"Engine",
+					"Chooser",
+					"AnimNext", 
 				}
 			);
 		}
diff --git a/Engine/Plugins/Experimental/AnimNextStateTree/Source/AnimNextStateTree/Internal/AnimStateTreeTrait.h b/Engine/Plugins/Experimental/AnimNextStateTree/Source/AnimNextStateTree/Internal/AnimStateTreeTrait.h
new file mode 100644
--- /dev/null
+++ b/Engine/Plugins/Experimental/AnimNextStateTree/Source/AnimNextStateTree/Internal/AnimStateTreeTrait.h
@@ -0,0 +1,62 @@
+// Copyright Epic Games, Inc. All Rights Reserved.
+
+#pragma once
+
+#include "CoreMinimal.h"
+#include "TraitCore/Trait.h"
+#include "TraitCore/TraitSharedData.h"
+#include "TraitInterfaces/IEvaluate.h"
+#include "TraitInterfaces/IGarbageCollection.h"
+#include "TraitInterfaces/IUpdate.h"
+
+#include "StateTree.h"
+#include "StateTreeReference.h"
+
+#include "AnimStateTreeTrait.generated.h"
+
+USTRUCT(meta = (DisplayName = "State Tree"))
+struct FAnimNextStateTreeTraitSharedData : public FAnimNextTraitSharedData
+{
+	GENERATED_BODY()
+
+	UPROPERTY(EditAnywhere, Category = "Default")
+	FStateTreeReference StateTreeReference;
+	
+	UPROPERTY(EditAnywhere, Category = "Default")
+	FStateTreeReferenceOverrides LinkedStateTreeOverrides;
+
+	// Latent pin support boilerplate
+#define TRAIT_LATENT_PROPERTIES_ENUMERATOR(GeneratorMacro) \
+GeneratorMacro(StateTreeReference) \
+GeneratorMacro(LinkedStateTreeOverrides) \
+
+GENERATE_TRAIT_LATENT_PROPERTIES(FAnimNextStateTreeTraitSharedData, TRAIT_LATENT_PROPERTIES_ENUMERATOR)
+#undef TRAIT_LATENT_PROPERTIES_ENUMERATOR
+};
+
+namespace UE::AnimNext
+{
+struct FStateTreeTrait : FAdditiveTrait, IUpdate, IGarbageCollection
+{
+	DECLARE_ANIM_TRAIT(FStateTreeTrait, 0x7a0dc697, FAdditiveTrait)
+
+	using FSharedData = FAnimNextStateTreeTraitSharedData;
+
+	struct FInstanceData : FTrait::FInstanceData
+	{
+		TObjectPtr<const UStateTree> StateTree;
+		FStateTreeInstanceData InstanceData;
+
+		void Construct(const FExecutionContext& Context, const FTraitBinding& Binding);
+		void Destruct(const FExecutionContext& Context, const FTraitBinding& Binding);
+	};
+
+
+	// IUpdate impl
+	virtual void OnBecomeRelevant(FUpdateTraversalContext& Context, const TTraitBinding<IUpdate>& Binding, const FTraitUpdateState& TraitState) const override;
+	virtual void PreUpdate(FUpdateTraversalContext& Context, const TTraitBinding<IUpdate>& Binding, const FTraitUpdateState& TraitState) const override;
+
+	// IGarbageCollection impl
+	virtual void AddReferencedObjects(const FExecutionContext& Context, const TTraitBinding<IGarbageCollection>& Binding, FReferenceCollector& Collector) const override;
+};
+}
\ No newline at end of file
diff --git a/Engine/Plugins/Experimental/AnimNextStateTree/Source/AnimNextStateTree/Private/AnimNextStateTreeContext.cpp b/Engine/Plugins/Experimental/AnimNextStateTree/Source/AnimNextStateTree/Private/AnimNextStateTreeContext.cpp
new file mode 100644
--- /dev/null
+++ b/Engine/Plugins/Experimental/AnimNextStateTree/Source/AnimNextStateTree/Private/AnimNextStateTreeContext.cpp
@@ -0,0 +1,38 @@
+// Copyright Epic Games, Inc. All Rights Reserved.
+
+#include "AnimNextStateTreeContext.h"
+
+#include "TraitCore/TraitStackBinding.h"
+#include "TraitInterfaces/IBlendStack.h"
+#include "TraitInterfaces/ISmoothBlend.h"
+
+bool FAnimNextStateTreeTraitContext::PushAnimationGraphOntoBlendStack(TNonNullPtr<UAnimNextAnimationGraph> InAnimationGraph, const FAlphaBlendArgs& InBlendArguments) const
+{		
+	if (Binding && Context)
+	{
+		UE::AnimNext::TTraitBinding<UE::AnimNext::IBlendStack> BlendStackBinding;
+		if(Binding->GetInterface<UE::AnimNext::IBlendStack>(BlendStackBinding))
+		{
+			UE::AnimNext::IBlendStack::FGraphRequest CurrentRequest;			
+			BlendStackBinding.GetActiveGraphRequest(*Context, CurrentRequest);
+
+			if (CurrentRequest.AnimationGraph != nullptr)
+			{
+				if (CurrentRequest.AnimationGraph == InAnimationGraph)
+				{
+					return true;
+				}
+			}
+			
+			FAnimNextGraphInstancePtr GraphPtr;
+			UE::AnimNext::IBlendStack::FGraphRequest GraphRequest;
+			GraphRequest.BlendArgs = InBlendArguments;
+			GraphRequest.AnimationGraph = InAnimationGraph.Get();
+			BlendStackBinding.PushGraph(*Context, GraphRequest, GraphPtr);
+
+			return true;
+		}
+	}
+
+	return false;
+}
\ No newline at end of file
diff --git a/Engine/Plugins/Experimental/AnimNextStateTree/Source/AnimNextStateTree/Private/AnimNextStateTreeGraphInstanceTask.cpp b/Engine/Plugins/Experimental/AnimNextStateTree/Source/AnimNextStateTree/Private/AnimNextStateTreeGraphInstanceTask.cpp
new file mode 100644
--- /dev/null
+++ b/Engine/Plugins/Experimental/AnimNextStateTree/Source/AnimNextStateTree/Private/AnimNextStateTreeGraphInstanceTask.cpp
@@ -0,0 +1,43 @@
+// Copyright Epic Games, Inc. All Rights Reserved.
+
+#include "Tasks/AnimNextStateTreeGraphInstanceTask.h"
+#include "StateTreeLinker.h"
+#include "StateTreeExecutionContext.h"
+
+FAnimNextStateTreeGraphInstanceTask::FAnimNextStateTreeGraphInstanceTask()
+{
+}
+
+bool FAnimNextStateTreeGraphInstanceTask::Link(FStateTreeLinker& Linker)
+{
+	Linker.LinkExternalData(TraitContextHandle);	
+	return true;
+}
+
+EStateTreeRunStatus FAnimNextStateTreeGraphInstanceTask::EnterState(FStateTreeExecutionContext& Context, const FStateTreeTransitionResult& Transition) const
+{
+	FAnimNextStateTreeTraitContext& ExecContext = Context.GetExternalData(TraitContextHandle);
+	FInstanceDataType& InstanceData = Context.GetInstanceData(*this);
+	if (ExecContext.PushAnimationGraphOntoBlendStack(InstanceData.AnimationGraph.Get(), InstanceData.BlendOptions))
+	{
+		return EStateTreeRunStatus::Running;	
+	}
+	
+	return EStateTreeRunStatus::Failed;
+}
+
+EStateTreeRunStatus FAnimNextStateTreeGraphInstanceTask::Tick(FStateTreeExecutionContext& Context, const float DeltaTime) const
+{
+	FInstanceDataType& InstanceData = Context.GetInstanceData(*this);
+	if (InstanceData.bContinueTicking)
+	{
+		return EStateTreeRunStatus::Running;
+	}
+	
+	return EStateTreeRunStatus::Succeeded;
+}
+
+void FAnimNextStateTreeGraphInstanceTask::ExitState(FStateTreeExecutionContext& Context, const FStateTreeTransitionResult& Transition) const
+{
+	FAnimNextStateTreeTaskBase::ExitState(Context, Transition);
+}
diff --git a/Engine/Plugins/Experimental/AnimNextStateTree/Source/AnimNextStateTree/Private/AnimNextStateTreeModule.cpp b/Engine/Plugins/Experimental/AnimNextStateTree/Source/AnimNextStateTree/Private/AnimNextStateTreeModule.cpp
--- a/Engine/Plugins/Experimental/AnimNextStateTree/Source/AnimNextStateTree/Private/AnimNextStateTreeModule.cpp
+++ b/Engine/Plugins/Experimental/AnimNextStateTree/Source/AnimNextStateTree/Private/AnimNextStateTreeModule.cpp
@@ -2,6 +2,11 @@
 
 #include "AnimNextStateTreeModule.h"
 
+#include "StateTree.h"
+#include "StateTreeReference.h"
+#include "Graph/AnimNextAnimationGraph.h"
+#include "RigVMCore/RigVMRegistry.h"
+
 namespace UE::AnimNext::StateTree
 {
 
@@ -11,6 +16,20 @@ void FAnimNextStateTreeModule::ShutdownModule()
 
 void FAnimNextStateTreeModule::StartupModule()
 {
+	static TPair<UClass*, FRigVMRegistry::ERegisterObjectOperation> const AllowedObjectTypes[] =
+	{
+		{ UStateTree::StaticClass(), FRigVMRegistry::ERegisterObjectOperation::Class },
+		{ UAnimNextAnimationGraph::StaticClass(), FRigVMRegistry::ERegisterObjectOperation::Class }
+	};
+
+	FRigVMRegistry::Get().RegisterObjectTypes(AllowedObjectTypes);
+
+	static UScriptStruct* const AllowedStructTypes[] =
+	{
+		FStateTreeReference::StaticStruct(),
+	};
+
+	FRigVMRegistry::Get().RegisterStructTypes(AllowedStructTypes);
 }
 
 IMPLEMENT_MODULE(FAnimNextStateTreeModule, AnimNextStateTree)
diff --git a/Engine/Plugins/Experimental/AnimNextStateTree/Source/AnimNextStateTree/Private/AnimNextStateTreeSchema.cpp b/Engine/Plugins/Experimental/AnimNextStateTree/Source/AnimNextStateTree/Private/AnimNextStateTreeSchema.cpp
new file mode 100644
--- /dev/null
+++ b/Engine/Plugins/Experimental/AnimNextStateTree/Source/AnimNextStateTree/Private/AnimNextStateTreeSchema.cpp
@@ -0,0 +1,42 @@
+// Copyright Epic Games, Inc. All Rights Reserved.
+
+#include "AnimNextStateTreeSchema.h"
+
+#include "AnimNextStateTreeTypes.h"
+#include "StateTreeConditionBase.h"
+#include "StateTreeConsiderationBase.h"
+#include "StateTreeEvaluatorBase.h"
+#include "StateTreePropertyFunctionBase.h"
+#include "StateTreeTaskBase.h"
+#include "Components/ActorComponent.h"
+#include "Subsystems/WorldSubsystem.h"
+#include "GameFramework/Actor.h"
+#include "AnimNextStateTreeContext.h"
+
+const FName UStateTreeAnimNextSchema::AnimStateTreeExecutionContextName("ExecutionContext");
+
+UStateTreeAnimNextSchema::UStateTreeAnimNextSchema() : ContextDataDescs(
+	{{ AnimStateTreeExecutionContextName, FAnimNextStateTreeTraitContext::StaticStruct(), FGuid(0xDFB93B7F, 0xEDBE4906, 0x851C66B2, 0x7585FA21) }})
+{
+}
+
+bool UStateTreeAnimNextSchema::IsStructAllowed(const UScriptStruct* InScriptStruct) const
+{
+	return InScriptStruct->IsChildOf(FStateTreeConditionCommonBase::StaticStruct())
+	|| InScriptStruct->IsChildOf(FStateTreeEvaluatorCommonBase::StaticStruct())
+	|| InScriptStruct->IsChildOf(FStateTreeTaskCommonBase::StaticStruct())
+	|| InScriptStruct->IsChildOf(FStateTreeConsiderationCommonBase::StaticStruct())
+	|| InScriptStruct->IsChildOf(FStateTreePropertyFunctionCommonBase::StaticStruct())
+	|| InScriptStruct->IsChildOf(FAnimNextStateTreeEvaluatorBase::StaticStruct())
+	|| InScriptStruct->IsChildOf(FAnimNextStateTreeTaskBase::StaticStruct());
+}
+
+bool UStateTreeAnimNextSchema::IsClassAllowed(const UClass* InClass) const
+{
+	return false;
+}
+
+bool UStateTreeAnimNextSchema::IsExternalItemAllowed(const UStruct& InStruct) const
+{
+	return InStruct.IsChildOf<FAnimNextStateTreeTraitContext>();
+}
diff --git a/Engine/Plugins/Experimental/AnimNextStateTree/Source/AnimNextStateTree/Private/AnimStateTreeTrait.cpp b/Engine/Plugins/Experimental/AnimNextStateTree/Source/AnimNextStateTree/Private/AnimStateTreeTrait.cpp
new file mode 100644
--- /dev/null
+++ b/Engine/Plugins/Experimental/AnimNextStateTree/Source/AnimNextStateTree/Private/AnimStateTreeTrait.cpp
@@ -0,0 +1,161 @@
+// Copyright Epic Games, Inc. All Rights Reserved.
+
+#include "AnimStateTreeTrait.h"
+
+#include "AnimNextStateTreeContext.h"
+#include "AnimNextStateTreeSchema.h"
+#include "StateTreeExecutionContext.h"
+#include "Graph/AnimNextGraphInstance.h"
+#include "Logging/LogScopedVerbosityOverride.h"
+#include "TraitCore/NodeInstance.h"
+
+namespace UE::AnimNext
+{
+	AUTO_REGISTER_ANIM_TRAIT(FStateTreeTrait)
+
+	// Trait implementation boilerplate
+	#define TRAIT_INTERFACE_ENUMERATOR(GeneratorMacro) \
+	GeneratorMacro(IUpdate) \
+	GeneratorMacro(IGarbageCollection) \
+
+	GENERATE_ANIM_TRAIT_IMPLEMENTATION(FStateTreeTrait, TRAIT_INTERFACE_ENUMERATOR, NULL_ANIM_TRAIT_INTERFACE_ENUMERATOR, NULL_ANIM_TRAIT_EVENT_ENUMERATOR)
+	#undef TRAIT_INTERFACE_ENUMERATOR
+
+
+	void FStateTreeTrait::FInstanceData::Construct(const FExecutionContext& Context, const FTraitBinding& Binding)
+	{
+		FTrait::FInstanceData::Construct(Context, Binding);
+		IGarbageCollection::RegisterWithGC(Context, Binding);
+	}
+
+	void FStateTreeTrait::FInstanceData::Destruct(const FExecutionContext& Context, const FTraitBinding& Binding)
+	{
+		FTrait::FInstanceData::Destruct(Context, Binding);
+		IGarbageCollection::UnregisterWithGC(Context, Binding);
+	}
+
+	void FStateTreeTrait::OnBecomeRelevant(FUpdateTraversalContext& Context, const TTraitBinding<IUpdate>& Binding, const FTraitUpdateState& TraitState) const
+	{
+		const FSharedData* SharedData = Binding.GetSharedData<FSharedData>();
+		FInstanceData* InstanceData = Binding.GetInstanceData<FInstanceData>();
+		InstanceData->StateTree = SharedData->StateTreeReference.GetStateTree();
+
+		if (InstanceData->StateTree)
+		{
+			UObject* Owner = GetTransientPackage();
+			FStateTreeExecutionContext StateTreeExecutionContext(*Owner, *InstanceData->StateTree, InstanceData->InstanceData);
+
+			FAnimNextStateTreeTraitContext TraitContext(Context, Binding.GetStack());
+			StateTreeExecutionContext.SetContextDataByName(UStateTreeAnimNextSchema::AnimStateTreeExecutionContextName, FStateTreeDataView(FAnimNextStateTreeTraitContext::StaticStruct(), reinterpret_cast<uint8*>(&TraitContext)));
+
+			FAnimNextGraphInstance* OwnerGraphInstance = &Binding.GetTraitPtr().GetNodeInstance()->GetOwner();
+			if(OwnerGraphInstance && InstanceData->StateTree->GetDefaultParameters().Identical(&OwnerGraphInstance->GetVariables(), 0))
+			{
+				const FInstancedPropertyBag& StateTreeParameters = InstanceData->StateTree->GetDefaultParameters();
+				FStructView MutableVariables = OwnerGraphInstance->GetVariables().GetMutableValue();
+				
+				const FRigVMExtendedExecuteContext& ExtendedExecuteContext = OwnerGraphInstance->GetExtendedExecuteContext();
+
+				// This copy-behaviour is temporary until we find a better way to directly patch the StateTree property binding/copies
+				const int32 NumVariables = StateTreeParameters.GetNumPropertiesInBag();
+				for(int32 VariableIndex = 0; VariableIndex < NumVariables; ++VariableIndex)
+				{
+					if (ExtendedExecuteContext.ExternalVariableRuntimeData.IsValidIndex(VariableIndex))
+					{
+						const FPropertyBagPropertyDesc& Desc = StateTreeParameters.GetPropertyBagStruct()->GetPropertyDescs()[VariableIndex];
+						void* TargetAddress = MutableVariables.GetMemory() + Desc.CachedProperty->GetOffset_ForInternal();
+						void const* SourceAddress = ExtendedExecuteContext.ExternalVariableRuntimeData[VariableIndex].Memory;
+
+						Desc.CachedProperty->CopyCompleteValue(TargetAddress, SourceAddress);
+					}
+				}
+				
+				StateTreeExecutionContext.SetCollectExternalDataCallback(FOnCollectStateTreeExternalData::CreateLambda([this, &TraitContext]( const FStateTreeExecutionContext& Context, const UStateTree* StateTree, TArrayView<const FStateTreeExternalDataDesc> ExternalDataDescs, TArrayView<FStateTreeDataView> OutDataViews) -> bool
+				{
+					for (int32 Index = 0; Index < ExternalDataDescs.Num(); Index++)
+					{
+						const FStateTreeExternalDataDesc& ItemDesc = ExternalDataDescs[Index];
+						if (ItemDesc.Struct != nullptr)
+						{
+							if (ItemDesc.Struct->IsChildOf(FAnimNextStateTreeTraitContext::StaticStruct()))
+							{
+								OutDataViews[Index] = FStateTreeDataView(FAnimNextStateTreeTraitContext::StaticStruct(), reinterpret_cast<uint8*>(&TraitContext));
+							}
+						}
+					}
+
+					return true;
+				}));
+		
+				if (StateTreeExecutionContext.IsValid())
+				{
+					StateTreeExecutionContext.Start(&OwnerGraphInstance->GetVariables());
+				}
+			}
+			
+		}	
+	}
+
+	void FStateTreeTrait::PreUpdate(FUpdateTraversalContext& Context, const TTraitBinding<IUpdate>& Binding, const FTraitUpdateState& TraitState) const
+	{
+		FInstanceData* InstanceData = Binding.GetInstanceData<FInstanceData>();
+		if (InstanceData->StateTree)
+		{
+			UObject* Owner = GetTransientPackage();
+			
+			FStateTreeExecutionContext StateTreeExecutionContext(*Owner, *InstanceData->StateTree, InstanceData->InstanceData);
+			FAnimNextGraphInstance* OwnerGraphInstance = &Binding.GetTraitPtr().GetNodeInstance()->GetOwner();
+			if(OwnerGraphInstance && InstanceData->StateTree->GetDefaultParameters().Identical(&OwnerGraphInstance->GetVariables(), 0))
+			{
+				FAnimNextStateTreeTraitContext TraitContext(Context, Binding.GetStack());
+				StateTreeExecutionContext.SetContextDataByName(UStateTreeAnimNextSchema::AnimStateTreeExecutionContextName, FStateTreeDataView(FAnimNextStateTreeTraitContext::StaticStruct(), reinterpret_cast<uint8*>(&TraitContext)));
+
+				const FInstancedPropertyBag& StateTreeParameters = InstanceData->StateTree->GetDefaultParameters();
+				FStructView MutableVariables = InstanceData->InstanceData.GetMutableStorage().GetMutableGlobalParameters();
+				
+				const FRigVMExtendedExecuteContext& ExtendedExecuteContext = OwnerGraphInstance->GetExtendedExecuteContext();
+
+				// This copy-behaviour is temporary until we find a better way to directly patch the StateTree property binding/copies
+				const int32 NumVariables = StateTreeParameters.GetNumPropertiesInBag();
+				for(int32 VariableIndex = 0; VariableIndex < NumVariables; ++VariableIndex)
+				{
+					if (ExtendedExecuteContext.ExternalVariableRuntimeData.IsValidIndex(VariableIndex))
+					{
+						const FPropertyBagPropertyDesc& Desc = StateTreeParameters.GetPropertyBagStruct()->GetPropertyDescs()[VariableIndex];
+						void* TargetAddress = MutableVariables.GetMemory() + Desc.CachedProperty->GetOffset_ForInternal();
+						void const* SourceAddress = ExtendedExecuteContext.ExternalVariableRuntimeData[VariableIndex].Memory;
+
+						Desc.CachedProperty->CopyCompleteValue(TargetAddress, SourceAddress);
+					}
+				}
+
+				StateTreeExecutionContext.SetCollectExternalDataCallback(FOnCollectStateTreeExternalData::CreateLambda([this, &TraitContext]( const FStateTreeExecutionContext& Context, const UStateTree* StateTree, TArrayView<const FStateTreeExternalDataDesc> ExternalDataDescs, TArrayView<FStateTreeDataView> OutDataViews) -> bool
+				{
+					for (int32 Index = 0; Index < ExternalDataDescs.Num(); Index++)
+					{
+						const FStateTreeExternalDataDesc& ItemDesc = ExternalDataDescs[Index];
+						if (ItemDesc.Struct != nullptr)
+						{
+							if (ItemDesc.Struct->IsChildOf(FAnimNextStateTreeTraitContext::StaticStruct()))
+							{
+								OutDataViews[Index] = FStateTreeDataView(FAnimNextStateTreeTraitContext::StaticStruct(), reinterpret_cast<uint8*>(&TraitContext));
+							}
+						}
+					}
+
+					return true;
+				}));
+
+				StateTreeExecutionContext.Tick(TraitState.GetDeltaTime());
+			}
+		}
+	}
+
+	void FStateTreeTrait::AddReferencedObjects(const FExecutionContext& Context, const TTraitBinding<IGarbageCollection>& Binding, FReferenceCollector& Collector) const
+	{
+		IGarbageCollection::AddReferencedObjects(Context, Binding, Collector);
+
+		FInstanceData* InstanceData = Binding.GetInstanceData<FInstanceData>();
+		Collector.AddReferencedObject(InstanceData->StateTree);
+	}
+}
diff --git a/Engine/Plugins/Experimental/AnimNextStateTree/Source/AnimNextStateTree/Public/AnimNextStateTree.h b/Engine/Plugins/Experimental/AnimNextStateTree/Source/AnimNextStateTree/Public/AnimNextStateTree.h
new file mode 100644
--- /dev/null
+++ b/Engine/Plugins/Experimental/AnimNextStateTree/Source/AnimNextStateTree/Public/AnimNextStateTree.h
@@ -0,0 +1,28 @@
+// Copyright Epic Games, Inc. All Rights Reserved.
+
+#pragma once
+
+#include "CoreMinimal.h"
+#include "StateTree.h"
+#include "Graph/AnimNextAnimationGraph.h"
+
+#include "AnimNextStateTree.generated.h"
+
+namespace UE::AnimNext::UncookedOnly
+{
+struct FUtils;
+}
+
+
+UCLASS(BlueprintType)
+class ANIMNEXTSTATETREE_API UAnimNextStateTree : public UAnimNextAnimationGraph
+{
+	GENERATED_BODY()
+
+public:
+	friend class UAnimNextStateTree_EditorData;
+	friend class UAnimNextStateTreeFactory;
+
+	UPROPERTY()
+	TObjectPtr<UStateTree> StateTree;
+};
\ No newline at end of file
diff --git a/Engine/Plugins/Experimental/AnimNextStateTree/Source/AnimNextStateTree/Public/AnimNextStateTreeContext.h b/Engine/Plugins/Experimental/AnimNextStateTree/Source/AnimNextStateTree/Public/AnimNextStateTreeContext.h
new file mode 100644
--- /dev/null
+++ b/Engine/Plugins/Experimental/AnimNextStateTree/Source/AnimNextStateTree/Public/AnimNextStateTreeContext.h
@@ -0,0 +1,33 @@
+// Copyright Epic Games, Inc. All Rights Reserved.
+
+#pragma once
+
+#include "AnimNextStateTreeContext.generated.h"
+
+namespace UE::AnimNext
+{
+struct FTraitStackBinding;
+struct FExecutionContext;
+struct FStateTreeTrait;
+}
+
+class UAnimNextAnimationGraph;
+struct FAlphaBlendArgs;
+
+USTRUCT()
+struct FAnimNextStateTreeTraitContext
+{	
+	GENERATED_BODY()
+	
+	friend UE::AnimNext::FStateTreeTrait;
+	
+	FAnimNextStateTreeTraitContext() {}
+
+	bool PushAnimationGraphOntoBlendStack(TNonNullPtr<UAnimNextAnimationGraph> InAnimationGraph, const FAlphaBlendArgs& InBlendArguments) const;
+protected:
+	FAnimNextStateTreeTraitContext(UE::AnimNext::FExecutionContext& InContext, const UE::AnimNext::FTraitStackBinding* InBinding) : Context(&InContext), Binding(InBinding) {}
+protected:
+	UE::AnimNext::FExecutionContext* Context = nullptr;
+	const UE::AnimNext::FTraitStackBinding* Binding = nullptr;
+};
+
diff --git a/Engine/Plugins/Experimental/AnimNextStateTree/Source/AnimNextStateTree/Public/AnimNextStateTreeSchema.h b/Engine/Plugins/Experimental/AnimNextStateTree/Source/AnimNextStateTree/Public/AnimNextStateTreeSchema.h
new file mode 100644
--- /dev/null
+++ b/Engine/Plugins/Experimental/AnimNextStateTree/Source/AnimNextStateTree/Public/AnimNextStateTreeSchema.h
@@ -0,0 +1,28 @@
+// Copyright Epic Games, Inc. All Rights Reserved.
+
+#pragma once
+
+#include "StateTreeSchema.h"
+
+#include "AnimNextStateTreeSchema.generated.h"
+
+
+UCLASS()
+class ANIMNEXTSTATETREE_API UStateTreeAnimNextSchema : public UStateTreeSchema
+{
+	GENERATED_BODY()
+public:
+	UStateTreeAnimNextSchema();
+	
+	static const FName AnimStateTreeExecutionContextName;
+
+	virtual bool IsStructAllowed(const UScriptStruct* InScriptStruct) const override;
+	virtual bool IsClassAllowed(const UClass* InScriptStruct) const override;
+	virtual bool IsExternalItemAllowed(const UStruct& InStruct) const override;
+	
+	virtual TConstArrayView<FStateTreeExternalDataDesc> GetContextDataDescs() const override { return ContextDataDescs; }
+protected:
+	/** List of named external data required by schema and provided to the state tree through the execution context. */
+	UPROPERTY()
+	TArray<FStateTreeExternalDataDesc> ContextDataDescs;
+};
diff --git a/Engine/Plugins/Experimental/AnimNextStateTree/Source/AnimNextStateTree/Public/AnimNextStateTreeTypes.h b/Engine/Plugins/Experimental/AnimNextStateTree/Source/AnimNextStateTree/Public/AnimNextStateTreeTypes.h
new file mode 100644
--- /dev/null
+++ b/Engine/Plugins/Experimental/AnimNextStateTree/Source/AnimNextStateTree/Public/AnimNextStateTreeTypes.h
@@ -0,0 +1,27 @@
+// Copyright Epic Games, Inc. All Rights Reserved.
+
+#pragma once
+
+#include "CoreMinimal.h"
+#include "StateTreeEvaluatorBase.h" 
+#include "StateTreeTaskBase.h"
+
+#include "AnimNextStateTreeTypes.generated.h"
+
+/**
+ * Base struct for all AnimNext StateTree Evaluators.
+ */
+USTRUCT(meta = (DisplayName = "AnimNext Evaluator Base", Hidden))
+struct ANIMNEXTSTATETREE_API FAnimNextStateTreeEvaluatorBase : public FStateTreeEvaluatorBase
+{
+	GENERATED_BODY()
+};
+
+/**
+ * Base struct for all AnimNext StateTree Tasks.
+ */
+USTRUCT(meta = (DisplayName = "AnimNext Task Base", Hidden))
+struct ANIMNEXTSTATETREE_API FAnimNextStateTreeTaskBase : public FStateTreeTaskBase
+{
+	GENERATED_BODY()
+};
diff --git a/Engine/Plugins/Experimental/AnimNextStateTree/Source/AnimNextStateTree/Public/Tasks/AnimNextStateTreeGraphInstanceTask.h b/Engine/Plugins/Experimental/AnimNextStateTree/Source/AnimNextStateTree/Public/Tasks/AnimNextStateTreeGraphInstanceTask.h
new file mode 100644
--- /dev/null
+++ b/Engine/Plugins/Experimental/AnimNextStateTree/Source/AnimNextStateTree/Public/Tasks/AnimNextStateTreeGraphInstanceTask.h
@@ -0,0 +1,46 @@
+// Copyright Epic Games, Inc. All Rights Reserved.
+
+#pragma once
+
+#include "AnimNextStateTreeContext.h"
+#include "AnimNextStateTreeTypes.h"
+#include "Graph/AnimNextAnimationGraph.h"
+#include "AlphaBlend.h"
+
+#include "AnimNextStateTreeGraphInstanceTask.generated.h"
+
+USTRUCT()
+struct ANIMNEXTSTATETREE_API FAnimNextGraphInstanceTaskInstanceData
+{
+	GENERATED_BODY()
+	
+	UPROPERTY(EditAnywhere, Category = Animation)
+	TObjectPtr<UAnimNextAnimationGraph> AnimationGraph = nullptr;
+
+	UPROPERTY(EditAnywhere, Category = Animation)
+	FAlphaBlendArgs BlendOptions;
+
+	UPROPERTY(EditAnywhere, Category = Animation)
+	bool bContinueTicking = true;
+};
+
+// Basic task pushing AnimationGraph onto blend stack
+USTRUCT(meta = (DisplayName = "AnimNext Graph"))
+struct ANIMNEXTSTATETREE_API FAnimNextStateTreeGraphInstanceTask : public FAnimNextStateTreeTaskBase
+{
+	GENERATED_BODY()
+
+	using FInstanceDataType = FAnimNextGraphInstanceTaskInstanceData;
+
+	FAnimNextStateTreeGraphInstanceTask();
+	
+	virtual bool Link(FStateTreeLinker& Linker) override;
+protected:
+	virtual const UStruct* GetInstanceDataType() const override { return FInstanceDataType::StaticStruct(); }
+	virtual EStateTreeRunStatus EnterState(FStateTreeExecutionContext& Context, const FStateTreeTransitionResult& Transition) const override;
+	virtual EStateTreeRunStatus Tick(FStateTreeExecutionContext& Context, const float DeltaTime) const override;
+	virtual void ExitState(FStateTreeExecutionContext& Context, const FStateTreeTransitionResult& Transition) const override;
+
+public:
+	TStateTreeExternalDataHandle<FAnimNextStateTreeTraitContext> TraitContextHandle;
+};
diff --git a/Engine/Plugins/Experimental/AnimNextStateTree/Source/AnimNextStateTreeEditor/AnimNextStateTreeEditor.Build.cs b/Engine/Plugins/Experimental/AnimNextStateTree/Source/AnimNextStateTreeEditor/AnimNextStateTreeEditor.Build.cs
--- a/Engine/Plugins/Experimental/AnimNextStateTree/Source/AnimNextStateTreeEditor/AnimNextStateTreeEditor.Build.cs
+++ b/Engine/Plugins/Experimental/AnimNextStateTree/Source/AnimNextStateTreeEditor/AnimNextStateTreeEditor.Build.cs
@@ -6,7 +6,21 @@ namespace UnrealBuildTool.Rules
 	{
 		public AnimNextStateTreeEditor(ReadOnlyTargetRules Target) : base(Target)
 		{
-			PrivateDependencyModuleNames.AddRange(new string[] { "AnimNextStateTree", "UnrealEd", "MessageLog" });
+			PrivateDependencyModuleNames.AddRange(
+				new string[]
+				{
+				    "AnimNextStateTree",
+				    "MessageLog",
+				    "EditorInteractiveToolsFramework",
+				    "InteractiveToolsFramework",
+				    "EditorFramework",
+				    "AnimNext",
+				    "UnrealEd",
+				    "AssetDefinition",
+				    "AnimNextAnimGraph",
+				    "RigVM"
+			    });
+
 			PublicDependencyModuleNames.AddRange(
 				new string[]
 				{
@@ -20,7 +34,9 @@ namespace UnrealBuildTool.Rules
 					"SlateCore",
 					"EditorSubsystem",
 					"Engine",
-					"Slate"
+					"Slate",
+					"RigVMDeveloper",
+					"AnimNextUncookedOnly"
 				}
 			);
 		}
diff --git a/Engine/Plugins/Experimental/AnimNextStateTree/Source/AnimNextStateTreeEditor/Internal/AnimNextStateTreeFactory.h b/Engine/Plugins/Experimental/AnimNextStateTree/Source/AnimNextStateTreeEditor/Internal/AnimNextStateTreeFactory.h
new file mode 100644
--- /dev/null
+++ b/Engine/Plugins/Experimental/AnimNextStateTree/Source/AnimNextStateTreeEditor/Internal/AnimNextStateTreeFactory.h
@@ -0,0 +1,22 @@
+// Copyright Epic Games, Inc. All Rights Reserved.
+
+#pragma once
+
+#include "CoreMinimal.h"
+#include "Factories/Factory.h"
+#include "AnimNextStateTreeFactory.generated.h"
+
+UCLASS(MinimalAPI)
+class UAnimNextStateTreeFactory : public UFactory
+{
+	GENERATED_BODY()
+
+	UAnimNextStateTreeFactory();
+
+	virtual bool ConfigureProperties() override;
+	virtual UObject* FactoryCreateNew(UClass* Class, UObject* InParent, FName Name, EObjectFlags Flags, UObject* Context, FFeedbackContext* Warn, FName CallingContext) override;
+	virtual UObject* FactoryCreateNew(UClass* Class, UObject* InParent, FName Name, EObjectFlags Flags, UObject* Context, FFeedbackContext* Warn) override
+	{
+		return FactoryCreateNew(Class, InParent, Name, Flags, Context, Warn, NAME_None);
+	}
+};
\ No newline at end of file
diff --git a/Engine/Plugins/Experimental/AnimNextStateTree/Source/AnimNextStateTreeEditor/Private/AnimNextStateTreeAssetDefinition.cpp b/Engine/Plugins/Experimental/AnimNextStateTree/Source/AnimNextStateTreeEditor/Private/AnimNextStateTreeAssetDefinition.cpp
new file mode 100644
--- /dev/null
+++ b/Engine/Plugins/Experimental/AnimNextStateTree/Source/AnimNextStateTreeEditor/Private/AnimNextStateTreeAssetDefinition.cpp
@@ -0,0 +1,20 @@
+// Copyright Epic Games, Inc. All Rights Reserved.
+
+#include "AnimNextStateTreeAssetDefinition.h"
+#include "IWorkspaceEditorModule.h"
+#include "Workspace/AnimNextWorkspaceFactory.h"
+#include "Modules/ModuleManager.h"
+
+EAssetCommandResult UAssetDefinition_AnimNextStateTree::OpenAssets(const FAssetOpenArgs& OpenArgs) const
+{
+	using namespace UE::AnimNext::Editor;
+	using namespace UE::Workspace;
+
+	for (UAnimNextStateTree* Asset : OpenArgs.LoadObjects<UAnimNextStateTree>())
+	{
+		IWorkspaceEditorModule& WorkspaceEditorModule = FModuleManager::Get().LoadModuleChecked<IWorkspaceEditorModule>("WorkspaceEditor");
+		WorkspaceEditorModule.OpenWorkspaceForObject(Asset, EOpenWorkspaceMethod::Default, UAnimNextWorkspaceFactory::StaticClass());
+	}
+
+	return EAssetCommandResult::Handled;
+}
diff --git a/Engine/Plugins/Experimental/AnimNextStateTree/Source/AnimNextStateTreeEditor/Private/AnimNextStateTreeAssetDefinition.h b/Engine/Plugins/Experimental/AnimNextStateTree/Source/AnimNextStateTreeEditor/Private/AnimNextStateTreeAssetDefinition.h
new file mode 100644
--- /dev/null
+++ b/Engine/Plugins/Experimental/AnimNextStateTree/Source/AnimNextStateTreeEditor/Private/AnimNextStateTreeAssetDefinition.h
@@ -0,0 +1,30 @@
+// Copyright Epic Games, Inc. All Rights Reserved.
+
+#pragma once
+
+#include "AnimNextStateTree.h"
+#include "AssetDefinitionDefault.h"
+#include "AnimNextStateTreeAssetDefinition.generated.h"
+
+#define LOCTEXT_NAMESPACE "AnimNextAssetDefinitions"
+
+UCLASS()
+class UAssetDefinition_AnimNextStateTree : public UAssetDefinitionDefault
+{
+	GENERATED_BODY()
+
+public:
+	// UAssetDefinition interface
+	virtual FText GetAssetDisplayName() const override { return LOCTEXT("AnimNextStateTree", "AnimNext State Tree"); }
+	virtual FLinearColor GetAssetColor() const override { return FLinearColor(FColor(96,128,48)); }
+	virtual TSoftClassPtr<UObject> GetAssetClass() const override { return UAnimNextStateTree::StaticClass(); }
+	virtual EAssetCommandResult OpenAssets(const FAssetOpenArgs& OpenArgs) const override;
+	virtual TConstArrayView<FAssetCategoryPath> GetAssetCategories() const override
+	{
+		static const auto Categories = { FAssetCategoryPath(EAssetCategoryPaths::Animation, LOCTEXT("AnimNextSubMenu", "AnimNext")) };
+		return Categories;
+	}
+	virtual bool ShouldSaveExternalPackages() const override { return true; }
+};
+
+#undef LOCTEXT_NAMESPACE
\ No newline at end of file
diff --git a/Engine/Plugins/Experimental/AnimNextStateTree/Source/AnimNextStateTreeEditor/Private/AnimNextStateTreeEditorHost.cpp b/Engine/Plugins/Experimental/AnimNextStateTree/Source/AnimNextStateTreeEditor/Private/AnimNextStateTreeEditorHost.cpp
new file mode 100644
--- /dev/null
+++ b/Engine/Plugins/Experimental/AnimNextStateTree/Source/AnimNextStateTreeEditor/Private/AnimNextStateTreeEditorHost.cpp
@@ -0,0 +1,69 @@
+// Copyright Epic Games, Inc. All Rights Reserved.
+
+#include "AnimNextStateTreeEditorHost.h"
+
+#include "AnimNextStateTree.h"
+#include "IAnimNextEditorModule.h"
+#include "StateTree.h"
+
+#include "IWorkspaceEditor.h"
+
+void FAnimNextStateTreeEditorHost::Init(const TWeakPtr<UE::Workspace::IWorkspaceEditor>& InWeakWorkspaceEditor)
+{
+	WeakWorkspaceEditor = InWeakWorkspaceEditor;
+	
+	TSharedPtr<UE::Workspace::IWorkspaceEditor> SharedEditor = InWeakWorkspaceEditor.Pin();
+	check(SharedEditor.IsValid());
+	SharedEditor->OnFocussedAssetChanged().AddSP(this, &FAnimNextStateTreeEditorHost::OnWorkspaceFocussedAssetChanged);
+}
+
+UStateTree* FAnimNextStateTreeEditorHost::GetStateTree() const
+{
+	if (TSharedPtr<UE::Workspace::IWorkspaceEditor> SharedWorkspaceEditor = WeakWorkspaceEditor.Pin())
+	{
+		const TObjectPtr<UAnimNextStateTree> AnimNextStateTreePtr = SharedWorkspaceEditor->GetFocussedAsset<UAnimNextStateTree>();
+		return AnimNextStateTreePtr ? AnimNextStateTreePtr->StateTree : nullptr;
+	}	
+
+	return nullptr;
+}
+
+FName FAnimNextStateTreeEditorHost::GetCompilerLogName() const
+{
+	return UE::AnimNext::Editor::LogListingName;
+}
+
+FName FAnimNextStateTreeEditorHost::GetCompilerTabName() const
+{
+	return UE::AnimNext::Editor::CompilerResultsTabName;
+}
+
+void FAnimNextStateTreeEditorHost::OnWorkspaceFocussedAssetChanged(TObjectPtr<UObject> InObject) const
+{
+	OnStateTreeChangedDelegate.Broadcast();
+}
+
+FSimpleMulticastDelegate& FAnimNextStateTreeEditorHost::OnStateTreeChanged()
+{
+	return OnStateTreeChangedDelegate;
+}
+
+TSharedPtr<IDetailsView> FAnimNextStateTreeEditorHost::GetAssetDetailsView()
+{
+	if (TSharedPtr<UE::Workspace::IWorkspaceEditor> SharedEditor = WeakWorkspaceEditor.Pin())
+	{
+		return SharedEditor->GetDetailsView();
+	}
+	
+	return nullptr;
+}
+
+TSharedPtr<IDetailsView> FAnimNextStateTreeEditorHost::GetDetailsView()
+{
+	if (TSharedPtr<UE::Workspace::IWorkspaceEditor> SharedEditor = WeakWorkspaceEditor.Pin())
+	{
+		return SharedEditor->GetDetailsView();
+	}
+	
+	return nullptr;
+}
diff --git a/Engine/Plugins/Experimental/AnimNextStateTree/Source/AnimNextStateTreeEditor/Private/AnimNextStateTreeEditorHost.h b/Engine/Plugins/Experimental/AnimNextStateTree/Source/AnimNextStateTreeEditor/Private/AnimNextStateTreeEditorHost.h
new file mode 100644
--- /dev/null
+++ b/Engine/Plugins/Experimental/AnimNextStateTree/Source/AnimNextStateTreeEditor/Private/AnimNextStateTreeEditorHost.h
@@ -0,0 +1,31 @@
+// Copyright Epic Games, Inc. All Rights Reserved.
+
+#pragma once
+
+#include "IStateTreeEditorHost.h"
+
+namespace UE::Workspace
+{
+	class IWorkspaceEditor;
+}
+
+class FAnimNextStateTreeEditorHost : public IStateTreeEditorHost
+{
+public:
+	void Init(const TWeakPtr<UE::Workspace::IWorkspaceEditor>& InWeakWorkspaceEditor);
+
+	// IStateTreeEditorHost overrides
+	virtual UStateTree* GetStateTree() const override;
+	virtual FName GetCompilerLogName() const override;
+	virtual FName GetCompilerTabName() const override;
+	
+	virtual FSimpleMulticastDelegate& OnStateTreeChanged() override;
+	virtual TSharedPtr<IDetailsView> GetAssetDetailsView() override;
+	virtual TSharedPtr<IDetailsView> GetDetailsView() override;
+
+protected:
+	void OnWorkspaceFocussedAssetChanged(TObjectPtr<UObject> InObject) const;
+
+	TWeakPtr<UE::Workspace::IWorkspaceEditor> WeakWorkspaceEditor;	
+	FSimpleMulticastDelegate OnStateTreeChangedDelegate;
+};
\ No newline at end of file
diff --git a/Engine/Plugins/Experimental/AnimNextStateTree/Source/AnimNextStateTreeEditor/Private/AnimNextStateTreeEditorModule.cpp b/Engine/Plugins/Experimental/AnimNextStateTree/Source/AnimNextStateTreeEditor/Private/AnimNextStateTreeEditorModule.cpp
--- a/Engine/Plugins/Experimental/AnimNextStateTree/Source/AnimNextStateTreeEditor/Private/AnimNextStateTreeEditorModule.cpp
+++ b/Engine/Plugins/Experimental/AnimNextStateTree/Source/AnimNextStateTreeEditor/Private/AnimNextStateTreeEditorModule.cpp
@@ -2,19 +2,263 @@
 
 #include "AnimNextStateTreeEditorModule.h"
 
+#include "AnimNextStateTree.h"
+#include "ContextObjectStore.h"
+#include "EditorModeManager.h"
+#include "StateTree.h"
+#include "AnimNextStateTreeWorkspaceExports.h"
+#include "IAnimNextEditorModule.h"
+#include "IWorkspaceEditor.h"
+#include "IWorkspaceEditorModule.h"
+#include "StateTreeEditorMode.h"
+#include "WorkspaceItemMenuContext.h"
+#include "Toolkits/AssetEditorToolkitMenuContext.h"
+#include "EdModeInteractiveToolsContext.h"
+#include "AnimNextStateTreeEditorHost.h"
+#include "StateTreeEditingSubsystem.h"
+#include "StateTreeEditorStyle.h"
+#include "Framework/Docking/LayoutExtender.h"
+#include "Toolkits/AssetEditorModeUILayer.h"
+
 #define LOCTEXT_NAMESPACE "AnimNextStateTreeEditorModule"
 
 namespace UE::AnimNext::StateTree
 {
-
 void FAnimNextStateTreeEditorModule::StartupModule()
 {
+	// Register StateTree as supported asset in AnimNext workspaces
+	Editor::IAnimNextEditorModule& AnimNextEditorModule = FModuleManager::Get().LoadModuleChecked<Editor::IAnimNextEditorModule>("AnimNextEditor");	
+	AnimNextEditorModule.AddWorkspaceSupportedAssetClass(UAnimNextStateTree::StaticClass()->GetClassPathName());
+
+	// Extend Workspace Editor layout to deal with StateTreeEditorMode tabs
+	Workspace::IWorkspaceEditorModule& WorkspaceEditorModule = FModuleManager::Get().LoadModuleChecked<Workspace::IWorkspaceEditorModule>("WorkspaceEditor");
+	WorkspaceEditorModule.OnExtendTabs().AddLambda([](FLayoutExtender& InLayoutExtender, TSharedPtr<UE::Workspace::IWorkspaceEditor> InEditorPtr)
+	{
+		FTabManager::FTab TreeOutlinerTab(FTabId(UAssetEditorUISubsystem::TopLeftTabID), ETabState::ClosedTab);
+		InLayoutExtender.ExtendLayout(FTabId(Workspace::WorkspaceTabs::TopLeftDocumentArea), ELayoutExtensionPosition::After, TreeOutlinerTab);
+
+		FTabManager::FTab StatisticsTab(FTabId(UAssetEditorUISubsystem::BottomRightTabID), ETabState::ClosedTab);
+		InLayoutExtender.ExtendLayout(FTabId(Workspace::WorkspaceTabs::BottomMiddleDocumentArea), ELayoutExtensionPosition::After, StatisticsTab);
+
+		FTabManager::FTab DebuggerTab(FTabId(UAssetEditorUISubsystem::TopRightTabID), ETabState::ClosedTab);
+		InLayoutExtender.ExtendLayout(FTabId(Workspace::WorkspaceTabs::BottomMiddleDocumentArea), ELayoutExtensionPosition::After, DebuggerTab);		
+	});
+	
+	// --- AnimNextStateTree ---
+	Workspace::FObjectDocumentArgs StateTreeDocumentArgs(
+		Workspace::FOnMakeDocumentWidget::CreateLambda([](const Workspace::FWorkspaceEditorContext& InContext)-> TSharedRef<SWidget>
+		{
+			UAnimNextStateTree* AnimNextStateTree = CastChecked<UAnimNextStateTree>(InContext.Object);
+			UStateTree* StateTree = AnimNextStateTree->StateTree;
+
+			TWeakPtr<Workspace::IWorkspaceEditor> WeakWorkspaceEditor = InContext.WorkspaceEditor;
+			UContextObjectStore* ContextStore = InContext.WorkspaceEditor->GetEditorModeManager().GetInteractiveToolsContext()->ContextObjectStore;
+			UStateTreeEditorContext* StateTreeEditorContext = ContextStore->FindContext<UStateTreeEditorContext>();
+			if (!StateTreeEditorContext)
+			{
+				StateTreeEditorContext = NewObject<UStateTreeEditorContext>();
+				TSharedPtr<FAnimNextStateTreeEditorHost> Host = MakeShared<FAnimNextStateTreeEditorHost>();
+				Host->Init(WeakWorkspaceEditor);
+				StateTreeEditorContext->EditorHostInterface = Host;
+				ContextStore->AddContextObject(StateTreeEditorContext);
+			}
+
+			if (UStateTreeEditingSubsystem* StateTreeEditingSubsystem = GEditor->GetEditorSubsystem<UStateTreeEditingSubsystem>())
+			{
+				TSharedRef<FStateTreeViewModel> StateTreeViewModel = StateTreeEditingSubsystem->FindOrAddViewModel(StateTree);
+				TSharedRef<SWidget> StateTreeViewWidget = StateTreeEditingSubsystem->GetStateTreeView(StateTreeViewModel, InContext.WorkspaceEditor->GetToolkitCommands());
+
+				TWeakPtr<SWidget> WeakStateTreeViewWidget = StateTreeViewWidget;
+				TWeakPtr<FStateTreeViewModel> WeakViewModel = StateTreeViewModel;
+				StateTreeViewModel->GetOnSelectionChanged().AddSPLambda(&StateTreeViewWidget.Get(), [WeakStateTreeViewWidget, WeakWorkspaceEditor, WeakViewModel](const TArray<TWeakObjectPtr<UStateTreeState>>& SelectedStates)
+				{
+					if (TSharedPtr<Workspace::IWorkspaceEditor> SharedWorkspaceEditor = WeakWorkspaceEditor.Pin())
+					{
+						TArray<UObject*> Selected;
+						for (const TWeakObjectPtr<UStateTreeState>& WeakState : SelectedStates)
+						{
+							if (UStateTreeState* State = WeakState.Get())
+							{
+								Selected.Add(State);
+							}
+						}
+
+						SharedWorkspaceEditor->SetGlobalSelection(WeakStateTreeViewWidget, UE::Workspace::FOnClearGlobalSelection::CreateLambda([WeakViewModel](){ if (const TSharedPtr<FStateTreeViewModel> SharedViewModel = WeakViewModel.Pin()) { SharedViewModel->ClearSelection(); } }));						
+						SharedWorkspaceEditor->SetDetailsObjects(Selected);
+					}					
+				});
+
+				return SNew(SVerticalBox)				
+				+SVerticalBox::Slot()
+				.FillHeight(1.f)
+				[
+					StateTreeViewWidget
+				];
+			}
+						
+			return SNullWidget::NullWidget;
+		}
+	), Workspace::WorkspaceTabs::TopMiddleDocumentArea);
+	
+	StateTreeDocumentArgs.OnGetTabName = Workspace::FOnGetTabName::CreateLambda([](const Workspace::FWorkspaceEditorContext& InContext)
+	{
+		UAnimNextStateTree* AnimNextStateTree = CastChecked<UAnimNextStateTree>(InContext.Object);
+		return FText::FromName(AnimNextStateTree->GetFName());
+	});
+
+	StateTreeDocumentArgs.DocumentEditorMode = UStateTreeEditorMode::EM_StateTree;
+
+	StateTreeDocumentArgs.OnGetDocumentBreadcrumbTrail = Workspace::FOnGetDocumentBreadcrumbTrail::CreateLambda([](const Workspace::FWorkspaceEditorContext& InContext, TArray<TSharedPtr<Workspace::FWorkspaceBreadcrumb>>& OutBreadcrumbs)
+	{
+		if (const UAnimNextStateTree* AnimNextStateTree = CastChecked<UAnimNextStateTree>(InContext.Object))
+		{
+			const TSharedPtr<Workspace::FWorkspaceBreadcrumb>& GraphCrumb = OutBreadcrumbs.Add_GetRef(MakeShared<Workspace::FWorkspaceBreadcrumb>());
+			GraphCrumb->OnGetLabel = Workspace::FWorkspaceBreadcrumb::FOnGetBreadcrumbLabel::CreateLambda([StateTreeName = AnimNextStateTree->GetFName()]{ return FText::FromName(StateTreeName); });
+
+			GraphCrumb->CanSave = Workspace::FWorkspaceBreadcrumb::FCanSaveBreadcrumb::CreateLambda(
+				[AnimNextStateTree]
+				{
+					return AnimNextStateTree->GetPackage()->IsDirty();
+				}
+			);
+		}
+	});
+
+	StateTreeDocumentArgs.OnGetTabIcon = Workspace::FOnGetTabIcon::CreateLambda([](const Workspace::FWorkspaceEditorContext& InContext)
+	{
+		return FAppStyle::GetBrush(TEXT("ClassIcon.Default"));
+	});
+	
+	WorkspaceEditorModule.RegisterObjectDocumentType(UAnimNextStateTree::StaticClass()->GetClassPathName(), StateTreeDocumentArgs);
+
+	class FStateTreeAssetOutlinerItemDetails : public Workspace::IWorkspaceOutlinerItemDetails
+	{
+	public:
+		virtual ~FStateTreeAssetOutlinerItemDetails() override = default;
+		virtual const FSlateBrush* GetItemIcon(const FWorkspaceOutlinerItemExport& Export) const override
+		{
+			if (Export.GetData().GetScriptStruct() == FAnimNextStateTreeOutlinerData::StaticStruct())
+			{
+				return FAppStyle::GetBrush(TEXT("ClassIcon.Default"));
+			}
+			else if (Export.GetData().GetScriptStruct() == FAnimNextStateTreeStateOutlinerData::StaticStruct())
+			{
+				const FAnimNextStateTreeStateOutlinerData& Data = Export.GetData().Get<FAnimNextStateTreeStateOutlinerData>();
+				return FStateTreeEditorStyle::GetBrushForSelectionBehaviorType(Data.SelectionBehavior, !Data.bIsLeafState, Data.Type);
+			}
+			
+			return nullptr;
+		}
+
+		virtual FSlateColor GetItemColor(const FWorkspaceOutlinerItemExport& Export) const override
+		{
+			if (Export.GetData().GetScriptStruct() == FAnimNextStateTreeStateOutlinerData::StaticStruct())
+			{
+				const FAnimNextStateTreeStateOutlinerData& Data = Export.GetData().Get<FAnimNextStateTreeStateOutlinerData>();
+				return Data.Color;
+			}
+			
+			return FSlateColor::UseForeground();
+		}
+		
+		virtual bool HandleSelected(const FToolMenuContext& ToolMenuContext) const override
+		{
+			const UWorkspaceItemMenuContext* WorkspaceItemContext = ToolMenuContext.FindContext<UWorkspaceItemMenuContext>();
+			const UAssetEditorToolkitMenuContext* AssetEditorContext = ToolMenuContext.FindContext<UAssetEditorToolkitMenuContext>(); 
+			if (WorkspaceItemContext && AssetEditorContext)
+			{				
+				if(const TSharedPtr<UE::Workspace::IWorkspaceEditor> WorkspaceEditor = StaticCastSharedPtr<UE::Workspace::IWorkspaceEditor>(AssetEditorContext->Toolkit.Pin()))
+				{
+					const uint32 NumSelectedExports = WorkspaceItemContext->SelectedExports.Num();
+					if (NumSelectedExports >= 1)
+					{
+						// Selecting AnimStateTree asset itself
+						if (NumSelectedExports == 1 && WorkspaceItemContext->SelectedExports[0].GetData().GetScriptStruct() == FAnimNextStateTreeOutlinerData::StaticStruct())
+						{
+							const FWorkspaceOutlinerItemExport& SelectedExport = WorkspaceItemContext->SelectedExports[0];
+							if (UAnimNextStateTree* LoadedStateTree = Cast<UAnimNextStateTree>(SelectedExport.GetAssetPath().ResolveObject()))
+							{
+								WorkspaceEditor->SetDetailsObjects({ LoadedStateTree->StateTree->EditorData } );
+								return true;
+							}
+						}
+
+						if(UStateTreeEditingSubsystem* EditingSubsystem = GEditor->GetEditorSubsystem<UStateTreeEditingSubsystem>())
+						{
+							UAnimNextStateTree* SelectionStateTree = [WorkspaceItemContext]() -> UAnimNextStateTree*
+							{
+								UAnimNextStateTree* FoundTree = nullptr;
+
+								for (const FWorkspaceOutlinerItemExport& SelectedExport : WorkspaceItemContext->SelectedExports)
+								{
+									if (UAnimNextStateTree* LoadedStateTree = Cast<UAnimNextStateTree>(SelectedExport.GetAssetPath().ResolveObject()))
+									{
+										if (FoundTree == nullptr)
+										{
+											FoundTree = LoadedStateTree;
+										}
+										if (FoundTree != LoadedStateTree)
+										{
+											FoundTree = nullptr;
+											break;
+										}
+									}
+								}
+
+								return FoundTree;
+							}();
+
+							if (SelectionStateTree)
+							{
+								TSharedPtr<FStateTreeViewModel> ViewModel = EditingSubsystem->FindOrAddViewModel(SelectionStateTree->StateTree.Get());
+								check(ViewModel.IsValid());
+								
+								TArray<TWeakObjectPtr<UStateTreeState>> ToBeSelectedStates;
+								for (const FWorkspaceOutlinerItemExport& SelectedExport : WorkspaceItemContext->SelectedExports)
+								{
+									if (SelectedExport.GetData().GetScriptStruct() == FAnimNextStateTreeStateOutlinerData::StaticStruct())
+									{
+										if (UAnimNextStateTree* LoadedStateTree = Cast<UAnimNextStateTree>(SelectedExport.GetAssetPath().ResolveObject()))
+										{
+											const FAnimNextStateTreeStateOutlinerData& StateData = SelectedExport.GetData().Get<FAnimNextStateTreeStateOutlinerData>();		
+											ensure(LoadedStateTree == SelectionStateTree);
+											if (UStateTreeState* State = ViewModel->GetMutableStateByID(StateData.StateId))
+											{
+												ToBeSelectedStates.Add(State);
+											}
+										}
+									}
+								}
+
+								ViewModel->SetSelection(ToBeSelectedStates);
+								return true;
+							}
+						}						
+					}
+				}			
+			}
+			return false;
+		}
+	};
 	
+	TSharedPtr<FStateTreeAssetOutlinerItemDetails> StateItemDetails = MakeShareable<FStateTreeAssetOutlinerItemDetails>(new FStateTreeAssetOutlinerItemDetails());
+	WorkspaceEditorModule.RegisterWorkspaceItemDetails(Workspace::FOutlinerItemDetailsId(FAnimNextStateTreeOutlinerData::StaticStruct()->GetFName()), StaticCastSharedPtr<UE::Workspace::IWorkspaceOutlinerItemDetails>(StateItemDetails));	
+	WorkspaceEditorModule.RegisterWorkspaceItemDetails(Workspace::FOutlinerItemDetailsId(FAnimNextStateTreeStateOutlinerData::StaticStruct()->GetFName()), StaticCastSharedPtr<UE::Workspace::IWorkspaceOutlinerItemDetails>(StateItemDetails));
 }
 
 void FAnimNextStateTreeEditorModule::ShutdownModule()
 {
+	if(Workspace::IWorkspaceEditorModule* WorkspaceEditorModule = FModuleManager::Get().GetModulePtr<Workspace::IWorkspaceEditorModule>("WorkspaceEditor"))
+	{
+		WorkspaceEditorModule->UnregisterObjectDocumentType(UAnimNextStateTree::StaticClass()->GetClassPathName());
+		WorkspaceEditorModule->UnregisterWorkspaceItemDetails(Workspace::FOutlinerItemDetailsId(FAnimNextStateTreeOutlinerData::StaticStruct()->GetFName()));
+		WorkspaceEditorModule->UnregisterWorkspaceItemDetails(Workspace::FOutlinerItemDetailsId(FAnimNextStateTreeStateOutlinerData::StaticStruct()->GetFName()));
+	}
 	
+	if( Editor::IAnimNextEditorModule* AnimNextEditorModule = FModuleManager::Get().GetModulePtr<Editor::IAnimNextEditorModule>("AnimNextEditor"))
+	{
+		AnimNextEditorModule->RemoveWorkspaceSupportedAssetClass(UStateTree::StaticClass()->GetClassPathName());
+	}
 }
 
 IMPLEMENT_MODULE(FAnimNextStateTreeEditorModule, AnimNextStateTreeEditor)
diff --git a/Engine/Plugins/Experimental/AnimNextStateTree/Source/AnimNextStateTreeEditor/Private/AnimNextStateTreeFactory.cpp b/Engine/Plugins/Experimental/AnimNextStateTree/Source/AnimNextStateTreeEditor/Private/AnimNextStateTreeFactory.cpp
new file mode 100644
--- /dev/null
+++ b/Engine/Plugins/Experimental/AnimNextStateTree/Source/AnimNextStateTreeEditor/Private/AnimNextStateTreeFactory.cpp
@@ -0,0 +1,50 @@
+// Copyright Epic Games, Inc. All Rights Reserved.
+
+#include "AnimNextStateTreeFactory.h"
+#include "AnimNextStateTree.h"
+#include "AnimNextStateTreeSchema.h"
+#include "AnimNextStateTree_EditorData.h"
+#include "StateTreeFactory.h"
+#include "AnimStateTreeTrait.h"
+
+UAnimNextStateTreeFactory::UAnimNextStateTreeFactory()
+{
+	bCreateNew = true;
+	bEditAfterNew = true;
+	SupportedClass = UAnimNextStateTree::StaticClass();
+}
+
+bool UAnimNextStateTreeFactory::ConfigureProperties()
+{
+	return true;
+}
+
+UObject* UAnimNextStateTreeFactory::FactoryCreateNew(UClass* Class, UObject* InParent, FName Name, EObjectFlags Flags, UObject* Context, FFeedbackContext* Warn, FName CallingContext)
+{
+	EObjectFlags FlagsToUse = Flags | RF_Public | RF_Standalone | RF_Transactional | RF_LoadCompleted;
+	if(InParent == GetTransientPackage())
+	{
+		FlagsToUse &= ~RF_Standalone;
+	}
+
+	UAnimNextStateTree* NewStateTree = NewObject<UAnimNextStateTree>(InParent, Class, Name, FlagsToUse);
+
+	// Create internal state tree	
+	UStateTreeFactory* StateTreeFactory = NewObject<UStateTreeFactory>(GetTransientPackage());
+	StateTreeFactory->SetSchemaClass(UStateTreeAnimNextSchema::StaticClass());
+	UStateTree* StateTree = CastChecked<UStateTree>(StateTreeFactory->FactoryCreateNew(UStateTree::StaticClass(), NewStateTree, TEXT("StateTree"), RF_Transactional, nullptr, nullptr));
+
+	NewStateTree->StateTree = StateTree;
+
+	// Create internal editor data
+	UAnimNextStateTree_EditorData* EditorData = NewObject<UAnimNextStateTree_EditorData>(NewStateTree, TEXT("EditorData"), RF_Transactional);
+	NewStateTree->EditorData = EditorData;
+	EditorData->Initialize(/*bRecompileVM*/false);
+	
+	// Compile the initial skeleton
+	EditorData->RecompileVM();
+	
+	check(!EditorData->bErrorsDuringCompilation);
+
+	return NewStateTree;
+}
\ No newline at end of file
diff --git a/Engine/Plugins/Experimental/AnimNextStateTree/Source/AnimNextStateTreeUncookedOnly/AnimNextStateTreeUncookedOnly.Build.cs b/Engine/Plugins/Experimental/AnimNextStateTree/Source/AnimNextStateTreeUncookedOnly/AnimNextStateTreeUncookedOnly.Build.cs
--- a/Engine/Plugins/Experimental/AnimNextStateTree/Source/AnimNextStateTreeUncookedOnly/AnimNextStateTreeUncookedOnly.Build.cs
+++ b/Engine/Plugins/Experimental/AnimNextStateTree/Source/AnimNextStateTreeUncookedOnly/AnimNextStateTreeUncookedOnly.Build.cs
@@ -6,7 +6,13 @@ namespace UnrealBuildTool.Rules
 	{
 		public AnimNextStateTreeUncookedOnly(ReadOnlyTargetRules Target) : base(Target)
 		{
-			PrivateDependencyModuleNames.AddRange(new string[] { "StateTreeEditorModule" });
+			PrivateDependencyModuleNames.AddRange(new string[] 
+			{
+				"StateTreeEditorModule",
+				"RigVM",
+				"RigVMDeveloper"
+			});
+			
 			PublicDependencyModuleNames.AddRange(
 				new string[]
 				{
@@ -14,7 +20,12 @@ namespace UnrealBuildTool.Rules
 					"CoreUObject",
 					"StateTreeModule",
 					"Engine",
-					"WorkspaceEditor"
+					"WorkspaceEditor",
+					"SlateCore",
+					"AnimNext",
+					"AnimNextUncookedOnly",
+					"AnimNextStateTree",
+					"AnimNextAnimGraph"
 				}
 			);
 		}
diff --git a/Engine/Plugins/Experimental/AnimNextStateTree/Source/AnimNextStateTreeUncookedOnly/Internal/AnimNextStateTreeWorkspaceAssetUserData.h b/Engine/Plugins/Experimental/AnimNextStateTree/Source/AnimNextStateTreeUncookedOnly/Internal/AnimNextStateTreeWorkspaceAssetUserData.h
new file mode 100644
--- /dev/null
+++ b/Engine/Plugins/Experimental/AnimNextStateTree/Source/AnimNextStateTreeUncookedOnly/Internal/AnimNextStateTreeWorkspaceAssetUserData.h
@@ -0,0 +1,15 @@
+// Copyright Epic Games, Inc. All Rights Reserved.
+
+#pragma once
+
+#include "Engine/AssetUserData.h"
+
+#include "AnimNextStateTreeWorkspaceAssetUserData.generated.h"
+
+UCLASS()
+class UAnimNextStateTreeWorkspaceAssetUserData : public UAssetUserData
+{
+	GENERATED_BODY()
+
+	virtual void GetAssetRegistryTags(FAssetRegistryTagsContext Context) const override;
+};
diff --git a/Engine/Plugins/Experimental/AnimNextStateTree/Source/AnimNextStateTreeUncookedOnly/Internal/AnimNextStateTree_EditorData.cpp b/Engine/Plugins/Experimental/AnimNextStateTree/Source/AnimNextStateTreeUncookedOnly/Internal/AnimNextStateTree_EditorData.cpp
new file mode 100644
--- /dev/null
+++ b/Engine/Plugins/Experimental/AnimNextStateTree/Source/AnimNextStateTreeUncookedOnly/Internal/AnimNextStateTree_EditorData.cpp
@@ -0,0 +1,127 @@
+// Copyright Epic Games, Inc. All Rights Reserved.
+
+#include "AnimNextStateTree_EditorData.h"
+
+#include "AnimNextStateTree.h"
+#include "AnimNextStateTreeWorkspaceAssetUserData.h"
+#include "StateTree.h"
+#include "StateTreeEditorData.h"
+#include "Entries/AnimNextVariableEntry.h"
+#include "Entries/AnimNextAnimationGraphEntry.h"
+#include "Graph/AnimNextAnimationGraphSchema.h"
+#include "Graph/RigUnit_AnimNextGraphRoot.h"
+#include "Graph/RigDecorator_AnimNextCppTrait.h"
+#include "Graph/RigUnit_AnimNextTraitStack.h"
+#include "TraitCore/TraitRegistry.h"
+#include "TraitCore/Trait.h"
+#include "AnimStateTreeTrait.h"
+#include "Traits/BlendStackTrait.h"
+
+TSubclassOf<UAssetUserData> UAnimNextStateTree_EditorData::GetAssetUserDataClass() const
+{
+	return UAnimNextStateTreeWorkspaceAssetUserData::StaticClass();
+}
+
+void UAnimNextStateTree_EditorData::RecompileVM()
+{
+	UAnimNextAnimationGraph_EditorData::RecompileVM();
+
+	UAnimNextStateTree* AnimationStateTree = UE::AnimNext::UncookedOnly::FUtils::GetAsset<UAnimNextStateTree>(this);
+
+	UStateTree* InnerStateTree = AnimationStateTree->StateTree;
+
+	UStateTreeEditorData* InnerEditorData = Cast<UStateTreeEditorData>(InnerStateTree->EditorData);
+	InnerEditorData->RootParameters.ResetParametersAndOverrides();
+	InnerEditorData->RootParameters.Parameters = AnimationStateTree->VariableDefaults;
+}
+
+TConstArrayView<TSubclassOf<UAnimNextRigVMAssetEntry>> UAnimNextStateTree_EditorData::GetEntryClasses() const
+{
+	static const TSubclassOf<UAnimNextRigVMAssetEntry> Classes[] =
+{
+		UAnimNextVariableEntry::StaticClass(),
+		UAnimNextAnimationGraphEntry::StaticClass(),
+	};
+
+	return Classes;
+}
+
+void UAnimNextStateTree_EditorData::GetProgrammaticGraphs(const FRigVMCompileSettings& InSettings, TArray<URigVMGraph*>& OutGraphs)
+{
+	if (UAnimNextStateTree* AnimStateTree = UE::AnimNext::UncookedOnly::FUtils::GetAsset<UAnimNextStateTree>(this))
+	{
+		URigVMGraph* Graph = NewObject<URigVMGraph>(this, NAME_None, RF_Transient);
+		Graph->SetSchemaClass(UAnimNextAnimationGraphSchema::StaticClass());
+
+		UAnimNextController* Controller = CastChecked<UAnimNextController>(RigVMClient.GetOrCreateController(Graph));
+		UE::AnimNext::UncookedOnly::FUtils::SetupAnimGraph(FRigUnit_AnimNextGraphRoot::DefaultEntryPoint, Controller);
+
+		if (Controller->GetGraph()->GetNodes().Num() != 1)
+		{
+			InSettings.ReportError(TEXT("Expected singular FRigUnit_AnimNextGraphRoot node"));
+			return;
+		}
+		
+		URigVMNode* EntryNode = Controller->GetGraph()->GetNodes()[0];			
+		URigVMPin* BeginExecutePin = EntryNode->FindPin(GET_MEMBER_NAME_STRING_CHECKED(FRigUnit_AnimNextGraphRoot, Result));
+
+		if (BeginExecutePin == nullptr)
+		{
+			InSettings.ReportError(TEXT("Failed to retrieve Result pin from FRigUnit_AnimNextGraphRoot node"));
+			return;
+		}
+
+		URigVMUnitNode* TraitStackNode = Controller->AddUnitNode(FRigUnit_AnimNextTraitStack::StaticStruct(), FRigVMStruct::ExecuteName, FVector2D(-800.0f, 0.0f), FString(), false);
+
+		if (TraitStackNode == nullptr)
+		{
+			InSettings.ReportError(TEXT("Failed to spawn FRigUnit_AnimNextTraitStack node"));
+			return;
+		}
+
+		const FName BlendStackTraitName = Controller->AddTraitByName(TraitStackNode->GetFName(), *UE::AnimNext::FTraitRegistry::Get().Find(UE::AnimNext::FBlendStackCoreTrait::TraitUID)->GetTraitName(), INDEX_NONE);
+		
+		if (BlendStackTraitName == NAME_None)
+		{
+			InSettings.ReportError(TEXT("Failed to add BlendStack trait to node"));
+			return;
+		}
+		
+		const FName StateTreeTraitName = Controller->AddTraitByName(TraitStackNode->GetFName(), *UE::AnimNext::FTraitRegistry::Get().Find(UE::AnimNext::FStateTreeTrait::TraitUID)->GetTraitName(), INDEX_NONE);
+
+		if (StateTreeTraitName == NAME_None)
+		{
+			InSettings.ReportError(TEXT("Failed to add StateTree trait to node"));
+			return;
+		}
+
+		URigVMPin* StateTreeReferencePin = TraitStackNode->FindTrait(StateTreeTraitName, GET_MEMBER_NAME_STRING_CHECKED(FAnimNextStateTreeTraitSharedData, StateTreeReference));
+		if (StateTreeReferencePin == nullptr)
+		{
+			InSettings.ReportError(TEXT("Failed to retrieve StateTreeReference pin"));
+			return;
+		}
+		
+		FStateTreeReference Ref;
+		Ref.SetStateTree(AnimStateTree->StateTree);
+
+		FString PinValue;
+		FStateTreeReference::StaticStruct()->ExportText(PinValue, &Ref, nullptr, nullptr, 0, nullptr);
+		Controller->SetPinDefaultValue(StateTreeReferencePin->GetPinPath(), PinValue);
+		
+		URigVMPin* TraitResult = TraitStackNode->FindPin(GET_MEMBER_NAME_STRING_CHECKED(FRigUnit_AnimNextTraitStack, Result));
+		if (TraitResult == nullptr)
+		{
+			InSettings.ReportError(TEXT("Failed to retrieve Result pin"));
+			return;
+		}
+
+		if (!Controller->AddLink(TraitResult, BeginExecutePin, false))
+		{
+			InSettings.ReportError(TEXT("Failed to link TraitStack and Graph Output pins"));
+			return;
+		}
+		
+		OutGraphs.Add(Graph);	
+	}
+}
diff --git a/Engine/Plugins/Experimental/AnimNextStateTree/Source/AnimNextStateTreeUncookedOnly/Internal/AnimNextStateTree_EditorData.h b/Engine/Plugins/Experimental/AnimNextStateTree/Source/AnimNextStateTreeUncookedOnly/Internal/AnimNextStateTree_EditorData.h
new file mode 100644
--- /dev/null
+++ b/Engine/Plugins/Experimental/AnimNextStateTree/Source/AnimNextStateTreeUncookedOnly/Internal/AnimNextStateTree_EditorData.h
@@ -0,0 +1,26 @@
+// Copyright Epic Games, Inc. All Rights Reserved.
+
+#pragma once
+
+#include "Graph/AnimNextAnimationGraph_EditorData.h"
+
+#include "AnimNextStateTree_EditorData.generated.h"
+
+UCLASS()
+class ANIMNEXTSTATETREEUNCOOKEDONLY_API UAnimNextStateTree_EditorData : public UAnimNextAnimationGraph_EditorData
+{
+	GENERATED_BODY()
+	
+protected:
+	virtual TSubclassOf<UAssetUserData> GetAssetUserDataClass() const override;
+
+	friend class UAnimNextStateTreeFactory;
+	// IRigVMClientHost interface
+	virtual void RecompileVM() override;
+	
+	// UAnimNextRigVMAssetEditorData interface
+	virtual TConstArrayView<TSubclassOf<UAnimNextRigVMAssetEntry>> GetEntryClasses() const override;
+
+	// Allows this asset to generate graphs to be injected at compilation time
+	virtual void GetProgrammaticGraphs(const FRigVMCompileSettings& InSettings, TArray<URigVMGraph*>& OutGraphs) override;
+};
\ No newline at end of file
diff --git a/Engine/Plugins/Experimental/AnimNextStateTree/Source/AnimNextStateTreeUncookedOnly/Private/AnimNextStateTreeWorkspaceAssetUserData.cpp b/Engine/Plugins/Experimental/AnimNextStateTree/Source/AnimNextStateTreeUncookedOnly/Private/AnimNextStateTreeWorkspaceAssetUserData.cpp
new file mode 100644
--- /dev/null
+++ b/Engine/Plugins/Experimental/AnimNextStateTree/Source/AnimNextStateTreeUncookedOnly/Private/AnimNextStateTreeWorkspaceAssetUserData.cpp
@@ -0,0 +1,59 @@
+// Copyright Epic Games, Inc. All Rights Reserved.
+
+#include "AnimNextStateTreeWorkspaceAssetUserData.h"
+
+#include "AnimNextStateTree.h"
+#include "StateTreeEditorData.h"
+#include "AnimNextStateTreeWorkspaceExports.h"
+#include "UObject/AssetRegistryTagsContext.h"
+
+void UAnimNextStateTreeWorkspaceAssetUserData::GetAssetRegistryTags(FAssetRegistryTagsContext Context) const
+{
+	FWorkspaceOutlinerItemExports OutlinerExports;
+
+	if (UAnimNextStateTree* AnimStateTree = CastChecked<UAnimNextStateTree>(GetOuter()))
+	{		
+		const FName AnimStateTreeIdentifier = AnimStateTree->GetFName();
+		
+		// StateTree asset export
+		UStateTree* StateTree = AnimStateTree->StateTree;
+		FWorkspaceOutlinerItemExport& RootAssetExport = OutlinerExports.Exports.Add_GetRef(FWorkspaceOutlinerItemExport(AnimStateTreeIdentifier, GetOuter()));			
+		RootAssetExport.GetData().InitializeAsScriptStruct(FAnimNextStateTreeOutlinerData::StaticStruct());
+		RootAssetExport.GetData().GetMutable<FAnimNextStateTreeOutlinerData>().Asset = AnimStateTree;
+
+		if (const UStateTreeEditorData* EditorData = Cast<UStateTreeEditorData>(StateTree->EditorData)) 
+		{
+			// Export each state as individual item as well (parented according to hierarchy)
+			TMap<UStateTreeState*, FWorkspaceOutlinerItemExport> ParentExports;
+			EditorData->VisitHierarchy([&OutlinerExports, RootAssetExport, &ParentExports, EditorData](UStateTreeState& State, const UStateTreeState* ParentState) -> EStateTreeVisitor
+			{
+				FWorkspaceOutlinerItemExport& StateExport = OutlinerExports.Exports.Add_GetRef(
+				ParentState == nullptr ? FWorkspaceOutlinerItemExport(State.Name, RootAssetExport)
+					: FWorkspaceOutlinerItemExport(State.Name, ParentExports.FindChecked(ParentState))
+				);
+
+				ParentExports.Add(&State, StateExport);
+
+				StateExport.GetData().InitializeAsScriptStruct(FAnimNextStateTreeStateOutlinerData::StaticStruct());
+				FAnimNextStateTreeStateOutlinerData& AssetData = StateExport.GetData().GetMutable<FAnimNextStateTreeStateOutlinerData>();
+				AssetData.StateName = State.Name;
+				AssetData.StateId = State.ID;
+				AssetData.bIsLeafState = State.Children.IsEmpty(); 
+				AssetData.Type = State.Type;
+				AssetData.SelectionBehavior = State.SelectionBehavior;
+
+				const FStateTreeEditorColor* StateTreeEditorColor = EditorData->FindColor(State.ColorRef);
+
+				AssetData.Color = StateTreeEditorColor ? FSlateColor(StateTreeEditorColor->Color) : FSlateColor::UseForeground();
+				
+				return EStateTreeVisitor::Continue;
+			});
+		}
+
+		{
+			FString TagValue;
+			FWorkspaceOutlinerItemExports::StaticStruct()->ExportText(TagValue, &OutlinerExports, nullptr, nullptr, PPF_None, nullptr);
+			Context.AddTag(UObject::FAssetRegistryTag(UE::Workspace::ExportsWorkspaceItemsRegistryTag, TagValue, UObject::FAssetRegistryTag::TT_Hidden));
+		}
+	}
+}
diff --git a/Engine/Plugins/Experimental/AnimNextStateTree/Source/AnimNextStateTreeUncookedOnly/Public/AnimNextStateTreeWorkspaceExports.h b/Engine/Plugins/Experimental/AnimNextStateTree/Source/AnimNextStateTreeUncookedOnly/Public/AnimNextStateTreeWorkspaceExports.h
new file mode 100644
--- /dev/null
+++ b/Engine/Plugins/Experimental/AnimNextStateTree/Source/AnimNextStateTreeUncookedOnly/Public/AnimNextStateTreeWorkspaceExports.h
@@ -0,0 +1,45 @@
+// Copyright Epic Games, Inc. All Rights Reserved.
+
+#pragma once
+
+#include "AnimNextAssetWorkspaceAssetUserData.h"
+#include "WorkspaceAssetRegistryInfo.h"
+#include "Engine/AssetUserData.h"
+#include "StateTreeTypes.h"
+#include "Styling/SlateColor.h"
+
+#include "AnimNextStateTreeWorkspaceExports.generated.h"
+
+USTRUCT()
+struct FAnimNextStateTreeOutlinerData : public FAnimNextRigVMAssetOutlinerData
+{
+	GENERATED_BODY()
+	
+	FAnimNextStateTreeOutlinerData() = default;
+};
+
+USTRUCT()
+struct FAnimNextStateTreeStateOutlinerData : public FWorkspaceOutlinerItemData
+{
+	GENERATED_BODY()
+	
+	FAnimNextStateTreeStateOutlinerData() = default;
+
+	UPROPERTY()
+	FName StateName = NAME_None;
+
+	UPROPERTY()
+	FGuid StateId;
+
+	UPROPERTY()
+	bool bIsLeafState;
+
+	UPROPERTY()
+	EStateTreeStateSelectionBehavior SelectionBehavior;
+
+	UPROPERTY()
+	EStateTreeStateType Type;
+
+	UPROPERTY()
+	FSlateColor Color;
+};
diff --git a/Engine/Plugins/Experimental/Workspace/Source/WorkspaceEditor/Private/Outliner/WorkspaceOutlinerMode.cpp b/Engine/Plugins/Experimental/Workspace/Source/WorkspaceEditor/Private/Outliner/WorkspaceOutlinerMode.cpp
--- a/Engine/Plugins/Experimental/Workspace/Source/WorkspaceEditor/Private/Outliner/WorkspaceOutlinerMode.cpp
+++ b/Engine/Plugins/Experimental/Workspace/Source/WorkspaceEditor/Private/Outliner/WorkspaceOutlinerMode.cpp
@@ -16,6 +16,7 @@
 #include "ToolMenus.h"
 #include "WorkspaceEditorCommands.h"
 #include "Framework/Commands/GenericCommands.h"
+#include "Algo/ForEach.h"
 
 #define LOCTEXT_NAMESPACE "FWorkspaceOutlinerMode"
 
@@ -220,12 +221,6 @@ void FWorkspaceOutlinerMode::OnItemDoubleClick(FSceneOutlinerTreeItemPtr Item)
 	OpenItems({Item} );
 }
 
-void FWorkspaceOutlinerMode::OnItemClicked(FSceneOutlinerTreeItemPtr Item)
-{
-	const FSceneOutlinerItemSelection& Selection = SceneOutliner->GetSelection();
-	HandleItemSelection(Selection);
-}
-
 FReply FWorkspaceOutlinerMode::OnKeyDown(const FKeyEvent& InKeyEvent)
 {
 	if (CommandList->ProcessCommandBindings(InKeyEvent))
@@ -240,24 +235,77 @@ void FWorkspaceOutlinerMode::HandleItemSelection(const FSceneOutlinerItemSelecti
 {
 	if (Selection.Num() > 0)
 	{
-		TArray<FSceneOutlinerTreeItemPtr> SelectedItems;
-		Selection.Get(SelectedItems);
-		
 		if (TSharedPtr<IWorkspaceEditor> SharedWorkspaceEditor = WeakWorkspaceEditor.Pin())
 		{
-			TArray<FWorkspaceOutlinerItemExport> Exports;
-			Exports.Reserve(SelectedItems.Num());
-			for(const FSceneOutlinerTreeItemPtr& SelectedItem : SelectedItems)
+			UWorkspaceItemMenuContext* MenuContext = NewObject<UWorkspaceItemMenuContext>();
+			MenuContext->SelectedExports.Reserve(Selection.Num());
+			Algo::TransformIf(SceneOutliner->GetSelectedItems(), MenuContext->SelectedExports,
+					[](const FSceneOutlinerTreeItemPtr& InItem)
+					{
+						if (FWorkspaceOutlinerTreeItem* TreeItem = InItem->CastTo<FWorkspaceOutlinerTreeItem>())
+						{
+							return TreeItem->Export.GetIdentifier() != NAME_None;
+						}
+					
+						return false;
+					},
+					[](const FSceneOutlinerTreeItemPtr& InItem)
+					{
+						return InItem->CastTo<FWorkspaceOutlinerTreeItem>()->Export;
+					});
+
+			// Singular item selected
+			if (MenuContext->SelectedExports.Num() == 1)
+			{
+				FWorkspaceOutlinerItemExport& SelectedExport = MenuContext->SelectedExports[0];
+
+				bool bHandled = false;
+
+				// Check whether item details should handle the selection
+				if (const TSharedPtr<IWorkspaceOutlinerItemDetails> SharedDetails = FWorkspaceEditorModule::GetOutlinerItemDetails(MakeOutlinerDetailsId(SelectedExport)))
+				{
+					FToolMenuContext Context(MenuContext);
+					WeakWorkspaceEditor.Pin()->InitToolMenuContext(Context);
+					bHandled = SharedDetails->HandleSelected(Context);
+				}
+				
+				if(!bHandled && SelectedExport.GetData().IsValid())
+				{
+					TSharedPtr<FStructOnScope> ExportDataView = MakeShared<FStructOnScope>(SelectedExport.GetData().GetScriptStruct(), SelectedExport.GetData().GetMutableMemory());
+					// TODO JDB handle struct selections
+					//SharedWorkspaceEditor->SetDetailsStruct(ExportDataView);
+				}
+
+				UObject* LoadedAsset = SelectedExport.GetAssetPath().ResolveObject();
+				if (!bHandled && LoadedAsset)
+				{
+					SharedWorkspaceEditor->SetDetailsObjects({LoadedAsset});
+				}
+			}
+			else
 			{
-				if(FWorkspaceOutlinerTreeItem* TreeItem = SelectedItem->CastTo<FWorkspaceOutlinerTreeItem>())
+				const UScriptStruct* DataType = MenuContext->SelectedExports[0].GetData().GetScriptStruct();
+				Algo::ForEach(MenuContext->SelectedExports, [&DataType](const FWorkspaceOutlinerItemExport& Export)
+				{
+					if (DataType != Export.GetData().GetScriptStruct())
+					{
+						DataType = nullptr;
+					}
+				});
+
+				const bool bSingleDataType = DataType != nullptr;
+				if (bSingleDataType)
 				{
-					if(TreeItem->Export.GetData().IsValid())
+					if (const TSharedPtr<IWorkspaceOutlinerItemDetails> SharedDetails = FWorkspaceEditorModule::GetOutlinerItemDetails(DataType->GetFName()))
 					{
-						Exports.Add(TreeItem->Export);
+						FToolMenuContext Context(MenuContext);
+						WeakWorkspaceEditor.Pin()->InitToolMenuContext(Context);
+						SharedDetails->HandleSelected(Context);
 					}
 				}
 			}
-			SharedWorkspaceEditor->OnOutlinerSelectionChanged().Broadcast(Exports);
+
+			SharedWorkspaceEditor->OnOutlinerSelectionChanged().Broadcast(MenuContext->SelectedExports);
 		}
 	}
 	else
@@ -271,11 +319,11 @@ void FWorkspaceOutlinerMode::HandleItemSelection(const FSceneOutlinerItemSelecti
 
 void FWorkspaceOutlinerMode::OnItemSelectionChanged(FSceneOutlinerTreeItemPtr Item, ESelectInfo::Type SelectionType, const FSceneOutlinerItemSelection& Selection)
 {
-	HandleItemSelection(Selection);
-
 	if (TSharedPtr<IWorkspaceEditor> SharedWorkspaceEditor = WeakWorkspaceEditor.Pin())
 	{
+		FWorkspaceEditorSelectionScope Scope(SharedWorkspaceEditor);		
 		SharedWorkspaceEditor->SetGlobalSelection(SceneOutliner->AsShared(), FOnClearGlobalSelection::CreateRaw(this, &FWorkspaceOutlinerMode::ResetOutlinerSelection));
+		HandleItemSelection(Selection);
 	}
 }
 
diff --git a/Engine/Plugins/Experimental/Workspace/Source/WorkspaceEditor/Private/Outliner/WorkspaceOutlinerMode.h b/Engine/Plugins/Experimental/Workspace/Source/WorkspaceEditor/Private/Outliner/WorkspaceOutlinerMode.h
--- a/Engine/Plugins/Experimental/Workspace/Source/WorkspaceEditor/Private/Outliner/WorkspaceOutlinerMode.h
+++ b/Engine/Plugins/Experimental/Workspace/Source/WorkspaceEditor/Private/Outliner/WorkspaceOutlinerMode.h
@@ -20,9 +20,7 @@ namespace UE::Workspace
 		virtual void Rebuild() override;
 		virtual TSharedPtr<SWidget> CreateContextMenu() override;		
 		virtual void OnItemDoubleClick(FSceneOutlinerTreeItemPtr Item) override;
-		virtual void OnItemClicked(FSceneOutlinerTreeItemPtr Item) override;
 		virtual FReply OnKeyDown(const FKeyEvent& InKeyEvent) override;
-		void HandleItemSelection(const FSceneOutlinerItemSelection& Selection);
 		virtual void OnItemSelectionChanged(FSceneOutlinerTreeItemPtr Item, ESelectInfo::Type SelectionType, const FSceneOutlinerItemSelection& Selection) override;
 		virtual bool CanCustomizeToolbar() const { return true; }
 		virtual ESelectionMode::Type GetSelectionMode() const { return ESelectionMode::Multi; }
@@ -34,6 +32,7 @@ namespace UE::Workspace
 		virtual TUniquePtr<ISceneOutlinerHierarchy> CreateHierarchy() override;		
 		// End ISceneOutlinerMode overrides
 
+		void HandleItemSelection(const FSceneOutlinerItemSelection& Selection);
 		void OnWorkspaceModified(UWorkspace* InWorkspace);
 		void ResetOutlinerSelection();
 
diff --git a/Engine/Plugins/Experimental/Workspace/Source/WorkspaceEditor/Private/Outliner/WorkspaceOutlinerTreeItem.cpp b/Engine/Plugins/Experimental/Workspace/Source/WorkspaceEditor/Private/Outliner/WorkspaceOutlinerTreeItem.cpp
--- a/Engine/Plugins/Experimental/Workspace/Source/WorkspaceEditor/Private/Outliner/WorkspaceOutlinerTreeItem.cpp
+++ b/Engine/Plugins/Experimental/Workspace/Source/WorkspaceEditor/Private/Outliner/WorkspaceOutlinerTreeItem.cpp
@@ -56,7 +56,18 @@ namespace UE::Workspace
 							
 							return FAppStyle::GetBrush(TEXT("ClassIcon.Default"));
 						})
-						.ColorAndOpacity(FSlateColor::UseForeground())
+						.ColorAndOpacity_Lambda([this]() -> FSlateColor
+						{
+							if (const TSharedPtr<FWorkspaceOutlinerTreeItem> SharedTreeItem = TreeItem.Pin())
+							{
+								if (SharedTreeItem->ItemDetails.IsValid())
+								{
+									return SharedTreeItem->ItemDetails->GetItemColor(SharedTreeItem->Export);
+								}
+							}
+							
+							return FSlateColor::UseForeground();
+						})
 					]
 				]
 				+ SHorizontalBox::Slot()
diff --git a/Engine/Plugins/Experimental/Workspace/Source/WorkspaceEditor/Private/WorkspaceEditor.cpp b/Engine/Plugins/Experimental/Workspace/Source/WorkspaceEditor/Private/WorkspaceEditor.cpp
--- a/Engine/Plugins/Experimental/Workspace/Source/WorkspaceEditor/Private/WorkspaceEditor.cpp
+++ b/Engine/Plugins/Experimental/Workspace/Source/WorkspaceEditor/Private/WorkspaceEditor.cpp
@@ -22,12 +22,34 @@
 #include "AssetEditorModeManager.h"
 #include "Selection.h"
 #include "EditorModeManager.h"
+#include "EdModeInteractiveToolsContext.h"
+#include "ContextObjectStore.h"
 
 #define LOCTEXT_NAMESPACE "WorkspaceEditor"
 
 namespace UE::Workspace
 {
 
+FWorkspaceEditorSelectionScope::FWorkspaceEditorSelectionScope(const TSharedPtr< IWorkspaceEditor>& InWorkspaceEditor) : WeakWorkspaceEditor(InWorkspaceEditor)
+{
+	const TSharedPtr<FWorkspaceEditor> SharedWorkspaceEditor = StaticCastSharedPtr<FWorkspaceEditor>(InWorkspaceEditor);
+	++SharedWorkspaceEditor->SelectionScopeDepth;
+}
+
+FWorkspaceEditorSelectionScope::~FWorkspaceEditorSelectionScope()
+{
+	if (const TSharedPtr<FWorkspaceEditor> SharedWorkspaceEditor = StaticCastSharedPtr<FWorkspaceEditor>(WeakWorkspaceEditor.Pin()))
+	{
+		--SharedWorkspaceEditor->SelectionScopeDepth;
+		check(SharedWorkspaceEditor->SelectionScopeDepth >= 0);
+		
+		if (SharedWorkspaceEditor->SelectionScopeDepth == 0)
+		{
+			SharedWorkspaceEditor->bSelectionScopeCleared = false;
+		}
+	}
+}
+
 namespace WorkspaceModes
 {
 	const FName WorkspaceEditor("WorkspaceEditorMode");
@@ -185,6 +207,36 @@ void FWorkspaceEditor::PostInitAssetEditor()
 	RegenerateMenusAndToolbars();
 }
 
+void FWorkspaceEditor::OnToolkitHostingFinished(const TSharedRef<IToolkit>& Toolkit)
+{
+	TSharedPtr<FWorkspaceEditorModeUILayer> ModeUILayer; 
+	ModeUILayers.RemoveAndCopyValue(Toolkit->GetToolkitFName(), ModeUILayer);
+	if(ModeUILayer)
+	{
+		ModeUILayer->OnToolkitHostingFinished(Toolkit);
+	}
+		
+	UToolMenus::UnregisterOwner(&(*Toolkit));	
+	HostedToolkits.Remove( Toolkit );
+}
+
+void FWorkspaceEditor::OnToolkitHostingStarted(const TSharedRef<IToolkit>& Toolkit)
+{
+	ensure(!ModeUILayers.Contains(Toolkit->GetToolkitFName()));
+
+	TSharedPtr<FWorkspaceEditorModeUILayer> ModeUILayer = MakeShared<FWorkspaceEditorModeUILayer>(ToolkitHost.Pin().Get());
+	ModeUILayer->SetModeMenuCategory(EditorMenuCategory);
+
+	// Actually re-use the main toolbar rather than a secondary, which also requires appending the UI layer commands
+	ModeUILayer->SetSecondaryModeToolbarName(GetToolMenuToolbarName());	
+	ToolkitCommands->Append(ModeUILayer->GetModeCommands());
+
+	ModeUILayer->OnToolkitHostingStarted(Toolkit);
+
+	ModeUILayers.Add(Toolkit->GetToolkitFName(), ModeUILayer);	
+	HostedToolkits.Add(Toolkit);
+}
+
 void FWorkspaceEditor::RestoreEditedObjectState()
 {
 	UWorkspaceState* State = Workspace->GetState();
@@ -320,14 +372,24 @@ IWorkspaceEditor::FOnOutlinerSelectionChanged& FWorkspaceEditor::OnOutlinerSelec
 
 void FWorkspaceEditor::SetGlobalSelection(FGlobalSelectionId SelectionId, FOnClearGlobalSelection OnClearSelectionDelegate)
 {
-	// Only execute if widget is still valid, and it is not the same as the previous call 
-	if (LastGlobalSelectionId.IsValid() && SelectionId != LastGlobalSelectionId)
-	{		
-		LastOnClearSelectionDelegate.ExecuteIfBound();
+	if(SelectionScopeDepth == 0 || bSelectionScopeCleared == false)
+	{
+		TArray<TPair<FGlobalSelectionId, FOnClearGlobalSelection>> SelectionsCopy = GlobalSelections;
+		GlobalSelections.Empty();
+		
+		for (TPair<FGlobalSelectionId, FOnClearGlobalSelection>& Selection : SelectionsCopy)
+		{
+			// Only execute if widget is still valid, and it is not the same as the previous call 
+			if (Selection.Key.IsValid() && SelectionId != Selection.Key)
+			{
+				Selection.Value.ExecuteIfBound();
+			}
+		}
+		
+		bSelectionScopeCleared = true;
 	}
 
-	LastGlobalSelectionId = SelectionId;
-	LastOnClearSelectionDelegate = OnClearSelectionDelegate;
+	GlobalSelections.Add({ SelectionId, OnClearSelectionDelegate});
 }
 
 void FWorkspaceEditor::SetFocussedAsset(const TObjectPtr<UObject> InAsset)
@@ -338,7 +400,30 @@ void FWorkspaceEditor::SetFocussedAsset(const TObjectPtr<UObject> InAsset)
 	if (InAsset != nullptr)
 	{
 		ModeManager->GetSelectedObjects()->Select(InAsset);
+
+		FWorkspaceEditorModule& WorkspaceEditorModule = FModuleManager::Get().LoadModuleChecked<FWorkspaceEditorModule>("WorkspaceEditor");
+		if (const FObjectDocumentArgs* Args = WorkspaceEditorModule.FindObjectDocumentType(InAsset))
+		{
+			const FEditorModeID NewEditorMode = Args->DocumentEditorMode;
+			if (NewEditorMode != NAME_None)
+			{
+				if(!ModeManager->IsModeActive(NewEditorMode))
+				{
+					ModeManager->ActivateMode(NewEditorMode);
+				}
+			}
+			else
+			{				
+				ModeManager->DeactivateAllModes();
+			}
+		}		
+	}
+	else
+	{
+		ModeManager->DeactivateAllModes();
 	}
+
+	OnFocussedAssetChangedDelegate.Broadcast(InAsset);
 }
 
 const TObjectPtr<UObject> FWorkspaceEditor::GetFocussedAssetOfClass(const TObjectPtr<UClass> AssetClass) const
@@ -357,6 +442,16 @@ void FWorkspaceEditor::HandleOutlinerSelectionChanged(TConstArrayView<FWorkspace
 	LastSelectedExports = InExports;
 }
 
+FOnFocussedAssetChanged& FWorkspaceEditor::OnFocussedAssetChanged()
+{
+	return OnFocussedAssetChangedDelegate;
+}
+
+TSharedPtr<IDetailsView> FWorkspaceEditor::GetDetailsView()
+{
+	return DetailsView;
+}
+
 void FWorkspaceEditor::BindCommands()
 {
 	const FWorkspaceAssetEditorCommands& Commands = FWorkspaceAssetEditorCommands::Get();
@@ -398,6 +493,8 @@ void FWorkspaceEditor::ExtendToolbar()
 void FWorkspaceEditor::RegisterTabSpawners(const TSharedRef<FTabManager>& InTabManager)
 {
 	FBaseAssetToolkit::RegisterTabSpawners(InTabManager);
+	
+	EditorMenuCategory = InTabManager->AddLocalWorkspaceMenuCategory(LOCTEXT("WorkspaceMenu_WorkspaceEditor", "Workspace Editor"));    
 
 	InTabManager->RegisterTabSpawner(WorkspaceTabs::WorkspaceView, FOnSpawnTab::CreateLambda([this](const FSpawnTabArgs& Args) -> TSharedRef<SDockTab>
 	{
@@ -647,6 +744,14 @@ void FWorkspaceEditor::OnClose()
 		DetailsView.Reset();
 	}
 
+	for (TSharedPtr<class IToolkit> Toolkit : HostedToolkits)
+	{
+		if (Toolkit.IsValid())
+		{
+			UToolMenus::UnregisterOwner(Toolkit.Get());
+		}
+	}
+
 	TabFactories.Clear();
 	FBaseAssetToolkit::OnClose();
 }
diff --git a/Engine/Plugins/Experimental/Workspace/Source/WorkspaceEditor/Private/WorkspaceEditor.h b/Engine/Plugins/Experimental/Workspace/Source/WorkspaceEditor/Private/WorkspaceEditor.h
--- a/Engine/Plugins/Experimental/Workspace/Source/WorkspaceEditor/Private/WorkspaceEditor.h
+++ b/Engine/Plugins/Experimental/Workspace/Source/WorkspaceEditor/Private/WorkspaceEditor.h
@@ -6,6 +6,7 @@
 #include "StructUtils/InstancedStruct.h"
 #include "IWorkspaceEditor.h"
 #include "WorkspaceAssetRegistryInfo.h"
+#include "WorkspaceEditorModeUILayer.h"
 #include "WorkflowOrientedApp/WorkflowTabManager.h"
 
 struct FWorkspaceDocumentState;
@@ -53,6 +54,7 @@ private:
 	friend struct FGraphDocumentSummoner;
 	friend struct FWorkspaceTabSummoner;
 	friend struct FAssetDocumentSummoner;
+	friend struct FWorkspaceEditorSelectionScope;
 
 	// FBaseAssetToolkit interface
 	virtual void RegisterTabSpawners(const TSharedRef<class FTabManager>& TabManager) override;
@@ -76,10 +78,10 @@ private:
 	virtual bool CanSaveAsset() const override;
 	virtual FText GetTabSuffix() const override;
 	virtual FText GetToolkitName() const override;
-	virtual FText GetToolkitToolTipText() const override;	
-	virtual bool IsFindInContentBrowserButtonVisible() const override { return true; }
+	virtual FText GetToolkitToolTipText() const override;
 	virtual void FindInContentBrowser_Execute() override;
-
+	virtual void OnToolkitHostingFinished(const TSharedRef<IToolkit>& Toolkit) override;
+	virtual void OnToolkitHostingStarted(const TSharedRef<IToolkit>& Toolkit) override;
 
 	// FGCObject interface
 	virtual void AddReferencedObjects(FReferenceCollector& Collector) override
@@ -102,7 +104,10 @@ private:
 	virtual FOnOutlinerSelectionChanged& OnOutlinerSelectionChanged() override;
 	virtual void SetGlobalSelection(FGlobalSelectionId SelectionId, FOnClearGlobalSelection OnClearSelectionDelegate) override;
 	virtual void SetFocussedAsset(const TObjectPtr<UObject> InAsset) override;	
-	virtual const TObjectPtr<UObject> GetFocussedAssetOfClass(const TObjectPtr<UClass> AssetClass) const override;	
+	virtual const TObjectPtr<UObject> GetFocussedAssetOfClass(const TObjectPtr<UClass> AssetClass) const override;
+	virtual FOnFocussedAssetChanged& OnFocussedAssetChanged() override;
+	virtual TSharedPtr<IDetailsView> GetDetailsView() override;
+
 	void HandleOutlinerSelectionChanged(TConstArrayView<FWorkspaceOutlinerItemExport> InExports);
 
 	void BindCommands();
@@ -137,6 +142,9 @@ private:
 
 	TSharedPtr<SWorkspaceView> WorkspaceView;
 
+	TMap<FName, TSharedPtr<FWorkspaceEditorModeUILayer>> ModeUILayers;
+	TArray<TSharedPtr<class IToolkit>> HostedToolkits;
+
 	/** Tabs to be registered into the Workspace */
 	FWorkflowAllowedTabSet TabFactories;
 
@@ -145,10 +153,14 @@ private:
 	bool bSavingAssetEntries = false;
 	bool bClosingDown = false;
 
-	FGlobalSelectionId LastGlobalSelectionId = nullptr;
-	FOnClearGlobalSelection LastOnClearSelectionDelegate = nullptr;
+	TArray<TPair<FGlobalSelectionId, FOnClearGlobalSelection>> GlobalSelections;
 	FOnOutlinerSelectionChanged OnOutlinerSelectionChangedDelegate;
+	bool bSelectionScopeCleared = true;
+	int32 SelectionScopeDepth = 0;
+	
 	TArray<FWorkspaceOutlinerItemExport> LastSelectedExports;
+	FOnFocussedAssetChanged OnFocussedAssetChangedDelegate;
+	TSharedPtr<FWorkspaceItem> EditorMenuCategory;
 };
 
 }
\ No newline at end of file
diff --git a/Engine/Plugins/Experimental/Workspace/Source/WorkspaceEditor/Private/WorkspaceEditorModeUILayer.cpp b/Engine/Plugins/Experimental/Workspace/Source/WorkspaceEditor/Private/WorkspaceEditorModeUILayer.cpp
new file mode 100644
--- /dev/null
+++ b/Engine/Plugins/Experimental/Workspace/Source/WorkspaceEditor/Private/WorkspaceEditorModeUILayer.cpp
@@ -0,0 +1,45 @@
+// Copyright Epic Games, Inc. All Rights Reserved.
+
+#include "WorkspaceEditorModeUILayer.h"
+#include "Toolkits/IToolkit.h"
+#include "ToolMenus.h"
+
+FWorkspaceEditorModeUILayer::FWorkspaceEditorModeUILayer(const IToolkitHost* InToolkitHost) : FAssetEditorModeUILayer(InToolkitHost)
+{
+}
+
+void FWorkspaceEditorModeUILayer::OnToolkitHostingStarted(const TSharedRef<IToolkit>& Toolkit)
+{
+	if (!Toolkit->IsAssetEditor())
+	{
+		FAssetEditorModeUILayer::OnToolkitHostingStarted(Toolkit);
+		HostedToolkit = Toolkit;
+		Toolkit->SetModeUILayer(SharedThis(this));
+		Toolkit->RegisterTabSpawners(ToolkitHost->GetTabManager().ToSharedRef());
+		RegisterModeTabSpawners();
+
+		OnToolkitHostReadyForUI.Execute();
+
+		// Set up an owner for the current scope so that we can cleanly clean up the toolbar extension on hosting finish
+		FToolMenuOwnerScoped Owner(&(*Toolkit));
+		UToolMenu* SecondaryModeToolbar = UToolMenus::Get()->ExtendMenu(GetSecondaryModeToolbarName());
+		OnRegisterSecondaryModeToolbarExtension.ExecuteIfBound(SecondaryModeToolbar);		
+	}
+}
+
+void FWorkspaceEditorModeUILayer::OnToolkitHostingFinished(const TSharedRef<IToolkit>& Toolkit)
+{	
+	FAssetEditorModeUILayer::OnToolkitHostingFinished(Toolkit);
+	UToolMenus::UnregisterOwner(&(*Toolkit));
+}
+
+void FWorkspaceEditorModeUILayer::SetModeMenuCategory(const TSharedPtr<FWorkspaceItem>& MenuCategoryIn)
+{
+	MenuCategory = MenuCategoryIn;
+}
+
+TSharedPtr<FWorkspaceItem> FWorkspaceEditorModeUILayer::GetModeMenuCategory() const
+{
+	check(MenuCategory);
+	return MenuCategory;
+}
\ No newline at end of file
diff --git a/Engine/Plugins/Experimental/Workspace/Source/WorkspaceEditor/Private/WorkspaceEditorModeUILayer.h b/Engine/Plugins/Experimental/Workspace/Source/WorkspaceEditor/Private/WorkspaceEditorModeUILayer.h
new file mode 100644
--- /dev/null
+++ b/Engine/Plugins/Experimental/Workspace/Source/WorkspaceEditor/Private/WorkspaceEditorModeUILayer.h
@@ -0,0 +1,20 @@
+// Copyright Epic Games, Inc. All Rights Reserved.
+
+
+#pragma once
+
+#include "CoreMinimal.h"
+#include "Toolkits/AssetEditorModeUILayer.h"
+
+class FWorkspaceEditorModeUILayer : public FAssetEditorModeUILayer
+{
+public:
+	FWorkspaceEditorModeUILayer(const IToolkitHost* InToolkitHost);
+	virtual void OnToolkitHostingStarted(const TSharedRef<IToolkit>& Toolkit) override;
+	virtual void OnToolkitHostingFinished(const TSharedRef<IToolkit>& Toolkit) override;
+	virtual TSharedPtr<FWorkspaceItem> GetModeMenuCategory() const override;
+	
+	void SetModeMenuCategory(const TSharedPtr<FWorkspaceItem>& MenuCategoryIn);
+protected:
+	TSharedPtr<FWorkspaceItem> MenuCategory;
+};
\ No newline at end of file
diff --git a/Engine/Plugins/Experimental/Workspace/Source/WorkspaceEditor/Public/IWorkspaceEditor.h b/Engine/Plugins/Experimental/Workspace/Source/WorkspaceEditor/Public/IWorkspaceEditor.h
--- a/Engine/Plugins/Experimental/Workspace/Source/WorkspaceEditor/Public/IWorkspaceEditor.h
+++ b/Engine/Plugins/Experimental/Workspace/Source/WorkspaceEditor/Public/IWorkspaceEditor.h
@@ -14,6 +14,17 @@ namespace UE::Workspace
 typedef TWeakPtr<SWidget> FGlobalSelectionId;
 using FOnClearGlobalSelection = FSimpleDelegate;
 
+DECLARE_MULTICAST_DELEGATE_OneParam(FOnFocussedAssetChanged, TObjectPtr<UObject>);
+
+// RAII helper allowing for a multi-widget selection scope within a WorkspaceEditor instance
+struct WORKSPACEEDITOR_API FWorkspaceEditorSelectionScope
+{
+	FWorkspaceEditorSelectionScope(const TSharedPtr<class IWorkspaceEditor>& InWorkspaceEditor);	
+	~FWorkspaceEditorSelectionScope();
+
+	TWeakPtr<class IWorkspaceEditor> WeakWorkspaceEditor; 
+};
+
 class IWorkspaceEditor : public FBaseAssetToolkit
 {
 public:
@@ -44,9 +55,12 @@ public:
 	virtual const TObjectPtr<UObject> GetFocussedAssetOfClass(const TObjectPtr<UClass> InClass ) const = 0;
 	
 	template<typename AssetClass>
-	TObjectPtr<AssetClass> GetFocussedAsset() const { return GetFocussedAssetOfClass(AssetClass::StaticClass()); }
+	TObjectPtr<AssetClass> GetFocussedAsset() const { return Cast<AssetClass>(GetFocussedAssetOfClass(AssetClass::StaticClass())); }
 	TObjectPtr<UObject> GetFocussedAsset() const { return GetFocussedAssetOfClass(UObject::StaticClass()); }
 
+	// Multi-cast delegate broadcasted whenever the asset focussed inside of the WorkspaceEditor changes
+	virtual FOnFocussedAssetChanged& OnFocussedAssetChanged() = 0;
+
 	// Get the current single selection of the outliner.
 	// @return true if a single selection is active
 	virtual bool GetOutlinerSelection(TArray<FWorkspaceOutlinerItemExport>& OutExports) const = 0;
@@ -54,6 +68,9 @@ public:
 	// Delegate fired when selection changes in the workspace outliner
 	using FOnOutlinerSelectionChanged = TMulticastDelegate<void(TConstArrayView<FWorkspaceOutlinerItemExport> InExports)>;
 	virtual FOnOutlinerSelectionChanged& OnOutlinerSelectionChanged() = 0;
+
+	// Retrieves the common DetailsView widget
+	virtual TSharedPtr<IDetailsView> GetDetailsView() = 0;
 };
 
 }
\ No newline at end of file
diff --git a/Engine/Plugins/Experimental/Workspace/Source/WorkspaceEditor/Public/IWorkspaceEditorModule.h b/Engine/Plugins/Experimental/Workspace/Source/WorkspaceEditor/Public/IWorkspaceEditorModule.h
--- a/Engine/Plugins/Experimental/Workspace/Source/WorkspaceEditor/Public/IWorkspaceEditorModule.h
+++ b/Engine/Plugins/Experimental/Workspace/Source/WorkspaceEditor/Public/IWorkspaceEditorModule.h
@@ -3,6 +3,7 @@
 #pragma once
 
 #include "CoreMinimal.h"
+#include "EditorModes.h"
 #include "GraphEditor.h"
 #include "StructUtils/InstancedStruct.h"
 #include "WorkspaceFactory.h"
@@ -117,6 +118,9 @@ struct FObjectDocumentArgs
 
 	// Delegate called to get the bread crumb trail for this document tab
 	FOnGetDocumentBreadcrumbTrail OnGetDocumentBreadcrumbTrail;
+
+	// EditorMode ID to be associated with this asset, will be used to try and activate matching editor mode when this asset is focussed
+	FEditorModeID DocumentEditorMode = NAME_None;
 };
 
 
diff --git a/Engine/Plugins/Experimental/Workspace/Source/WorkspaceEditor/Public/IWorkspaceOutlinerItemDetails.h b/Engine/Plugins/Experimental/Workspace/Source/WorkspaceEditor/Public/IWorkspaceOutlinerItemDetails.h
--- a/Engine/Plugins/Experimental/Workspace/Source/WorkspaceEditor/Public/IWorkspaceOutlinerItemDetails.h
+++ b/Engine/Plugins/Experimental/Workspace/Source/WorkspaceEditor/Public/IWorkspaceOutlinerItemDetails.h
@@ -3,6 +3,7 @@
 #pragma once
 
 #include "WorkspaceAssetRegistryInfo.h"
+#include "Styling/SlateColor.h"
 
 class UPackage;
 struct FToolMenuContext;
@@ -22,6 +23,7 @@ class IWorkspaceOutlinerItemDetails : public TSharedFromThis<IWorkspaceOutlinerI
 public:
     virtual ~IWorkspaceOutlinerItemDetails() = default;
     virtual const FSlateBrush* GetItemIcon(const FWorkspaceOutlinerItemExport& Export) const { return nullptr; }
+	virtual FSlateColor GetItemColor(const FWorkspaceOutlinerItemExport& Export) const { return FSlateColor::UseForeground(); }
     virtual void HandleDoubleClick(const FToolMenuContext& ToolMenuContext) const {}
 	virtual bool CanDelete(const FWorkspaceOutlinerItemExport& Export) const { return true; }
 	virtual void Delete(TConstArrayView<FWorkspaceOutlinerItemExport> Exports) const {}
@@ -29,6 +31,7 @@ public:
     virtual void Rename(const FWorkspaceOutlinerItemExport& Export, const FText& InName) const {}
     virtual bool ValidateName(const FWorkspaceOutlinerItemExport& Export, const FText& InName, FText& OutErrorMessage) const { return false; }
     virtual UPackage* GetPackage(const FWorkspaceOutlinerItemExport& Export) const { return nullptr; }
+	virtual bool HandleSelected(const FToolMenuContext& ToolMenuContext) const { return false; }
 };
 
 }
diff --git a/Engine/Plugins/Experimental/Workspace/Source/WorkspaceEditor/WorkspaceEditor.Build.cs b/Engine/Plugins/Experimental/Workspace/Source/WorkspaceEditor/WorkspaceEditor.Build.cs
--- a/Engine/Plugins/Experimental/Workspace/Source/WorkspaceEditor/WorkspaceEditor.Build.cs
+++ b/Engine/Plugins/Experimental/Workspace/Source/WorkspaceEditor/WorkspaceEditor.Build.cs
@@ -37,7 +37,10 @@ namespace UnrealBuildTool.Rules
 					"JsonUtilities",
 					"AssetRegistry",
 					"ContentBrowser",
-					"SceneOutliner"
+					"SceneOutliner",
+					"EditorInteractiveToolsFramework",
+					"InteractiveToolsFramework",
+					"EditorFramework"
 				}
 			);
 		}
