Implement curve and attribute support in AnimNext
Fix mesh space additives in AnimNext

diff --git a/Engine/Plugins/Experimental/AnimNext/Source/AnimNext/Internal/AnimNextModuleImpl.h b/Engine/Plugins/Experimental/AnimNext/Source/AnimNext/Internal/AnimNextModuleImpl.h
--- a/Engine/Plugins/Experimental/AnimNext/Source/AnimNext/Internal/AnimNextModuleImpl.h
+++ b/Engine/Plugins/Experimental/AnimNext/Source/AnimNext/Internal/AnimNextModuleImpl.h
@@ -30,7 +30,7 @@ public:
 	virtual void RegisterAnimNextAnimGraph(const IAnimNextAnimGraph& InAnimGraphImpl) override;
 	virtual void UnregisterAnimNextAnimGraph() override;
 	virtual void UpdateGraph(const FAnimNextGraphInstancePtr& GraphInstance, float DeltaTime, FTraitEventList& InputEventList, FTraitEventList& OutputEventList) override;
-	virtual void EvaluateGraph(const FAnimNextGraphInstancePtr& GraphInstance, const FReferencePose& RefPose, int32 GraphLODLevel, FLODPoseHeap& OutputPose) const override;
+	virtual void EvaluateGraph(const FAnimNextGraphInstancePtr& GraphInstance, const FReferencePose& RefPose, int32 GraphLODLevel, FAnimNextGraphLODPose& OutputPose) const override;
 
 #if WITH_ANIMNEXT_CONSOLE_COMMANDS
 	void ListNodeTemplates(const TArray<FString>& Args);
diff --git a/Engine/Plugins/Experimental/AnimNext/Source/AnimNext/Internal/Component/SkinnedMeshComponentExtensions.h b/Engine/Plugins/Experimental/AnimNext/Source/AnimNext/Internal/Component/SkinnedMeshComponentExtensions.h
--- a/Engine/Plugins/Experimental/AnimNext/Source/AnimNext/Internal/Component/SkinnedMeshComponentExtensions.h
+++ b/Engine/Plugins/Experimental/AnimNext/Source/AnimNext/Internal/Component/SkinnedMeshComponentExtensions.h
@@ -1,4 +1,4 @@
-﻿// Copyright Epic Games, Inc. All Rights Reserved.
+// Copyright Epic Games, Inc. All Rights Reserved.
 
 #pragma once
 
@@ -15,7 +15,11 @@ namespace UE::Anim
 struct FSkinnedMeshComponentExtensions
 {
 	// Converts local to component space, flips buffers, updates bounds and dispatches to renderer
-	static ANIMNEXT_API void CompleteAndDispatch(USkinnedMeshComponent* InComponent, TConstArrayView<FBoneIndexType> InParentIndices, TConstArrayView<FBoneIndexType> InRequiredBoneIndices, TConstArrayView<FTransform> InLocalSpaceTransforms);
+	static ANIMNEXT_API void CompleteAndDispatch(
+		USkinnedMeshComponent* InComponent,
+		TConstArrayView<FBoneIndexType> InParentIndices,
+		TConstArrayView<FBoneIndexType> InRequiredBoneIndices,
+		TConstArrayView<FTransform> InLocalSpaceTransforms);
 };
 
 }
diff --git a/Engine/Plugins/Experimental/AnimNext/Source/AnimNext/Internal/Graph/AnimNext_LODPose.h b/Engine/Plugins/Experimental/AnimNext/Source/AnimNext/Internal/Graph/AnimNext_LODPose.h
--- a/Engine/Plugins/Experimental/AnimNext/Source/AnimNext/Internal/Graph/AnimNext_LODPose.h
+++ b/Engine/Plugins/Experimental/AnimNext/Source/AnimNext/Internal/Graph/AnimNext_LODPose.h
@@ -5,6 +5,9 @@
 #include "CoreMinimal.h"
 #include "ReferencePose.h"
 #include "LODPose.h"
+#include "Animation/AnimCurveTypes.h"
+#include "Animation/AttributesRuntime.h"
+
 #include "AnimNext_LODPose.generated.h"
 
 USTRUCT(BlueprintType, meta = (DisplayName = "ReferencePose"))
@@ -27,16 +30,13 @@ struct FAnimNextGraphLODPose
 {
 	GENERATED_BODY()
 
-	FAnimNextGraphLODPose() = default;
+	// Joint transforms
+	UE::AnimNext::FLODPoseHeap LODPose;
 
-	explicit FAnimNextGraphLODPose(const UE::AnimNext::FLODPoseHeap& InLODPose)
-		: LODPose(InLODPose)
-	{
-	}
-	explicit FAnimNextGraphLODPose(UE::AnimNext::FLODPoseHeap&& InLODPose)
-		: LODPose(MoveTemp(InLODPose))
-	{
-	}
+	// Float curves
+	FBlendedHeapCurve Curves;
 
-	UE::AnimNext::FLODPoseHeap LODPose;
+	// Attributes
+	// Note that attribute bone indices are LOD bone indices matching the LOD pose
+	UE::Anim::FHeapAttributeContainer Attributes;
 };
diff --git a/Engine/Plugins/Experimental/AnimNext/Source/AnimNext/Private/AnimNextModuleImpl.cpp b/Engine/Plugins/Experimental/AnimNext/Source/AnimNext/Private/AnimNextModuleImpl.cpp
--- a/Engine/Plugins/Experimental/AnimNext/Source/AnimNext/Private/AnimNextModuleImpl.cpp
+++ b/Engine/Plugins/Experimental/AnimNext/Source/AnimNext/Private/AnimNextModuleImpl.cpp
@@ -164,7 +164,7 @@ namespace UE::AnimNext
 		}
 	}
 
-	void FAnimNextModuleImpl::EvaluateGraph(const FAnimNextGraphInstancePtr& GraphInstance, const UE::AnimNext::FReferencePose& RefPose, int32 GraphLODLevel, FLODPoseHeap& OutputPose) const
+	void FAnimNextModuleImpl::EvaluateGraph(const FAnimNextGraphInstancePtr& GraphInstance, const UE::AnimNext::FReferencePose& RefPose, int32 GraphLODLevel, FAnimNextGraphLODPose& OutputPose) const
 	{
 		if (AnimGraphImpl != nullptr)
 		{
diff --git a/Engine/Plugins/Experimental/AnimNext/Source/AnimNext/Private/DecompressionTools.cpp b/Engine/Plugins/Experimental/AnimNext/Source/AnimNext/Private/DecompressionTools.cpp
--- a/Engine/Plugins/Experimental/AnimNext/Source/AnimNext/Private/DecompressionTools.cpp
+++ b/Engine/Plugins/Experimental/AnimNext/Source/AnimNext/Private/DecompressionTools.cpp
@@ -44,6 +44,9 @@ struct FGetBonePoseScratchArea : public TThreadSingleton<FGetBonePoseScratchArea
 	BoneTrackArray AnimScaleRetargetingPairs;
 	BoneTrackArray AnimRelativeRetargetingPairs;
 	BoneTrackArray OrientAndScaleRetargetingPairs;
+
+	// A bit set that specifies whether a compact bone index has its rotation animated by the sequence or not
+	TBitArray<> AnimatedCompactRotations;
 };
 
 
@@ -56,7 +59,7 @@ static bool CanEvaluateRawAnimationData(const UAnimSequence* AnimSequence)
 #endif
 }
 
-static bool UseRawDataForPoseExtraction(const UAnimSequence* AnimSequence, FLODPose& AnimationPoseData)
+static bool UseRawDataForPoseExtraction(const UAnimSequence* AnimSequence, const FLODPose& AnimationPoseData)
 {
 	return CanEvaluateRawAnimationData(AnimSequence) && (
 #if WITH_EDITOR
@@ -71,9 +74,20 @@ static bool UseRawDataForPoseExtraction(const UAnimSequence* AnimSequence, FLODP
 		AnimationPoseData.ShouldUseSourceData());
 }
 
+bool FDecompressionTools::ShouldUseRawData(const UAnimSequence* AnimSequence, const FLODPose& AnimationPoseData)
+{
+	const bool bUseRawData =
+#if WITH_EDITOR
+		GetForceRawData() ||
+#endif // WITH_EDITOR
+		UseRawDataForPoseExtraction(AnimSequence, AnimationPoseData);
+
+	return bUseRawData;
+}
+
 // --- ---
 
-/*static*/ void FDecompressionTools::GetAnimationPose(const UAnimSequence* AnimSequence, FLODPose& OutAnimationPoseData, const FAnimExtractContext& ExtractionContext)
+/*static*/ void FDecompressionTools::GetAnimationPose(const UAnimSequence* AnimSequence, const FAnimExtractContext& ExtractionContext, FLODPose& OutAnimationPoseData, bool bForceUseRawData)
 {
 	if (OutAnimationPoseData.GetRefPose().IsValid() == false)
 	{
@@ -81,7 +95,7 @@ static bool UseRawDataForPoseExtraction(const UAnimSequence* AnimSequence, FLODP
 	}
 
 	// @todo anim: if compressed and baked in the future, we don't have to do this 
-	if (UseRawDataForPoseExtraction(AnimSequence, OutAnimationPoseData) && AnimSequence->IsValidAdditive())
+	if (bForceUseRawData && AnimSequence->IsValidAdditive())
 	{
 		switch (AnimSequence->GetAdditiveAnimType())
 		{
@@ -103,21 +117,12 @@ static bool UseRawDataForPoseExtraction(const UAnimSequence* AnimSequence, FLODP
 	}
 	else
 	{
-		GetBonePose(AnimSequence, ExtractionContext, OutAnimationPoseData);
-	}
-
-	// If the sequence has root motion enabled, allow sampling of a root motion delta into the custom attribute container of the outgoing pose
-	if (AnimSequence->HasRootMotion())
-	{
-		if (const UE::Anim::IAnimRootMotionProvider* RootMotionProvider = UE::Anim::IAnimRootMotionProvider::Get())
-		{
-			//RootMotionProvider->SampleRootMotion(ExtractionContext.DeltaTimeRecord, *AnimSequence, ExtractionContext.bLooping, OutAnimationPoseData.GetAttributes());
-		}
+		GetBonePose(AnimSequence, ExtractionContext, OutAnimationPoseData, bForceUseRawData);
 	}
 
 	// Check that all bone atoms coming from animation are normalized
 #if DO_CHECK && WITH_EDITORONLY_DATA
-	//check(OutAnimationPoseData.LocalTransforms.IsNormalized());
+	check(OutAnimationPoseData.LocalTransformsView.IsValid());
 #endif
 
 }
@@ -130,13 +135,8 @@ void FDecompressionTools::GetBonePose(const UAnimSequence* AnimSequence, const F
 	const TArrayView<const FBoneIndexType> LODBoneIndexToSkeletonBoneIndexMap = OutAnimationPoseData.GetLODBoneIndexToSkeletonBoneIndexMap();
 
 	check(!bForceUseRawData || CanEvaluateRawAnimationData(AnimSequence));
-	const bool bUseRawDataForPoseExtraction = (CanEvaluateRawAnimationData(AnimSequence) && bForceUseRawData) ||
-#if WITH_EDITOR
-		GetForceRawData() ||
-#endif // WITH_EDITOR
-		UseRawDataForPoseExtraction(AnimSequence, OutAnimationPoseData);
 
-	const bool bIsBakedAdditive = !bUseRawDataForPoseExtraction && AnimSequence->IsValidAdditive();
+	const bool bIsBakedAdditive = !bForceUseRawData && AnimSequence->IsValidAdditive();
 
 	const USkeleton* MySkeleton = AnimSequence->GetSkeleton();
 	if (!MySkeleton)
@@ -181,16 +181,16 @@ void FDecompressionTools::GetBonePose(const UAnimSequence* AnimSequence, const F
 	}
 
 #if WITH_EDITOR
-	const int32 NumTracks = bUseRawDataForPoseExtraction ? AnimSequence->GetDataModelInterface()->GetNumBoneTracks() : AnimSequence->CompressedData.CompressedTrackToSkeletonMapTable.Num();
+	const int32 NumTracks = bForceUseRawData ? AnimSequence->GetDataModelInterface()->GetNumBoneTracks() : AnimSequence->CompressedData.CompressedTrackToSkeletonMapTable.Num();
 #else
 	const int32 NumTracks = AnimSequence->CompressedData.CompressedTrackToSkeletonMapTable.Num();
 #endif 
-	const bool bTreatAnimAsAdditive = (AnimSequence->IsValidAdditive() && !bUseRawDataForPoseExtraction); // Raw data is never additive
+	const bool bTreatAnimAsAdditive = (AnimSequence->IsValidAdditive() && !bForceUseRawData); // Raw data is never additive
 	const FRootMotionReset RootMotionReset(AnimSequence->bEnableRootMotion, AnimSequence->RootMotionRootLock, AnimSequence->bForceRootLock, AnimSequence->ExtractRootTrackTransform(0.0f, /*&RequiredBones*/nullptr), bTreatAnimAsAdditive);
 
 #if WITH_EDITOR
 	// Evaluate raw (source) curve and bone data
-	if (bUseRawDataForPoseExtraction)
+	if (bForceUseRawData)
 	{
 		// TODO : Curves support
 		//{
@@ -210,6 +210,7 @@ void FDecompressionTools::GetBonePose(const UAnimSequence* AnimSequence, const F
 	}
 	else
 #endif // WITH_EDITOR
+	{
 		// Only try and evaluate compressed bone data if the animation contains any bone tracks
 		if (NumTracks != 0)
 		{
@@ -227,24 +228,7 @@ void FDecompressionTools::GetBonePose(const UAnimSequence* AnimSequence, const F
 
 			DecompressPose(OutAnimationPoseData, AnimSequence->CompressedData, ExtractionContext, DecompContext, AnimSequence->GetRetargetTransforms(), RootMotionReset);
 		}
-
-//	// TODO : Curves support
-//	// (Always) evaluate compressed curve data
-//	{
-//		// Scoped so that the RemappedCurve destructs and isn't kept alive longer than needed.
-//		FSkeletonRemappingCurve RemappedCurve(OutAnimationPoseData.GetCurve(), RequiredBones, GetSkeleton());
-//#if WITH_EDITOR
-//		// When evaluating from raw animation data, UE::Anim::BuildPoseFromModel will populate the curve data
-//		if (!bUseRawDataForPoseExtraction)
-//#endif // WITH_EDITOR
-//		{
-//			AnimSequence->EvaluateCurveData(RemappedCurve.GetCurve(), ExtractionContext.CurrentTime, bUseRawDataForPoseExtraction);
-//		}
-//	}
-
-	// TODO : attributes support
-	// Evaluate animation attributes (no compressed format yet)
-	//EvaluateAttributes(OutAnimationPoseData, ExtractionContext, false);
+	}
 }
 
 void FDecompressionTools::GetBonePose_Additive(const UAnimSequence* AnimSequence, const FAnimExtractContext& ExtractionContext, FLODPose& OutAnimationPoseData)
@@ -255,6 +239,58 @@ void FDecompressionTools::GetBonePose_AdditiveMeshRotationOnly(const UAnimSequen
 {
 }
 
+void FDecompressionTools::GetAnimationCurves(const UAnimSequence* AnimSequence, const FAnimExtractContext& ExtractionContext, FBlendedCurve& OutCurves, bool bForceUseRawData)
+{
+	check(!bForceUseRawData || CanEvaluateRawAnimationData(AnimSequence));
+
+	if (bForceUseRawData)
+	{
+		//FAnimationPoseData RemappedPoseData(OutAnimationPoseData.GetPose(), OutCurves, OutAnimationPoseData.GetAttributes());
+		//const UE::Anim::DataModel::FEvaluationContext EvaluationContext(ExtractionContext.CurrentTime, DataModelInterface->GetFrameRate(), GetRetargetTransformsSourceName(), GetRetargetTransforms(), Interpolation);
+		//DataModelInterface->Evaluate(RemappedPoseData, EvaluationContext);
+	}
+	else
+	{
+		AnimSequence->EvaluateCurveData(OutCurves, ExtractionContext.CurrentTime, bForceUseRawData);
+	}
+}
+
+void FDecompressionTools::GetAnimationAttributes(const UAnimSequence* AnimSequence, const FAnimExtractContext& ExtractionContext, const FReferencePose& RefPose, UE::Anim::FStackAttributeContainer& OutAttributes, bool bForceUseRawData)
+{
+	QUICK_SCOPE_CYCLE_COUNTER(STAT_EvaluateAttributes);
+
+#if WITH_EDITOR
+	if (bForceUseRawData)
+	{
+		AnimSequence->ValidateModel();
+
+		for (const FAnimatedBoneAttribute& Attribute : AnimSequence->GetDataModel()->GetAttributes())
+		{
+			const int32 LODBoneIndex = RefPose.GetLODBoneIndexFromSkeletonBoneIndex(Attribute.Identifier.GetBoneIndex());
+			// Only add attribute if the bone its tied to exists in the currently evaluated set of bones
+			if (LODBoneIndex != INDEX_NONE)
+			{
+				UE::Anim::Attributes::GetAttributeValue(OutAttributes, FCompactPoseBoneIndex(LODBoneIndex), Attribute, ExtractionContext.CurrentTime);
+			}
+		}
+	}
+	else
+#endif // WITH_EDITOR
+	{
+		for (const TPair<FAnimationAttributeIdentifier, FAttributeCurve>& BakedAttribute : AnimSequence->AttributeCurves)
+		{
+			const int32 LODBoneIndex = RefPose.GetLODBoneIndexFromSkeletonBoneIndex(BakedAttribute.Key.GetBoneIndex());
+			// Only add attribute if the bone its tied to exists in the currently evaluated set of bones
+			if (LODBoneIndex != INDEX_NONE)
+			{
+				UE::Anim::FAttributeId Info(BakedAttribute.Key.GetName(), FCompactPoseBoneIndex(LODBoneIndex));
+				uint8* AttributePtr = OutAttributes.FindOrAdd(BakedAttribute.Key.GetType(), Info);
+				BakedAttribute.Value.EvaluateToPtr(BakedAttribute.Key.GetType(), ExtractionContext.CurrentTime, AttributePtr);
+			}
+		}
+	}
+}
+
 
 // --- ---
 
@@ -286,11 +322,12 @@ void FDecompressionTools::DecompressPose(FLODPose& OutAnimationPoseData,
 	const USkeleton* TargetSkeleton = OutAnimationPoseData.GetSkeletonAsset();
 	const FSkeletonRemapping& SkeletonRemapping = UE::Anim::FSkeletonRemappingRegistry::Get().GetRemapping(DecompressionContext.GetSourceSkeleton(), TargetSkeleton);
 
-	BoneTrackArray& RotationScalePairs = FGetBonePoseScratchArea::Get().RotationScalePairs;
-	BoneTrackArray& TranslationPairs = FGetBonePoseScratchArea::Get().TranslationPairs;
-	BoneTrackArray& AnimScaleRetargetingPairs = FGetBonePoseScratchArea::Get().AnimScaleRetargetingPairs;
-	BoneTrackArray& AnimRelativeRetargetingPairs = FGetBonePoseScratchArea::Get().AnimRelativeRetargetingPairs;
-	BoneTrackArray& OrientAndScaleRetargetingPairs = FGetBonePoseScratchArea::Get().OrientAndScaleRetargetingPairs;
+	FGetBonePoseScratchArea& ScratchArea = FGetBonePoseScratchArea::Get();
+	BoneTrackArray& RotationScalePairs = ScratchArea.RotationScalePairs;
+	BoneTrackArray& TranslationPairs = ScratchArea.TranslationPairs;
+	BoneTrackArray& AnimScaleRetargetingPairs = ScratchArea.AnimScaleRetargetingPairs;
+	BoneTrackArray& AnimRelativeRetargetingPairs = ScratchArea.AnimRelativeRetargetingPairs;
+	BoneTrackArray& OrientAndScaleRetargetingPairs = ScratchArea.OrientAndScaleRetargetingPairs;
 
 	// build a list of desired bones
 	RotationScalePairs.Reset();
@@ -299,6 +336,13 @@ void FDecompressionTools::DecompressPose(FLODPose& OutAnimationPoseData,
 	AnimRelativeRetargetingPairs.Reset();
 	OrientAndScaleRetargetingPairs.Reset();
 
+	const bool bIsMeshSpaceAdditive = DecompressionContext.GetAdditiveType() == AAT_RotationOffsetMeshSpace;
+	TBitArray<>& AnimatedCompactRotations = ScratchArea.AnimatedCompactRotations;
+	if (bIsMeshSpaceAdditive)
+	{
+		AnimatedCompactRotations.Init(false, NumLODBoneIndexes);
+	}
+
 	// Optimization: assuming first index is root bone. That should always be the case in Skeletons.
 	checkSlow((LODBoneIndexToSkeletonBoneIndexMap[0] == FMeshPoseBoneIndex(0).GetInt()));
 	// this is not guaranteed for AnimSequences though... If Root is not animated, Track will not exist.
@@ -321,6 +365,11 @@ void FDecompressionTools::DecompressPose(FLODPose& OutAnimationPoseData,
 				{
 					RotationScalePairs.Add(BoneTrackPair(LODBoneIndex, TrackIndex));
 
+					if (bIsMeshSpaceAdditive)
+					{
+						AnimatedCompactRotations[LODBoneIndex] = true;
+					}
+
 					// Skip extracting translation component for EBoneTranslationRetargetingMode::Skeleton.
 					switch (TargetSkeleton->GetBoneTranslationRetargetingMode(TargetSkeletonBoneIndex, OutAnimationPoseData.GetDisableRetargeting()))
 					{
@@ -419,11 +468,9 @@ void FDecompressionTools::DecompressPose(FLODPose& OutAnimationPoseData,
 	// Retarget the pose onto the target skeleton (correcting for differences in rest poses)
 	if (SkeletonRemapping.RequiresReferencePoseRetarget())
 	{
-		const int32 LODNumBones = OutAnimationPoseData.GetNumBones();
-
 		if (DecompressionContext.IsAdditiveAnimation())
 		{
-			for (int32 LODBoneIndex = (bFirstTrackIsRootBone ? 1 : 0); LODBoneIndex < LODNumBones; ++LODBoneIndex)
+			for (int32 LODBoneIndex = (bFirstTrackIsRootBone ? 1 : 0); LODBoneIndex < NumLODBoneIndexes; ++LODBoneIndex)
 			{
 				const int32 TargetSkeletonBoneIndex = ReferencePose.GetSkeletonBoneIndexFromLODBoneIndex(LODBoneIndex);
 
@@ -441,7 +488,7 @@ void FDecompressionTools::DecompressPose(FLODPose& OutAnimationPoseData,
 		}
 		else
 		{
-			for (int32 LODBoneIndex = (bFirstTrackIsRootBone ? 1 : 0); LODBoneIndex < LODNumBones; ++LODBoneIndex)
+			for (int32 LODBoneIndex = (bFirstTrackIsRootBone ? 1 : 0); LODBoneIndex < NumLODBoneIndexes; ++LODBoneIndex)
 			{
 				const int32 TargetSkeletonBoneIndex = LODBoneIndexToSkeletonBoneIndexMap[LODBoneIndex]; // ReferencePose.GetSkeletonBoneIndexFromLODBoneIndex(LODBoneIndex);
 				OutAnimationPoseData.LocalTransformsView[LODBoneIndex].SetRotation(SkeletonRemapping.RetargetBoneRotationToTargetSkeleton(TargetSkeletonBoneIndex, OutAnimationPoseData.LocalTransformsView[LODBoneIndex].GetRotation()));
@@ -544,6 +591,31 @@ void FDecompressionTools::DecompressPose(FLODPose& OutAnimationPoseData,
 	//		}
 	//	}
 	//}
+
+	if (bIsMeshSpaceAdditive)
+	{
+		// When an animation is a mesh-space additive, bones that aren't animated will end up with some non-identity
+		// delta relative to the base used to create the additive. This is because the delta is calculated in mesh-space
+		// unlike regular additive animations where bones that aren't animated has an identity delta. For rotations,
+		// this mesh-space delta will be the parent bone rotation.
+		// However, if a bone isn't animated in the sequence but present on the target skeleton, we have no data for it
+		// and the output pose will contain an identity delta which isn't what we want. As such, bones missing from
+		// the sequence have their rotation set to their parent.
+
+		const TArrayView<const FBoneIndexType> LODBoneIndexToParentLODBoneIndexMap = OutAnimationPoseData.GetLODBoneIndexToParentLODBoneIndexMap();
+
+		// We always skip the root since it has no parent (its delta value is fine as the identity)
+		for (int32 LODBoneIndex = 1; LODBoneIndex < NumLODBoneIndexes; ++LODBoneIndex)
+		{
+			if (!AnimatedCompactRotations[LODBoneIndex])
+			{
+				// This bone wasn't animated in the sequence, fix it up
+				const FBoneIndexType ParentLODBoneIndex = LODBoneIndexToParentLODBoneIndexMap[LODBoneIndex];
+				const FQuat ParentRotation = OutAnimationPoseData.LocalTransformsView[ParentLODBoneIndex].GetRotation();
+				OutAnimationPoseData.LocalTransformsView[LODBoneIndex].SetRotation(ParentRotation);
+			}
+		}
+	}
 }
 
 void FDecompressionTools::RetargetBoneTransform(const FReferencePose& ReferencePose
diff --git a/Engine/Plugins/Experimental/AnimNext/Source/AnimNext/Private/EvaluationVM/Tasks/PushAnimSequenceKeyframe.cpp b/Engine/Plugins/Experimental/AnimNext/Source/AnimNext/Private/EvaluationVM/Tasks/PushAnimSequenceKeyframe.cpp
--- a/Engine/Plugins/Experimental/AnimNext/Source/AnimNext/Private/EvaluationVM/Tasks/PushAnimSequenceKeyframe.cpp
+++ b/Engine/Plugins/Experimental/AnimNext/Source/AnimNext/Private/EvaluationVM/Tasks/PushAnimSequenceKeyframe.cpp
@@ -3,6 +3,7 @@
 #include "EvaluationVM/Tasks/PushAnimSequenceKeyframe.h"
 
 #include "Animation/AnimSequence.h"
+#include "Animation/AnimRootMotionProvider.h"
 #include "BonePose.h"
 #include "DecompressionTools.h"
 #include "EvaluationVM/EvaluationVM.h"
@@ -38,59 +39,39 @@ void FAnimNextAnimSequenceKeyframeTask::Execute(UE::AnimNext::FEvaluationVM& VM)
 		FDeltaTimeRecord DeltaTimeRecord;	// Not needed
 		const bool bExtractRootMotion = bExtractTrajectory;
 		const bool bLooping = false;		// Not needed
-		const bool bUseRawData = false;
 
 		const FAnimExtractContext ExtractionContext(SampleTime, bExtractRootMotion, DeltaTimeRecord, bLooping);
 
 		FKeyframeState Keyframe = VM.MakeUninitializedKeyframe(bIsAdditive);
 
+		const bool bUseRawData = FDecompressionTools::ShouldUseRawData(AnimSequencePtr, Keyframe.Pose);
+
 		if (EnumHasAnyFlags(VM.GetFlags(), EEvaluationFlags::Bones))
 		{
-			FDecompressionTools::GetAnimationPose(AnimSequencePtr, Keyframe.Pose, ExtractionContext);
+			FDecompressionTools::GetAnimationPose(AnimSequencePtr, ExtractionContext, Keyframe.Pose, bUseRawData);
 		}
 
 		if (EnumHasAnyFlags(VM.GetFlags(), EEvaluationFlags::Curves))
 		{
-			AnimSequencePtr->EvaluateCurveData(Keyframe.Curves, static_cast<float>(SampleTime), bUseRawData);
+			FDecompressionTools::GetAnimationCurves(AnimSequencePtr, ExtractionContext, Keyframe.Curves, bUseRawData);
 		}
 
 		if (EnumHasAnyFlags(VM.GetFlags(), EEvaluationFlags::Attributes))
 		{
-			// TODO: A few notes here:
-			// The attributes store a compact/lod bone index, this needs fixup if we feed the output into an AnimBP or whoever reads this later
-			// To be able to sample attributes, we need to sample bones as well, otherwise we'll have no refpose object
-			// We could assign the ref pose even if we don't sample bones, that would allow us to be more selective
-
-			QUICK_SCOPE_CYCLE_COUNTER(STAT_EvaluateAttributes);
-
-#if WITH_EDITOR
-			if (bUseRawData)
-			{
-				AnimSequencePtr->ValidateModel();
+			FDecompressionTools::GetAnimationAttributes(AnimSequencePtr, ExtractionContext, Keyframe.Pose.GetRefPose(), Keyframe.Attributes, bUseRawData);
+		}
 
-				for (const FAnimatedBoneAttribute& Attribute : AnimSequencePtr->GetDataModel()->GetAttributes())
-				{
-					const int32 LODBoneIndex = Keyframe.Pose.GetRefPose().GetLODBoneIndexFromSkeletonBoneIndex(Attribute.Identifier.GetBoneIndex());
-					// Only add attribute if the bone its tied to exists in the currently evaluated set of bones
-					if (LODBoneIndex != INDEX_NONE)
-					{
-						UE::Anim::Attributes::GetAttributeValue(Keyframe.Attributes, FCompactPoseBoneIndex(LODBoneIndex), Attribute, ExtractionContext.CurrentTime);
-					}
-				}
-			}
-			else
-#endif // WITH_EDITOR
+		// Trajectory is currently held as an attribute
+		if (EnumHasAnyFlags(VM.GetFlags(), EEvaluationFlags::Attributes | EEvaluationFlags::Trajectory))
+		{
+			// If the sequence has root motion enabled, allow sampling of a root motion delta into the custom attribute container of the outgoing pose
+			if (AnimSequence->HasRootMotion())
 			{
-				for (const TPair<FAnimationAttributeIdentifier, FAttributeCurve>& BakedAttribute : AnimSequencePtr->AttributeCurves)
+				// TODO: We should cache the provider in the VM
+				// We have to grab two locks to get it and it won't change during graph evaluation
+				if (const UE::Anim::IAnimRootMotionProvider* RootMotionProvider = UE::Anim::IAnimRootMotionProvider::Get())
 				{
-					const int32 LODBoneIndex = Keyframe.Pose.GetRefPose().GetLODBoneIndexFromSkeletonBoneIndex(BakedAttribute.Key.GetBoneIndex());
-					// Only add attribute if the bone its tied to exists in the currently evaluated set of bones
-					if (LODBoneIndex != INDEX_NONE)
-					{
-						UE::Anim::FAttributeId Info(BakedAttribute.Key.GetName(), FCompactPoseBoneIndex(LODBoneIndex));
-						uint8* AttributePtr = Keyframe.Attributes.FindOrAdd(BakedAttribute.Key.GetType(), Info);
-						BakedAttribute.Value.EvaluateToPtr(BakedAttribute.Key.GetType(), ExtractionContext.CurrentTime, AttributePtr);
-					}
+					RootMotionProvider->SampleRootMotion(ExtractionContext.DeltaTimeRecord, *AnimSequence, ExtractionContext.bLooping, Keyframe.Attributes);
 				}
 			}
 		}
diff --git a/Engine/Plugins/Experimental/AnimNext/Source/AnimNext/Private/GenerationTools.cpp b/Engine/Plugins/Experimental/AnimNext/Source/AnimNext/Private/GenerationTools.cpp
--- a/Engine/Plugins/Experimental/AnimNext/Source/AnimNext/Private/GenerationTools.cpp
+++ b/Engine/Plugins/Experimental/AnimNext/Source/AnimNext/Private/GenerationTools.cpp
@@ -146,7 +146,37 @@ struct FCompareBoneIndexType
 					LODBoneIndexToSkeletonBoneIndexMap[LODBoneIndex] = static_cast<FBoneIndexType>(SkeletonBoneIndex.GetInt());
 				}
 
-				OutAnimationReferencePose.Initialize(SkeletalMesh->GetRefSkeleton(), { LODBoneIndexToMeshBoneIndexMap }, { LODBoneIndexToSkeletonBoneIndexMap }, { SkeletonBoneIndexToLODBoneIndexMap }, LODNumBones, bCanGenerateSingleBonesList);
+				TArray<FBoneIndexType> MeshBoneIndexToLODBoneIndex;
+				MeshBoneIndexToLODBoneIndex.Init(INDEX_NONE, NumOrderedBones);
+
+				for (int32 LODBoneIndex = 0; LODBoneIndex < NumOrderedBones; ++LODBoneIndex)
+				{
+					const int32 MeshBoneIndex = LODBoneIndexToMeshBoneIndexMap[LODBoneIndex];
+					MeshBoneIndexToLODBoneIndex[MeshBoneIndex] = LODBoneIndex;
+				}
+
+				const FReferenceSkeleton& RefSkeleton = SkeletalMesh->GetRefSkeleton();
+
+				TArray<FBoneIndexType> LODBoneIndexToParentLODBoneIndexMap;
+				LODBoneIndexToParentLODBoneIndexMap.Init(INDEX_NONE, NumOrderedBones);
+
+				for (int32 LODBoneIndex = 0; LODBoneIndex < NumOrderedBones; ++LODBoneIndex)
+				{
+					const int32 MeshBoneIndex = LODBoneIndexToMeshBoneIndexMap[LODBoneIndex];
+					const int32 ParentMeshBoneIndex = RefSkeleton.GetParentIndex(MeshBoneIndex);
+					const int32 ParentLODBoneIndex = ParentMeshBoneIndex != INDEX_NONE ? MeshBoneIndexToLODBoneIndex[ParentMeshBoneIndex] : INDEX_NONE;
+
+					LODBoneIndexToParentLODBoneIndexMap[LODBoneIndex] = ParentLODBoneIndex;
+				}
+
+				OutAnimationReferencePose.Initialize(
+					RefSkeleton,
+					{ LODBoneIndexToParentLODBoneIndexMap },
+					{ LODBoneIndexToMeshBoneIndexMap },
+					{ LODBoneIndexToSkeletonBoneIndexMap },
+					{ SkeletonBoneIndexToLODBoneIndexMap },
+					LODNumBones,
+					bCanGenerateSingleBonesList);
 
 				ReferencePoseGenerated = true;
 			}
@@ -566,37 +596,116 @@ void FGenerationTools::FixLODRequiredBones(const int32 NumLODs
 	}
 }
 
+void FGenerationTools::RemapAttributes(
+	const FLODPose& LODPose,
+	const UE::Anim::FHeapAttributeContainer& InAttributes,
+	UE::Anim::FMeshAttributeContainer& OutAttributes)
+{
+	const TArrayView<const FBoneIndexType> LODBoneIndexToMeshBoneIndexMap = LODPose.GetLODBoneIndexToMeshBoneIndexMap();
+
+	for (const TWeakObjectPtr<UScriptStruct> WeakScriptStruct : InAttributes.GetUniqueTypes())
+	{
+		const UScriptStruct* ScriptStruct = WeakScriptStruct.Get();
+		const int32 TypeIndex = InAttributes.FindTypeIndex(ScriptStruct);
+		if (TypeIndex != INDEX_NONE)
+		{
+			const TArray<UE::Anim::TWrappedAttribute<FDefaultAllocator>, FDefaultAllocator>& SourceValues = InAttributes.GetValues(TypeIndex);
+			const TArray<UE::Anim::FAttributeId, FDefaultAllocator>& AttributeIds = InAttributes.GetKeys(TypeIndex);
+
+			// Try and remap all the source attributes to their respective new bone indices
+			for (int32 EntryIndex = 0; EntryIndex < AttributeIds.Num(); ++EntryIndex)
+			{
+				const UE::Anim::FAttributeId& AttributeId = AttributeIds[EntryIndex];
+
+				const int32 LODBoneIndex = AttributeId.GetIndex();
+				const int32 MeshBoneIndex = LODBoneIndexToMeshBoneIndexMap[LODBoneIndex];
+
+				const UE::Anim::FAttributeId NewInfo(AttributeId.GetName(), MeshBoneIndex, AttributeId.GetNamespace());
+				uint8* NewAttribute = OutAttributes.FindOrAdd(ScriptStruct, NewInfo);
+				ScriptStruct->CopyScriptStruct(NewAttribute, SourceValues[EntryIndex].GetPtr<void>());
+			}
+		}
+	}
+}
+
+void FGenerationTools::RemapAttributes(
+	const FLODPose& LODPose,
+	const UE::Anim::FHeapAttributeContainer& InAttributes,
+	FPoseContext& OutPose)
+{
+	const TArrayView<const FBoneIndexType> LODBoneIndexToMeshBoneIndexMap = LODPose.GetLODBoneIndexToMeshBoneIndexMap();
+	const FBoneContainer& BoneContainer = OutPose.Pose.GetBoneContainer();
+
+	for (const TWeakObjectPtr<UScriptStruct> WeakScriptStruct : InAttributes.GetUniqueTypes())
+	{
+		const UScriptStruct* ScriptStruct = WeakScriptStruct.Get();
+		const int32 TypeIndex = InAttributes.FindTypeIndex(ScriptStruct);
+		if (TypeIndex != INDEX_NONE)
+		{
+			const TArray<UE::Anim::TWrappedAttribute<FDefaultAllocator>, FDefaultAllocator>& SourceValues = InAttributes.GetValues(TypeIndex);
+			const TArray<UE::Anim::FAttributeId, FDefaultAllocator>& AttributeIds = InAttributes.GetKeys(TypeIndex);
+
+			// Try and remap all the source attributes to their respective new bone indices
+			for (int32 EntryIndex = 0; EntryIndex < AttributeIds.Num(); ++EntryIndex)
+			{
+				const UE::Anim::FAttributeId& AttributeId = AttributeIds[EntryIndex];
+
+				const int32 LODBoneIndex = AttributeId.GetIndex();
+				const FMeshPoseBoneIndex MeshBoneIndex(LODBoneIndexToMeshBoneIndexMap[LODBoneIndex]);
+
+				// Remap our skeletal mesh bone index into the skeleton bone index we output for
+				const FSkeletonPoseBoneIndex SkeletonBoneIndex = BoneContainer.GetSkeletonPoseIndexFromMeshPoseIndex(MeshBoneIndex);
+				ensure(SkeletonBoneIndex.IsValid());	// We expect the skeletal mesh bone to map to a valid skeleton bone
+
+				// Remap our skeleton bone index into the compact pose bone index we output for
+				const FCompactPoseBoneIndex CompactBoneIndex = BoneContainer.GetCompactPoseIndexFromSkeletonPoseIndex(SkeletonBoneIndex);
+
+				if (ensure(CompactBoneIndex.IsValid()))	// We expect the skeleton bone to map to a valid compact pose bone
+				{
+					const UE::Anim::FAttributeId NewInfo(AttributeId.GetName(), CompactBoneIndex);
+					uint8* NewAttribute = OutPose.CustomAttributes.FindOrAdd(ScriptStruct, NewInfo);
+					ScriptStruct->CopyScriptStruct(NewAttribute, SourceValues[EntryIndex].GetPtr<void>());
+				}
+				else
+				{
+					// This bone is part of the LOD but isn't part of the required bones
+				}
+			}
+		}
+	}
+}
+
 void FGenerationTools::ConvertLocalSpaceToComponentSpace(
-												 TConstArrayView<FBoneIndexType> InParentIndices,
+												 TConstArrayView<FBoneIndexType> InMeshBoneIndexToParentMeshBoneIndexMap,
 												 TConstArrayView<FTransform> InBoneSpaceTransforms,
-												 TConstArrayView<FBoneIndexType> InRequiredBoneIndices, 
+												 TConstArrayView<FBoneIndexType> InLODBoneIndexToMeshBoneIndexMap,
 												 TArrayView<FTransform> OutComponentSpaceTransforms)
 {
 	SCOPE_CYCLE_COUNTER(STAT_AnimNext_ConvertLocalSpaceToComponentSpace);
 
-	const int32 NumBones = InParentIndices.Num();
-	checkf(NumBones == InBoneSpaceTransforms.Num(), TEXT("Buffer mismatch: %d:%d"), NumBones, InBoneSpaceTransforms.Num());
-	checkf(NumBones == OutComponentSpaceTransforms.Num(), TEXT("Buffer mismatch: %d:%d"), NumBones, OutComponentSpaceTransforms.Num());
+	checkf(InMeshBoneIndexToParentMeshBoneIndexMap.Num() == InBoneSpaceTransforms.Num(), TEXT("Buffer mismatch: %d:%d"), InMeshBoneIndexToParentMeshBoneIndexMap.Num(), InBoneSpaceTransforms.Num());
+	checkf(InMeshBoneIndexToParentMeshBoneIndexMap.Num() == OutComponentSpaceTransforms.Num(), TEXT("Buffer mismatch: %d:%d"), InMeshBoneIndexToParentMeshBoneIndexMap.Num(), OutComponentSpaceTransforms.Num());
 
-	const FBoneIndexType* RESTRICT RequiredBoneIndicesData = InRequiredBoneIndices.GetData();
-	const FBoneIndexType* RESTRICT ParentIndicesData = InParentIndices.GetData();
+	const FBoneIndexType* RESTRICT LODBoneIndexToMeshBoneIndexMapPtr = InLODBoneIndexToMeshBoneIndexMap.GetData();
+	const FBoneIndexType* RESTRICT MeshBoneIndexToParentMeshBoneIndexMapPtr = InMeshBoneIndexToParentMeshBoneIndexMap.GetData();
 	const FTransform* RESTRICT LocalTransformsData = InBoneSpaceTransforms.GetData();
 	FTransform* RESTRICT ComponentSpaceData = OutComponentSpaceTransforms.GetData();
 
 	// First bone (if we have one) is always root bone, and it doesn't have a parent.
 	{
-		check(InRequiredBoneIndices.Num() == 0 || InRequiredBoneIndices[0] == 0);
+		check(InLODBoneIndexToMeshBoneIndexMap.Num() == 0 || InLODBoneIndexToMeshBoneIndexMap[0] == 0);
 		OutComponentSpaceTransforms[0] = InBoneSpaceTransforms[0];
 	}
 
-	const int32 NumRequiredBones = InRequiredBoneIndices.Num();
-	for (int32 RequiredBoneIndex = 1; RequiredBoneIndex < NumRequiredBones; ++RequiredBoneIndex)
+	const int32 NumLODBones = InLODBoneIndexToMeshBoneIndexMap.Num();
+	for (int32 LODBoneIndex = 1; LODBoneIndex < NumLODBones; ++LODBoneIndex)
 	{
-		const FBoneIndexType* BoneIndex = RequiredBoneIndicesData + RequiredBoneIndex;
-		FTransform* RESTRICT ComponentSpaceTransform = ComponentSpaceData + *BoneIndex;
-		const FBoneIndexType* RESTRICT ParentIndex = ParentIndicesData + *BoneIndex;
-		const FTransform* RESTRICT ParentComponentSpaceTransform = ComponentSpaceData + *ParentIndex;
-		const FTransform* RESTRICT LocalSpaceTransform = LocalTransformsData + *BoneIndex;
+		const FBoneIndexType MeshBoneIndex = LODBoneIndexToMeshBoneIndexMapPtr[LODBoneIndex];
+		const FBoneIndexType ParentMeshBoneIndex = MeshBoneIndexToParentMeshBoneIndexMapPtr[MeshBoneIndex];
+
+		FTransform* RESTRICT ComponentSpaceTransform = ComponentSpaceData + MeshBoneIndex;
+		const FTransform* RESTRICT ParentComponentSpaceTransform = ComponentSpaceData + ParentMeshBoneIndex;
+		const FTransform* RESTRICT LocalSpaceTransform = LocalTransformsData + MeshBoneIndex;
 
 		FTransform::Multiply(ComponentSpaceTransform, LocalSpaceTransform, ParentComponentSpaceTransform);
 
diff --git a/Engine/Plugins/Experimental/AnimNext/Source/AnimNext/Private/Graph/RigUnit_AnimNextRunAnimationGraph.cpp b/Engine/Plugins/Experimental/AnimNext/Source/AnimNext/Private/Graph/RigUnit_AnimNextRunAnimationGraph.cpp
--- a/Engine/Plugins/Experimental/AnimNext/Source/AnimNext/Private/Graph/RigUnit_AnimNextRunAnimationGraph.cpp
+++ b/Engine/Plugins/Experimental/AnimNext/Source/AnimNext/Private/Graph/RigUnit_AnimNextRunAnimationGraph.cpp
@@ -1,4 +1,4 @@
-﻿// Copyright Epic Games, Inc. All Rights Reserved.
+// Copyright Epic Games, Inc. All Rights Reserved.
 
 #include "Graph/RigUnit_AnimNextRunAnimationGraph.h"
 
@@ -105,7 +105,7 @@ FRigUnit_AnimNextRunAnimationGraph_Execute()
 
 	IAnimNextModuleInterface& AnimNextModule = IAnimNextModuleInterface::Get();
 	AnimNextModule.UpdateGraph(Instance, ExecuteContext.GetDeltaTime(), InputEventList, OutputEventList);
-	AnimNextModule.EvaluateGraph(Instance, RefPose, LOD, Result.LODPose);
+	AnimNextModule.EvaluateGraph(Instance, RefPose, LOD, Result);
 
 	// We might have appended new input/output events, append them
 	{
diff --git a/Engine/Plugins/Experimental/AnimNext/Source/AnimNext/Private/Graph/RigUnit_AnimNextWriteSkeletalMeshComponentPose.cpp b/Engine/Plugins/Experimental/AnimNext/Source/AnimNext/Private/Graph/RigUnit_AnimNextWriteSkeletalMeshComponentPose.cpp
--- a/Engine/Plugins/Experimental/AnimNext/Source/AnimNext/Private/Graph/RigUnit_AnimNextWriteSkeletalMeshComponentPose.cpp
+++ b/Engine/Plugins/Experimental/AnimNext/Source/AnimNext/Private/Graph/RigUnit_AnimNextWriteSkeletalMeshComponentPose.cpp
@@ -46,10 +46,15 @@ FRigUnit_AnimNextWriteSkeletalMeshComponentPose_Execute()
 	// Map LOD pose into local-space scratch buffer
 	FGenerationTools::RemapPose(Pose.LODPose, LocalSpaceTransforms);
 
+	SkeletalMeshComponent->AnimCurves.CopyFrom(Pose.Curves);
+
+	// Attributes require remapping since the indices are LOD indices and we want mesh indices
+	FGenerationTools::RemapAttributes(Pose.LODPose, Pose.Attributes, SkeletalMeshComponent->CustomAttributes);
+
 	// Convert and dispatch to renderer
 	UE::Anim::FSkinnedMeshComponentExtensions::CompleteAndDispatch(
 		SkeletalMeshComponent,
-		RefPose.GetParentIndices(),
+		RefPose.GetMeshBoneIndexToParentMeshBoneIndexMap(),
 		RefPose.GetLODBoneIndexToMeshBoneIndexMap(Pose.LODPose.LODLevel),
 		LocalSpaceTransforms);
 }
diff --git a/Engine/Plugins/Experimental/AnimNext/Source/AnimNext/Public/DecompressionTools.h b/Engine/Plugins/Experimental/AnimNext/Source/AnimNext/Public/DecompressionTools.h
--- a/Engine/Plugins/Experimental/AnimNext/Source/AnimNext/Public/DecompressionTools.h
+++ b/Engine/Plugins/Experimental/AnimNext/Source/AnimNext/Public/DecompressionTools.h
@@ -5,6 +5,8 @@
 #include "CoreMinimal.h"
 #include "ReferencePose.h"
 #include "LODPose.h"
+#include "Animation/AnimCurveTypes.h"
+#include "Animation/AttributesRuntime.h"
 
 struct FCompactPose;
 struct FCompressedAnimSequence;
@@ -19,12 +21,12 @@ namespace UE::AnimNext
 class ANIMNEXT_API FDecompressionTools
 {
 public:
+	// Returns whether decompression should use raw data or not
+	static bool ShouldUseRawData(const UAnimSequence* AnimSequence, const FLODPose& AnimationPoseData);
 
 	// Extracted from UAnimSequence
 	// Extracts Animation Data from the provided AnimSequence, using AnimExtractContext extraction parameters
-	static void GetAnimationPose(const UAnimSequence* AnimSequence
-		, FLODPose& OutAnimationPoseData
-		, const FAnimExtractContext& ExtractionContext);
+	static void GetAnimationPose(const UAnimSequence* AnimSequence, const FAnimExtractContext& ExtractionContext, FLODPose& OutAnimationPoseData, bool bForceUseRawData = false);
 
 	/**
 	* Get Bone Transform of the Time given, relative to Parent for all RequiredBones
@@ -35,6 +37,9 @@ public:
 	static void GetBonePose_Additive(const UAnimSequence* AnimSequence, const FAnimExtractContext& ExtractionContext, FLODPose& OutAnimationPoseData);
 	static void GetBonePose_AdditiveMeshRotationOnly(const UAnimSequence* AnimSequence, const FAnimExtractContext& ExtractionContext, FLODPose& OutAnimationPoseData);
 
+	static void GetAnimationCurves(const UAnimSequence* AnimSequence, const FAnimExtractContext& ExtractionContext, FBlendedCurve& OutCurves, bool bForceUseRawData = false);
+	static void GetAnimationAttributes(const UAnimSequence* AnimSequence, const FAnimExtractContext& ExtractionContext, const FReferencePose& RefPose, UE::Anim::FStackAttributeContainer& OutAttributes, bool bForceUseRawData = false);
+
 	// Decompress and retarget animation data using provided RetargetTransforms
 	static void DecompressPose(FLODPose& OutAnimationPoseData,
 		const FCompressedAnimSequence& CompressedData,
diff --git a/Engine/Plugins/Experimental/AnimNext/Source/AnimNext/Public/GenerationTools.h b/Engine/Plugins/Experimental/AnimNext/Source/AnimNext/Public/GenerationTools.h
--- a/Engine/Plugins/Experimental/AnimNext/Source/AnimNext/Public/GenerationTools.h
+++ b/Engine/Plugins/Experimental/AnimNext/Source/AnimNext/Public/GenerationTools.h
@@ -3,6 +3,7 @@
 #pragma once
 
 #include "BoneIndices.h"
+#include "Animation/AttributesRuntime.h"
 #include "Containers/IndirectArray.h"
 #include "ReferencePose.h"
 #include "LODPose.h"
@@ -106,6 +107,18 @@ public:
 	// in by this call if they are not affected by the current LOD.
 	static void RemapPose(const FLODPose& SourcePose, TArrayView<FTransform> TargetTransforms);
 
+	// Converts AnimNext attributes to AnimBP attributes
+	static void RemapAttributes(
+		const FLODPose& LODPose,
+		const UE::Anim::FHeapAttributeContainer& InAttributes,
+		UE::Anim::FMeshAttributeContainer& OutAttributes);
+
+	// Converts AnimNext attributes to AnimBP attributes
+	static void RemapAttributes(
+		const FLODPose& LODPose,
+		const UE::Anim::FHeapAttributeContainer& InAttributes,
+		FPoseContext& OutPose);
+
 	// Converts a local space to component space buffer given a number of required bones
 	static void ConvertLocalSpaceToComponentSpace(TConstArrayView<FBoneIndexType> InParentIndices, TConstArrayView<FTransform> InBoneSpaceTransforms, TConstArrayView<FBoneIndexType> InRequiredBoneIndices, TArrayView<FTransform> OutComponentSpaceTransforms);
 };
diff --git a/Engine/Plugins/Experimental/AnimNext/Source/AnimNext/Public/IAnimNextModuleInterface.h b/Engine/Plugins/Experimental/AnimNext/Source/AnimNext/Public/IAnimNextModuleInterface.h
--- a/Engine/Plugins/Experimental/AnimNext/Source/AnimNext/Public/IAnimNextModuleInterface.h
+++ b/Engine/Plugins/Experimental/AnimNext/Source/AnimNext/Public/IAnimNextModuleInterface.h
@@ -4,11 +4,11 @@
 
 #include "CoreMinimal.h"
 #include "Modules/ModuleInterface.h"
-#include "LODPose.h"
 #include "ReferencePose.h"
 
 struct FAnimNextGraphInstancePtr;
 struct FAnimNextVariableBindingData;
+struct FAnimNextGraphLODPose;
 
 namespace UE::AnimNext
 {
@@ -26,7 +26,7 @@ public:
 	virtual ~IAnimNextAnimGraph() = default;
 
 	virtual void UpdateGraph(const FAnimNextGraphInstancePtr& GraphInstance, float DeltaTime, FTraitEventList& InputEventList, FTraitEventList& OutputEventList) const = 0;
-	virtual void EvaluateGraph(const FAnimNextGraphInstancePtr& GraphInstance, const FReferencePose& RefPose, int32 GraphLODLevel, FLODPoseHeap& OutputPose) const = 0;
+	virtual void EvaluateGraph(const FAnimNextGraphInstancePtr& GraphInstance, const FReferencePose& RefPose, int32 GraphLODLevel, FAnimNextGraphLODPose& OutputPose) const = 0;
 };
 
 class IAnimNextModuleInterface : public IModuleInterface
@@ -43,7 +43,7 @@ public:
 	virtual void RegisterAnimNextAnimGraph(const IAnimNextAnimGraph& InAnimGraphImpl) = 0;
 	virtual void UnregisterAnimNextAnimGraph() = 0;
 	virtual void UpdateGraph(const FAnimNextGraphInstancePtr& GraphInstance, float DeltaTime, FTraitEventList& InputEventList, FTraitEventList& OutputEventList) = 0;
-	virtual void EvaluateGraph(const FAnimNextGraphInstancePtr& GraphInstance, const FReferencePose& RefPose, int32 GraphLODLevel, FLODPoseHeap& OutputPose) const = 0;
+	virtual void EvaluateGraph(const FAnimNextGraphInstancePtr& GraphInstance, const FReferencePose& RefPose, int32 GraphLODLevel, FAnimNextGraphLODPose& OutputPose) const = 0;
 };
 
 }
diff --git a/Engine/Plugins/Experimental/AnimNext/Source/AnimNext/Public/LODPose.h b/Engine/Plugins/Experimental/AnimNext/Source/AnimNext/Public/LODPose.h
--- a/Engine/Plugins/Experimental/AnimNext/Source/AnimNext/Public/LODPose.h
+++ b/Engine/Plugins/Experimental/AnimNext/Source/AnimNext/Public/LODPose.h
@@ -81,6 +81,18 @@ struct FLODPose
 		return RefPose != nullptr ? RefPose->GetNumBonesForLOD(LODLevel) : 0;
 	}
 
+	const TArrayView<const FBoneIndexType> GetLODBoneIndexToParentLODBoneIndexMap() const
+	{
+		if (LODLevel != INVALID_LOD_LEVEL && RefPose != nullptr)
+		{
+			return RefPose->GetLODBoneIndexToParentLODBoneIndexMap(LODLevel);
+		}
+		else
+		{
+			return TArrayView<const FBoneIndexType>();
+		}
+	}
+
 	const TArrayView<const FBoneIndexType> GetLODBoneIndexToMeshBoneIndexMap() const
 	{
 		if (LODLevel != INVALID_LOD_LEVEL && RefPose != nullptr)
diff --git a/Engine/Plugins/Experimental/AnimNext/Source/AnimNext/Public/ReferencePose.h b/Engine/Plugins/Experimental/AnimNext/Source/AnimNext/Public/ReferencePose.h
--- a/Engine/Plugins/Experimental/AnimNext/Source/AnimNext/Public/ReferencePose.h
+++ b/Engine/Plugins/Experimental/AnimNext/Source/AnimNext/Public/ReferencePose.h
@@ -34,6 +34,11 @@ struct TReferencePose
 	// Higher LOD come first
 	TTransformArray<AllocatorType> ReferenceLocalTransforms;
 
+	// A mapping of LOD sorted bone indices to their parent LOD sorted bone indices per LOD
+	// Each list of bone indices is a mapping of: LODSortedBoneIndex -> LODSortedBoneIndex
+	// When fast path is enabled, we have a single LOD entry that we truncate to the number of bones for each LOD
+	TArray<TArray<FBoneIndexType, AllocatorType>, AllocatorType> LODBoneIndexToParentLODBoneIndexMapPerLOD;
+
 	// A mapping of LOD sorted bone indices to skeletal mesh indices per LOD
 	// Each list of bone indices is a mapping of: LODSortedBoneIndex -> SkeletalMeshBoneIndex
 	// When fast path is enabled, we have a single LOD entry that we truncate to the number of bones for each LOD
@@ -53,7 +58,7 @@ struct TReferencePose
 	TArray<int32, AllocatorType> LODNumBones;
 
 	// Mapping of mesh bone indices to mesh parent indices for each bone
-	TArray<FBoneIndexType, AllocatorType> ParentIndices;
+	TArray<FBoneIndexType, AllocatorType> MeshBoneIndexToParentMeshBoneIndexMap;
 
 	TWeakObjectPtr<const USkeletalMesh> SkeletalMesh = nullptr;
 	TWeakObjectPtr<const USkeleton> Skeleton = nullptr;
@@ -79,6 +84,7 @@ struct TReferencePose
 	}
 
 	void Initialize(const FReferenceSkeleton& RefSkeleton
+		, const TArray<TArray<FBoneIndexType>>& InLODBoneIndexToParentLODBoneIndexMapPerLOD
 		, const TArray<TArray<FBoneIndexType>>& InLODBoneIndexToMeshBoneIndexMapPerLOD
 		, const TArray<TArray<FBoneIndexType>>& InLODBoneIndexToSkeletonBoneIndexMapPerLOD
 		, const TArray<TArray<FBoneIndexType>>& InSkeletonBoneIndexToLODBoneIndexMapPerLOD
@@ -88,8 +94,9 @@ struct TReferencePose
 		const int32 NumBonesLOD0 = !InLODNumBones.IsEmpty() ? InLODNumBones[0] : 0;
 		const int32 NumBonesMesh = RefSkeleton.GetRefBoneInfo().Num();
 		
-		ParentIndices.SetNum(NumBonesMesh);
+		MeshBoneIndexToParentMeshBoneIndexMap.SetNum(NumBonesMesh);
 		ReferenceLocalTransforms.SetNum(NumBonesLOD0);
+		LODBoneIndexToParentLODBoneIndexMapPerLOD = InLODBoneIndexToParentLODBoneIndexMapPerLOD;
 		LODBoneIndexToMeshBoneIndexMapPerLOD = InLODBoneIndexToMeshBoneIndexMapPerLOD;
 		LODBoneIndexToSkeletonBoneIndexMapPerLOD = InLODBoneIndexToSkeletonBoneIndexMapPerLOD;
 		SkeletonBoneIndexToLODBoneIndexMapPerLOD = InSkeletonBoneIndexToLODBoneIndexMapPerLOD;
@@ -105,14 +112,30 @@ struct TReferencePose
 			ReferenceLocalTransforms[LODBoneIndex] = RefBonePose[BoneLODIndexToMeshIndexMap0[LODBoneIndex]];
 		}
 
-		for(int32 BoneIndex = 0; BoneIndex < NumBonesMesh; ++BoneIndex)
+		for (int32 BoneIndex = 0; BoneIndex < NumBonesMesh; ++BoneIndex)
 		{
-			ParentIndices[BoneIndex] = RefBoneInfo[BoneIndex].ParentIndex;
+			MeshBoneIndexToParentMeshBoneIndexMap[BoneIndex] = RefBoneInfo[BoneIndex].ParentIndex;
 		}
 
 		GenerationFlags = bFastPath ? EReferencePoseGenerationFlags::FastPath : EReferencePoseGenerationFlags::None;
 	}
 
+	// Returns a list of LOD sorted parent bone indices, a mapping of: LODSortedBoneIndex -> LODSortedBoneIndex
+	const TArrayView<const FBoneIndexType> GetLODBoneIndexToParentLODBoneIndexMap(int32 LODLevel) const
+	{
+		TArrayView<const FBoneIndexType> ArrayView;
+
+		if (LODLevel >= 0 && (IsFastPath() || LODLevel < LODBoneIndexToParentLODBoneIndexMapPerLOD.Num()))
+		{
+			const int32 NumBonesForLOD = GetNumBonesForLOD(LODLevel);
+			const int32 LODIndex = IsFastPath() ? 0 : LODLevel;
+
+			ArrayView = MakeArrayView(LODBoneIndexToParentLODBoneIndexMapPerLOD[LODIndex].GetData(), NumBonesForLOD);
+		}
+
+		return ArrayView;
+	}
+
 	// Returns a list of LOD sorted skeletal mesh bone indices, a mapping of: LODSortedBoneIndex -> SkeletalMeshBoneIndex
 	const TArrayView<const FBoneIndexType> GetLODBoneIndexToMeshBoneIndexMap(int32 LODLevel) const
 	{
@@ -167,9 +190,9 @@ struct TReferencePose
 	}
 
 	// Returns a mapping of mesh bone indices to mesh parent indices for each bone
-	TConstArrayView<FBoneIndexType> GetParentIndices() const
+	TConstArrayView<FBoneIndexType> GetMeshBoneIndexToParentMeshBoneIndexMap() const
 	{
-		return ParentIndices;
+		return MeshBoneIndexToParentMeshBoneIndexMap;
 	}
 
 	int32 GetSkeletonBoneIndexFromLODBoneIndex(int32 LODBoneIndex) const
diff --git a/Engine/Plugins/Experimental/AnimNext/Source/AnimNext/Public/TransformArray.h b/Engine/Plugins/Experimental/AnimNext/Source/AnimNext/Public/TransformArray.h
--- a/Engine/Plugins/Experimental/AnimNext/Source/AnimNext/Public/TransformArray.h
+++ b/Engine/Plugins/Experimental/AnimNext/Source/AnimNext/Public/TransformArray.h
@@ -503,21 +503,21 @@ struct TTransformArraySoA
 
 	bool ContainsNaN() const
 	{
-		for (const auto& Rotation : Rotations)
+		for (const FQuat& Rotation : Rotations)
 		{
 			if (Rotation.ContainsNaN())
 			{
 				return true;
 			}
 		}
-		for (const auto& Translation : Translations)
+		for (const FVector& Translation : Translations)
 		{
 			if (Translation.ContainsNaN())
 			{
 				return true;
 			}
 		}
-		for (const auto& Scale3D : Scales3D)
+		for (const FVector& Scale3D : Scales3D)
 		{
 			if (Scale3D.ContainsNaN())
 			{
@@ -535,7 +535,7 @@ struct TTransformArraySoA
 			return false;
 		}
 
-		for (const auto& Rotation : Rotations)
+		for (const FQuat& Rotation : Rotations)
 		{
 			if (Rotation.IsNormalized() == false)
 			{
@@ -596,15 +596,15 @@ private:
 
 	FORCEINLINE void DiagnosticCheckNaN_All() const
 	{
-		for (const auto& Rotation : Rotations)
+		for (const FQuat& Rotation : Rotations)
 		{
 			DiagnosticCheckNaN_Rotate(Rotation);
 		}
-		for (const auto& Translation : Translations)
+		for (const FVector& Translation : Translations)
 		{
 			DiagnosticCheckNaN_Translate(Translation);
 		}
-		for (const auto& Scale3D : Scales3D)
+		for (const FVector& Scale3D : Scales3D)
 		{
 			DiagnosticCheckNaN_Scale3D(Scale3D);
 		}
diff --git a/Engine/Plugins/Experimental/AnimNext/Source/AnimNext/Public/TransformArrayView.h b/Engine/Plugins/Experimental/AnimNext/Source/AnimNext/Public/TransformArrayView.h
--- a/Engine/Plugins/Experimental/AnimNext/Source/AnimNext/Public/TransformArrayView.h
+++ b/Engine/Plugins/Experimental/AnimNext/Source/AnimNext/Public/TransformArrayView.h
@@ -105,6 +105,35 @@ namespace UE::AnimNext
 		{
 			return FTransformSoAAdapterConst(Rotations[Index], Translations[Index], Scales3D[Index]);
 		}
+
+		bool IsValid() const
+		{
+			for (const FQuat& Rotation : Rotations)
+			{
+				if (Rotation.ContainsNaN() || !Rotation.IsNormalized())
+				{
+					return false;
+				}
+			}
+
+			for (const FVector& Translation : Translations)
+			{
+				if (Translation.ContainsNaN())
+				{
+					return false;
+				}
+			}
+
+			for (const FVector& Scale3D : Scales3D)
+			{
+				if (Scale3D.ContainsNaN())
+				{
+					return false;
+				}
+			}
+
+			return true;
+		}
 	};
 
 	struct FTransformArraySoAConstView
@@ -127,6 +156,35 @@ namespace UE::AnimNext
 		{
 			return FTransformSoAAdapterConst(Rotations[Index], Translations[Index], Scales3D[Index]);
 		}
+
+		bool IsValid() const
+		{
+			for (const FQuat& Rotation : Rotations)
+			{
+				if (Rotation.ContainsNaN() || !Rotation.IsNormalized())
+				{
+					return false;
+				}
+			}
+
+			for (const FVector& Translation : Translations)
+			{
+				if (Translation.ContainsNaN())
+				{
+					return false;
+				}
+			}
+
+			for (const FVector& Scale3D : Scales3D)
+			{
+				if (Scale3D.ContainsNaN())
+				{
+					return false;
+				}
+			}
+
+			return true;
+		}
 	};
 
 #define DEFAULT_SOA_VIEW 1
diff --git a/Engine/Plugins/Experimental/AnimNextAnimGraph/Source/AnimNextAnimGraph/Private/Graph/AnimGraph/AnimNode_AnimNextGraph.cpp b/Engine/Plugins/Experimental/AnimNextAnimGraph/Source/AnimNextAnimGraph/Private/Graph/AnimGraph/AnimNode_AnimNextGraph.cpp
--- a/Engine/Plugins/Experimental/AnimNextAnimGraph/Source/AnimNextAnimGraph/Private/Graph/AnimGraph/AnimNode_AnimNextGraph.cpp
+++ b/Engine/Plugins/Experimental/AnimNextAnimGraph/Source/AnimNextAnimGraph/Private/Graph/AnimGraph/AnimNode_AnimNextGraph.cpp
@@ -111,7 +111,8 @@ void FAnimNode_AnimNextGraph::Evaluate_AnyThread(FPoseContext& Output)
 		const int32 LODLevel = Output.AnimInstanceProxy->GetLODLevel();
 
 		const UE::AnimNext::FReferencePose& RefPose = RefPoseHandle.GetRef<UE::AnimNext::FReferencePose>();
-		FAnimNextGraphLODPose ResultPose(FLODPoseHeap(RefPose, LODLevel, true, Output.ExpectsAdditivePose()));
+		FAnimNextGraphLODPose ResultPose;
+		ResultPose.LODPose = FLODPoseHeap(RefPose, LODLevel, true, Output.ExpectsAdditivePose());
 
 		{
 			const FEvaluationProgram EvaluationProgram = UE::AnimNext::EvaluateGraph(GraphInstance);
@@ -127,6 +128,8 @@ void FAnimNode_AnimNextGraph::Evaluate_AnyThread(FPoseContext& Output)
 				if (EvaluationVM.PopValue(KEYFRAME_STACK_NAME, EvaluatedKeyframe))
 				{
 					ResultPose.LODPose.CopyFrom(EvaluatedKeyframe->Pose);
+					ResultPose.Curves.CopyFrom(EvaluatedKeyframe->Curves);
+					ResultPose.Attributes.CopyFrom(EvaluatedKeyframe->Attributes);
 					bHasValidOutput = true;
 				}
 			}
@@ -136,10 +139,14 @@ void FAnimNode_AnimNextGraph::Evaluate_AnyThread(FPoseContext& Output)
 				// We need to output a valid pose, generate one
 				FKeyframeState ReferenceKeyframe = EvaluationVM.MakeReferenceKeyframe(Output.ExpectsAdditivePose());
 				ResultPose.LODPose.CopyFrom(ReferenceKeyframe.Pose);
+				ResultPose.Curves.CopyFrom(ReferenceKeyframe.Curves);
+				ResultPose.Attributes.CopyFrom(ReferenceKeyframe.Attributes);
 			}
 		}
 
 		FGenerationTools::RemapPose(ResultPose.LODPose, Output);
+		Output.Curve.CopyFrom(ResultPose.Curves);
+		FGenerationTools::RemapAttributes(ResultPose.LODPose, ResultPose.Attributes, Output);
 	}
 	else
 	{
diff --git a/Engine/Plugins/Experimental/AnimNextAnimGraph/Source/AnimNextAnimGraph/Private/Module.cpp b/Engine/Plugins/Experimental/AnimNextAnimGraph/Source/AnimNextAnimGraph/Private/Module.cpp
--- a/Engine/Plugins/Experimental/AnimNextAnimGraph/Source/AnimNextAnimGraph/Private/Module.cpp
+++ b/Engine/Plugins/Experimental/AnimNextAnimGraph/Source/AnimNextAnimGraph/Private/Module.cpp
@@ -14,6 +14,7 @@
 #include "TraitInterfaces/IEvaluate.h"
 #include "TraitInterfaces/IUpdate.h"
 #include "EvaluationVM/EvaluationVM.h"
+#include "Graph/AnimNext_LODPose.h"
 
 namespace UE::AnimNext::AnimGraph
 {
@@ -48,7 +49,7 @@ public:
 		UE::AnimNext::UpdateGraph(GraphInstance, DeltaTime, InputEventList, OutputEventList);
 	}
 
-	virtual void EvaluateGraph(const FAnimNextGraphInstancePtr& GraphInstance, const UE::AnimNext::FReferencePose& RefPose, int32 GraphLODLevel, FLODPoseHeap& OutputPose) const override
+	virtual void EvaluateGraph(const FAnimNextGraphInstancePtr& GraphInstance, const UE::AnimNext::FReferencePose& RefPose, int32 GraphLODLevel, FAnimNextGraphLODPose& OutputPose) const override
 	{
 		const FEvaluationProgram EvaluationProgram = UE::AnimNext::EvaluateGraph(GraphInstance);
 
@@ -62,7 +63,9 @@ public:
 			TUniquePtr<FKeyframeState> EvaluatedKeyframe;
 			if (EvaluationVM.PopValue(KEYFRAME_STACK_NAME, EvaluatedKeyframe))
 			{
-				OutputPose.CopyFrom(EvaluatedKeyframe->Pose);
+				OutputPose.LODPose.CopyFrom(EvaluatedKeyframe->Pose);
+				OutputPose.Curves.CopyFrom(EvaluatedKeyframe->Curves);
+				OutputPose.Attributes.CopyFrom(EvaluatedKeyframe->Attributes);
 				bHasValidOutput = true;
 			}
 		}
@@ -71,7 +74,9 @@ public:
 		{
 			// We need to output a valid pose, generate one
 			FKeyframeState ReferenceKeyframe = EvaluationVM.MakeReferenceKeyframe(false);
-			OutputPose.CopyFrom(ReferenceKeyframe.Pose);
+			OutputPose.LODPose.CopyFrom(ReferenceKeyframe.Pose);
+			OutputPose.Curves.CopyFrom(ReferenceKeyframe.Curves);
+			OutputPose.Attributes.CopyFrom(ReferenceKeyframe.Attributes);
 		}
 	}
 };
diff --git a/Engine/Source/Runtime/Engine/Classes/Animation/AnimSequence.h b/Engine/Source/Runtime/Engine/Classes/Animation/AnimSequence.h
--- a/Engine/Source/Runtime/Engine/Classes/Animation/AnimSequence.h
+++ b/Engine/Source/Runtime/Engine/Classes/Animation/AnimSequence.h
@@ -39,6 +39,7 @@ struct FAnimSequenceDecompressionContext;
 struct FCompactPose;
 
 namespace UE { namespace Anim { class FAnimSequenceCompilingManager; namespace Compression { struct FScopedCompressionGuard; } class FAnimationSequenceAsyncCacheTask; } }
+namespace UE::AnimNext { class FDecompressionTools; }
 
 extern ENGINE_API int32 GPerformFrameStripping;
 
@@ -945,6 +946,6 @@ public:
 	friend struct UE::Anim::Compression::FScopedCompressionGuard;
 	friend class FAnimDataControllerTestBase;
 	friend class UE::Anim::FAnimSequenceCompilingManager;
-	friend struct FAnimNextAnimSequenceKeyframeTask;
+	friend UE::AnimNext::FDecompressionTools;
 	friend class FAnimSequenceDetails;
 };
diff --git a/Engine/Source/Runtime/Engine/Classes/Animation/AttributeCurve.h b/Engine/Source/Runtime/Engine/Classes/Animation/AttributeCurve.h
--- a/Engine/Source/Runtime/Engine/Classes/Animation/AttributeCurve.h
+++ b/Engine/Source/Runtime/Engine/Classes/Animation/AttributeCurve.h
@@ -14,6 +14,7 @@
 #include "AttributeCurve.generated.h"
 
 namespace UE { namespace Anim { class IAttributeBlendOperator; struct Attributes; } }
+namespace UE::AnimNext { class FDecompressionTools; }
 
 typedef UE::Anim::TWrappedAttribute<FDefaultAllocator> FWrappedAttribute;
 
@@ -190,7 +191,7 @@ protected:
 
 	friend class UAnimSequence;
 	friend struct UE::Anim::Attributes;
-	friend struct FAnimNextAnimSequenceKeyframeTask;
+	friend UE::AnimNext::FDecompressionTools;
 };
 
 template<>
diff --git a/Engine/Source/Runtime/Engine/Classes/Components/SkeletalMeshComponent.h b/Engine/Source/Runtime/Engine/Classes/Components/SkeletalMeshComponent.h
--- a/Engine/Source/Runtime/Engine/Classes/Components/SkeletalMeshComponent.h
+++ b/Engine/Source/Runtime/Engine/Classes/Components/SkeletalMeshComponent.h
@@ -308,6 +308,7 @@ class USkeletalMeshComponent : public USkinnedMeshComponent, public IInterface_C
 	friend struct FAnimNode_LinkedAnimLayer;
 	friend struct FLinkedInstancesAdapter;
 	friend struct FLinkedAnimLayerClassData;
+	friend struct FRigUnit_AnimNextWriteSkeletalMeshComponentPose;
 	
 #if WITH_EDITORONLY_DATA 
 private:
