AnimNext: Fix crash running automated tests

Disable undo/redo when adding an internal default refpose trait

diff --git a/Engine/Plugins/Experimental/AnimNext/Source/AnimNextUncookedOnly/Private/Graph/AnimNextAnimationGraph_EditorData.cpp b/Engine/Plugins/Experimental/AnimNext/Source/AnimNextUncookedOnly/Private/Graph/AnimNextAnimationGraph_EditorData.cpp
--- a/Engine/Plugins/Experimental/AnimNext/Source/AnimNextUncookedOnly/Private/Graph/AnimNextAnimationGraph_EditorData.cpp
+++ b/Engine/Plugins/Experimental/AnimNext/Source/AnimNextUncookedOnly/Private/Graph/AnimNextAnimationGraph_EditorData.cpp
@@ -363,7 +363,7 @@ namespace UE::AnimNext::UncookedOnly::Private
 
 			UAnimNextController* AnimNextController = CastChecked<UAnimNextController>(VMController);
 			constexpr UE::AnimNext::FTraitUID ReferencePoseTraitUID(0x7508ab89); // Trait header is private, reference by UID directly
-			const FName RigVMTraitName =  AnimNextController->AddTraitByName(VMNode->GetFName(), *UE::AnimNext::FTraitRegistry::Get().Find(ReferencePoseTraitUID)->GetTraitName(), INDEX_NONE);
+			const FName RigVMTraitName =  AnimNextController->AddTraitByName(VMNode->GetFName(), *UE::AnimNext::FTraitRegistry::Get().Find(ReferencePoseTraitUID)->GetTraitName(), INDEX_NONE, TEXT(""), false);
 		
 			check(RigVMTraitName != NAME_None);
 
