Fix AnimNext refpose array size

diff --git a/Engine/Plugins/Experimental/AnimNext/Source/AnimNext/Private/GenerationTools.cpp b/Engine/Plugins/Experimental/AnimNext/Source/AnimNext/Private/GenerationTools.cpp
--- a/Engine/Plugins/Experimental/AnimNext/Source/AnimNext/Private/GenerationTools.cpp
+++ b/Engine/Plugins/Experimental/AnimNext/Source/AnimNext/Private/GenerationTools.cpp
@@ -125,12 +125,12 @@ struct FCompareBoneIndexType
 				const FSkeletonToMeshLinkup& LinkupTable = const_cast<USkeleton*>(Skeleton)->FindOrAddMeshLinkupData(SkeletalMesh);
 
 				// Generate a Skeleton to LOD look up table
-				const int32 SkeletonToMeshLinkupTableSize = LinkupTable.SkeletonToMeshTable.Num();
+				const int32 NumMeshBones = LinkupTable.SkeletonToMeshTable.Num();
 				const int32 NumOrderedBones = LODBoneIndexToMeshBoneIndexMap.Num();
 
 				TArray<FBoneIndexType> SkeletonBoneIndexToLODBoneIndexMap;
 				TArray<FBoneIndexType> LODBoneIndexToSkeletonBoneIndexMap;
-				SkeletonBoneIndexToLODBoneIndexMap.Init(INDEX_NONE, SkeletonToMeshLinkupTableSize);
+				SkeletonBoneIndexToLODBoneIndexMap.Init(INDEX_NONE, NumMeshBones);
 				LODBoneIndexToSkeletonBoneIndexMap.Init(INDEX_NONE, NumOrderedBones);
 
 				for (int32 LODBoneIndex = 0; LODBoneIndex < NumOrderedBones; ++LODBoneIndex)
@@ -147,7 +147,7 @@ struct FCompareBoneIndexType
 				}
 
 				TArray<FBoneIndexType> MeshBoneIndexToLODBoneIndex;
-				MeshBoneIndexToLODBoneIndex.Init(INDEX_NONE, NumOrderedBones);
+				MeshBoneIndexToLODBoneIndex.Init(INDEX_NONE, NumMeshBones);
 
 				for (int32 LODBoneIndex = 0; LODBoneIndex < NumOrderedBones; ++LODBoneIndex)
 				{
