Fixed uninitialized variable at function schema

#rnx

diff --git a/Engine/Plugins/Experimental/AnimNext/Source/AnimNextEditor/Private/Common/GraphEditorSchemaActions.h b/Engine/Plugins/Experimental/AnimNext/Source/AnimNextEditor/Private/Common/GraphEditorSchemaActions.h
--- a/Engine/Plugins/Experimental/AnimNext/Source/AnimNextEditor/Private/Common/GraphEditorSchemaActions.h
+++ b/Engine/Plugins/Experimental/AnimNext/Source/AnimNextEditor/Private/Common/GraphEditorSchemaActions.h
@@ -154,7 +154,7 @@ private:
 
 	/** Marked as true for local function definitions */
 	UPROPERTY(Transient)
-	bool bIsLocalFunction;
+	bool bIsLocalFunction = false;
 
 	/** Holds the node type that this spawner will instantiate. */
 	UPROPERTY(Transient)
