AnimNext: Added debug drawing

Added debug draw support via RigVM's already-existing FRigVMDrawInterface. All existing RigVM debug draw helpers work now within the context of AnimNext.
Debug drawing uses a custom render proxy supported by a bare-bones custom FPrimitiveSceneProxyDesc (as the AnimNext component is not a primitive component, and we want this to work in Scene Graph too). May be able to pair this back even further, but this is what I ended up with after trying to run in various scnearios.
Added a convinence RigUnit for getting the current actor transform. As this is not thread-safe, it should be changed to use a module component that precaches the actor transform once that subsystem is checked in.
Tweaked the component communication layer a little to allow for more flexible payloads so this is easier to add to in the future.

#rnx

diff --git a/Engine/Plugins/Experimental/AnimNext/Source/AnimNext/AnimNext.Build.cs b/Engine/Plugins/Experimental/AnimNext/Source/AnimNext/AnimNext.Build.cs
--- a/Engine/Plugins/Experimental/AnimNext/Source/AnimNext/AnimNext.Build.cs
+++ b/Engine/Plugins/Experimental/AnimNext/Source/AnimNext/AnimNext.Build.cs
@@ -13,7 +13,8 @@ namespace UnrealBuildTool.Rules
 					"CoreUObject",
 					"RigVM",
 					"ControlRig",
-					"Engine"
+					"Engine",
+					"RenderCore"
 				}
 			);
 
diff --git a/Engine/Plugins/Experimental/AnimNext/Source/AnimNext/Internal/AnimNextDebugDraw.h b/Engine/Plugins/Experimental/AnimNext/Source/AnimNext/Internal/AnimNextDebugDraw.h
new file mode 100644
--- /dev/null
+++ b/Engine/Plugins/Experimental/AnimNext/Source/AnimNext/Internal/AnimNextDebugDraw.h
@@ -0,0 +1,89 @@
+﻿// Copyright Epic Games, Inc. All Rights Reserved.
+
+#pragma once
+
+#include "CoreMinimal.h"
+#include "PrimitiveSceneDesc.h"
+#include "PrimitiveSceneProxy.h"
+#include "PrimitiveSceneProxyDesc.h"
+#include "RigVMCore/RigVMDrawInterface.h"
+
+class UObject;
+struct FRigVMDrawInterface;
+struct FAnimNextModuleInstance;
+
+namespace UE::AnimNext
+{
+	struct FModuleEventTickFunction;
+}
+
+namespace UE::AnimNext::Debug
+{
+
+#if UE_ENABLE_DEBUG_DRAWING
+
+// Simple scene proxy to perform debug drawing with
+class FAnimNextDebugSceneProxy : public FPrimitiveSceneProxy
+{
+	explicit FAnimNextDebugSceneProxy(const FPrimitiveSceneProxyDesc& InProxyDesc);
+
+	virtual void GetDynamicMeshElements(const TArray<const FSceneView*>& Views, const FSceneViewFamily& ViewFamily, uint32 VisibilityMap, FMeshElementCollector& Collector) const override;
+	virtual FPrimitiveViewRelevance GetViewRelevance(const FSceneView* View) const override;
+	virtual uint32 GetMemoryFootprint() const override;
+	virtual SIZE_T GetTypeHash() const override;
+
+	FRigVMDrawInterface DrawInterface;
+	bool bIsEnabled = false;
+
+	friend struct FAnimNextModuleInstance;
+	friend struct FDebugDraw;
+};
+
+// Render data for debug drawing
+struct FDebugDraw
+{
+public:
+	explicit FDebugDraw(UObject* InOwner);
+	~FDebugDraw();
+
+	// Draw any debug items in DrawInterface
+	void Draw();
+
+	// Enable/disable debug drawing
+	// Enqueues a render command to update the enabled state
+	void SetEnabled(bool bInIsEnabled);
+
+private:
+	struct FCustomSceneProxyDesc : FPrimitiveSceneProxyDesc
+	{
+		explicit FCustomSceneProxyDesc(UObject* InOwner);
+
+		// FPrimitiveSceneProxyDesc interface
+		virtual void GetUsedMaterials(TArray<UMaterialInterface*>& OutMaterials, bool bGetDebugMaterials = false) const override {}
+
+		FCustomPrimitiveData DummyCustomPrimitiveData;
+	};
+
+	void CalcBounds();
+
+	FCustomSceneProxyDesc SceneProxyDesc;
+	FPrimitiveSceneDesc SceneDesc;
+	FPrimitiveSceneInfoData SceneInfoData;
+	FAnimNextDebugSceneProxy* SceneProxy;
+	FCustomPrimitiveData CustomPrimitiveData;
+
+	FSceneInterface* Scene = nullptr;
+
+	// Anim thread accessible draw interface
+	FRigVMDrawInterface DrawInterface;
+
+	// Anim thread accessible enabled flag
+	bool bIsEnabled = false;
+
+	friend struct ::FAnimNextModuleInstance;
+	friend struct UE::AnimNext::FModuleEventTickFunction;
+};
+
+#endif
+
+}
\ No newline at end of file
diff --git a/Engine/Plugins/Experimental/AnimNext/Source/AnimNext/Internal/Component/AnimNextComponent.h b/Engine/Plugins/Experimental/AnimNext/Source/AnimNext/Internal/Component/AnimNextComponent.h
--- a/Engine/Plugins/Experimental/AnimNext/Source/AnimNext/Internal/Component/AnimNextComponent.h
+++ b/Engine/Plugins/Experimental/AnimNext/Source/AnimNext/Internal/Component/AnimNextComponent.h
@@ -61,6 +61,10 @@ public:
 	UFUNCTION(BlueprintCallable, Category = "AnimNext")
 	ANIMNEXT_API void SetEnabled(bool bEnabled);
 
+	// Enable or disable debug drawing. Note only works in builds with UE_ENABLE_DEBUG_DRAWING enabled
+	UFUNCTION(BlueprintCallable, Category = "AnimNext")
+	ANIMNEXT_API void ShowDebugDrawing(bool bShowDebugDrawing);
+
 	// Queues an input trait event
 	// Input events will be processed in the next graph update after they are queued
 	ANIMNEXT_API void QueueInputTraitEvent(FAnimNextTraitEventPtr Event);
@@ -78,7 +82,7 @@ private:
 
 	UPROPERTY(Transient)
 	TObjectPtr<UAnimNextComponentWorldSubsystem> Subsystem = nullptr;
-	
+
 	// Handle to the registered module
 	UE::AnimNext::FModuleHandle ModuleHandle;
 
@@ -95,4 +99,8 @@ private:
 	// How to initialize the module
 	UPROPERTY(EditAnywhere, Category="Module")
 	EAnimNextModuleInitMethod InitMethod = EAnimNextModuleInitMethod::InitializeAndPauseInEditor;
+
+	/** When checked, the module's debug drawing instructions are drawn in the viewport */
+	UPROPERTY(EditAnywhere, Category = "Rendering")
+	uint8 bShowDebugDrawing : 1 = false;
 };
diff --git a/Engine/Plugins/Experimental/AnimNext/Source/AnimNext/Internal/Component/AnimNextWorldSubsystem.h b/Engine/Plugins/Experimental/AnimNext/Source/AnimNext/Internal/Component/AnimNextWorldSubsystem.h
--- a/Engine/Plugins/Experimental/AnimNext/Source/AnimNext/Internal/Component/AnimNextWorldSubsystem.h
+++ b/Engine/Plugins/Experimental/AnimNext/Source/AnimNext/Internal/Component/AnimNextWorldSubsystem.h
@@ -22,14 +22,34 @@ struct FModulePendingAction
 	enum class EType : int32
 	{
 		None = 0,
-		ReleaseHandle,
-		EnableHandle,
-		DisableHandle,
+		ReleaseHandle,			// Payload = FModulePayloadNone
+		EnableHandle,			// Payload = bool
+		EnableDebugDrawing,		// Payload = bool
+	};
+
+	// Marker struct for an 'empty' payload
+	struct FModulePayloadNone
+	{
 	};
 
 	FModulePendingAction() = default;
 
-	FModulePendingAction(EType InType, FModuleHandle InHandle);
+	FModulePendingAction(FModuleHandle InHandle, EType InType)
+		: Handle(InHandle)
+		, Type(InType)
+	{
+		Payload.Set<FModulePayloadNone>(FModulePayloadNone());
+	}
+
+	template<typename PayloadType>
+	FModulePendingAction(FModuleHandle InHandle, EType InType, PayloadType InPayload)
+		: Handle(InHandle)
+		, Type(InType)
+	{
+		Payload.Set<PayloadType>(InPayload);
+	}
+
+	TVariant<FModulePayloadNone, bool> Payload;
 
 	FModuleHandle Handle;
 
@@ -74,6 +94,12 @@ protected:
 	// This operation is deferred until the next time the schedule ticks
 	void EnableHandle(UE::AnimNext::FModuleHandle InHandle, bool bInEnabled);
 
+#if UE_ENABLE_DEBUG_DRAWING
+	// Enables or disables the module's debug drawing
+	// This operation is deferred until the next time the schedule ticks
+	void ShowDebugDrawingHandle(UE::AnimNext::FModuleHandle InHandle, bool bInShowDebugDrawing);
+#endif
+
 	// Queue a task to run at a particular point in a schedule
 	// @param	InHandle			The handle to execute the task on
 	// @param	InModuleEventName	The name of the event in the module to run the supplied task relative to. If this is NAME_None, then the first valid event will be used.
diff --git a/Engine/Plugins/Experimental/AnimNext/Source/AnimNext/Internal/Module/AnimNextModuleInstance.h b/Engine/Plugins/Experimental/AnimNext/Source/AnimNext/Internal/Module/AnimNextModuleInstance.h
--- a/Engine/Plugins/Experimental/AnimNext/Source/AnimNext/Internal/Module/AnimNextModuleInstance.h
+++ b/Engine/Plugins/Experimental/AnimNext/Source/AnimNext/Internal/Module/AnimNextModuleInstance.h
@@ -20,6 +20,11 @@ namespace UE::AnimNext
 	struct FProxyVariablesContext;
 }
 
+namespace UE::AnimNext::Debug
+{
+	struct FDebugDraw;
+}
+
 namespace UE::AnimNext
 {
 
@@ -50,7 +55,7 @@ struct FAnimNextModuleInstance : public FAnimNextDataInterfaceInstance
 {
 	GENERATED_BODY()
 
-	FAnimNextModuleInstance() = default;
+	FAnimNextModuleInstance();
 	FAnimNextModuleInstance(
 		UAnimNextModule* InModule,
 		UObject* InObject,
@@ -87,6 +92,14 @@ struct FAnimNextModuleInstance : public FAnimNextDataInterfaceInstance
 	// Get the module that this instance represents
 	ANIMNEXT_API const UAnimNextModule* GetModule() const;
 
+#if UE_ENABLE_DEBUG_DRAWING
+	// Get the debug draw interface
+	ANIMNEXT_API FRigVMDrawInterface& GetDebugDrawInterface();
+
+	// Set whether to show the module's debug drawing instructions in the viewport
+	ANIMNEXT_API void ShowDebugDrawing(bool bInShowDebugDrawing);
+#endif
+
 	// Object this entry is bound to
 	UPROPERTY(Transient)
 	TObjectPtr<UObject> Object = nullptr;
@@ -113,6 +126,12 @@ struct FAnimNextModuleInstance : public FAnimNextDataInterfaceInstance
 	UPROPERTY(Transient)
 	FAnimNextPublicVariablesProxy PublicVariablesProxy;
 
+#if UE_ENABLE_DEBUG_DRAWING
+	// Debug draw support
+	// TODO: Move this to a 'module component' once that code is checked in
+	TUniquePtr<UE::AnimNext::Debug::FDebugDraw> DebugDraw;
+#endif
+
 	enum class ERunState : uint8
 	{
 		None,
@@ -133,7 +152,7 @@ struct FAnimNextModuleInstance : public FAnimNextDataInterfaceInstance
 
 	// Transition to the specified run state, verifying that the current state is valid
 	void TransitionToRunState(ERunState InNewState);
-	
+
 	// How this entry initializes
 	EAnimNextModuleInitMethod InitMethod = EAnimNextModuleInitMethod::InitializeAndPauseInEditor;
 
diff --git a/Engine/Plugins/Experimental/AnimNext/Source/AnimNext/Private/AnimNextDebugDraw.cpp b/Engine/Plugins/Experimental/AnimNext/Source/AnimNext/Private/AnimNextDebugDraw.cpp
new file mode 100644
--- /dev/null
+++ b/Engine/Plugins/Experimental/AnimNext/Source/AnimNext/Private/AnimNextDebugDraw.cpp
@@ -0,0 +1,219 @@
+﻿// Copyright Epic Games, Inc. All Rights Reserved.
+
+#include "AnimNextDebugDraw.h"
+
+#include "PrimitiveSceneDesc.h"
+#include "PrimitiveSceneProxyDesc.h"
+#include "SceneInterface.h"
+#include "RigVMCore/RigVMDrawInstruction.h"
+#include "Engine/EngineTypes.h"
+#include "Engine/World.h"
+#include "SceneManagement.h"
+#include "Module/AnimNextModuleInstance.h"
+
+namespace UE::AnimNext::Debug
+{
+
+#if UE_ENABLE_DEBUG_DRAWING
+
+FAnimNextDebugSceneProxy::FAnimNextDebugSceneProxy(const FPrimitiveSceneProxyDesc& InProxyDesc)
+	: FPrimitiveSceneProxy(InProxyDesc)
+{
+	bWillEverBeLit = false;
+}
+
+SIZE_T FAnimNextDebugSceneProxy::GetTypeHash() const
+{
+	static size_t UniquePointer;
+	return reinterpret_cast<size_t>(&UniquePointer);
+}
+
+void FAnimNextDebugSceneProxy::GetDynamicMeshElements(const TArray<const FSceneView*>& Views, const FSceneViewFamily& ViewFamily, uint32 VisibilityMap, FMeshElementCollector& Collector) const
+{
+	for (int32 ViewIndex = 0; ViewIndex < Views.Num(); ViewIndex++)
+	{
+		if (VisibilityMap & (1 << ViewIndex))
+		{
+			const FSceneView* View = Views[ViewIndex];
+			FPrimitiveDrawInterface* PDI = Collector.GetPDI(ViewIndex);
+
+			if (bIsEnabled)
+			{ 
+				for (int32 InstructionIndex = 0; InstructionIndex < DrawInterface.Num(); InstructionIndex++)
+				{
+					const FRigVMDrawInstruction& Instruction = DrawInterface[InstructionIndex];
+					if (Instruction.Positions.Num() == 0)
+					{
+						continue;
+					}
+
+					switch (Instruction.PrimitiveType)
+					{
+						case ERigVMDrawSettings::Points:
+						{
+							for (const FVector& Point : Instruction.Positions)
+							{
+								PDI->DrawPoint(Instruction.Transform.TransformPosition(Point), Instruction.Color, Instruction.Thickness, SDPG_World);
+							}
+							break;
+						}
+						case ERigVMDrawSettings::Lines:
+						{
+							const TArray<FVector>& Points = Instruction.Positions;
+							PDI->AddReserveLines(SDPG_World, Points.Num() / 2, false, Instruction.Thickness > SMALL_NUMBER);
+							for (int32 PointIndex = 0; PointIndex < Points.Num() - 1; PointIndex += 2)
+							{
+								PDI->DrawLine(Instruction.Transform.TransformPosition(Points[PointIndex]), Instruction.Transform.TransformPosition(Points[PointIndex + 1]), Instruction.Color, SDPG_World, Instruction.Thickness);
+							}
+							break;
+						}
+						case ERigVMDrawSettings::LineStrip:
+						{
+							const TArray<FVector>& Points = Instruction.Positions;
+							PDI->AddReserveLines(SDPG_World, Points.Num() - 1, false, Instruction.Thickness > SMALL_NUMBER);
+							for (int32 PointIndex = 0; PointIndex < Points.Num() - 1; PointIndex++)
+							{
+								PDI->DrawLine(Instruction.Transform.TransformPosition(Points[PointIndex]), Instruction.Transform.TransformPosition(Points[PointIndex + 1]), Instruction.Color, SDPG_World, Instruction.Thickness);
+							}
+							break;
+						}
+						case ERigVMDrawSettings::DynamicMesh:
+						{
+							FDynamicMeshBuilder MeshBuilder(PDI->View->GetFeatureLevel());
+							MeshBuilder.AddVertices(Instruction.MeshVerts);
+							MeshBuilder.AddTriangles(Instruction.MeshIndices);
+							MeshBuilder.Draw(PDI, Instruction.Transform.ToMatrixWithScale(), Instruction.MaterialRenderProxy, SDPG_World/*SDPG_Foreground*/);
+							break;
+						}
+					default:
+						break;
+					}
+				}
+			}
+		}
+	}
+}
+
+FPrimitiveViewRelevance FAnimNextDebugSceneProxy::GetViewRelevance(const FSceneView* View) const
+{
+	FPrimitiveViewRelevance ViewRelevance;
+	ViewRelevance.bDrawRelevance = IsShown(View);
+	ViewRelevance.bDynamicRelevance = true;
+	// ideally the TranslucencyRelevance should be filled out by the material, here we do it conservative
+	ViewRelevance.bSeparateTranslucency = ViewRelevance.bNormalTranslucency = true;
+	return ViewRelevance;
+}
+
+uint32 FAnimNextDebugSceneProxy::GetMemoryFootprint(void) const
+{
+	return(sizeof(*this) + FPrimitiveSceneProxy::GetAllocatedSize());
+}
+
+FDebugDraw::FCustomSceneProxyDesc::FCustomSceneProxyDesc(UObject* InOwner)
+{
+	bReceivesDecals = false;
+	bVisibleInReflectionCaptures = false;
+	bVisibleInRealTimeSkyCaptures = false;
+	bVisibleInRayTracing = false;
+
+	bCastDynamicShadow = false;
+	bCastStaticShadow = false;
+	bAffectDynamicIndirectLighting = false;
+	bAffectDistanceFieldLighting = false;
+	bCastContactShadow = false;
+	bSelectable = false;
+	bReceiveMobileCSMShadows = false;
+	bHiddenInSceneCapture = true;
+	bSupportsWorldPositionOffsetVelocity = false;
+
+	World = InOwner->GetWorld();
+	Scene = World->Scene;
+	CustomPrimitiveData = &DummyCustomPrimitiveData;
+}
+
+FDebugDraw::FDebugDraw(UObject* InOwner)
+	: SceneProxyDesc(InOwner)
+	, SceneProxy(new FAnimNextDebugSceneProxy(SceneProxyDesc))
+{
+	SceneInfoData.SceneProxy = SceneProxy;
+	SceneInfoData.OwnerLastRenderTimePtr = nullptr;
+	SceneDesc.ProxyDesc = &SceneProxyDesc;
+	SceneDesc.PrimitiveSceneData = &SceneInfoData;
+	SceneDesc.World = InOwner->GetWorld();
+
+	SceneDesc.RenderMatrix = FMatrix::Identity;
+	SceneDesc.AttachmentRootPosition = FVector::ZeroVector;
+	SceneDesc.Bounds = FBoxSphereBounds(FVector::ZeroVector, FVector::OneVector, 1.f);
+	SceneDesc.LocalBounds = SceneDesc.Bounds;
+
+	Scene = SceneDesc.World->Scene;
+
+	Scene->AddPrimitive(&SceneDesc);
+}
+
+FDebugDraw::~FDebugDraw()
+{
+	Scene->RemovePrimitive(&SceneDesc);
+}
+
+void FDebugDraw::Draw()
+{
+	// Recalc bounds according to what we need to draw
+	CalcBounds();
+
+	// Update the primitive transform/bounds
+	Scene->UpdatePrimitiveTransform(&SceneDesc);
+
+	// Move ownership of the draw interface to the RT
+	ENQUEUE_RENDER_COMMAND(AnimNextDebugDraw)([bIsEnabled = bIsEnabled, SceneProxy = SceneProxy, DrawInterface = MoveTemp(DrawInterface)](FRHICommandList& RHICmdList) mutable
+	{
+		SceneProxy->bIsEnabled = bIsEnabled;
+		SceneProxy->DrawInterface = MoveTemp(DrawInterface);
+	});
+
+	DrawInterface.Reset();
+}
+
+void FDebugDraw::SetEnabled(bool bInIsEnabled)
+{
+	if(bInIsEnabled != bIsEnabled)
+	{
+		// Update the RT copy of the flag if it differs
+		ENQUEUE_RENDER_COMMAND(AnimNextDebugDraw)([bInIsEnabled, SceneProxy = SceneProxy](FRHICommandList& RHICmdList) mutable
+		{
+			SceneProxy->bIsEnabled = bInIsEnabled;
+		});
+	}
+	bIsEnabled = bInIsEnabled; 
+}
+
+void FDebugDraw::CalcBounds()
+{
+	FBox Box(ForceInit);
+
+	// Get bounding box for the debug drawings if they are drawn 
+	for (const FRigVMDrawInstruction& Instruction : DrawInterface.Instructions)
+	{
+		FTransform Transform = Instruction.Transform /** GetComponentToWorld()*/;
+		for (const FVector& Position : Instruction.Positions)
+		{
+			Box += Transform.TransformPosition(Position);
+		}
+	}
+
+	if (Box.IsValid)
+	{
+		// Points are in world space, so no need to transform.
+		SceneDesc.Bounds = FBoxSphereBounds(Box);
+	}
+	else
+	{
+		SceneDesc.Bounds = FBoxSphereBounds(FVector::ZeroVector, FVector::OneVector, 1.f);
+	}
+
+	SceneDesc.LocalBounds = SceneDesc.Bounds;
+}
+
+#endif
+
+}
diff --git a/Engine/Plugins/Experimental/AnimNext/Source/AnimNext/Private/Component/AnimNextComponent.cpp b/Engine/Plugins/Experimental/AnimNext/Source/AnimNext/Private/Component/AnimNextComponent.cpp
--- a/Engine/Plugins/Experimental/AnimNext/Source/AnimNext/Private/Component/AnimNextComponent.cpp
+++ b/Engine/Plugins/Experimental/AnimNext/Source/AnimNext/Private/Component/AnimNextComponent.cpp
@@ -22,6 +22,9 @@ void UAnimNextComponent::OnRegister()
 
 		CreatePublicVariablesProxy();
 		Subsystem->Register(this);
+#if UE_ENABLE_DEBUG_DRAWING
+		Subsystem->ShowDebugDrawing(this, bShowDebugDrawing);
+#endif
 	}
 }
 
@@ -164,6 +167,17 @@ void UAnimNextComponent::SetEnabled(bool bEnabled)
 	}
 }
 
+void UAnimNextComponent::ShowDebugDrawing(bool bInShowDebugDrawing)
+{
+#if UE_ENABLE_DEBUG_DRAWING
+	bShowDebugDrawing = bInShowDebugDrawing;
+	if(Subsystem)
+	{
+		Subsystem->ShowDebugDrawing(this, bInShowDebugDrawing);
+	}
+#endif
+}
+
 void UAnimNextComponent::QueueInputTraitEvent(FAnimNextTraitEventPtr Event)
 {
 	using namespace UE::AnimNext;
diff --git a/Engine/Plugins/Experimental/AnimNext/Source/AnimNext/Private/Component/AnimNextComponentWorldSubsystem.cpp b/Engine/Plugins/Experimental/AnimNext/Source/AnimNext/Private/Component/AnimNextComponentWorldSubsystem.cpp
--- a/Engine/Plugins/Experimental/AnimNext/Source/AnimNext/Private/Component/AnimNextComponentWorldSubsystem.cpp
+++ b/Engine/Plugins/Experimental/AnimNext/Source/AnimNext/Private/Component/AnimNextComponentWorldSubsystem.cpp
@@ -33,6 +33,16 @@ void UAnimNextComponentWorldSubsystem::SetEnabled(UAnimNextComponent* InComponen
 	EnableHandle(InComponent->ModuleHandle, bInEnabled);
 }
 
+#if UE_ENABLE_DEBUG_DRAWING
+void UAnimNextComponentWorldSubsystem::ShowDebugDrawing(UAnimNextComponent* InComponent, bool bInShowDebugDrawing)
+{
+	using namespace UE::AnimNext;
+	check(IsInGameThread());
+	check(InComponent);
+	ShowDebugDrawingHandle(InComponent->ModuleHandle, bInShowDebugDrawing);
+}
+#endif
+
 void UAnimNextComponentWorldSubsystem::QueueTask(UAnimNextComponent* InComponent, FName InModuleEventName, TUniqueFunction<void(const UE::AnimNext::FModuleTaskContext&)>&& InTaskFunction, UE::AnimNext::ETaskRunLocation InLocation)
 {
 	using namespace UE::AnimNext;
diff --git a/Engine/Plugins/Experimental/AnimNext/Source/AnimNext/Private/Component/AnimNextComponentWorldSubsystem.h b/Engine/Plugins/Experimental/AnimNext/Source/AnimNext/Private/Component/AnimNextComponentWorldSubsystem.h
--- a/Engine/Plugins/Experimental/AnimNext/Source/AnimNext/Private/Component/AnimNextComponentWorldSubsystem.h
+++ b/Engine/Plugins/Experimental/AnimNext/Source/AnimNext/Private/Component/AnimNextComponentWorldSubsystem.h
@@ -30,6 +30,12 @@ class UAnimNextComponentWorldSubsystem : public UAnimNextWorldSubsystem
 	// This operation is deferred until the next time the schedule ticks
 	void SetEnabled(UAnimNextComponent* InComponent, bool bInEnabled);
 
+#if UE_ENABLE_DEBUG_DRAWING
+	// Enables or disables debug drawing for the module represented by the supplied handle
+	// This operation is deferred until the next time the schedule ticks
+	void ShowDebugDrawing(UAnimNextComponent* InComponent, bool bInShowDebugDrawing);
+#endif
+
 	// Queue a task to run at a particular point in a schedule
 	// @param	InComponent			The component to execute the task on
 	// @param	InModuleEventName	The name of the event in the module to run the supplied task relative to. If this is NAME_None, then the first valid event will be used.
diff --git a/Engine/Plugins/Experimental/AnimNext/Source/AnimNext/Private/Component/AnimNextWorldSubsystem.cpp b/Engine/Plugins/Experimental/AnimNext/Source/AnimNext/Private/Component/AnimNextWorldSubsystem.cpp
--- a/Engine/Plugins/Experimental/AnimNext/Source/AnimNext/Private/Component/AnimNextWorldSubsystem.cpp
+++ b/Engine/Plugins/Experimental/AnimNext/Source/AnimNext/Private/Component/AnimNextWorldSubsystem.cpp
@@ -11,11 +11,6 @@
 namespace UE::AnimNext
 {
 
-FModulePendingAction::FModulePendingAction(EType InType, FModuleHandle InHandle)
-	: Handle(InHandle)
-	, Type(InType)
-{
-}
 
 }
 
@@ -107,15 +102,17 @@ void UAnimNextWorldSubsystem::FlushPendingActions()
 				if (IsValidHandle(PendingAction.Handle))
 				{
 					FAnimNextModuleInstance& Instance = Instances.Get(PendingAction.Handle);
-					Instance.Enable(true);
+					Instance.Enable(PendingAction.Payload.Get<bool>());
 				}
 				break;
-			case FModulePendingAction::EType::DisableHandle:
+			case FModulePendingAction::EType::EnableDebugDrawing:
+#if UE_ENABLE_DEBUG_DRAWING
 				if (IsValidHandle(PendingAction.Handle))
 				{
 					FAnimNextModuleInstance& Instance = Instances.Get(PendingAction.Handle);
-					Instance.Enable(false);
+					Instance.ShowDebugDrawing(PendingAction.Payload.Get<bool>());
 				}
+#endif
 				break;
 			default:
 				break;
@@ -155,15 +152,27 @@ void UAnimNextWorldSubsystem::UnregisterHandle(UE::AnimNext::FModuleHandle& InOu
 	}
 }
 
-void UAnimNextWorldSubsystem::EnableHandle(UE::AnimNext::FModuleHandle InOutHandle, bool bInEnabled)
+void UAnimNextWorldSubsystem::EnableHandle(UE::AnimNext::FModuleHandle InHandle, bool bInEnabled)
 {
 	using namespace UE::AnimNext;
 	check(IsInGameThread());
-	if (IsValidHandle(InOutHandle))
+	if (IsValidHandle(InHandle))
+	{
+		PendingActions.Emplace(InHandle, FModulePendingAction::EType::EnableHandle, bInEnabled);
+	}
+}
+
+#if UE_ENABLE_DEBUG_DRAWING
+void UAnimNextWorldSubsystem::ShowDebugDrawingHandle(UE::AnimNext::FModuleHandle InHandle, bool bInShowDebugDrawing)
+{
+	using namespace UE::AnimNext;
+	check(IsInGameThread());
+	if (IsValidHandle(InHandle))
 	{
-		PendingActions.Emplace(bInEnabled ? FModulePendingAction::EType::EnableHandle : FModulePendingAction::EType::DisableHandle, InOutHandle);
+		PendingActions.Emplace(InHandle, FModulePendingAction::EType::EnableDebugDrawing, bInShowDebugDrawing);
 	}
 }
+#endif
 
 void UAnimNextWorldSubsystem::QueueTaskHandle(UE::AnimNext::FModuleHandle InOutHandle, FName InModuleEventName, TUniqueFunction<void(const UE::AnimNext::FModuleTaskContext&)>&& InTaskFunction, UE::AnimNext::ETaskRunLocation InLocation)
 {
diff --git a/Engine/Plugins/Experimental/AnimNext/Source/AnimNext/Private/Graph/RigUnit_GetActorTransform.cpp b/Engine/Plugins/Experimental/AnimNext/Source/AnimNext/Private/Graph/RigUnit_GetActorTransform.cpp
new file mode 100644
--- /dev/null
+++ b/Engine/Plugins/Experimental/AnimNext/Source/AnimNext/Private/Graph/RigUnit_GetActorTransform.cpp
@@ -0,0 +1,31 @@
+﻿// Copyright Epic Games, Inc. All Rights Reserved.
+
+#include "RigUnit_GetActorTransform.h"
+
+#include "Component/AnimNextComponent.h"
+
+FRigUnit_GetActorTransform_Execute()
+{
+	using namespace UE::AnimNext;
+
+	const FAnimNextModuleContextData& ModuleContextData = ExecuteContext.GetContextData<FAnimNextModuleContextData>();
+	FAnimNextModuleInstance& ModuleInstance = ModuleContextData.GetModuleInstance();
+
+	// TODO: this is not thread safe at the moment. Once we get the 'weak semantics' CL checked in, we can move commonly-used debug info
+	// into a module component that can be used to access GT data where needed. 
+	UAnimNextComponent* Component = Cast<UAnimNextComponent>(ModuleInstance.Object);
+	if(Component == nullptr)
+	{
+		Transform = FTransform::Identity;
+		return;
+	}
+
+	AActor* Owner = Component->GetOwner();
+	if(Owner == nullptr)
+	{
+		Transform = FTransform::Identity;
+		return;
+	}
+
+	Transform = Owner->GetActorTransform();
+}
\ No newline at end of file
diff --git a/Engine/Plugins/Experimental/AnimNext/Source/AnimNext/Private/Graph/RigUnit_GetActorTransform.h b/Engine/Plugins/Experimental/AnimNext/Source/AnimNext/Private/Graph/RigUnit_GetActorTransform.h
new file mode 100644
--- /dev/null
+++ b/Engine/Plugins/Experimental/AnimNext/Source/AnimNext/Private/Graph/RigUnit_GetActorTransform.h
@@ -0,0 +1,24 @@
+﻿// Copyright Epic Games, Inc. All Rights Reserved.
+
+#pragma once
+
+#include "AnimNextExecuteContext.h"
+#include "Graph/RigUnit_AnimNextBase.h"
+
+#include "RigUnit_GetActorTransform.generated.h"
+
+/** Gets the transform of the host of this AnimNext module */
+USTRUCT(meta=(DisplayName="Get Host Transform", Category="Animation Graph", NodeColor="0, 1, 1", Keywords="Actor, Transform"))
+struct ANIMNEXT_API FRigUnit_GetActorTransform : public FRigUnit_AnimNextBase
+{
+	GENERATED_BODY()
+
+	RIGVM_METHOD()
+	void Execute();
+
+	virtual FString GetUnitSubTitle() const { return TEXT("Actor"); };
+
+	// The transform of our host actor
+	UPROPERTY(EditAnywhere, Category = "Graph", meta = (Output))
+	FTransform Transform;
+};
diff --git a/Engine/Plugins/Experimental/AnimNext/Source/AnimNext/Private/Module/AnimNextModuleInstance.cpp b/Engine/Plugins/Experimental/AnimNext/Source/AnimNext/Private/Module/AnimNextModuleInstance.cpp
--- a/Engine/Plugins/Experimental/AnimNext/Source/AnimNext/Private/Module/AnimNextModuleInstance.cpp
+++ b/Engine/Plugins/Experimental/AnimNext/Source/AnimNext/Private/Module/AnimNextModuleInstance.cpp
@@ -2,8 +2,12 @@
 
 #include "Module/AnimNextModuleInstance.h"
 
+#if UE_ENABLE_DEBUG_DRAWING
+#include "AnimNextDebugDraw.h"
+#endif
 #include "Engine/World.h"
 #include "AnimNextStats.h"
+#include "SceneInterface.h"
 #include "Algo/TopologicalSort.h"
 #include "Logging/StructuredLog.h"
 #include "Misc/EnumerateRange.h"
@@ -14,6 +18,8 @@
 
 DEFINE_STAT(STAT_AnimNext_InitializeInstance);
 
+FAnimNextModuleInstance::FAnimNextModuleInstance() = default;
+
 FAnimNextModuleInstance::FAnimNextModuleInstance(
 		UAnimNextModule* InModule,
 		UObject* InObject,
@@ -23,12 +29,20 @@ FAnimNextModuleInstance::FAnimNextModuleInstance(
 	, InitMethod(InInitMethod)
 {
 	DataInterface = InModule;
+
+#if UE_ENABLE_DEBUG_DRAWING
+	DebugDraw = MakeUnique<UE::AnimNext::Debug::FDebugDraw>(Object);
+#endif
 }
 
 FAnimNextModuleInstance::~FAnimNextModuleInstance()
 {
 	ResetBindingsAndInstanceData();
 
+#if UE_ENABLE_DEBUG_DRAWING
+	DebugDraw.Reset();
+#endif
+
 	Object = nullptr;
 	DataInterface = nullptr;
 	Handle.Reset();
@@ -367,4 +381,18 @@ void FAnimNextModuleInstance::OnModuleCompiled()
 	ResetBindingsAndInstanceData();
 	Initialize();
 }
+#endif
+
+#if UE_ENABLE_DEBUG_DRAWING
+FRigVMDrawInterface& FAnimNextModuleInstance::GetDebugDrawInterface()
+{
+	check(DebugDraw.IsValid());
+	return DebugDraw->DrawInterface;
+}
+
+void FAnimNextModuleInstance::ShowDebugDrawing(bool bInShowDebugDrawing)
+{
+	check(DebugDraw.IsValid());
+	DebugDraw->SetEnabled(bInShowDebugDrawing);
+}
 #endif
\ No newline at end of file
diff --git a/Engine/Plugins/Experimental/AnimNext/Source/AnimNext/Private/Module/ModuleTickFunction.cpp b/Engine/Plugins/Experimental/AnimNext/Source/AnimNext/Private/Module/ModuleTickFunction.cpp
--- a/Engine/Plugins/Experimental/AnimNext/Source/AnimNext/Private/Module/ModuleTickFunction.cpp
+++ b/Engine/Plugins/Experimental/AnimNext/Source/AnimNext/Private/Module/ModuleTickFunction.cpp
@@ -2,6 +2,7 @@
 
 #include "Module/ModuleTickFunction.h"
 
+#include "AnimNextDebugDraw.h"
 #include "RigVMRuntimeDataRegistry.h"
 #include "Algo/TopologicalSort.h"
 #include "Module/AnimNextModule.h"
@@ -114,6 +115,11 @@ void FModuleEndTickFunction::Run()
 				}
 			});
 	}
+
+#if UE_ENABLE_DEBUG_DRAWING
+	// Perform any debug drawing
+	ModuleInstance->DebugDraw->Draw();
+#endif
 }
 
 FString FModuleEndTickFunction::DiagnosticMessage()
@@ -145,6 +151,10 @@ void FModuleEventTickFunction::Run(float InDeltaTime)
 		// RigVM setup
 		AnimNextContext.SetDeltaTime(InDeltaTime);
 
+#if UE_ENABLE_DEBUG_DRAWING
+		AnimNextContext.SetDrawInterface(&ModuleInstance->GetDebugDrawInterface());
+#endif
+
 		// Module setup
 		AnimNextContext.SetContextData<FAnimNextModuleContextData>(ModuleInstance);
 
diff --git a/Engine/Plugins/Experimental/AnimNext/Source/AnimNext/Private/TraitCore/ExecutionContext.cpp b/Engine/Plugins/Experimental/AnimNext/Source/AnimNext/Private/TraitCore/ExecutionContext.cpp
--- a/Engine/Plugins/Experimental/AnimNext/Source/AnimNext/Private/TraitCore/ExecutionContext.cpp
+++ b/Engine/Plugins/Experimental/AnimNext/Source/AnimNext/Private/TraitCore/ExecutionContext.cpp
@@ -586,4 +586,12 @@ namespace UE::AnimNext
 		check(Template.GetRegistryHandle().IsValid());
 		return TraitRegistry.Find(Template.GetRegistryHandle());
 	}
+
+#if UE_ENABLE_DEBUG_DRAWING
+	FRigVMDrawInterface& FExecutionContext::GetDebugDrawInterface()
+	{
+		check(IsBound());
+		return RootGraphInstance->GetModuleInstance()->GetDebugDrawInterface();
+	}
+#endif
 }
diff --git a/Engine/Plugins/Experimental/AnimNext/Source/AnimNext/Public/TraitCore/ExecutionContext.h b/Engine/Plugins/Experimental/AnimNext/Source/AnimNext/Public/TraitCore/ExecutionContext.h
--- a/Engine/Plugins/Experimental/AnimNext/Source/AnimNext/Public/TraitCore/ExecutionContext.h
+++ b/Engine/Plugins/Experimental/AnimNext/Source/AnimNext/Public/TraitCore/ExecutionContext.h
@@ -15,11 +15,13 @@
 #include "TraitCore/NodeHandle.h"
 #include "Graph/AnimNextGraphInstancePtr.h"
 #include "Graph/GraphInstanceComponent.h"
+#include "EngineDefines.h"
 
 #include <type_traits>
 
 struct FAnimNextGraphInstance;
 class FMemStack;
+struct FRigVMDrawInterface;
 
 namespace UE::AnimNext
 {
@@ -202,6 +204,10 @@ namespace UE::AnimNext
 		// Returns the local thread memstack
 		FMemStack& GetMemStack() const { return MemStack; }
 
+#if UE_ENABLE_DEBUG_DRAWING
+		// Get the debug draw interface
+		FRigVMDrawInterface& GetDebugDrawInterface();
+#endif
 	private:
 		// No copy or move
 		FExecutionContext(const FExecutionContext&) = delete;
