Anim Next: Added rewind debugger and visual logger support

diff --git a/Engine/Plugins/Experimental/AnimNext/Source/AnimNext/Private/Component/SkinnedMeshComponentExtensions.cpp b/Engine/Plugins/Experimental/AnimNext/Source/AnimNext/Private/Component/SkinnedMeshComponentExtensions.cpp
--- a/Engine/Plugins/Experimental/AnimNext/Source/AnimNext/Private/Component/SkinnedMeshComponentExtensions.cpp
+++ b/Engine/Plugins/Experimental/AnimNext/Source/AnimNext/Private/Component/SkinnedMeshComponentExtensions.cpp
@@ -3,6 +3,7 @@
 #include "Component/SkinnedMeshComponentExtensions.h"
 #include "GenerationTools.h"
 #include "Components/SkinnedMeshComponent.h"
+#include "Animation/AnimTrace.h"
 
 namespace UE::Anim
 {
@@ -30,6 +31,11 @@ void FSkinnedMeshComponentExtensions::CompleteAndDispatch(
 
 		// Send updated transforms to the renderer
 		InComponent->SendRenderDynamicData_Concurrent();
+
+		if (const USkeletalMeshComponent* SkeletalMeshComponent = Cast<USkeletalMeshComponent>(InComponent))
+		{
+			TRACE_SKELETAL_MESH_COMPONENT(SkeletalMeshComponent);
+		}
 	}
 }
 
diff --git a/Engine/Plugins/Experimental/AnimNext/Source/AnimNext/Private/TraitCore/ExecutionContext.cpp b/Engine/Plugins/Experimental/AnimNext/Source/AnimNext/Private/TraitCore/ExecutionContext.cpp
--- a/Engine/Plugins/Experimental/AnimNext/Source/AnimNext/Private/TraitCore/ExecutionContext.cpp
+++ b/Engine/Plugins/Experimental/AnimNext/Source/AnimNext/Private/TraitCore/ExecutionContext.cpp
@@ -11,6 +11,7 @@
 #include "TraitCore/NodeTemplateRegistry.h"
 #include "Graph/AnimNextAnimationGraph.h"
 #include "Graph/AnimNextGraphInstance.h"
+#include "Components/SceneComponent.h"
 
 namespace UE::AnimNext
 {
@@ -593,5 +594,33 @@ namespace UE::AnimNext
 		check(IsBound());
 		return RootGraphInstance->GetModuleInstance()->GetDebugDrawInterface();
 	}
+
+	const UObject* FExecutionContext::GetHostObject() const
+	{
+		const FAnimNextModuleInstance* ModuleInstance = GetRootGraphInstance().GetModuleInstance();
+		if (ModuleInstance)
+		{
+			return ModuleInstance->Object;
+		}
+
+		return nullptr;
+	}
+
+	const FTransform FExecutionContext::GetHostTransform() const
+	{
+		if (const UObject* ContextObject = GetHostObject())
+		{
+			if (const UActorComponent* ActorComponent = Cast<UActorComponent>(ContextObject))
+			{
+				if (const AActor* OwningActor = ActorComponent->GetOwner())
+				{
+					return OwningActor->GetActorTransform();
+				}
+			}
+		}
+
+		// @TODO: Add support for non-actors.
+		return FTransform::Identity;
+	}
 #endif
 }
diff --git a/Engine/Plugins/Experimental/AnimNext/Source/AnimNext/Public/TraitCore/ExecutionContext.h b/Engine/Plugins/Experimental/AnimNext/Source/AnimNext/Public/TraitCore/ExecutionContext.h
--- a/Engine/Plugins/Experimental/AnimNext/Source/AnimNext/Public/TraitCore/ExecutionContext.h
+++ b/Engine/Plugins/Experimental/AnimNext/Source/AnimNext/Public/TraitCore/ExecutionContext.h
@@ -3,6 +3,7 @@
 #pragma once
 
 #include "CoreMinimal.h"
+#include "EngineDefines.h"
 #include "TraitCore/TraitBinding.h"
 #include "TraitCore/TraitEvent.h"
 #include "TraitCore/TraitHandle.h"
@@ -58,6 +59,14 @@ namespace UE::AnimNext
 		// Derived types can control various features
 		virtual ~FExecutionContext() = default;
 
+#if UE_ENABLE_DEBUG_DRAWING
+		// Owner object our module instance belongs to. Usually the Anim Next Component. This access is NOT thread-safe.
+		const class UObject* GetHostObject() const;
+		
+		// The world transform of our host. Usually the actor containing our Anim Next Component
+		// @TODO: Not thread-safe and should be replaced ASAP by a subsystem so we can use this in scene graph/mass. Doesn't work for non-actors.
+		const FTransform GetHostTransform() const;
+#endif
 
 		//////////////////////////////////////////////////////////////////////////
 		// The following functions handle the context binding
