Fix static init ordering issue in AnimNext trait registry

diff --git a/Engine/Plugins/Experimental/AnimNext/Source/AnimNext/Private/TraitCore/NodeDescription.cpp b/Engine/Plugins/Experimental/AnimNext/Source/AnimNext/Private/TraitCore/NodeDescription.cpp
--- a/Engine/Plugins/Experimental/AnimNext/Source/AnimNext/Private/TraitCore/NodeDescription.cpp
+++ b/Engine/Plugins/Experimental/AnimNext/Source/AnimNext/Private/TraitCore/NodeDescription.cpp
@@ -37,17 +37,19 @@ namespace UE::AnimNext
 		}
 
 		// Use our template to serialize our traits
+		const FTraitRegistry& TraitRegistry = FTraitRegistry::Get();
 		const FNodeTemplate* NodeTemplate = NodeTemplateRegistry.Find(TemplateHandle);
 
 		const uint32 NumTraits = NodeTemplate->GetNumTraits();
 		const FTraitTemplate* TraitTemplates = NodeTemplate->GetTraits();
 		for (uint32 TraitIndex = 0; TraitIndex < NumTraits; ++TraitIndex)
 		{
-			const FTraitRegistryHandle TraitHandle = TraitTemplates[TraitIndex].GetRegistryHandle();
-			const FTrait* Trait = FTraitRegistry::Get().Find(TraitHandle);
+			const FTraitTemplate& TraitTemplate = TraitTemplates[TraitIndex];
 
-			FAnimNextTraitSharedData* SharedData = TraitTemplates[TraitIndex].GetTraitDescription(*this);
+			const FTraitRegistryHandle TraitHandle = TraitTemplate.GetRegistryHandle();
+			FAnimNextTraitSharedData* SharedData = TraitTemplate.GetTraitDescription(*this);
 
+			const FTrait* Trait = TraitRegistry.Find(TraitHandle);
 			Trait->SerializeTraitSharedData(Ar, *SharedData);
 		}
 	}
diff --git a/Engine/Plugins/Experimental/AnimNext/Source/AnimNext/Private/TraitCore/TraitRegistry.cpp b/Engine/Plugins/Experimental/AnimNext/Source/AnimNext/Private/TraitCore/TraitRegistry.cpp
--- a/Engine/Plugins/Experimental/AnimNext/Source/AnimNext/Private/TraitCore/TraitRegistry.cpp
+++ b/Engine/Plugins/Experimental/AnimNext/Source/AnimNext/Private/TraitCore/TraitRegistry.cpp
@@ -8,9 +8,25 @@ namespace UE::AnimNext
 {
 	namespace Private
 	{
-		static FTraitRegistry* GTraitRegistry = nullptr;
+		// The trait registry relies on static initialization to queue traits that attempt to register
+		// before the module is ready. However, we need to store these pending requests in the meantime.
+		// To do so safely, we return this through a function to ensure that it initializes before we
+		// use it.
+		struct FTraitRegistryStaticData
+		{
+			TArray<TraitConstructorFunc> PendingRegistrationQueue;
+		};
+
+		static FTraitRegistryStaticData& GetTraitRegistryStaticData()
+		{
+			static FTraitRegistryStaticData StaticData;
+			return StaticData;
+		}
 
-		static TArray<TraitConstructorFunc> GPendingRegistrationQueue;
+		// Just a regular static to keep codegen clean as there is no race when accessing this during static
+		// init. This pointer will be null until the module initializes (for the duration of static init) and
+		// only returns to null when we shutdown the process and the module unloads.
+		static FTraitRegistry* GTraitRegistry = nullptr;
 	}
 
 	FTraitRegistry& FTraitRegistry::Get()
@@ -26,13 +42,14 @@ namespace UE::AnimNext
 			Private::GTraitRegistry = new FTraitRegistry();
 
 			// Register all our pending static init traits
-			for (TraitConstructorFunc TraitConstructor : Private::GPendingRegistrationQueue)
+			Private::FTraitRegistryStaticData& StaticData = Private::GetTraitRegistryStaticData();
+			for (TraitConstructorFunc TraitConstructor : StaticData.PendingRegistrationQueue)
 			{
 				Private::GTraitRegistry->AutoRegisterImpl(TraitConstructor);
 			}
 
 			// Reset the registration queue, it won't be used anymore now that the registry is up and ready
-			Private::GPendingRegistrationQueue.Empty(0);
+			StaticData.PendingRegistrationQueue.Empty(0);
 		}
 	}
 
@@ -64,7 +81,8 @@ namespace UE::AnimNext
 		{
 			// Registry isn't ready yet, queue up our trait
 			// Once Init() is called, our queue will be processed
-			Private::GPendingRegistrationQueue.Add(TraitConstructor);
+			Private::FTraitRegistryStaticData& StaticData = Private::GetTraitRegistryStaticData();
+			StaticData.PendingRegistrationQueue.Add(TraitConstructor);
 		}
 	}
 
@@ -78,10 +96,11 @@ namespace UE::AnimNext
 		else
 		{
 			// Registry isn't ready yet or it got destroyed before the traits are unregistering
-			const int32 TraitIndex = Private::GPendingRegistrationQueue.IndexOfByKey(TraitConstructor);
+			Private::FTraitRegistryStaticData& StaticData = Private::GetTraitRegistryStaticData();
+			const int32 TraitIndex = StaticData.PendingRegistrationQueue.IndexOfByKey(TraitConstructor);
 			if (TraitIndex != INDEX_NONE)
 			{
-				Private::GPendingRegistrationQueue.RemoveAtSwap(TraitIndex);
+				StaticData.PendingRegistrationQueue.RemoveAtSwap(TraitIndex);
 			}
 		}
 	}
