Make sure UAnimNextAssetWorkspaceAssetUserData is marked as EditorOnly and excluded from cook/packaging

diff --git a/Engine/Plugins/Experimental/AnimNext/Source/AnimNextUncookedOnly/Internal/AnimNextAssetWorkspaceAssetUserData.h b/Engine/Plugins/Experimental/AnimNext/Source/AnimNextUncookedOnly/Internal/AnimNextAssetWorkspaceAssetUserData.h
--- a/Engine/Plugins/Experimental/AnimNext/Source/AnimNextUncookedOnly/Internal/AnimNextAssetWorkspaceAssetUserData.h
+++ b/Engine/Plugins/Experimental/AnimNext/Source/AnimNextUncookedOnly/Internal/AnimNextAssetWorkspaceAssetUserData.h
@@ -134,6 +134,10 @@ struct FAnimNextGraphOutlinerData : public FAnimNextAssetEntryOutlinerData
 UCLASS()
 class UAnimNextAssetWorkspaceAssetUserData : public UAssetUserData
 {
+public:
+	virtual bool IsEditorOnly() const override { return true; }
+
+private:
 	GENERATED_BODY()
 
 	virtual void GetAssetRegistryTags(FAssetRegistryTagsContext Context) const override;
