Composite character support in AnimNext

* Added new RigUnit to remap from one skeletal mesh to another.
* Added re-usable remap pose data so we can use it on other places later on (like a copy from mesh anim graph node).
* Added stats for it and separated the other remapping stats (from, to AnimBP, to local transforms)

diff --git a/Engine/Plugins/Experimental/AnimNext/Source/AnimNext/Private/GenerationTools.cpp b/Engine/Plugins/Experimental/AnimNext/Source/AnimNext/Private/GenerationTools.cpp
--- a/Engine/Plugins/Experimental/AnimNext/Source/AnimNext/Private/GenerationTools.cpp
+++ b/Engine/Plugins/Experimental/AnimNext/Source/AnimNext/Private/GenerationTools.cpp
@@ -20,7 +20,9 @@
 #include "AnimNextStats.h"
 
 DEFINE_STAT(STAT_AnimNext_GenerateReferencePose);
-DEFINE_STAT(STAT_AnimNext_RemapPose);
+DEFINE_STAT(STAT_AnimNext_RemapPose_FromAnimBP);
+DEFINE_STAT(STAT_AnimNext_RemapPose_ToAnimBP);
+DEFINE_STAT(STAT_AnimNext_RemapPose_ToLocalTransforms);
 DEFINE_STAT(STAT_AnimNext_ConvertLocalSpaceToComponentSpace);
 
 DEFINE_LOG_CATEGORY_STATIC(LogAnimGenerationTools, Log, All)
@@ -506,7 +508,7 @@ void FGenerationTools::FixLODRequiredBones(const int32 NumLODs
 // The target pose should be assigned to the correct reference pose prior to this call
 /*static*/ void FGenerationTools::RemapPose(const FPoseContext& SourcePose, FLODPose& TargetPose)
 {
-	SCOPE_CYCLE_COUNTER(STAT_AnimNext_RemapPose);
+	SCOPE_CYCLE_COUNTER(STAT_AnimNext_RemapPose_FromAnimBP);
 	
 	const FBoneContainer& BoneContainer = SourcePose.Pose.GetBoneContainer();
 	const FReferencePose& RefPose = TargetPose.GetRefPose();
@@ -543,7 +545,7 @@ void FGenerationTools::FixLODRequiredBones(const int32 NumLODs
 // The target pose should be assigned to the correct reference pose prior to this call
 /*static*/ void FGenerationTools::RemapPose(const FLODPose& SourcePose, FPoseContext& TargetPose)
 {
-	SCOPE_CYCLE_COUNTER(STAT_AnimNext_RemapPose);
+	SCOPE_CYCLE_COUNTER(STAT_AnimNext_RemapPose_ToAnimBP);
 	
 	const FBoneContainer& BoneContainer = TargetPose.Pose.GetBoneContainer();
 	const FReferencePose& RefPose = SourcePose.GetRefPose();
@@ -582,7 +584,7 @@ void FGenerationTools::FixLODRequiredBones(const int32 NumLODs
 // in by this call if they are not affected by the current LOD.
 /*static*/ void FGenerationTools::RemapPose(const FLODPose& SourcePose, TArrayView<FTransform> TargetTransforms)
 {
-	SCOPE_CYCLE_COUNTER(STAT_AnimNext_RemapPose);
+	SCOPE_CYCLE_COUNTER(STAT_AnimNext_RemapPose_ToLocalTransforms);
 	
 	const FReferencePose& RefPose = SourcePose.GetRefPose();
 	const TArrayView<const FBoneIndexType> LODBoneIndexes = RefPose.GetLODBoneIndexToMeshBoneIndexMap(SourcePose.LODLevel);
diff --git a/Engine/Plugins/Experimental/AnimNext/Source/AnimNext/Private/Graph/RigUnit_AnimNextRemapPose.cpp b/Engine/Plugins/Experimental/AnimNext/Source/AnimNext/Private/Graph/RigUnit_AnimNextRemapPose.cpp
new file mode 100644
--- /dev/null
+++ b/Engine/Plugins/Experimental/AnimNext/Source/AnimNext/Private/Graph/RigUnit_AnimNextRemapPose.cpp
@@ -0,0 +1,54 @@
+// Copyright Epic Games, Inc. All Rights Reserved.
+
+#include "RigUnit_AnimNextRemapPose.h"
+
+#include "Units/RigUnitContext.h"
+#include "GenerationTools.h"
+#include "Component/SkinnedMeshComponentExtensions.h"
+#include "Components/SkeletalMeshComponent.h"
+#include "Engine/SkeletalMesh.h"
+
+#include UE_INLINE_GENERATED_CPP_BY_NAME(RigUnit_AnimNextRemapPose)
+
+FRigUnit_AnimNextRemapPose_Execute()
+{
+	using namespace UE::AnimNext;
+
+	if(!TargetAnimGraphRefPose.ReferencePose.IsValid())
+	{
+		return;
+	}
+	
+	if(!Pose.LODPose.IsValid())
+	{
+		return;
+	}
+
+	FMemMark MemMark(FMemStack::Get());
+
+	const UE::AnimNext::FLODPoseHeap& SourcePose = Pose.LODPose;
+	const UE::AnimNext::FReferencePose& SourceRefPose = SourcePose.GetRefPose();
+	const USkeletalMesh* SourceMesh = SourceRefPose.SkeletalMesh.Get();
+
+	const UE::AnimNext::FReferencePose& TargetRefPose = TargetAnimGraphRefPose.ReferencePose.GetRef<UE::AnimNext::FReferencePose>();
+	UE::AnimNext::FLODPoseHeap& TargetPose = Result.LODPose;
+	const USkeletalMesh* TargetMesh = TargetRefPose.SkeletalMesh.Get();
+
+	if (SourceMesh == TargetMesh)
+	{
+		// Just copy the pose to the target in case there is nothing to convert.
+		TargetPose.CopyFrom(SourcePose);
+	}
+	else
+	{
+		// Pre-create the mapping and cache it to avoid runtime lookups.
+		if (RemapPoseData.ShouldReinit(SourceRefPose, TargetRefPose))
+		{
+			RemapPoseData.Reinit(SourceRefPose, TargetRefPose);
+		}
+
+		Result.Curves.CopyFrom(Pose.Curves);
+		RemapPoseData.RemapAttributes(Pose.LODPose, Pose.Attributes, Result.LODPose, Result.Attributes);
+		RemapPoseData.RemapPose(SourcePose, TargetPose);
+	}
+}
\ No newline at end of file
diff --git a/Engine/Plugins/Experimental/AnimNext/Source/AnimNext/Private/Graph/RigUnit_AnimNextRemapPose.h b/Engine/Plugins/Experimental/AnimNext/Source/AnimNext/Private/Graph/RigUnit_AnimNextRemapPose.h
new file mode 100644
--- /dev/null
+++ b/Engine/Plugins/Experimental/AnimNext/Source/AnimNext/Private/Graph/RigUnit_AnimNextRemapPose.h
@@ -0,0 +1,37 @@
+// Copyright Epic Games, Inc. All Rights Reserved.
+
+#pragma once
+
+#include "AnimNextExecuteContext.h"
+#include "Graph/AnimNext_LODPose.h"
+#include "Graph/RigUnit_AnimNextBase.h"
+#include "RemapPoseData.h"
+
+#include "RigUnit_AnimNextRemapPose.generated.h"
+
+/** Remaps an anim graph pose from to another skeletal mesh component. */
+USTRUCT(meta=(DisplayName="Remap Pose", Category="Animation Graph", NodeColor="0, 1, 1", Keywords="Output,Pose,Port"))
+struct ANIMNEXT_API FRigUnit_AnimNextRemapPose : public FRigUnit_AnimNextBase
+{
+	GENERATED_BODY()
+
+	RIGVM_METHOD()
+	void Execute();
+
+	// Pose to read
+	UPROPERTY(EditAnywhere, DisplayName = "Input Pose", Category = "Graph", meta = (Input))
+	FAnimNextGraphLODPose Pose;
+
+	UPROPERTY(EditAnywhere, DisplayName = "Converted Pose", Category = "Graph", meta = (Output))
+	FAnimNextGraphLODPose Result;
+
+	UPROPERTY(EditAnywhere, DisplayName = "Target Ref Pose", Category = "Graph", meta = (Input))
+	FAnimNextGraphReferencePose TargetAnimGraphRefPose;
+
+	// The execution result
+	UPROPERTY(EditAnywhere, DisplayName = "Execute", Category = "BeginExecution", meta = (Input, Output))
+	FAnimNextExecuteContext ExecuteContext;
+
+	UPROPERTY(meta = (Hidden))
+	FRemapPoseData RemapPoseData; // TODO: Store this per remap node or inside the data registry? We don't have ref counting ref poses and don't know when to remove these though.
+};
diff --git a/Engine/Plugins/Experimental/AnimNext/Source/AnimNext/Private/RemapPoseData.cpp b/Engine/Plugins/Experimental/AnimNext/Source/AnimNext/Private/RemapPoseData.cpp
new file mode 100644
--- /dev/null
+++ b/Engine/Plugins/Experimental/AnimNext/Source/AnimNext/Private/RemapPoseData.cpp
@@ -0,0 +1,157 @@
+﻿// Copyright Epic Games, Inc. All Rights Reserved.
+
+#include "RemapPoseData.h"
+#include "AnimNextStats.h"
+#include "Animation/Skeleton.h"
+#include "Engine/SkeletalMesh.h"
+
+DEFINE_STAT(STAT_AnimNext_RemapPose_Mesh2Mesh);
+
+bool FRemapPoseData::ShouldReinit(const UE::AnimNext::FReferencePose& InSourceRefPose,
+	const UE::AnimNext::FReferencePose& InTargetRefPose) const
+{
+	if (!SourceRefPose || !TargetRefPose)
+	{
+		// First time use.
+		return true;
+	}
+
+	if (InSourceRefPose.SkeletalMesh != SourceRefPose->SkeletalMesh)
+	{
+		// Source mesh changed.
+		return true;
+	}
+
+	if (InTargetRefPose.SkeletalMesh != TargetRefPose->SkeletalMesh)
+	{
+		// Target mesh changed.
+		return true;
+	}
+
+	return false;
+}
+
+void FRemapPoseData::Reinit(const UE::AnimNext::FReferencePose& InSourceRefPose,
+	const UE::AnimNext::FReferencePose& InTargetRefPose)
+{
+	SourceRefPose = &InSourceRefPose;
+	TargetRefPose = &InTargetRefPose;
+
+	const USkeletalMesh* SourceSkeletalMesh = InSourceRefPose.SkeletalMesh.Get();
+	const USkeleton* SourceSkeleton = SourceSkeletalMesh->GetSkeleton();
+	const FReferenceSkeleton& SourceRefSkeleton = SourceSkeleton->GetReferenceSkeleton();
+
+	const USkeletalMesh* TargetSkeletalMesh = InTargetRefPose.SkeletalMesh.Get();
+	const FReferenceSkeleton& TargetRefSkeleton = TargetSkeletalMesh->GetRefSkeleton();
+	const int32 TargetNumLODLevels = TargetSkeletalMesh->GetLODNum();
+
+	SourceToTargetBoneIndexMapPerLOD.Empty();
+	SourceToTargetBoneIndexMapPerLOD.AddDefaulted(TargetNumLODLevels);
+
+	for (int32 LODLevel = 0; LODLevel < TargetNumLODLevels; ++LODLevel)
+	{
+		const TArrayView<const FBoneIndexType> SourceSkeletonIndexToPoseIndex = InSourceRefPose.GetSkeletonBoneIndexToLODBoneIndexMap(LODLevel); // skeleton -> pose
+		const TArrayView<const FBoneIndexType> TargetPoseToMeshBoneIndexMap = InTargetRefPose.GetLODBoneIndexToMeshBoneIndexMap(LODLevel); // pose -> mesh
+		const int32 SourceNumBonesInLOD = SourceRefPose->GetNumBonesForLOD(LODLevel);
+
+		// Iterate over the bones present in the given LOD on the target skeleton.
+		// These bones are candidates for our mapping table, only the bones from the target skeleton that are available in the source skeleton are of interest.
+		for (FBoneIndexType BoneIndex = 0; BoneIndex < TargetPoseToMeshBoneIndexMap.Num(); ++BoneIndex)
+		{
+			const FName BoneName = TargetRefSkeleton.GetBoneName(TargetPoseToMeshBoneIndexMap[BoneIndex]);
+			const int32 SourceSkeletonBoneIndex = SourceRefSkeleton.FindBoneIndex(BoneName);
+
+			// Is the current bone on the target skeleton part of the source skeleton as well? Only add bones that are present on the target.
+			if (SourceSkeletonBoneIndex != INDEX_NONE)
+			{
+				// TODO: Do we need to do this at runtime? This might fail if the LOD level of the source and the target differ. We either move this to the runtime pose remapping or store a matrix of remappings.
+				const FBoneIndexType RemappedIndex = SourceSkeletonIndexToPoseIndex[SourceSkeletonBoneIndex];
+
+				// There are cases where the given bone are present in the skeleton, but not part of the actual pose/skeletal mesh,
+				// like e.g. when they are disabled for the given LOD on the target while they are not on the source. Skip these bones.
+				if (RemappedIndex != INDEX_NONE && RemappedIndex < SourceNumBonesInLOD)
+				{
+					FRemappedBone RemappedBone;
+					RemappedBone.SourceBoneIndex = RemappedIndex;
+					RemappedBone.TargetBoneIndex = BoneIndex;
+
+					SourceToTargetBoneIndexMapPerLOD[LODLevel].BoneIndexMap.Add(RemappedBone);
+				}
+			}
+		}
+	}
+}
+
+// TODO: Source pose dictates the LOD level of the target at the moment. What if the LOD levels of source and target are not in sync? We probably need to use the target skel mesh component LOD level when preparing the target pose.
+void FRemapPoseData::RemapPose(const UE::AnimNext::FLODPoseHeap& SourcePose,
+	UE::AnimNext::FLODPoseHeap& OutTargetPose) const
+{
+	SCOPE_CYCLE_COUNTER(STAT_AnimNext_RemapPose_Mesh2Mesh);
+
+	const int32 SourceLODLevel = SourcePose.LODLevel;
+
+	const bool IsAdditive = false;
+	if (OutTargetPose.ShouldPrepareForLOD(*TargetRefPose, SourceLODLevel, SourcePose.IsAdditive()))
+	{
+		OutTargetPose.PrepareForLOD(*TargetRefPose, SourceLODLevel, /*bSetRefPose=*/true, IsAdditive);
+	}
+
+	OutTargetPose.Flags = SourcePose.Flags;
+
+	const TArray<FRemappedBone>& SourceToTargetBoneIndexMap = SourceToTargetBoneIndexMapPerLOD[SourceLODLevel].BoneIndexMap;
+	for (int32 MappingElementIndex = 0; MappingElementIndex < SourceToTargetBoneIndexMap.Num(); ++MappingElementIndex)
+	{
+		const int32 SourceBoneIndex = SourceToTargetBoneIndexMap[MappingElementIndex].SourceBoneIndex;
+		const int32 TargetBoneIndex = SourceToTargetBoneIndexMap[MappingElementIndex].TargetBoneIndex;
+
+		const auto& TargetTransform = SourcePose.LocalTransformsView[SourceBoneIndex];
+		OutTargetPose.LocalTransforms[TargetBoneIndex] = TargetTransform;
+	}
+}
+
+void FRemapPoseData::RemapAttributes(const UE::AnimNext::FLODPose& SourceLODPose,
+	const UE::Anim::FHeapAttributeContainer& InAttributes,
+	const UE::AnimNext::FLODPose& TargetLODPose,
+	UE::Anim::FHeapAttributeContainer& OutAttributes)
+{
+	const USkeletalMesh* SourceSkeletalMesh = SourceRefPose->SkeletalMesh.Get();
+	const FReferenceSkeleton& SourceRefSkeleton = SourceSkeletalMesh->GetRefSkeleton();
+	const TArrayView<const FBoneIndexType> SourceLODBoneIndexToSkeletonBoneIndexMap = SourceLODPose.GetLODBoneIndexToSkeletonBoneIndexMap();
+
+	const USkeletalMesh* TargetSkeletalMesh = TargetRefPose->SkeletalMesh.Get();
+	const FReferenceSkeleton& TargetRefSkeleton = TargetSkeletalMesh->GetRefSkeleton();
+	const TArrayView<const FBoneIndexType> TargetSkeletonBoneToLODBoneIndexMap = TargetLODPose.GetSkeletonBoneIndexToLODBoneIndexMap();
+
+	for (const TWeakObjectPtr<UScriptStruct> WeakScriptStruct : InAttributes.GetUniqueTypes())
+	{
+		const UScriptStruct* ScriptStruct = WeakScriptStruct.Get();
+		const int32 TypeIndex = InAttributes.FindTypeIndex(ScriptStruct);
+		if (TypeIndex != INDEX_NONE)
+		{
+			const TArray<UE::Anim::TWrappedAttribute<FDefaultAllocator>, FDefaultAllocator>& SourceValues = InAttributes.GetValues(TypeIndex);
+			const TArray<UE::Anim::FAttributeId, FDefaultAllocator>& AttributeIds = InAttributes.GetKeys(TypeIndex);
+
+			// Try and remap all the source attributes to their respective new bone indices
+			for (int32 EntryIndex = 0; EntryIndex < AttributeIds.Num(); ++EntryIndex)
+			{
+				const UE::Anim::FAttributeId& AttributeId = AttributeIds[EntryIndex];
+
+				// Remap the source bone from LOD pose bone index to a skeleton bone index and get the bone name.
+				const int32 SourceLODBoneIndex = AttributeId.GetIndex();
+				const int32 SourceSkeletonBoneIndex = SourceLODBoneIndexToSkeletonBoneIndexMap[SourceLODBoneIndex];
+				const FName BoneName = SourceRefSkeleton.GetBoneName(SourceSkeletonBoneIndex);
+				
+				// Find the given bone inside the target skeleton, if it is in there, remap the attribute.
+				const int32 TargetSkeletonBoneIndex = TargetRefSkeleton.FindBoneIndex(BoneName);
+				if (TargetSkeletonBoneIndex != INDEX_NONE)
+				{
+					const int32 TargetLODBoneIndex = TargetSkeletonBoneToLODBoneIndexMap[TargetSkeletonBoneIndex];
+
+					const UE::Anim::FAttributeId NewInfo(AttributeId.GetName(), TargetLODBoneIndex, AttributeId.GetNamespace());
+					uint8* NewAttribute = OutAttributes.FindOrAdd(ScriptStruct, NewInfo);
+					ScriptStruct->CopyScriptStruct(NewAttribute, SourceValues[EntryIndex].GetPtr<void>());
+				}
+			}
+		}
+	}
+}
\ No newline at end of file
diff --git a/Engine/Plugins/Experimental/AnimNext/Source/AnimNext/Public/AnimNextStats.h b/Engine/Plugins/Experimental/AnimNext/Source/AnimNext/Public/AnimNextStats.h
--- a/Engine/Plugins/Experimental/AnimNext/Source/AnimNext/Public/AnimNextStats.h
+++ b/Engine/Plugins/Experimental/AnimNext/Source/AnimNext/Public/AnimNextStats.h
@@ -28,7 +28,10 @@ DECLARE_CYCLE_STAT_EXTERN(TEXT("AnimNext: Param Stack Decoalesce"), STAT_AnimNex
 DECLARE_CYCLE_STAT_EXTERN(TEXT("AnimNext: Update Param Layer"), STAT_AnimNext_Graph_UpdateParamLayer, STATGROUP_AnimNext, ANIMNEXT_API);
 
 DECLARE_CYCLE_STAT_EXTERN(TEXT("AnimNext: Generate Reference Pose"), STAT_AnimNext_GenerateReferencePose, STATGROUP_AnimNext, ANIMNEXT_API);
-DECLARE_CYCLE_STAT_EXTERN(TEXT("AnimNext: Remap Pose"), STAT_AnimNext_RemapPose, STATGROUP_AnimNext, ANIMNEXT_API);
+DECLARE_CYCLE_STAT_EXTERN(TEXT("AnimNext: Remap Pose (From AnimBP)"), STAT_AnimNext_RemapPose_FromAnimBP, STATGROUP_AnimNext, ANIMNEXT_API);
+DECLARE_CYCLE_STAT_EXTERN(TEXT("AnimNext: Remap Pose (To AnimBP)"), STAT_AnimNext_RemapPose_ToAnimBP, STATGROUP_AnimNext, ANIMNEXT_API);
+DECLARE_CYCLE_STAT_EXTERN(TEXT("AnimNext: Remap Pose (To LocalTransforms)"), STAT_AnimNext_RemapPose_ToLocalTransforms, STATGROUP_AnimNext, ANIMNEXT_API);
+DECLARE_CYCLE_STAT_EXTERN(TEXT("AnimNext: Remap Pose (Mesh->Mesh)"), STAT_AnimNext_RemapPose_Mesh2Mesh, STATGROUP_AnimNext, ANIMNEXT_API);
 DECLARE_CYCLE_STAT_EXTERN(TEXT("AnimNext: Convert Local Space To Component Space"), STAT_AnimNext_ConvertLocalSpaceToComponentSpace, STATGROUP_AnimNext, ANIMNEXT_API);
 
 DECLARE_CYCLE_STAT_EXTERN(TEXT("AnimNext: Skeletal Mesh Component Port"), STAT_AnimNext_Port_SkeletalMeshComponent, STATGROUP_AnimNext, ANIMNEXT_API);
diff --git a/Engine/Plugins/Experimental/AnimNext/Source/AnimNext/Public/RemapPoseData.h b/Engine/Plugins/Experimental/AnimNext/Source/AnimNext/Public/RemapPoseData.h
new file mode 100644
--- /dev/null
+++ b/Engine/Plugins/Experimental/AnimNext/Source/AnimNext/Public/RemapPoseData.h
@@ -0,0 +1,67 @@
+// Copyright Epic Games, Inc. All Rights Reserved.
+
+#pragma once
+
+#include "Animation/AttributesRuntime.h"
+#include "ReferencePose.h"
+#include "LODPose.h"
+
+#include "RemapPoseData.generated.h"
+
+USTRUCT()
+struct FRemappedBone
+{
+	GENERATED_BODY()
+
+	int32 SourceBoneIndex = INDEX_NONE;
+	int32 TargetBoneIndex = INDEX_NONE;
+};
+
+USTRUCT()
+struct FBoneRemapping
+{
+	GENERATED_BODY()
+
+	/**
+	 * Source to target bone index array for bones from the target pose that are also present in the source pose are stored.
+	 * This is a subset of the bones from the target pose. The smaller the subset, the less bones to remap.
+	 */
+	TArray<FRemappedBone> BoneIndexMap;
+};
+
+USTRUCT()
+struct FRemapPoseData
+{
+	GENERATED_BODY()
+
+	/**
+	 * Check if we need to re-initialize the mapping. This will return true in case of first time usage of the data,
+	 * or in the case that either the source or target pose changed.
+	 */
+	bool ShouldReinit(const UE::AnimNext::FReferencePose& InSourceRefPose,
+		const UE::AnimNext::FReferencePose& InTargetRefPose) const;
+
+	/**
+	 * Re-initializes the mapping for each LOD. Call ShouldReinit() before to avoid unnecessary compute.
+	 * Compare the bones from the source and target pose and find the subset of the target that is also present in the source.
+	 */
+	void Reinit(const UE::AnimNext::FReferencePose& InSourceRefPose,
+		const UE::AnimNext::FReferencePose& InTargetRefPose);
+
+	/**
+	 * Remap a pose using the previously cached mapping. Call this at runtime.
+	 */
+	void RemapPose(const UE::AnimNext::FLODPoseHeap& SourcePose,
+		UE::AnimNext::FLODPoseHeap& OutTargetPose) const;
+
+	void RemapAttributes(const UE::AnimNext::FLODPose& SourceLODPose,
+		const UE::Anim::FHeapAttributeContainer& InAttributes,
+		const UE::AnimNext::FLODPose& TargetLODPose,
+		UE::Anim::FHeapAttributeContainer& OutAttributes);
+
+	const UE::AnimNext::FReferencePose* SourceRefPose = nullptr;
+	const UE::AnimNext::FReferencePose* TargetRefPose = nullptr;
+
+	/** Bone mapping for each LOD level. */
+	TArray<FBoneRemapping> SourceToTargetBoneIndexMapPerLOD;
+};
\ No newline at end of file
diff --git a/Engine/Source/Runtime/Engine/Classes/Components/SkeletalMeshComponent.h b/Engine/Source/Runtime/Engine/Classes/Components/SkeletalMeshComponent.h
--- a/Engine/Source/Runtime/Engine/Classes/Components/SkeletalMeshComponent.h
+++ b/Engine/Source/Runtime/Engine/Classes/Components/SkeletalMeshComponent.h
@@ -309,6 +309,7 @@ class USkeletalMeshComponent : public USkinnedMeshComponent, public IInterface_C
 	friend struct FLinkedInstancesAdapter;
 	friend struct FLinkedAnimLayerClassData;
 	friend struct FRigUnit_AnimNextWriteSkeletalMeshComponentPose;
+	friend struct FRigUnit_AnimNextRemapPose;
 	
 #if WITH_EDITORONLY_DATA 
 private:
