Add scope cycle stats to various AnimNext functions

diff --git a/Engine/Plugins/Experimental/AnimNext/Source/AnimNext/Private/Graph/RigUnit_MakeReferencePoseFromSkeletalMeshComponent.cpp b/Engine/Plugins/Experimental/AnimNext/Source/AnimNext/Private/Graph/RigUnit_MakeReferencePoseFromSkeletalMeshComponent.cpp
--- a/Engine/Plugins/Experimental/AnimNext/Source/AnimNext/Private/Graph/RigUnit_MakeReferencePoseFromSkeletalMeshComponent.cpp
+++ b/Engine/Plugins/Experimental/AnimNext/Source/AnimNext/Private/Graph/RigUnit_MakeReferencePoseFromSkeletalMeshComponent.cpp
@@ -1,16 +1,21 @@
-﻿// Copyright Epic Games, Inc. All Rights Reserved.
+// Copyright Epic Games, Inc. All Rights Reserved.
 
 #include "RigUnit_MakeReferencePoseFromSkeletalMeshComponent.h"
 
 #include "Units/RigUnitContext.h"
+#include "AnimNextStats.h"
 #include "GenerationTools.h"
 #include "Components/SkeletalMeshComponent.h"
 #include "Engine/SkeletalMesh.h"
 
 #include UE_INLINE_GENERATED_CPP_BY_NAME(RigUnit_MakeReferencePoseFromSkeletalMeshComponent)
 
+DEFINE_STAT(STAT_AnimNext_Make_RefPose);
+
 FRigUnit_MakeReferencePoseFromSkeletalMeshComponent_Execute()
 {
+	SCOPE_CYCLE_COUNTER(STAT_AnimNext_Make_RefPose);
+
 	using namespace UE::AnimNext;
 
 	if(SkeletalMeshComponent == nullptr)
diff --git a/Engine/Plugins/Experimental/AnimNext/Source/AnimNext/Private/Module/ModuleTickFunction.cpp b/Engine/Plugins/Experimental/AnimNext/Source/AnimNext/Private/Module/ModuleTickFunction.cpp
--- a/Engine/Plugins/Experimental/AnimNext/Source/AnimNext/Private/Module/ModuleTickFunction.cpp
+++ b/Engine/Plugins/Experimental/AnimNext/Source/AnimNext/Private/Module/ModuleTickFunction.cpp
@@ -3,6 +3,7 @@
 #include "Module/ModuleTickFunction.h"
 
 #include "AnimNextDebugDraw.h"
+#include "AnimNextStats.h"
 #include "RigVMRuntimeDataRegistry.h"
 #include "Algo/TopologicalSort.h"
 #include "Module/AnimNextModule.h"
@@ -12,6 +13,9 @@
 #include "Module/ModuleEvents.h"
 #include "TraitCore/TraitEventList.h"
 
+DEFINE_STAT(STAT_AnimNext_Module_EventTick);
+DEFINE_STAT(STAT_AnimNext_Module_EndTick);
+
 namespace UE::AnimNext
 {
 
@@ -22,6 +26,8 @@ void FModuleEndTickFunction::ExecuteTick(float DeltaTime, ELevelTick TickType, E
 
 void FModuleEndTickFunction::Run()
 {
+	SCOPE_CYCLE_COUNTER(STAT_AnimNext_Module_EndTick);
+
 	// Decrement the remaining lifetime of the input events we processed and queue up any remaining events
 	DecrementLifetimeAndPurgeExpired(ModuleInstance->InputEventList, ModuleInstance->OutputEventList);
 
@@ -134,6 +140,8 @@ void FModuleEventTickFunction::ExecuteTick(float DeltaTime, ELevelTick TickType,
 
 void FModuleEventTickFunction::Run(float InDeltaTime)
 {
+	SCOPE_CYCLE_COUNTER(STAT_AnimNext_Module_EventTick);
+
 	while (!PreExecuteTasks.IsEmpty())
 	{
 		TOptional<TUniqueFunction<void(const FModuleTaskContext&)>> Function = PreExecuteTasks.Dequeue();
diff --git a/Engine/Plugins/Experimental/AnimNext/Source/AnimNext/Public/AnimNextStats.h b/Engine/Plugins/Experimental/AnimNext/Source/AnimNext/Public/AnimNextStats.h
--- a/Engine/Plugins/Experimental/AnimNext/Source/AnimNext/Public/AnimNextStats.h
+++ b/Engine/Plugins/Experimental/AnimNext/Source/AnimNext/Public/AnimNextStats.h
@@ -6,9 +6,13 @@
 
 DECLARE_STATS_GROUP(TEXT("AnimNext"), STATGROUP_AnimNext, STATCAT_Advanced);
 
+DECLARE_CYCLE_STAT_EXTERN(TEXT("AnimNext: Module Event Tick"), STAT_AnimNext_Module_EventTick, STATGROUP_AnimNext, ANIMNEXT_API);
+DECLARE_CYCLE_STAT_EXTERN(TEXT("AnimNext: Module End Tick"), STAT_AnimNext_Module_EndTick, STATGROUP_AnimNext, ANIMNEXT_API);
+
 DECLARE_CYCLE_STAT_EXTERN(TEXT("AnimNext: Initialize Instance"), STAT_AnimNext_InitializeInstance, STATGROUP_AnimNext, ANIMNEXT_API);
 DECLARE_CYCLE_STAT_EXTERN(TEXT("AnimNext: Create Instance Data"), STAT_AnimNext_CreateInstanceData, STATGROUP_AnimNext, ANIMNEXT_API);
 
+DECLARE_CYCLE_STAT_EXTERN(TEXT("AnimNext: Make RefPose"), STAT_AnimNext_Make_RefPose, STATGROUP_AnimNext, ANIMNEXT_API);
 DECLARE_CYCLE_STAT_EXTERN(TEXT("AnimNext: Run Graph"), STAT_AnimNext_Run_Graph, STATGROUP_AnimNext, ANIMNEXT_API);
 DECLARE_CYCLE_STAT_EXTERN(TEXT("AnimNext: Write Pose"), STAT_AnimNext_Write_Pose, STATGROUP_AnimNext, ANIMNEXT_API);
 
