AnimNext: Fix public variables pin-exposure load/compile order dependency

Move the pin-exposure support to depend on the source editor data (the variable entries) rather than the compiled data (as compiled data may not be saved, so reloading may not correspond to the last-observed state)

#rnx

diff --git a/Engine/Plugins/Experimental/AnimNext/Config/BaseAnimNext.ini b/Engine/Plugins/Experimental/AnimNext/Config/BaseAnimNext.ini
--- a/Engine/Plugins/Experimental/AnimNext/Config/BaseAnimNext.ini
+++ b/Engine/Plugins/Experimental/AnimNext/Config/BaseAnimNext.ini
@@ -63,3 +63,6 @@
 
 +EnumRedirects=(OldName="EAnimNextModuleInitMethod",NewName="/Script/AnimNext.EAnimNextModuleInitMethod")
 +EnumRedirects=(OldName="EAnimNextParameterFlags",NewName="/Script/AnimNext.EAnimNextExportedVariableFlags")
+
++PropertyRedirects=(OldName="/Script/AnimNext.RigVMTrait_AnimNextPublicVariables.Asset", NewName="InternalAsset")
++PropertyRedirects=(OldName="/Script/AnimNext.RigVMTrait_AnimNextPublicVariables.VariableNames", NewName="InternalVariableNames")
\ No newline at end of file
diff --git a/Engine/Plugins/Experimental/AnimNext/Source/AnimNext/Internal/Graph/RigVMTrait_AnimNextPublicVariables.h b/Engine/Plugins/Experimental/AnimNext/Source/AnimNext/Internal/Graph/RigVMTrait_AnimNextPublicVariables.h
--- a/Engine/Plugins/Experimental/AnimNext/Source/AnimNext/Internal/Graph/RigVMTrait_AnimNextPublicVariables.h
+++ b/Engine/Plugins/Experimental/AnimNext/Source/AnimNext/Internal/Graph/RigVMTrait_AnimNextPublicVariables.h
@@ -8,6 +8,7 @@
 
 class UAnimNextRigVMAsset;
 class FRigVMTraitScope;
+class UAnimNextDataInterface;
 
 // Represents public variables of an asset via a trait 
 USTRUCT(BlueprintType)
@@ -17,17 +18,22 @@ struct ANIMNEXT_API FRigVMTrait_AnimNextPublicVariables : public FRigVMTrait
 
 	// The data interface that any programmatic pins will be derived from
 	UPROPERTY(meta = (Hidden))
-	TObjectPtr<UAnimNextDataInterface> Asset = nullptr;
+	TObjectPtr<UAnimNextDataInterface> InternalAsset = nullptr;
 
 	// Variable names that are exposed
 	UPROPERTY(meta = (Hidden))
-	TArray<FName> VariableNames;
+	TArray<FName> InternalVariableNames;
 
 	// FRigVMTrait interface
 #if WITH_EDITOR
 	virtual FString GetDisplayName() const override;
 	virtual void GetProgrammaticPins(URigVMController* InController, int32 InParentPinIndex, const FString& InDefaultValue, FRigVMPinInfoArray& OutPinArray) const override;
 	virtual bool ShouldCreatePinForProperty(const FProperty* InProperty) const override;
+
+	// Editor-implemented function ptrs
+	static FString (*GetDisplayNameFunc)(const FRigVMTrait_AnimNextPublicVariables& InTrait);
+	static void (*GetProgrammaticPinsFunc)(const FRigVMTrait_AnimNextPublicVariables& InTrait, URigVMController* InController, int32 InParentPinIndex, const FString& InDefaultValue, FRigVMPinInfoArray& OutPinArray);
+	static bool (*ShouldCreatePinForPropertyFunc)(const FRigVMTrait_AnimNextPublicVariables& InTrait, const FProperty* InProperty);
 #endif
 };
 
diff --git a/Engine/Plugins/Experimental/AnimNext/Source/AnimNext/Private/Graph/RigVMTrait_AnimNextPublicVariables.cpp b/Engine/Plugins/Experimental/AnimNext/Source/AnimNext/Private/Graph/RigVMTrait_AnimNextPublicVariables.cpp
--- a/Engine/Plugins/Experimental/AnimNext/Source/AnimNext/Private/Graph/RigVMTrait_AnimNextPublicVariables.cpp
+++ b/Engine/Plugins/Experimental/AnimNext/Source/AnimNext/Private/Graph/RigVMTrait_AnimNextPublicVariables.cpp
@@ -3,49 +3,20 @@
 #include "Graph/RigVMTrait_AnimNextPublicVariables.h"
 #include "AnimNextRigVMAsset.h"
 #include "Logging/StructuredLog.h"
-#if WITH_EDITOR
-#include "RigVMModel/RigVMController.h"
-#include "RigVMModel/RigVMPin.h"
-#endif
 
 #if WITH_EDITOR
+FString (*FRigVMTrait_AnimNextPublicVariables::GetDisplayNameFunc)(const FRigVMTrait_AnimNextPublicVariables& InTrait);
+void (*FRigVMTrait_AnimNextPublicVariables::GetProgrammaticPinsFunc)(const FRigVMTrait_AnimNextPublicVariables& InTrait, URigVMController* InController, int32 InParentPinIndex, const FString& InDefaultValue, FRigVMPinInfoArray& OutPinArray);
+bool (*FRigVMTrait_AnimNextPublicVariables::ShouldCreatePinForPropertyFunc)(const FRigVMTrait_AnimNextPublicVariables& InTrait, const FProperty* InProperty);
+
 FString FRigVMTrait_AnimNextPublicVariables::GetDisplayName() const
 {
-	TStringBuilder<256> StringBuilder;
-	StringBuilder.Appendf(TEXT("Variables: %s"), Asset ? *Asset->GetFName().ToString() : TEXT("None"));
-	return StringBuilder.ToString();
+	return GetDisplayNameFunc(*this);
 }
 
 void FRigVMTrait_AnimNextPublicVariables::GetProgrammaticPins(URigVMController* InController, int32 InParentPinIndex, const FString& InDefaultValue, FRigVMPinInfoArray& OutPinArray) const
 {
-	if (Asset == nullptr)
-	{
-		return;
-	}
-
-	const FInstancedPropertyBag& UserDefaults = Asset->GetPublicVariableDefaults();
-	if(!UserDefaults.IsValid())
-	{
-		return;
-	}
-
-	FInstancedPropertyBag Defaults;
-	Defaults.InitializeFromBagStruct(UserDefaults.GetPropertyBagStruct());
-
-	const TFunction<ERigVMPinDefaultValueType(const FName&)> DefaultValueTypeGetter = [&UserDefaults, &Defaults](const FName& InPropertyName)
-	{
-		if(const FProperty* Property = UserDefaults.GetPropertyBagStruct()->FindPropertyByName(InPropertyName))
-		{
-			if(Property->Identical_InContainer(UserDefaults.GetValue().GetMemory(), Defaults.GetValue().GetMemory()))
-			{
-				return ERigVMPinDefaultValueType::Unset;
-			}
-			return ERigVMPinDefaultValueType::Override;
-		}
-		return ERigVMPinDefaultValueType::AutoDetect;
-	};
-
-	OutPinArray.AddPins(const_cast<UPropertyBag*>(UserDefaults.GetPropertyBagStruct()), InController, ERigVMPinDirection::Input, InParentPinIndex, DefaultValueTypeGetter, UserDefaults.GetValue().GetMemory(), true);
+	GetProgrammaticPinsFunc(*this, InController, InParentPinIndex, InDefaultValue, OutPinArray);
 }
 
 bool FRigVMTrait_AnimNextPublicVariables::ShouldCreatePinForProperty(const FProperty* InProperty) const
@@ -54,10 +25,7 @@ bool FRigVMTrait_AnimNextPublicVariables::ShouldCreatePinForProperty(const FProp
 	{
 		return false;
 	}
-	return
-		InProperty->GetFName() == GET_MEMBER_NAME_CHECKED(FRigVMTrait_AnimNextPublicVariables, Asset) ||
-		InProperty->GetFName() == GET_MEMBER_NAME_CHECKED(FRigVMTrait_AnimNextPublicVariables, VariableNames) ||
-		VariableNames.Contains(InProperty->GetFName());
+	return ShouldCreatePinForPropertyFunc(*this, InProperty);
 }
 #endif
 
@@ -66,14 +34,14 @@ namespace UE::AnimNext
 
 const UAnimNextDataInterface* FPublicVariablesTraitToDataInterfaceHostAdapter::GetDataInterface() const
 {
-	return Trait.Asset;
+	return Trait.InternalAsset;
 }
 
 uint8* FPublicVariablesTraitToDataInterfaceHostAdapter::GetMemoryForVariable(int32 InVariableIndex, FName InVariableName, const FProperty* InVariableProperty) const
 {
 	// Note we dont use InVariableIndex here as we may not have bound all variables in an interface
 
-	int32 TraitVariableIndex = Trait.VariableNames.Find(InVariableName);
+	int32 TraitVariableIndex = Trait.InternalVariableNames.Find(InVariableName);
 	if(TraitVariableIndex == INDEX_NONE)
 	{
 		// Variable not bound here
diff --git a/Engine/Plugins/Experimental/AnimNext/Source/AnimNextEditor/Private/AnimNextEditorModule.cpp b/Engine/Plugins/Experimental/AnimNext/Source/AnimNextEditor/Private/AnimNextEditorModule.cpp
--- a/Engine/Plugins/Experimental/AnimNext/Source/AnimNextEditor/Private/AnimNextEditorModule.cpp
+++ b/Engine/Plugins/Experimental/AnimNext/Source/AnimNextEditor/Private/AnimNextEditorModule.cpp
@@ -48,6 +48,7 @@
 #include "Graph/AnimNextAnimationGraph_EditorData.h"
 #include "Common/AnimNextAssetItemDetails.h"
 #include "DataInterface/AnimNextDataInterface.h"
+#include "Graph/RigVMTrait_AnimNextPublicVariablesEditor.h"
 #include "Module/RigUnit_AnimNextModuleEvents.h"
 #include "Variables/SVariablesView.h"
 #include "Variables/VariableOverrideCommands.h"
@@ -61,6 +62,7 @@ namespace UE::AnimNext::Editor
 
 void FAnimNextEditorModule::StartupModule()
 {
+	FPublicVariablesImpl::Register();
 	FVariableOverrideCommands::Register();
 
 	// Register settings for user editing
diff --git a/Engine/Plugins/Experimental/AnimNext/Source/AnimNextEditor/Private/Graph/RigVMTrait_AnimNextPublicVariablesEditor.cpp b/Engine/Plugins/Experimental/AnimNext/Source/AnimNextEditor/Private/Graph/RigVMTrait_AnimNextPublicVariablesEditor.cpp
new file mode 100644
--- /dev/null
+++ b/Engine/Plugins/Experimental/AnimNext/Source/AnimNextEditor/Private/Graph/RigVMTrait_AnimNextPublicVariablesEditor.cpp
@@ -0,0 +1,80 @@
+﻿// Copyright Epic Games, Inc. All Rights Reserved.
+
+#include "RigVMTrait_AnimNextPublicVariablesEditor.h"
+#include "AnimNextRigVMAssetEditorData.h"
+#include "UncookedOnlyUtils.h"
+#include "Entries/AnimNextVariableEntry.h"
+#include "Graph/RigVMTrait_AnimNextPublicVariables.h"
+#include "RigVMModel/RigVMController.h"
+#include "RigVMModel/RigVMPin.h"
+
+namespace UE::AnimNext::Editor
+{
+
+void FPublicVariablesImpl::Register()
+{
+	FRigVMTrait_AnimNextPublicVariables::GetDisplayNameFunc = GetDisplayName;
+	FRigVMTrait_AnimNextPublicVariables::GetProgrammaticPinsFunc = GetProgrammaticPins;
+	FRigVMTrait_AnimNextPublicVariables::ShouldCreatePinForPropertyFunc = ShouldCreatePinForProperty;
+}
+
+FString FPublicVariablesImpl::GetDisplayName(const FRigVMTrait_AnimNextPublicVariables& InTrait)
+{
+	TStringBuilder<256> StringBuilder;
+	StringBuilder.Appendf(TEXT("Variables: %s"), InTrait.InternalAsset ? *InTrait.InternalAsset->GetFName().ToString() : TEXT("None"));
+	return StringBuilder.ToString();
+}
+
+void FPublicVariablesImpl::GetProgrammaticPins(const FRigVMTrait_AnimNextPublicVariables& InTrait, URigVMController* InController, int32 InParentPinIndex, const FString& InDefaultValue, FRigVMPinInfoArray& OutPinArray)
+{
+	if (InTrait.InternalAsset == nullptr)
+	{
+		return;
+	}
+
+	UAnimNextRigVMAssetEditorData* EditorData = UE::AnimNext::UncookedOnly::FUtils::GetEditorData<UAnimNextRigVMAssetEditorData>(InTrait.InternalAsset.Get());
+	if(EditorData == nullptr)
+	{
+		return;
+	}
+
+	TArray<UAnimNextVariableEntry*> PublicVariables;
+	EditorData->GetPublicVariables(PublicVariables);
+	if(PublicVariables.Num() == 0)
+	{
+		return;
+	}
+
+	TArray<FPropertyBagPropertyDesc> Descs;
+	Descs.Reserve(PublicVariables.Num());
+	TArray<TConstArrayView<uint8>> Values;
+	Values.Reserve(PublicVariables.Num());
+	for(UAnimNextVariableEntry* PublicVariable : PublicVariables)
+	{
+		const FProperty* Property;
+		TConstArrayView<uint8> Value;
+		PublicVariable->GetDefaultValue(Property, Value);
+		Descs.Emplace(PublicVariable->GetEntryName(), Property);
+		Values.Emplace(Value);
+	}
+
+	FInstancedPropertyBag Defaults;
+	Defaults.ReplaceAllPropertiesAndValues(Descs, Values);
+
+	const TFunction<ERigVMPinDefaultValueType(const FName&)> DefaultValueTypeGetter = [](const FName& InPropertyName)
+	{
+		return ERigVMPinDefaultValueType::AutoDetect;
+	};
+
+	OutPinArray.AddPins(const_cast<UPropertyBag*>(Defaults.GetPropertyBagStruct()), InController, ERigVMPinDirection::Input, InParentPinIndex, DefaultValueTypeGetter, Defaults.GetValue().GetMemory(), true);
+}
+
+bool FPublicVariablesImpl::ShouldCreatePinForProperty(const FRigVMTrait_AnimNextPublicVariables& InTrait, const FProperty* InProperty)
+{
+	return
+		InProperty->GetFName() == GET_MEMBER_NAME_CHECKED(FRigVMTrait_AnimNextPublicVariables, InternalAsset) ||
+		InProperty->GetFName() == GET_MEMBER_NAME_CHECKED(FRigVMTrait_AnimNextPublicVariables, InternalVariableNames) ||
+		InTrait.InternalVariableNames.Contains(InProperty->GetFName());
+}
+
+}
\ No newline at end of file
diff --git a/Engine/Plugins/Experimental/AnimNext/Source/AnimNextEditor/Private/Graph/RigVMTrait_AnimNextPublicVariablesEditor.h b/Engine/Plugins/Experimental/AnimNext/Source/AnimNextEditor/Private/Graph/RigVMTrait_AnimNextPublicVariablesEditor.h
new file mode 100644
--- /dev/null
+++ b/Engine/Plugins/Experimental/AnimNext/Source/AnimNextEditor/Private/Graph/RigVMTrait_AnimNextPublicVariablesEditor.h
@@ -0,0 +1,23 @@
+﻿// Copyright Epic Games, Inc. All Rights Reserved.
+
+#pragma once
+
+#include "CoreMinimal.h"
+
+struct FRigVMTrait_AnimNextPublicVariables;
+class URigVMController;
+struct FRigVMPinInfoArray;
+
+namespace UE::AnimNext::Editor
+{
+
+// Editor-side impl for FRigVMTrait_AnimNextPublicVariables
+struct ANIMNEXTEDITOR_API FPublicVariablesImpl
+{
+	static void Register();
+	static FString GetDisplayName(const FRigVMTrait_AnimNextPublicVariables& InTrait);
+	static void GetProgrammaticPins(const FRigVMTrait_AnimNextPublicVariables& InTrait, URigVMController* InController, int32 InParentPinIndex, const FString& InDefaultValue, FRigVMPinInfoArray& OutPinArray);
+	static bool ShouldCreatePinForProperty(const FRigVMTrait_AnimNextPublicVariables& InTrait, const FProperty* InProperty);
+};
+
+}
\ No newline at end of file
diff --git a/Engine/Plugins/Experimental/AnimNext/Source/AnimNextUncookedOnly/Internal/AnimNextRigVMAssetEditorData.h b/Engine/Plugins/Experimental/AnimNext/Source/AnimNextUncookedOnly/Internal/AnimNextRigVMAssetEditorData.h
--- a/Engine/Plugins/Experimental/AnimNext/Source/AnimNextUncookedOnly/Internal/AnimNextRigVMAssetEditorData.h
+++ b/Engine/Plugins/Experimental/AnimNext/Source/AnimNextUncookedOnly/Internal/AnimNextRigVMAssetEditorData.h
@@ -41,6 +41,7 @@ namespace UE::AnimNext::Editor
 	class SVariablesOutliner;
 	class SAddVariablesDialog;
 	class FVariableProxyCustomization;
+	struct FPublicVariablesImpl;
 }
 
 namespace UE::AnimNext::Tests
@@ -162,6 +163,7 @@ protected:
 	friend class UE::AnimNext::Editor::SAddVariablesDialog;
 	friend class UE::AnimNext::Editor::FVariableProxyCustomization;
 	friend class UAnimNextDataInterfaceEntry;
+	friend struct UE::AnimNext::Editor::FPublicVariablesImpl;
 
 	// UObject interface
 	virtual void Serialize(FArchive& Ar) override;
@@ -385,6 +387,10 @@ protected:
 	// Checks all entries to see if any are public variables
 	bool HasPublicVariables() const;
 
+	// Gets any public variables that this asset has. Variables have no specified order.
+	// May recurse into implemented data interfaces, so public variables returned by this function may not be directly owned by this asset.
+	void GetPublicVariables(TArray<UAnimNextVariableEntry*>& OutPublicVariables) const;
+
 	// Refresh the 'external' models for the RigVM client to reference
 	void RefreshExternalModels();
 
diff --git a/Engine/Plugins/Experimental/AnimNext/Source/AnimNextUncookedOnly/Internal/Entries/AnimNextVariableEntry.h b/Engine/Plugins/Experimental/AnimNext/Source/AnimNextUncookedOnly/Internal/Entries/AnimNextVariableEntry.h
--- a/Engine/Plugins/Experimental/AnimNext/Source/AnimNextUncookedOnly/Internal/Entries/AnimNextVariableEntry.h
+++ b/Engine/Plugins/Experimental/AnimNext/Source/AnimNextUncookedOnly/Internal/Entries/AnimNextVariableEntry.h
@@ -29,6 +29,7 @@ namespace UE::AnimNext::Editor
 	class FVariablesOutlinerHierarchy;
 	class SVariablesOutlinerValue;
 	class FVariableProxyCustomization;
+	struct FPublicVariablesImpl;
 }
 
 namespace UE::AnimNext::Tests
@@ -63,6 +64,7 @@ class UAnimNextVariableEntry : public UAnimNextRigVMAssetEntry, public IAnimNext
 	friend class UE::AnimNext::Editor::FVariablesOutlinerHierarchy;
 	friend struct UE::AnimNext::UncookedOnly::FUtils;
 	friend class UE::AnimNext::Editor::FVariableProxyCustomization;
+	friend struct UE::AnimNext::Editor::FPublicVariablesImpl;
 
 	UAnimNextVariableEntry();
 
diff --git a/Engine/Plugins/Experimental/AnimNext/Source/AnimNextUncookedOnly/Private/AnimNextEdGraphNode.cpp b/Engine/Plugins/Experimental/AnimNext/Source/AnimNextUncookedOnly/Private/AnimNextEdGraphNode.cpp
--- a/Engine/Plugins/Experimental/AnimNext/Source/AnimNextUncookedOnly/Private/AnimNextEdGraphNode.cpp
+++ b/Engine/Plugins/Experimental/AnimNext/Source/AnimNextUncookedOnly/Private/AnimNextEdGraphNode.cpp
@@ -205,13 +205,13 @@ void UAnimNextEdGraphNode::BuildExposeVariablesContextMenu(UToolMenu* SubMenu)
 		FRigVMTrait_AnimNextPublicVariables DefaultTrait;
 		FRigVMTrait_AnimNextPublicVariables NewTrait;
 		UAnimNextDataInterface* Asset = CastChecked<UAnimNextDataInterface>(InAssetData.GetAsset());
-		NewTrait.Asset = Asset;
+		NewTrait.InternalAsset = Asset;
 		const FInstancedPropertyBag& PublicVariableDefaults = Asset->GetPublicVariableDefaults();
 		TConstArrayView<FPropertyBagPropertyDesc> Descs = PublicVariableDefaults.GetPropertyBagStruct()->GetPropertyDescs();
-		NewTrait.VariableNames.Reserve(Descs.Num());
+		NewTrait.InternalVariableNames.Reserve(Descs.Num());
 		for(const FPropertyBagPropertyDesc& Desc : Descs)
 		{
-			NewTrait.VariableNames.Add(Desc.Name);
+			NewTrait.InternalVariableNames.Add(Desc.Name);
 		}
 		FRigVMTrait_AnimNextPublicVariables::StaticStruct()->ExportText(DefaultValue, &NewTrait, &DefaultTrait, nullptr, PPF_SerializedAsImportText, nullptr);
 
diff --git a/Engine/Plugins/Experimental/AnimNext/Source/AnimNextUncookedOnly/Private/AnimNextRigVMAssetEditorData.cpp b/Engine/Plugins/Experimental/AnimNext/Source/AnimNextUncookedOnly/Private/AnimNextRigVMAssetEditorData.cpp
--- a/Engine/Plugins/Experimental/AnimNext/Source/AnimNextUncookedOnly/Private/AnimNextRigVMAssetEditorData.cpp
+++ b/Engine/Plugins/Experimental/AnimNext/Source/AnimNextUncookedOnly/Private/AnimNextRigVMAssetEditorData.cpp
@@ -1799,6 +1799,28 @@ bool UAnimNextRigVMAssetEditorData::HasPublicVariables() const
 	return false;
 }
 
+void UAnimNextRigVMAssetEditorData::GetPublicVariables(TArray<UAnimNextVariableEntry*>& OutPublicVariables) const
+{
+	for(UAnimNextRigVMAssetEntry* Entry : Entries)
+	{
+		if(UAnimNextVariableEntry* VariableEntry = Cast<UAnimNextVariableEntry>(Entry))
+		{
+			if(VariableEntry->GetExportAccessSpecifier() == EAnimNextExportAccessSpecifier::Public)
+			{
+				OutPublicVariables.Add(VariableEntry);
+			}
+		}
+		else if(UAnimNextDataInterfaceEntry* DataInterfaceEntry = Cast<UAnimNextDataInterfaceEntry>(Entry))
+		{
+			if(DataInterfaceEntry->DataInterface)
+			{
+				UAnimNextDataInterface_EditorData* EditorData = UE::AnimNext::UncookedOnly::FUtils::GetEditorData<UAnimNextDataInterface_EditorData>(DataInterfaceEntry->DataInterface.Get());
+				EditorData->GetPublicVariables(OutPublicVariables);
+			}
+		}
+	}
+}
+
 void UAnimNextRigVMAssetEditorData::RefreshExternalModels()
 {
 	GraphModels.Reset();
