Fix missing variable and comment actions in AnimNext graph context menu

diff --git a/Engine/Plugins/Experimental/AnimNext/Source/AnimNextEditor/Private/Common/SActionMenu.cpp b/Engine/Plugins/Experimental/AnimNext/Source/AnimNextEditor/Private/Common/SActionMenu.cpp
--- a/Engine/Plugins/Experimental/AnimNext/Source/AnimNextEditor/Private/Common/SActionMenu.cpp
+++ b/Engine/Plugins/Experimental/AnimNext/Source/AnimNextEditor/Private/Common/SActionMenu.cpp
@@ -148,6 +148,15 @@ void SActionMenu::CollectAllAnimNextGraphActions(FGraphContextMenuBuilder& MenuB
 		}
 	}
 
+	TArray<FRigVMExternalVariable> ExternalVariables = RigVMHost->GetVM()->GetExternalVariables(RigVMHost->GetRigVMExtendedExecuteContext());
+	for (const FRigVMExternalVariable& ExternalVariable : ExternalVariables)
+	{
+		MenuBuilder.AddAction(MakeShared<FAnimNextSchemaAction_Variable>(ExternalVariable, true));
+		MenuBuilder.AddAction(MakeShared<FAnimNextSchemaAction_Variable>(ExternalVariable, false));
+	}
+
+	MenuBuilder.AddAction(MakeShared<FAnimNextSchemaAction_AddComment>());
+
 	RigVMController->EnableReporting(true);
 }
 
