AnimNext: Removed the old RigVM asset view

The older view was not maintained and was confusing to users. It was only present to allow upgrading of older assets.
Also removed the ability to add animation graphs to modules

diff --git a/Engine/Plugins/Experimental/AnimNext/Source/AnimNextEditor/Private/AnimNextEditorModule.cpp b/Engine/Plugins/Experimental/AnimNext/Source/AnimNextEditor/Private/AnimNextEditorModule.cpp
--- a/Engine/Plugins/Experimental/AnimNext/Source/AnimNextEditor/Private/AnimNextEditorModule.cpp
+++ b/Engine/Plugins/Experimental/AnimNext/Source/AnimNextEditor/Private/AnimNextEditorModule.cpp
@@ -9,7 +9,6 @@
 #include "ScopedTransaction.h"
 #include "SSimpleButton.h"
 #include "UncookedOnlyUtils.h"
-#include "Common/SRigVMAssetView.h"
 #include "Editor/RigVMEditorTools.h"
 #include "Framework/Application/SlateApplication.h"
 #include "Module/AnimNextModule.h"
@@ -48,6 +47,7 @@
 #include "Graph/AnimNextAnimationGraph_EditorData.h"
 #include "Common/AnimNextAssetItemDetails.h"
 #include "DataInterface/AnimNextDataInterface.h"
+#include "Entries/AnimNextEventGraphEntry.h"
 #include "Graph/RigVMTrait_AnimNextPublicVariablesEditor.h"
 #include "Module/RigUnit_AnimNextModuleEvents.h"
 #include "Variables/SVariablesView.h"
@@ -148,67 +148,6 @@ void FAnimNextEditorModule::StartupModule()
 			}
 		});
 
-	SRigVMAssetView::RegisterCategoryFactory("Variables", [](UAnimNextRigVMAssetEditorData* InEditorData)
-	{
-		UAnimNextRigVMAssetEditorData* EditorData = CastChecked<UAnimNextRigVMAssetEditorData>(InEditorData);
-		UAnimNextRigVMAsset* Asset = UncookedOnly::FUtils::GetAsset(InEditorData);
-		return SNew(SSimpleButton)
-			.Text(LOCTEXT("AddVariableButton", "Add Variable"))
-			.Icon(FAppStyle::Get().GetBrush("Icons.Plus"))
-			.OnClicked_Lambda([EditorData, Asset]()
-			{
-				TSharedRef<SAddVariablesDialog> AddVariablesDialog =
-					SNew(SAddVariablesDialog, TArray<UAnimNextRigVMAssetEditorData*>({ EditorData }));
-
-				TArray<FVariableToAdd> VariablesToAdd;
-				TArray<FDataInterfaceToAdd> DataInterfacesToAdd;
-				if(AddVariablesDialog->ShowModal(VariablesToAdd, DataInterfacesToAdd))
-				{
-					FScopedTransaction Transaction(LOCTEXT("AddVariables", "Add variable(s)"));
-					for (const FVariableToAdd& VariableToAdd : VariablesToAdd)
-					{
-						check(EditorData->FindEntry(VariableToAdd.Name) == nullptr);
-						EditorData->AddVariable(VariableToAdd.Name, VariableToAdd.Type);
-					}
-				}
-				return FReply::Handled();
-			});
-	});
-
-	SRigVMAssetView::RegisterCategoryFactory("Event Graphs", [](UAnimNextRigVMAssetEditorData* InEditorData)
-	{
-		UAnimNextRigVMAssetEditorData* EditorData = CastChecked<UAnimNextRigVMAssetEditorData>(InEditorData);
-		return SNew(SSimpleButton)
-			.Text(LOCTEXT("AddEventGraphButton", "Add Event Graph"))
-			.Icon(FAppStyle::Get().GetBrush("Icons.Plus"))
-			.OnClicked_Lambda([EditorData]()
-			{
-				FScopedTransaction Transaction(LOCTEXT("AddEventGraph", "Add Event Graph"));
-
-				// Create a new entry for the graph
-				EditorData->AddEventGraph(TEXT("NewGraph"), FRigUnit_AnimNextPrePhysicsEvent::StaticStruct());
-
-				return FReply::Handled();
-			});
-	});
-
-	SRigVMAssetView::RegisterCategoryFactory("Animation Graphs", [](UAnimNextRigVMAssetEditorData* InEditorData)
-	{
-		UAnimNextRigVMAssetEditorData* EditorData = CastChecked<UAnimNextRigVMAssetEditorData>(InEditorData);
-		return SNew(SSimpleButton)
-			.Text(LOCTEXT("AddGraphButton", "Add Animation Graph"))
-			.Icon(FAppStyle::Get().GetBrush("Icons.Plus"))
-			.OnClicked_Lambda([EditorData]()
-			{
-				FScopedTransaction Transaction(LOCTEXT("AddAnimationGraph", "Add Animation Graph"));
-
-				// Create a new entry for the graph
-				EditorData->AddAnimationGraph(TEXT("NewGraph"));
-
-				return FReply::Handled();
-			});
-	});
-
 	UE::UniversalObjectLocator::IUniversalObjectLocatorEditorModule& UolEditorModule = FModuleManager::LoadModuleChecked<UE::UniversalObjectLocator::IUniversalObjectLocatorEditorModule>("UniversalObjectLocatorEditor");
 	UolEditorModule.RegisterLocatorEditor("AnimNextObjectFunction", MakeShared<FObjectFunctionLocatorEditor>());
 	UolEditorModule.RegisterLocatorEditor("AnimNextObjectProperty", MakeShared<FObjectPropertyLocatorEditor>());
@@ -267,10 +206,6 @@ void FAnimNextEditorModule::ShutdownModule()
 
 	UnregisterWorkspaceDocumentTypes();
 
-	SRigVMAssetView::UnregisterCategoryFactory("Parameters");
-	SRigVMAssetView::UnregisterCategoryFactory("Event Graphs");
-	SRigVMAssetView::UnregisterCategoryFactory("Animation Graphs");
-
 	if(FModuleManager::Get().IsModuleLoaded("UniversalObjectLocatorEditor"))
 	{
 		UE::UniversalObjectLocator::IUniversalObjectLocatorEditorModule& UolEditorModule = FModuleManager::GetModuleChecked<UE::UniversalObjectLocator::IUniversalObjectLocatorEditorModule>("UniversalObjectLocatorEditor");
@@ -334,113 +269,23 @@ void FAnimNextEditorModule::RemoveWorkspaceSupportedAssetClass(const FTopLevelAs
 
 void FAnimNextEditorModule::RegisterWorkspaceDocumentTypes(Workspace::IWorkspaceEditorModule& WorkspaceEditorModule)
 {
-	// --- AnimNextRigVMAsset ---
-	Workspace::FObjectDocumentArgs AnimNextAssetDocumentArgs(
-			Workspace::FOnMakeDocumentWidget::CreateLambda([](const Workspace::FWorkspaceEditorContext& InContext)
+	// --- AnimNextModule ---
+	Workspace::FObjectDocumentArgs AnimNextModuleDocumentArgs(
+			Workspace::FOnRedirectWorkspaceContext::CreateLambda([](UObject* InObject) -> UObject*
 			{
-				UAnimNextRigVMAsset* Asset = CastChecked<UAnimNextRigVMAsset>(InContext.Object);
-				UAnimNextRigVMAssetEditorData* EditorData = UncookedOnly::FUtils::GetEditorData<UAnimNextRigVMAssetEditorData>(Asset);
-
-				TWeakPtr<Workspace::IWorkspaceEditor> WeakWorkspaceEditor = InContext.WorkspaceEditor;
-
-				// TODO: For now (so we can convert older assets over to new formats manually) we open the asset view, but in the future we
-				// should just open any nested graphs (and set the variables view up)
-				TSharedPtr<SRigVMAssetView> SharedAssetView = SNew(SRigVMAssetView, EditorData)
-					.OnOpenGraph_Lambda([WeakWorkspaceEditor](URigVMGraph* InGraph)
-					{
-						if(TSharedPtr<Workspace::IWorkspaceEditor> WorkspaceEditor = WeakWorkspaceEditor.Pin())
-						{
-							if(IRigVMClientHost* RigVMClientHost = InGraph->GetImplementingOuter<IRigVMClientHost>())
-							{
-								if(UObject* EditorObject = RigVMClientHost->GetEditorObjectForRigVMGraph(InGraph))
-								{
-									WorkspaceEditor->OpenObjects({EditorObject});
-								}
-							}
-						}
-					})
-					.OnDeleteEntries_Lambda([WeakWorkspaceEditor](const TArray<UAnimNextRigVMAssetEntry*>& InEntries)
-					{
-						if(TSharedPtr<Workspace::IWorkspaceEditor> WorkspaceEditor = WeakWorkspaceEditor.Pin())
-						{
-							if(InEntries.Num() > 0)
-							{
-								TArray<UObject*> EdGraphsToClose;
-								EdGraphsToClose.Reserve(InEntries.Num());
-								for(UAnimNextRigVMAssetEntry* Entry : InEntries)
-								{
-									if(IAnimNextRigVMGraphInterface* GraphInterface = Cast<IAnimNextRigVMGraphInterface>(Entry))
-									{
-										if(URigVMEdGraph* EdGraph = GraphInterface->GetEdGraph())
-										{
-											EdGraphsToClose.Add(EdGraph);
-										}
-									}
-								}
-
-								WorkspaceEditor->CloseObjects(EdGraphsToClose);
-							}
-						}
-					});
+				UAnimNextModule* Module = CastChecked<UAnimNextModule>(InObject);
+				UAnimNextModule_EditorData* EditorData = UncookedOnly::FUtils::GetEditorData<UAnimNextModule_EditorData>(Module);
 
-				TWeakPtr<SRigVMAssetView> WeakAssetView = SharedAssetView;
-				SharedAssetView->SetOnSelectionChanged(SRigVMAssetView::FOnSelectionChanged::CreateLambda([WeakWorkspaceEditor, WeakAssetView](const TArray<UObject*>& InEntries)
+				// Redirect to the inner graph, if any
+				UAnimNextEventGraphEntry* EventGraphEntry = EditorData->FindFirstEntryOfType<UAnimNextEventGraphEntry>();
+				if(EventGraphEntry)
 				{
-					if(TSharedPtr<SRigVMAssetView> SharedAssetView = WeakAssetView.Pin())
-					{
-						if(TSharedPtr<Workspace::IWorkspaceEditor> WorkspaceEditor = WeakWorkspaceEditor.Pin())
-						{	
-							WorkspaceEditor->SetGlobalSelection(SharedAssetView, UE::Workspace::FOnClearGlobalSelection::CreateLambda([WeakAssetView](){ if (const TSharedPtr<SRigVMAssetView> SharedAssetView = WeakAssetView.Pin()) { SharedAssetView->ClearSelection(); } }));
-							WorkspaceEditor->SetDetailsObjects(InEntries);
-						}
-				}}));
-
-
-				return SharedAssetView.ToSharedRef();
-			}),
-			Workspace::WorkspaceTabs::TopMiddleDocumentArea);
-	AnimNextAssetDocumentArgs.OnGetTabName = Workspace::FOnGetTabName::CreateLambda([](const Workspace::FWorkspaceEditorContext& InContext)
-	{
-		const UAnimNextRigVMAsset* Asset = CastChecked<UAnimNextRigVMAsset>(InContext.Object);
-		return FText::FromName(Asset->GetFName());
-	});
-
-	AnimNextAssetDocumentArgs.OnGetDocumentBreadcrumbTrail = Workspace::FOnGetDocumentBreadcrumbTrail::CreateLambda([](const Workspace::FWorkspaceEditorContext& InContext, TArray<TSharedPtr<Workspace::FWorkspaceBreadcrumb>>& OutBreadcrumbs)
-	{
-		if (const UAnimNextRigVMAsset* Asset = Cast<UAnimNextRigVMAsset>(InContext.Object))
-		{
-			TWeakObjectPtr<const UAnimNextRigVMAsset> WeakAsset = Asset;
-			const TSharedPtr<Workspace::FWorkspaceBreadcrumb>& GraphCrumb = OutBreadcrumbs.Add_GetRef(MakeShared<Workspace::FWorkspaceBreadcrumb>());
-			GraphCrumb->OnGetLabel = Workspace::FWorkspaceBreadcrumb::FOnGetBreadcrumbLabel::CreateLambda([AssetName = Asset->GetFName()]{ return FText::FromName(AssetName); });
-			GraphCrumb->CanSave = Workspace::FWorkspaceBreadcrumb::FCanSaveBreadcrumb::CreateLambda(
-				[WeakAsset]
-				{
-					if (const UAnimNextRigVMAsset* Asset = WeakAsset.Get())
-					{
-						return Asset->GetPackage()->IsDirty();						
-					}
-
-					return false;
+					return EventGraphEntry->GetEdGraph();
 				}
-			);
-			GraphCrumb->OnSave = Workspace::FWorkspaceBreadcrumb::FOnSaveBreadcrumb::CreateLambda(
-				[WeakAsset]
-				{
-					if (const UAnimNextRigVMAsset* Asset = WeakAsset.Get())
-					{
-						FEditorFileUtils::PromptForCheckoutAndSave({Asset->GetPackage()}, false, /*bPromptToSave=*/ false);
-					}
-				}
-			);
-		}
-	});
-
-	AnimNextAssetDocumentArgs.OnGetTabIcon = Workspace::FOnGetTabIcon::CreateLambda([](const Workspace::FWorkspaceEditorContext& InContext)
-	{
-		return FAppStyle::GetBrush(TEXT("ClassIcon.Default"));
-	});
+				return nullptr;
+			}));
 
-	WorkspaceEditorModule.RegisterObjectDocumentType(FTopLevelAssetPath(TEXT("/Script/AnimNext.AnimNextModule")), AnimNextAssetDocumentArgs);
+	WorkspaceEditorModule.RegisterObjectDocumentType(FTopLevelAssetPath(TEXT("/Script/AnimNext.AnimNextModule")), AnimNextModuleDocumentArgs);
 
 	// --- AnimNextAnimationGraph ---
 	Workspace::FObjectDocumentArgs AnimNextAnimationGraphDocumentArgs(
diff --git a/Engine/Plugins/Experimental/AnimNext/Source/AnimNextEditor/Private/Common/SRigVMAssetView.cpp b/Engine/Plugins/Experimental/AnimNext/Source/AnimNextEditor/Private/Common/SRigVMAssetView.cpp
deleted file mode 100644
diff --git a/Engine/Plugins/Experimental/AnimNext/Source/AnimNextEditor/Private/Common/SRigVMAssetView.h b/Engine/Plugins/Experimental/AnimNext/Source/AnimNextEditor/Private/Common/SRigVMAssetView.h
deleted file mode 100644
diff --git a/Engine/Plugins/Experimental/AnimNext/Source/AnimNextUncookedOnly/Internal/AnimNextRigVMAssetEditorData.h b/Engine/Plugins/Experimental/AnimNext/Source/AnimNextUncookedOnly/Internal/AnimNextRigVMAssetEditorData.h
--- a/Engine/Plugins/Experimental/AnimNext/Source/AnimNextUncookedOnly/Internal/AnimNextRigVMAssetEditorData.h
+++ b/Engine/Plugins/Experimental/AnimNext/Source/AnimNextUncookedOnly/Internal/AnimNextRigVMAssetEditorData.h
@@ -310,6 +310,19 @@ protected:
 		}
 	}
 
+	// Find the first entry of the specified type
+	template<typename EntryType>
+	EntryType* FindFirstEntryOfType() const
+	{
+		EntryType* FirstEntry = nullptr;
+		ForEachEntryOfType<EntryType>([&FirstEntry](EntryType* InEntry)
+		{
+			FirstEntry = InEntry;
+			return false;
+		});
+		return FirstEntry;
+	}
+
 	// Returns all nodes in all graphs of the specified class
 	template<class T>
 	void GetAllNodesOfClass(TArray<T*>& OutNodes) const
diff --git a/Engine/Plugins/Experimental/AnimNext/Source/AnimNextUncookedOnly/Internal/Entries/AnimNextEventGraphEntry.h b/Engine/Plugins/Experimental/AnimNext/Source/AnimNextUncookedOnly/Internal/Entries/AnimNextEventGraphEntry.h
--- a/Engine/Plugins/Experimental/AnimNext/Source/AnimNextUncookedOnly/Internal/Entries/AnimNextEventGraphEntry.h
+++ b/Engine/Plugins/Experimental/AnimNext/Source/AnimNextUncookedOnly/Internal/Entries/AnimNextEventGraphEntry.h
@@ -20,6 +20,11 @@ namespace UE::AnimNext::Tests
 	class FVariables;
 }
 
+namespace UE::AnimNext::Editor
+{
+	class FAnimNextEditorModule;
+}
+
 UCLASS(MinimalAPI, Category = "Event Graphs", DisplayName = "Event Graph")
 class UAnimNextEventGraphEntry : public UAnimNextRigVMAssetEntry, public IAnimNextRigVMGraphInterface
 {
@@ -30,6 +35,7 @@ class UAnimNextEventGraphEntry : public UAnimNextRigVMAssetEntry, public IAnimNe
 	friend class UE::AnimNext::Tests::FEditor_Variables;
 	friend class UE::AnimNext::Tests::FEditor_Graphs;
 	friend class UE::AnimNext::Tests::FVariables;
+	friend class UE::AnimNext::Editor::FAnimNextEditorModule;
 
 	// UAnimNextRigVMAssetEntry interface
 	virtual FName GetEntryName() const override { return GraphName; }
diff --git a/Engine/Plugins/Experimental/AnimNext/Source/AnimNextUncookedOnly/Private/Graph/AnimNextAnimationGraph_EditorData.cpp b/Engine/Plugins/Experimental/AnimNext/Source/AnimNextUncookedOnly/Private/Graph/AnimNextAnimationGraph_EditorData.cpp
--- a/Engine/Plugins/Experimental/AnimNext/Source/AnimNextUncookedOnly/Private/Graph/AnimNextAnimationGraph_EditorData.cpp
+++ b/Engine/Plugins/Experimental/AnimNext/Source/AnimNextUncookedOnly/Private/Graph/AnimNextAnimationGraph_EditorData.cpp
@@ -684,7 +684,6 @@ TConstArrayView<TSubclassOf<UAnimNextRigVMAssetEntry>> UAnimNextAnimationGraph_E
 	static const TSubclassOf<UAnimNextRigVMAssetEntry> Classes[] =
 	{
 		UAnimNextAnimationGraphEntry::StaticClass(),
-		UAnimNextEventGraphEntry::StaticClass(),	// TODO: remove when assets are reworked post-refactor
 		UAnimNextVariableEntry::StaticClass(),
 		UAnimNextDataInterfaceEntry::StaticClass(),
 	};
diff --git a/Engine/Plugins/Experimental/AnimNext/Source/AnimNextUncookedOnly/Private/Module/AnimNextModule_EditorData.cpp b/Engine/Plugins/Experimental/AnimNext/Source/AnimNextUncookedOnly/Private/Module/AnimNextModule_EditorData.cpp
--- a/Engine/Plugins/Experimental/AnimNext/Source/AnimNextUncookedOnly/Private/Module/AnimNextModule_EditorData.cpp
+++ b/Engine/Plugins/Experimental/AnimNext/Source/AnimNextUncookedOnly/Private/Module/AnimNextModule_EditorData.cpp
@@ -168,7 +168,6 @@ TConstArrayView<TSubclassOf<UAnimNextRigVMAssetEntry>> UAnimNextModule_EditorDat
 {
 	static const TSubclassOf<UAnimNextRigVMAssetEntry> Classes[] =
 	{
-		UAnimNextAnimationGraphEntry::StaticClass(),	// TEMP: Remove this once we have ported all old data
 		UAnimNextEventGraphEntry::StaticClass(),
 		UAnimNextVariableEntry::StaticClass(),
 		UAnimNextDataInterfaceEntry::StaticClass(),
diff --git a/Engine/Plugins/Experimental/Workspace/Source/WorkspaceEditor/Private/WorkspaceEditorModule.cpp b/Engine/Plugins/Experimental/Workspace/Source/WorkspaceEditor/Private/WorkspaceEditorModule.cpp
--- a/Engine/Plugins/Experimental/Workspace/Source/WorkspaceEditor/Private/WorkspaceEditorModule.cpp
+++ b/Engine/Plugins/Experimental/Workspace/Source/WorkspaceEditor/Private/WorkspaceEditorModule.cpp
@@ -268,7 +268,7 @@ const FObjectDocumentArgs* FWorkspaceEditorModule::FindObjectDocumentType(const
 	while(FoundArgs && FoundArgs->OnRedirectWorkspaceContext.IsBound())
 	{
 		InObject = FoundArgs->OnRedirectWorkspaceContext.Execute(const_cast<UObject*>(InObject));
-		FoundArgs = ObjectDocumentArgs.Find(InObject->GetClass()->GetClassPathName());
+		FoundArgs = InObject ? ObjectDocumentArgs.Find(InObject->GetClass()->GetClassPathName()) : nullptr;
 	}
 
 	return FoundArgs;
