AnimNext: Fix initialize event running more than once

Up to now this was predicated on the state transition between initialization and running, but this conflates it with the enabled/disabled state which is under seperate user control.
Switching to a simple task that runs after the true execution of the event runs, rather than coordinating via the EndTickFunction.

#rnx

diff --git a/Engine/Plugins/Experimental/AnimNext/Source/AnimNext/Internal/Component/AnimNextWorldSubsystem.h b/Engine/Plugins/Experimental/AnimNext/Source/AnimNext/Internal/Component/AnimNextWorldSubsystem.h
--- a/Engine/Plugins/Experimental/AnimNext/Source/AnimNext/Internal/Component/AnimNextWorldSubsystem.h
+++ b/Engine/Plugins/Experimental/AnimNext/Source/AnimNext/Internal/Component/AnimNextWorldSubsystem.h
@@ -12,6 +12,12 @@
 #include "AnimNextWorldSubsystem.generated.h"
 
 class UAnimNextComponent;
+struct FRigUnit_AnimNextInitializeEvent;
+
+namespace UE::AnimNext
+{
+	struct FModuleEndTickFunction;
+}
 
 namespace UE::AnimNext
 {
@@ -68,6 +74,9 @@ public:
 	UAnimNextWorldSubsystem();
 	virtual ~UAnimNextWorldSubsystem() override;
 
+	// Get the subsystem for the specified object's world
+	static UAnimNextWorldSubsystem* Get(UObject* InObject);
+	
 	static void AddReferencedObjects(UObject* InThis, FReferenceCollector& Collector);
 
 private:
@@ -107,7 +116,13 @@ protected:
 	// @param	InLocation			Where to run the task, before or after
 	void QueueTaskHandle(UE::AnimNext::FModuleHandle InHandle, FName InModuleEventName, TUniqueFunction<void(const UE::AnimNext::FModuleTaskContext&)>&& InTaskFunction, UE::AnimNext::ETaskRunLocation InLocation = UE::AnimNext::ETaskRunLocation::Before);
 
+	// Disables the initialize event for a module one it has run once
+	void DisableInitializeEvent(UE::AnimNext::FModuleHandle InHandle);
+	
 protected:
+	friend struct FRigUnit_AnimNextInitializeEvent;
+	friend struct UE::AnimNext::FModuleEndTickFunction;
+
 	// Currently running instances, pooled
 	UE::AnimNext::TPool<FAnimNextModuleInstance> Instances;
 
diff --git a/Engine/Plugins/Experimental/AnimNext/Source/AnimNext/Internal/Module/AnimNextModuleInstance.h b/Engine/Plugins/Experimental/AnimNext/Source/AnimNext/Internal/Module/AnimNextModuleInstance.h
--- a/Engine/Plugins/Experimental/AnimNext/Source/AnimNext/Internal/Module/AnimNextModuleInstance.h
+++ b/Engine/Plugins/Experimental/AnimNext/Source/AnimNext/Internal/Module/AnimNextModuleInstance.h
@@ -124,6 +124,12 @@ struct FAnimNextModuleInstance : public FAnimNextDataInterfaceInstance
 	ANIMNEXT_API void ShowDebugDrawing(bool bInShowDebugDrawing);
 #endif
 
+	// Called to disable the initialize event for this module on the game thread when it has executed for the first time
+	void DisableInitializeEvent();
+
+	// Run a simple task on the GT via FFunctionGraphTask::CreateAndDispatchWhenReady
+	static void RunTaskOnGameThread(TUniqueFunction<void(void)>&& InFunction);
+
 	// Object this entry is bound to
 	UPROPERTY(Transient)
 	TObjectPtr<UObject> Object = nullptr;
@@ -159,6 +165,9 @@ struct FAnimNextModuleInstance : public FAnimNextDataInterfaceInstance
 	// Components of this module
 	UE::AnimNext::ModuleInstanceComponentMapType Components; 
 
+	// Initialize event index (INDEX_NONE if no initialize event is found)
+	int32 InitializeEventIndex = INDEX_NONE;
+
 	enum class ERunState : uint8
 	{
 		None,
diff --git a/Engine/Plugins/Experimental/AnimNext/Source/AnimNext/Private/Component/AnimNextWorldSubsystem.cpp b/Engine/Plugins/Experimental/AnimNext/Source/AnimNext/Private/Component/AnimNextWorldSubsystem.cpp
--- a/Engine/Plugins/Experimental/AnimNext/Source/AnimNext/Private/Component/AnimNextWorldSubsystem.cpp
+++ b/Engine/Plugins/Experimental/AnimNext/Source/AnimNext/Private/Component/AnimNextWorldSubsystem.cpp
@@ -47,6 +47,22 @@ UAnimNextWorldSubsystem::~UAnimNextWorldSubsystem()
 	}
 }
 
+UAnimNextWorldSubsystem* UAnimNextWorldSubsystem::Get(UObject* InObject)
+{
+	if(InObject == nullptr)
+	{
+		return nullptr;
+	}
+
+	UWorld* World = InObject->GetWorld();
+	if(World == nullptr)
+	{
+		return nullptr;
+	}
+
+	return World->GetSubsystem<UAnimNextWorldSubsystem>();
+}
+
 void UAnimNextWorldSubsystem::AddReferencedObjects(UObject* InThis, FReferenceCollector& Collector)
 {
 	Super::AddReferencedObjects(InThis, Collector);
@@ -240,4 +256,15 @@ void UAnimNextWorldSubsystem::OnModuleCompiled(UAnimNextModule* InModule)
 	}
 }
 
-#endif
\ No newline at end of file
+#endif
+
+void UAnimNextWorldSubsystem::DisableInitializeEvent(UE::AnimNext::FModuleHandle InHandle)
+{
+	using namespace UE::AnimNext;
+	check(IsInGameThread());
+	if (IsValidHandle(InHandle))
+	{
+		FAnimNextModuleInstance& Instance = Instances.Get(InHandle);
+		Instance.DisableInitializeEvent();
+	}
+}
\ No newline at end of file
diff --git a/Engine/Plugins/Experimental/AnimNext/Source/AnimNext/Private/Module/AnimNextModuleInstance.cpp b/Engine/Plugins/Experimental/AnimNext/Source/AnimNext/Private/Module/AnimNextModuleInstance.cpp
--- a/Engine/Plugins/Experimental/AnimNext/Source/AnimNext/Private/Module/AnimNextModuleInstance.cpp
+++ b/Engine/Plugins/Experimental/AnimNext/Source/AnimNext/Private/Module/AnimNextModuleInstance.cpp
@@ -410,4 +410,27 @@ void FAnimNextModuleInstance::ShowDebugDrawing(bool bInShowDebugDrawing)
 	check(DebugDraw.IsValid());
 	DebugDraw->SetEnabled(bInShowDebugDrawing);
 }
-#endif
\ No newline at end of file
+#endif
+
+void FAnimNextModuleInstance::DisableInitializeEvent()
+{
+	if(InitializeEventIndex != INDEX_NONE)
+	{
+		check(IsInGameThread());
+		check(TickFunctions.IsValidIndex(InitializeEventIndex));
+		check(TickFunctions[InitializeEventIndex].EventName == FRigUnit_AnimNextInitializeEvent::EventName);
+		TickFunctions[InitializeEventIndex].SetTickFunctionEnable(false);
+	}
+}
+
+void FAnimNextModuleInstance::RunTaskOnGameThread(TUniqueFunction<void(void)>&& InFunction)
+{
+	if(IsInGameThread())
+	{
+		InFunction();
+	}
+	else
+	{
+		FFunctionGraphTask::CreateAndDispatchWhenReady(MoveTemp(InFunction), TStatId(), nullptr, ENamedThreads::GameThread);
+	}
+}
\ No newline at end of file
diff --git a/Engine/Plugins/Experimental/AnimNext/Source/AnimNext/Private/Module/ModuleTickFunction.cpp b/Engine/Plugins/Experimental/AnimNext/Source/AnimNext/Private/Module/ModuleTickFunction.cpp
--- a/Engine/Plugins/Experimental/AnimNext/Source/AnimNext/Private/Module/ModuleTickFunction.cpp
+++ b/Engine/Plugins/Experimental/AnimNext/Source/AnimNext/Private/Module/ModuleTickFunction.cpp
@@ -6,6 +6,7 @@
 #include "AnimNextStats.h"
 #include "RigVMRuntimeDataRegistry.h"
 #include "Algo/TopologicalSort.h"
+#include "Component/AnimNextWorldSubsystem.h"
 #include "Module/AnimNextModule.h"
 #include "Module/ModuleTaskContext.h"
 #include "Module/ProxyVariablesContext.h"
@@ -60,27 +61,6 @@ void FModuleEndTickFunction::Run()
 		// Reset our list of output events, we don't retain any
 		ModuleInstance->OutputEventList.Reset();
 	}
-	
-	auto RunTaskOnGameThread = [](TUniqueFunction<void(void)>&& InFunction)
-	{
-		if(IsInGameThread())
-		{
-			InFunction();
-		}
-		else
-		{
-			FFunctionGraphTask::CreateAndDispatchWhenReady(MoveTemp(InFunction), TStatId(), nullptr, ENamedThreads::GameThread);
-		}
-	};
-
-	auto DisableInitializeEvent = [this]()
-	{
-		if(ModuleInstance->TickFunctions.Num() > 0 && ModuleInstance->TickFunctions[0].EventName == FRigUnit_AnimNextInitializeEvent::EventName)
-		{
-			// Disable the initialize event as we have already run it
-			ModuleInstance->TickFunctions[0].SetTickFunctionEnable(false);
-		}
-	};
 
 	if(ModuleInstance->RunState == FAnimNextModuleInstance::ERunState::PendingInitialUpdate)
 	{
@@ -88,30 +68,21 @@ void FModuleEndTickFunction::Run()
 #if WITH_EDITOR
 			|| (ModuleInstance->InitMethod == EAnimNextModuleInitMethod::InitializeAndPauseInEditor && ModuleInstance->bIsEditor)
 #endif
-			)
+		)
 		{
-			// Queue task to disable our tick functions now we have performed our initial update
-			RunTaskOnGameThread([this, &DisableInitializeEvent]()
+			// Queue task to disable ourselves
+			FAnimNextModuleInstance::RunTaskOnGameThread([ModuleHandle = ModuleInstance->Handle, WeakObject = TWeakObjectPtr<UObject>(ModuleInstance->Object)]()
 			{
 				check(IsInGameThread());
-				DisableInitializeEvent();
-				ModuleInstance->Enable(false);
+				UAnimNextWorldSubsystem* WorldSubsystem = UAnimNextWorldSubsystem::Get(WeakObject.Get());
+				WorldSubsystem->EnableHandle(ModuleHandle, false);
 			});
 		}
 	}
-	else if(ModuleInstance->RunState != FAnimNextModuleInstance::ERunState::Running)
-	{
-		RunTaskOnGameThread([this, &DisableInitializeEvent]()
-		{
-			check(IsInGameThread());
-			DisableInitializeEvent();
-			ModuleInstance->TransitionToRunState(FAnimNextModuleInstance::ERunState::Running);
-		});
-	}
 
 	if (!MainThreadActionEventList.IsEmpty())
 	{
-		RunTaskOnGameThread([MainThreadActionEventList = MoveTemp(MainThreadActionEventList)]()
+		FAnimNextModuleInstance::RunTaskOnGameThread([MainThreadActionEventList = MoveTemp(MainThreadActionEventList)]()
 			{
 				check(IsInGameThread());
 				for (const FAnimNextTraitEventPtr& Event : MainThreadActionEventList)
diff --git a/Engine/Plugins/Experimental/AnimNext/Source/AnimNext/Private/Module/RigUnit_AnimNextModuleEvents.cpp b/Engine/Plugins/Experimental/AnimNext/Source/AnimNext/Private/Module/RigUnit_AnimNextModuleEvents.cpp
--- a/Engine/Plugins/Experimental/AnimNext/Source/AnimNext/Private/Module/RigUnit_AnimNextModuleEvents.cpp
+++ b/Engine/Plugins/Experimental/AnimNext/Source/AnimNext/Private/Module/RigUnit_AnimNextModuleEvents.cpp
@@ -2,6 +2,8 @@
 
 #include "Module/RigUnit_AnimNextModuleEvents.h"
 
+#include "Component/AnimNextWorldSubsystem.h"
+#include "Engine/World.h"
 #include "Module/AnimNextModuleInstance.h"
 
 FRigUnit_AnimNextExecuteBindings_Execute()
@@ -26,12 +28,22 @@ UE::AnimNext::FModuleEventBindingFunction FRigUnit_AnimNextExecuteBindings::GetB
 
 FRigUnit_AnimNextInitializeEvent_Execute()
 {
+	FAnimNextModuleInstance& ModuleInstance = ExecuteContext.GetContextData<FAnimNextModuleContextData>().GetModuleInstance();
+
+	// Queue a GT task to disable ourselves (we can only enable/disable tick functions on the GT)
+	FAnimNextModuleInstance::RunTaskOnGameThread([ModuleHandle = ModuleInstance.Handle, WeakObject = TWeakObjectPtr<UObject>(ModuleInstance.Object)]()
+	{
+		UAnimNextWorldSubsystem* WorldSubsystem = UAnimNextWorldSubsystem::Get(WeakObject.Get());
+		WorldSubsystem->DisableInitializeEvent(ModuleHandle);
+	});
 }
 
 UE::AnimNext::FModuleEventBindingFunction FRigUnit_AnimNextInitializeEvent::GetBindingFunction() const
 {
 	return [](const UE::AnimNext::FTickFunctionBindingContext& InContext, FTickFunction& InTickFunction)
 	{
+		InContext.ModuleInstance.InitializeEventIndex = InContext.EventIndex;
+
 		InTickFunction.TickGroup = ETickingGroup::TG_PrePhysics;
 
 		// Add prerequisite on this for all subsequent tick functions (assumes tick functions are sorted by phase)
