Add getter/setter dialog when dragging variables into AnimNext graph

diff --git a/Engine/Plugins/Experimental/AnimNext/Source/AnimNextEditor/Private/Common/GraphEditorSchemaActions.cpp b/Engine/Plugins/Experimental/AnimNext/Source/AnimNextEditor/Private/Common/GraphEditorSchemaActions.cpp
--- a/Engine/Plugins/Experimental/AnimNext/Source/AnimNextEditor/Private/Common/GraphEditorSchemaActions.cpp
+++ b/Engine/Plugins/Experimental/AnimNext/Source/AnimNextEditor/Private/Common/GraphEditorSchemaActions.cpp
@@ -12,6 +12,8 @@
 #include "Units/RigUnit.h"
 #include "EdGraphSchema_K2.h"
 #include "GraphEditor.h"
+#include "Framework/MultiBox/MultiBoxBuilder.h"
+#include "Framework/Application/SlateApplication.h"
 
 #define LOCTEXT_NAMESPACE "AnimNextSchemaActions"
 
@@ -137,9 +139,9 @@ UEdGraphNode* FAnimNextSchemaAction_DispatchFactory::PerformAction(UEdGraph* Par
 	return NewNode;
 }
 
-FAnimNextSchemaAction_Variable::FAnimNextSchemaAction_Variable(const FRigVMExternalVariable& InExternalVariable, bool bInIsGetter)
+FAnimNextSchemaAction_Variable::FAnimNextSchemaAction_Variable(const FRigVMExternalVariable& InExternalVariable, const EVariableAccessorChoice InVariableAccessorChoice)
 	: ExternalVariable(InExternalVariable)
-	, bIsGetter(bInIsGetter)
+	, VariableAccessorChoice(InVariableAccessorChoice)
 {
 	static const FText VariablesCategory = LOCTEXT("Variables", "Variables");
 	static const FTextFormat GetVariableFormat = LOCTEXT("GetVariableFormat", "Get {0}");
@@ -148,16 +150,20 @@ FAnimNextSchemaAction_Variable::FAnimNextSchemaAction_Variable(const FRigVMExter
 	FText MenuDesc;
 	FText ToolTip;
 
-	if(bInIsGetter)
+	if(InVariableAccessorChoice == EVariableAccessorChoice::Get)
 	{
 		MenuDesc = FText::Format(GetVariableFormat, FText::FromName(ExternalVariable.Name));
 		ToolTip = FText::FromString(FString::Printf(TEXT("Get the value of variable %s"), *ExternalVariable.Name.ToString()));
 	}
-	else
+	else if (InVariableAccessorChoice == EVariableAccessorChoice::Set)
 	{
 		MenuDesc = FText::Format(SetVariableFormat, FText::FromName(ExternalVariable.Name));
 		ToolTip = FText::FromString(FString::Printf(TEXT("Set the value of variable %s"), *ExternalVariable.Name.ToString()));
 	}
+	else
+	{
+		MenuDesc = FText::FromName(ExternalVariable.Name);
+	}
 
 	UpdateSearchData(MenuDesc, ToolTip, VariablesCategory, FText::GetEmpty());
 
@@ -167,64 +173,118 @@ FAnimNextSchemaAction_Variable::FAnimNextSchemaAction_Variable(const FRigVMExter
 
 UEdGraphNode* FAnimNextSchemaAction_Variable::PerformAction(UEdGraph* ParentGraph, TArray<UEdGraphPin*>& FromPins, const FVector2D Location, bool bSelectNewNode)
 {
-	IRigVMClientHost* Host = ParentGraph->GetImplementingOuter<IRigVMClientHost>();
-	if(Host == nullptr)
+	const auto AddNodeLambda = [](UEdGraph* ParentGraph, FRigVMExternalVariable ExtVariable, const FVector2D Location, const bool bIsGetter) -> URigVMEdGraphNode*
 	{
-		return nullptr;
-	}
+		IRigVMClientHost* Host = ParentGraph->GetImplementingOuter<IRigVMClientHost>();
+		if(Host == nullptr)
+		{
+			return nullptr;
+		}
 
-	URigVMEdGraph* EdGraph = Cast<URigVMEdGraph>(ParentGraph);
-	if(EdGraph == nullptr)
-	{
-		return nullptr;
-	}
+		URigVMEdGraph* EdGraph = Cast<URigVMEdGraph>(ParentGraph);
+		if(EdGraph == nullptr)
+		{
+			return nullptr;
+		}
 
-	URigVMEdGraphNode* NewNode = nullptr;
+		URigVMEdGraphNode* NewNode = nullptr;
 
-	FString ObjectPath;
-	if (ExternalVariable.TypeObject)
-	{
-		ObjectPath = ExternalVariable.TypeObject->GetPathName();
-	}
+		FString ObjectPath;
+		if (ExtVariable.TypeObject)
+		{
+			ObjectPath = ExtVariable.TypeObject->GetPathName();
+		}
 
-	FString TypeName = ExternalVariable.TypeName.ToString();
-	if (ExternalVariable.bIsArray)
-	{
-		TypeName = FString::Printf(TEXT("TArray<%s>"), *TypeName);
-	}
+		FString TypeName = ExtVariable.TypeName.ToString();
+		if (ExtVariable.bIsArray)
+		{
+			TypeName = FString::Printf(TEXT("TArray<%s>"), *TypeName);
+		}
 
-	FString NodeName;
+		FString NodeName;
 
-	URigVMController* Controller = Host->GetRigVMClient()->GetController(ParentGraph);
-	Controller->OpenUndoBracket(TEXT("Add Variable"));
+		URigVMController* Controller = Host->GetRigVMClient()->GetController(ParentGraph);	
+		Controller->OpenUndoBracket(TEXT("Add Variable"));
 
-	if (URigVMNode* ModelNode = Controller->AddVariableNodeFromObjectPath(ExternalVariable.Name, TypeName, ObjectPath, bIsGetter, FString(), Location, NodeName, true, true))
-	{
-		for (UEdGraphNode* Node : ParentGraph->Nodes)
+		if (URigVMNode* ModelNode = Controller->AddVariableNodeFromObjectPath(ExtVariable.Name, TypeName, ObjectPath, bIsGetter, FString(), Location, NodeName, true, true))
 		{
-			if (URigVMEdGraphNode* RigNode = Cast<URigVMEdGraphNode>(Node))
+			for (UEdGraphNode* Node : ParentGraph->Nodes)
 			{
-				if (RigNode->GetModelNodeName() == ModelNode->GetFName())
+				if (URigVMEdGraphNode* RigNode = Cast<URigVMEdGraphNode>(Node))
 				{
-					NewNode = RigNode;
-					break;
+					if (RigNode->GetModelNodeName() == ModelNode->GetFName())
+					{
+						NewNode = RigNode;
+						break;
+					}
 				}
 			}
-		}
 
-		if (NewNode)
+			if (NewNode)
+			{
+				Controller->ClearNodeSelection(true);
+				Controller->SelectNode(ModelNode, true, true);
+			}
+			Controller->CloseUndoBracket();
+		}
+		else
 		{
-			Controller->ClearNodeSelection(true);
-			Controller->SelectNode(ModelNode, true, true);
+			Controller->CancelUndoBracket();
 		}
-		Controller->CloseUndoBracket();
+
+		return NewNode;
+	};
+
+	if (VariableAccessorChoice == EVariableAccessorChoice::Deferred)
+	{
+		FMenuBuilder MenuBuilder(true, NULL);
+		const FText SectionText = FText::FromString(FString::Printf(TEXT("Variable %s"), *ExternalVariable.Name.ToString()));
+
+		MenuBuilder.BeginSection("VariableDropped", SectionText);
+
+		MenuBuilder.AddMenuEntry(
+			FText::FromString(FString::Printf(TEXT("Get %s"), *ExternalVariable.Name.ToString())),
+			FText::FromString(FString::Printf(TEXT("Adds a getter node for variable %s"), *ExternalVariable.Name.ToString())),
+			FSlateIcon(),
+			FUIAction(
+				FExecuteAction::CreateLambda([AddNodeLambda, ParentGraph, ExtVariable = ExternalVariable, Location]
+				{
+					AddNodeLambda(ParentGraph, ExtVariable, Location, true);
+				}),
+				FCanExecuteAction()
+			)
+		);
+
+		MenuBuilder.AddMenuEntry(
+			FText::FromString(FString::Printf(TEXT("Set %s"), *ExternalVariable.Name.ToString())),
+			FText::FromString(FString::Printf(TEXT("Adds a setter node for variable %s"), *ExternalVariable.Name.ToString())),
+			FSlateIcon(),
+			FUIAction(
+				FExecuteAction::CreateLambda([AddNodeLambda, ParentGraph, Location, ExtVariable = ExternalVariable]
+				{
+					AddNodeLambda(ParentGraph, ExtVariable, Location, false);
+				}),
+				FCanExecuteAction()
+			)
+		);
+
+		MenuBuilder.EndSection();
+
+		FSlateApplication::Get().PushMenu(
+			FSlateApplication::Get().GetInteractiveTopLevelWindows()[0],
+			FWidgetPath(),
+			MenuBuilder.MakeWidget(),
+			FSlateApplication::Get().GetCursorPos(),
+			FPopupTransitionEffect(FPopupTransitionEffect::ContextMenu)
+		);
+
+		return nullptr;
 	}
 	else
 	{
-		Controller->CancelUndoBracket();
+		const bool bIsGetter = VariableAccessorChoice == EVariableAccessorChoice::Get;
+		return AddNodeLambda(ParentGraph, ExternalVariable, Location, bIsGetter);
 	}
-
-	return NewNode;
 }
 
 FAnimNextSchemaAction_AddComment::FAnimNextSchemaAction_AddComment()
diff --git a/Engine/Plugins/Experimental/AnimNext/Source/AnimNextEditor/Private/Common/GraphEditorSchemaActions.h b/Engine/Plugins/Experimental/AnimNext/Source/AnimNextEditor/Private/Common/GraphEditorSchemaActions.h
--- a/Engine/Plugins/Experimental/AnimNext/Source/AnimNextEditor/Private/Common/GraphEditorSchemaActions.h
+++ b/Engine/Plugins/Experimental/AnimNext/Source/AnimNextEditor/Private/Common/GraphEditorSchemaActions.h
@@ -90,9 +90,16 @@ struct FAnimNextSchemaAction_Variable : public FAnimNextSchemaAction
 {
 	GENERATED_BODY()
 
+	enum class EVariableAccessorChoice
+	{
+		Set,
+		Get,
+		Deferred
+	};
+
 	FAnimNextSchemaAction_Variable() = default;
 
-	FAnimNextSchemaAction_Variable(const FRigVMExternalVariable& InExternalVariable, bool bInIsGetter = true);
+	FAnimNextSchemaAction_Variable(const FRigVMExternalVariable& InExternalVariable, const EVariableAccessorChoice InVariableAccessorChoice);
 
 	// FEdGraphSchemaAction Interface
 	virtual UEdGraphNode* PerformAction(UEdGraph* ParentGraph, TArray<UEdGraphPin*>& FromPins, const FVector2D Location, bool bSelectNewNode = true) override;
@@ -110,7 +117,7 @@ struct FAnimNextSchemaAction_Variable : public FAnimNextSchemaAction
 
 private:
 	FRigVMExternalVariable ExternalVariable;
-	bool bIsGetter = false;
+	EVariableAccessorChoice VariableAccessorChoice = EVariableAccessorChoice::Set;
 	FLinearColor VariableColor;
 };
 
diff --git a/Engine/Plugins/Experimental/AnimNext/Source/AnimNextEditor/Private/Common/SActionMenu.cpp b/Engine/Plugins/Experimental/AnimNext/Source/AnimNextEditor/Private/Common/SActionMenu.cpp
--- a/Engine/Plugins/Experimental/AnimNext/Source/AnimNextEditor/Private/Common/SActionMenu.cpp
+++ b/Engine/Plugins/Experimental/AnimNext/Source/AnimNextEditor/Private/Common/SActionMenu.cpp
@@ -151,8 +151,8 @@ void SActionMenu::CollectAllAnimNextGraphActions(FGraphContextMenuBuilder& MenuB
 	TArray<FRigVMExternalVariable> ExternalVariables = RigVMHost->GetVM()->GetExternalVariables(RigVMHost->GetRigVMExtendedExecuteContext());
 	for (const FRigVMExternalVariable& ExternalVariable : ExternalVariables)
 	{
-		MenuBuilder.AddAction(MakeShared<FAnimNextSchemaAction_Variable>(ExternalVariable, true));
-		MenuBuilder.AddAction(MakeShared<FAnimNextSchemaAction_Variable>(ExternalVariable, false));
+		MenuBuilder.AddAction(MakeShared<FAnimNextSchemaAction_Variable>(ExternalVariable, FAnimNextSchemaAction_Variable::EVariableAccessorChoice::Set));
+		MenuBuilder.AddAction(MakeShared<FAnimNextSchemaAction_Variable>(ExternalVariable, FAnimNextSchemaAction_Variable::EVariableAccessorChoice::Get));
 	}
 
 	MenuBuilder.AddAction(MakeShared<FAnimNextSchemaAction_AddComment>());
diff --git a/Engine/Plugins/Experimental/AnimNext/Source/AnimNextEditor/Private/Variables/Outliner/VariablesOutlinerMode.cpp b/Engine/Plugins/Experimental/AnimNext/Source/AnimNextEditor/Private/Variables/Outliner/VariablesOutlinerMode.cpp
--- a/Engine/Plugins/Experimental/AnimNext/Source/AnimNextEditor/Private/Variables/Outliner/VariablesOutlinerMode.cpp
+++ b/Engine/Plugins/Experimental/AnimNext/Source/AnimNextEditor/Private/Variables/Outliner/VariablesOutlinerMode.cpp
@@ -363,7 +363,7 @@ TSharedPtr<FDragDropOperation> FVariablesOutlinerMode::CreateDragDropOperation(c
 		if(ExternalVariable.Name == VariableName)
 		{
 			const bool bIsGetter = !(MouseEvent.IsShiftDown() || MouseEvent.IsAltDown());
-			TSharedPtr<FAnimNextSchemaAction_Variable> Action = MakeShared<FAnimNextSchemaAction_Variable>(ExternalVariable, bIsGetter);
+			TSharedPtr<FAnimNextSchemaAction_Variable> Action = MakeShared<FAnimNextSchemaAction_Variable>(ExternalVariable, FAnimNextSchemaAction_Variable::EVariableAccessorChoice::Deferred);
 			return FVariableDragDropOp::New(Action);
 		}
 	}
