AnimNext: Fix memory stomp in debug rendering

The FPrimitiveSceneInfoData lifetime needs to live as long as the primitive does on the RT, so we allocate it seperately and release it on the RT.

diff --git a/Engine/Plugins/Experimental/AnimNext/Source/AnimNext/Internal/AnimNextDebugDraw.h b/Engine/Plugins/Experimental/AnimNext/Source/AnimNext/Internal/AnimNextDebugDraw.h
--- a/Engine/Plugins/Experimental/AnimNext/Source/AnimNext/Internal/AnimNextDebugDraw.h
+++ b/Engine/Plugins/Experimental/AnimNext/Source/AnimNext/Internal/AnimNextDebugDraw.h
@@ -68,9 +68,8 @@ private:
 
 	FCustomSceneProxyDesc SceneProxyDesc;
 	FPrimitiveSceneDesc SceneDesc;
-	FPrimitiveSceneInfoData SceneInfoData;
-	FAnimNextDebugSceneProxy* SceneProxy;
-	FCustomPrimitiveData CustomPrimitiveData;
+	FPrimitiveSceneInfoData* SceneInfoData = nullptr;
+	FAnimNextDebugSceneProxy* SceneProxy = nullptr;
 
 	FSceneInterface* Scene = nullptr;
 
diff --git a/Engine/Plugins/Experimental/AnimNext/Source/AnimNext/Private/AnimNextDebugDraw.cpp b/Engine/Plugins/Experimental/AnimNext/Source/AnimNext/Private/AnimNextDebugDraw.cpp
--- a/Engine/Plugins/Experimental/AnimNext/Source/AnimNext/Private/AnimNextDebugDraw.cpp
+++ b/Engine/Plugins/Experimental/AnimNext/Source/AnimNext/Private/AnimNextDebugDraw.cpp
@@ -135,10 +135,11 @@ FDebugDraw::FDebugDraw(UObject* InOwner)
 	: SceneProxyDesc(InOwner)
 	, SceneProxy(new FAnimNextDebugSceneProxy(SceneProxyDesc))
 {
-	SceneInfoData.SceneProxy = SceneProxy;
-	SceneInfoData.OwnerLastRenderTimePtr = nullptr;
+	SceneInfoData = new FPrimitiveSceneInfoData();
+	SceneInfoData->SceneProxy = SceneProxy;
+	SceneInfoData->OwnerLastRenderTimePtr = nullptr;
 	SceneDesc.ProxyDesc = &SceneProxyDesc;
-	SceneDesc.PrimitiveSceneData = &SceneInfoData;
+	SceneDesc.PrimitiveSceneData = SceneInfoData;
 	SceneDesc.World = InOwner->GetWorld();
 
 	SceneDesc.RenderMatrix = FMatrix::Identity;
@@ -154,6 +155,12 @@ FDebugDraw::FDebugDraw(UObject* InOwner)
 FDebugDraw::~FDebugDraw()
 {
 	Scene->RemovePrimitive(&SceneDesc);
+
+	// Delete the scene info data on the RT as we no longer own it
+	ENQUEUE_RENDER_COMMAND(AnimNextDebugDraw)([SceneInfoData = SceneInfoData](FRHICommandList& RHICmdList)
+	{
+		delete SceneInfoData;
+	});
 }
 
 void FDebugDraw::Draw()
@@ -179,7 +186,7 @@ void FDebugDraw::SetEnabled(bool bInIsEnabled)
 	if(bInIsEnabled != bIsEnabled)
 	{
 		// Update the RT copy of the flag if it differs
-		ENQUEUE_RENDER_COMMAND(AnimNextDebugDraw)([bInIsEnabled, SceneProxy = SceneProxy](FRHICommandList& RHICmdList) mutable
+		ENQUEUE_RENDER_COMMAND(AnimNextDebugDraw)([bInIsEnabled, SceneProxy = SceneProxy](FRHICommandList& RHICmdList)
 		{
 			SceneProxy->bIsEnabled = bInIsEnabled;
 		});
