AnimNext: Align tests with current asset entry permissions

diff --git a/Engine/Plugins/Experimental/AnimNext/Source/AnimNextTestSuite/Private/AnimNextEditorTest.cpp b/Engine/Plugins/Experimental/AnimNext/Source/AnimNextTestSuite/Private/AnimNextEditorTest.cpp
--- a/Engine/Plugins/Experimental/AnimNext/Source/AnimNextTestSuite/Private/AnimNextEditorTest.cpp
+++ b/Engine/Plugins/Experimental/AnimNext/Source/AnimNextTestSuite/Private/AnimNextEditorTest.cpp
@@ -144,8 +144,8 @@ bool FEditor_Graphs::RunTest(const FString& InParameters)
 	
 	FTestSettings TestSettings[] =
 	{
-		{ UAnimNextAnimationGraphFactory::StaticClass(), UAnimNextAnimationGraph::StaticClass(), true, false, true, true },
-		{ UAnimNextModuleFactory::StaticClass(), UAnimNextModule::StaticClass(), true, true, true, false },
+		{ UAnimNextAnimationGraphFactory::StaticClass(), UAnimNextAnimationGraph::StaticClass(), false, false, true, true },
+		{ UAnimNextModuleFactory::StaticClass(), UAnimNextModule::StaticClass(), true, true, false, false },
 	};
 
 	for(const FTestSettings& TestSetting : TestSettings)
diff --git a/Engine/Plugins/Experimental/AnimNext/Source/AnimNextTestSuite/Private/AnimNextVariablesTest.cpp b/Engine/Plugins/Experimental/AnimNext/Source/AnimNextTestSuite/Private/AnimNextVariablesTest.cpp
--- a/Engine/Plugins/Experimental/AnimNext/Source/AnimNextTestSuite/Private/AnimNextVariablesTest.cpp
+++ b/Engine/Plugins/Experimental/AnimNext/Source/AnimNextTestSuite/Private/AnimNextVariablesTest.cpp
@@ -16,7 +16,6 @@
 #include "Entries/AnimNextVariableEntry.h"
 #include "Factories/Factory.h"
 #include "Graph/AnimNextAnimationGraph.h"
-#include "Graph/AnimNextAnimationGraphFactory.h"
 #include "Module/AnimNextModuleFactory.h"
 #include "Module/RigUnit_AnimNextModuleEvents.h"
 
@@ -217,97 +216,82 @@ IMPLEMENT_SIMPLE_AUTOMATION_TEST(FVariables, "Animation.AnimNext.Variables", EAu
 
 bool FVariables::RunTest(const FString& InParameters)
 {
-	struct FFactoryAndClass
+	ON_SCOPE_EXIT{ FUtils::CleanupAfterTests(); };
+
+	UFactory* Factory = NewObject<UFactory>(GetTransientPackage(), UAnimNextModuleFactory::StaticClass());
+	UAnimNextRigVMAsset* Asset = CastChecked<UAnimNextRigVMAsset>(Factory->FactoryCreateNew(UAnimNextModule::StaticClass(), GetTransientPackage(), TEXT("TestAsset"), RF_Transient, nullptr, nullptr, NAME_None));
+	UE_RETURN_ON_ERROR(Asset != nullptr, "FEditor_Graphs -> Failed to create asset");
+
+	UAnimNextRigVMAssetEditorData* EditorData = UncookedOnly::FUtils::GetEditorData<UAnimNextRigVMAssetEditorData>(Asset);
+	UE_RETURN_ON_ERROR(EditorData != nullptr, "FEditor_Graphs -> Asset has no editor data.");
+
+	// Add variables
+	UAnimNextVariableEntry* OperandAEntry = EditorData->AddVariable(TEXT("A"), FAnimNextParamType::GetType<int32>(), TEXT("1"));
+	UE_RETURN_ON_ERROR(OperandAEntry != nullptr, TEXT("Could not create new variable in graph."));
+	UAnimNextVariableEntry* OperandBEntry = EditorData->AddVariable(TEXT("B"), FAnimNextParamType::GetType<int32>(), TEXT("2"));
+	UE_RETURN_ON_ERROR(OperandBEntry != nullptr, TEXT("Could not create new variable in graph."));
+	UAnimNextVariableEntry* ResultEntry = EditorData->AddVariable(TEXT("Result"), FAnimNextParamType::GetType<int32>(), TEXT("12"));
+	UE_RETURN_ON_ERROR(ResultEntry != nullptr, TEXT("Could not create new variable in graph."));
+
+	// Get event graph
+	UAnimNextEventGraphEntry* EventGraph = Cast<UAnimNextEventGraphEntry>(EditorData->FindEntry(TEXT("PrePhysics")));
+	if(EventGraph == nullptr)
 	{
-		TSubclassOf<UFactory> FactoryClass;
-		TSubclassOf<UAnimNextRigVMAsset> Class;
-	};
-
-	FFactoryAndClass FactoryClassPairs[] =
+		EventGraph = EditorData->AddEventGraph(TEXT("PrePhysics"), FRigUnit_AnimNextPrePhysicsEvent::StaticStruct());
+	}
+	UE_RETURN_ON_ERROR(EventGraph != nullptr, TEXT("Could not create new event graph in asset."));
+	
+	URigVMGraph* RigVMGraph = EventGraph->GetRigVMGraph();
+	UE_RETURN_ON_ERROR(RigVMGraph->GetNodes().Num() == 1, TEXT("Unexpected number of nodes in new event graph."));
+
+	URigVMNode* EventNode = RigVMGraph->GetNodes()[0];
+	check(EventNode);
+	URigVMPin* ExecutePin = EventNode->FindPin("ExecuteContext");
+	UE_RETURN_ON_ERROR(ExecutePin != nullptr, TEXT("Could find initial execute pin."));
+
+	UAnimNextController* Controller = Cast<UAnimNextController>(EditorData->GetController(EventGraph->GetRigVMGraph()));
+
+	URigVMVariableNode* VariableANode = Controller->AddVariableNode("A", RigVMTypeUtils::Int32Type, nullptr, true, TEXT(""));
+	UE_RETURN_ON_ERROR(VariableANode != nullptr, TEXT("Could not add get variable node."));
+	URigVMVariableNode* VariableBNode = Controller->AddVariableNode("B", RigVMTypeUtils::Int32Type, nullptr, true, TEXT(""));
+	UE_RETURN_ON_ERROR(VariableBNode != nullptr, TEXT("Could not add get variable node."));
+	URigVMVariableNode* SetResultNode = Controller->AddVariableNode("Result", RigVMTypeUtils::Int32Type, nullptr, false, TEXT(""));
+	UE_RETURN_ON_ERROR(SetResultNode != nullptr, TEXT("Could not add set variable node."));
+
+	URigVMUnitNode* TestOpUnitNode = Controller->AddUnitNode(FAnimNextTests_TestOperation::StaticStruct());
+	bool bLinkAAdded = Controller->AddLink(VariableANode->FindPin("Value"), TestOpUnitNode->FindPin("A"));
+	UE_RETURN_ON_ERROR(bLinkAAdded, TEXT("Could not link variable node."));
+	bool bLinkBAdded = Controller->AddLink(VariableBNode->FindPin("Value"), TestOpUnitNode->FindPin("B"));
+	UE_RETURN_ON_ERROR(bLinkBAdded, TEXT("Could not link variable node."));
+	bool bLinkResultAdded = Controller->AddLink(TestOpUnitNode->FindPin("Result"), SetResultNode->FindPin("Value"));
+	UE_RETURN_ON_ERROR(bLinkResultAdded, TEXT("Could not link variable node."));
+
+	bool bLinkExec1Added = Controller->AddLink(EventNode->FindPin(FRigVMStruct::ExecuteContextName.ToString()), TestOpUnitNode->FindPin(FRigVMStruct::ExecuteContextName.ToString()));
+	UE_RETURN_ON_ERROR(bLinkExec1Added, TEXT("Could not link variable node exec."));
+
+	bool bLinkExec2Added = Controller->AddLink(TestOpUnitNode->FindPin(FRigVMStruct::ExecuteContextName.ToString()), SetResultNode->FindPin(FRigVMStruct::ExecuteContextName.ToString()));
+	UE_RETURN_ON_ERROR(bLinkExec1Added, TEXT("Could not link variable node exec."));
+
+	URigVMUnitNode* PrintResultUnitNode = Controller->AddUnitNode(FAnimNextTests_PrintResult::StaticStruct());
+	bool bLinkExec3Added = Controller->AddLink(SetResultNode->FindPin(FRigVMStruct::ExecuteContextName.ToString()), PrintResultUnitNode->FindPin(FRigVMStruct::ExecuteContextName.ToString()));
+	URigVMVariableNode* GetResultNode = Controller->AddVariableNode("Result", RigVMTypeUtils::Int32Type, nullptr, true, TEXT(""));
+	UE_RETURN_ON_ERROR(GetResultNode != nullptr, TEXT("Could not add get variable node."));
+	bool bLinkResult2Added = Controller->AddLink(GetResultNode->FindPin("Value"), PrintResultUnitNode->FindPin("Result"));
+	UE_RETURN_ON_ERROR(bLinkResult2Added, TEXT("Could not link variable node."));
+
+	TArray<FString> Messages;
+	FRigVMRuntimeSettings RuntimeSettings;
+	RuntimeSettings.SetLogFunction([&Messages](const FRigVMLogSettings& InLogSettings, const FRigVMExecuteContext* InContext, const FString& Message)
 	{
-		{ UAnimNextAnimationGraphFactory::StaticClass(), UAnimNextAnimationGraph::StaticClass() },
-		{ UAnimNextModuleFactory::StaticClass(), UAnimNextModule::StaticClass() },
-	};
+		Messages.Add(Message);
+	});
+	Asset->GetRigVMExtendedExecuteContext().SetRuntimeSettings(RuntimeSettings);
 
-	for(const FFactoryAndClass& FactoryAndClass : FactoryClassPairs)
-	{
-		ON_SCOPE_EXIT{ FUtils::CleanupAfterTests(); };
-
-		UFactory* Factory = NewObject<UFactory>(GetTransientPackage(), FactoryAndClass.FactoryClass);
-		UAnimNextRigVMAsset* Asset = CastChecked<UAnimNextRigVMAsset>(Factory->FactoryCreateNew(FactoryAndClass.Class, GetTransientPackage(), TEXT("TestAsset"), RF_Transient, nullptr, nullptr, NAME_None));
-		UE_RETURN_ON_ERROR(Asset != nullptr, "FEditor_Graphs -> Failed to create asset");
-
-		UAnimNextRigVMAssetEditorData* EditorData = UncookedOnly::FUtils::GetEditorData<UAnimNextRigVMAssetEditorData>(Asset);
-		UE_RETURN_ON_ERROR(EditorData != nullptr, "FEditor_Graphs -> Asset has no editor data.");
-
-		// Add variables
-		UAnimNextVariableEntry* OperandAEntry = EditorData->AddVariable(TEXT("A"), FAnimNextParamType::GetType<int32>(), TEXT("1"));
-		UE_RETURN_ON_ERROR(OperandAEntry != nullptr, TEXT("Could not create new variable in graph."));
-		UAnimNextVariableEntry* OperandBEntry = EditorData->AddVariable(TEXT("B"), FAnimNextParamType::GetType<int32>(), TEXT("2"));
-		UE_RETURN_ON_ERROR(OperandBEntry != nullptr, TEXT("Could not create new variable in graph."));
-		UAnimNextVariableEntry* ResultEntry = EditorData->AddVariable(TEXT("Result"), FAnimNextParamType::GetType<int32>(), TEXT("12"));
-		UE_RETURN_ON_ERROR(ResultEntry != nullptr, TEXT("Could not create new variable in graph."));
-
-		// Add event graph
-		UAnimNextEventGraphEntry* EventGraph = Cast<UAnimNextEventGraphEntry>(EditorData->FindEntry(TEXT("PrePhysics")));
-		if(EventGraph == nullptr)
-		{
-			EventGraph = EditorData->AddEventGraph(TEXT("PrePhysics"), FRigUnit_AnimNextPrePhysicsEvent::StaticStruct());
-		}
-		UE_RETURN_ON_ERROR(EventGraph != nullptr, TEXT("Could not create new event graph in asset."));
-		
-		URigVMGraph* RigVMGraph = EventGraph->GetRigVMGraph();
-		UE_RETURN_ON_ERROR(RigVMGraph->GetNodes().Num() == 1, TEXT("Unexpected number of nodes in new event graph."));
-
-		URigVMNode* EventNode = RigVMGraph->GetNodes()[0];
-		check(EventNode);
-		URigVMPin* ExecutePin = EventNode->FindPin("ExecuteContext");
-		UE_RETURN_ON_ERROR(ExecutePin != nullptr, TEXT("Could find initial execute pin."));
-
-		UAnimNextController* Controller = Cast<UAnimNextController>(EditorData->GetController(EventGraph->GetRigVMGraph()));
-
-		URigVMVariableNode* VariableANode = Controller->AddVariableNode("A", RigVMTypeUtils::Int32Type, nullptr, true, TEXT(""));
-		UE_RETURN_ON_ERROR(VariableANode != nullptr, TEXT("Could not add get variable node."));
-		URigVMVariableNode* VariableBNode = Controller->AddVariableNode("B", RigVMTypeUtils::Int32Type, nullptr, true, TEXT(""));
-		UE_RETURN_ON_ERROR(VariableBNode != nullptr, TEXT("Could not add get variable node."));
-		URigVMVariableNode* SetResultNode = Controller->AddVariableNode("Result", RigVMTypeUtils::Int32Type, nullptr, false, TEXT(""));
-		UE_RETURN_ON_ERROR(SetResultNode != nullptr, TEXT("Could not add set variable node."));
-
-		URigVMUnitNode* TestOpUnitNode = Controller->AddUnitNode(FAnimNextTests_TestOperation::StaticStruct());
-		bool bLinkAAdded = Controller->AddLink(VariableANode->FindPin("Value"), TestOpUnitNode->FindPin("A"));
-		UE_RETURN_ON_ERROR(bLinkAAdded, TEXT("Could not link variable node."));
-		bool bLinkBAdded = Controller->AddLink(VariableBNode->FindPin("Value"), TestOpUnitNode->FindPin("B"));
-		UE_RETURN_ON_ERROR(bLinkBAdded, TEXT("Could not link variable node."));
-		bool bLinkResultAdded = Controller->AddLink(TestOpUnitNode->FindPin("Result"), SetResultNode->FindPin("Value"));
-		UE_RETURN_ON_ERROR(bLinkResultAdded, TEXT("Could not link variable node."));
-
-		bool bLinkExec1Added = Controller->AddLink(EventNode->FindPin(FRigVMStruct::ExecuteContextName.ToString()), TestOpUnitNode->FindPin(FRigVMStruct::ExecuteContextName.ToString()));
-		UE_RETURN_ON_ERROR(bLinkExec1Added, TEXT("Could not link variable node exec."));
-
-		bool bLinkExec2Added = Controller->AddLink(TestOpUnitNode->FindPin(FRigVMStruct::ExecuteContextName.ToString()), SetResultNode->FindPin(FRigVMStruct::ExecuteContextName.ToString()));
-		UE_RETURN_ON_ERROR(bLinkExec1Added, TEXT("Could not link variable node exec."));
-
-		URigVMUnitNode* PrintResultUnitNode = Controller->AddUnitNode(FAnimNextTests_PrintResult::StaticStruct());
-		bool bLinkExec3Added = Controller->AddLink(SetResultNode->FindPin(FRigVMStruct::ExecuteContextName.ToString()), PrintResultUnitNode->FindPin(FRigVMStruct::ExecuteContextName.ToString()));
-		URigVMVariableNode* GetResultNode = Controller->AddVariableNode("Result", RigVMTypeUtils::Int32Type, nullptr, true, TEXT(""));
-		UE_RETURN_ON_ERROR(GetResultNode != nullptr, TEXT("Could not add get variable node."));
-		bool bLinkResult2Added = Controller->AddLink(GetResultNode->FindPin("Value"), PrintResultUnitNode->FindPin("Result"));
-		UE_RETURN_ON_ERROR(bLinkResult2Added, TEXT("Could not link variable node."));
-
-		TArray<FString> Messages;
-		FRigVMRuntimeSettings RuntimeSettings;
-		RuntimeSettings.SetLogFunction([&Messages](const FRigVMLogSettings& InLogSettings, const FRigVMExecuteContext* InContext, const FString& Message)
-		{
-			Messages.Add(Message);
-		});
-		Asset->GetRigVMExtendedExecuteContext().SetRuntimeSettings(RuntimeSettings);
-
-		Asset->GetVM()->ExecuteVM(Asset->GetRigVMExtendedExecuteContext(), FRigUnit_AnimNextPrePhysicsEvent::EventName);
+	Asset->GetVM()->ExecuteVM(Asset->GetRigVMExtendedExecuteContext(), FRigUnit_AnimNextPrePhysicsEvent::EventName);
 
-		UE_RETURN_ON_ERROR(Messages.Num() == 1, TEXT("unexpected number of messages"));
-		UE_RETURN_ON_ERROR(Messages[0] == TEXT("Result = 3"), TEXT("unexpected result message"));
-
-	}
+	UE_RETURN_ON_ERROR(Messages.Num() == 1, TEXT("unexpected number of messages"));
+	UE_RETURN_ON_ERROR(Messages[0] == TEXT("Result = 3"), TEXT("unexpected result message"));
+	
 	return true;
 }
 
